Animation: Pose search experimental plugin
- This is a plugin that enables the animation runtime to search within a sequence for a given pose
- UPoseSearchSchema: specifies list of bones to sample, will include other info in the future like which properties to sample
- UPoseSearchIndexConfig: anim asset metadata object that specifies sample rate and range
- UPoseSearchIndex: anim asset metadata object containing data sampled from sequence offline, currently just translations per bone in the schema
- FPoseSearchPoseHistory: maintains a buffer of recorded poses over time for constructing search queries with past poses
- PoseSearchBuildIndex(): Constructs a search index for an anim sequence
- PoseSearchBuildQuery(): Constructs a search query given the current pose and a pose search schema
- PoseSearch(): Searches for the nearest matching pose in an indexed sequence given a query that was built for the same pose search schema
- FAnimNode_PoseMatchingSequencePlayer: Subclass of sequence player that adds the option to offset playback based on a pose search
- FAnimNode_PoseSearchHistoryCollector: a new pose history anim node for explicit control of pose search query context
  - Pose history is now updated during graph evaluation instead of always using the final pose result at the end of evaluation
  - Removed pose history from anim instance to decouple plugin from anim runtime
- Removed dependencies on USkeletalMeshComponent
- Promoted anim update shared context to base context so ancestor tracking can be used during node initialization
- Added FBoneIndexBase::IsValid()
- Exported UAnimationAsset API so plugins can use it
- Improved documentation
- Register asset type actions on module startup

diff --git a/Engine/Plugins/Animation/PoseSearch/PoseSearch.uplugin b/Engine/Plugins/Animation/PoseSearch/PoseSearch.uplugin
new file mode 100644
--- /dev/null
+++ b/Engine/Plugins/Animation/PoseSearch/PoseSearch.uplugin
@@ -0,0 +1,30 @@
+{
+	"FileVersion": 3,
+	"Version": 1,
+	"VersionName": "0.1",
+	"FriendlyName": "Pose Search",
+	"Description": "Experimental pose search API",
+	"Category": "Animation",
+	"CreatedBy": "Epic Games, Inc.",
+	"CreatedByURL": "http://epicgames.com",
+	"DocsURL": "",
+	"MarketplaceURL": "",
+	"SupportURL": "",
+	"EnabledByDefault": false,
+	"CanContainContent": true,
+	"IsBetaVersion": false,
+	"IsExperimentalVersion": true,
+	"Installed": false,
+	"Modules": [
+		{
+			"Name": "PoseSearch",
+			"Type": "Runtime",
+			"LoadingPhase": "PreDefault"
+		},
+		{
+			"Name": "PoseSearchEditor",
+			"Type": "UncookedOnly",
+			"LoadingPhase": "PreDefault"
+		}
+	]
+}
\ No newline at end of file
diff --git a/Engine/Plugins/Animation/PoseSearch/Source/Editor/PoseSearchEditor.Build.cs b/Engine/Plugins/Animation/PoseSearch/Source/Editor/PoseSearchEditor.Build.cs
new file mode 100644
--- /dev/null
+++ b/Engine/Plugins/Animation/PoseSearch/Source/Editor/PoseSearchEditor.Build.cs
@@ -0,0 +1,33 @@
+// Copyright Epic Games, Inc. All Rights Reserved.
+
+using UnrealBuildTool;
+
+public class PoseSearchEditor : ModuleRules
+{
+	public PoseSearchEditor(ReadOnlyTargetRules Target) : base(Target)
+	{
+		PCHUsage = ModuleRules.PCHUsageMode.UseExplicitOrSharedPCHs;
+
+		PublicDependencyModuleNames.AddRange(
+			new string[] {
+				"AnimGraph",
+				"AnimationCore",
+				"Core",
+				"CoreUObject",
+				"Engine",
+				"PoseSearch",
+				"ToolMenus",
+				"UnrealEd",
+			}
+		);
+
+		PrivateDependencyModuleNames.AddRange(
+			new string[] {
+				"AssetTools",
+				"BlueprintGraph",
+				"GraphEditor",
+				"SlateCore",
+			}
+		);
+	}
+}
\ No newline at end of file
diff --git a/Engine/Plugins/Animation/PoseSearch/Source/Editor/Private/AnimGraphNode_PoseMatchingSequencePlayer.cpp b/Engine/Plugins/Animation/PoseSearch/Source/Editor/Private/AnimGraphNode_PoseMatchingSequencePlayer.cpp
new file mode 100644
--- /dev/null
+++ b/Engine/Plugins/Animation/PoseSearch/Source/Editor/Private/AnimGraphNode_PoseMatchingSequencePlayer.cpp
@@ -0,0 +1,472 @@
+// Copyright Epic Games, Inc. All Rights Reserved.
+
+#include "AnimGraphNode_PoseMatchingSequencePlayer.h"
+#include "EdGraphSchema_K2_Actions.h"
+#include "Modules/ModuleManager.h"
+#include "ToolMenus.h"
+#include "Kismet2/CompilerResultsLog.h"
+#include "GraphEditorActions.h"
+#include "ARFilter.h"
+#include "AssetRegistryModule.h"
+#include "BlueprintActionFilter.h"
+#include "BlueprintActionDatabaseRegistrar.h"
+#include "EditorCategoryUtils.h"
+#include "BlueprintNodeSpawner.h"
+#include "Animation/AnimComposite.h"
+#include "Animation/AnimSequence.h"
+
+#define LOCTEXT_NAMESPACE "A3Nodes"
+
+/////////////////////////////////////////////////////
+// FNewSequencePlayerAction
+
+// Action to add a sequence player node to the graph
+struct FNewSequencePlayerAction : public FEdGraphSchemaAction_K2NewNode
+{
+protected:
+	FAssetData AssetInfo;
+public:
+	FNewSequencePlayerAction(const FAssetData& InAssetInfo, FText Title)
+		: FEdGraphSchemaAction_K2NewNode(LOCTEXT("Animation", "Animations"), Title, LOCTEXT("EvalAnimSequenceToMakePose", "Evaluates an animation sequence to produce a pose"), 0, FText::FromName(InAssetInfo.ObjectPath))
+	{
+		AssetInfo = InAssetInfo;
+
+		UAnimGraphNode_PoseMatchingSequencePlayer* Template = NewObject<UAnimGraphNode_PoseMatchingSequencePlayer>();
+		NodeTemplate = Template;
+	}
+
+	virtual UEdGraphNode* PerformAction(class UEdGraph* ParentGraph, UEdGraphPin* FromPin, const FVector2D Location, bool bSelectNewNode = true) override
+	{
+		UAnimGraphNode_PoseMatchingSequencePlayer* SpawnedNode = CastChecked<UAnimGraphNode_PoseMatchingSequencePlayer>(FEdGraphSchemaAction_K2NewNode::PerformAction(ParentGraph, FromPin, Location, bSelectNewNode));
+		SpawnedNode->Node.Sequence = Cast<UAnimSequence>(AssetInfo.GetAsset());
+
+		return SpawnedNode;
+	}
+};
+
+/////////////////////////////////////////////////////
+// UAnimGraphNode_PoseMatchingSequencePlayer
+
+UAnimGraphNode_PoseMatchingSequencePlayer::UAnimGraphNode_PoseMatchingSequencePlayer(const FObjectInitializer& ObjectInitializer)
+	: Super(ObjectInitializer)
+{
+}
+
+void UAnimGraphNode_PoseMatchingSequencePlayer::PreloadRequiredAssets()
+{
+	PreloadObject(Node.Sequence);
+
+	Super::PreloadRequiredAssets();
+}
+
+FLinearColor UAnimGraphNode_PoseMatchingSequencePlayer::GetNodeTitleColor() const
+{
+	if ((Node.Sequence != NULL) && Node.Sequence->IsValidAdditive())
+	{
+		return FLinearColor(0.10f, 0.60f, 0.12f);
+	}
+	else
+	{
+		return FColor(200, 100, 100);
+	}
+}
+
+FText UAnimGraphNode_PoseMatchingSequencePlayer::GetTooltipText() const
+{
+	if (!Node.Sequence)
+	{
+		return FText();
+	}
+
+	const bool bAdditive = Node.Sequence->IsValidAdditive();
+	return GetTitleGivenAssetInfo(FText::FromString(Node.Sequence->GetPathName()), bAdditive);
+}
+
+FText UAnimGraphNode_PoseMatchingSequencePlayer::GetNodeTitleForSequence(ENodeTitleType::Type TitleType, UAnimSequenceBase* InSequence) const
+{
+	const bool bAdditive = InSequence->IsValidAdditive();
+	const FText BasicTitle = GetTitleGivenAssetInfo(FText::FromName(InSequence->GetFName()), bAdditive);
+
+	if (SyncGroup.GroupName == NAME_None)
+	{
+		return BasicTitle;
+	}
+	else
+	{
+		const FText SyncGroupName = FText::FromName(SyncGroup.GroupName);
+
+		FFormatNamedArguments Args;
+		Args.Add(TEXT("Title"), BasicTitle);
+		Args.Add(TEXT("SyncGroup"), SyncGroupName);
+
+		if (TitleType == ENodeTitleType::FullTitle)
+		{
+			return FText::Format(LOCTEXT("SequenceNodeGroupWithSubtitleFull", "{Title}\nSync group {SyncGroup}"), Args);
+		}
+		else
+		{
+			return FText::Format(LOCTEXT("SequenceNodeGroupWithSubtitleList", "{Title} (Sync group {SyncGroup})"), Args);
+		}
+	}
+}
+
+FText UAnimGraphNode_PoseMatchingSequencePlayer::GetNodeTitle(ENodeTitleType::Type TitleType) const
+{
+	if (Node.Sequence == nullptr)
+	{
+		// we may have a valid variable connected or default pin value
+		UEdGraphPin* SequencePin = FindPin(GET_MEMBER_NAME_STRING_CHECKED(FAnimNode_PoseMatchingSequencePlayer, Sequence));
+		if (SequencePin && SequencePin->LinkedTo.Num() > 0)
+		{
+			return LOCTEXT("SequenceNodeTitleVariable", "Play Animation Sequence (pose matching)");
+		}
+		else if (SequencePin && SequencePin->DefaultObject != nullptr)
+		{
+			return GetNodeTitleForSequence(TitleType, CastChecked<UAnimSequenceBase>(SequencePin->DefaultObject));
+		}
+		else
+		{
+			return LOCTEXT("SequenceNullTitle", "Play (None) (pose matching)");
+		}
+	}
+	else
+	{
+		return GetNodeTitleForSequence(TitleType, Node.Sequence);
+	}
+}
+
+FText UAnimGraphNode_PoseMatchingSequencePlayer::GetTitleGivenAssetInfo(const FText& AssetName, bool bKnownToBeAdditive)
+{
+	FFormatNamedArguments Args;
+	Args.Add(TEXT("AssetName"), AssetName);
+
+	if (bKnownToBeAdditive)
+	{
+		return FText::Format(LOCTEXT("SequenceNodeTitleAdditive", "Play {AssetName} (additive, pose matching)"), Args);
+	}
+	else
+	{
+		return FText::Format(LOCTEXT("SequenceNodeTitle", "Play {AssetName} (pose matching)"), Args);
+	}
+}
+
+FText UAnimGraphNode_PoseMatchingSequencePlayer::GetMenuCategory() const
+{
+	return FEditorCategoryUtils::GetCommonCategory(FCommonEditorCategory::Animation);
+}
+
+void UAnimGraphNode_PoseMatchingSequencePlayer::GetMenuActions(FBlueprintActionDatabaseRegistrar& ActionRegistrar) const
+{
+	auto LoadedAssetSetup = [](UEdGraphNode* NewNode, bool /*bIsTemplateNode*/, TWeakObjectPtr<UAnimSequence> SequencePtr)
+	{
+		UAnimGraphNode_PoseMatchingSequencePlayer* SequencePlayerNode = CastChecked<UAnimGraphNode_PoseMatchingSequencePlayer>(NewNode);
+		SequencePlayerNode->Node.Sequence = SequencePtr.Get();
+	};
+
+	auto UnloadedAssetSetup = [](UEdGraphNode* NewNode, bool bIsTemplateNode, const FAssetData AssetData)
+	{
+		UAnimGraphNode_PoseMatchingSequencePlayer* SequencePlayerNode = CastChecked<UAnimGraphNode_PoseMatchingSequencePlayer>(NewNode);
+		if (bIsTemplateNode)
+		{
+			AssetData.GetTagValue("Skeleton", SequencePlayerNode->UnloadedSkeletonName);
+		}
+		else
+		{
+			UAnimSequence* Sequence = Cast<UAnimSequence>(AssetData.GetAsset());
+			check(Sequence != nullptr);
+			SequencePlayerNode->Node.Sequence = Sequence;
+		}
+	};	
+
+	const UObject* QueryObject = ActionRegistrar.GetActionKeyFilter();
+	if (QueryObject == nullptr)
+	{
+		FAssetRegistryModule& AssetRegistryModule = FModuleManager::LoadModuleChecked<FAssetRegistryModule>(TEXT("AssetRegistry"));
+		// define a filter to help in pulling UAnimSequence asset data from the registry
+		FARFilter Filter;
+		Filter.ClassNames.Add(UAnimSequence::StaticClass()->GetFName());
+		Filter.bRecursiveClasses = true;
+		// Find matching assets and add an entry for each one
+		TArray<FAssetData> SequenceList;
+		AssetRegistryModule.Get().GetAssets(Filter, /*out*/SequenceList);
+
+		for (auto AssetIt = SequenceList.CreateConstIterator(); AssetIt; ++AssetIt)
+		{
+			const FAssetData& Asset = *AssetIt;			
+
+			UBlueprintNodeSpawner* NodeSpawner = UBlueprintNodeSpawner::Create(GetClass());
+			if (Asset.IsAssetLoaded())
+			{
+				UAnimSequence* AnimSequence = Cast<UAnimSequence>(Asset.GetAsset());
+				if(AnimSequence)
+				{
+					NodeSpawner->CustomizeNodeDelegate = UBlueprintNodeSpawner::FCustomizeNodeDelegate::CreateStatic(LoadedAssetSetup, TWeakObjectPtr<UAnimSequence>(AnimSequence));
+					NodeSpawner->DefaultMenuSignature.MenuName = GetTitleGivenAssetInfo(FText::FromName(AnimSequence->GetFName()), AnimSequence->IsValidAdditive());
+					NodeSpawner->DefaultMenuSignature.Tooltip = GetTitleGivenAssetInfo(FText::FromString(AnimSequence->GetPathName()), AnimSequence->IsValidAdditive());
+				}
+			}
+			else
+			{
+				const FString TagValue = Asset.GetTagValueRef<FString>(GET_MEMBER_NAME_CHECKED(UAnimSequence, AdditiveAnimType));
+				const bool bKnownToBeAdditive = (!TagValue.IsEmpty() && !TagValue.Equals(TEXT("AAT_None")));
+
+				NodeSpawner->CustomizeNodeDelegate = UBlueprintNodeSpawner::FCustomizeNodeDelegate::CreateStatic(UnloadedAssetSetup, Asset);
+				NodeSpawner->DefaultMenuSignature.MenuName = GetTitleGivenAssetInfo(FText::FromName(Asset.AssetName), bKnownToBeAdditive);
+				NodeSpawner->DefaultMenuSignature.Tooltip = GetTitleGivenAssetInfo(FText::FromName(Asset.ObjectPath), bKnownToBeAdditive);
+			}
+			ActionRegistrar.AddBlueprintAction(Asset, NodeSpawner);
+		}
+	}
+	else if (const UAnimSequence* AnimSequence = Cast<UAnimSequence>(QueryObject))
+	{
+		UBlueprintNodeSpawner* NodeSpawner = UBlueprintNodeSpawner::Create(GetClass());
+
+		TWeakObjectPtr<UAnimSequence> SequencePtr = MakeWeakObjectPtr(const_cast<UAnimSequence*>(AnimSequence));
+		NodeSpawner->CustomizeNodeDelegate = UBlueprintNodeSpawner::FCustomizeNodeDelegate::CreateStatic(LoadedAssetSetup, SequencePtr);
+		NodeSpawner->DefaultMenuSignature.MenuName = GetTitleGivenAssetInfo(FText::FromName(AnimSequence->GetFName()), AnimSequence->IsValidAdditive());
+		NodeSpawner->DefaultMenuSignature.Tooltip = GetTitleGivenAssetInfo(FText::FromString(AnimSequence->GetPathName()), AnimSequence->IsValidAdditive());
+
+		ActionRegistrar.AddBlueprintAction(QueryObject, NodeSpawner);
+	}
+	else if (QueryObject == GetClass())
+	{
+		FAssetRegistryModule& AssetRegistryModule = FModuleManager::LoadModuleChecked<FAssetRegistryModule>(TEXT("AssetRegistry"));
+		// define a filter to help in pulling UAnimSequence asset data from the registry
+		FARFilter Filter;
+		Filter.ClassNames.Add(UAnimSequence::StaticClass()->GetFName());
+		Filter.bRecursiveClasses = true;
+		// Find matching assets and add an entry for each one
+		TArray<FAssetData> SequenceList;
+		AssetRegistryModule.Get().GetAssets(Filter, /*out*/SequenceList);
+
+		for (auto AssetIt = SequenceList.CreateConstIterator(); AssetIt; ++AssetIt)
+		{
+			const FAssetData& Asset = *AssetIt;
+			if (Asset.IsAssetLoaded())
+			{
+				continue;
+			}
+
+			UBlueprintNodeSpawner* NodeSpawner = UBlueprintNodeSpawner::Create(GetClass());
+			NodeSpawner->CustomizeNodeDelegate = UBlueprintNodeSpawner::FCustomizeNodeDelegate::CreateStatic(UnloadedAssetSetup, Asset);
+
+			const FString TagValue = Asset.GetTagValueRef<FString>(GET_MEMBER_NAME_CHECKED(UAnimSequence, AdditiveAnimType));
+			const bool bKnownToBeAdditive = (!TagValue.IsEmpty() && !TagValue.Equals(TEXT("AAT_None")));
+
+			NodeSpawner->DefaultMenuSignature.MenuName = GetTitleGivenAssetInfo(FText::FromName(Asset.AssetName), bKnownToBeAdditive);
+			NodeSpawner->DefaultMenuSignature.Tooltip = GetTitleGivenAssetInfo(FText::FromName(Asset.ObjectPath), bKnownToBeAdditive);
+			ActionRegistrar.AddBlueprintAction(Asset, NodeSpawner);
+		}
+	}	
+}
+
+bool UAnimGraphNode_PoseMatchingSequencePlayer::IsActionFilteredOut(class FBlueprintActionFilter const& Filter)
+{
+	bool bIsFilteredOut = false;
+	FBlueprintActionContext const& FilterContext = Filter.Context;
+
+	for (UBlueprint* Blueprint : FilterContext.Blueprints)
+	{
+		if (UAnimBlueprint* AnimBlueprint = Cast<UAnimBlueprint>(Blueprint))
+		{
+			if(Node.Sequence)
+			{
+				if(Node.Sequence->GetSkeleton() != AnimBlueprint->TargetSkeleton)
+				{
+					// Sequence does not use the same skeleton as the Blueprint, cannot use
+					bIsFilteredOut = true;
+					break;
+				}
+			}
+			else 
+			{
+				FAssetData SkeletonData(AnimBlueprint->TargetSkeleton);
+				if(UnloadedSkeletonName != SkeletonData.GetExportTextName())
+				{
+					bIsFilteredOut = true;
+					break;
+				}
+			}
+		}
+		else
+		{
+			// Not an animation Blueprint, cannot use
+			bIsFilteredOut = true;
+			break;
+		}
+	}
+	return bIsFilteredOut;
+}
+
+EAnimAssetHandlerType UAnimGraphNode_PoseMatchingSequencePlayer::SupportsAssetClass(const UClass* AssetClass) const
+{
+	if (AssetClass->IsChildOf(UAnimSequence::StaticClass()) || AssetClass->IsChildOf(UAnimComposite::StaticClass()))
+	{
+		return EAnimAssetHandlerType::PrimaryHandler;
+	}
+	else
+	{
+		return EAnimAssetHandlerType::NotSupported;
+	}
+}
+
+void UAnimGraphNode_PoseMatchingSequencePlayer::ValidateAnimNodeDuringCompilation(class USkeleton* ForSkeleton, class FCompilerResultsLog& MessageLog)
+{
+	Super::ValidateAnimNodeDuringCompilation(ForSkeleton, MessageLog);
+
+	UAnimSequenceBase* SequenceToCheck = Node.Sequence;
+	UEdGraphPin* SequencePin = FindPin(GET_MEMBER_NAME_STRING_CHECKED(FAnimNode_PoseMatchingSequencePlayer, Sequence));
+	if (SequencePin != nullptr && SequenceToCheck == nullptr)
+	{
+		SequenceToCheck = Cast<UAnimSequenceBase>(SequencePin->DefaultObject);
+	}
+
+	if (SequenceToCheck == nullptr)
+	{
+		// we may have a connected node
+		if (SequencePin == nullptr || SequencePin->LinkedTo.Num() == 0)
+		{
+			MessageLog.Error(TEXT("@@ references an unknown sequence"), this);
+		}
+	}
+	else if(SupportsAssetClass(SequenceToCheck->GetClass()) == EAnimAssetHandlerType::NotSupported)
+	{
+		MessageLog.Error(*FText::Format(LOCTEXT("UnsupportedAssetError", "@@ is trying to play a {0} as a sequence, which is not allowed."), SequenceToCheck->GetClass()->GetDisplayNameText()).ToString(), this);
+	}
+	else
+	{
+		USkeleton* SeqSkeleton = SequenceToCheck->GetSkeleton();
+		if (SeqSkeleton&& // if anim sequence doesn't have skeleton, it might be due to anim sequence not loaded yet, @todo: wait with anim blueprint compilation until all assets are loaded?
+			!SeqSkeleton->IsCompatible(ForSkeleton))
+		{
+			MessageLog.Error(TEXT("@@ references sequence that uses different skeleton @@"), this, SeqSkeleton);
+		}
+	}
+}
+
+void UAnimGraphNode_PoseMatchingSequencePlayer::GetNodeContextMenuActions(UToolMenu* Menu, UGraphNodeContextMenuContext* Context) const
+{
+	if (!Context->bIsDebugging)
+	{
+		// add an option to convert to single frame
+		{
+			FToolMenuSection& Section = Menu->AddSection("AnimGraphNodeSequencePlayer", NSLOCTEXT("A3Nodes", "SequencePlayerHeading", "Sequence Player"));
+			Section.AddMenuEntry(FGraphEditorCommands::Get().OpenRelatedAsset);
+			Section.AddMenuEntry(FGraphEditorCommands::Get().ConvertToSeqEvaluator);
+		}
+	}
+}
+
+void UAnimGraphNode_PoseMatchingSequencePlayer::SetAnimationAsset(UAnimationAsset* Asset)
+{
+	if (UAnimSequenceBase* Seq = Cast<UAnimSequenceBase>(Asset))
+	{
+		Node.Sequence = Seq;
+	}
+}
+
+void UAnimGraphNode_PoseMatchingSequencePlayer::BakeDataDuringCompilation(class FCompilerResultsLog& MessageLog)
+{
+	UAnimBlueprint* AnimBlueprint = GetAnimBlueprint();
+	Node.GroupIndex = AnimBlueprint->FindOrAddGroup(SyncGroup.GroupName);
+	Node.GroupRole = SyncGroup.GroupRole;
+}
+
+void UAnimGraphNode_PoseMatchingSequencePlayer::GetAllAnimationSequencesReferred(TArray<UAnimationAsset*>& AnimationAssets) const
+{
+	if(Node.Sequence)
+	{
+		HandleAnimReferenceCollection(Node.Sequence, AnimationAssets);
+	}
+}
+
+void UAnimGraphNode_PoseMatchingSequencePlayer::ReplaceReferredAnimations(const TMap<UAnimationAsset*, UAnimationAsset*>& AnimAssetReplacementMap)
+{
+	HandleAnimReferenceReplacement(Node.Sequence, AnimAssetReplacementMap);
+}
+
+
+bool UAnimGraphNode_PoseMatchingSequencePlayer::DoesSupportTimeForTransitionGetter() const
+{
+	return true;
+}
+
+UAnimationAsset* UAnimGraphNode_PoseMatchingSequencePlayer::GetAnimationAsset() const 
+{
+	UAnimSequenceBase* Sequence = Node.Sequence;
+	UEdGraphPin* SequencePin = FindPin(GET_MEMBER_NAME_STRING_CHECKED(FAnimNode_PoseMatchingSequencePlayer, Sequence));
+	if (SequencePin != nullptr && Sequence == nullptr)
+	{
+		Sequence = Cast<UAnimSequenceBase>(SequencePin->DefaultObject);
+	}
+
+	return Sequence;
+}
+
+const TCHAR* UAnimGraphNode_PoseMatchingSequencePlayer::GetTimePropertyName() const 
+{
+	return TEXT("InternalTimeAccumulator");
+}
+
+UScriptStruct* UAnimGraphNode_PoseMatchingSequencePlayer::GetTimePropertyStruct() const 
+{
+	return FAnimNode_PoseMatchingSequencePlayer::StaticStruct();
+}
+
+void UAnimGraphNode_PoseMatchingSequencePlayer::CustomizePinData(UEdGraphPin* Pin, FName SourcePropertyName, int32 ArrayIndex) const
+{
+	Super::CustomizePinData(Pin, SourcePropertyName, ArrayIndex);
+
+	if (Pin->PinName == GET_MEMBER_NAME_STRING_CHECKED(FAnimNode_PoseMatchingSequencePlayer, PlayRate))
+	{
+		if (!Pin->bHidden)
+		{
+			// Draw value for PlayRateBasis if the pin is not exposed
+			UEdGraphPin* PlayRateBasisPin = FindPin(GET_MEMBER_NAME_STRING_CHECKED(FAnimNode_PoseMatchingSequencePlayer, PlayRateBasis));
+			if (!PlayRateBasisPin || PlayRateBasisPin->bHidden)
+			{
+				if (Node.PlayRateBasis != 1.f)
+				{
+					FFormatNamedArguments Args;
+					Args.Add(TEXT("PinFriendlyName"), Pin->PinFriendlyName);
+					Args.Add(TEXT("PlayRateBasis"), FText::AsNumber(Node.PlayRateBasis));
+					Pin->PinFriendlyName = FText::Format(LOCTEXT("FAnimNode_PoseMatchingSequencePlayer_PlayRateBasis_Value", "({PinFriendlyName} / {PlayRateBasis})"), Args);
+				}
+			}
+			else // PlayRateBasisPin is visible
+			{
+				FFormatNamedArguments Args;
+				Args.Add(TEXT("PinFriendlyName"), Pin->PinFriendlyName);
+				Pin->PinFriendlyName = FText::Format(LOCTEXT("FAnimNode_PoseMatchingSequencePlayer_PlayRateBasis_Name", "({PinFriendlyName} / PlayRateBasis)"), Args);
+			}
+
+			Pin->PinFriendlyName = Node.PlayRateScaleBiasClamp.GetFriendlyName(Pin->PinFriendlyName);
+		}
+	}
+}
+
+void UAnimGraphNode_PoseMatchingSequencePlayer::PostEditChangeProperty(struct FPropertyChangedEvent& PropertyChangedEvent)
+{
+	const FName PropertyName = (PropertyChangedEvent.Property ? PropertyChangedEvent.Property->GetFName() : NAME_None);
+
+	// Reconstruct node to show updates to PinFriendlyNames.
+	if ((PropertyName == GET_MEMBER_NAME_STRING_CHECKED(FAnimNode_PoseMatchingSequencePlayer, PlayRateBasis))
+		|| (PropertyName == GET_MEMBER_NAME_STRING_CHECKED(FInputScaleBiasClamp, bMapRange))
+		|| (PropertyName == GET_MEMBER_NAME_STRING_CHECKED(FInputRange, Min))
+		|| (PropertyName == GET_MEMBER_NAME_STRING_CHECKED(FInputRange, Max))
+		|| (PropertyName == GET_MEMBER_NAME_STRING_CHECKED(FInputScaleBiasClamp, Scale))
+		|| (PropertyName == GET_MEMBER_NAME_STRING_CHECKED(FInputScaleBiasClamp, Bias))
+		|| (PropertyName == GET_MEMBER_NAME_STRING_CHECKED(FInputScaleBiasClamp, bClampResult))
+		|| (PropertyName == GET_MEMBER_NAME_STRING_CHECKED(FInputScaleBiasClamp, ClampMin))
+		|| (PropertyName == GET_MEMBER_NAME_STRING_CHECKED(FInputScaleBiasClamp, ClampMax))
+		|| (PropertyName == GET_MEMBER_NAME_STRING_CHECKED(FInputScaleBiasClamp, bInterpResult))
+		|| (PropertyName == GET_MEMBER_NAME_STRING_CHECKED(FInputScaleBiasClamp, InterpSpeedIncreasing))
+		|| (PropertyName == GET_MEMBER_NAME_STRING_CHECKED(FInputScaleBiasClamp, InterpSpeedDecreasing)))
+	{
+		ReconstructNode();
+	}
+
+	Super::PostEditChangeProperty(PropertyChangedEvent);
+}
+
+#undef LOCTEXT_NAMESPACE
diff --git a/Engine/Plugins/Animation/PoseSearch/Source/Editor/Private/AnimGraphNode_PoseSearchHistoryCollector.cpp b/Engine/Plugins/Animation/PoseSearch/Source/Editor/Private/AnimGraphNode_PoseSearchHistoryCollector.cpp
new file mode 100644
--- /dev/null
+++ b/Engine/Plugins/Animation/PoseSearch/Source/Editor/Private/AnimGraphNode_PoseSearchHistoryCollector.cpp
@@ -0,0 +1,29 @@
+// Copyright Epic Games, Inc. All Rights Reserved.
+
+#include "AnimGraphNode_PoseSearchHistoryCollector.h"
+
+#define LOCTEXT_NAMESPACE "AnimGraphNode_PoseSearchHistoryCollector"
+
+
+FLinearColor UAnimGraphNode_PoseSearchHistoryCollector::GetNodeTitleColor() const
+{
+	return FColor(86, 182, 194);
+}
+
+FText UAnimGraphNode_PoseSearchHistoryCollector::GetTooltipText() const
+{
+	return LOCTEXT("NodeToolTip", "Use this pose for matching");
+}
+
+FText UAnimGraphNode_PoseSearchHistoryCollector::GetNodeTitle(ENodeTitleType::Type TitleType) const
+{
+	return LOCTEXT("NodeTitle", "Match Pose");
+}
+
+FText UAnimGraphNode_PoseSearchHistoryCollector::GetMenuCategory() const
+{
+	return LOCTEXT("NodeCategory", "Pose Search");
+}
+
+
+#undef LOCTEXT_NAMESPACE
diff --git a/Engine/Plugins/Animation/PoseSearch/Source/Editor/Private/PoseSearchEditor.cpp b/Engine/Plugins/Animation/PoseSearch/Source/Editor/Private/PoseSearchEditor.cpp
new file mode 100644
--- /dev/null
+++ b/Engine/Plugins/Animation/PoseSearch/Source/Editor/Private/PoseSearchEditor.cpp
@@ -0,0 +1,60 @@
+// Copyright Epic Games, Inc. All Rights Reserved.
+
+#include "PoseSearchEditor.h"
+#include "Editor.h"
+#include "Animation/AnimSequence.h"
+#include "IAssetTools.h"
+#include "AssetToolsModule.h"
+#include "Modules/ModuleManager.h"
+
+void UPoseSearchBlueprintLibrary::BuildPoseSearchIndex(const UAnimSequence* AnimationSequence, const UPoseSearchIndexConfig* Config, const UPoseSearchSchema* Schema, UPoseSearchIndex* SearchIndex)
+{
+	if (!ensure(AnimationSequence))
+	{
+		return;
+	}
+	if (!ensure(Config))
+	{
+		return;
+	}
+	if (!ensure(Schema))
+	{
+		return;
+	}
+	if (!ensure(SearchIndex))
+	{
+		return;
+	}
+
+	PoseSearchBuildIndex(*AnimationSequence, *Config, *Schema, SearchIndex);
+}
+
+UPoseSearchSchemaFactory::UPoseSearchSchemaFactory(const FObjectInitializer& ObjectInitializer)
+    : Super(ObjectInitializer)
+{
+	bCreateNew = true;
+	bEditAfterNew = true;
+	SupportedClass = UPoseSearchSchema::StaticClass();
+}
+
+UObject* UPoseSearchSchemaFactory::FactoryCreateNew(UClass* InClass, UObject* InParent, FName InName, EObjectFlags Flags, UObject* Context, FFeedbackContext* Warn)
+{
+	UPoseSearchSchema* Schema = NewObject<UPoseSearchSchema>(InParent, InClass, InName, Flags | RF_Transactional);
+	return Schema;
+}
+
+class FPoseSearchEditorModule : public IPoseSearchEditorModuleInterface
+{
+public:
+	virtual void StartupModule() override
+	{
+		FAssetToolsModule::GetModule().Get().RegisterAssetTypeActions(MakeShareable(new FAssetTypeActions_PoseSearchSchema));
+		//UAssetToolsHelpers::GetAssetTools()->RegisterAssetTypeActions(MakeShareable(new FAssetTypeActions_PoseSearchSchema));
+	}
+
+	virtual void ShutdownModule() override
+	{
+	}
+};
+
+IMPLEMENT_MODULE(FPoseSearchEditorModule, PoseSearchEditor);
diff --git a/Engine/Plugins/Animation/PoseSearch/Source/Editor/Public/AnimGraphNode_PoseMatchingSequencePlayer.h b/Engine/Plugins/Animation/PoseSearch/Source/Editor/Public/AnimGraphNode_PoseMatchingSequencePlayer.h
new file mode 100644
--- /dev/null
+++ b/Engine/Plugins/Animation/PoseSearch/Source/Editor/Public/AnimGraphNode_PoseMatchingSequencePlayer.h
@@ -0,0 +1,64 @@
+// Copyright Epic Games, Inc. All Rights Reserved.
+
+#pragma once
+
+#include "PoseSearch/AnimNode_PoseMatchingSequencePlayer.h"
+#include "Editor/AnimGraph/Classes/AnimGraphNode_AssetPlayerBase.h"
+#include "CoreMinimal.h"
+#include "EdGraph/EdGraphNodeUtils.h"
+#include "UObject/ObjectMacros.h"
+
+#include "AnimGraphNode_PoseMatchingSequencePlayer.generated.h"
+
+class FBlueprintActionDatabaseRegistrar;
+
+UCLASS(MinimalAPI)
+class UAnimGraphNode_PoseMatchingSequencePlayer : public UAnimGraphNode_AssetPlayerBase
+{
+	GENERATED_UCLASS_BODY()
+
+	UPROPERTY(EditAnywhere, Category = Settings)
+	FAnimNode_PoseMatchingSequencePlayer Node;
+
+	// UEdGraphNode interface
+	virtual FLinearColor GetNodeTitleColor() const override;
+	virtual FText GetTooltipText() const override;
+	virtual FText GetNodeTitle(ENodeTitleType::Type TitleType) const override;
+	virtual FText GetMenuCategory() const override;
+	virtual void PostEditChangeProperty(struct FPropertyChangedEvent& PropertyChangedEvent) override;
+	// End of UEdGraphNode interface
+
+	// UAnimGraphNode_Base interface
+	virtual void CustomizePinData(UEdGraphPin* Pin, FName SourcePropertyName, int32 ArrayIndex) const override;
+	virtual void ValidateAnimNodeDuringCompilation(class USkeleton* ForSkeleton, class FCompilerResultsLog& MessageLog) override;
+	virtual void PreloadRequiredAssets() override;
+	virtual void BakeDataDuringCompilation(class FCompilerResultsLog& MessageLog) override;
+	virtual bool DoesSupportTimeForTransitionGetter() const override;
+	virtual UAnimationAsset* GetAnimationAsset() const override;
+	virtual const TCHAR* GetTimePropertyName() const override;
+	virtual UScriptStruct* GetTimePropertyStruct() const override;
+	virtual void GetAllAnimationSequencesReferred(TArray<UAnimationAsset*>& AnimationAssets) const override;
+	virtual void ReplaceReferredAnimations(const TMap<UAnimationAsset*, UAnimationAsset*>& AnimAssetReplacementMap) override;
+	virtual void GetMenuActions(FBlueprintActionDatabaseRegistrar& ActionRegistrar) const override;
+	virtual bool IsActionFilteredOut(class FBlueprintActionFilter const& Filter) override;
+	virtual EAnimAssetHandlerType SupportsAssetClass(const UClass* AssetClass) const override;
+	// End of UAnimGraphNode_Base interface
+
+	// UK2Node interface
+	virtual void GetNodeContextMenuActions(class UToolMenu* Menu, class UGraphNodeContextMenuContext* Context) const override;
+	// End of UK2Node interface
+
+	// UAnimGraphNode_AssetPlayerBase interface
+	virtual void SetAnimationAsset(UAnimationAsset* Asset) override;
+	// End of UAnimGraphNode_AssetPlayerBase interface
+
+private:
+	static FText GetTitleGivenAssetInfo(const FText& AssetName, bool bKnownToBeAdditive);
+	FText GetNodeTitleForSequence(ENodeTitleType::Type TitleType, UAnimSequenceBase* InSequence) const;
+
+	/** Constructing FText strings can be costly, so we cache the node's title */
+	FNodeTitleTextTable CachedNodeTitles;
+
+	/** Used for filtering in the Blueprint context menu when the sequence asset this node uses is unloaded */
+	FString UnloadedSkeletonName;
+};
diff --git a/Engine/Plugins/Animation/PoseSearch/Source/Editor/Public/AnimGraphNode_PoseSearchHistoryCollector.h b/Engine/Plugins/Animation/PoseSearch/Source/Editor/Public/AnimGraphNode_PoseSearchHistoryCollector.h
new file mode 100644
--- /dev/null
+++ b/Engine/Plugins/Animation/PoseSearch/Source/Editor/Public/AnimGraphNode_PoseSearchHistoryCollector.h
@@ -0,0 +1,24 @@
+// Copyright Epic Games, Inc. All Rights Reserved.
+
+#pragma once
+
+#include "CoreMinimal.h"
+#include "UObject/ObjectMacros.h"
+#include "AnimGraphNode_Base.h"
+#include "PoseSearch/AnimNode_PoseSearchHistoryCollector.h"
+#include "AnimGraphNode_PoseSearchHistoryCollector.generated.h"
+
+
+UCLASS(MinimalAPI)
+class UAnimGraphNode_PoseSearchHistoryCollector : public UAnimGraphNode_Base
+{
+	GENERATED_BODY()
+
+	UPROPERTY(EditAnywhere, Category = Settings)
+	FAnimNode_PoseSearchHistoryCollector Node;
+
+	virtual FLinearColor GetNodeTitleColor() const override;
+	virtual FText GetTooltipText() const override;
+	virtual FText GetNodeTitle(ENodeTitleType::Type TitleType) const override;
+	virtual FText GetMenuCategory() const override;
+};
diff --git a/Engine/Plugins/Animation/PoseSearch/Source/Editor/Public/PoseSearchEditor.h b/Engine/Plugins/Animation/PoseSearch/Source/Editor/Public/PoseSearchEditor.h
new file mode 100644
--- /dev/null
+++ b/Engine/Plugins/Animation/PoseSearch/Source/Editor/Public/PoseSearchEditor.h
@@ -0,0 +1,51 @@
+// Copyright Epic Games, Inc. All Rights Reserved.
+
+#pragma once
+
+
+#include "AssetTypeActions/AssetTypeActions_AnimationAsset.h"
+#include "Engine/EngineTypes.h"
+#include "Factories/Factory.h"
+#include "Kismet/BlueprintFunctionLibrary.h"
+#include "PoseSearch/PoseSearch.h"
+#include "CoreMinimal.h"
+#include "Modules/ModuleInterface.h"
+
+#include "PoseSearchEditor.generated.h"
+
+
+class UAnimSequence;
+
+UCLASS()
+class POSESEARCHEDITOR_API UPoseSearchBlueprintLibrary : public UBlueprintFunctionLibrary
+{
+	GENERATED_BODY()
+public:
+	/** Compiles a pose search index */
+	UFUNCTION(BlueprintCallable, Category = "AnimationBlueprintLibrary|PoseSearch")
+	static void BuildPoseSearchIndex(
+	    const UAnimSequence* AnimationSequence,
+	    const UPoseSearchIndexConfig* SearchConfig,
+	    const UPoseSearchSchema* SearchSchema,
+	    UPoseSearchIndex* SearchIndex);
+};
+
+UCLASS()
+class POSESEARCHEDITOR_API UPoseSearchSchemaFactory : public UFactory
+{
+	GENERATED_BODY()
+public:
+	UPoseSearchSchemaFactory(const FObjectInitializer& ObjectInitializer);
+	virtual UObject* FactoryCreateNew(UClass* InClass, UObject* InParent, FName InName, EObjectFlags Flags, UObject* Context, FFeedbackContext* Warn) override;
+};
+
+class POSESEARCHEDITOR_API FAssetTypeActions_PoseSearchSchema : public FAssetTypeActions_AnimationAsset
+{
+	virtual FText GetName() const override { return NSLOCTEXT("AssetTypeActions", "PoseSearchEditor", "Pose Search Schema"); }
+	virtual UClass* GetSupportedClass() const override { return UPoseSearchSchema::StaticClass(); }
+	virtual bool CanFilter() override { return true; }
+};
+
+class IPoseSearchEditorModuleInterface : public IModuleInterface
+{
+};
\ No newline at end of file
diff --git a/Engine/Plugins/Animation/PoseSearch/Source/Runtime/PoseSearch.Build.cs b/Engine/Plugins/Animation/PoseSearch/Source/Runtime/PoseSearch.Build.cs
new file mode 100644
--- /dev/null
+++ b/Engine/Plugins/Animation/PoseSearch/Source/Runtime/PoseSearch.Build.cs
@@ -0,0 +1,34 @@
+// Copyright Epic Games, Inc. All Rights Reserved.
+
+using UnrealBuildTool;
+
+public class PoseSearch : ModuleRules
+{
+	public PoseSearch(ReadOnlyTargetRules Target) : base(Target)
+	{
+		PCHUsage = ModuleRules.PCHUsageMode.UseExplicitOrSharedPCHs;
+
+		PublicIncludePaths.AddRange(
+			new string[] {
+				System.IO.Path.Combine(ModuleDirectory, "Public"),
+			}
+			);			
+		
+		PublicDependencyModuleNames.AddRange(
+			new string[]
+			{
+				"Core",
+				"CoreUObject",
+				"Engine",
+				"AnimationCore"
+			}
+			);
+
+		PrivateDependencyModuleNames.AddRange(
+			new string[]
+            {
+				"Engine",
+            }
+			);
+	}
+}
diff --git a/Engine/Plugins/Animation/PoseSearch/Source/Runtime/Private/AnimNode_PoseMatchingSequencePlayer.cpp b/Engine/Plugins/Animation/PoseSearch/Source/Runtime/Private/AnimNode_PoseMatchingSequencePlayer.cpp
new file mode 100644
--- /dev/null
+++ b/Engine/Plugins/Animation/PoseSearch/Source/Runtime/Private/AnimNode_PoseMatchingSequencePlayer.cpp
@@ -0,0 +1,109 @@
+// Copyright Epic Games, Inc. All Rights Reserved.
+
+#include "PoseSearch/AnimNode_PoseMatchingSequencePlayer.h"
+#include "PoseSearch/AnimNode_PoseSearchHistoryCollector.h"
+#include "Animation/AnimInstanceProxy.h"
+#include "PoseSearch/PoseSearch.h"
+
+#define LOCTEXT_NAMESPACE "AnimNode_PoseMatchingSequencePlayer"
+
+TAutoConsoleVariable<int32> CVarAnimPoseMatchingSequencePlayerEnable(TEXT("a.AnimNode.PoseMatchingSequencePlayer.Enable"), 1, TEXT("Enable / Disable Pose Matching"));
+TAutoConsoleVariable<int32> CVarAnimPoseMatchingSequencePlayerDebugVis(TEXT("a.AnimNode.PoseMatchingSequencePlayer.DebugVis"), 0, TEXT("Enable / Disable Pose Matching Debug Visualization"));
+
+const UPoseSearchIndex* GetPoseSearchDataIndex(const UAnimSequenceBase& Sequence)
+{
+	UPoseSearchIndex* PoseSearchIndex = nullptr;
+
+	const TArray<UAnimMetaData*>& MetaData = Sequence.GetMetaData();
+
+	for (UAnimMetaData* MetaDataInstance : MetaData)
+	{
+		if (MetaDataInstance->GetClass() == UPoseSearchIndex::StaticClass())
+		{
+			PoseSearchIndex = Cast<UPoseSearchIndex>(MetaDataInstance);
+			break;
+		}
+	}
+
+	return PoseSearchIndex;
+}
+
+static float FindStartPosition(FAnimInstanceProxy* AnimProxy, const UAnimSequenceBase& Sequence, const FPoseSearchPoseHistory& History, FPoseSearchBuildQueryScratch* Scratch, TArray<float>* Query, bool bEnableDebugVis)
+{
+	const UPoseSearchIndex* PoseSearchIndex = GetPoseSearchDataIndex(Sequence);
+
+	float Offset = -1.0f;
+	if (PoseSearchIndex && PoseSearchIndex->Schema)
+	{
+		bool bQueryBuilt = PoseSearchBuildQuery(
+		    *PoseSearchIndex->Schema,
+		    PoseSearchIndex->SequenceSampleRate,
+			History,
+		    Scratch,
+		    Query);
+
+		if (bQueryBuilt)
+		{
+			USkeletalMeshComponent* SkeletalMeshComponent = AnimProxy->GetSkelMeshComponent();
+			check(SkeletalMeshComponent);
+
+			FPoseSearchDebugDrawParams DebugDraw;
+			DebugDraw.World = SkeletalMeshComponent->GetWorld();
+			DebugDraw.Flags = bEnableDebugVis ? EPoseSearchDebugDrawFlags::DrawAll : EPoseSearchDebugDrawFlags::None;
+			DebugDraw.DefaultLifeTime = 2.0f;
+			DebugDraw.ComponentTransform = SkeletalMeshComponent->GetComponentTransform();
+			FPoseSearchResult Result = PoseSearch(*PoseSearchIndex, *Query, DebugDraw);
+			Offset = Result.TimeOffsetSeconds;
+		}
+	}
+
+	return Offset;
+}
+
+/////////////////////////////////////////////////////
+// FAnimNode_PoseMatchingSequencePlayer
+
+void FAnimNode_PoseMatchingSequencePlayer::Initialize_AnyThread(const FAnimationInitializeContext& Context)
+{
+	DECLARE_SCOPE_HIERARCHICAL_COUNTER_ANIMNODE(Initialize_AnyThread);
+
+	const FAnimNode_PoseSearchHistoryCollector* HistoryNode = Context.GetAncestor<FAnimNode_PoseSearchHistoryCollector>();
+	const FPoseSearchPoseHistory* History = HistoryNode ? &HistoryNode->GetPoseHistory() : nullptr;
+	
+	if (Sequence && StartFromNearestPose && History && CVarAnimPoseMatchingSequencePlayerEnable.GetValueOnAnyThread())
+	{
+		StartPosition = FindStartPosition(Context.AnimInstanceProxy, *Sequence, *History, &Scratch, &SearchQuery, (bool)CVarAnimPoseMatchingSequencePlayerDebugVis.GetValueOnAnyThread());
+		if (StartPosition < 0.0f)
+		{
+			UE_LOG(LogAnimation, Warning, TEXT("Pose search history buffer too small. Increase history time or decrease fragment offsets in search index."));
+			StartPosition = 0.0f;
+		}
+	}
+	else
+	{
+		StartPosition = 0.0f;
+	}
+
+	Super::Initialize_AnyThread(Context);
+}
+
+void FAnimNode_PoseMatchingSequencePlayer::UpdateAssetPlayer(const FAnimationUpdateContext& Context)
+{
+	DECLARE_SCOPE_HIERARCHICAL_COUNTER_ANIMNODE(UpdateAssetPlayer);
+
+	Super::UpdateAssetPlayer(Context);
+}
+
+void FAnimNode_PoseMatchingSequencePlayer::Evaluate_AnyThread(FPoseContext& Output)
+{
+	DECLARE_SCOPE_HIERARCHICAL_COUNTER_ANIMNODE(Evaluate_AnyThread);
+
+	Super::Evaluate_AnyThread(Output);
+}
+
+void FAnimNode_PoseMatchingSequencePlayer::OnInitializeAnimInstance(const FAnimInstanceProxy* InProxy, const UAnimInstance* InAnimInstance)
+{
+	Super::OnInitializeAnimInstance(InProxy, InAnimInstance);
+}
+
+#undef LOCTEXT_NAMESPACE
\ No newline at end of file
diff --git a/Engine/Plugins/Animation/PoseSearch/Source/Runtime/Private/AnimNode_PoseSearchHistoryCollector.cpp b/Engine/Plugins/Animation/PoseSearch/Source/Runtime/Private/AnimNode_PoseSearchHistoryCollector.cpp
new file mode 100644
--- /dev/null
+++ b/Engine/Plugins/Animation/PoseSearch/Source/Runtime/Private/AnimNode_PoseSearchHistoryCollector.cpp
@@ -0,0 +1,49 @@
+// Copyright Epic Games, Inc. All Rights Reserved.
+
+#include "PoseSearch/AnimNode_PoseSearchHistoryCollector.h"
+#include "Animation/AnimInstanceProxy.h"
+
+#define LOCTEXT_NAMESPACE "AnimNode_PoseSearchHistoryCollector"
+
+/////////////////////////////////////////////////////
+// FAnimNode_PoseSearchHistoryCollector
+
+void FAnimNode_PoseSearchHistoryCollector::Initialize_AnyThread(const FAnimationInitializeContext& Context)
+{
+	DECLARE_SCOPE_HIERARCHICAL_COUNTER_ANIMNODE(Initialize_AnyThread);
+
+	Super::Initialize_AnyThread(Context);
+
+	// @@ need to do this once based on descendant node's search schema, not every node init
+	PoseHistory.Init(32, 1.0f);
+
+	FScopedAnimNodeTracker ScopedNodeTracker = Context.TrackAncestor(this);
+
+	Source.Initialize(Context);
+
+	EvalDeltaTime = 0.0f;
+}
+
+void FAnimNode_PoseSearchHistoryCollector::Evaluate_AnyThread(FPoseContext& Output)
+{
+	DECLARE_SCOPE_HIERARCHICAL_COUNTER_ANIMNODE(Evaluate_AnyThread);
+
+	Source.Evaluate(Output);
+
+	PoseHistory.Update(EvalDeltaTime, Output.Pose);
+
+	EvalDeltaTime = 0.0f;
+}
+
+void FAnimNode_PoseSearchHistoryCollector::Update_AnyThread(const FAnimationUpdateContext& Context)
+{
+	DECLARE_SCOPE_HIERARCHICAL_COUNTER_ANIMNODE(Update_AnyThread);
+
+	FScopedAnimNodeTracker ScopedNodeTracker = Context.TrackAncestor(this);
+
+	Source.Update(Context);
+
+	EvalDeltaTime += Context.GetDeltaTime();
+}
+
+#undef LOCTEXT_NAMESPACE
\ No newline at end of file
diff --git a/Engine/Plugins/Animation/PoseSearch/Source/Runtime/Private/PoseSearch.cpp b/Engine/Plugins/Animation/PoseSearch/Source/Runtime/Private/PoseSearch.cpp
new file mode 100644
--- /dev/null
+++ b/Engine/Plugins/Animation/PoseSearch/Source/Runtime/Private/PoseSearch.cpp
@@ -0,0 +1,442 @@
+// Copyright Epic Games, Inc. All Rights Reserved.
+
+#include "PoseSearch/PoseSearch.h"
+
+#include "Animation/AnimCurveTypes.h"
+#include "Animation/AnimInstance.h"
+#include "Animation/AnimSequenceBase.h"
+#include "AnimationRuntime.h"
+#include "BonePose.h"
+#include "DrawDebugHelpers.h"
+
+//////////////////////////////////////////////////////////////////////////
+// UPoseSearchSchema
+
+void UPoseSearchSchema::PreSave(const class ITargetPlatform* TargetPlatform)
+{
+	// Initialize references to obtain bone indices
+	for (FBoneReference& BoneRef : Bones)
+	{
+		BoneRef.Initialize(GetSkeleton());
+	}
+
+	// Fill out bone index array and sort by bone index
+	BoneIndices.SetNum(Bones.Num());
+	for (int32 Index = 0; Index != Bones.Num(); ++Index)
+	{
+		BoneIndices[Index] = Bones[Index].BoneIndex;
+	}
+	BoneIndices.Sort();
+
+	// Build separate index array with parent indices guaranteed to be present
+	BoneIndicesWithParents = BoneIndices;
+	FAnimationRuntime::EnsureParentsPresent(BoneIndicesWithParents, GetSkeleton()->GetReferenceSkeleton());
+
+	// Sort fragment offsets by largest offset first since larger offsets are closer to the beginning of the sample array
+	FragmentSampleOffsets.Sort(TGreater<>());
+
+	// Ensure we have at least one offset at zero which corresponds with instantaneous sample matching
+	if (FragmentSampleOffsets.Num() == 0)
+	{
+		FragmentSampleOffsets.Add(0);
+	}
+
+	FloatsPerPose = BoneIndices.Num() * 3;
+
+	Super::PreSave(TargetPlatform);
+}
+
+
+//////////////////////////////////////////////////////////////////////////
+// UPoseSearchPoseHistory
+
+/**
+* Fills skeleton transforms with evaluated compact pose transforms.
+* Bones that weren't evaluated are filled with the bone's reference pose.
+*/
+static void CopyCompactToSkeletonPose (const FCompactPose& Pose, TArray<FTransform>& OutTransforms)
+{
+	const FBoneContainer& BoneContainer = Pose.GetBoneContainer();
+	const FReferenceSkeleton& RefSkeleton = BoneContainer.GetReferenceSkeleton();
+	TArrayView<const FTransform> RefSkeletonTransforms = MakeArrayView(RefSkeleton.GetRefBonePose());
+
+	const int32 NumSkeletonBones = BoneContainer.GetNumBones();
+	OutTransforms.SetNum(NumSkeletonBones);
+
+	for (auto SkeletonBoneIdx = FSkeletonPoseBoneIndex(0); SkeletonBoneIdx != NumSkeletonBones; ++SkeletonBoneIdx)
+	{
+		FCompactPoseBoneIndex CompactBoneIdx = BoneContainer.GetCompactPoseIndexFromSkeletonIndex(SkeletonBoneIdx.GetInt());
+		OutTransforms[SkeletonBoneIdx.GetInt()] = CompactBoneIdx.IsValid() ? Pose[CompactBoneIdx] : RefSkeletonTransforms[SkeletonBoneIdx.GetInt()];
+	}
+}
+
+void FPoseSearchPoseHistory::Init(int32 InNumPoses, float InTimeHorizon)
+{
+	int32 Capacity = FMath::RoundUpToPowerOfTwo(InNumPoses);
+
+	if ((Queue.GetCapacity() == Capacity) && (TimeHorizon == InTimeHorizon))
+	{
+		return;
+	}
+
+	Poses.SetNum(Capacity);
+	Knots.SetNum(Capacity);
+	Queue.Init(Capacity);
+	TimeHorizon = InTimeHorizon;
+}
+
+void FPoseSearchPoseHistory::Init(const FPoseSearchPoseHistory& History)
+{
+	Poses = History.Poses;
+	Knots = History.Knots;
+	Queue = History.Queue;
+	TimeHorizon = History.TimeHorizon;
+}
+
+bool FPoseSearchPoseHistory::Sample(float Time, const TArray<FBoneIndexType>& RequiredBones, TArray<FTransform>& OutPose) const
+{
+	// Find the upper bound knot
+	uint32 UpperBoundIndex = MAX_uint32;
+	int32 UpperBoundOffset = 1;
+	for (; UpperBoundOffset < (int32)Queue.Num(); ++UpperBoundOffset)
+	{
+		uint32 TestIndex = Queue.GetOffsetFromBack(UpperBoundOffset);
+		if (Knots[TestIndex] >= Time)
+		{
+			UpperBoundIndex = TestIndex;
+			break;
+		}
+	}
+
+	if (UpperBoundIndex == MAX_uint32)
+	{
+		return false;
+	}
+
+	// The lower bound knot is adjacent
+	int32 LowerBoundOffset = UpperBoundOffset - 1;
+	uint32 LowerBoundIndex = Queue.GetOffsetFromBack(LowerBoundOffset);
+
+	// Compute alpha between upper and lower bound knots
+	float Alpha = FMath::GetMappedRangeValueUnclamped(
+		FVector2D(Knots[LowerBoundIndex], Knots[UpperBoundIndex]),
+		FVector2D(0.0f, 1.0f),
+		Time);
+
+	// Lerp between upper and lower bound poses by alpha to produce output pose at requested sample time
+	OutPose = Poses[LowerBoundIndex].LocalTransforms;
+	FAnimationRuntime::LerpBoneTransforms(
+		OutPose,
+		Poses[UpperBoundIndex].LocalTransforms,
+		Alpha,
+		RequiredBones);
+
+	return true;
+}
+
+void FPoseSearchPoseHistory::Update(float TimeDelta, const FCompactPose& Pose)
+{
+	// Age our elapsed times
+	for (int32 Offset = 0; Offset != (int32)Queue.Num(); ++Offset)
+	{
+		uint32 Index = Queue.GetOffsetFromFront(Offset);
+		Knots[Index] += TimeDelta;
+	}
+
+	if (!Queue.IsFull())
+	{
+		// Consume every pose until the queue is full
+		Queue.PushBack();
+	}
+	else
+	{
+		// Exercise pose retention policy. We must guarantee there is always one additional knot
+		// at or beyond the desired time horizon H so we can fulfill sample requests at t=H. We also
+		// want to evenly distribute knots across the entire history buffer so we only push additional
+		// poses when enough time has elapsed.
+
+		const float SampleInterval = GetSampleInterval();
+
+		uint32 SecondOldest = Queue.GetOffsetFromFront(1);
+		bool bCanEvictOldest = Knots[SecondOldest] >= TimeHorizon;
+
+		uint32 SecondNewest = Queue.GetOffsetFromBack(1);
+		bool bShouldPushNewest = Knots[SecondNewest] >= SampleInterval;
+
+		if (bCanEvictOldest && bShouldPushNewest)
+		{
+			Queue.PopFront();
+			Queue.PushBack();
+		}
+	}
+
+	// Regardless of the retention policy, we always update the most recent pose
+	uint32 Newest = Queue.GetOffsetFromBack(0);
+	Knots[Newest] = 0.0f;
+	CopyCompactToSkeletonPose(Pose, Poses[Newest].LocalTransforms);
+}
+
+float FPoseSearchPoseHistory::GetSampleInterval() const
+{
+	return TimeHorizon / Queue.GetCapacity();
+}
+
+
+//////////////////////////////////////////////////////////////////////////
+// PoseSearch API
+
+void PoseSearchDrawPose(const FPoseSearchDebugDrawParams& DrawParams, const FColor& Color, TArrayView<const float> Pose, float LifeTimeDelta, uint8 DepthPriorityDelta = 0)
+{
+	const float LifeTime = DrawParams.DefaultLifeTime + LifeTimeDelta;
+	const uint8 DepthPriority = ESceneDepthPriorityGroup::SDPG_Foreground + 1 + DepthPriorityDelta;
+
+	FVector PointPrev(
+		Pose[0],
+		Pose[1],
+		Pose[2]);
+
+	PointPrev = DrawParams.ComponentTransform.TransformPosition(PointPrev);
+
+	const int32 NumPoints = Pose.Num() / 3;
+
+	for (int32 PointIdx = 1; PointIdx != NumPoints; ++PointIdx)
+	{
+		FVector PointNext(
+			Pose[PointIdx * 3 + 0],
+			Pose[PointIdx * 3 + 1],
+			Pose[PointIdx * 3 + 2]);
+
+		PointNext = DrawParams.ComponentTransform.TransformPosition(PointNext);
+
+		DrawDebugPoint(DrawParams.World, PointNext, 1.0f, Color, false, LifeTime, DepthPriority);
+
+		bool bIsChildOfPrev = DrawParams.Schema->GetSkeleton()->GetReferenceSkeleton().BoneIsChildOf(
+			DrawParams.Schema->BoneIndices[PointIdx],
+			DrawParams.Schema->BoneIndices[PointIdx-1]);
+
+		if (bIsChildOfPrev)
+		{
+			DrawDebugLine(DrawParams.World, PointPrev, PointNext, Color, false, LifeTime, DepthPriority);
+		}
+		PointPrev = PointNext;
+	}
+}
+
+void PoseSearchDrawSearchIndex(const FPoseSearchDebugDrawParams& DrawParams, const UPoseSearchIndex& SearchIndex, int32 HighlightPoseIdx)
+{
+	if (!DrawParams.CanDraw())
+	{
+		return;
+	}
+
+	int32 LastPoseIdx = SearchIndex.NumPoses;
+	int32 StartPoseIdx = 0;
+	if (!(DrawParams.Flags & EPoseSearchDebugDrawFlags::DrawSearchIndex))
+	{
+		StartPoseIdx = HighlightPoseIdx;
+		LastPoseIdx = StartPoseIdx + 1;
+	}
+
+	if (StartPoseIdx < 0)
+	{
+		return;
+	}
+
+	for (int32 PoseDrawIdx = StartPoseIdx; PoseDrawIdx != LastPoseIdx; ++PoseDrawIdx)
+	{
+		float LifeTimeDelta = 0.0f;
+		FLinearColor Color;
+		if (PoseDrawIdx == HighlightPoseIdx)
+		{
+			Color = FLinearColor::Yellow;
+		}
+		else
+		{
+			float Lerp = (float)(PoseDrawIdx + 1) / SearchIndex.NumPoses;
+			LifeTimeDelta = Lerp - 1.0f;
+			Color = FLinearColor::LerpUsingHSV(FLinearColor(FColor::Cyan), FLinearColor(FColor::Blue), Lerp);
+		}
+
+		TArrayView<const float> Pose = MakeArrayView(&SearchIndex.Values[PoseDrawIdx * SearchIndex.FloatsPerPose], SearchIndex.FloatsPerPose);
+		PoseSearchDrawPose(DrawParams, Color.ToFColor(true), Pose, LifeTimeDelta);
+	}
+}
+
+void PoseSearchDrawQuery(const FPoseSearchDebugDrawParams& DrawParams, TArrayView<const float> Query)
+{
+	if (!DrawParams.CanDraw())
+	{
+		return;
+	}
+
+	int32 NumSamples = DrawParams.Schema->FragmentSampleOffsets.Num();
+	for (int32 Sample = 0; Sample != NumSamples; ++Sample)
+	{
+		float Lerp = (Sample + 1) / NumSamples;
+		FLinearColor Color = FLinearColor::LerpUsingHSV(FLinearColor(FColor::Magenta), FLinearColor(FColor::Purple), Lerp);
+
+		TArrayView<const float> Pose = MakeArrayView(&Query[Sample * DrawParams.Schema->FloatsPerPose], DrawParams.Schema->FloatsPerPose);
+		PoseSearchDrawPose(DrawParams, Color.ToFColor(true), Pose, 0.0f, 1);
+	}
+}
+
+void PoseSearchBuildIndex(const UAnimSequenceBase& AnimSequence, const UPoseSearchIndexConfig& SearchConfig, const UPoseSearchSchema& SearchSchema, UPoseSearchIndex* SearchIndex)
+{
+	USkeleton* Skeleton = AnimSequence.GetSkeleton();
+	check(Skeleton);
+	check(Skeleton->IsCompatible(SearchSchema.GetSkeleton()));
+
+	FBoneContainer BoneContainer;
+	BoneContainer.InitializeTo(SearchSchema.BoneIndicesWithParents, FCurveEvaluationOption(false), *Skeleton);
+
+	FBlendedCurve UnusedCurve;
+	FAnimExtractContext ExtractionCtx;
+	// ExtractionCtx.PoseCurves is intentionally left empty
+	// ExtractionCtx.BonesRequired is unused by UAnimSequence::GetAnimationPose
+	ExtractionCtx.bExtractRootMotion = true;
+
+	double CurrTime = SearchConfig.FrameSamplingRange.GetLowerBoundValue();
+	const double EndTime = FMath::Min(AnimSequence.SequenceLength, SearchConfig.FrameSamplingRange.GetUpperBoundValue());
+	const double DeltaTime = 1.0 / SearchConfig.SampleRate;
+
+	const int32 NumPoses = FMath::FloorToInt((EndTime - CurrTime) / DeltaTime);
+
+	SearchIndex->Values.Reset(SearchSchema.FloatsPerPose * NumPoses);
+
+	FCompactPose Pose;
+	Pose.SetBoneContainer(&BoneContainer);
+	FCSPose<FCompactPose> ComponentSpacePose;
+
+	for (int32 PoseIdx = 0; PoseIdx != NumPoses; ++PoseIdx, CurrTime += DeltaTime)
+	{
+		// Extract pose
+		ExtractionCtx.CurrentTime = CurrTime;
+		AnimSequence.GetAnimationPose(Pose, UnusedCurve, ExtractionCtx);
+		ComponentSpacePose.InitPose(Pose);
+
+		for (int32 BoneIndex : SearchSchema.BoneIndices)
+		{
+			FCompactPoseBoneIndex CompactBoneIndex = BoneContainer.MakeCompactPoseIndex(FMeshPoseBoneIndex(BoneIndex));
+			const FTransform& Transform = ComponentSpacePose.GetComponentSpaceTransform(CompactBoneIndex);
+ 
+			FVector Translation = Transform.GetTranslation();
+			SearchIndex->Values.Add(Translation.X);
+			SearchIndex->Values.Add(Translation.Y);
+			SearchIndex->Values.Add(Translation.Z);
+		}
+	}
+
+	SearchIndex->NumPoses = NumPoses;
+	SearchIndex->FloatsPerPose = SearchSchema.FloatsPerPose;
+	SearchIndex->Schema = &SearchSchema;
+	SearchIndex->SequenceSampleRate = SearchConfig.SampleRate;
+	SearchIndex->SequenceFrameSamplingRange = SearchConfig.FrameSamplingRange;
+}
+
+bool PoseSearchBuildQuery(const UPoseSearchSchema& SearchSchema, int32 AssetSampleRate, const FPoseSearchPoseHistory& History, FPoseSearchBuildQueryScratch* Scratch, TArray<float>* Query)
+{
+	Query->Reset(0);
+	Query->Reserve(SearchSchema.FragmentSampleOffsets.Num() * SearchSchema.FloatsPerPose);
+
+	for (int32 Offset : SearchSchema.FragmentSampleOffsets)
+	{
+		float TimeDelta = Offset * (1.0f / AssetSampleRate);
+
+		if (!History.Sample(TimeDelta, SearchSchema.BoneIndicesWithParents, Scratch->LocalPose))
+		{
+			return false;
+		}
+
+		FAnimationRuntime::FillUpComponentSpaceTransforms(SearchSchema.GetSkeleton()->GetReferenceSkeleton(), Scratch->LocalPose, Scratch->ComponentPose);
+
+		for (int32 SkeletonBoneIndex : SearchSchema.BoneIndices)
+		{
+			const FTransform& Transform = Scratch->ComponentPose[SkeletonBoneIndex];
+			FVector Translation = Transform.GetTranslation();
+			Query->Add(Translation.X);
+			Query->Add(Translation.Y);
+			Query->Add(Translation.Z);
+		}
+	}
+
+	return true;
+}
+
+FPoseSearchResult PoseSearch(const UPoseSearchIndex& SearchIndex, TArrayView<const float> Query, FPoseSearchDebugDrawParams DebugDrawParams)
+{
+	check(SearchIndex.Schema != nullptr);
+	check(SearchIndex.NumPoses * SearchIndex.FloatsPerPose == SearchIndex.Values.Num());
+
+	if (Query.Num() != SearchIndex.FloatsPerPose * SearchIndex.Schema->FragmentSampleOffsets.Num())
+	{
+		return FPoseSearchResult();
+	}
+
+	float BestPoseDifference = TNumericLimits<float>::Max();
+	int32 BestPoseIdx = INDEX_NONE;
+
+	for (int32 PoseIdx = SearchIndex.Schema->FragmentSampleOffsets[0]; PoseIdx < SearchIndex.NumPoses; ++PoseIdx)
+	{
+		float PoseDifference = 0.0f;
+		int32 QueryValueIdx = 0;
+		for (int32 PoseOffset : SearchIndex.Schema->FragmentSampleOffsets)
+		{
+			int32 SearchValueIdx = (PoseIdx - PoseOffset) * SearchIndex.FloatsPerPose;
+
+			for (int32 PoseElemCount = 0; PoseElemCount != SearchIndex.FloatsPerPose; ++PoseElemCount, ++QueryValueIdx, ++SearchValueIdx)
+			{
+				PoseDifference += FMath::Square(Query[QueryValueIdx] - SearchIndex.Values[SearchValueIdx]);
+			}
+		}
+
+		if (PoseDifference < BestPoseDifference)
+		{
+			BestPoseDifference = PoseDifference;
+			BestPoseIdx = PoseIdx;
+		}
+	}
+
+	check(BestPoseIdx != INDEX_NONE);
+
+	const float SampleDelta = 1.0f / SearchIndex.SequenceSampleRate;
+
+	float BestPoseTime = FMath::Min(
+		SearchIndex.SequenceFrameSamplingRange.GetUpperBoundValue(),
+		SampleDelta * BestPoseIdx + SearchIndex.SequenceFrameSamplingRange.GetLowerBoundValue());
+
+	// Do debug visualization
+	DebugDrawParams.Schema = SearchIndex.Schema;
+	if (DebugDrawParams.CanDraw())
+	{
+		if (EnumHasAnyFlags(DebugDrawParams.Flags, EPoseSearchDebugDrawFlags::DrawQuery))
+		{
+			PoseSearchDrawQuery(DebugDrawParams, Query);
+		}
+
+		if (EnumHasAnyFlags(DebugDrawParams.Flags, EPoseSearchDebugDrawFlags::DrawSearchIndex | EPoseSearchDebugDrawFlags::DrawBest))
+		{
+			int32 HighlightPoseIdx = EnumHasAnyFlags(DebugDrawParams.Flags, EPoseSearchDebugDrawFlags::DrawBest) ? BestPoseIdx : -1;
+			PoseSearchDrawSearchIndex(DebugDrawParams, SearchIndex, HighlightPoseIdx);
+		}
+	}
+
+	FPoseSearchResult Result;
+	Result.Dissimilarity = BestPoseDifference;
+	Result.TimeOffsetSeconds = BestPoseTime;
+	Result.PoseIdx = BestPoseIdx;
+	return Result;
+}
+
+
+//////////////////////////////////////////////////////////////////////////
+// FPoseSearchModule
+
+void FPoseSearchModule::StartupModule()
+{
+}
+
+void FPoseSearchModule::ShutdownModule()
+{
+}
+
+IMPLEMENT_MODULE(FPoseSearchModule, PoseSearch)
\ No newline at end of file
diff --git a/Engine/Plugins/Animation/PoseSearch/Source/Runtime/Public/PoseSearch/AnimNode_PoseMatchingSequencePlayer.h b/Engine/Plugins/Animation/PoseSearch/Source/Runtime/Public/PoseSearch/AnimNode_PoseMatchingSequencePlayer.h
new file mode 100644
--- /dev/null
+++ b/Engine/Plugins/Animation/PoseSearch/Source/Runtime/Public/PoseSearch/AnimNode_PoseMatchingSequencePlayer.h
@@ -0,0 +1,41 @@
+// Copyright Epic Games, Inc. All Rights Reserved.
+
+#pragma once
+
+#include "Animation/AnimNode_SequencePlayer.h"
+#include "PoseSearch/PoseSearch.h"
+
+#include "AnimNode_PoseMatchingSequencePlayer.generated.h"
+
+USTRUCT(BlueprintInternalUseOnly)
+struct POSESEARCH_API FAnimNode_PoseMatchingSequencePlayer : public FAnimNode_SequencePlayer
+{
+	GENERATED_BODY()
+
+public:
+	UPROPERTY(EditAnywhere, BlueprintReadWrite, Category = Settings, meta = (PinHiddenByDefault))
+	bool StartFromNearestPose;
+
+public:
+	FAnimNode_PoseMatchingSequencePlayer()
+	    : StartFromNearestPose(true)
+	{}
+
+	// FAnimNode_Base interface
+	virtual void Initialize_AnyThread(const FAnimationInitializeContext& Context) override;
+	virtual void Evaluate_AnyThread(FPoseContext& Output) override;
+	virtual bool NeedsOnInitializeAnimInstance() const { return true; }
+	// End of FAnimNode_Base interface
+
+	// FAnimNode_SequencePlayerBase interface
+	virtual void UpdateAssetPlayer(const FAnimationUpdateContext& Context) override;
+	// End of FAnimNode_SequencePlayerBase interface
+
+protected:
+	virtual void OnInitializeAnimInstance(const FAnimInstanceProxy* InProxy, const UAnimInstance* InAnimInstance);
+
+private:
+	FPoseSearchBuildQueryScratch Scratch;
+	TArray<float> SearchQuery;
+	float InitDelta;
+};
\ No newline at end of file
diff --git a/Engine/Plugins/Animation/PoseSearch/Source/Runtime/Public/PoseSearch/AnimNode_PoseSearchHistoryCollector.h b/Engine/Plugins/Animation/PoseSearch/Source/Runtime/Public/PoseSearch/AnimNode_PoseSearchHistoryCollector.h
new file mode 100644
--- /dev/null
+++ b/Engine/Plugins/Animation/PoseSearch/Source/Runtime/Public/PoseSearch/AnimNode_PoseSearchHistoryCollector.h
@@ -0,0 +1,36 @@
+// Copyright Epic Games, Inc. All Rights Reserved.
+
+#pragma once
+
+#include "PoseSearch/PoseSearch.h"
+
+#include "AnimNode_PoseSearchHistoryCollector.generated.h"
+
+USTRUCT(BlueprintInternalUseOnly)
+struct POSESEARCH_API FAnimNode_PoseSearchHistoryCollector : public FAnimNode_Base
+{
+	GENERATED_BODY()
+
+public:
+	UPROPERTY(EditAnywhere, BlueprintReadWrite, Category = Links)
+	FPoseLink Source;
+
+public:
+	FAnimNode_PoseSearchHistoryCollector()
+	{}
+
+	// FAnimNode_Base interface
+	virtual void Initialize_AnyThread(const FAnimationInitializeContext& Context) override;
+	virtual void Evaluate_AnyThread(FPoseContext& Output) override;
+	virtual void Update_AnyThread(const FAnimationUpdateContext& Context);
+	// End of FAnimNode_Base interface
+
+	const FPoseSearchPoseHistory& GetPoseHistory() const { return PoseHistory; }
+	void Init(int32 InNumPoses, float InTimeHorizon);
+
+protected:
+
+	FPoseSearchPoseHistory PoseHistory;
+	float EvalDeltaTime = 0.0f;
+
+};
\ No newline at end of file
diff --git a/Engine/Plugins/Animation/PoseSearch/Source/Runtime/Public/PoseSearch/PoseSearch.h b/Engine/Plugins/Animation/PoseSearch/Source/Runtime/Public/PoseSearch/PoseSearch.h
new file mode 100644
--- /dev/null
+++ b/Engine/Plugins/Animation/PoseSearch/Source/Runtime/Public/PoseSearch/PoseSearch.h
@@ -0,0 +1,331 @@
+// Copyright Epic Games, Inc. All Rights Reserved.
+
+#pragma once
+
+#include "CoreMinimal.h"
+#include "Animation/AnimMetaData.h"
+#include "Animation/AnimationAsset.h"
+#include "BoneIndices.h"
+#include "Modules/ModuleManager.h"
+#include "Misc/EnumClassFlags.h"
+#include "UObject/Object.h"
+#include "UObject/ObjectMacros.h"
+
+#include "PoseSearch.generated.h"
+
+
+class UAnimSequenceBase;
+
+/** Sampling parameters for building a pose search index for an asset. */
+UCLASS(BlueprintType, Category = "Animation|PoseSearch")
+class POSESEARCH_API UPoseSearchIndexConfig : public UAnimMetaData
+{
+	GENERATED_BODY()
+
+public:
+	static constexpr int32 DefaultSampleRate = 15;
+
+	UPROPERTY(EditAnywhere, meta = (ClampMin = "1", ClampMax = "60"), Category = "Config")
+	int32 SampleRate;
+
+	UPROPERTY(EditAnywhere, Category = "Config")
+	FFloatRange FrameSamplingRange;
+
+	UPoseSearchIndexConfig()
+	    : SampleRate(DefaultSampleRate)
+	{}
+};
+
+
+/**
+* Specifies the format of a pose search index. At runtime, queries are built according to the schema for searching.
+*/
+UCLASS(BlueprintType, Category = "Animation|PoseSearch")
+class POSESEARCH_API UPoseSearchSchema : public UAnimationAsset
+{
+	GENERATED_BODY()
+
+public:
+	UPROPERTY(EditAnywhere, Category = "Schema")
+	TArray<FBoneReference> Bones;
+
+	UPROPERTY(EditAnywhere, Category = "Schema")
+	TArray<int32> FragmentSampleOffsets;
+
+	UPROPERTY()
+	int32 FloatsPerPose;
+
+	UPROPERTY()
+	TArray<uint16> BoneIndices;
+
+	UPROPERTY()
+	TArray<uint16> BoneIndicesWithParents;
+
+public: // UObject
+	virtual void PreSave(const class ITargetPlatform* TargetPlatform);
+};
+
+
+/**
+* Contains pose data sampled at editor time for searching at runtime.
+*/
+UCLASS(BlueprintType, Category = "Animation|PoseSearch")
+class POSESEARCH_API UPoseSearchIndex : public UAnimMetaData
+{
+	GENERATED_BODY()
+public:
+	UPROPERTY()
+	int32 NumPoses;
+
+	UPROPERTY()
+	int32 FloatsPerPose;
+
+	UPROPERTY()
+	TArray<float> Values;
+
+	UPROPERTY()
+	const UPoseSearchSchema* Schema;
+
+	UPROPERTY()
+	int32 SequenceSampleRate;
+
+	UPROPERTY()
+	FFloatRange SequenceFrameSamplingRange;
+};
+
+struct POSESEARCH_API FPoseSearchPoseSnapshot
+{
+	TArray<FTransform> LocalTransforms;
+};
+
+
+namespace PoseSearchDetail {
+
+/**
+* Helper to adapt a fixed-size buffer to a circular queue.
+* FCircularView simply manages circular indexing arithmetic and is meant be paired with a container.
+*/
+class FCircularView
+{
+public:
+	FCircularView(uint32 InCapacity = 0)
+	    : Capacity(0)
+	    , VirtualFirst(0)
+	    , VirtualLast(0)
+	{
+		if (InCapacity)
+		{
+			Init(InCapacity);
+		}
+	}
+
+	uint32 GetCapacity() const
+	{
+		return Capacity;
+	}
+
+	uint32 GetFront() const
+	{
+		checkSlow(!IsEmpty());
+		return ToPhysicalIndex(VirtualFirst);
+	}
+
+	uint32 GetBack() const
+	{
+		checkSlow(!IsEmpty());
+		return ToPhysicalIndex(VirtualLast - 1);
+	}
+
+	uint32 Num() const
+	{
+		return VirtualLast - VirtualFirst;
+	}
+
+	void Init(uint32 InCapacity)
+	{
+		checkSlow(FMath::IsPowerOfTwo(InCapacity));
+		Capacity = InCapacity;
+		VirtualFirst = 0;
+		VirtualLast = 0;
+	}
+
+	bool IsEmpty() const
+	{
+		return Num() == 0;
+	}
+
+	bool IsFull() const
+	{
+		return Num() == Capacity;
+	}
+
+	uint32 GetOffsetFromFront(int32 Offset) const
+	{
+		checkSlow((uint32)FMath::Abs(Offset) < Num());
+		return ToPhysicalIndex((uint32)(VirtualFirst + Offset));
+	}
+
+	uint32 GetOffsetFromBack(int32 Offset) const
+	{
+		return GetOffsetFromFront(Num() - 1 - Offset);
+	}
+
+	uint32 operator[](int32 OffsetFromFirst) const
+	{
+		return GetOffsetFromFront(OffsetFromFirst);
+	}
+
+	void PushBack()
+	{
+		checkSlow(!IsFull());
+		++VirtualLast;
+	}
+
+	void PopFront()
+	{
+		checkSlow(!IsEmpty());
+		++VirtualFirst;
+	}
+
+	void PushFront()
+	{
+		checkSlow(!IsFull());
+		--VirtualFirst;
+	}
+
+	void PopBack()
+	{
+		checkSlow(!IsEmpty());
+		--VirtualLast;
+	}
+
+private:
+	uint32 ToPhysicalIndex(uint32 VirtualIndex) const { return VirtualIndex & (Capacity - 1); }
+
+	// Maximum available elements. Must be a power of two.
+	uint32 Capacity;
+
+	// The beginning of the view. Must be converted by ToPhysicalIndex to index physical memory.
+	uint32 VirtualFirst;
+
+	// One element past the back of the view. Must be converted by ToPhysicalIndex to index physical memory.
+	uint32 VirtualLast;
+};
+
+} // namespace PoseSearchDetail
+
+
+/** Records poses over time in a ring buffer. PoseSearchBuildQuery() uses this to sample from the present or past poses according to the search schema. */
+class POSESEARCH_API FPoseSearchPoseHistory
+{
+public:
+	void Init(int32 InNumPoses, float InTimeHorizon);
+	void Init(const FPoseSearchPoseHistory& History);
+	bool Sample(float Time, const TArray<FBoneIndexType>& RequiredBones, TArray<FTransform>& OutPose) const;
+	void Update(float TimeDelta, const FCompactPose& EvalContext);
+	float GetSampleInterval() const;
+	float GetTimeHorizon() const { return TimeHorizon; }
+
+private:
+	TArray<FPoseSearchPoseSnapshot> Poses;
+	TArray<float> Knots;
+	PoseSearchDetail::FCircularView Queue;
+	float TimeHorizon = 0.0f;
+};
+
+struct FPoseSearchBuildQueryScratch
+{
+	TArray<FTransform> LocalPose;
+	TArray<FTransform> ComponentPose;
+};
+
+enum class EPoseSearchDebugDrawFlags : uint32
+{
+	None			= 0,
+	DrawQuery		= 1 << 0,
+	DrawSearchIndex = 1 << 1,
+	DrawBest        = 1 << 2,
+	DrawAll		    = MAX_uint32
+};
+ENUM_CLASS_FLAGS(EPoseSearchDebugDrawFlags);
+
+struct FPoseSearchDebugDrawParams
+{
+	const UWorld* World = nullptr;
+	float DefaultLifeTime = 5.0f;
+	EPoseSearchDebugDrawFlags Flags = EPoseSearchDebugDrawFlags::None;
+	const UPoseSearchSchema* Schema = nullptr;
+	FTransform ComponentTransform = FTransform::Identity;
+
+	bool CanDraw () const { return World && Schema; }
+};
+
+struct FPoseSearchResult
+{
+	int32 PoseIdx = -1;
+	float TimeOffsetSeconds = 0.0f;
+	float Dissimilarity = MAX_flt;
+};
+
+
+//////////////////////////////////////////////////////////////////////////
+// Main PoseSearch API
+
+/**
+* Draws a pose search query
+*
+* @param DrawParams Visualization options
+* @param Query The pose search query created with PoseSearchBuildQuery()
+*/
+void POSESEARCH_API PoseSearchDrawQuery(const FPoseSearchDebugDrawParams& DrawParams, TArrayView<const float> Query);
+
+/**
+* Draws a pose search index
+*
+* @param DrawParams Visualization options
+* @param SearchIndex The pose search index to visualize
+* @param HighlightPoseIdx Optional index of a pose in the search index to highlight. Specify -1 to not highlight a pose.
+*/
+void POSESEARCH_API PoseSearchDrawSearchIndex(const FPoseSearchDebugDrawParams& DrawParams, const UPoseSearchIndex& SearchIndex, int32 HighlightPoseIdx = -1);
+
+/**
+* Creates a pose search index by sampling from the animation sequence
+*
+* @param AnimSequence The asset to sample from
+* @param SearchConfig Options for sampling from this asset
+* @param SearchSchema The format of the search index
+* @param SearchIndex Output pose search index
+*/
+void POSESEARCH_API PoseSearchBuildIndex(const UAnimSequenceBase& AnimSequence, const UPoseSearchIndexConfig& SearchConfig, const UPoseSearchSchema& SearchSchema, UPoseSearchIndex* SearchIndex);
+
+/**
+* Builds a pose search query as an array of floats according to the search schema
+*
+* @param SearchSchema The format of the search index and the query being built
+* @param AssetSampleRate The sampling rate used to index the animation asset
+* @param History Recorded pose to sample from to build the query
+* @param Scratch Working memory to re-use for query construction
+* @Query Output pose search query
+*/
+bool POSESEARCH_API PoseSearchBuildQuery(const UPoseSearchSchema& SearchSchema, int32 AssetSampleRate, const FPoseSearchPoseHistory& History, FPoseSearchBuildQueryScratch* Scratch, TArray<float>* Query);
+
+/**
+* Performs a pose search
+*
+* @param SearchIndex The index to search within. See PoseSearchBuildIndex().
+* @param Query The pose query to search for. See PoseSearchBuildQuery().
+* @param DrawParams Visualization options
+* @return The pose in the SearchIndex that most closely matches the Query.
+*/
+FPoseSearchResult POSESEARCH_API PoseSearch(const UPoseSearchIndex& SearchIndex, TArrayView<const float> Query, FPoseSearchDebugDrawParams DrawParams = FPoseSearchDebugDrawParams());
+
+
+//////////////////////////////////////////////////////////////////////////
+// FPoseSearchModule
+
+class FPoseSearchModule : public IModuleInterface
+{
+public:
+	virtual void StartupModule() override;
+	virtual void ShutdownModule() override;
+};
\ No newline at end of file
diff --git a/Engine/Source/Developer/AssetTools/Private/AssetTypeActions/AssetTypeActions_AnimationAsset.h b/Engine/Source/Developer/AssetTools/Private/AssetTypeActions/AssetTypeActions_AnimationAsset.h
deleted file mode 100644
diff --git a/Engine/Source/Developer/AssetTools/Public/AssetTypeActions/AssetTypeActions_AnimationAsset.h b/Engine/Source/Developer/AssetTools/Public/AssetTypeActions/AssetTypeActions_AnimationAsset.h
new file mode 100644
--- /dev/null
+++ b/Engine/Source/Developer/AssetTools/Public/AssetTypeActions/AssetTypeActions_AnimationAsset.h
@@ -0,0 +1,44 @@
+// Copyright Epic Games, Inc. All Rights Reserved.
+
+#pragma once
+
+#include "CoreMinimal.h"
+#include "Toolkits/IToolkitHost.h"
+#include "AssetTypeActions_Base.h"
+#include "Animation/AnimationAsset.h"
+#include "EditorAnimUtils.h"
+
+class ASSETTOOLS_API FAssetTypeActions_AnimationAsset : public FAssetTypeActions_Base
+{
+public:
+	// IAssetTypeActions Implementation
+	virtual FText GetName() const override { return NSLOCTEXT("AssetTypeActions", "AssetTypeActions_AnimationAsset", "AnimationAsset"); }
+	virtual FColor GetTypeColor() const override { return FColor(80,123,72); }
+	virtual UClass* GetSupportedClass() const override { return UAnimationAsset::StaticClass(); }
+	virtual bool HasActions ( const TArray<UObject*>& InObjects ) const override { return true; }
+	virtual void GetActions(const TArray<UObject*>& InObjects, struct FToolMenuSection& Section) override;
+	virtual void OpenAssetEditor( const TArray<UObject*>& InObjects, TSharedPtr<class IToolkitHost> EditWithinLevelEditor = TSharedPtr<IToolkitHost>() ) override;
+	virtual bool CanFilter() override { return false; }
+	virtual uint32 GetCategories() override { return EAssetTypeCategories::Animation; }
+	virtual class UThumbnailInfo* GetThumbnailInfo(UObject* Asset) const override;
+
+private:
+	/** Handler to fill the retarget submenu */
+	void FillRetargetMenu( FMenuBuilder& MenuBuilder, const TArray<UObject*> InObjects );
+
+	/** Handler for when FindSkeleton is selected */
+	void ExecuteFindSkeleton(TArray<TWeakObjectPtr<UAnimationAsset>> Objects);
+
+	/** Handle menu item for wanting to open asset in a new editor */
+	void ExecuteOpenInNewWindow(TArray<TWeakObjectPtr<UAnimationAsset>> Objects);
+
+	/** Open animation asset, will find existing editor if desired. */
+	void OpenAnimAssetEditor(const TArray<UObject*>& InObjects, bool bForceNewEditor, TSharedPtr<IToolkitHost> EditWithinLevelEditor);
+
+	/** Context menu item handler for changing the supplied assets skeletons */ 
+	void RetargetAssets(TArray<UObject*> InAnimBlueprints, bool bDuplicateAssets, bool bOpenEditor, TSharedPtr<class IToolkitHost> EditWithinLevelEditor);
+
+	/** Handler for retargeting */
+	void RetargetAnimationHandler(USkeleton* OldSkeleton, USkeleton* NewSkeleton, bool bRemapReferencedAssets, bool bAllowRemapToExisting, bool bConvertSpaces, const EditorAnimUtils::FNameDuplicationRule* NameRule, TArray<TWeakObjectPtr<UObject>> InAnimAssets);
+	void RetargetNonSkeletonAnimationHandler(USkeleton* OldSkeleton, USkeleton* NewSkeleton, bool bRemapReferencedAssets, bool bAllowRemapToExisting, bool bConvertSpaces, const EditorAnimUtils::FNameDuplicationRule* NameRule, TArray<TWeakObjectPtr<UObject>> InAnimAssets, TWeakPtr<class IToolkitHost> EditWithinLevelEditor);
+};
diff --git a/Engine/Source/Runtime/AnimationCore/Public/BoneIndices.h b/Engine/Source/Runtime/AnimationCore/Public/BoneIndices.h
--- a/Engine/Source/Runtime/AnimationCore/Public/BoneIndices.h
+++ b/Engine/Source/Runtime/AnimationCore/Public/BoneIndices.h
@@ -15,6 +15,8 @@ struct FBoneIndexBase
 
 	FORCEINLINE bool IsRootBone() const { return BoneIndex == 0; }
 
+	FORCEINLINE bool IsValid() const { return BoneIndex != INDEX_NONE; }
+
 protected:
 	int32 BoneIndex;
 };
diff --git a/Engine/Source/Runtime/Engine/Classes/Animation/AnimMetaData.h b/Engine/Source/Runtime/Engine/Classes/Animation/AnimMetaData.h
--- a/Engine/Source/Runtime/Engine/Classes/Animation/AnimMetaData.h
+++ b/Engine/Source/Runtime/Engine/Classes/Animation/AnimMetaData.h
@@ -14,7 +14,7 @@
 #include "UObject/Object.h"
 #include "AnimMetaData.generated.h"
 
-UCLASS(Blueprintable, abstract, const, editinlinenew, hidecategories=Object, collapsecategories)
+UCLASS(Blueprintable, abstract, editinlinenew, hidecategories=Object, collapsecategories)
 class ENGINE_API UAnimMetaData : public UObject
 {
 	GENERATED_UCLASS_BODY()
diff --git a/Engine/Source/Runtime/Engine/Classes/Animation/AnimNodeBase.h b/Engine/Source/Runtime/Engine/Classes/Animation/AnimNodeBase.h
--- a/Engine/Source/Runtime/Engine/Classes/Animation/AnimNodeBase.h
+++ b/Engine/Source/Runtime/Engine/Classes/Animation/AnimNodeBase.h
@@ -138,13 +138,15 @@ struct FAnimationBaseContext
 public:
 	FAnimInstanceProxy* AnimInstanceProxy;
 
+	FAnimationUpdateSharedContext* SharedContext;
+
 	FAnimationBaseContext();
 
 protected:
 	// DEPRECATED - Please use constructor that uses an FAnimInstanceProxy*
 	ENGINE_API FAnimationBaseContext(UAnimInstance* InAnimInstance);
 
-	ENGINE_API FAnimationBaseContext(FAnimInstanceProxy* InAnimInstanceProxy);
+	ENGINE_API FAnimationBaseContext(FAnimInstanceProxy* InAnimInstanceProxy, FAnimationUpdateSharedContext* InSharedContext = nullptr);
 
 public:
 	// we define a copy constructor here simply to avoid deprecation warnings with clang
@@ -161,6 +163,28 @@ public:
 	ENGINE_API UAnimBlueprint* GetAnimBlueprint() const;
 #endif //WITH_EDITORONLY_DATA
 
+	template<typename NodeType>
+	FScopedAnimNodeTracker TrackAncestor(NodeType* Node) const {
+		if (ensure(SharedContext != nullptr))
+		{
+			FAnimNodeTracker::FKey Key = SharedContext->AncestorTracker.Push<NodeType>(Node);
+			return FScopedAnimNodeTracker(&SharedContext->AncestorTracker, Key);
+		}
+
+		return FScopedAnimNodeTracker();
+	}
+
+	template<typename NodeType>
+	NodeType* GetAncestor() const {
+		if (ensure(SharedContext != nullptr))
+		{
+			FAnimNode_Base* Node = SharedContext->AncestorTracker.Top<NodeType>();
+			return static_cast<NodeType*>(Node);
+		}
+		
+		return nullptr;
+	}
+
 #if ANIM_NODE_IDS_AVAILABLE
 	// Get the current node Id, set when we recurse into graph traversal functions from pose links
 	ENGINE_API int32 GetCurrentNodeId() const { return CurrentNodeId; }
@@ -174,7 +198,7 @@ protected:
 
 	// The previous node ID, set when we recurse into graph traversal functions from pose links
 	int32 PreviousNodeId;
-#endif
+#endif // ANIM_NODE_IDS_AVAILABLE
 
 protected:
 
@@ -187,8 +211,8 @@ protected:
 struct FAnimationInitializeContext : public FAnimationBaseContext
 {
 public:
-	FAnimationInitializeContext(FAnimInstanceProxy* InAnimInstanceProxy)
-		: FAnimationBaseContext(InAnimInstanceProxy)
+	FAnimationInitializeContext(FAnimInstanceProxy* InAnimInstanceProxy, FAnimationUpdateSharedContext* InSharedContext = nullptr)
+		: FAnimationBaseContext(InAnimInstanceProxy, InSharedContext)
 	{
 	}
 };
@@ -210,8 +234,6 @@ public:
 struct FAnimationUpdateContext : public FAnimationBaseContext
 {
 private:
-	FAnimationUpdateSharedContext* SharedContext;
-
 	float CurrentWeight;
 	float RootMotionWeightModifier;
 
@@ -220,7 +242,6 @@ private:
 public:
 	FAnimationUpdateContext(FAnimInstanceProxy* InAnimInstanceProxy = nullptr)
 		: FAnimationBaseContext(InAnimInstanceProxy)
-		, SharedContext(nullptr)
 		, CurrentWeight(1.0f)
 		, RootMotionWeightModifier(1.0f)
 		, DeltaTime(0.0f)
@@ -228,18 +249,18 @@ public:
 	}
 
 	FAnimationUpdateContext(FAnimInstanceProxy* InAnimInstanceProxy, float InDeltaTime, FAnimationUpdateSharedContext* InSharedContext = nullptr)
-		: FAnimationUpdateContext(InAnimInstanceProxy)
+		: FAnimationBaseContext(InAnimInstanceProxy, InSharedContext)
+		, CurrentWeight(1.0f)
+		, RootMotionWeightModifier(1.0f)
+		, DeltaTime(InDeltaTime)
 	{
-		SharedContext = InSharedContext;
-		DeltaTime = InDeltaTime;
 	}
 
 
 	FAnimationUpdateContext(const FAnimationUpdateContext& Copy) = default;
 
 	FAnimationUpdateContext(const FAnimationUpdateContext& Copy, FAnimInstanceProxy* InAnimInstanceProxy)
-		: FAnimationBaseContext(InAnimInstanceProxy)
-		, SharedContext(Copy.SharedContext)
+		: FAnimationBaseContext(InAnimInstanceProxy, Copy.SharedContext)
 		, CurrentWeight(Copy.CurrentWeight)
 		, RootMotionWeightModifier(Copy.RootMotionWeightModifier)
 		, DeltaTime(Copy.DeltaTime)
@@ -314,32 +335,6 @@ public:
 	}
 #endif
 
-	// Add a node to the list of tracked ancestors
-	template<typename NodeType>
-	FScopedAnimNodeTracker TrackAncestor(NodeType* Node) const
-	{
-		if (ensure(SharedContext != nullptr))
-		{
-			FAnimNodeTracker::FKey Key = SharedContext->AncestorTracker.Push<NodeType>(Node);
-			return FScopedAnimNodeTracker(&SharedContext->AncestorTracker, Key);
-		}
-
-		return FScopedAnimNodeTracker();
-	}
-
-	// Returns the nearest ancestor node of a particular type
-	template<typename NodeType>
-	NodeType* GetAncestor() const
-	{
-		if (ensure(SharedContext != nullptr))
-		{
-			FAnimNode_Base* Node = SharedContext->AncestorTracker.Top<NodeType>();
-			return static_cast<NodeType*>(Node);
-		}
-		
-		return nullptr;
-	}
-
 	// Returns persistent state that is tracked through animation tree update
 	FAnimationUpdateSharedContext* GetSharedContext() const
 	{
diff --git a/Engine/Source/Runtime/Engine/Classes/Animation/AnimationAsset.h b/Engine/Source/Runtime/Engine/Classes/Animation/AnimationAsset.h
--- a/Engine/Source/Runtime/Engine/Classes/Animation/AnimationAsset.h
+++ b/Engine/Source/Runtime/Engine/Classes/Animation/AnimationAsset.h
@@ -19,6 +19,7 @@
 #include "Interfaces/Interface_PreviewMeshProvider.h"
 #include "AnimationAsset.generated.h"
 
+class UAnimMetaData;
 class UAssetMappingTable;
 class UAssetUserData;
 class USkeleton;
@@ -799,10 +800,10 @@ struct FAnimationGroupReference
 	}
 };
 
-UCLASS(abstract, MinimalAPI)
-class UAnimationAsset : public UObject, public IInterface_AssetUserData, public IInterface_PreviewMeshProvider
+UCLASS(abstract)
+class ENGINE_API UAnimationAsset : public UObject, public IInterface_AssetUserData, public IInterface_PreviewMeshProvider
 {
-	GENERATED_UCLASS_BODY()
+	GENERATED_BODY()
 
 private:
 	/** Pointer to the Skeleton this asset can be played on .	*/
@@ -820,7 +821,7 @@ private:
 	 * You can query by GetMetaData function
 	 */
 	UPROPERTY(Category=MetaData, instanced, EditAnywhere)
-	TArray<class UAnimMetaData*> MetaData;
+	TArray<UAnimMetaData*> MetaData;
 
 public:
 	/* 
@@ -878,7 +879,7 @@ public:
 	// but for BlendSpace CurrentTime is normalized [0,1], so this is 1
 	virtual float GetMaxCurrentTime() { return 0.f; }
 
-	ENGINE_API void SetSkeleton(USkeleton* NewSkeleton);
+	void SetSkeleton(USkeleton* NewSkeleton);
 	void ResetSkeleton(USkeleton* NewSkeleton);
 	virtual void PostLoad() override;
 
@@ -889,11 +890,11 @@ public:
 
 	/** Get available Metadata within the animation asset
 	 */
-	ENGINE_API const TArray<class UAnimMetaData*>& GetMetaData() const { return MetaData; }
-	ENGINE_API void AddMetaData(class UAnimMetaData* MetaDataInstance); 
-	ENGINE_API void EmptyMetaData() { MetaData.Empty(); }	
-	ENGINE_API void RemoveMetaData(class UAnimMetaData* MetaDataInstance);
-	ENGINE_API void RemoveMetaData(const TArray<class UAnimMetaData*> MetaDataInstances);
+	const TArray<UAnimMetaData*>& GetMetaData() const { return MetaData; }
+	void AddMetaData(UAnimMetaData* MetaDataInstance); 
+	void EmptyMetaData() { MetaData.Empty(); }	
+	void RemoveMetaData(UAnimMetaData* MetaDataInstance);
+	void RemoveMetaData(TArrayView<UAnimMetaData*> MetaDataInstances);
 
 	/** IInterface_PreviewMeshProvider interface */
 	virtual void SetPreviewMesh(USkeletalMesh* PreviewMesh, bool bMarkAsDirty = true) override;
@@ -905,33 +906,33 @@ public:
 	 * 
 	 * @param NewSkeleton	NewSkeleton to change to 
 	 */
-	ENGINE_API bool ReplaceSkeleton(USkeleton* NewSkeleton, bool bConvertSpaces=false);
+	bool ReplaceSkeleton(USkeleton* NewSkeleton, bool bConvertSpaces=false);
 
 	// Helper function for GetAllAnimationSequencesReferred, it adds itself first and call GetAllAnimationSEquencesReferred
-	ENGINE_API void HandleAnimReferenceCollection(TArray<UAnimationAsset*>& AnimationAssets, bool bRecursive);
+	void HandleAnimReferenceCollection(TArray<UAnimationAsset*>& AnimationAssets, bool bRecursive);
 
 protected:
 	/** Retrieve all animations that are used by this asset 
 	 * 
 	 * @param (out)		AnimationSequences 
 	 **/
-	ENGINE_API virtual bool GetAllAnimationSequencesReferred(TArray<class UAnimationAsset*>& AnimationSequences, bool bRecursive = true);
+	virtual bool GetAllAnimationSequencesReferred(TArray<class UAnimationAsset*>& AnimationSequences, bool bRecursive = true);
 
 public:
 	/** Replace this assets references to other animations based on ReplacementMap 
 	 * 
 	 * @param ReplacementMap	Mapping of original asset to new asset
 	 **/
-	ENGINE_API virtual void ReplaceReferredAnimations(const TMap<UAnimationAsset*, UAnimationAsset*>& ReplacementMap);
+	virtual void ReplaceReferredAnimations(const TMap<UAnimationAsset*, UAnimationAsset*>& ReplacementMap);
 
-	ENGINE_API virtual int32 GetMarkerUpdateCounter() const { return 0; }
+	virtual int32 GetMarkerUpdateCounter() const { return 0; }
 
 	/** 
 	 * Parent Asset related function. Used by editor
 	 */
-	ENGINE_API void SetParentAsset(UAnimationAsset* InParentAsset);
-	ENGINE_API bool HasParentAsset() { return ParentAsset != nullptr;  }
-	ENGINE_API bool RemapAsset(UAnimationAsset* SourceAsset, UAnimationAsset* TargetAsset);
+	void SetParentAsset(UAnimationAsset* InParentAsset);
+	bool HasParentAsset() { return ParentAsset != nullptr;  }
+	bool RemapAsset(UAnimationAsset* SourceAsset, UAnimationAsset* TargetAsset);
 	// we have to update whenever we have anything loaded
 	void UpdateParentAsset();
 protected:
@@ -940,18 +941,18 @@ protected:
 
 public:
 	/** Return a list of unique marker names for blending compatibility */
-	ENGINE_API virtual TArray<FName>* GetUniqueMarkerNames() { return NULL; }
+	virtual TArray<FName>* GetUniqueMarkerNames() { return NULL; }
 
 	//~ Begin IInterface_AssetUserData Interface
-	ENGINE_API virtual void AddAssetUserData(UAssetUserData* InUserData) override;
-	ENGINE_API virtual void RemoveUserDataOfClass(TSubclassOf<UAssetUserData> InUserDataClass) override;
-	ENGINE_API virtual UAssetUserData* GetAssetUserDataOfClass(TSubclassOf<UAssetUserData> InUserDataClass) override;
-	ENGINE_API virtual const TArray<UAssetUserData*>* GetAssetUserDataArray() const override;
+	virtual void AddAssetUserData(UAssetUserData* InUserData) override;
+	virtual void RemoveUserDataOfClass(TSubclassOf<UAssetUserData> InUserDataClass) override;
+	virtual UAssetUserData* GetAssetUserDataOfClass(TSubclassOf<UAssetUserData> InUserDataClass) override;
+	virtual const TArray<UAssetUserData*>* GetAssetUserDataArray() const override;
 	//~ End IInterface_AssetUserData Interface
 
 	//~ Begin UObject Interface.
 #if WITH_EDITOR
-	ENGINE_API virtual void PostEditChangeProperty(FPropertyChangedEvent& PropertyChangedEvent) override;
+	virtual void PostEditChangeProperty(FPropertyChangedEvent& PropertyChangedEvent) override;
 #endif // WITH_EDITOR
 
 	/**
diff --git a/Engine/Source/Runtime/Engine/Private/Animation/AnimInstanceProxy.cpp b/Engine/Source/Runtime/Engine/Private/Animation/AnimInstanceProxy.cpp
--- a/Engine/Source/Runtime/Engine/Private/Animation/AnimInstanceProxy.cpp
+++ b/Engine/Source/Runtime/Engine/Private/Animation/AnimInstanceProxy.cpp
@@ -289,7 +289,8 @@ void FAnimInstanceProxy::InitializeRootNode_WithRoot(FAnimNode_Base* InRootNode)
 
 	if (InRootNode != nullptr)
 	{
-		FAnimationInitializeContext InitContext(this);
+		FAnimationUpdateSharedContext SharedContext;
+		FAnimationInitializeContext InitContext(this, &SharedContext);
 
 		if(InRootNode == RootNode)
 		{
diff --git a/Engine/Source/Runtime/Engine/Private/Animation/AnimNodeBase.cpp b/Engine/Source/Runtime/Engine/Private/Animation/AnimNodeBase.cpp
--- a/Engine/Source/Runtime/Engine/Private/Animation/AnimNodeBase.cpp
+++ b/Engine/Source/Runtime/Engine/Private/Animation/AnimNodeBase.cpp
@@ -59,6 +59,7 @@ const void* FExposedValueCopyRecord::GetSourceAddr(FAnimInstanceProxy* Proxy) co
 
 FAnimationBaseContext::FAnimationBaseContext()
 	: AnimInstanceProxy(nullptr)
+	, SharedContext(nullptr)
 #if ANIM_TRACE_ENABLED
 	, CurrentNodeId(INDEX_NONE)
 	, PreviousNodeId(INDEX_NONE)
@@ -66,8 +67,9 @@ FAnimationBaseContext::FAnimationBaseContext()
 {
 }
 
-FAnimationBaseContext::FAnimationBaseContext(FAnimInstanceProxy* InAnimInstanceProxy)
+FAnimationBaseContext::FAnimationBaseContext(FAnimInstanceProxy* InAnimInstanceProxy, FAnimationUpdateSharedContext* InSharedContext)
 	: AnimInstanceProxy(InAnimInstanceProxy)
+	, SharedContext(InSharedContext)
 #if ANIM_TRACE_ENABLED
 	, CurrentNodeId(INDEX_NONE)
 	, PreviousNodeId(INDEX_NONE)
@@ -77,6 +79,7 @@ FAnimationBaseContext::FAnimationBaseContext(FAnimInstanceProxy* InAnimInstanceP
 
 FAnimationBaseContext::FAnimationBaseContext(const FAnimationBaseContext& InContext)
 	: AnimInstanceProxy(InContext.AnimInstanceProxy)
+	, SharedContext(InContext.SharedContext)
 #if ANIM_TRACE_ENABLED
 	, CurrentNodeId(InContext.CurrentNodeId)
 	, PreviousNodeId(InContext.PreviousNodeId)
diff --git a/Engine/Source/Runtime/Engine/Private/Animation/AnimNode_StateMachine.cpp b/Engine/Source/Runtime/Engine/Private/Animation/AnimNode_StateMachine.cpp
--- a/Engine/Source/Runtime/Engine/Private/Animation/AnimNode_StateMachine.cpp
+++ b/Engine/Source/Runtime/Engine/Private/Animation/AnimNode_StateMachine.cpp
@@ -88,7 +88,7 @@ void FAnimationActiveTransitionEntry::InitializeCustomGraphLinks(const FAnimatio
 		if (const IAnimClassInterface* AnimBlueprintClass = Context.GetAnimClass())
 		{
 			CustomTransitionGraph.LinkID = AnimBlueprintClass->GetAnimNodeProperties().Num() - 1 - TransitionRule.CustomResultNodeIndex; //@TODO: Crazysauce
-			FAnimationInitializeContext InitContext(Context.AnimInstanceProxy);
+			FAnimationInitializeContext InitContext(Context.AnimInstanceProxy, Context.SharedContext);
 			CustomTransitionGraph.Initialize(InitContext);
 
 			if (Context.AnimInstanceProxy)
@@ -365,7 +365,7 @@ void FAnimNode_StateMachine::Update_AnyThread(const FAnimationUpdateContext& Con
 	// If we just became relevant and haven't been initialized yet, then reinitialize state machine.
 	if (!bFirstUpdate && bReinitializeOnBecomingRelevant && UpdateCounter.HasEverBeenUpdated() && !UpdateCounter.WasSynchronizedCounter(Context.AnimInstanceProxy->GetUpdateCounter()) && (CVarAnimStateMachineRelevancyReset.GetValueOnAnyThread() == 1))
 	{
-		FAnimationInitializeContext InitializationContext(Context.AnimInstanceProxy);
+		FAnimationInitializeContext InitializationContext(Context.AnimInstanceProxy, Context.SharedContext);
 		Initialize_AnyThread(InitializationContext);
 	}
 	UpdateCounter.SynchronizeWith(Context.AnimInstanceProxy->GetUpdateCounter());
@@ -1017,7 +1017,7 @@ void FAnimNode_StateMachine::SetState(const FAnimationBaseContext& Context, int3
 		if ((!bAlreadyActive || bForceReset) && !IsAConduitState(NewStateIndex))
 		{
 			// Initialize the new state since it's not part of an active transition (and thus not still initialized)
-			FAnimationInitializeContext InitContext(Context.AnimInstanceProxy);
+			FAnimationInitializeContext InitContext(Context.AnimInstanceProxy, Context.SharedContext);
 			StatePoseLinks[NewStateIndex].Initialize(InitContext);
 
 			// Also call cache bones if needed
diff --git a/Engine/Source/Runtime/Engine/Private/AnimationAsset.cpp b/Engine/Source/Runtime/Engine/Private/AnimationAsset.cpp
--- a/Engine/Source/Runtime/Engine/Private/AnimationAsset.cpp
+++ b/Engine/Source/Runtime/Engine/Private/AnimationAsset.cpp
@@ -178,11 +178,6 @@ void FAnimGroupInstance::Prepare(const FAnimGroupInstance* PreviousGroup)
 //////////////////////////////////////////////////////////////////////////
 // UAnimationAsset
 
-UAnimationAsset::UAnimationAsset(const FObjectInitializer& ObjectInitializer)
-	: Super(ObjectInitializer)
-{
-}
-
 void UAnimationAsset::PostLoad()
 {
 	LLM_SCOPE(ELLMTag::Animation);
@@ -242,22 +237,22 @@ void UAnimationAsset::Serialize(FArchive& Ar)
 	}
 }
 
-void UAnimationAsset::AddMetaData(class UAnimMetaData* MetaDataInstance)
+void UAnimationAsset::AddMetaData(UAnimMetaData* MetaDataInstance)
 {
 	MetaData.Add(MetaDataInstance);
 }
 
-void UAnimationAsset::RemoveMetaData(class UAnimMetaData* MetaDataInstance)
+void UAnimationAsset::RemoveMetaData(UAnimMetaData* MetaDataInstance)
 {
 	MetaData.Remove(MetaDataInstance);
 }
 
-void UAnimationAsset::RemoveMetaData(const TArray<UAnimMetaData*> MetaDataInstances)
+void UAnimationAsset::RemoveMetaData(TArrayView<UAnimMetaData*> MetaDataInstances)
 {
 	MetaData.RemoveAll(
 		[&](UAnimMetaData* MetaDataInstance)
 	{
-		return MetaDataInstances.Find(MetaDataInstance);
+		return MetaDataInstances.Contains(MetaDataInstance);
 	});
 }
 
