Non-unity build fixes

diff --git a/Engine/Plugins/Animation/PoseSearch/Source/Editor/Private/PoseSearchEditor.cpp b/Engine/Plugins/Animation/PoseSearch/Source/Editor/Private/PoseSearchEditor.cpp
--- a/Engine/Plugins/Animation/PoseSearch/Source/Editor/Private/PoseSearchEditor.cpp
+++ b/Engine/Plugins/Animation/PoseSearch/Source/Editor/Private/PoseSearchEditor.cpp
@@ -49,7 +49,6 @@ public:
 	virtual void StartupModule() override
 	{
 		FAssetToolsModule::GetModule().Get().RegisterAssetTypeActions(MakeShareable(new FAssetTypeActions_PoseSearchSchema));
-		//UAssetToolsHelpers::GetAssetTools()->RegisterAssetTypeActions(MakeShareable(new FAssetTypeActions_PoseSearchSchema));
 	}
 
 	virtual void ShutdownModule() override
diff --git a/Engine/Plugins/Animation/PoseSearch/Source/Runtime/Public/PoseSearch/AnimNode_PoseSearchHistoryCollector.h b/Engine/Plugins/Animation/PoseSearch/Source/Runtime/Public/PoseSearch/AnimNode_PoseSearchHistoryCollector.h
--- a/Engine/Plugins/Animation/PoseSearch/Source/Runtime/Public/PoseSearch/AnimNode_PoseSearchHistoryCollector.h
+++ b/Engine/Plugins/Animation/PoseSearch/Source/Runtime/Public/PoseSearch/AnimNode_PoseSearchHistoryCollector.h
@@ -3,7 +3,7 @@
 #pragma once
 
 #include "PoseSearch/PoseSearch.h"
-
+#include "Animation/AnimNodeBase.h"
 #include "AnimNode_PoseSearchHistoryCollector.generated.h"
 
 USTRUCT(BlueprintInternalUseOnly)
diff --git a/Engine/Plugins/Animation/PoseSearch/Source/Runtime/Public/PoseSearch/PoseSearch.h b/Engine/Plugins/Animation/PoseSearch/Source/Runtime/Public/PoseSearch/PoseSearch.h
--- a/Engine/Plugins/Animation/PoseSearch/Source/Runtime/Public/PoseSearch/PoseSearch.h
+++ b/Engine/Plugins/Animation/PoseSearch/Source/Runtime/Public/PoseSearch/PoseSearch.h
@@ -15,6 +15,7 @@
 
 
 class UAnimSequenceBase;
+struct FCompactPose;
 
 /** Sampling parameters for building a pose search index for an asset. */
 UCLASS(BlueprintType, Category = "Animation|PoseSearch")
