Fixed potential null deref
#rnx

diff --git a/Engine/Plugins/Animation/PoseSearch/Source/Runtime/Private/AnimNode_PoseMatchingSequencePlayer.cpp b/Engine/Plugins/Animation/PoseSearch/Source/Runtime/Private/AnimNode_PoseMatchingSequencePlayer.cpp
--- a/Engine/Plugins/Animation/PoseSearch/Source/Runtime/Private/AnimNode_PoseMatchingSequencePlayer.cpp
+++ b/Engine/Plugins/Animation/PoseSearch/Source/Runtime/Private/AnimNode_PoseMatchingSequencePlayer.cpp
@@ -18,7 +18,7 @@ const UPoseSearchIndex* GetPoseSearchDataIndex(const UAnimSequenceBase& Sequence
 
 	for (UAnimMetaData* MetaDataInstance : MetaData)
 	{
-		if (MetaDataInstance->GetClass() == UPoseSearchIndex::StaticClass())
+		if (MetaDataInstance && (MetaDataInstance->GetClass() == UPoseSearchIndex::StaticClass()))
 		{
 			PoseSearchIndex = Cast<UPoseSearchIndex>(MetaDataInstance);
 			break;
