Fixup deprecated API call - GetAnimationPose taking different arguments

diff --git a/Engine/Plugins/Animation/PoseSearch/Source/Runtime/Private/PoseSearch.cpp b/Engine/Plugins/Animation/PoseSearch/Source/Runtime/Private/PoseSearch.cpp
--- a/Engine/Plugins/Animation/PoseSearch/Source/Runtime/Private/PoseSearch.cpp
+++ b/Engine/Plugins/Animation/PoseSearch/Source/Runtime/Private/PoseSearch.cpp
@@ -290,6 +290,8 @@ void PoseSearchBuildIndex(const UAnimSequenceBase& AnimSequence, const UPoseSear
 	BoneContainer.InitializeTo(SearchSchema.BoneIndicesWithParents, FCurveEvaluationOption(false), *Skeleton);
 
 	FBlendedCurve UnusedCurve;
+	FStackCustomAttributes UnusedAttributes;
+
 	FAnimExtractContext ExtractionCtx;
 	// ExtractionCtx.PoseCurves is intentionally left empty
 	// ExtractionCtx.BonesRequired is unused by UAnimSequence::GetAnimationPose
@@ -307,11 +309,12 @@ void PoseSearchBuildIndex(const UAnimSequenceBase& AnimSequence, const UPoseSear
 	Pose.SetBoneContainer(&BoneContainer);
 	FCSPose<FCompactPose> ComponentSpacePose;
 
+	FAnimationPoseData AnimPoseData(Pose, UnusedCurve, UnusedAttributes);
 	for (int32 PoseIdx = 0; PoseIdx != NumPoses; ++PoseIdx, CurrTime += DeltaTime)
 	{
 		// Extract pose
 		ExtractionCtx.CurrentTime = CurrTime;
-		AnimSequence.GetAnimationPose(Pose, UnusedCurve, ExtractionCtx);
+		AnimSequence.GetAnimationPose(AnimPoseData, ExtractionCtx);
 		ComponentSpacePose.InitPose(Pose);
 
 		for (int32 BoneIndex : SearchSchema.BoneIndices)
