Fix deprecation warnings/errors in plugins

#rnx

diff --git a/Engine/Plugins/Animation/PoseSearch/Source/Editor/Private/AnimGraphNode_PoseMatchingSequencePlayer.cpp b/Engine/Plugins/Animation/PoseSearch/Source/Editor/Private/AnimGraphNode_PoseMatchingSequencePlayer.cpp
--- a/Engine/Plugins/Animation/PoseSearch/Source/Editor/Private/AnimGraphNode_PoseMatchingSequencePlayer.cpp
+++ b/Engine/Plugins/Animation/PoseSearch/Source/Editor/Private/AnimGraphNode_PoseMatchingSequencePlayer.cpp
@@ -368,8 +368,10 @@ void UAnimGraphNode_PoseMatchingSequencePlayer::SetAnimationAsset(UAnimationAsse
 void UAnimGraphNode_PoseMatchingSequencePlayer::BakeDataDuringCompilation(class FCompilerResultsLog& MessageLog)
 {
 	UAnimBlueprint* AnimBlueprint = GetAnimBlueprint();
-	Node.GroupIndex = AnimBlueprint->FindOrAddGroup(SyncGroup.GroupName);
+	AnimBlueprint->FindOrAddGroup(SyncGroup.GroupName);
+	Node.GroupName = SyncGroup.GroupName;
 	Node.GroupRole = SyncGroup.GroupRole;
+	Node.GroupScope = SyncGroup.GroupScope;
 }
 
 void UAnimGraphNode_PoseMatchingSequencePlayer::GetAllAnimationSequencesReferred(TArray<UAnimationAsset*>& AnimationAssets) const
