PoseSearch: Pose feature vector API, support for indexing and searching collections of animations, and lots of refactoring and cleanup.
+ Pose feature vector API
  + FPoseSearchFeatureVectorLayout
    - Explicit description of a pose feature vector (FV)
    - Determined by options set in a UPoseSearchSchema and owned by the schema. See UPoseSearchSchema::GenerateLayout().
  + FPoseSearchFeatureDesc
    - Describes each feature of a vector, including data type, sampling options, and buffer offset.
    - New data types introduced (previously only positions were supported)
    - Could potentially be extended in the future to support externally registered features, but that's not the goal for now.
  + FFeatureVectorBuilder
    - Helper object for writing features into a float buffer according to the layout
    - Keeps track of which features are present, allowing the FV to be constructed piecemeal
  + FFeatureVectorReader
    - Helper object for extracting features. Presently only used for debug rendering.
+ UPoseSearchDatabase
  - A data asset for indexing a collection of animation sequences
  - Support added for searching the collection and returning the matching sequence with time offset
+ Pose matching refactoring
  - Setup no longer requires an animation modifier or a special graph node
  + UPoseSearchSequenceMetaData
    - Replaces UPoseSearchIndexConfig and UPoseSearchIndex
    - Counterpart to UPoseSearchDatabase for a single animation instead of a collection
    - Simply add to sequence metadata, fill it out, and save
  + IPoseSearchProvider
    - Modular feature interface that replaces the need for AnimGraphNode_PoseMatchingSequencePlayer
    - The PoseSearch plugin now registers itself as an implementer of IPoseSearchProvider
    - AnimGraphNode_SequencePlayer now detects if a pose search provider exists and offers optional pose matching support
    - Currently exposed as checkbox in an "experimental" section of the node's settings
+ Other cleanup/refactoring
  + Namespacing
    - Pulled all non-UCLASS/USTRUCTs into a UE::PoseSearch namespace
    - Renamed PoseSearchDetail namespace to UE::PoseSearch::Private per standards
  + UPoseSearchSchema is no longer an animation asset
    - This was a kludge to get the bone reference UI to work
    - Now it's a data asset and implements IBoneReferenceSkeletonProvider
    - Removed related asset type actions and factory
  + FPoseSearchIndex replaces bits of UPoseSearchIndex
    - Has narrower scope and is shared by UPoseSearchSequenceMetaData and UPoseSearchDatabase
  + Odds n' ends
    - Removed FPoseSearchBuildQueryScratch. FPoseHistory now owns memory required for taking samples.
    - Debug drawing reworked in terms of FV API
    - Added a.PoseSearch.DrawSearchIndex console command to render the search index of any animation sequences currently open in the editor
    - Removed Blueprint function library for building search indices since animation modifiers are no longer needed

diff --git a/Engine/Plugins/Animation/PoseSearch/Source/Editor/PoseSearchEditor.Build.cs b/Engine/Plugins/Animation/PoseSearch/Source/Editor/PoseSearchEditor.Build.cs
--- a/Engine/Plugins/Animation/PoseSearch/Source/Editor/PoseSearchEditor.Build.cs
+++ b/Engine/Plugins/Animation/PoseSearch/Source/Editor/PoseSearchEditor.Build.cs
@@ -9,7 +9,8 @@ public class PoseSearchEditor : ModuleRules
 		PCHUsage = ModuleRules.PCHUsageMode.UseExplicitOrSharedPCHs;
 
 		PublicDependencyModuleNames.AddRange(
-			new string[] {
+			new string[]
+			{
 				"AnimGraph",
 				"AnimationCore",
 				"Core",
@@ -23,11 +24,9 @@ public class PoseSearchEditor : ModuleRules
 		);
 
 		PrivateDependencyModuleNames.AddRange(
-			new string[] {
-				"AssetTools",
+			new string[]
+			{
 				"BlueprintGraph",
-				"GraphEditor",
-				"SlateCore",
 			}
 		);
 	}
diff --git a/Engine/Plugins/Animation/PoseSearch/Source/Editor/Private/AnimGraphNode_PoseMatchingSequencePlayer.cpp b/Engine/Plugins/Animation/PoseSearch/Source/Editor/Private/AnimGraphNode_PoseMatchingSequencePlayer.cpp
deleted file mode 100644
diff --git a/Engine/Plugins/Animation/PoseSearch/Source/Editor/Private/PoseSearchEditor.cpp b/Engine/Plugins/Animation/PoseSearch/Source/Editor/Private/PoseSearchEditor.cpp
--- a/Engine/Plugins/Animation/PoseSearch/Source/Editor/Private/PoseSearchEditor.cpp
+++ b/Engine/Plugins/Animation/PoseSearch/Source/Editor/Private/PoseSearchEditor.cpp
@@ -1,59 +1,91 @@
 // Copyright Epic Games, Inc. All Rights Reserved.
 
 #include "PoseSearchEditor.h"
-#include "Editor.h"
 #include "Animation/AnimSequence.h"
-#include "IAssetTools.h"
-#include "AssetToolsModule.h"
 #include "Modules/ModuleManager.h"
+#include "Editor.h"
+#include "Subsystems/AssetEditorSubsystem.h"
+#include "IAnimationEditor.h"
+#include "IPersonaToolkit.h"
+#include "IPersonaPreviewScene.h"
 
-void UPoseSearchBlueprintLibrary::BuildPoseSearchIndex(const UAnimSequence* AnimationSequence, const UPoseSearchIndexConfig* Config, const UPoseSearchSchema* Schema, UPoseSearchIndex* SearchIndex)
-{
-	if (!ensure(AnimationSequence))
-	{
-		return;
-	}
-	if (!ensure(Config))
-	{
-		return;
-	}
-	if (!ensure(Schema))
-	{
-		return;
-	}
-	if (!ensure(SearchIndex))
-	{
-		return;
-	}
 
-	PoseSearchBuildIndex(*AnimationSequence, *Config, *Schema, SearchIndex);
-}
+//////////////////////////////////////////////////////////////////////////
+// FEditorCommands
 
-UPoseSearchSchemaFactory::UPoseSearchSchemaFactory(const FObjectInitializer& ObjectInitializer)
-    : Super(ObjectInitializer)
+namespace UE { namespace PoseSearch {
+
+struct FEditorCommands
 {
-	bCreateNew = true;
-	bEditAfterNew = true;
-	SupportedClass = UPoseSearchSchema::StaticClass();
-}
+	static void DrawSearchIndex(UWorld * World);
+};
 
-UObject* UPoseSearchSchemaFactory::FactoryCreateNew(UClass* InClass, UObject* InParent, FName InName, EObjectFlags Flags, UObject* Context, FFeedbackContext* Warn)
+void FEditorCommands::DrawSearchIndex(UWorld * World)
 {
-	UPoseSearchSchema* Schema = NewObject<UPoseSearchSchema>(InParent, InClass, InName, Flags | RF_Transactional);
-	return Schema;
+	FDebugDrawParams DrawParams;
+	DrawParams.DefaultLifeTime = 60.0f;
+	DrawParams.Flags = EDebugDrawFlags::DrawSearchIndex;
+
+	TArray<UObject*> EditedAssets = GEditor->GetEditorSubsystem<UAssetEditorSubsystem>()->GetAllEditedAssets();
+	for (UObject* EditedAsset : EditedAssets)
+	{
+		if (UAnimSequence* Sequence = Cast<UAnimSequence>(EditedAsset))
+		{
+			const UPoseSearchSequenceMetaData* MetaData = Sequence->FindMetaDataByClass<UPoseSearchSequenceMetaData>();
+			DrawParams.SearchIndex = MetaData ? &MetaData->SearchIndex : nullptr;
+			if (DrawParams.SearchIndex)
+			{
+				IAssetEditorInstance* EditorInstance = GEditor->GetEditorSubsystem<UAssetEditorSubsystem>()->FindEditorForAsset(Sequence, true /*bFocusIfOpen*/);
+				if (EditorInstance && EditorInstance->GetEditorName() == TEXT("AnimationEditor"))
+				{
+					IAnimationEditor* Editor = static_cast<IAnimationEditor*>(EditorInstance);
+					TSharedRef<IPersonaToolkit> Toolkit = Editor->GetPersonaToolkit();
+					TSharedRef<IPersonaPreviewScene> Scene = Toolkit->GetPreviewScene();
+
+					DrawParams.World = Scene->GetWorld();
+					Draw(DrawParams);
+				}
+			}
+		}
+	}
 }
 
-class FPoseSearchEditorModule : public IPoseSearchEditorModuleInterface
+
+//////////////////////////////////////////////////////////////////////////
+// FPoseSearchEditorModule
+
+class FEditorModule : public IPoseSearchEditorModuleInterface
 {
 public:
-	virtual void StartupModule() override
+	void StartupModule() override;
+	void ShutdownModule() override;
+
+private:
+	TArray<class IConsoleObject*> ConsoleCommands;
+};
+
+void FEditorModule::StartupModule()
+{
+	if (GIsEditor && !IsRunningCommandlet())
 	{
-		FAssetToolsModule::GetModule().Get().RegisterAssetTypeActions(MakeShareable(new FAssetTypeActions_PoseSearchSchema));
+		ConsoleCommands.Add(IConsoleManager::Get().RegisterConsoleCommand(
+			TEXT("a.PoseSearch.DrawSearchIndex"),
+			TEXT("Draw the search index for the selected asset"),
+			FConsoleCommandWithWorldDelegate::CreateStatic(&FEditorCommands::DrawSearchIndex),
+			ECVF_Default
+		));
 	}
+}
 
-	virtual void ShutdownModule() override
+void FEditorModule::ShutdownModule()
+{
+	for (IConsoleObject* ConsoleCmd : ConsoleCommands)
 	{
+		IConsoleManager::Get().UnregisterConsoleObject(ConsoleCmd);
 	}
-};
+	ConsoleCommands.Empty();
+}
+
+}} // namespace UE::PoseSearch
 
-IMPLEMENT_MODULE(FPoseSearchEditorModule, PoseSearchEditor);
+IMPLEMENT_MODULE(UE::PoseSearch::FEditorModule, PoseSearchEditor);
diff --git a/Engine/Plugins/Animation/PoseSearch/Source/Editor/Public/AnimGraphNode_PoseMatchingSequencePlayer.h b/Engine/Plugins/Animation/PoseSearch/Source/Editor/Public/AnimGraphNode_PoseMatchingSequencePlayer.h
deleted file mode 100644
diff --git a/Engine/Plugins/Animation/PoseSearch/Source/Editor/Public/PoseSearchEditor.h b/Engine/Plugins/Animation/PoseSearch/Source/Editor/Public/PoseSearchEditor.h
--- a/Engine/Plugins/Animation/PoseSearch/Source/Editor/Public/PoseSearchEditor.h
+++ b/Engine/Plugins/Animation/PoseSearch/Source/Editor/Public/PoseSearchEditor.h
@@ -2,49 +2,9 @@
 
 #pragma once
 
-
-#include "AssetTypeActions/AssetTypeActions_AnimationAsset.h"
-#include "Engine/EngineTypes.h"
-#include "Factories/Factory.h"
-#include "Kismet/BlueprintFunctionLibrary.h"
-#include "PoseSearch/PoseSearch.h"
 #include "CoreMinimal.h"
 #include "Modules/ModuleInterface.h"
-
-#include "PoseSearchEditor.generated.h"
-
-
-class UAnimSequence;
-
-UCLASS()
-class POSESEARCHEDITOR_API UPoseSearchBlueprintLibrary : public UBlueprintFunctionLibrary
-{
-	GENERATED_BODY()
-public:
-	/** Compiles a pose search index */
-	UFUNCTION(BlueprintCallable, Category = "AnimationBlueprintLibrary|PoseSearch")
-	static void BuildPoseSearchIndex(
-	    const UAnimSequence* AnimationSequence,
-	    const UPoseSearchIndexConfig* SearchConfig,
-	    const UPoseSearchSchema* SearchSchema,
-	    UPoseSearchIndex* SearchIndex);
-};
-
-UCLASS()
-class POSESEARCHEDITOR_API UPoseSearchSchemaFactory : public UFactory
-{
-	GENERATED_BODY()
-public:
-	UPoseSearchSchemaFactory(const FObjectInitializer& ObjectInitializer);
-	virtual UObject* FactoryCreateNew(UClass* InClass, UObject* InParent, FName InName, EObjectFlags Flags, UObject* Context, FFeedbackContext* Warn) override;
-};
-
-class POSESEARCHEDITOR_API FAssetTypeActions_PoseSearchSchema : public FAssetTypeActions_AnimationAsset
-{
-	virtual FText GetName() const override { return NSLOCTEXT("AssetTypeActions", "PoseSearchEditor", "Pose Search Schema"); }
-	virtual UClass* GetSupportedClass() const override { return UPoseSearchSchema::StaticClass(); }
-	virtual bool CanFilter() override { return true; }
-};
+#include "PoseSearch/PoseSearch.h"
 
 class IPoseSearchEditorModuleInterface : public IModuleInterface
 {
diff --git a/Engine/Plugins/Animation/PoseSearch/Source/Runtime/PoseSearch.Build.cs b/Engine/Plugins/Animation/PoseSearch/Source/Runtime/PoseSearch.Build.cs
--- a/Engine/Plugins/Animation/PoseSearch/Source/Runtime/PoseSearch.Build.cs
+++ b/Engine/Plugins/Animation/PoseSearch/Source/Runtime/PoseSearch.Build.cs
@@ -9,11 +9,19 @@ public class PoseSearch : ModuleRules
 		PCHUsage = ModuleRules.PCHUsageMode.UseExplicitOrSharedPCHs;
 
 		PublicIncludePaths.AddRange(
-			new string[] {
+			new string[]
+			{
 				System.IO.Path.Combine(ModuleDirectory, "Public"),
 			}
-			);			
-		
+		);
+
+		PrivateIncludePaths.AddRange(
+			new string[]
+			{
+				System.IO.Path.Combine(ModuleDirectory, "Private"),
+			}
+		);
+
 		PublicDependencyModuleNames.AddRange(
 			new string[]
 			{
@@ -22,13 +30,15 @@ public class PoseSearch : ModuleRules
 				"Engine",
 				"AnimationCore"
 			}
-			);
+		);
 
 		PrivateDependencyModuleNames.AddRange(
 			new string[]
             {
+				"Core",
+				"CoreUObject",
 				"Engine",
             }
-			);
+		);
 	}
 }
diff --git a/Engine/Plugins/Animation/PoseSearch/Source/Runtime/Private/AnimNode_PoseMatchingSequencePlayer.cpp b/Engine/Plugins/Animation/PoseSearch/Source/Runtime/Private/AnimNode_PoseMatchingSequencePlayer.cpp
deleted file mode 100644
diff --git a/Engine/Plugins/Animation/PoseSearch/Source/Runtime/Private/AnimNode_PoseSearchHistoryCollector.cpp b/Engine/Plugins/Animation/PoseSearch/Source/Runtime/Private/AnimNode_PoseSearchHistoryCollector.cpp
--- a/Engine/Plugins/Animation/PoseSearch/Source/Runtime/Private/AnimNode_PoseSearchHistoryCollector.cpp
+++ b/Engine/Plugins/Animation/PoseSearch/Source/Runtime/Private/AnimNode_PoseSearchHistoryCollector.cpp
@@ -14,7 +14,7 @@ void FAnimNode_PoseSearchHistoryCollector::Initialize_AnyThread(const FAnimation
 
 	Super::Initialize_AnyThread(Context);
 
-	// @@ need to do this once based on descendant node's search schema, not every node init
+	// @@ need to do this once based on descendant node's (or input param?) search schema, not every node init
 	PoseHistory.Init(32, 1.0f);
 
 	FScopedAnimNodeTracker ScopedNodeTracker = Context.TrackAncestor(this);
@@ -46,4 +46,13 @@ void FAnimNode_PoseSearchHistoryCollector::Update_AnyThread(const FAnimationUpda
 	EvalDeltaTime += Context.GetDeltaTime();
 }
 
+TArrayView<const float> FAnimNode_PoseSearchHistoryCollector::BuildQuery(const UPoseSearchSchema* Schema)
+{
+	check(Schema);
+
+	Query.SetNumZeroed(Schema->Layout.NumFloats);
+	UE::PoseSearch::BuildQuery(Schema, &PoseHistory, Query);
+	return Query;
+}
+
 #undef LOCTEXT_NAMESPACE
\ No newline at end of file
diff --git a/Engine/Plugins/Animation/PoseSearch/Source/Runtime/Private/PoseSearch.cpp b/Engine/Plugins/Animation/PoseSearch/Source/Runtime/Private/PoseSearch.cpp
--- a/Engine/Plugins/Animation/PoseSearch/Source/Runtime/Private/PoseSearch.cpp
+++ b/Engine/Plugins/Animation/PoseSearch/Source/Runtime/Private/PoseSearch.cpp
@@ -2,12 +2,99 @@
 
 #include "PoseSearch/PoseSearch.h"
 
+#include "Async/ParallelFor.h"
+#include "Features/IModularFeatures.h"
+#include "DrawDebugHelpers.h"
+#include "Animation/AnimPoseSearchProvider.h"
 #include "Animation/AnimCurveTypes.h"
 #include "Animation/AnimInstance.h"
+#include "Animation/AnimMetaData.h"
+#include "Animation/AnimSequence.h"
 #include "Animation/AnimSequenceBase.h"
+#include "Animation/AnimationPoseData.h"
 #include "AnimationRuntime.h"
 #include "BonePose.h"
-#include "DrawDebugHelpers.h"
+#include "PoseSearch/AnimNode_PoseSearchHistoryCollector.h"
+
+
+namespace UE { namespace PoseSearch {
+
+//////////////////////////////////////////////////////////////////////////
+// FFeatureTypeTraits
+
+struct FFeatureTypeTraits
+{
+	EPoseSearchFeatureType Type = EPoseSearchFeatureType::Invalid;
+	uint32 NumFloats = 0;
+};
+
+// Could upgrade to class objects in the future with value reader/writer functions
+static constexpr FFeatureTypeTraits FeatureTypeTraits[] =
+{
+	{ EPoseSearchFeatureType::Position, 3 },
+	{ EPoseSearchFeatureType::Rotation, 6 },
+	{ EPoseSearchFeatureType::LinearVelocity, 3 },
+	{ EPoseSearchFeatureType::AngularVelocity, 3 },
+};
+
+FFeatureTypeTraits GetFeatureTypeTraits(EPoseSearchFeatureType Type)
+{
+	// Could allow external registration to a TSet of traits in the future
+	// For now just use a simple local array
+	for (const FFeatureTypeTraits& Traits : FeatureTypeTraits)
+	{
+		if (Traits.Type == Type)
+		{
+			return Traits;
+		}
+	}
+
+	return FFeatureTypeTraits();
+}
+
+}} // namespace UE::PoseSearch
+
+
+//////////////////////////////////////////////////////////////////////////
+// FPoseSearchFeatureDesc
+
+bool FPoseSearchFeatureDesc::operator==(const FPoseSearchFeatureDesc& Other) const
+{
+	return
+		(SchemaBoneIdx == Other.SchemaBoneIdx) &&
+		(SubsampleIdx == Other.SubsampleIdx) &&
+		(Type == Other.Type) &&
+		(Domain == Other.Domain);
+}
+
+
+//////////////////////////////////////////////////////////////////////////
+// FPoseSearchFeatureVectorLayout
+
+void FPoseSearchFeatureVectorLayout::Init()
+{
+	uint32 FloatCount = 0;
+
+	for (FPoseSearchFeatureDesc& Element : Features)
+	{
+		Element.ValueOffset = FloatCount;
+		FloatCount += UE::PoseSearch::GetFeatureTypeTraits(Element.Type).NumFloats;
+	}
+
+	NumFloats = FloatCount;
+}
+
+void FPoseSearchFeatureVectorLayout::Reset()
+{
+	Features.Reset();
+	NumFloats = 0;
+}
+
+bool FPoseSearchFeatureVectorLayout::IsValid() const
+{
+	return NumFloats != 0.0f;
+}
+
 
 //////////////////////////////////////////////////////////////////////////
 // UPoseSearchSchema
@@ -17,7 +104,7 @@ void UPoseSearchSchema::PreSave(const class ITargetPlatform* TargetPlatform)
 	// Initialize references to obtain bone indices
 	for (FBoneReference& BoneRef : Bones)
 	{
-		BoneRef.Initialize(GetSkeleton());
+		BoneRef.Initialize(Skeleton);
 	}
 
 	// Fill out bone index array and sort by bone index
@@ -30,47 +117,471 @@ void UPoseSearchSchema::PreSave(const class ITargetPlatform* TargetPlatform)
 
 	// Build separate index array with parent indices guaranteed to be present
 	BoneIndicesWithParents = BoneIndices;
-	FAnimationRuntime::EnsureParentsPresent(BoneIndicesWithParents, GetSkeleton()->GetReferenceSkeleton());
+	if (Skeleton)
+	{
+		FAnimationRuntime::EnsureParentsPresent(BoneIndicesWithParents, Skeleton->GetReferenceSkeleton());
+	}
+
+	BoneIndices = BoneIndicesWithParents;
+
+	PoseSampleOffsets.Sort(TLess<>());
+	TrajectorySampleOffsets.Sort(TLess<>());
+	TrajectoryDistanceOffsets.Sort(TLess<>());
+
+	GenerateLayout();
+
+	Super::PreSave(TargetPlatform);
+}
+
+bool UPoseSearchSchema::IsValid() const
+{
+	return (Skeleton != nullptr) && Layout.IsValid();
+}
+
+void UPoseSearchSchema::GenerateLayout()
+{
+	Layout.Reset();
+
+	for (int32 TrajectoryTimeSubsampleIdx = 0; TrajectoryTimeSubsampleIdx != TrajectorySampleOffsets.Num(); ++TrajectoryTimeSubsampleIdx)
+	{
+		FPoseSearchFeatureDesc Element;
+		Element.SchemaBoneIdx = FPoseSearchFeatureDesc::TrajectoryBoneIndex;
+		Element.SubsampleIdx = TrajectoryTimeSubsampleIdx;
+		Element.Domain = EPoseSearchFeatureDomain::Time;
+
+		if (bUseTrajectoryPositions)
+		{
+			Element.Type = EPoseSearchFeatureType::Position;
+			Layout.Features.Add(Element);
+		}
+
+		if (bUseTrajectoryVelocities)
+		{
+			Element.Type = EPoseSearchFeatureType::LinearVelocity;
+			Layout.Features.Add(Element);
+		}
+	}
+
+ 	for (int32 TrajectoryDistSubsampleIdx = 0; TrajectoryDistSubsampleIdx != TrajectoryDistanceOffsets.Num(); ++TrajectoryDistSubsampleIdx)
+ 	{
+ 		FPoseSearchFeatureDesc Element;
+ 		Element.SchemaBoneIdx = FPoseSearchFeatureDesc::TrajectoryBoneIndex;
+ 		Element.SubsampleIdx = TrajectoryDistSubsampleIdx;
+ 		Element.Domain = EPoseSearchFeatureDomain::Distance;
+
+		if (bUseTrajectoryPositions)
+		{
+			Element.Type = EPoseSearchFeatureType::Position;
+			Layout.Features.Add(Element);
+		}
 
-	// Sort fragment offsets by largest offset first since larger offsets are closer to the beginning of the sample array
-	FragmentSampleOffsets.Sort(TGreater<>());
+		if (bUseTrajectoryVelocities)
+		{
+			Element.Type = EPoseSearchFeatureType::LinearVelocity;
+			Layout.Features.Add(Element);
+		}
+ 	}
 
-	// Ensure we have at least one offset at zero which corresponds with instantaneous sample matching
-	if (FragmentSampleOffsets.Num() == 0)
+	for (int32 PoseSubsampleIdx = 0; PoseSubsampleIdx != PoseSampleOffsets.Num(); ++PoseSubsampleIdx)
 	{
-		FragmentSampleOffsets.Add(0);
+		FPoseSearchFeatureDesc Element;
+		Element.SubsampleIdx = PoseSubsampleIdx;
+		Element.Domain = EPoseSearchFeatureDomain::Time;
+
+		for (int32 SchemaBoneIdx = 0; SchemaBoneIdx != BoneIndices.Num(); ++SchemaBoneIdx)
+		{
+			Element.SchemaBoneIdx = SchemaBoneIdx;
+			if (bUseBonePositions)
+			{
+				Element.Type = EPoseSearchFeatureType::Position;
+				Layout.Features.Add(Element);
+			}
+
+			if (bUseBoneVelocities)
+			{
+				Element.Type = EPoseSearchFeatureType::LinearVelocity;
+				Layout.Features.Add(Element);
+			}
+		}
 	}
 
-	FloatsPerPose = BoneIndices.Num() * 3;
+	Layout.Init();
+}
+
+//////////////////////////////////////////////////////////////////////////
+// FPoseSearchIndex
+
+bool FPoseSearchIndex::IsValid() const
+{
+	bool bSchemaValid = Schema && Schema->IsValid();
+	bool bSearchIndexValid = bSchemaValid && (NumPoses * Schema->Layout.NumFloats == Values.Num());
+
+	return bSearchIndexValid;
+}
+
+
+//////////////////////////////////////////////////////////////////////////
+// UPoseSearchSequenceMetaData
+
+void UPoseSearchSequenceMetaData::PreSave(const class ITargetPlatform* TargetPlatform)
+{
+	if (Schema && SamplingRange.Size() > 0.0f)
+	{
+		UObject* Outer = GetOuter();
+		if (UAnimSequence* Sequence = Cast<UAnimSequence>(Outer))
+		{
+			UE::PoseSearch::BuildIndex(Sequence, this);
+		}
+	}
 
 	Super::PreSave(TargetPlatform);
 }
 
+bool UPoseSearchSequenceMetaData::IsValidForIndexing() const
+{
+	return Schema && Schema->IsValid() && (SamplingRange.Size() > 0.0f);
+}
+
+bool UPoseSearchSequenceMetaData::IsValidForSearch() const
+{
+	return IsValidForIndexing() && SearchIndex.IsValid();
+}
+
+
+//////////////////////////////////////////////////////////////////////////
+// UPoseSearchDatabase
+
+const FPoseSearchDatabaseSequence* UPoseSearchDatabase::FindSequenceByPoseIdx(int32 PoseIdx) const
+{
+	for (const FPoseSearchDatabaseSequence& Sequence : Sequences)
+	{
+		if (PoseIdx >= Sequence.FirstPoseIdx && PoseIdx < Sequence.FirstPoseIdx + Sequence.NumPoses)
+		{
+			return &Sequence;
+		}
+	}
+
+	return nullptr;
+}
+
+bool UPoseSearchDatabase::IsValidForIndexing() const
+{
+	return Schema && Schema->IsValid() && !Sequences.IsEmpty();
+}
+
+bool UPoseSearchDatabase::IsValidForSearch() const
+{
+	return IsValidForIndexing() && SearchIndex.IsValid();
+}
+
+void UPoseSearchDatabase::PreSave(const class ITargetPlatform* TargetPlatform)
+{
+	if (IsValidForIndexing())
+	{
+		UE::PoseSearch::BuildIndex(this);
+	}
+
+	Super::PreSave(TargetPlatform);
+}
+
+
+namespace UE { namespace PoseSearch {
+
+//////////////////////////////////////////////////////////////////////////
+// FFeatureVectorBuilder
+
+void FFeatureVectorBuilder::Init(const FPoseSearchFeatureVectorLayout* InLayout, TArrayView<float> Buffer)
+{
+	check(InLayout);
+	check(Buffer.Num() == InLayout->NumFloats);
+	Layout = InLayout;
+	Values = Buffer;
+	ResetFeatures();
+}
+
+void FFeatureVectorBuilder::ResetFeatures()
+{
+	NumFeaturesAdded = 0;
+	FeaturesAdded.Init(false, Layout->Features.Num());
+}
+
+void FFeatureVectorBuilder::SetTransform(FPoseSearchFeatureDesc Element, const FTransform& Transform)
+{
+	SetPosition(Element, Transform.GetTranslation());
+	SetRotation(Element, Transform.GetRotation());
+}
+
+void FFeatureVectorBuilder::SetTransformDerivative(FPoseSearchFeatureDesc Element, const FTransform& Transform, const FTransform& PrevTransform, float DeltaTime)
+{
+	SetLinearVelocity(Element, Transform, PrevTransform, DeltaTime);
+	SetAngularVelocity(Element, Transform, PrevTransform, DeltaTime);
+}
+
+void FFeatureVectorBuilder::SetPosition(FPoseSearchFeatureDesc Element, const FVector& Position)
+{
+	Element.Type = EPoseSearchFeatureType::Position;
+	SetVector(Element, Position);
+}
+
+void FFeatureVectorBuilder::SetRotation(FPoseSearchFeatureDesc Element, const FQuat& Rotation)
+{
+	Element.Type = EPoseSearchFeatureType::Rotation;
+	int32 ElementIndex = Layout->Features.Find(Element);
+	if (ElementIndex >= 0)
+	{
+		FVector X = Rotation.GetAxisX();
+		FVector Y = Rotation.GetAxisY();
+
+		const FPoseSearchFeatureDesc& FoundElement = Layout->Features[ElementIndex];
+
+		Values[FoundElement.ValueOffset + 0] = X.X;
+		Values[FoundElement.ValueOffset + 1] = X.Y;
+		Values[FoundElement.ValueOffset + 2] = X.Z;
+		Values[FoundElement.ValueOffset + 3] = Y.X;
+		Values[FoundElement.ValueOffset + 4] = Y.Y;
+		Values[FoundElement.ValueOffset + 5] = Y.Z;
+
+		if (!FeaturesAdded[ElementIndex])
+		{
+			FeaturesAdded[ElementIndex] = true;
+			++NumFeaturesAdded;
+		}
+	}
+}
+
+void FFeatureVectorBuilder::SetLinearVelocity(FPoseSearchFeatureDesc Element, const FTransform& Transform, const FTransform& PrevTransform, float DeltaTime)
+{
+	Element.Type = EPoseSearchFeatureType::LinearVelocity;
+	FVector LinearVelocity = (Transform.GetTranslation() - PrevTransform.GetTranslation()) / DeltaTime;
+	SetVector(Element, LinearVelocity);
+}
+
+void FFeatureVectorBuilder::SetAngularVelocity(FPoseSearchFeatureDesc Element, const FTransform& Transform, const FTransform& PrevTransform, float DeltaTime)
+{
+	Element.Type = EPoseSearchFeatureType::AngularVelocity;
+	int32 ElementIndex = Layout->Features.Find(Element);
+	if (ElementIndex >= 0)
+	{
+		FQuat Q0 = PrevTransform.GetRotation();
+		FQuat Q1 = Transform.GetRotation();
+		Q1.EnforceShortestArcWith(Q0);
+
+		// Given angular velocity vector w, quaternion differentiation can be represented as
+		//   dq/dt = (w * q)/2
+		// Solve for w
+		//   w = 2 * dq/dt * q^-1
+		// And let dq/dt be expressed as the finite difference
+		//   dq/dt = (q(t+h) - q(t)) / h
+		FQuat DQDt = (Q1 - Q0) / DeltaTime;
+		FQuat QInv = Q0.Inverse();
+		FQuat W = (DQDt * QInv) * 2.0f;
+
+		FVector AngularVelocity(W.X, W.Y, W.Z);
+
+		const FPoseSearchFeatureDesc& FoundElement = Layout->Features[ElementIndex];
+
+		Values[FoundElement.ValueOffset + 0] = AngularVelocity[0];
+		Values[FoundElement.ValueOffset + 1] = AngularVelocity[1];
+		Values[FoundElement.ValueOffset + 2] = AngularVelocity[2];
+
+		if (!FeaturesAdded[ElementIndex])
+		{
+			FeaturesAdded[ElementIndex] = true;
+			++NumFeaturesAdded;
+		}
+	}
+}
+
+void FFeatureVectorBuilder::SetVector(FPoseSearchFeatureDesc Element, const FVector& Vector)
+{
+	int32 ElementIndex = Layout->Features.Find(Element);
+	if (ElementIndex >= 0)
+	{
+		const FPoseSearchFeatureDesc& FoundElement = Layout->Features[ElementIndex];
+
+		Values[FoundElement.ValueOffset + 0] = Vector[0];
+		Values[FoundElement.ValueOffset + 1] = Vector[1];
+		Values[FoundElement.ValueOffset + 2] = Vector[2];
+
+		if (!FeaturesAdded[ElementIndex])
+		{
+			FeaturesAdded[ElementIndex] = true;
+			++NumFeaturesAdded;
+		}
+	}
+}
+
+bool FFeatureVectorBuilder::SetPoseFeatures(const UPoseSearchSchema* Schema, FPoseHistory* History)
+{
+	check(Schema && Schema->IsValid());
+	check(History);
+
+	FPoseSearchFeatureDesc Feature;
+	Feature.Domain = EPoseSearchFeatureDomain::Time;
+
+	for (int32 SubsampleIdx = 0; SubsampleIdx != Schema->PoseSampleOffsets.Num(); ++SubsampleIdx)
+	{
+		Feature.SubsampleIdx = SubsampleIdx;
+
+		int32 Offset = Schema->PoseSampleOffsets[SubsampleIdx];
+		float TimeDelta = -Offset * (1.0f /  Schema->SampleRate);
+
+		if (!History->Sample(TimeDelta, Schema->Skeleton->GetReferenceSkeleton(), Schema->BoneIndicesWithParents))
+		{
+			return false;
+		}
+
+		TArrayView<const FTransform> ComponentPose = History->GetComponentPoseSample();
+		TArrayView<const FTransform> ComponentPrevPose = History->GetPrevComponentPoseSample();
+		for (int32 SchemaBoneIdx = 0; SchemaBoneIdx != Schema->BoneIndices.Num(); ++SchemaBoneIdx)
+		{
+			Feature.SchemaBoneIdx = SchemaBoneIdx;
+
+			int32 SkeletonBoneIndex = Schema->BoneIndices[SchemaBoneIdx]; // @@@ need to resolve bone indices at runtime, hold in FPoseHistory?
+			const FTransform& Transform = ComponentPose[SkeletonBoneIndex];
+			const FTransform& PrevTransform = ComponentPrevPose[SkeletonBoneIndex];
+			SetTransform(Feature, Transform);
+			SetTransformDerivative(Feature, Transform, PrevTransform, History->GetSampleInterval());
+		}
+	}
+
+	return true;
+}
+
+bool FFeatureVectorBuilder::IsComplete() const
+{
+	return NumFeaturesAdded == Layout->Features.Num();
+}
+
 
 //////////////////////////////////////////////////////////////////////////
-// UPoseSearchPoseHistory
+// FPoseSearchFeatureVectorReader
+
+void FFeatureVectorReader::Init(const FPoseSearchFeatureVectorLayout* InLayout)
+{
+	check(InLayout);
+	Layout = InLayout;
+}
+
+void FFeatureVectorReader::SetValues(TArrayView<const float> InValues)
+{
+	check(Layout);
+	check(Layout->NumFloats == InValues.Num());
+	Values = InValues;
+}
+
+bool FFeatureVectorReader::IsValid() const
+{
+	return Layout && (Layout->NumFloats == Values.Num());
+}
+
+bool FFeatureVectorReader::GetTransform(FPoseSearchFeatureDesc Element, FTransform* OutTransform) const
+{
+	FVector Position;
+	bool bResult = GetPosition(Element, &Position);
+
+	FQuat Rotation;
+	bResult |= GetRotation(Element, &Rotation);
+
+	OutTransform->SetComponents(Rotation, Position, FVector::OneVector);
+	return bResult;
+}
+
+bool FFeatureVectorReader::GetPosition(FPoseSearchFeatureDesc Element, FVector* OutPosition) const
+{
+	Element.Type = EPoseSearchFeatureType::Position;
+	return GetVector(Element, OutPosition);
+}
+
+bool FFeatureVectorReader::GetRotation(FPoseSearchFeatureDesc Element, FQuat* OutRotation) const
+{
+	Element.Type = EPoseSearchFeatureType::Rotation;
+	int32 ElementIndex = IsValid() ? Layout->Features.Find(Element) : -1;
+	if (ElementIndex >= 0)
+	{
+		const FPoseSearchFeatureDesc& FoundElement = Layout->Features[ElementIndex];
+
+		FVector X;
+		FVector Y;
+
+		X.X = Values[FoundElement.ValueOffset + 0];
+		X.Y = Values[FoundElement.ValueOffset + 1];
+		X.Z = Values[FoundElement.ValueOffset + 2];
+		Y.X = Values[FoundElement.ValueOffset + 3];
+		Y.Y = Values[FoundElement.ValueOffset + 4];
+		Y.Z = Values[FoundElement.ValueOffset + 5];
+
+		FVector Z = FVector::CrossProduct(X, Y);
+
+		FMatrix M(FMatrix::Identity);
+		M.SetColumn(0, X);
+		M.SetColumn(1, Y);
+		M.SetColumn(2, Z);
+
+		*OutRotation = FQuat(M);
+		return true;
+	}
+
+	*OutRotation = FQuat::Identity;
+	return false;
+}
+
+bool FFeatureVectorReader::GetLinearVelocity(FPoseSearchFeatureDesc Element, FVector* OutLinearVelocity) const
+{
+	Element.Type = EPoseSearchFeatureType::LinearVelocity;
+	return GetVector(Element, OutLinearVelocity);
+}
+
+bool FFeatureVectorReader::GetAngularVelocity(FPoseSearchFeatureDesc Element, FVector* OutAngularVelocity) const
+{
+	Element.Type = EPoseSearchFeatureType::AngularVelocity;
+	return GetVector(Element, OutAngularVelocity);
+}
+
+bool FFeatureVectorReader::GetVector(FPoseSearchFeatureDesc Element, FVector* OutVector) const
+{
+	int32 ElementIndex = IsValid() ? Layout->Features.Find(Element) : -1;
+	if (ElementIndex >= 0)
+	{
+		const FPoseSearchFeatureDesc& FoundElement = Layout->Features[ElementIndex];
+
+		FVector V;
+		V.X = Values[FoundElement.ValueOffset + 0];
+		V.Y = Values[FoundElement.ValueOffset + 1];
+		V.Z = Values[FoundElement.ValueOffset + 2];
+		*OutVector = V;
+		return true;
+	}
+
+	*OutVector = FVector::ZeroVector;
+	return false;
+}
+
+
+//////////////////////////////////////////////////////////////////////////
+// FPoseHistory
 
 /**
 * Fills skeleton transforms with evaluated compact pose transforms.
 * Bones that weren't evaluated are filled with the bone's reference pose.
 */
-static void CopyCompactToSkeletonPose (const FCompactPose& Pose, TArray<FTransform>& OutTransforms)
+static void CopyCompactToSkeletonPose(const FCompactPose& Pose, TArray<FTransform>& OutLocalTransforms)
 {
 	const FBoneContainer& BoneContainer = Pose.GetBoneContainer();
 	const FReferenceSkeleton& RefSkeleton = BoneContainer.GetReferenceSkeleton();
 	TArrayView<const FTransform> RefSkeletonTransforms = MakeArrayView(RefSkeleton.GetRefBonePose());
 
 	const int32 NumSkeletonBones = BoneContainer.GetNumBones();
-	OutTransforms.SetNum(NumSkeletonBones);
+	OutLocalTransforms.SetNum(NumSkeletonBones);
 
 	for (auto SkeletonBoneIdx = FSkeletonPoseBoneIndex(0); SkeletonBoneIdx != NumSkeletonBones; ++SkeletonBoneIdx)
 	{
 		FCompactPoseBoneIndex CompactBoneIdx = BoneContainer.GetCompactPoseIndexFromSkeletonIndex(SkeletonBoneIdx.GetInt());
-		OutTransforms[SkeletonBoneIdx.GetInt()] = CompactBoneIdx.IsValid() ? Pose[CompactBoneIdx] : RefSkeletonTransforms[SkeletonBoneIdx.GetInt()];
+		OutLocalTransforms[SkeletonBoneIdx.GetInt()] = CompactBoneIdx.IsValid() ? Pose[CompactBoneIdx] : RefSkeletonTransforms[SkeletonBoneIdx.GetInt()];
 	}
 }
 
-void FPoseSearchPoseHistory::Init(int32 InNumPoses, float InTimeHorizon)
+void FPoseHistory::Init(int32 InNumPoses, float InTimeHorizon)
 {
 	int32 Capacity = FMath::RoundUpToPowerOfTwo(InNumPoses);
 
@@ -85,7 +596,7 @@ void FPoseSearchPoseHistory::Init(int32 InNumPoses, float InTimeHorizon)
 	TimeHorizon = InTimeHorizon;
 }
 
-void FPoseSearchPoseHistory::Init(const FPoseSearchPoseHistory& History)
+void FPoseHistory::Init(const FPoseHistory& History)
 {
 	Poses = History.Poses;
 	Knots = History.Knots;
@@ -93,54 +604,86 @@ void FPoseSearchPoseHistory::Init(const FPoseSearchPoseHistory& History)
 	TimeHorizon = History.TimeHorizon;
 }
 
-bool FPoseSearchPoseHistory::Sample(float Time, const TArray<FBoneIndexType>& RequiredBones, TArray<FTransform>& OutPose) const
+bool FPoseHistory::SampleLocalPose(float SecondsAgo, const FReferenceSkeleton& RefSkeleton, const TArray<FBoneIndexType>& RequiredBones, TArray<FTransform>& LocalPose)
 {
-	// Find the upper bound knot
-	uint32 UpperBoundIndex = MAX_uint32;
-	int32 UpperBoundOffset = 1;
-	for (; UpperBoundOffset < (int32)Queue.Num(); ++UpperBoundOffset)
+	// Find the lower bound knot
+	uint32 NextIndex = MAX_uint32;
+	int32 NextOffset = 1;
+	for (; NextOffset < (int32)Queue.Num(); ++NextOffset)
 	{
-		uint32 TestIndex = Queue.GetOffsetFromBack(UpperBoundOffset);
-		if (Knots[TestIndex] >= Time)
+		uint32 TestIndex = Queue.GetOffsetFromBack(NextOffset);
+		if (Knots[TestIndex] >= SecondsAgo)
 		{
-			UpperBoundIndex = TestIndex;
+			NextIndex = TestIndex;
 			break;
 		}
 	}
 
-	if (UpperBoundIndex == MAX_uint32)
+	if (NextIndex == MAX_uint32)
 	{
 		return false;
 	}
 
-	// The lower bound knot is adjacent
-	int32 LowerBoundOffset = UpperBoundOffset - 1;
-	uint32 LowerBoundIndex = Queue.GetOffsetFromBack(LowerBoundOffset);
+	// Get the previous knot
+	int32 PrevOffset = NextOffset - 1;
+	uint32 PrevIndex = Queue.GetOffsetFromBack(PrevOffset);
 
-	// Compute alpha between upper and lower bound knots
+	// Compute alpha between previous and next knots
 	float Alpha = FMath::GetMappedRangeValueUnclamped(
-		FVector2D(Knots[LowerBoundIndex], Knots[UpperBoundIndex]),
+		FVector2D(Knots[PrevIndex], Knots[NextIndex]),
 		FVector2D(0.0f, 1.0f),
-		Time);
+		SecondsAgo);
 
-	// Lerp between upper and lower bound poses by alpha to produce output pose at requested sample time
-	OutPose = Poses[LowerBoundIndex].LocalTransforms;
+	TArray<FTransform>& PrevPose = Poses[PrevIndex].LocalTransforms;
+	TArray<FTransform>& NextPose = Poses[NextIndex].LocalTransforms;
+
+	// We may not have accumulated enough poses yet
+	if (PrevPose.Num() != NextPose.Num())
+	{
+		return false;
+	}
+
+	if (RequiredBones.Num() > PrevPose.Num())
+	{
+		return false;
+	}
+
+	// Lerp between poses by alpha to produce output local pose at requested sample time
+	LocalPose = PrevPose;
 	FAnimationRuntime::LerpBoneTransforms(
-		OutPose,
-		Poses[UpperBoundIndex].LocalTransforms,
+		LocalPose,
+		NextPose,
 		Alpha,
 		RequiredBones);
 
 	return true;
 }
 
-void FPoseSearchPoseHistory::Update(float TimeDelta, const FCompactPose& Pose)
+bool FPoseHistory::Sample(float SecondsAgo, const FReferenceSkeleton& RefSkeleton, const TArray<FBoneIndexType>& RequiredBones)
+{
+	// Compute local space pose at requested time
+	bool bSampled = SampleLocalPose(SecondsAgo, RefSkeleton, RequiredBones, SampledLocalPose);
+
+	// Compute local space pose one sample interval in the past
+	bSampled = bSampled && SampleLocalPose(SecondsAgo + GetSampleInterval(), RefSkeleton, RequiredBones, SampledPrevLocalPose);
+
+	// Convert local to component space
+	if (bSampled)
+	{
+		FAnimationRuntime::FillUpComponentSpaceTransforms(RefSkeleton, SampledLocalPose, SampledComponentPose);
+		FAnimationRuntime::FillUpComponentSpaceTransforms(RefSkeleton, SampledPrevLocalPose, SampledPrevComponentPose);
+	}
+
+	return bSampled;
+}
+
+void FPoseHistory::Update(float SecondsElapsed, const FCompactPose& Pose)
 {
 	// Age our elapsed times
 	for (int32 Offset = 0; Offset != (int32)Queue.Num(); ++Offset)
 	{
 		uint32 Index = Queue.GetOffsetFromFront(Offset);
-		Knots[Index] += TimeDelta;
+		Knots[Index] += SecondsElapsed;
 	}
 
 	if (!Queue.IsFull())
@@ -176,64 +719,391 @@ void FPoseSearchPoseHistory::Update(float TimeDelta, const FCompactPose& Pose)
 	CopyCompactToSkeletonPose(Pose, Poses[Newest].LocalTransforms);
 }
 
-float FPoseSearchPoseHistory::GetSampleInterval() const
+float FPoseHistory::GetSampleInterval() const
 {
 	return TimeHorizon / Queue.GetCapacity();
 }
 
 
 //////////////////////////////////////////////////////////////////////////
-// PoseSearch API
+// FIndexer
+
+class FSequenceIndexer
+{
+public:
+	struct Result
+	{
+		int32 NumIndexedPoses;
+		TArrayView<const float> Values;
+	};
+	Result Process(const UPoseSearchSchema* Schema, const UAnimSequence* Sequence, FFloatInterval SamplingRange);
+	Result GetResult() const;
+
+private:
+	void SampleBegin(int32 SampleIdx);
+	void SampleEnd();
+	void ExtractPoses(const UAnimSequence* Sequence);
+	void ExtractRootMotion(const UAnimSequence* Sequence);
+	void AddPoseFeatures(int32 SampleIdx);
+	void AddTrajectoryTimeFeatures(int32 SampleIdx);
+	void AddTrajectoryDistanceFeatures(int32 SampleIdx);
+
+	struct FSampleContext
+	{
+		TArray<FTransform> ComponentSpacePose;		// Indexed by SampleIdx * NumBones + SchemaBoneIdx
+		TArray<FTransform> LocalRootMotion;			// Indexed by SampleIdx
+		TArray<FTransform> AccumulatedRootMotion;	// Indexed by SampleIdx
+		TArray<float> AccumulatedRootDistance;		// Indexed by SampleIdx
+
+		int32 TotalSamples = 0;
+		int32 FirstIndexedSample = 0;
+		int32 LastIndexedSample = 0;
+		int32 NumIndexedSamples = 0;
+		int32 NumBones = 0;
+
+		void Reset();
+		void Reserve();
+	};
+
+	const UPoseSearchSchema* Schema = nullptr;
+
+	TArray<float> Values;
+	
+	FFeatureVectorBuilder Builder;
+	FSampleContext Context;
+
+	float DeltaTime = 0.0f; //@@@ move to Schema?
+};
+
+void FSequenceIndexer::FSampleContext::Reset()
+{
+	TotalSamples = 0;
+	FirstIndexedSample = 0;
+	LastIndexedSample = 0;
+	NumIndexedSamples = 0;
+	NumBones = 0;
+
+	ComponentSpacePose.Reset(0);
+	LocalRootMotion.Reset(0);
+	AccumulatedRootMotion.Reset(0);
+	AccumulatedRootDistance.Reset(0);
+}
+
+void FSequenceIndexer::FSampleContext::Reserve()
+{
+	ComponentSpacePose.Reserve(NumBones * TotalSamples);
+	LocalRootMotion.Reserve(TotalSamples);
+	AccumulatedRootMotion.Reserve(TotalSamples);
+	AccumulatedRootDistance.Reserve(TotalSamples);
+}
+
+FSequenceIndexer::Result FSequenceIndexer::Process(const UPoseSearchSchema* InSchema, const UAnimSequence* Sequence, FFloatInterval SamplingRange)
+{
+	check(InSchema);
+	check(Sequence);
+
+	Schema = InSchema;
+
+	USkeleton* Skeleton = Sequence->GetSkeleton();
+	check(Skeleton);
+	check(Skeleton->IsCompatible(Schema->Skeleton));
+
+	const float BeginTime = SamplingRange.Min;
+	const float EndTime = FMath::Min(Sequence->GetPlayLength(), SamplingRange.Max);
+
+	DeltaTime = 1.0f / Schema->SampleRate;
+
+	Context.Reset();
+	Context.NumBones = Schema->BoneIndices.Num();
+	Context.TotalSamples = FMath::FloorToInt(Sequence->GetPlayLength() * Schema->SampleRate);
+	Context.FirstIndexedSample = FMath::FloorToInt(BeginTime * Schema->SampleRate);
+	Context.LastIndexedSample = FMath::Max(0, FMath::FloorToInt(EndTime * Schema->SampleRate) - 1);
+	Context.NumIndexedSamples = Context.LastIndexedSample - Context.FirstIndexedSample + 1;
+	Context.Reserve();
+
+	Values.SetNumZeroed(Schema->Layout.NumFloats * Context.NumIndexedSamples);
+
+	ExtractPoses(Sequence);
+	ExtractRootMotion(Sequence);
+
+	for (int32 SampleIdx = Context.FirstIndexedSample; SampleIdx <= Context.LastIndexedSample; ++SampleIdx)
+	{
+		SampleBegin(SampleIdx);
+
+		AddPoseFeatures(SampleIdx);
+		AddTrajectoryTimeFeatures(SampleIdx);
+		AddTrajectoryDistanceFeatures(SampleIdx);
+
+		SampleEnd();
+	}
+
+	return GetResult();
+}
+
+FSequenceIndexer::Result FSequenceIndexer::GetResult() const
+{
+	Result Result;
+	Result.NumIndexedPoses = Context.NumIndexedSamples;
+	Result.Values = Values;
+	return Result;
+}
+
+void FSequenceIndexer::SampleBegin(int32 SampleIdx)
+{
+	int32 FirstValueIdx = (SampleIdx - Context.FirstIndexedSample) * Schema->Layout.NumFloats;
+	TArrayView<float> FeatureVectorValues = MakeArrayView(&Values[FirstValueIdx], Schema->Layout.NumFloats);
+	Builder.Init(&Schema->Layout, FeatureVectorValues);
+}
+
+void FSequenceIndexer::SampleEnd()
+{
+	check(Builder.IsComplete());
+}
+
+void FSequenceIndexer::ExtractPoses(const UAnimSequence* Sequence)
+{
+	USkeleton* Skeleton = Sequence->GetSkeleton();
+	FBoneContainer BoneContainer;
+	BoneContainer.InitializeTo(Schema->BoneIndicesWithParents, FCurveEvaluationOption(false), *Skeleton);
+
+	FCompactPose Pose;
+	Pose.SetBoneContainer(&BoneContainer);
+	FCSPose<FCompactPose> ComponentSpacePose;
+
+	FBlendedCurve UnusedCurve;
+	FStackCustomAttributes UnusedAttributes;
+
+	FAnimExtractContext ExtractionCtx;
+	// ExtractionCtx.PoseCurves is intentionally left empty
+	// ExtractionCtx.BonesRequired is unused by UAnimSequence::GetAnimationPose
+	ExtractionCtx.bExtractRootMotion = true;
+
+	FAnimationPoseData AnimPoseData(Pose, UnusedCurve, UnusedAttributes);
+	for (int32 SampleIdx = 0; SampleIdx != Context.TotalSamples; ++SampleIdx)
+	{
+		const float CurrentTime = SampleIdx * DeltaTime;
+
+		ExtractionCtx.CurrentTime = CurrentTime;
+		Sequence->GetAnimationPose(AnimPoseData, ExtractionCtx);
+		ComponentSpacePose.InitPose(Pose);
+
+		for (int32 BoneIndex : Schema->BoneIndices)
+		{
+			FCompactPoseBoneIndex CompactBoneIndex = BoneContainer.MakeCompactPoseIndex(FMeshPoseBoneIndex(BoneIndex));
+			const FTransform& Transform = ComponentSpacePose.GetComponentSpaceTransform(CompactBoneIndex);
+			Context.ComponentSpacePose.Add(Transform);
+		}
+	}
+}
+
+void FSequenceIndexer::ExtractRootMotion(const UAnimSequence* Sequence)
+{
+	double AccumulatedRootDistance = 0.0;
+	FTransform AccumulatedRootMotion = FTransform::Identity;
+	for (int32 SampleIdx = 0; SampleIdx != Context.TotalSamples; ++SampleIdx)
+	{
+		const float CurrentTime = SampleIdx * DeltaTime;
 
-void PoseSearchDrawPose(const FPoseSearchDebugDrawParams& DrawParams, const FColor& Color, TArrayView<const float> Pose, float LifeTimeDelta, uint8 DepthPriorityDelta = 0)
+		FTransform LocalRootMotion = Sequence->ExtractRootMotion(CurrentTime, DeltaTime, false /*!allowLooping*/);
+		Context.LocalRootMotion.Add(LocalRootMotion);
+
+		AccumulatedRootMotion = LocalRootMotion * AccumulatedRootMotion;
+		AccumulatedRootDistance += LocalRootMotion.GetTranslation().Size();
+		Context.AccumulatedRootMotion.Add(AccumulatedRootMotion);
+		Context.AccumulatedRootDistance.Add((float)AccumulatedRootDistance);
+	}
+}
+
+void FSequenceIndexer::AddPoseFeatures(int32 SampleIdx)
 {
-	const float LifeTime = DrawParams.DefaultLifeTime + LifeTimeDelta;
-	const uint8 DepthPriority = ESceneDepthPriorityGroup::SDPG_Foreground + 1 + DepthPriorityDelta;
+	FPoseSearchFeatureDesc CurrentElement;
+	CurrentElement.Domain = EPoseSearchFeatureDomain::Time;
+
+	FTransform SampleSpaceOrigin = Context.AccumulatedRootMotion[SampleIdx];
+	
+	for (int32 SchemaSubsampleIdx = 0; SchemaSubsampleIdx != Schema->PoseSampleOffsets.Num(); ++SchemaSubsampleIdx)
+	{
+		CurrentElement.SubsampleIdx = SchemaSubsampleIdx;
+
+		const int32 SampleOffset = Schema->PoseSampleOffsets[SchemaSubsampleIdx];
+		const int32 SubsampleIdx = FMath::Clamp(SampleIdx + SampleOffset, 0, Context.AccumulatedRootMotion.Num() - 1);
+
+		FTransform SubsampleRoot = Context.AccumulatedRootMotion[SubsampleIdx];
+		SubsampleRoot.SetToRelativeTransform(SampleSpaceOrigin);
 
-	FVector PointPrev(
-		Pose[0],
-		Pose[1],
-		Pose[2]);
+		for (int32 SchemaBoneIndex = 0; SchemaBoneIndex != Context.NumBones; ++SchemaBoneIndex)
+		{
+			CurrentElement.SchemaBoneIdx = SchemaBoneIndex;
+
+			int32 BoneSampleIdx = Context.NumBones * (SampleIdx + SampleOffset) + SchemaBoneIndex;
+			int32 BonePrevSampleIdx = Context.NumBones * (SampleIdx - 1 + SampleOffset) + SchemaBoneIndex;
+			
+			// @@@Add extrapolation. Clamp for now
+			BoneSampleIdx = FMath::Clamp(BoneSampleIdx, 0, Context.ComponentSpacePose.Num() - 1);
+			BonePrevSampleIdx = FMath::Clamp(BonePrevSampleIdx, 0, Context.ComponentSpacePose.Num() - 1);
+
+			FTransform BoneInSampleSpace = Context.ComponentSpacePose[BoneSampleIdx] * SubsampleRoot;
+			FTransform BonePrevInSampleSpace = Context.ComponentSpacePose[BonePrevSampleIdx] * SubsampleRoot;
+
+			Builder.SetTransform(CurrentElement, BoneInSampleSpace);
+			Builder.SetTransformDerivative(CurrentElement, BoneInSampleSpace, BonePrevInSampleSpace, DeltaTime);
+		}
+	}
+}
 
-	PointPrev = DrawParams.ComponentTransform.TransformPosition(PointPrev);
+void FSequenceIndexer::AddTrajectoryTimeFeatures(int32 SampleIdx)
+{
+	FPoseSearchFeatureDesc CurrentElement;
+	CurrentElement.Domain = EPoseSearchFeatureDomain::Time;
+	CurrentElement.SchemaBoneIdx = FPoseSearchFeatureDesc::TrajectoryBoneIndex;
 
-	const int32 NumPoints = Pose.Num() / 3;
+	FTransform SampleSpaceOrigin = Context.AccumulatedRootMotion[SampleIdx];
 
-	for (int32 PointIdx = 1; PointIdx != NumPoints; ++PointIdx)
+	for (int32 SubsampleIdx = 0; SubsampleIdx != Schema->TrajectorySampleOffsets.Num(); ++SubsampleIdx)
 	{
-		FVector PointNext(
-			Pose[PointIdx * 3 + 0],
-			Pose[PointIdx * 3 + 1],
-			Pose[PointIdx * 3 + 2]);
+		CurrentElement.SubsampleIdx = SubsampleIdx;
 
-		PointNext = DrawParams.ComponentTransform.TransformPosition(PointNext);
+		int32 RootMotionIdx = SampleIdx + Schema->TrajectorySampleOffsets[SubsampleIdx];
+		int32 RootMotionPrevIdx = RootMotionIdx - 1;
 
-		DrawDebugPoint(DrawParams.World, PointNext, 1.0f, Color, false, LifeTime, DepthPriority);
+		// @@@ Add extrapolation. Clamp for now
+		RootMotionIdx = FMath::Clamp(RootMotionIdx, 0, Context.AccumulatedRootMotion.Num() - 1);
+		RootMotionPrevIdx = FMath::Clamp(RootMotionPrevIdx, 0, Context.AccumulatedRootMotion.Num() - 1);
 
-		bool bIsChildOfPrev = DrawParams.Schema->GetSkeleton()->GetReferenceSkeleton().BoneIsChildOf(
-			DrawParams.Schema->BoneIndices[PointIdx],
-			DrawParams.Schema->BoneIndices[PointIdx-1]);
+		FTransform SubsampleRoot = Context.AccumulatedRootMotion[RootMotionIdx];
+		SubsampleRoot.SetToRelativeTransform(SampleSpaceOrigin);
+
+		FTransform SubsamplePrevRoot = Context.AccumulatedRootMotion[RootMotionPrevIdx];
+		SubsamplePrevRoot.SetToRelativeTransform(SampleSpaceOrigin);
+
+		Builder.SetTransform(CurrentElement, SubsampleRoot);
+		Builder.SetTransformDerivative(CurrentElement, SubsampleRoot, SubsamplePrevRoot, DeltaTime);
+	}
+}
+
+void FSequenceIndexer::AddTrajectoryDistanceFeatures(int32 SampleIdx)
+{
+	FPoseSearchFeatureDesc CurrentElement;
+	CurrentElement.Domain = EPoseSearchFeatureDomain::Distance;
+	CurrentElement.SchemaBoneIdx = FPoseSearchFeatureDesc::TrajectoryBoneIndex;
+
+	TArrayView<const float> AccumulatedRootDistances = Context.AccumulatedRootDistance;
+
+	FTransform SampleSpaceOrigin = Context.AccumulatedRootMotion[SampleIdx];
+
+	for (int32 SubsampleIdx = 0; SubsampleIdx != Schema->TrajectoryDistanceOffsets.Num(); ++SubsampleIdx)
+	{
+		CurrentElement.SubsampleIdx = SubsampleIdx;
+
+		const float TrajectoryDistance = Schema->TrajectoryDistanceOffsets[SubsampleIdx];
+		const float SampleAccumulatedRootDistance = TrajectoryDistance + AccumulatedRootDistances[SampleIdx];
+
+		int32 LowerBoundSampleIdx = Algo::LowerBound(AccumulatedRootDistances, SampleAccumulatedRootDistance);
+
+		// @@@ Add extrapolation. Clamp for now
+		int32 PrevSampleIdx = FMath::Clamp(LowerBoundSampleIdx - 1, 0, AccumulatedRootDistances.Num() - 1);
+		int32 NextSampleIdx = FMath::Clamp(LowerBoundSampleIdx, 0, AccumulatedRootDistances.Num() - 1);
+
+		const float PrevSampleDistance = AccumulatedRootDistances[PrevSampleIdx];
+		const float NextSampleDistance = AccumulatedRootDistances[NextSampleIdx];
+
+		FTransform PrevRootInSampleSpace = Context.AccumulatedRootMotion[PrevSampleIdx];
+		PrevRootInSampleSpace.SetToRelativeTransform(SampleSpaceOrigin);
+
+		FTransform NextRootInSampleSpace = Context.AccumulatedRootMotion[NextSampleIdx];
+		NextRootInSampleSpace.SetToRelativeTransform(SampleSpaceOrigin);
+		
+		float Alpha = FMath::GetRangePct(PrevSampleDistance, NextSampleDistance, SampleAccumulatedRootDistance);
+		FTransform BlendedRootInSampleSpace;
+		BlendedRootInSampleSpace.Blend(PrevRootInSampleSpace, NextRootInSampleSpace, Alpha);
+
+		Builder.SetTransform(CurrentElement, BlendedRootInSampleSpace);
+	}
+}
+
+
+//////////////////////////////////////////////////////////////////////////
+// PoseSearch API
+
+static void DrawFeatureVector(const FDebugDrawParams& DrawParams, const FFeatureVectorReader& Reader, const FLinearColor& Color1, const FLinearColor& Color2)
+{
+	const float LifeTime = DrawParams.DefaultLifeTime;
+	const uint8 DepthPriority = ESceneDepthPriorityGroup::SDPG_Foreground + 2;
+
+	FPoseSearchFeatureDesc Element;
+	Element.Domain = EPoseSearchFeatureDomain::Time;
+
+	const int32 NumSubsamples = DrawParams.SearchIndex->Schema->PoseSampleOffsets.Num();
+	const int32 NumBones = DrawParams.SearchIndex->Schema->BoneIndices.Num();
+
+	if ((NumSubsamples * NumBones) == 0)
+	{
+		return;
+	}
 
-		if (bIsChildOfPrev)
+	for (int32 SubsampleIdx = 0; SubsampleIdx != NumSubsamples; ++SubsampleIdx)
+	{
+		Element.SubsampleIdx = SubsampleIdx;
+		float Lerp = (SubsampleIdx + 1.0f) / NumSubsamples;
+		FColor Color = FLinearColor::LerpUsingHSV(Color1, Color2, Lerp).ToFColor(true);
+
+		FTransform Adjust;
+		Adjust.SetTranslation(FVector((SubsampleIdx - 2.0f) * 5.0f, 0, 0));
+		Adjust *= DrawParams.ComponentTransform;
+
+		Element.SchemaBoneIdx = 0;
+		FVector BonePosPrev;
+		Reader.GetPosition(Element, &BonePosPrev);
+		BonePosPrev = Adjust.TransformPosition(BonePosPrev);
+
+
+		FVector BoneVel;
+		Reader.GetLinearVelocity(Element, &BoneVel);
+		BoneVel *= 0.1f;
+		DrawDebugDirectionalArrow(DrawParams.World, BonePosPrev, BonePosPrev + BoneVel, 5.0f, FColor::Red, false, LifeTime, DepthPriority, 0.0f);
+		
+		for (int32 SchemaBoneIdx = 1; SchemaBoneIdx != NumBones; ++SchemaBoneIdx)
 		{
-			DrawDebugLine(DrawParams.World, PointPrev, PointNext, Color, false, LifeTime, DepthPriority);
+			Element.SchemaBoneIdx = SchemaBoneIdx;
+
+			FVector BonePosNext;
+			Reader.GetPosition(Element, &BonePosNext);
+			BonePosNext = Adjust.TransformPosition(BonePosNext);
+
+			DrawDebugPoint(DrawParams.World, BonePosNext, 1.0f, Color, false, LifeTime, DepthPriority);
+
+			Reader.GetLinearVelocity(Element, &BoneVel);
+			BoneVel *= 0.1f;
+			DrawDebugDirectionalArrow(DrawParams.World, BonePosNext, BonePosNext + BoneVel, 5.0f, FColor::Red, false, LifeTime, DepthPriority, 0.0f);
+
+			bool bIsChildOfPrev = DrawParams.SearchIndex->Schema->Skeleton->GetReferenceSkeleton().BoneIsChildOf(
+				DrawParams.SearchIndex->Schema->BoneIndices[SchemaBoneIdx],
+				DrawParams.SearchIndex->Schema->BoneIndices[SchemaBoneIdx-1]);
+
+			if (bIsChildOfPrev)
+			{
+				DrawDebugLine(DrawParams.World, BonePosPrev, BonePosNext, Color, false, LifeTime, DepthPriority);
+			}
+			BonePosPrev = BonePosNext;
 		}
-		PointPrev = PointNext;
 	}
 }
 
-void PoseSearchDrawSearchIndex(const FPoseSearchDebugDrawParams& DrawParams, const UPoseSearchIndex& SearchIndex, int32 HighlightPoseIdx)
+static void DrawSearchIndex(const FDebugDrawParams& DrawParams)
 {
 	if (!DrawParams.CanDraw())
 	{
 		return;
 	}
 
-	int32 LastPoseIdx = SearchIndex.NumPoses;
+	FFeatureVectorReader Reader;
+	Reader.Init(&DrawParams.SearchIndex->Schema->Layout);
+
+	int32 LastPoseIdx = DrawParams.SearchIndex->NumPoses;
 	int32 StartPoseIdx = 0;
-	if (!(DrawParams.Flags & EPoseSearchDebugDrawFlags::DrawSearchIndex))
+	if (!(DrawParams.Flags & EDebugDrawFlags::DrawSearchIndex))
 	{
-		StartPoseIdx = HighlightPoseIdx;
+		StartPoseIdx = DrawParams.HighlightPoseIdx;
 		LastPoseIdx = StartPoseIdx + 1;
 	}
 
@@ -242,204 +1112,347 @@ void PoseSearchDrawSearchIndex(const FPoseSearchDebugDrawParams& DrawParams, con
 		return;
 	}
 
-	for (int32 PoseDrawIdx = StartPoseIdx; PoseDrawIdx != LastPoseIdx; ++PoseDrawIdx)
+	for (int32 PoseIdx = StartPoseIdx; PoseIdx != LastPoseIdx; ++PoseIdx)
 	{
-		float LifeTimeDelta = 0.0f;
-		FLinearColor Color;
-		if (PoseDrawIdx == HighlightPoseIdx)
+		FLinearColor Color1;
+		FLinearColor Color2;
+		if (PoseIdx == DrawParams.HighlightPoseIdx)
 		{
-			Color = FLinearColor::Yellow;
+			Color1 = Color2 = FLinearColor::Yellow;
 		}
 		else
 		{
-			float Lerp = (float)(PoseDrawIdx + 1) / SearchIndex.NumPoses;
-			LifeTimeDelta = Lerp - 1.0f;
-			Color = FLinearColor::LerpUsingHSV(FLinearColor(FColor::Cyan), FLinearColor(FColor::Blue), Lerp);
+			float Lerp = (float)(PoseIdx + 1) / DrawParams.SearchIndex->NumPoses;
+			Color1 = FLinearColor(FColor::Cyan);
+			Color2 = FLinearColor(FColor::Blue);
 		}
 
-		TArrayView<const float> Pose = MakeArrayView(&SearchIndex.Values[PoseDrawIdx * SearchIndex.FloatsPerPose], SearchIndex.FloatsPerPose);
-		PoseSearchDrawPose(DrawParams, Color.ToFColor(true), Pose, LifeTimeDelta);
+		// @@@ make a search index helper for pose slicing
+		int32 ValueOffset = PoseIdx * DrawParams.SearchIndex->Schema->Layout.NumFloats;
+		TArrayView<const float> Values = MakeArrayView(&DrawParams.SearchIndex->Values[ValueOffset], DrawParams.SearchIndex->Schema->Layout.NumFloats);
+		Reader.SetValues(Values);
+		
+		DrawFeatureVector(DrawParams, Reader, Color1, Color2);
 	}
 }
 
-void PoseSearchDrawQuery(const FPoseSearchDebugDrawParams& DrawParams, TArrayView<const float> Query)
+static void DrawQuery(const FDebugDrawParams& DrawParams)
 {
 	if (!DrawParams.CanDraw())
 	{
 		return;
 	}
 
-	int32 NumSamples = DrawParams.Schema->FragmentSampleOffsets.Num();
-	for (int32 Sample = 0; Sample != NumSamples; ++Sample)
+	FFeatureVectorReader Reader;
+	Reader.Init(&DrawParams.SearchIndex->Schema->Layout);
+	Reader.SetValues(DrawParams.Query);
+	DrawFeatureVector(DrawParams, Reader, FLinearColor(FColor::Magenta), FLinearColor(FColor::Purple));
+}
+
+void Draw(const FDebugDrawParams& DebugDrawParams)
+{
+	if (DebugDrawParams.CanDraw())
 	{
-		float Lerp = (Sample + 1) / NumSamples;
-		FLinearColor Color = FLinearColor::LerpUsingHSV(FLinearColor(FColor::Magenta), FLinearColor(FColor::Purple), Lerp);
+		if (EnumHasAnyFlags(DebugDrawParams.Flags, EDebugDrawFlags::DrawQuery))
+		{
+			DrawQuery(DebugDrawParams);
+		}
 
-		TArrayView<const float> Pose = MakeArrayView(&Query[Sample * DrawParams.Schema->FloatsPerPose], DrawParams.Schema->FloatsPerPose);
-		PoseSearchDrawPose(DrawParams, Color.ToFColor(true), Pose, 0.0f, 1);
+		if (EnumHasAnyFlags(DebugDrawParams.Flags, EDebugDrawFlags::DrawSearchIndex | EDebugDrawFlags::DrawBest))
+		{
+			DrawSearchIndex(DebugDrawParams);
+		}
 	}
 }
 
-void PoseSearchBuildIndex(const UAnimSequenceBase& AnimSequence, const UPoseSearchIndexConfig& SearchConfig, const UPoseSearchSchema& SearchSchema, UPoseSearchIndex* SearchIndex)
+bool BuildIndex(const UAnimSequence* Sequence, UPoseSearchSequenceMetaData* SequenceMetaData)
 {
-	USkeleton* Skeleton = AnimSequence.GetSkeleton();
-	check(Skeleton);
-	check(Skeleton->IsCompatible(SearchSchema.GetSkeleton()));
-
-	FBoneContainer BoneContainer;
-	BoneContainer.InitializeTo(SearchSchema.BoneIndicesWithParents, FCurveEvaluationOption(false), *Skeleton);
+	check(Sequence);
+	check(SequenceMetaData);
 
-	FBlendedCurve UnusedCurve;
-	FStackCustomAttributes UnusedAttributes;
-
-	FAnimExtractContext ExtractionCtx;
-	// ExtractionCtx.PoseCurves is intentionally left empty
-	// ExtractionCtx.BonesRequired is unused by UAnimSequence::GetAnimationPose
-	ExtractionCtx.bExtractRootMotion = true;
+	if (!SequenceMetaData->IsValidForIndexing())
+	{
+		return false;
+	}
 
-	double CurrTime = SearchConfig.FrameSamplingRange.GetLowerBoundValue();
-	const double EndTime = FMath::Min(AnimSequence.GetPlayLength(), SearchConfig.FrameSamplingRange.GetUpperBoundValue());
-	const double DeltaTime = 1.0 / SearchConfig.SampleRate;
+	USkeleton* SeqSkeleton = Sequence->GetSkeleton();
+	if (!SeqSkeleton || !SeqSkeleton->IsCompatible(SequenceMetaData->Schema->Skeleton))
+	{
+		return false;
+	}
 
-	const int32 NumPoses = FMath::FloorToInt((EndTime - CurrTime) / DeltaTime);
+	FSequenceIndexer Indexer;
+	FSequenceIndexer::Result Result = Indexer.Process(SequenceMetaData->Schema, Sequence, SequenceMetaData->SamplingRange);
 
-	SearchIndex->Values.Reset(SearchSchema.FloatsPerPose * NumPoses);
+	SequenceMetaData->SearchIndex.Values = Result.Values;
+	SequenceMetaData->SearchIndex.NumPoses = Result.NumIndexedPoses;
+	SequenceMetaData->SearchIndex.Schema = SequenceMetaData->Schema;
+	return true;
+}
 
-	FCompactPose Pose;
-	Pose.SetBoneContainer(&BoneContainer);
-	FCSPose<FCompactPose> ComponentSpacePose;
+bool BuildIndex(UPoseSearchDatabase* Database)
+{
+	check(Database);
 
-	FAnimationPoseData AnimPoseData(Pose, UnusedCurve, UnusedAttributes);
-	for (int32 PoseIdx = 0; PoseIdx != NumPoses; ++PoseIdx, CurrTime += DeltaTime)
+	if (!Database->IsValidForIndexing())
 	{
-		// Extract pose
-		ExtractionCtx.CurrentTime = CurrTime;
-		AnimSequence.GetAnimationPose(AnimPoseData, ExtractionCtx);
-		ComponentSpacePose.InitPose(Pose);
+		return false;
+	}
 
-		for (int32 BoneIndex : SearchSchema.BoneIndices)
+	for (const FPoseSearchDatabaseSequence& DbSequence : Database->Sequences)
+	{
+		USkeleton* SeqSkeleton = DbSequence.Sequence->GetSkeleton();
+		if (!SeqSkeleton || !SeqSkeleton->IsCompatible(Database->Schema->Skeleton))
 		{
-			FCompactPoseBoneIndex CompactBoneIndex = BoneContainer.MakeCompactPoseIndex(FMeshPoseBoneIndex(BoneIndex));
-			const FTransform& Transform = ComponentSpacePose.GetComponentSpaceTransform(CompactBoneIndex);
- 
-			FVector Translation = Transform.GetTranslation();
-			SearchIndex->Values.Add(Translation.X);
-			SearchIndex->Values.Add(Translation.Y);
-			SearchIndex->Values.Add(Translation.Z);
+			return false;
 		}
 	}
 
-	SearchIndex->NumPoses = NumPoses;
-	SearchIndex->FloatsPerPose = SearchSchema.FloatsPerPose;
-	SearchIndex->Schema = &SearchSchema;
-	SearchIndex->SequenceSampleRate = SearchConfig.SampleRate;
-	SearchIndex->SequenceFrameSamplingRange = SearchConfig.FrameSamplingRange;
+	// Prepare animation indexing tasks
+	TArray<FSequenceIndexer> Indexers;
+	Indexers.SetNum(Database->Sequences.Num());
+
+	auto IndexerTask = [&Database, &Indexers](int32 SequenceIdx)
+	{
+		const FPoseSearchDatabaseSequence& DbSequence = Database->Sequences[SequenceIdx];
+		FSequenceIndexer& Indexer = Indexers[SequenceIdx];
+		Indexer.Process(Database->Schema, DbSequence.Sequence, DbSequence.SamplingRange);
+	};
+
+	// Index animations independently
+	ParallelFor(Database->Sequences.Num(), IndexerTask);
+
+	// Write index info to sequence and count up total poses and storage required
+	int32 TotalPoses = 0;
+	int32 TotalFloats = 0;
+	for (int32 SequenceIdx = 0; SequenceIdx != Database->Sequences.Num(); ++SequenceIdx)
+	{
+		FPoseSearchDatabaseSequence& DbSequence = Database->Sequences[SequenceIdx];
+		FSequenceIndexer::Result Result = Indexers[SequenceIdx].GetResult();
+		DbSequence.NumPoses = Result.NumIndexedPoses;
+		DbSequence.FirstPoseIdx = TotalPoses;
+		TotalPoses += Result.NumIndexedPoses;
+		TotalFloats += Result.Values.Num();
+	}
+
+	// Join animation data into a single search index
+	Database->SearchIndex.Values.Reset(TotalFloats);
+	for (const FSequenceIndexer& Indexer : Indexers)
+	{
+		FSequenceIndexer::Result Result = Indexer.GetResult();
+		Database->SearchIndex.Values.Append(Result.Values.GetData(), Result.Values.Num());
+	}
+
+	Database->SearchIndex.NumPoses = TotalPoses;
+	Database->SearchIndex.Schema = Database->Schema;
+	return true;
 }
 
-bool PoseSearchBuildQuery(const UPoseSearchSchema& SearchSchema, int32 AssetSampleRate, const FPoseSearchPoseHistory& History, FPoseSearchBuildQueryScratch* Scratch, TArray<float>* Query)
+bool BuildQuery(const UPoseSearchSchema* Schema, FPoseHistory* History, TArrayView<float> Query)
 {
-	Query->Reset(0);
-	Query->Reserve(SearchSchema.FragmentSampleOffsets.Num() * SearchSchema.FloatsPerPose);
+	check(Schema);
+	check(History);
 
-	for (int32 Offset : SearchSchema.FragmentSampleOffsets)
+	FFeatureVectorBuilder Builder;
+	Builder.Init(&Schema->Layout, Query);
+
+	FPoseSearchFeatureDesc Feature;
+	Feature.Domain = EPoseSearchFeatureDomain::Time; // @@@ making assumption here, need to refer to schema
+
+	for (int32 SubsampleIdx = 0; SubsampleIdx != Schema->PoseSampleOffsets.Num(); ++SubsampleIdx)
 	{
-		float TimeDelta = Offset * (1.0f / AssetSampleRate);
+		Feature.SubsampleIdx = SubsampleIdx;
 
-		if (!History.Sample(TimeDelta, SearchSchema.BoneIndicesWithParents, Scratch->LocalPose))
+		int32 Offset = Schema->PoseSampleOffsets[SubsampleIdx];
+		float TimeDelta = -Offset * (1.0f /  Schema->SampleRate);
+
+		if (!History->Sample(TimeDelta, Schema->Skeleton->GetReferenceSkeleton(), Schema->BoneIndicesWithParents))
 		{
-			return false;
+			break;
 		}
 
-		FAnimationRuntime::FillUpComponentSpaceTransforms(SearchSchema.GetSkeleton()->GetReferenceSkeleton(), Scratch->LocalPose, Scratch->ComponentPose);
-
-		for (int32 SkeletonBoneIndex : SearchSchema.BoneIndices)
+		TArrayView<const FTransform> ComponentPose = History->GetComponentPoseSample();
+		TArrayView<const FTransform> ComponentPrevPose = History->GetPrevComponentPoseSample();
+		for (int32 SchemaBoneIdx = 0; SchemaBoneIdx != Schema->BoneIndices.Num(); ++SchemaBoneIdx)
 		{
-			const FTransform& Transform = Scratch->ComponentPose[SkeletonBoneIndex];
-			FVector Translation = Transform.GetTranslation();
-			Query->Add(Translation.X);
-			Query->Add(Translation.Y);
-			Query->Add(Translation.Z);
+			Feature.SchemaBoneIdx = SchemaBoneIdx;
+
+			int32 SkeletonBoneIndex = Schema->BoneIndices[SchemaBoneIdx]; // @@@ need to resolve bone indices at runtime, hold in FPoseHistory?
+			const FTransform& Transform = ComponentPose[SkeletonBoneIndex];
+			const FTransform& PrevTransform = ComponentPrevPose[SkeletonBoneIndex];
+			Builder.SetTransform(Feature, Transform);
+			Builder.SetTransformDerivative(Feature, Transform, PrevTransform, History->GetSampleInterval());
 		}
 	}
 
-	return true;
+	return Builder.IsComplete();
 }
 
-FPoseSearchResult PoseSearch(const UPoseSearchIndex& SearchIndex, TArrayView<const float> Query, FPoseSearchDebugDrawParams DebugDrawParams)
+static FSearchResult Search(const FPoseSearchIndex& SearchIndex, TArrayView<const float> Query)
 {
-	check(SearchIndex.Schema != nullptr);
-	check(SearchIndex.NumPoses * SearchIndex.FloatsPerPose == SearchIndex.Values.Num());
+	FSearchResult Result;
+
+	if (!ensure(SearchIndex.IsValid()))
+	{
+		return Result;
+	}
 
-	if (Query.Num() != SearchIndex.FloatsPerPose * SearchIndex.Schema->FragmentSampleOffsets.Num())
+	if(!ensure(Query.Num() == SearchIndex.Schema->Layout.NumFloats))
 	{
-		return FPoseSearchResult();
+		return Result;
 	}
 
-	float BestPoseDifference = TNumericLimits<float>::Max();
+	float BestPoseDissimilarity = MAX_flt;
 	int32 BestPoseIdx = INDEX_NONE;
 
-	for (int32 PoseIdx = SearchIndex.Schema->FragmentSampleOffsets[0]; PoseIdx < SearchIndex.NumPoses; ++PoseIdx)
+	for (int32 PoseIdx = 0; PoseIdx != SearchIndex.NumPoses; ++PoseIdx)
 	{
-		float PoseDifference = 0.0f;
-		int32 QueryValueIdx = 0;
-		for (int32 PoseOffset : SearchIndex.Schema->FragmentSampleOffsets)
-		{
-			int32 SearchValueIdx = (PoseIdx - PoseOffset) * SearchIndex.FloatsPerPose;
+		const int32 FeatureValueOffset = PoseIdx * SearchIndex.Schema->Layout.NumFloats;
 
-			for (int32 PoseElemCount = 0; PoseElemCount != SearchIndex.FloatsPerPose; ++PoseElemCount, ++QueryValueIdx, ++SearchValueIdx)
-			{
-				PoseDifference += FMath::Square(Query[QueryValueIdx] - SearchIndex.Values[SearchValueIdx]);
-			}
+		float PoseDissimilarity = 0.0f;
+		for (int32 ValueIdx = 0; ValueIdx != SearchIndex.Schema->Layout.NumFloats; ++ValueIdx)
+		{
+			PoseDissimilarity += FMath::Square(Query[ValueIdx] - SearchIndex.Values[ValueIdx + FeatureValueOffset]);
 		}
 
-		if (PoseDifference < BestPoseDifference)
+		if (PoseDissimilarity < BestPoseDissimilarity)
 		{
-			BestPoseDifference = PoseDifference;
+			BestPoseDissimilarity = PoseDissimilarity;
 			BestPoseIdx = PoseIdx;
 		}
 	}
 
-	check(BestPoseIdx != INDEX_NONE);
+	ensure(BestPoseIdx != INDEX_NONE);
+
+	Result.Dissimilarity = BestPoseDissimilarity;
+	Result.PoseIdx = BestPoseIdx;
+	// Result.TimeOffsetSeconds is set by caller
+
+	return Result;
+}
+
+FSearchResult Search(const UPoseSearchSequenceMetaData* Sequence, TArrayView<const float> Query, FDebugDrawParams DebugDrawParams)
+{
+	FSearchResult Result;
+
+	if (!ensure(Sequence && Sequence->IsValidForSearch()))
+	{
+		return Result;
+	}
+
+	const FPoseSearchIndex& SearchIndex = Sequence->SearchIndex;
+
+	Result = Search(SearchIndex, Query);
+	if (!Result.IsValid())
+	{
+		return Result;
+	}
 
-	const float SampleDelta = 1.0f / SearchIndex.SequenceSampleRate;
+	const float SampleDelta = 1.0f / SearchIndex.Schema->SampleRate;
 
-	float BestPoseTime = FMath::Min(
-		SearchIndex.SequenceFrameSamplingRange.GetUpperBoundValue(),
-		SampleDelta * BestPoseIdx + SearchIndex.SequenceFrameSamplingRange.GetLowerBoundValue());
+	Result.TimeOffsetSeconds = SampleDelta * Result.PoseIdx + Sequence->SamplingRange.Min;
 
 	// Do debug visualization
-	DebugDrawParams.Schema = SearchIndex.Schema;
-	if (DebugDrawParams.CanDraw())
+	DebugDrawParams.SearchIndex = &SearchIndex;
+	DebugDrawParams.Query = Query;
+	DebugDrawParams.HighlightPoseIdx = Result.PoseIdx;
+	Draw(DebugDrawParams);
+
+	return Result;
+}
+
+FDbSearchResult Search(const UPoseSearchDatabase* Database, TArrayView<const float> Query, FDebugDrawParams DebugDrawParams)
+{
+	if (!ensure(Database && Database->IsValidForSearch()))
 	{
-		if (EnumHasAnyFlags(DebugDrawParams.Flags, EPoseSearchDebugDrawFlags::DrawQuery))
-		{
-			PoseSearchDrawQuery(DebugDrawParams, Query);
-		}
+		return FDbSearchResult();
+	}
 
-		if (EnumHasAnyFlags(DebugDrawParams.Flags, EPoseSearchDebugDrawFlags::DrawSearchIndex | EPoseSearchDebugDrawFlags::DrawBest))
-		{
-			int32 HighlightPoseIdx = EnumHasAnyFlags(DebugDrawParams.Flags, EPoseSearchDebugDrawFlags::DrawBest) ? BestPoseIdx : -1;
-			PoseSearchDrawSearchIndex(DebugDrawParams, SearchIndex, HighlightPoseIdx);
-		}
+	const FPoseSearchIndex& SearchIndex = Database->SearchIndex;
+
+	FDbSearchResult Result = Search(SearchIndex, Query);
+	if (!Result.IsValid())
+	{
+		return FDbSearchResult();
 	}
 
-	FPoseSearchResult Result;
-	Result.Dissimilarity = BestPoseDifference;
-	Result.TimeOffsetSeconds = BestPoseTime;
-	Result.PoseIdx = BestPoseIdx;
+	const FPoseSearchDatabaseSequence* DbSequence = Database->FindSequenceByPoseIdx(Result.PoseIdx);
+	if (!ensure(DbSequence))
+	{
+		return FDbSearchResult();
+	}
+
+	Result.DbSequence = DbSequence;
+
+	const float SampleDelta = 1.0f / SearchIndex.Schema->SampleRate;
+	Result.TimeOffsetSeconds = SampleDelta * (DbSequence->FirstPoseIdx - Result.PoseIdx) + DbSequence->SamplingRange.Min;
+
+	// Do debug visualization
+	DebugDrawParams.SearchIndex = &SearchIndex;
+	DebugDrawParams.Query = Query;
+	DebugDrawParams.HighlightPoseIdx = Result.PoseIdx;
+	Draw(DebugDrawParams);
+
+	return Result;
+}
+
+FSearchResult Search(const FAnimationBaseContext& GraphContext, const UAnimSequenceBase* Sequence)
+{
+	FSearchResult Result;
+
+	const UPoseSearchSequenceMetaData* MetaData = Sequence->FindMetaDataByClass<UPoseSearchSequenceMetaData>();
+	if (!MetaData || !MetaData->IsValidForSearch())
+	{
+		return Result;
+	}
+
+	FAnimNode_PoseSearchHistoryCollector* HistoryNode = GraphContext.GetAncestor<FAnimNode_PoseSearchHistoryCollector>();
+	if (!HistoryNode)
+	{
+		return Result;
+	}
+
+	TArrayView<const float> Query = HistoryNode->BuildQuery(MetaData->Schema);
+
+	Result = Search(MetaData, Query);
 	return Result;
 }
 
 
 //////////////////////////////////////////////////////////////////////////
-// FPoseSearchModule
+// FModule
 
-void FPoseSearchModule::StartupModule()
+class FModule : public IModuleInterface, public UE::Anim::IPoseSearchProvider
 {
+public: // IModuleInterface
+	virtual void StartupModule() override;
+	virtual void ShutdownModule() override;
+
+public: // IPoseSearchProvider
+	virtual UE::Anim::IPoseSearchProvider::FSearchResult Search(const FAnimationBaseContext& GraphContext, const UAnimSequenceBase* Sequence) override;
+};
+
+void FModule::StartupModule()
+{
+	IModularFeatures::Get().RegisterModularFeature(UE::Anim::IPoseSearchProvider::ModularFeatureName, this);
 }
 
-void FPoseSearchModule::ShutdownModule()
+void FModule::ShutdownModule()
 {
+	IModularFeatures::Get().UnregisterModularFeature(UE::Anim::IPoseSearchProvider::ModularFeatureName, this);
 }
 
-IMPLEMENT_MODULE(FPoseSearchModule, PoseSearch)
\ No newline at end of file
+UE::Anim::IPoseSearchProvider::FSearchResult FModule::Search(const FAnimationBaseContext& GraphContext, const UAnimSequenceBase* Sequence)
+{
+	::UE::PoseSearch::FSearchResult Result = ::UE::PoseSearch::Search(GraphContext, Sequence);
+
+	UE::Anim::IPoseSearchProvider::FSearchResult ProviderResult;
+	ProviderResult.Dissimilarity = Result.Dissimilarity;
+	ProviderResult.PoseIdx = Result.PoseIdx;
+	ProviderResult.TimeOffsetSeconds = Result.TimeOffsetSeconds;
+	return ProviderResult;
+}
+
+}} // namespace UE::PoseSearch
+
+IMPLEMENT_MODULE(UE::PoseSearch::FModule, PoseSearch)
\ No newline at end of file
diff --git a/Engine/Plugins/Animation/PoseSearch/Source/Runtime/Public/PoseSearch/AnimNode_PoseMatchingSequencePlayer.h b/Engine/Plugins/Animation/PoseSearch/Source/Runtime/Public/PoseSearch/AnimNode_PoseMatchingSequencePlayer.h
deleted file mode 100644
diff --git a/Engine/Plugins/Animation/PoseSearch/Source/Runtime/Public/PoseSearch/AnimNode_PoseSearchHistoryCollector.h b/Engine/Plugins/Animation/PoseSearch/Source/Runtime/Public/PoseSearch/AnimNode_PoseSearchHistoryCollector.h
--- a/Engine/Plugins/Animation/PoseSearch/Source/Runtime/Public/PoseSearch/AnimNode_PoseSearchHistoryCollector.h
+++ b/Engine/Plugins/Animation/PoseSearch/Source/Runtime/Public/PoseSearch/AnimNode_PoseSearchHistoryCollector.h
@@ -25,12 +25,14 @@ public:
 	virtual void Update_AnyThread(const FAnimationUpdateContext& Context);
 	// End of FAnimNode_Base interface
 
-	const FPoseSearchPoseHistory& GetPoseHistory() const { return PoseHistory; }
-	void Init(int32 InNumPoses, float InTimeHorizon);
+	UE::PoseSearch::FPoseHistory& GetPoseHistory() { return PoseHistory; }
+	const UE::PoseSearch::FPoseHistory& GetPoseHistory() const { return PoseHistory; }
+	TArrayView<const float> BuildQuery(const UPoseSearchSchema* Schema);
+
 
 protected:
 
-	FPoseSearchPoseHistory PoseHistory;
+	UE::PoseSearch::FPoseHistory PoseHistory;
+	TArray<float> Query;
 	float EvalDeltaTime = 0.0f;
-
 };
\ No newline at end of file
diff --git a/Engine/Plugins/Animation/PoseSearch/Source/Runtime/Public/PoseSearch/PoseSearch.h b/Engine/Plugins/Animation/PoseSearch/Source/Runtime/Public/PoseSearch/PoseSearch.h
--- a/Engine/Plugins/Animation/PoseSearch/Source/Runtime/Public/PoseSearch/PoseSearch.h
+++ b/Engine/Plugins/Animation/PoseSearch/Source/Runtime/Public/PoseSearch/PoseSearch.h
@@ -3,104 +3,300 @@
 #pragma once
 
 #include "CoreMinimal.h"
-#include "Animation/AnimMetaData.h"
-#include "Animation/AnimationAsset.h"
-#include "BoneIndices.h"
+#include "Engine/DataAsset.h"
 #include "Modules/ModuleManager.h"
 #include "Misc/EnumClassFlags.h"
 #include "UObject/Object.h"
 #include "UObject/ObjectMacros.h"
+#include "Animation/AnimMetaData.h"
+#include "Animation/AnimationAsset.h"
+#include "BoneIndices.h"
+#include "Interfaces/Interface_BoneReferenceSkeletonProvider.h"
 
 #include "PoseSearch.generated.h"
 
 
-class UAnimSequenceBase;
+class UAnimSequence;
 struct FCompactPose;
+struct FReferenceSkeleton;
 
-/** Sampling parameters for building a pose search index for an asset. */
-UCLASS(BlueprintType, Category = "Animation|PoseSearch")
-class POSESEARCH_API UPoseSearchIndexConfig : public UAnimMetaData
+UENUM()
+enum class EPoseSearchFeatureType : int32
+{
+	Invalid			= -1,
+
+	Position		= 0,	
+	Rotation		= 1,
+	LinearVelocity	= 2,
+	AngularVelocity	= 3,
+
+	Num
+};
+
+UENUM()
+enum class EPoseSearchFeatureDomain : int32
+{
+	Invalid		= -1,
+
+	Time		= 0,
+	Distance	= 1,
+
+	Num
+};
+
+/** Describes each feature of a vector, including data type, sampling options, and buffer offset. */
+USTRUCT()
+struct POSESEARCH_API FPoseSearchFeatureDesc
 {
 	GENERATED_BODY()
 
-public:
-	static constexpr int32 DefaultSampleRate = 15;
+	static constexpr int32 TrajectoryBoneIndex = -1;
 
-	UPROPERTY(EditAnywhere, meta = (ClampMin = "1", ClampMax = "60"), Category = "Config")
-	int32 SampleRate;
+	UPROPERTY()
+	int32 SchemaBoneIdx;
 
-	UPROPERTY(EditAnywhere, Category = "Config")
-	FFloatRange FrameSamplingRange;
+	UPROPERTY()
+	int32 SubsampleIdx;
 
-	UPoseSearchIndexConfig()
-	    : SampleRate(DefaultSampleRate)
+	UPROPERTY()
+	EPoseSearchFeatureType Type;
+
+	UPROPERTY()
+	EPoseSearchFeatureDomain Domain;
+
+	// Set via FPoseSearchFeatureLayout::Init() and ignored by operator==
+	UPROPERTY()
+	int32 ValueOffset;
+
+	bool operator==(const FPoseSearchFeatureDesc& Other) const;
+
+	FPoseSearchFeatureDesc()
+		: SchemaBoneIdx(MAX_int32)
+		, SubsampleIdx(MAX_int32)
+		, Type(EPoseSearchFeatureType::Invalid)
+		, Domain(EPoseSearchFeatureDomain::Invalid)
+		, ValueOffset(MAX_int32)
 	{}
 };
 
+/**
+* Explicit description of a pose feature vector.
+* Determined by options set in a UPoseSearchSchema and owned by the schema.
+* See UPoseSearchSchema::GenerateLayout().
+*/
+USTRUCT()
+struct POSESEARCH_API FPoseSearchFeatureVectorLayout
+{
+	GENERATED_BODY()
+
+	void Init();
+	void Reset();
+
+	UPROPERTY()
+	TArray<FPoseSearchFeatureDesc> Features;
+
+	UPROPERTY()
+	uint32 NumFloats;
+
+	bool IsValid() const;
+};
+
+
+//USTRUCT()
+//struct FPoseSearchSchemaBoneRef
+//{
+//	GENERATED_BODY()
+//
+//	/** Index into UPoseSearchSchema::Bones. */
+//	UPROPERTY()
+//	int32 BoneIdx = -1;
+//
+//	/**
+//	* Optional index into UPoseSearchSchema::Bones.
+//	* Causes BoneIdx to be sampled in the reference frame of this bone.
+//	*/
+//	UPROPERTY()
+//	int32 RelativeBoneIdx = -1;
+//};
+
 
 /**
 * Specifies the format of a pose search index. At runtime, queries are built according to the schema for searching.
 */
 UCLASS(BlueprintType, Category = "Animation|PoseSearch")
-class POSESEARCH_API UPoseSearchSchema : public UAnimationAsset
+class POSESEARCH_API UPoseSearchSchema : public UDataAsset, public IBoneReferenceSkeletonProvider
 {
 	GENERATED_BODY()
 
 public:
+
+	static constexpr int32 DefaultSampleRate = 10;
+
+	UPROPERTY(EditAnywhere, Category = "Schema")
+	USkeleton* Skeleton;
+
+	UPROPERTY(EditAnywhere, meta = (ClampMin = "1", ClampMax = "60"), Category = "Schema")
+	int32 SampleRate;
+
+	UPROPERTY(EditAnywhere, Category = "Schema")
+	bool bUseBoneVelocities;
+
+	UPROPERTY(EditAnywhere, Category = "Schema")
+	bool bUseBonePositions;
+
+	UPROPERTY(EditAnywhere, Category = "Schema")
+	bool bUseTrajectoryVelocities;
+
+	UPROPERTY(EditAnywhere, Category = "Schema")
+	bool bUseTrajectoryPositions;
+
 	UPROPERTY(EditAnywhere, Category = "Schema")
 	TArray<FBoneReference> Bones;
 
+	/*
+	UPROPERTY(EditAnywhere, Category = "Schema")
+	TArray<FPoseSearchSchemaBoneRef> PoseSampleBones;
+	*/
+
+	UPROPERTY(EditAnywhere, Category = "Schema")
+	TArray<int32> PoseSampleOffsets;
+
+	UPROPERTY(EditAnywhere, Category = "Schema")
+	TArray<int32> TrajectorySampleOffsets;
+
 	UPROPERTY(EditAnywhere, Category = "Schema")
-	TArray<int32> FragmentSampleOffsets;
+	TArray<float> TrajectoryDistanceOffsets;
 
 	UPROPERTY()
-	int32 FloatsPerPose;
+	FPoseSearchFeatureVectorLayout Layout;
+
 
+	// @@@ these should be deprecated. Need to gather bone indices at runtime instead of baking.
 	UPROPERTY()
 	TArray<uint16> BoneIndices;
 
 	UPROPERTY()
 	TArray<uint16> BoneIndicesWithParents;
 
+	UPoseSearchSchema()
+		: SampleRate(DefaultSampleRate)
+	{}
+
+	bool IsValid () const;
+
 public: // UObject
 	virtual void PreSave(const class ITargetPlatform* TargetPlatform);
-};
 
+public: // IBoneReferenceSkeletonProvider
+	class USkeleton* GetSkeleton(bool& bInvalidSkeletonIsError) override { bInvalidSkeletonIsError = false; return Skeleton; }
+
+private:
+	void GenerateLayout();
+};
 
 /**
-* Contains pose data sampled at editor time for searching at runtime.
+* A search index for animation poses. The structure of the search index is determined by its UPoseSearchSchema.
+* May represent a single animation (see UPoseSearchSequenceMetaData) or a collection (see UPoseSearchDatabase).
 */
-UCLASS(BlueprintType, Category = "Animation|PoseSearch")
-class POSESEARCH_API UPoseSearchIndex : public UAnimMetaData
+USTRUCT()
+struct POSESEARCH_API FPoseSearchIndex
 {
 	GENERATED_BODY()
 public:
 	UPROPERTY()
 	int32 NumPoses;
 
-	UPROPERTY()
-	int32 FloatsPerPose;
-
 	UPROPERTY()
 	TArray<float> Values;
 
 	UPROPERTY()
 	const UPoseSearchSchema* Schema;
 
+	bool IsValid() const;
+};
+
+/** Animation metadata object for indexing a single animation. */
+UCLASS(BlueprintType, Category = "Animation|PoseSearch")
+class POSESEARCH_API UPoseSearchSequenceMetaData : public UAnimMetaData
+{
+	GENERATED_BODY()
+public:
+
+	UPROPERTY(EditAnywhere, Category="Settings")
+	const UPoseSearchSchema* Schema;
+
+	UPROPERTY(EditAnywhere, Category="Settings")
+	FFloatInterval SamplingRange;
+
+	UPROPERTY();
+	FPoseSearchIndex SearchIndex;
+
+	UPoseSearchSequenceMetaData ()
+		: SamplingRange(0.0f, 0.0f)
+	{}
+
+	bool IsValidForIndexing() const;
+	bool IsValidForSearch() const;
+
+public: // UObject
+	virtual void PreSave(const class ITargetPlatform* TargetPlatform);
+};
+
+/** An entry in a UPoseSearchDatabase. */
+USTRUCT(BlueprintType, Category = "Animation|PoseSearch")
+struct POSESEARCH_API FPoseSearchDatabaseSequence
+{
+	GENERATED_BODY()
+
+	UPROPERTY(EditAnywhere, Category="Sequence")
+	UAnimSequence* Sequence;
+
+	UPROPERTY(EditAnywhere, Category="Sequence")
+	FFloatInterval SamplingRange;
+
+	UPROPERTY(EditAnywhere, Category="Sequence")
+	bool bLoopAnimation;
+
 	UPROPERTY()
-	int32 SequenceSampleRate;
+	int32 FirstPoseIdx;
 
 	UPROPERTY()
-	FFloatRange SequenceFrameSamplingRange;
+	int32 NumPoses;
+
+	FPoseSearchDatabaseSequence()
+		: SamplingRange(0.0f, 0.0f)
+	{}
 };
 
-struct POSESEARCH_API FPoseSearchPoseSnapshot
+/** A data asset for indexing a collection of animation sequences. */
+UCLASS(BlueprintType, Category = "Animation|PoseSearch")
+class POSESEARCH_API UPoseSearchDatabase : public UDataAsset
 {
-	TArray<FTransform> LocalTransforms;
+	GENERATED_BODY()
+public:
+
+	UPROPERTY(EditAnywhere, Category="Database")
+	const UPoseSearchSchema* Schema;
+
+	UPROPERTY(EditAnywhere, Category="Database")
+	TArray<FPoseSearchDatabaseSequence> Sequences;
+
+	UPROPERTY()
+	FPoseSearchIndex SearchIndex;
+
+	const FPoseSearchDatabaseSequence* FindSequenceByPoseIdx(int32 PoseIdx) const;
+
+	bool IsValidForIndexing() const;
+	bool IsValidForSearch() const;
+
+public: // UObject
+	virtual void PreSave(const class ITargetPlatform* TargetPlatform);
 };
 
 
-namespace PoseSearchDetail {
+namespace UE { namespace PoseSearch {
+
+namespace Private {
 
 /**
 * Helper to adapt a fixed-size buffer to a circular queue.
@@ -213,34 +409,100 @@ private:
 	uint32 VirtualLast;
 };
 
-} // namespace PoseSearchDetail
+struct FSnapshot
+{
+	TArray<FTransform> LocalTransforms;
+};
+
+} // namespace Private
 
 
-/** Records poses over time in a ring buffer. PoseSearchBuildQuery() uses this to sample from the present or past poses according to the search schema. */
-class POSESEARCH_API FPoseSearchPoseHistory
+/** Records poses over time in a ring buffer. BuildQuery() uses this to sample from the present or past poses according to the search schema. */
+// @@@Rename to FPoseQueryContext?
+class POSESEARCH_API FPoseHistory
 {
 public:
 	void Init(int32 InNumPoses, float InTimeHorizon);
-	void Init(const FPoseSearchPoseHistory& History);
-	bool Sample(float Time, const TArray<FBoneIndexType>& RequiredBones, TArray<FTransform>& OutPose) const;
-	void Update(float TimeDelta, const FCompactPose& EvalContext);
+	void Init(const FPoseHistory& History);
+	bool Sample(float SecondsAgo, const FReferenceSkeleton& RefSkeleton, const TArray<FBoneIndexType>& RequiredBones);
+	void Update(float SecondsElapsed, const FCompactPose& EvalContext);
 	float GetSampleInterval() const;
+	TArrayView<const FTransform> GetLocalPoseSample() const { return SampledLocalPose; }
+	TArrayView<const FTransform> GetComponentPoseSample() const { return SampledComponentPose; }
+	TArrayView<const FTransform> GetPrevLocalPoseSample() const { return SampledPrevLocalPose; }
+	TArrayView<const FTransform> GetPrevComponentPoseSample() const { return SampledPrevComponentPose; }
 	float GetTimeHorizon() const { return TimeHorizon; }
 
 private:
-	TArray<FPoseSearchPoseSnapshot> Poses;
+	bool SampleLocalPose(float Time, const FReferenceSkeleton& RefSkeleton, const TArray<FBoneIndexType>& RequiredBones, TArray<FTransform>& LocalPose);
+
+	TArray<Private::FSnapshot> Poses;
 	TArray<float> Knots;
-	PoseSearchDetail::FCircularView Queue;
+	TArray<FTransform> SampledLocalPose;
+	TArray<FTransform> SampledComponentPose;
+	TArray<FTransform> SampledPrevLocalPose;
+	TArray<FTransform> SampledPrevComponentPose;
+	Private::FCircularView Queue;
 	float TimeHorizon = 0.0f;
 };
 
-struct FPoseSearchBuildQueryScratch
+
+/** 
+* Helper object for writing features into a float buffer according to a feature vector layout.
+* Keeps track of which features are present, allowing the feature vector to be built up piecemeal.
+*/
+class POSESEARCH_API FFeatureVectorBuilder
 {
-	TArray<FTransform> LocalPose;
-	TArray<FTransform> ComponentPose;
+public:
+	void Init(const FPoseSearchFeatureVectorLayout* Layout, TArrayView<float> Buffer);
+	void ResetFeatures();
+
+	void SetTransform(FPoseSearchFeatureDesc Feature, const FTransform& Transform);
+	void SetTransformDerivative(FPoseSearchFeatureDesc Feature, const FTransform& Transform, const FTransform& PrevTransform, float DeltaTime);
+
+	void SetPosition(FPoseSearchFeatureDesc Feature, const FVector& Translation);
+	void SetRotation(FPoseSearchFeatureDesc Feature, const FQuat& Rotation);
+	void SetLinearVelocity(FPoseSearchFeatureDesc Feature, const FTransform& Transform, const FTransform& PrevTransform, float DeltaTime);
+	void SetAngularVelocity(FPoseSearchFeatureDesc Feature, const FTransform& Transform, const FTransform& PrevTransform, float DeltaTime);
+
+	void SetVector(FPoseSearchFeatureDesc Feature, const FVector& Vector);
+
+	bool SetPoseFeatures(const UPoseSearchSchema* Schema, FPoseHistory* History);
+
+	bool IsComplete() const;
+
+private:
+	const FPoseSearchFeatureVectorLayout* Layout = nullptr;
+	TArrayView<float> Values;
+	TBitArray<> FeaturesAdded;
+	int32 NumFeaturesAdded = 0;
+};
+
+/** Helper object for extracting features from a float buffer according to the feature vector layout. */
+class POSESEARCH_API FFeatureVectorReader
+{
+public:
+	void Init(const FPoseSearchFeatureVectorLayout* Layout);
+	void SetValues(TArrayView<const float> Values);
+	bool IsValid() const;
+
+	bool GetTransform(FPoseSearchFeatureDesc Feature, FTransform* OutTransform) const;
+	bool GetPosition(FPoseSearchFeatureDesc Feature, FVector* OutPosition) const;
+	bool GetRotation(FPoseSearchFeatureDesc Feature, FQuat* OutRotation) const;
+	bool GetLinearVelocity(FPoseSearchFeatureDesc Feature, FVector* OutLinearVelocity) const;
+	bool GetAngularVelocity(FPoseSearchFeatureDesc Feature, FVector* OutAngularVelocity) const;
+	bool GetVector(FPoseSearchFeatureDesc Feature, FVector* OutVector) const;
+
+private:
+	const FPoseSearchFeatureVectorLayout* Layout = nullptr;
+	TArrayView<const float> Values;
 };
 
-enum class EPoseSearchDebugDrawFlags : uint32
+
+//////////////////////////////////////////////////////////////////////////
+// Main PoseSearch API
+
+enum class EDebugDrawFlags : uint32
 {
 	None			= 0,
 	DrawQuery		= 1 << 0,
@@ -248,85 +510,98 @@ enum class EPoseSearchDebugDrawFlags : uint32
 	DrawBest        = 1 << 2,
 	DrawAll		    = MAX_uint32
 };
-ENUM_CLASS_FLAGS(EPoseSearchDebugDrawFlags);
+ENUM_CLASS_FLAGS(EDebugDrawFlags);
 
-struct FPoseSearchDebugDrawParams
+struct FDebugDrawParams
 {
 	const UWorld* World = nullptr;
+	const FPoseSearchIndex* SearchIndex = nullptr;
+	EDebugDrawFlags Flags = EDebugDrawFlags::None;
 	float DefaultLifeTime = 5.0f;
-	EPoseSearchDebugDrawFlags Flags = EPoseSearchDebugDrawFlags::None;
-	const UPoseSearchSchema* Schema = nullptr;
 	FTransform ComponentTransform = FTransform::Identity;
+	int32 HighlightPoseIdx = -1;
+	TArrayView<const float> Query;
 
-	bool CanDraw () const { return World && Schema; }
+	bool CanDraw () const { return World && SearchIndex; }
 };
 
-struct FPoseSearchResult
+struct FSearchResult
 {
 	int32 PoseIdx = -1;
 	float TimeOffsetSeconds = 0.0f;
 	float Dissimilarity = MAX_flt;
+
+	bool IsValid() const { return PoseIdx >= 0; }
 };
 
+struct FDbSearchResult : public FSearchResult
+{
+	FDbSearchResult() = default;
+
+	FDbSearchResult(const FSearchResult& Result)
+		: FSearchResult(Result)
+	{}
+
+	const FPoseSearchDatabaseSequence* DbSequence = nullptr;
+};
 
-//////////////////////////////////////////////////////////////////////////
-// Main PoseSearch API
 
 /**
-* Draws a pose search query
+* Visualize pose search debug information
 *
-* @param DrawParams Visualization options
-* @param Query The pose search query created with PoseSearchBuildQuery()
+* @param DrawParams		Visualization options
 */
-void POSESEARCH_API PoseSearchDrawQuery(const FPoseSearchDebugDrawParams& DrawParams, TArrayView<const float> Query);
+POSESEARCH_API void Draw(const FDebugDrawParams& DrawParams);
+
 
 /**
-* Draws a pose search index
-*
-* @param DrawParams Visualization options
-* @param SearchIndex The pose search index to visualize
-* @param HighlightPoseIdx Optional index of a pose in the search index to highlight. Specify -1 to not highlight a pose.
+* Creates a pose search index for an animation sequence
+* 
+* @param Database	The input sequence and output metadata to write the search index to
 */
-void POSESEARCH_API PoseSearchDrawSearchIndex(const FPoseSearchDebugDrawParams& DrawParams, const UPoseSearchIndex& SearchIndex, int32 HighlightPoseIdx = -1);
+POSESEARCH_API bool BuildIndex(const UAnimSequence* Sequence, UPoseSearchSequenceMetaData* SequenceMetaData);
+
 
 /**
-* Creates a pose search index by sampling from the animation sequence
-*
-* @param AnimSequence The asset to sample from
-* @param SearchConfig Options for sampling from this asset
-* @param SearchSchema The format of the search index
-* @param SearchIndex Output pose search index
+* Creates a pose search index for a collection of animations
+* 
+* @param Database	The input collection of animations and output search index
 */
-void POSESEARCH_API PoseSearchBuildIndex(const UAnimSequenceBase& AnimSequence, const UPoseSearchIndexConfig& SearchConfig, const UPoseSearchSchema& SearchSchema, UPoseSearchIndex* SearchIndex);
+POSESEARCH_API bool BuildIndex(UPoseSearchDatabase* Database);
+
 
 /**
 * Builds a pose search query as an array of floats according to the search schema
 *
-* @param SearchSchema The format of the search index and the query being built
-* @param AssetSampleRate The sampling rate used to index the animation asset
-* @param History Recorded pose to sample from to build the query
-* @param Scratch Working memory to re-use for query construction
-* @Query Output pose search query
+* @param Schema			The schema to build a query for
+* @param History		Recorded pose to sample from to build the query
+* @param Scratch		Working memory to re-use for query construction
+* @param Query			Output pose search query. Query.Num() must be equal to Schema.Layout.NumFloats.
 */
-bool POSESEARCH_API PoseSearchBuildQuery(const UPoseSearchSchema& SearchSchema, int32 AssetSampleRate, const FPoseSearchPoseHistory& History, FPoseSearchBuildQueryScratch* Scratch, TArray<float>* Query);
+POSESEARCH_API bool BuildQuery(const UPoseSearchSchema* Schema, FPoseHistory* History, TArrayView<float> Query);
+
 
 /**
 * Performs a pose search
 *
-* @param SearchIndex The index to search within. See PoseSearchBuildIndex().
-* @param Query The pose query to search for. See PoseSearchBuildQuery().
-* @param DrawParams Visualization options
-* @return The pose in the SearchIndex that most closely matches the Query.
+* @param Sequence		The sequence to search within
+* @param Query			The pose query to search for. See BuildQuery().
+* @param DrawParams		Visualization options
+* 
+* @return				The pose in the sequence that most closely matches the Query.
 */
-FPoseSearchResult POSESEARCH_API PoseSearch(const UPoseSearchIndex& SearchIndex, TArrayView<const float> Query, FPoseSearchDebugDrawParams DrawParams = FPoseSearchDebugDrawParams());
+POSESEARCH_API FSearchResult Search(const UPoseSearchSequenceMetaData* Sequence, TArrayView<const float> Query, FDebugDrawParams DrawParams = FDebugDrawParams());
 
 
-//////////////////////////////////////////////////////////////////////////
-// FPoseSearchModule
+/**
+* Performs a pose search
+*
+* @param Database		The database to search within
+* @param Query			The pose query to search for. See BuildQuery().
+* @param DrawParams		Visualization options
+* 
+* @return				The pose in the database that most closely matches the Query.
+*/
+POSESEARCH_API FDbSearchResult Search(const UPoseSearchDatabase* Database, TArrayView<const float> Query, FDebugDrawParams DrawParams = FDebugDrawParams());
 
-class FPoseSearchModule : public IModuleInterface
-{
-public:
-	virtual void StartupModule() override;
-	virtual void ShutdownModule() override;
-};
\ No newline at end of file
+}} // namespace UE::PoseSearch
\ No newline at end of file
diff --git a/Engine/Source/Editor/AnimGraph/Classes/AnimGraphNode_SequencePlayer.h b/Engine/Source/Editor/AnimGraph/Classes/AnimGraphNode_SequencePlayer.h
--- a/Engine/Source/Editor/AnimGraph/Classes/AnimGraphNode_SequencePlayer.h
+++ b/Engine/Source/Editor/AnimGraph/Classes/AnimGraphNode_SequencePlayer.h
@@ -28,6 +28,7 @@ class UAnimGraphNode_SequencePlayer : public UAnimGraphNode_AssetPlayerBase
 	// End of UEdGraphNode interface
 
 	// UAnimGraphNode_Base interface
+	virtual void CustomizeDetails(IDetailLayoutBuilder& DetailBuilder) override;
 	virtual void CustomizePinData(UEdGraphPin* Pin, FName SourcePropertyName, int32 ArrayIndex) const override;
 	virtual void ValidateAnimNodeDuringCompilation(class USkeleton* ForSkeleton, class FCompilerResultsLog& MessageLog) override;
 	virtual void PreloadRequiredAssets() override;		
diff --git a/Engine/Source/Editor/AnimGraph/Private/AnimGraphNode_SequencePlayer.cpp b/Engine/Source/Editor/AnimGraph/Private/AnimGraphNode_SequencePlayer.cpp
--- a/Engine/Source/Editor/AnimGraph/Private/AnimGraphNode_SequencePlayer.cpp
+++ b/Engine/Source/Editor/AnimGraph/Private/AnimGraphNode_SequencePlayer.cpp
@@ -11,10 +11,12 @@
 #include "AssetRegistryModule.h"
 #include "BlueprintActionFilter.h"
 #include "BlueprintActionDatabaseRegistrar.h"
-#include "EditorCategoryUtils.h"
 #include "BlueprintNodeSpawner.h"
+#include "DetailLayoutBuilder.h"
+#include "EditorCategoryUtils.h"
 #include "Animation/AnimComposite.h"
 #include "Animation/AnimSequence.h"
+#include "Animation/AnimPoseSearchProvider.h"
 
 #define LOCTEXT_NAMESPACE "A3Nodes"
 
@@ -426,6 +428,16 @@ UScriptStruct* UAnimGraphNode_SequencePlayer::GetTimePropertyStruct() const
 	return FAnimNode_SequencePlayer::StaticStruct();
 }
 
+void UAnimGraphNode_SequencePlayer::CustomizeDetails(IDetailLayoutBuilder& DetailBuilder)
+{
+	Super::CustomizeDetails(DetailBuilder);
+
+	if (!UE::Anim::IPoseSearchProvider::IsAvailable())
+	{
+		DetailBuilder.HideCategory(TEXT("PoseMatching"));
+	}
+}
+
 void UAnimGraphNode_SequencePlayer::CustomizePinData(UEdGraphPin* Pin, FName SourcePropertyName, int32 ArrayIndex) const
 {
 	Super::CustomizePinData(Pin, SourcePropertyName, ArrayIndex);
diff --git a/Engine/Source/Runtime/Engine/Classes/Animation/AnimNode_SequencePlayer.h b/Engine/Source/Runtime/Engine/Classes/Animation/AnimNode_SequencePlayer.h
--- a/Engine/Source/Runtime/Engine/Classes/Animation/AnimNode_SequencePlayer.h
+++ b/Engine/Source/Runtime/Engine/Classes/Animation/AnimNode_SequencePlayer.h
@@ -43,6 +43,10 @@ public:
 	UPROPERTY(EditAnywhere, BlueprintReadWrite, Category = Settings, meta = (PinHiddenByDefault))
 	bool bLoopAnimation;
 
+	// Use pose matching to choose the start position. Requires experimental PoseSearch plugin.
+	UPROPERTY(EditAnywhere, BlueprintReadWrite, Category = PoseMatching, meta = (PinHiddenByDefault))
+	bool bStartFromMatchingPose;
+
 public:
 	FAnimNode_SequencePlayer()
 		: Sequence(nullptr)
@@ -50,6 +54,7 @@ public:
 		, PlayRate(1.0f)
 		, StartPosition(0.f)
 		, bLoopAnimation(true)
+		, bStartFromMatchingPose(false)
 	{
 	}
 
@@ -70,4 +75,5 @@ public:
 	// End of FAnimNode_Base interface
 
 	float GetTimeFromEnd(float CurrentNodeTime);
+	float GetEffectiveStartPosition(const FAnimationBaseContext& Context) const;
 };
diff --git a/Engine/Source/Runtime/Engine/Classes/Animation/AnimationAsset.h b/Engine/Source/Runtime/Engine/Classes/Animation/AnimationAsset.h
--- a/Engine/Source/Runtime/Engine/Classes/Animation/AnimationAsset.h
+++ b/Engine/Source/Runtime/Engine/Classes/Animation/AnimationAsset.h
@@ -914,6 +914,19 @@ public:
 	/** Get available Metadata within the animation asset
 	 */
 	const TArray<UAnimMetaData*>& GetMetaData() const { return MetaData; }
+	
+	/** Returns the first metadata of the specified class */
+	UAnimMetaData* FindMetaDataByClass(const TSubclassOf<UAnimMetaData> MetaDataClass) const;
+
+	/** Templatized version of FindMetaDataByClass that handles casting for you */
+	template<class T>
+	T* FindMetaDataByClass() const
+	{
+		static_assert(TPointerIsConvertibleFromTo<T, const UAnimMetaData>::Value, "'T' template parameter to FindMetaDataByClass must be derived from UAnimMetaData");
+
+		return (T*)FindMetaDataByClass(T::StaticClass());
+	}
+	
 	void AddMetaData(UAnimMetaData* MetaDataInstance); 
 	void EmptyMetaData() { MetaData.Empty(); }	
 	void RemoveMetaData(UAnimMetaData* MetaDataInstance);
diff --git a/Engine/Source/Runtime/Engine/Private/Animation/AnimNode_SequencePlayer.cpp b/Engine/Source/Runtime/Engine/Private/Animation/AnimNode_SequencePlayer.cpp
--- a/Engine/Source/Runtime/Engine/Private/Animation/AnimNode_SequencePlayer.cpp
+++ b/Engine/Source/Runtime/Engine/Private/Animation/AnimNode_SequencePlayer.cpp
@@ -5,6 +5,7 @@
 #include "AnimEncoding.h"
 #include "Animation/AnimInstanceProxy.h"
 #include "Animation/AnimTrace.h"
+#include "Animation/AnimPoseSearchProvider.h"
 
 #define LOCTEXT_NAMESPACE "AnimNode_SequencePlayer"
 
@@ -47,10 +48,12 @@ void FAnimNode_SequencePlayer::Initialize_AnyThread(const FAnimationInitializeCo
 
 	if (Sequence != nullptr)
 	{
-		InternalTimeAccumulator = FMath::Clamp(StartPosition, 0.f, Sequence->GetPlayLength());
+		const float EffectiveStartPosition = GetEffectiveStartPosition(Context);
+
+		InternalTimeAccumulator = FMath::Clamp(EffectiveStartPosition, 0.f, Sequence->GetPlayLength());
 		const float AdjustedPlayRate = PlayRateScaleBiasClamp.ApplyTo(FMath::IsNearlyZero(PlayRateBasis) ? 0.f : (PlayRate / PlayRateBasis), 0.f);
 		const float EffectivePlayrate = Sequence->RateScale * AdjustedPlayRate;
-		if ((StartPosition == 0.f) && (EffectivePlayrate < 0.f))
+		if ((EffectiveStartPosition == 0.f) && (EffectivePlayrate < 0.f))
 		{
 			InternalTimeAccumulator = Sequence->GetPlayLength();
 		}
@@ -138,4 +141,23 @@ float FAnimNode_SequencePlayer::GetTimeFromEnd(float CurrentNodeTime)
 	return Sequence->GetPlayLength() - CurrentNodeTime;
 }
 
+float FAnimNode_SequencePlayer::GetEffectiveStartPosition(const FAnimationBaseContext& Context) const
+{
+	// Override the start position if pose matching is enabled
+	if (Sequence != nullptr && bStartFromMatchingPose)
+	{
+		UE::Anim::IPoseSearchProvider* PoseSearchProvider = UE::Anim::IPoseSearchProvider::Get();
+		if (PoseSearchProvider)
+		{
+			UE::Anim::IPoseSearchProvider::FSearchResult Result = PoseSearchProvider->Search(Context, Sequence);
+			if (Result.PoseIdx >= 0)
+			{
+				return Result.TimeOffsetSeconds;
+			}
+		}
+	}
+
+	return StartPosition;
+}
+
 #undef LOCTEXT_NAMESPACE
\ No newline at end of file
diff --git a/Engine/Source/Runtime/Engine/Private/Animation/AnimPoseSearchProvider.cpp b/Engine/Source/Runtime/Engine/Private/Animation/AnimPoseSearchProvider.cpp
new file mode 100644
--- /dev/null
+++ b/Engine/Source/Runtime/Engine/Private/Animation/AnimPoseSearchProvider.cpp
@@ -0,0 +1,25 @@
+// Copyright Epic Games, Inc. All Rights Reserved.
+
+#include "Animation/AnimPoseSearchProvider.h"
+#include "Features/IModularFeatures.h"
+
+
+namespace UE { namespace Anim {
+
+const FName IPoseSearchProvider::ModularFeatureName(TEXT("AnimPoseSearch"));
+
+IPoseSearchProvider* IPoseSearchProvider::Get()
+{
+	if (IsAvailable())
+	{
+		return &IModularFeatures::Get().GetModularFeature<IPoseSearchProvider>(ModularFeatureName);
+	}
+	return nullptr;
+}
+
+bool IPoseSearchProvider::IsAvailable()
+{
+	return IModularFeatures::Get().IsModularFeatureAvailable(ModularFeatureName);
+}
+
+}} // namespace UE::Anim
\ No newline at end of file
diff --git a/Engine/Source/Runtime/Engine/Private/AnimationAsset.cpp b/Engine/Source/Runtime/Engine/Private/AnimationAsset.cpp
--- a/Engine/Source/Runtime/Engine/Private/AnimationAsset.cpp
+++ b/Engine/Source/Runtime/Engine/Private/AnimationAsset.cpp
@@ -3,6 +3,7 @@
 #include "Animation/AnimationAsset.h"
 #include "Engine/AssetUserData.h"
 #include "Animation/AssetMappingTable.h"
+#include "Animation/AnimMetaData.h"
 #include "Animation/AnimSequence.h"
 #include "AnimationUtils.h"
 #include "Animation/AnimInstance.h"
@@ -237,6 +238,25 @@ void UAnimationAsset::Serialize(FArchive& Ar)
 	}
 }
 
+UAnimMetaData* UAnimationAsset::FindMetaDataByClass(const TSubclassOf<UAnimMetaData> MetaDataClass) const
+{
+	UAnimMetaData* FoundMetaData = nullptr;
+
+	if (UClass* TargetClass = MetaDataClass.Get())
+	{
+		for (UAnimMetaData* MetaDataInstance : MetaData)
+		{
+			if (MetaDataInstance && MetaDataInstance->IsA(TargetClass))
+			{
+				FoundMetaData = MetaDataInstance;
+				break;
+			}
+		}
+	}
+
+	return FoundMetaData;
+}
+
 void UAnimationAsset::AddMetaData(UAnimMetaData* MetaDataInstance)
 {
 	MetaData.Add(MetaDataInstance);
diff --git a/Engine/Source/Runtime/Engine/Public/Animation/AnimPoseSearchProvider.h b/Engine/Source/Runtime/Engine/Public/Animation/AnimPoseSearchProvider.h
new file mode 100644
--- /dev/null
+++ b/Engine/Source/Runtime/Engine/Public/Animation/AnimPoseSearchProvider.h
@@ -0,0 +1,43 @@
+// Copyright Epic Games, Inc. All Rights Reserved.
+
+#pragma once
+
+#include "CoreTypes.h"
+#include "CoreMinimal.h"
+#include "Features/IModularFeature.h"
+
+struct FAnimationBaseContext;
+class UAnimSequenceBase;
+
+namespace UE { namespace Anim {
+
+/** Modular feature interface for PoseSearch */
+class ENGINE_API IPoseSearchProvider : public IModularFeature
+{
+public:
+	static const FName ModularFeatureName; // "AnimPoseSearch"
+
+	virtual ~IPoseSearchProvider() {}
+
+	static bool IsAvailable();
+	static IPoseSearchProvider* Get();
+
+	struct FSearchResult
+	{
+		int32 PoseIdx = -1;
+		float TimeOffsetSeconds = 0.0f;
+		float Dissimilarity = MAX_flt;
+	};
+
+	/**
+	* Finds a matching pose in the sequence given the current graph context
+	* 
+	* @param	GraphContext	Graph execution context used to construct a pose search query
+	* @param	Sequence		The sequence to search for the pose query
+	* 
+	* @return	The pose in the sequence that most closely matches the query
+	*/
+	virtual FSearchResult Search(const FAnimationBaseContext& GraphContext, const UAnimSequenceBase* Sequence) = 0;
+};
+
+}} // namespace UE::Anim
\ No newline at end of file
