Move Animation/PoseSearch plugin to Experimental

diff --git a/Engine/Plugins/Experimental/Animation/PoseSearch/PoseSearch.uplugin b/Engine/Plugins/Experimental/Animation/PoseSearch/PoseSearch.uplugin
new file mode 100644
--- /dev/null
+++ b/Engine/Plugins/Experimental/Animation/PoseSearch/PoseSearch.uplugin
@@ -0,0 +1,30 @@
+{
+	"FileVersion": 3,
+	"Version": 1,
+	"VersionName": "0.1",
+	"FriendlyName": "Pose Search",
+	"Description": "Experimental pose search API",
+	"Category": "Animation",
+	"CreatedBy": "Epic Games, Inc.",
+	"CreatedByURL": "http://epicgames.com",
+	"DocsURL": "",
+	"MarketplaceURL": "",
+	"SupportURL": "",
+	"EnabledByDefault": false,
+	"CanContainContent": true,
+	"IsBetaVersion": false,
+	"IsExperimentalVersion": true,
+	"Installed": false,
+	"Modules": [
+		{
+			"Name": "PoseSearch",
+			"Type": "Runtime",
+			"LoadingPhase": "PreDefault"
+		},
+		{
+			"Name": "PoseSearchEditor",
+			"Type": "UncookedOnly",
+			"LoadingPhase": "PreDefault"
+		}
+	]
+}
\ No newline at end of file
diff --git a/Engine/Plugins/Experimental/Animation/PoseSearch/Source/Editor/PoseSearchEditor.Build.cs b/Engine/Plugins/Experimental/Animation/PoseSearch/Source/Editor/PoseSearchEditor.Build.cs
new file mode 100644
--- /dev/null
+++ b/Engine/Plugins/Experimental/Animation/PoseSearch/Source/Editor/PoseSearchEditor.Build.cs
@@ -0,0 +1,33 @@
+// Copyright Epic Games, Inc. All Rights Reserved.
+
+using UnrealBuildTool;
+
+public class PoseSearchEditor : ModuleRules
+{
+	public PoseSearchEditor(ReadOnlyTargetRules Target) : base(Target)
+	{
+		PCHUsage = ModuleRules.PCHUsageMode.UseExplicitOrSharedPCHs;
+
+		PublicDependencyModuleNames.AddRange(
+			new string[]
+			{
+				"AnimGraph",
+				"AnimationCore",
+				"Core",
+				"CoreUObject",
+				"Engine",
+				"PoseSearch",
+				"ToolMenus",
+				"EditorFramework",
+				"UnrealEd",
+			}
+		);
+
+		PrivateDependencyModuleNames.AddRange(
+			new string[]
+			{
+				"BlueprintGraph",
+			}
+		);
+	}
+}
\ No newline at end of file
diff --git a/Engine/Plugins/Experimental/Animation/PoseSearch/Source/Editor/Private/AnimGraphNode_MotionMatching.cpp b/Engine/Plugins/Experimental/Animation/PoseSearch/Source/Editor/Private/AnimGraphNode_MotionMatching.cpp
new file mode 100644
--- /dev/null
+++ b/Engine/Plugins/Experimental/Animation/PoseSearch/Source/Editor/Private/AnimGraphNode_MotionMatching.cpp
@@ -0,0 +1,29 @@
+// Copyright Epic Games, Inc. All Rights Reserved.
+
+#include "AnimGraphNode_MotionMatching.h"
+
+#define LOCTEXT_NAMESPACE "AnimGraphNode_MotionMatching"
+
+
+FLinearColor UAnimGraphNode_MotionMatching::GetNodeTitleColor() const
+{
+	return FColor(86, 182, 194);
+}
+
+FText UAnimGraphNode_MotionMatching::GetTooltipText() const
+{
+	return LOCTEXT("NodeToolTip", "Motion Matching");
+}
+
+FText UAnimGraphNode_MotionMatching::GetNodeTitle(ENodeTitleType::Type TitleType) const
+{
+	return LOCTEXT("NodeTitle", "Motion Matching");
+}
+
+FText UAnimGraphNode_MotionMatching::GetMenuCategory() const
+{
+	return LOCTEXT("NodeCategory", "Pose Search");
+}
+
+
+#undef LOCTEXT_NAMESPACE
diff --git a/Engine/Plugins/Experimental/Animation/PoseSearch/Source/Editor/Private/AnimGraphNode_PoseSearchHistoryCollector.cpp b/Engine/Plugins/Experimental/Animation/PoseSearch/Source/Editor/Private/AnimGraphNode_PoseSearchHistoryCollector.cpp
new file mode 100644
--- /dev/null
+++ b/Engine/Plugins/Experimental/Animation/PoseSearch/Source/Editor/Private/AnimGraphNode_PoseSearchHistoryCollector.cpp
@@ -0,0 +1,29 @@
+// Copyright Epic Games, Inc. All Rights Reserved.
+
+#include "AnimGraphNode_PoseSearchHistoryCollector.h"
+
+#define LOCTEXT_NAMESPACE "AnimGraphNode_PoseSearchHistoryCollector"
+
+
+FLinearColor UAnimGraphNode_PoseSearchHistoryCollector::GetNodeTitleColor() const
+{
+	return FColor(86, 182, 194);
+}
+
+FText UAnimGraphNode_PoseSearchHistoryCollector::GetTooltipText() const
+{
+	return LOCTEXT("NodeToolTip", "Use this pose for matching");
+}
+
+FText UAnimGraphNode_PoseSearchHistoryCollector::GetNodeTitle(ENodeTitleType::Type TitleType) const
+{
+	return LOCTEXT("NodeTitle", "Match Pose");
+}
+
+FText UAnimGraphNode_PoseSearchHistoryCollector::GetMenuCategory() const
+{
+	return LOCTEXT("NodeCategory", "Pose Search");
+}
+
+
+#undef LOCTEXT_NAMESPACE
diff --git a/Engine/Plugins/Experimental/Animation/PoseSearch/Source/Editor/Private/PoseSearchEditor.cpp b/Engine/Plugins/Experimental/Animation/PoseSearch/Source/Editor/Private/PoseSearchEditor.cpp
new file mode 100644
--- /dev/null
+++ b/Engine/Plugins/Experimental/Animation/PoseSearch/Source/Editor/Private/PoseSearchEditor.cpp
@@ -0,0 +1,91 @@
+// Copyright Epic Games, Inc. All Rights Reserved.
+
+#include "PoseSearchEditor.h"
+#include "Animation/AnimSequence.h"
+#include "Modules/ModuleManager.h"
+#include "Editor.h"
+#include "Subsystems/AssetEditorSubsystem.h"
+#include "IAnimationEditor.h"
+#include "IPersonaToolkit.h"
+#include "IPersonaPreviewScene.h"
+
+
+//////////////////////////////////////////////////////////////////////////
+// FEditorCommands
+
+namespace UE { namespace PoseSearch {
+
+struct FEditorCommands
+{
+	static void DrawSearchIndex();
+};
+
+void FEditorCommands::DrawSearchIndex()
+{
+	FDebugDrawParams DrawParams;
+	DrawParams.DefaultLifeTime = 60.0f;
+	DrawParams.Flags = EDebugDrawFlags::DrawSearchIndex;
+
+	TArray<UObject*> EditedAssets = GEditor->GetEditorSubsystem<UAssetEditorSubsystem>()->GetAllEditedAssets();
+	for (UObject* EditedAsset : EditedAssets)
+	{
+		if (UAnimSequence* Sequence = Cast<UAnimSequence>(EditedAsset))
+		{
+			const UPoseSearchSequenceMetaData* MetaData = Sequence->FindMetaDataByClass<UPoseSearchSequenceMetaData>();
+			if (MetaData)
+			{
+				DrawParams.SequenceMetaData = MetaData;
+				IAssetEditorInstance* EditorInstance = GEditor->GetEditorSubsystem<UAssetEditorSubsystem>()->FindEditorForAsset(Sequence, true /*bFocusIfOpen*/);
+				if (EditorInstance && EditorInstance->GetEditorName() == TEXT("AnimationEditor"))
+				{
+					IAnimationEditor* Editor = static_cast<IAnimationEditor*>(EditorInstance);
+					TSharedRef<IPersonaToolkit> Toolkit = Editor->GetPersonaToolkit();
+					TSharedRef<IPersonaPreviewScene> Scene = Toolkit->GetPreviewScene();
+
+					DrawParams.World = Scene->GetWorld();
+					Draw(DrawParams);
+				}
+			}
+		}
+	}
+}
+
+
+//////////////////////////////////////////////////////////////////////////
+// FPoseSearchEditorModule
+
+class FEditorModule : public IPoseSearchEditorModuleInterface
+{
+public:
+	void StartupModule() override;
+	void ShutdownModule() override;
+
+private:
+	TArray<class IConsoleObject*> ConsoleCommands;
+};
+
+void FEditorModule::StartupModule()
+{
+	if (GIsEditor && !IsRunningCommandlet())
+	{
+		ConsoleCommands.Add(IConsoleManager::Get().RegisterConsoleCommand(
+			TEXT("a.PoseSearch.DrawSearchIndex"),
+			TEXT("Draw the search index for the selected asset"),
+			FConsoleCommandDelegate::CreateStatic(&FEditorCommands::DrawSearchIndex),
+			ECVF_Default
+		));
+	}
+}
+
+void FEditorModule::ShutdownModule()
+{
+	for (IConsoleObject* ConsoleCmd : ConsoleCommands)
+	{
+		IConsoleManager::Get().UnregisterConsoleObject(ConsoleCmd);
+	}
+	ConsoleCommands.Empty();
+}
+
+}} // namespace UE::PoseSearch
+
+IMPLEMENT_MODULE(UE::PoseSearch::FEditorModule, PoseSearchEditor);
diff --git a/Engine/Plugins/Experimental/Animation/PoseSearch/Source/Editor/Public/AnimGraphNode_MotionMatching.h b/Engine/Plugins/Experimental/Animation/PoseSearch/Source/Editor/Public/AnimGraphNode_MotionMatching.h
new file mode 100644
--- /dev/null
+++ b/Engine/Plugins/Experimental/Animation/PoseSearch/Source/Editor/Public/AnimGraphNode_MotionMatching.h
@@ -0,0 +1,24 @@
+// Copyright Epic Games, Inc. All Rights Reserved.
+
+#pragma once
+
+#include "CoreMinimal.h"
+#include "UObject/ObjectMacros.h"
+#include "AnimGraphNode_Base.h"
+#include "PoseSearch/AnimNode_MotionMatching.h"
+#include "AnimGraphNode_MotionMatching.generated.h"
+
+
+UCLASS(MinimalAPI)
+class UAnimGraphNode_MotionMatching : public UAnimGraphNode_Base
+{
+	GENERATED_BODY()
+
+	UPROPERTY(EditAnywhere, Category = Settings)
+	FAnimNode_MotionMatching Node;
+
+	virtual FLinearColor GetNodeTitleColor() const override;
+	virtual FText GetTooltipText() const override;
+	virtual FText GetNodeTitle(ENodeTitleType::Type TitleType) const override;
+	virtual FText GetMenuCategory() const override;
+};
diff --git a/Engine/Plugins/Experimental/Animation/PoseSearch/Source/Editor/Public/AnimGraphNode_PoseSearchHistoryCollector.h b/Engine/Plugins/Experimental/Animation/PoseSearch/Source/Editor/Public/AnimGraphNode_PoseSearchHistoryCollector.h
new file mode 100644
--- /dev/null
+++ b/Engine/Plugins/Experimental/Animation/PoseSearch/Source/Editor/Public/AnimGraphNode_PoseSearchHistoryCollector.h
@@ -0,0 +1,24 @@
+// Copyright Epic Games, Inc. All Rights Reserved.
+
+#pragma once
+
+#include "CoreMinimal.h"
+#include "UObject/ObjectMacros.h"
+#include "AnimGraphNode_Base.h"
+#include "PoseSearch/AnimNode_PoseSearchHistoryCollector.h"
+#include "AnimGraphNode_PoseSearchHistoryCollector.generated.h"
+
+
+UCLASS(MinimalAPI)
+class UAnimGraphNode_PoseSearchHistoryCollector : public UAnimGraphNode_Base
+{
+	GENERATED_BODY()
+
+	UPROPERTY(EditAnywhere, Category = Settings)
+	FAnimNode_PoseSearchHistoryCollector Node;
+
+	virtual FLinearColor GetNodeTitleColor() const override;
+	virtual FText GetTooltipText() const override;
+	virtual FText GetNodeTitle(ENodeTitleType::Type TitleType) const override;
+	virtual FText GetMenuCategory() const override;
+};
diff --git a/Engine/Plugins/Experimental/Animation/PoseSearch/Source/Editor/Public/PoseSearchEditor.h b/Engine/Plugins/Experimental/Animation/PoseSearch/Source/Editor/Public/PoseSearchEditor.h
new file mode 100644
--- /dev/null
+++ b/Engine/Plugins/Experimental/Animation/PoseSearch/Source/Editor/Public/PoseSearchEditor.h
@@ -0,0 +1,11 @@
+// Copyright Epic Games, Inc. All Rights Reserved.
+
+#pragma once
+
+#include "CoreMinimal.h"
+#include "Modules/ModuleInterface.h"
+#include "PoseSearch/PoseSearch.h"
+
+class IPoseSearchEditorModuleInterface : public IModuleInterface
+{
+};
\ No newline at end of file
diff --git a/Engine/Plugins/Experimental/Animation/PoseSearch/Source/Runtime/PoseSearch.Build.cs b/Engine/Plugins/Experimental/Animation/PoseSearch/Source/Runtime/PoseSearch.Build.cs
new file mode 100644
--- /dev/null
+++ b/Engine/Plugins/Experimental/Animation/PoseSearch/Source/Runtime/PoseSearch.Build.cs
@@ -0,0 +1,44 @@
+// Copyright Epic Games, Inc. All Rights Reserved.
+
+using UnrealBuildTool;
+
+public class PoseSearch : ModuleRules
+{
+	public PoseSearch(ReadOnlyTargetRules Target) : base(Target)
+	{
+		PCHUsage = ModuleRules.PCHUsageMode.UseExplicitOrSharedPCHs;
+
+		PublicIncludePaths.AddRange(
+			new string[]
+			{
+				System.IO.Path.Combine(ModuleDirectory, "Public"),
+			}
+		);
+
+		PrivateIncludePaths.AddRange(
+			new string[]
+			{
+				System.IO.Path.Combine(ModuleDirectory, "Private"),
+			}
+		);
+
+		PublicDependencyModuleNames.AddRange(
+			new string[]
+			{
+				"Core",
+				"CoreUObject",
+				"Engine",
+				"AnimationCore"
+			}
+		);
+
+		PrivateDependencyModuleNames.AddRange(
+			new string[]
+            {
+				"Core",
+				"CoreUObject",
+				"Engine",
+            }
+		);
+	}
+}
diff --git a/Engine/Plugins/Experimental/Animation/PoseSearch/Source/Runtime/Private/AnimNode_MotionMatching.cpp b/Engine/Plugins/Experimental/Animation/PoseSearch/Source/Runtime/Private/AnimNode_MotionMatching.cpp
new file mode 100644
--- /dev/null
+++ b/Engine/Plugins/Experimental/Animation/PoseSearch/Source/Runtime/Private/AnimNode_MotionMatching.cpp
@@ -0,0 +1,203 @@
+// Copyright Epic Games, Inc. All Rights Reserved.
+
+#include "PoseSearch/AnimNode_MotionMatching.h"
+#include "PoseSearch/PoseSearch.h"
+#include "PoseSearch/AnimNode_PoseSearchHistoryCollector.h"
+#include "Animation/AnimInstanceProxy.h"
+#include "Animation/AnimNode_Inertialization.h"
+#include "Animation/AnimSequence.h"
+
+#define LOCTEXT_NAMESPACE "AnimNode_MotionMatching"
+
+/////////////////////////////////////////////////////
+// FAnimNode_MotionMatching
+
+void FAnimNode_MotionMatching::Initialize_AnyThread(const FAnimationInitializeContext& Context)
+{
+	DECLARE_SCOPE_HIERARCHICAL_COUNTER_ANIMNODE(Initialize_AnyThread);
+
+	GetEvaluateGraphExposedInputs().Execute(Context);
+
+	if (IsValidForSearch())
+	{
+		Query.Reset();
+		Query.SetNumZeroed(Database->Schema->Layout.NumFloats);
+		QueryBuilder.Init(&Database->Schema->Layout, Query);
+
+		// Set initial pose features
+		UE::PoseSearch::IPoseHistoryProvider* PoseHistoryProvider = Context.GetMessage<UE::PoseSearch::IPoseHistoryProvider>();
+		if (PoseHistoryProvider)
+		{
+			UE::PoseSearch::FPoseHistory& History = PoseHistoryProvider->GetPoseHistory();
+			QueryBuilder.SetPoseFeatures(Database->Schema, &History);
+		}
+	}
+
+	DbPoseIdx = INDEX_NONE;
+
+	Source.SetLinkNode(&SequencePlayerNode);
+	Source.Initialize(Context);
+}
+
+void FAnimNode_MotionMatching::Evaluate_AnyThread(FPoseContext& Output)
+{
+	DECLARE_SCOPE_HIERARCHICAL_COUNTER_ANIMNODE(Evaluate_AnyThread);
+
+	Source.Evaluate(Output);
+}
+
+void FAnimNode_MotionMatching::Update_AnyThread(const FAnimationUpdateContext& Context)
+{
+	DECLARE_SCOPE_HIERARCHICAL_COUNTER_ANIMNODE(Update_AnyThread);
+
+	GetEvaluateGraphExposedInputs().Execute(Context);
+	// Note: What if the input database changes? That's not being handled at all!
+
+	if (IsValidForSearch())
+	{
+		float AssetTime = SequencePlayerNode.GetAccumulatedTime();
+
+		// Step the pose forward
+		if (DbPoseIdx != INDEX_NONE)
+		{
+			check(DbSequenceIdx != INDEX_NONE);
+
+			// Determine roughly which pose we're playing from the database based on the sequence player time
+			DbPoseIdx = Database->GetPoseIndexFromAssetTime(DbSequenceIdx, AssetTime);
+
+			// Overwrite query feature vector with the stepped pose from the database
+			if (DbPoseIdx != INDEX_NONE)
+			{
+				TArrayView<const float> PoseFeatureVector = Database->SearchIndex.GetPoseValues(DbPoseIdx);
+				QueryBuilder.Copy(PoseFeatureVector);
+			}
+		}
+
+		// Update trajectory features in the query with the latest inputs
+		SetTrajectoryFeatures();
+
+		// Determine how much the updated query vector deviates from the current pose vector
+		float CurrentDissimilarity = MAX_flt;
+		if (DbPoseIdx != INDEX_NONE)
+		{
+			CurrentDissimilarity = UE::PoseSearch::ComparePoses(Database->SearchIndex, DbPoseIdx, Query);
+		}
+
+		// Search the database for the nearest match to the updated query vector
+		UE::PoseSearch::FDbSearchResult Result = UE::PoseSearch::Search(Database, MakeArrayView(Query));
+		if (Result.IsValid())
+		{
+			const FPoseSearchDatabaseSequence& ResultDbSequence = Database->Sequences[Result.DbSequenceIdx];
+
+			// Consider the search result better if it is more similar to the query than the current pose we're playing back from the database
+			bool bBetterPose = Result.Dissimilarity < CurrentDissimilarity;
+
+			// We'll ignore the candidate pose if it is too near to our current pose
+			bool bNearbyPose = false;
+			if (DbSequenceIdx == Result.DbSequenceIdx)
+			{
+				bNearbyPose = FMath::Abs(AssetTime - Result.TimeOffsetSeconds) < PoseJumpThreshold;
+				if (!bNearbyPose && ResultDbSequence.bLoopAnimation)
+				{
+					bNearbyPose = FMath::Abs(SequencePlayerNode.GetCurrentAssetLength() - AssetTime - Result.TimeOffsetSeconds) < PoseJumpThreshold;
+				}
+			}
+
+			// And we won't bother to jump to another pose within the same looping sequence we're already playing
+			// (This should probably should be configurable)
+			bool bSameCycle = (DbSequenceIdx == Result.DbSequenceIdx) && ResultDbSequence.bLoopAnimation;
+
+			// Start playback from the candidate pose if we determined it was a better option
+			if (bBetterPose && !bNearbyPose && !bSameCycle)
+			{
+				JumpToPose(Context, Result);
+			}
+		}
+	}
+
+	Source.Update(Context);
+}
+
+bool FAnimNode_MotionMatching::HasPreUpdate() const
+{
+#if WITH_EDITORONLY_DATA
+	return true;
+#else
+	return false;
+#endif
+}
+
+void FAnimNode_MotionMatching::PreUpdate(const UAnimInstance* InAnimInstance)
+{
+#if WITH_EDITORONLY_DATA
+	if (bDebugDraw)
+	{
+		USkeletalMeshComponent* SkeletalMeshComponent = InAnimInstance->GetSkelMeshComponent();
+		check(SkeletalMeshComponent);
+
+		UE::PoseSearch::FDebugDrawParams DrawParams;
+		DrawParams.RootTransform = SkeletalMeshComponent->GetBoneTransform(0);
+		DrawParams.Flags = UE::PoseSearch::EDebugDrawFlags::DrawQuery;
+		DrawParams.Database = Database;
+		DrawParams.Query = Query;
+		DrawParams.World = SkeletalMeshComponent->GetWorld();
+		DrawParams.DefaultLifeTime = 0.0f;
+		UE::PoseSearch::Draw(DrawParams);
+	}
+#endif
+}
+
+void FAnimNode_MotionMatching::GatherDebugData(FNodeDebugData& DebugData)
+{
+	Source.GatherDebugData(DebugData);
+}
+
+bool FAnimNode_MotionMatching::IsValidForSearch() const
+{
+	return Database && Database->IsValidForSearch();
+}
+
+void FAnimNode_MotionMatching::SetTrajectoryFeatures()
+{
+	// This is all placeholder-- trajectory should be provided externally rather than computed in this node.
+	// And we shouldn't be assuming the format of the database's schema!
+
+	// Add instantaneous root velocity to query
+	FPoseSearchFeatureDesc TrajectoryFeature;
+	TrajectoryFeature.Domain = EPoseSearchFeatureDomain::Time;
+	TrajectoryFeature.SchemaBoneIdx = FPoseSearchFeatureDesc::TrajectoryBoneIndex;
+	TrajectoryFeature.SubsampleIdx = 0;
+	TrajectoryFeature.Type = EPoseSearchFeatureType::LinearVelocity;
+
+	// I'm sure there is probably a less magical way to get LocalVelocity into the same space as the animation data...
+	FQuat Rotate(FVector::ZAxisVector, HALF_PI);
+
+	// Apply the updated trajectory to the query feature vector, leaving any existing pose information intact
+	QueryBuilder.SetVector(TrajectoryFeature, Rotate.RotateVector(LocalVelocity));
+}
+
+void FAnimNode_MotionMatching::JumpToPose(const FAnimationUpdateContext& Context, UE::PoseSearch::FDbSearchResult Result)
+{
+	// Remember which pose and sequence we're playing from the database
+	DbPoseIdx = Result.PoseIdx;
+	DbSequenceIdx = Result.DbSequenceIdx;
+
+	// Immediately jump to the pose by updating the embedded sequence player node
+	const FPoseSearchDatabaseSequence& ResultDbSequence = Database->Sequences[Result.DbSequenceIdx];
+	SequencePlayerNode.Sequence = ResultDbSequence.Sequence;
+	SequencePlayerNode.SetAccumulatedTime(Result.TimeOffsetSeconds);
+	SequencePlayerNode.bLoopAnimation = ResultDbSequence.bLoopAnimation;
+
+	// Use inertial blending to smooth over the transition
+	// It would be cool in the future to adjust the blend time by amount of dissimilarity, but we'll need a standardized distance metric first.
+	if (BlendTime > 0.0f)
+	{
+		UE::Anim::IInertializationRequester* InertializationRequester = Context.GetMessage<UE::Anim::IInertializationRequester>();
+		if (InertializationRequester)
+		{
+			InertializationRequester->RequestInertialization(BlendTime);
+		}
+	}
+}
+
+#undef LOCTEXT_NAMESPACE
\ No newline at end of file
diff --git a/Engine/Plugins/Experimental/Animation/PoseSearch/Source/Runtime/Private/AnimNode_PoseSearchHistoryCollector.cpp b/Engine/Plugins/Experimental/Animation/PoseSearch/Source/Runtime/Private/AnimNode_PoseSearchHistoryCollector.cpp
new file mode 100644
--- /dev/null
+++ b/Engine/Plugins/Experimental/Animation/PoseSearch/Source/Runtime/Private/AnimNode_PoseSearchHistoryCollector.cpp
@@ -0,0 +1,81 @@
+// Copyright Epic Games, Inc. All Rights Reserved.
+
+#include "PoseSearch/AnimNode_PoseSearchHistoryCollector.h"
+#include "Animation/AnimInstanceProxy.h"
+#include "Animation/AnimNodeMessages.h"
+
+#define LOCTEXT_NAMESPACE "AnimNode_PoseSearchHistoryCollector"
+
+namespace UE { namespace PoseSearch { namespace Private {
+
+class FPoseHistoryProvider : public IPoseHistoryProvider
+{
+public:
+	FPoseHistoryProvider(FAnimNode_PoseSearchHistoryCollector* InNode)
+		: Node(*InNode)
+	{}
+
+	// IPoseHistoryProvider interface
+	virtual const FPoseHistory& GetPoseHistory() const override
+	{
+		return Node.GetPoseHistory();
+	}
+
+	virtual FPoseHistory& GetPoseHistory() override
+	{
+		return Node.GetPoseHistory();
+	}
+
+	// Node we wrap
+	FAnimNode_PoseSearchHistoryCollector& Node;
+};
+
+}}} // namespace UE::PoseSearch::Private
+
+/////////////////////////////////////////////////////
+// FAnimNode_PoseSearchHistoryCollector
+
+void FAnimNode_PoseSearchHistoryCollector::Initialize_AnyThread(const FAnimationInitializeContext& Context)
+{
+	DECLARE_SCOPE_HIERARCHICAL_COUNTER_ANIMNODE(Initialize_AnyThread);
+
+	Super::Initialize_AnyThread(Context);
+
+	// @@ need to do this once based on descendant node's (or input param?) search schema, not every node init
+	PoseHistory.Init(32, 1.0f);
+
+	UE::Anim::TScopedGraphMessage<UE::PoseSearch::Private::FPoseHistoryProvider> ScopedMessage(Context, this);
+
+	Source.Initialize(Context);
+
+	EvalDeltaTime = 0.0f;
+}
+
+void FAnimNode_PoseSearchHistoryCollector::Evaluate_AnyThread(FPoseContext& Output)
+{
+	DECLARE_SCOPE_HIERARCHICAL_COUNTER_ANIMNODE(Evaluate_AnyThread);
+
+	Source.Evaluate(Output);
+
+	PoseHistory.Update(EvalDeltaTime, Output.Pose);
+
+	EvalDeltaTime = 0.0f;
+}
+
+void FAnimNode_PoseSearchHistoryCollector::Update_AnyThread(const FAnimationUpdateContext& Context)
+{
+	DECLARE_SCOPE_HIERARCHICAL_COUNTER_ANIMNODE(Update_AnyThread);
+
+	UE::Anim::TScopedGraphMessage<UE::PoseSearch::Private::FPoseHistoryProvider> ScopedMessage(Context, this);
+
+	Source.Update(Context);
+
+	EvalDeltaTime += Context.GetDeltaTime();
+}
+
+void FAnimNode_PoseSearchHistoryCollector::GatherDebugData(FNodeDebugData& DebugData)
+{
+	Source.GatherDebugData(DebugData);
+}
+
+#undef LOCTEXT_NAMESPACE
\ No newline at end of file
diff --git a/Engine/Plugins/Experimental/Animation/PoseSearch/Source/Runtime/Private/PoseSearch.cpp b/Engine/Plugins/Experimental/Animation/PoseSearch/Source/Runtime/Private/PoseSearch.cpp
new file mode 100644
--- /dev/null
+++ b/Engine/Plugins/Experimental/Animation/PoseSearch/Source/Runtime/Private/PoseSearch.cpp
@@ -0,0 +1,1694 @@
+// Copyright Epic Games, Inc. All Rights Reserved.
+
+#include "PoseSearch/PoseSearch.h"
+
+#include "Async/ParallelFor.h"
+#include "Templates/IdentityFunctor.h"
+#include "Templates/Invoke.h"
+#include "Templates/Less.h"
+#include "Features/IModularFeatures.h"
+#include "DrawDebugHelpers.h"
+#include "Animation/AnimPoseSearchProvider.h"
+#include "Animation/AnimCurveTypes.h"
+#include "Animation/AnimInstance.h"
+#include "Animation/AnimNodeBase.h"
+#include "Animation/AnimMetaData.h"
+#include "Animation/AnimNodeBase.h"
+#include "Animation/AnimSequence.h"
+#include "Animation/AnimSequenceBase.h"
+#include "Animation/AnimationPoseData.h"
+#include "AnimationRuntime.h"
+#include "BonePose.h"
+
+IMPLEMENT_ANIMGRAPH_MESSAGE(UE::PoseSearch::IPoseHistoryProvider);
+
+namespace UE { namespace PoseSearch {
+
+//////////////////////////////////////////////////////////////////////////
+// Constants and utilities
+
+constexpr float DrawDebugLineThickness = 1.0f;
+constexpr float DrawDebugPointSize = 3.0f;
+constexpr float DrawDebugVelocityScale = 0.1f;
+constexpr float DrawDebugArrowSize = 5.0f;
+constexpr float DrawDebugSphereSize = 1.0f;
+constexpr int32 DrawDebugSphereSegments = 8;
+constexpr float DrawDebugSphereLineThickness = 0.2f;
+
+static bool IsSamplingRangeValid(FFloatInterval Range)
+{
+	return Range.IsValid() && (Range.Min >= 0.0f);
+}
+
+static FFloatInterval GetEffectiveSamplingRange(const UAnimSequenceBase* Sequence, FFloatInterval SamplingRange)
+{
+	const bool bSampleAll = (SamplingRange.Min == 0.0f) && (SamplingRange.Max == 0.0f);
+	const float SequencePlayLength = Sequence->GetPlayLength();
+
+	FFloatInterval Range;
+	Range.Min = bSampleAll ? 0.0f : SamplingRange.Min;
+	Range.Max = bSampleAll ? SequencePlayLength : FMath::Min(SequencePlayLength, SamplingRange.Max);
+	return Range;
+}
+
+static float CompareFeatureVectors(int32 NumValues, const float* A, const float* B)
+{
+	float Dissimilarity = 0.0f;
+
+	for (int32 ValueIdx = 0; ValueIdx != NumValues; ++ValueIdx)
+	{
+		float Diff = A[ValueIdx] - B[ValueIdx];
+		Dissimilarity += Diff * Diff;
+	}
+
+	return Dissimilarity;
+}
+
+FLinearColor GetColorForFeature(FPoseSearchFeatureDesc Feature, const FPoseSearchFeatureVectorLayout* Layout)
+{
+	int32 FeatureIdx = Layout->Features.IndexOfByKey(Feature);
+	check(FeatureIdx != INDEX_NONE);
+	float Lerp = (float)(FeatureIdx) / (Layout->Features.Num() - 1);
+	FLinearColor ColorHSV(Lerp * 360.0f, 0.8f, 0.5f, 0.6f);
+	return ColorHSV.HSVToLinearRGB();
+}
+
+/**
+* Algo::LowerBound adapted to TIndexedContainerIterator for use with indexable but not necessarily contiguous containers. Used here with TRingBuffer.
+*
+* Performs binary search, resulting in position of the first element >= Value using predicate
+*
+* @param First TIndexedContainerIterator beginning of range to search through, must be already sorted by SortPredicate
+* @param Last TIndexedContainerIterator end of range
+* @param Value Value to look for
+* @param SortPredicate Predicate for sort comparison, defaults to <
+*
+* @returns Position of the first element >= Value, may be position after last element in range
+*/
+template <typename IteratorType, typename ValueType, typename ProjectionType, typename SortPredicateType>
+FORCEINLINE auto LowerBound(IteratorType First, IteratorType Last, const ValueType& Value, ProjectionType Projection, SortPredicateType SortPredicate) -> decltype(First.GetIndex())
+{
+	using SizeType = decltype(First.GetIndex());
+
+	check(First.GetIndex() <= Last.GetIndex());
+
+	// Current start of sequence to check
+	SizeType Start = First.GetIndex();
+
+	// Size of sequence to check
+	SizeType Size = Last.GetIndex() - Start;
+
+	// With this method, if Size is even it will do one more comparison than necessary, but because Size can be predicted by the CPU it is faster in practice
+	while (Size > 0)
+	{
+		const SizeType LeftoverSize = Size % 2;
+		Size = Size / 2;
+
+		const SizeType CheckIndex = Start + Size;
+		const SizeType StartIfLess = CheckIndex + LeftoverSize;
+
+		auto&& CheckValue = Invoke(Projection, *(First + CheckIndex));
+		Start = SortPredicate(CheckValue, Value) ? StartIfLess : Start;
+	}
+	return Start;
+}
+
+template <typename IteratorType, typename ValueType, typename SortPredicateType>
+FORCEINLINE auto LowerBound(IteratorType First, IteratorType Last, const ValueType& Value, SortPredicateType SortPredicate) -> decltype(First.GetIndex())
+{
+	return LowerBound(First, Last, Value, FIdentityFunctor(), SortPredicate);
+}
+
+template <typename IteratorType, typename ValueType>
+FORCEINLINE auto LowerBound(IteratorType First, IteratorType Last, const ValueType& Value) -> decltype(First.GetIndex())
+{
+	return LowerBound(First, Last, Value, FIdentityFunctor(), TLess<>());
+}
+
+
+//////////////////////////////////////////////////////////////////////////
+// FFeatureTypeTraits
+
+struct FFeatureTypeTraits
+{
+	EPoseSearchFeatureType Type = EPoseSearchFeatureType::Invalid;
+	uint32 NumFloats = 0;
+};
+
+// Could upgrade to class objects in the future with value reader/writer functions
+static constexpr FFeatureTypeTraits FeatureTypeTraits[] =
+{
+	{ EPoseSearchFeatureType::Position, 3 },
+	{ EPoseSearchFeatureType::Rotation, 6 },
+	{ EPoseSearchFeatureType::LinearVelocity, 3 },
+	{ EPoseSearchFeatureType::AngularVelocity, 3 },
+};
+
+FFeatureTypeTraits GetFeatureTypeTraits(EPoseSearchFeatureType Type)
+{
+	// Could allow external registration to a TSet of traits in the future
+	// For now just use a simple local array
+	for (const FFeatureTypeTraits& Traits : FeatureTypeTraits)
+	{
+		if (Traits.Type == Type)
+		{
+			return Traits;
+		}
+	}
+
+	return FFeatureTypeTraits();
+}
+
+}} // namespace UE::PoseSearch
+
+
+//////////////////////////////////////////////////////////////////////////
+// FPoseSearchFeatureDesc
+
+bool FPoseSearchFeatureDesc::operator==(const FPoseSearchFeatureDesc& Other) const
+{
+	return
+		(SchemaBoneIdx == Other.SchemaBoneIdx) &&
+		(SubsampleIdx == Other.SubsampleIdx) &&
+		(Type == Other.Type) &&
+		(Domain == Other.Domain);
+}
+
+
+//////////////////////////////////////////////////////////////////////////
+// FPoseSearchFeatureVectorLayout
+
+void FPoseSearchFeatureVectorLayout::Init()
+{
+	uint32 FloatCount = 0;
+
+	for (FPoseSearchFeatureDesc& Feature : Features)
+	{
+		Feature.ValueOffset = FloatCount;
+		FloatCount += UE::PoseSearch::GetFeatureTypeTraits(Feature.Type).NumFloats;
+	}
+
+	NumFloats = FloatCount;
+}
+
+void FPoseSearchFeatureVectorLayout::Reset()
+{
+	Features.Reset();
+	NumFloats = 0;
+}
+
+bool FPoseSearchFeatureVectorLayout::IsValid(int32 MaxNumBones) const
+{
+	if (NumFloats == 0.0f)
+	{
+		return false;
+	}
+
+	for (const FPoseSearchFeatureDesc& Feature : Features)
+	{
+		if (Feature.SchemaBoneIdx >= MaxNumBones)
+		{
+			return false;
+		}
+	}
+
+	return true;
+}
+
+
+//////////////////////////////////////////////////////////////////////////
+// UPoseSearchSchema
+
+void UPoseSearchSchema::PreSave(const class ITargetPlatform* TargetPlatform)
+{
+	SamplingInterval = 1.0f / SampleRate;
+
+	PoseSampleOffsets.Sort(TLess<>());
+	TrajectorySampleOffsets.Sort(TLess<>());
+	TrajectoryDistanceOffsets.Sort(TLess<>());
+
+	GenerateLayout();
+	ResolveBoneReferences();
+
+	Super::PreSave(TargetPlatform);
+}
+
+void UPoseSearchSchema::PostLoad()
+{
+	Super::PostLoad();
+
+	ResolveBoneReferences();
+}
+
+bool UPoseSearchSchema::IsValid() const
+{
+	bool bValid = Skeleton != nullptr;
+
+	for (const FBoneReference& BoneRef : Bones)
+	{
+		bValid &= BoneRef.HasValidSetup();
+	}
+
+	bValid &= (Bones.Num() == BoneIndices.Num());
+
+	bValid &= Layout.IsValid(BoneIndices.Num());
+	
+	return bValid;
+}
+
+void UPoseSearchSchema::GenerateLayout()
+{
+	Layout.Reset();
+
+	for (int32 TrajectoryTimeSubsampleIdx = 0; TrajectoryTimeSubsampleIdx != TrajectorySampleOffsets.Num(); ++TrajectoryTimeSubsampleIdx)
+	{
+		FPoseSearchFeatureDesc Element;
+		Element.SchemaBoneIdx = FPoseSearchFeatureDesc::TrajectoryBoneIndex;
+		Element.SubsampleIdx = TrajectoryTimeSubsampleIdx;
+		Element.Domain = EPoseSearchFeatureDomain::Time;
+
+		if (bUseTrajectoryPositions)
+		{
+			Element.Type = EPoseSearchFeatureType::Position;
+			Layout.Features.Add(Element);
+		}
+
+		if (bUseTrajectoryVelocities)
+		{
+			Element.Type = EPoseSearchFeatureType::LinearVelocity;
+			Layout.Features.Add(Element);
+		}
+	}
+
+ 	for (int32 TrajectoryDistSubsampleIdx = 0; TrajectoryDistSubsampleIdx != TrajectoryDistanceOffsets.Num(); ++TrajectoryDistSubsampleIdx)
+ 	{
+ 		FPoseSearchFeatureDesc Element;
+ 		Element.SchemaBoneIdx = FPoseSearchFeatureDesc::TrajectoryBoneIndex;
+ 		Element.SubsampleIdx = TrajectoryDistSubsampleIdx;
+ 		Element.Domain = EPoseSearchFeatureDomain::Distance;
+
+		if (bUseTrajectoryPositions)
+		{
+			Element.Type = EPoseSearchFeatureType::Position;
+			Layout.Features.Add(Element);
+		}
+
+		if (bUseTrajectoryVelocities)
+		{
+			Element.Type = EPoseSearchFeatureType::LinearVelocity;
+			Layout.Features.Add(Element);
+		}
+ 	}
+
+	for (int32 PoseSubsampleIdx = 0; PoseSubsampleIdx != PoseSampleOffsets.Num(); ++PoseSubsampleIdx)
+	{
+		FPoseSearchFeatureDesc Element;
+		Element.SubsampleIdx = PoseSubsampleIdx;
+		Element.Domain = EPoseSearchFeatureDomain::Time;
+
+		for (int32 SchemaBoneIdx = 0; SchemaBoneIdx != Bones.Num(); ++SchemaBoneIdx)
+		{
+			Element.SchemaBoneIdx = SchemaBoneIdx;
+			if (bUseBonePositions)
+			{
+				Element.Type = EPoseSearchFeatureType::Position;
+				Layout.Features.Add(Element);
+			}
+
+			if (bUseBoneVelocities)
+			{
+				Element.Type = EPoseSearchFeatureType::LinearVelocity;
+				Layout.Features.Add(Element);
+			}
+		}
+	}
+
+	Layout.Init();
+}
+
+void UPoseSearchSchema::ResolveBoneReferences()
+{
+	// Initialize references to obtain bone indices
+	for (FBoneReference& BoneRef : Bones)
+	{
+		BoneRef.Initialize(Skeleton);
+	}
+
+	// Fill out bone index array and sort by bone index
+	BoneIndices.SetNum(Bones.Num());
+	for (int32 Index = 0; Index != Bones.Num(); ++Index)
+	{
+		BoneIndices[Index] = Bones[Index].BoneIndex;
+	}
+	BoneIndices.Sort();
+
+	// Build separate index array with parent indices guaranteed to be present
+	BoneIndicesWithParents = BoneIndices;
+	if (Skeleton)
+	{
+		FAnimationRuntime::EnsureParentsPresent(BoneIndicesWithParents, Skeleton->GetReferenceSkeleton());
+	}
+}
+
+
+//////////////////////////////////////////////////////////////////////////
+// FPoseSearchIndex
+
+bool FPoseSearchIndex::IsValid() const
+{
+	bool bSchemaValid = Schema && Schema->IsValid();
+	bool bSearchIndexValid = bSchemaValid && (NumPoses * Schema->Layout.NumFloats == Values.Num());
+
+	return bSearchIndexValid;
+}
+
+TArrayView<const float> FPoseSearchIndex::GetPoseValues(int32 PoseIdx) const
+{
+	check(PoseIdx < NumPoses);
+	int32 ValueOffset = PoseIdx * Schema->Layout.NumFloats;
+	return MakeArrayView(&Values[ValueOffset], Schema->Layout.NumFloats);
+}
+
+void FPoseSearchIndex::Reset()
+{
+	NumPoses = 0;
+	Values.Reset();
+	Schema = nullptr;
+}
+
+
+//////////////////////////////////////////////////////////////////////////
+// UPoseSearchSequenceMetaData
+
+void UPoseSearchSequenceMetaData::PreSave(const class ITargetPlatform* TargetPlatform)
+{
+	SearchIndex.Reset();
+
+	if (IsValidForIndexing())
+	{
+		UObject* Outer = GetOuter();
+		if (UAnimSequence* Sequence = Cast<UAnimSequence>(Outer))
+		{
+			UE::PoseSearch::BuildIndex(Sequence, this);
+		}
+	}
+
+	Super::PreSave(TargetPlatform);
+}
+
+bool UPoseSearchSequenceMetaData::IsValidForIndexing() const
+{
+	return Schema && Schema->IsValid() && UE::PoseSearch::IsSamplingRangeValid(SamplingRange);
+}
+
+bool UPoseSearchSequenceMetaData::IsValidForSearch() const
+{
+	return IsValidForIndexing() && SearchIndex.IsValid();
+}
+
+
+//////////////////////////////////////////////////////////////////////////
+// UPoseSearchDatabase
+
+int32 UPoseSearchDatabase::FindSequenceForPose(int32 PoseIdx) const
+{
+	auto Predicate = [PoseIdx](const FPoseSearchDatabaseSequence& DbSequence)
+	{
+		return (PoseIdx >= DbSequence.FirstPoseIdx) && (PoseIdx < DbSequence.FirstPoseIdx + DbSequence.NumPoses);
+	};
+
+	return Sequences.IndexOfByPredicate(Predicate);
+}
+
+int32 UPoseSearchDatabase::GetPoseIndexFromAssetTime(int32 DbSequenceIdx, float AssetTime) const
+{
+	const FPoseSearchDatabaseSequence& DbSequence = Sequences[DbSequenceIdx];
+	FFloatInterval Range = UE::PoseSearch::GetEffectiveSamplingRange(DbSequence.Sequence, DbSequence.SamplingRange);
+	if (Range.Contains(AssetTime))
+	{
+		int32 PoseOffset = FMath::RoundToInt(Schema->SampleRate * (AssetTime - Range.Min));		
+		if (PoseOffset >= DbSequence.NumPoses)
+		{
+			if (DbSequence.bLoopAnimation)
+			{
+				PoseOffset -= DbSequence.NumPoses;
+			}
+			else
+			{
+				PoseOffset = DbSequence.NumPoses - 1;
+			}
+		}
+		
+		int32 PoseIdx = DbSequence.FirstPoseIdx + PoseOffset;
+		return PoseIdx;
+	}
+	
+	return INDEX_NONE;
+}
+
+bool UPoseSearchDatabase::IsValidForIndexing() const
+{
+	return Schema && Schema->IsValid() && !Sequences.IsEmpty();
+}
+
+bool UPoseSearchDatabase::IsValidForSearch() const
+{
+	return IsValidForIndexing() && SearchIndex.IsValid();
+}
+
+void UPoseSearchDatabase::PreSave(const class ITargetPlatform* TargetPlatform)
+{
+	SearchIndex.Reset();
+
+	if (IsValidForIndexing())
+	{
+		UE::PoseSearch::BuildIndex(this);
+	}
+
+	Super::PreSave(TargetPlatform);
+}
+
+
+namespace UE { namespace PoseSearch {
+
+//////////////////////////////////////////////////////////////////////////
+// FPoseHistory
+
+/**
+* Fills skeleton transforms with evaluated compact pose transforms.
+* Bones that weren't evaluated are filled with the bone's reference pose.
+*/
+static void CopyCompactToSkeletonPose(const FCompactPose& Pose, TArray<FTransform>& OutLocalTransforms)
+{
+	const FBoneContainer& BoneContainer = Pose.GetBoneContainer();
+	const FReferenceSkeleton& RefSkeleton = BoneContainer.GetReferenceSkeleton();
+	TArrayView<const FTransform> RefSkeletonTransforms = MakeArrayView(RefSkeleton.GetRefBonePose());
+
+	const int32 NumSkeletonBones = BoneContainer.GetNumBones();
+	OutLocalTransforms.SetNum(NumSkeletonBones);
+
+	for (auto SkeletonBoneIdx = FSkeletonPoseBoneIndex(0); SkeletonBoneIdx != NumSkeletonBones; ++SkeletonBoneIdx)
+	{
+		FCompactPoseBoneIndex CompactBoneIdx = BoneContainer.GetCompactPoseIndexFromSkeletonIndex(SkeletonBoneIdx.GetInt());
+		OutLocalTransforms[SkeletonBoneIdx.GetInt()] = CompactBoneIdx.IsValid() ? Pose[CompactBoneIdx] : RefSkeletonTransforms[SkeletonBoneIdx.GetInt()];
+	}
+}
+
+void FPoseHistory::Init(int32 InNumPoses, float InTimeHorizon)
+{
+	Poses.Reserve(InNumPoses);
+	Knots.Reserve(InNumPoses);
+	TimeHorizon = InTimeHorizon;
+}
+
+void FPoseHistory::Init(const FPoseHistory& History)
+{
+	Poses = History.Poses;
+	Knots = History.Knots;
+	TimeHorizon = History.TimeHorizon;
+}
+
+bool FPoseHistory::SampleLocalPose(float SecondsAgo, const TArray<FBoneIndexType>& RequiredBones, TArray<FTransform>& LocalPose)
+{
+	int32 NextIdx = LowerBound(Knots.begin(), Knots.end(), SecondsAgo, TGreater<>());
+	if (NextIdx >= Knots.Num())
+	{
+		return false;
+	}
+
+	int32 PrevIdx = NextIdx - 1;
+
+	// Compute alpha between previous and next knots
+	float Alpha = FMath::GetMappedRangeValueUnclamped(
+		FVector2D(Knots[PrevIdx], Knots[NextIdx]),
+		FVector2D(0.0f, 1.0f),
+		SecondsAgo);
+
+	TArray<FTransform>& PrevPose = Poses[PrevIdx].LocalTransforms;
+	TArray<FTransform>& NextPose = Poses[NextIdx].LocalTransforms;
+
+	// We may not have accumulated enough poses yet
+	if (PrevPose.Num() != NextPose.Num())
+	{
+		return false;
+	}
+
+	if (RequiredBones.Num() > PrevPose.Num())
+	{
+		return false;
+	}
+
+	// Lerp between poses by alpha to produce output local pose at requested sample time
+	LocalPose = PrevPose;
+	FAnimationRuntime::LerpBoneTransforms(
+		LocalPose,
+		NextPose,
+		Alpha,
+		RequiredBones);
+
+	return true;
+}
+
+bool FPoseHistory::Sample(float SecondsAgo, const FReferenceSkeleton& RefSkeleton, const TArray<FBoneIndexType>& RequiredBones)
+{
+	// Compute local space pose at requested time
+	bool bSampled = SampleLocalPose(SecondsAgo, RequiredBones, SampledLocalPose);
+
+	// Compute local space pose one sample interval in the past
+	bSampled = bSampled && SampleLocalPose(SecondsAgo + GetSampleInterval(), RequiredBones, SampledPrevLocalPose);
+
+	// Convert local to component space
+	if (bSampled)
+	{
+		FAnimationRuntime::FillUpComponentSpaceTransforms(RefSkeleton, SampledLocalPose, SampledComponentPose);
+		FAnimationRuntime::FillUpComponentSpaceTransforms(RefSkeleton, SampledPrevLocalPose, SampledPrevComponentPose);
+	}
+
+	return bSampled;
+}
+
+void FPoseHistory::Update(float SecondsElapsed, const FCompactPose& Pose)
+{
+	// Age our elapsed times
+	for (float& Knot : Knots)
+	{
+		Knot += SecondsElapsed;
+	}
+
+	if (Knots.Num() != Knots.Max())
+	{
+		// Consume every pose until the queue is full
+		Knots.AddUninitialized();
+		Poses.Emplace();
+	}
+	else
+	{
+		// Exercise pose retention policy. We must guarantee there is always one additional knot
+		// at or beyond the desired time horizon H so we can fulfill sample requests at t=H. We also
+		// want to evenly distribute knots across the entire history buffer so we only push additional
+		// poses when enough time has elapsed.
+
+		const float SampleInterval = GetSampleInterval();
+
+		bool bCanEvictOldest = Knots[1] >= TimeHorizon;
+		bool bShouldPushNewest = Knots[Knots.Num() - 2] >= SampleInterval;
+
+		if (bCanEvictOldest && bShouldPushNewest)
+		{
+			FPose PoseTemp = MoveTemp(Poses.First());
+			Poses.PopFront();
+			Poses.Emplace(MoveTemp(PoseTemp));
+
+			Knots.PopFront();
+			Knots.AddUninitialized();
+		}
+	}
+
+	// Regardless of the retention policy, we always update the most recent pose
+	Knots.Last() = 0.0f;
+	CopyCompactToSkeletonPose(Pose, Poses.Last().LocalTransforms);
+}
+
+float FPoseHistory::GetSampleInterval() const
+{
+	return TimeHorizon / Knots.Max();
+}
+
+
+//////////////////////////////////////////////////////////////////////////
+// FFeatureVectorBuilder
+
+void FFeatureVectorBuilder::Init(const FPoseSearchFeatureVectorLayout* InLayout, TArrayView<float> Buffer)
+{
+	check(InLayout);
+	check(Buffer.Num() == InLayout->NumFloats);
+	Layout = InLayout;
+	Values = Buffer;
+	ResetFeatures();
+}
+
+void FFeatureVectorBuilder::ResetFeatures()
+{
+	NumFeaturesAdded = 0;
+	FeaturesAdded.Init(false, Layout->Features.Num());
+}
+
+void FFeatureVectorBuilder::SetTransform(FPoseSearchFeatureDesc Element, const FTransform& Transform)
+{
+	SetPosition(Element, Transform.GetTranslation());
+	SetRotation(Element, Transform.GetRotation());
+}
+
+void FFeatureVectorBuilder::SetTransformDerivative(FPoseSearchFeatureDesc Element, const FTransform& Transform, const FTransform& PrevTransform, float DeltaTime)
+{
+	SetLinearVelocity(Element, Transform, PrevTransform, DeltaTime);
+	SetAngularVelocity(Element, Transform, PrevTransform, DeltaTime);
+}
+
+void FFeatureVectorBuilder::SetPosition(FPoseSearchFeatureDesc Element, const FVector& Position)
+{
+	Element.Type = EPoseSearchFeatureType::Position;
+	SetVector(Element, Position);
+}
+
+void FFeatureVectorBuilder::SetRotation(FPoseSearchFeatureDesc Element, const FQuat& Rotation)
+{
+	Element.Type = EPoseSearchFeatureType::Rotation;
+	int32 ElementIndex = Layout->Features.Find(Element);
+	if (ElementIndex >= 0)
+	{
+		FVector X = Rotation.GetAxisX();
+		FVector Y = Rotation.GetAxisY();
+
+		const FPoseSearchFeatureDesc& FoundElement = Layout->Features[ElementIndex];
+
+		Values[FoundElement.ValueOffset + 0] = X.X;
+		Values[FoundElement.ValueOffset + 1] = X.Y;
+		Values[FoundElement.ValueOffset + 2] = X.Z;
+		Values[FoundElement.ValueOffset + 3] = Y.X;
+		Values[FoundElement.ValueOffset + 4] = Y.Y;
+		Values[FoundElement.ValueOffset + 5] = Y.Z;
+
+		if (!FeaturesAdded[ElementIndex])
+		{
+			FeaturesAdded[ElementIndex] = true;
+			++NumFeaturesAdded;
+		}
+	}
+}
+
+void FFeatureVectorBuilder::SetLinearVelocity(FPoseSearchFeatureDesc Element, const FTransform& Transform, const FTransform& PrevTransform, float DeltaTime)
+{
+	Element.Type = EPoseSearchFeatureType::LinearVelocity;
+	FVector LinearVelocity = (Transform.GetTranslation() - PrevTransform.GetTranslation()) / DeltaTime;
+	SetVector(Element, LinearVelocity);
+}
+
+void FFeatureVectorBuilder::SetAngularVelocity(FPoseSearchFeatureDesc Element, const FTransform& Transform, const FTransform& PrevTransform, float DeltaTime)
+{
+	Element.Type = EPoseSearchFeatureType::AngularVelocity;
+	int32 ElementIndex = Layout->Features.Find(Element);
+	if (ElementIndex >= 0)
+	{
+		FQuat Q0 = PrevTransform.GetRotation();
+		FQuat Q1 = Transform.GetRotation();
+		Q1.EnforceShortestArcWith(Q0);
+
+		// Given angular velocity vector w, quaternion differentiation can be represented as
+		//   dq/dt = (w * q)/2
+		// Solve for w
+		//   w = 2 * dq/dt * q^-1
+		// And let dq/dt be expressed as the finite difference
+		//   dq/dt = (q(t+h) - q(t)) / h
+		FQuat DQDt = (Q1 - Q0) / DeltaTime;
+		FQuat QInv = Q0.Inverse();
+		FQuat W = (DQDt * QInv) * 2.0f;
+
+		FVector AngularVelocity(W.X, W.Y, W.Z);
+
+		const FPoseSearchFeatureDesc& FoundElement = Layout->Features[ElementIndex];
+
+		Values[FoundElement.ValueOffset + 0] = AngularVelocity[0];
+		Values[FoundElement.ValueOffset + 1] = AngularVelocity[1];
+		Values[FoundElement.ValueOffset + 2] = AngularVelocity[2];
+
+		if (!FeaturesAdded[ElementIndex])
+		{
+			FeaturesAdded[ElementIndex] = true;
+			++NumFeaturesAdded;
+		}
+	}
+}
+
+void FFeatureVectorBuilder::SetVector(FPoseSearchFeatureDesc Element, const FVector& Vector)
+{
+	int32 ElementIndex = Layout->Features.Find(Element);
+	if (ElementIndex >= 0)
+	{
+		const FPoseSearchFeatureDesc& FoundElement = Layout->Features[ElementIndex];
+
+		Values[FoundElement.ValueOffset + 0] = Vector[0];
+		Values[FoundElement.ValueOffset + 1] = Vector[1];
+		Values[FoundElement.ValueOffset + 2] = Vector[2];
+
+		if (!FeaturesAdded[ElementIndex])
+		{
+			FeaturesAdded[ElementIndex] = true;
+			++NumFeaturesAdded;
+		}
+	}
+}
+
+bool FFeatureVectorBuilder::SetPoseFeatures(const UPoseSearchSchema* Schema, FPoseHistory* History)
+{
+	check(Schema && Schema->IsValid());
+	check(History);
+
+	FPoseSearchFeatureDesc Feature;
+	Feature.Domain = EPoseSearchFeatureDomain::Time;
+
+	for (int32 SubsampleIdx = 0; SubsampleIdx != Schema->PoseSampleOffsets.Num(); ++SubsampleIdx)
+	{
+		Feature.SubsampleIdx = SubsampleIdx;
+
+		int32 Offset = Schema->PoseSampleOffsets[SubsampleIdx];
+		float TimeDelta = -Offset * Schema->SamplingInterval;
+
+		if (!History->Sample(TimeDelta, Schema->Skeleton->GetReferenceSkeleton(), Schema->BoneIndicesWithParents))
+		{
+			return false;
+		}
+
+		TArrayView<const FTransform> ComponentPose = History->GetComponentPoseSample();
+		TArrayView<const FTransform> ComponentPrevPose = History->GetPrevComponentPoseSample();
+		for (int32 SchemaBoneIdx = 0; SchemaBoneIdx != Schema->BoneIndices.Num(); ++SchemaBoneIdx)
+		{
+			Feature.SchemaBoneIdx = SchemaBoneIdx;
+
+			int32 SkeletonBoneIndex = Schema->BoneIndices[SchemaBoneIdx];
+			const FTransform& Transform = ComponentPose[SkeletonBoneIndex];
+			const FTransform& PrevTransform = ComponentPrevPose[SkeletonBoneIndex];
+			SetTransform(Feature, Transform);
+			SetTransformDerivative(Feature, Transform, PrevTransform, History->GetSampleInterval());
+		}
+	}
+
+	return true;
+}
+
+void FFeatureVectorBuilder::Copy(TArrayView<const float> FeatureVector)
+{
+	check(FeatureVector.Num() == Values.Num());
+	FMemory::Memcpy(Values.GetData(), FeatureVector.GetData(), FeatureVector.GetTypeSize() * FeatureVector.Num());
+	NumFeaturesAdded = Layout->Features.Num();
+	FeaturesAdded.SetRange(0, FeaturesAdded.Num(), true);
+}
+
+bool FFeatureVectorBuilder::IsComplete() const
+{
+	return NumFeaturesAdded == Layout->Features.Num();
+}
+
+
+//////////////////////////////////////////////////////////////////////////
+// FPoseSearchFeatureVectorReader
+
+void FFeatureVectorReader::Init(const FPoseSearchFeatureVectorLayout* InLayout)
+{
+	check(InLayout);
+	Layout = InLayout;
+}
+
+void FFeatureVectorReader::SetValues(TArrayView<const float> InValues)
+{
+	check(Layout);
+	check(Layout->NumFloats == InValues.Num());
+	Values = InValues;
+}
+
+bool FFeatureVectorReader::IsValid() const
+{
+	return Layout && (Layout->NumFloats == Values.Num());
+}
+
+bool FFeatureVectorReader::GetTransform(FPoseSearchFeatureDesc Element, FTransform* OutTransform) const
+{
+	FVector Position;
+	bool bResult = GetPosition(Element, &Position);
+
+	FQuat Rotation;
+	bResult |= GetRotation(Element, &Rotation);
+
+	OutTransform->SetComponents(Rotation, Position, FVector::OneVector);
+	return bResult;
+}
+
+bool FFeatureVectorReader::GetPosition(FPoseSearchFeatureDesc Element, FVector* OutPosition) const
+{
+	Element.Type = EPoseSearchFeatureType::Position;
+	return GetVector(Element, OutPosition);
+}
+
+bool FFeatureVectorReader::GetRotation(FPoseSearchFeatureDesc Element, FQuat* OutRotation) const
+{
+	Element.Type = EPoseSearchFeatureType::Rotation;
+	int32 ElementIndex = IsValid() ? Layout->Features.Find(Element) : -1;
+	if (ElementIndex >= 0)
+	{
+		const FPoseSearchFeatureDesc& FoundElement = Layout->Features[ElementIndex];
+
+		FVector X;
+		FVector Y;
+
+		X.X = Values[FoundElement.ValueOffset + 0];
+		X.Y = Values[FoundElement.ValueOffset + 1];
+		X.Z = Values[FoundElement.ValueOffset + 2];
+		Y.X = Values[FoundElement.ValueOffset + 3];
+		Y.Y = Values[FoundElement.ValueOffset + 4];
+		Y.Z = Values[FoundElement.ValueOffset + 5];
+
+		FVector Z = FVector::CrossProduct(X, Y);
+
+		FMatrix M(FMatrix::Identity);
+		M.SetColumn(0, X);
+		M.SetColumn(1, Y);
+		M.SetColumn(2, Z);
+
+		*OutRotation = FQuat(M);
+		return true;
+	}
+
+	*OutRotation = FQuat::Identity;
+	return false;
+}
+
+bool FFeatureVectorReader::GetLinearVelocity(FPoseSearchFeatureDesc Element, FVector* OutLinearVelocity) const
+{
+	Element.Type = EPoseSearchFeatureType::LinearVelocity;
+	return GetVector(Element, OutLinearVelocity);
+}
+
+bool FFeatureVectorReader::GetAngularVelocity(FPoseSearchFeatureDesc Element, FVector* OutAngularVelocity) const
+{
+	Element.Type = EPoseSearchFeatureType::AngularVelocity;
+	return GetVector(Element, OutAngularVelocity);
+}
+
+bool FFeatureVectorReader::GetVector(FPoseSearchFeatureDesc Element, FVector* OutVector) const
+{
+	int32 ElementIndex = IsValid() ? Layout->Features.Find(Element) : -1;
+	if (ElementIndex >= 0)
+	{
+		const FPoseSearchFeatureDesc& FoundElement = Layout->Features[ElementIndex];
+
+		FVector V;
+		V.X = Values[FoundElement.ValueOffset + 0];
+		V.Y = Values[FoundElement.ValueOffset + 1];
+		V.Z = Values[FoundElement.ValueOffset + 2];
+		*OutVector = V;
+		return true;
+	}
+
+	*OutVector = FVector::ZeroVector;
+	return false;
+}
+
+
+//////////////////////////////////////////////////////////////////////////
+// FSequenceIndexer
+
+bool FDebugDrawParams::CanDraw () const
+{
+	if (!World || !EnumHasAnyFlags(Flags, EDebugDrawFlags::DrawAll))
+	{
+		return false;
+	}
+
+	const FPoseSearchIndex* SearchIndex = GetSearchIndex();
+	if (!SearchIndex)
+	{
+		return false;
+	}
+
+	return SearchIndex->IsValid();
+}
+
+const FPoseSearchIndex* FDebugDrawParams::GetSearchIndex() const
+{
+	if (Database)
+	{
+		return &Database->SearchIndex;
+	}
+
+	if (SequenceMetaData)
+	{
+		return &SequenceMetaData->SearchIndex;
+	}
+
+	return nullptr;
+}
+
+const UPoseSearchSchema* FDebugDrawParams::GetSchema() const
+{
+	if (Database)
+	{
+		return Database->Schema;
+	}
+
+	if (SequenceMetaData)
+	{
+		return SequenceMetaData->Schema;
+	}
+
+	return nullptr;
+}
+
+
+//////////////////////////////////////////////////////////////////////////
+// FSequenceIndexer
+
+class FSequenceIndexer
+{
+public:
+	struct Result
+	{
+		int32 NumIndexedPoses = 0;
+		TArrayView<const float> Values;
+	};
+	Result Process(const UPoseSearchSchema* Schema, const UAnimSequence* Sequence, FFloatInterval SamplingRange);
+	Result GetResult() const;
+
+private:
+	void SampleBegin(int32 SampleIdx);
+	void SampleEnd();
+	void ExtractPoses(const UAnimSequence* Sequence);
+	void ExtractRootMotion(const UAnimSequence* Sequence);
+	void AddPoseFeatures(int32 SampleIdx);
+	void AddTrajectoryTimeFeatures(int32 SampleIdx);
+	void AddTrajectoryDistanceFeatures(int32 SampleIdx);
+
+	struct FSampleContext
+	{
+		TArray<FTransform> ComponentSpacePose;		// Indexed by SampleIdx * NumBones + SchemaBoneIdx
+		TArray<FTransform> LocalRootMotion;			// Indexed by SampleIdx
+		TArray<FTransform> AccumulatedRootMotion;	// Indexed by SampleIdx
+		TArray<float> AccumulatedRootDistance;		// Indexed by SampleIdx
+
+		int32 TotalSamples = 0;
+		int32 FirstIndexedSample = 0;
+		int32 LastIndexedSample = 0;
+		int32 NumIndexedSamples = 0;
+		int32 NumBones = 0;
+
+		void Reset();
+		void Reserve();
+	};
+
+	const UPoseSearchSchema* Schema = nullptr;
+
+	TArray<float> Values;
+	
+	FFeatureVectorBuilder Builder;
+	FSampleContext Context;
+};
+
+void FSequenceIndexer::FSampleContext::Reset()
+{
+	TotalSamples = 0;
+	FirstIndexedSample = 0;
+	LastIndexedSample = 0;
+	NumIndexedSamples = 0;
+	NumBones = 0;
+
+	ComponentSpacePose.Reset(0);
+	LocalRootMotion.Reset(0);
+	AccumulatedRootMotion.Reset(0);
+	AccumulatedRootDistance.Reset(0);
+}
+
+void FSequenceIndexer::FSampleContext::Reserve()
+{
+	ComponentSpacePose.Reserve(NumBones * TotalSamples);
+	LocalRootMotion.Reserve(TotalSamples);
+	AccumulatedRootMotion.Reserve(TotalSamples);
+	AccumulatedRootDistance.Reserve(TotalSamples);
+}
+
+FSequenceIndexer::Result FSequenceIndexer::Process(const UPoseSearchSchema* InSchema, const UAnimSequence* Sequence, FFloatInterval RequestedSamplingRange)
+{
+	check(InSchema);
+	check(Sequence);
+
+	Schema = InSchema;
+
+	USkeleton* Skeleton = Sequence->GetSkeleton();
+	check(Skeleton);
+	check(Skeleton->IsCompatible(Schema->Skeleton));
+
+	const float SequencePlayLength = Sequence->GetPlayLength();
+	const FFloatInterval SamplingRange = GetEffectiveSamplingRange(Sequence, RequestedSamplingRange);
+
+	Context.Reset();
+	Context.NumBones = Schema->BoneIndices.Num();
+	Context.TotalSamples = FMath::FloorToInt(SequencePlayLength * Schema->SampleRate);
+	Context.FirstIndexedSample = FMath::FloorToInt(SamplingRange.Min * Schema->SampleRate);
+	Context.LastIndexedSample = FMath::Max(0, FMath::FloorToInt(SamplingRange.Max * Schema->SampleRate) - 1);
+	Context.NumIndexedSamples = Context.LastIndexedSample - Context.FirstIndexedSample + 1;
+	Context.Reserve();
+
+	Values.SetNumZeroed(Schema->Layout.NumFloats * Context.NumIndexedSamples);
+
+	ExtractPoses(Sequence);
+	ExtractRootMotion(Sequence);
+
+	for (int32 SampleIdx = Context.FirstIndexedSample; SampleIdx <= Context.LastIndexedSample; ++SampleIdx)
+	{
+		SampleBegin(SampleIdx);
+
+		AddPoseFeatures(SampleIdx);
+		AddTrajectoryTimeFeatures(SampleIdx);
+		AddTrajectoryDistanceFeatures(SampleIdx);
+
+		SampleEnd();
+	}
+
+	return GetResult();
+}
+
+FSequenceIndexer::Result FSequenceIndexer::GetResult() const
+{
+	Result Result;
+	Result.NumIndexedPoses = Context.NumIndexedSamples;
+	Result.Values = Values;
+	return Result;
+}
+
+void FSequenceIndexer::SampleBegin(int32 SampleIdx)
+{
+	int32 FirstValueIdx = (SampleIdx - Context.FirstIndexedSample) * Schema->Layout.NumFloats;
+	TArrayView<float> FeatureVectorValues = MakeArrayView(&Values[FirstValueIdx], Schema->Layout.NumFloats);
+	Builder.Init(&Schema->Layout, FeatureVectorValues);
+}
+
+void FSequenceIndexer::SampleEnd()
+{
+	check(Builder.IsComplete());
+}
+
+void FSequenceIndexer::ExtractPoses(const UAnimSequence* Sequence)
+{
+	if (Schema->Bones.IsEmpty())
+	{
+		return;
+	}
+
+	USkeleton* Skeleton = Sequence->GetSkeleton();
+	FBoneContainer BoneContainer;
+	BoneContainer.InitializeTo(Schema->BoneIndicesWithParents, FCurveEvaluationOption(false), *Skeleton);
+
+	FCompactPose Pose;
+	Pose.SetBoneContainer(&BoneContainer);
+	FCSPose<FCompactPose> ComponentSpacePose;
+
+	FBlendedCurve UnusedCurve;
+	FStackCustomAttributes UnusedAttributes;
+
+	FAnimExtractContext ExtractionCtx;
+	// ExtractionCtx.PoseCurves is intentionally left empty
+	// ExtractionCtx.BonesRequired is unused by UAnimSequence::GetAnimationPose
+	ExtractionCtx.bExtractRootMotion = true;
+
+	FAnimationPoseData AnimPoseData(Pose, UnusedCurve, UnusedAttributes);
+	for (int32 SampleIdx = 0; SampleIdx != Context.TotalSamples; ++SampleIdx)
+	{
+		const float CurrentTime = SampleIdx * Schema->SamplingInterval;
+
+		ExtractionCtx.CurrentTime = CurrentTime;
+		Sequence->GetAnimationPose(AnimPoseData, ExtractionCtx);
+		ComponentSpacePose.InitPose(Pose);
+
+		for (int32 BoneIndex : Schema->BoneIndices)
+		{
+			FCompactPoseBoneIndex CompactBoneIndex = BoneContainer.MakeCompactPoseIndex(FMeshPoseBoneIndex(BoneIndex));
+			const FTransform& Transform = ComponentSpacePose.GetComponentSpaceTransform(CompactBoneIndex);
+			Context.ComponentSpacePose.Add(Transform);
+		}
+	}
+}
+
+void FSequenceIndexer::ExtractRootMotion(const UAnimSequence* Sequence)
+{
+	double AccumulatedRootDistance = 0.0;
+	FTransform AccumulatedRootMotion = FTransform::Identity;
+	for (int32 SampleIdx = 0; SampleIdx != Context.TotalSamples; ++SampleIdx)
+	{
+		const float CurrentTime = SampleIdx * Schema->SamplingInterval;
+
+		FTransform LocalRootMotion = Sequence->ExtractRootMotion(CurrentTime, Schema->SamplingInterval, false /*!allowLooping*/);
+		Context.LocalRootMotion.Add(LocalRootMotion);
+
+		AccumulatedRootMotion = LocalRootMotion * AccumulatedRootMotion;
+		AccumulatedRootDistance += LocalRootMotion.GetTranslation().Size();
+		Context.AccumulatedRootMotion.Add(AccumulatedRootMotion);
+		Context.AccumulatedRootDistance.Add((float)AccumulatedRootDistance);
+	}
+}
+
+void FSequenceIndexer::AddPoseFeatures(int32 SampleIdx)
+{
+	FPoseSearchFeatureDesc CurrentElement;
+	CurrentElement.Domain = EPoseSearchFeatureDomain::Time;
+
+	FTransform SampleSpaceOrigin = Context.AccumulatedRootMotion[SampleIdx];
+	
+	for (int32 SchemaSubsampleIdx = 0; SchemaSubsampleIdx != Schema->PoseSampleOffsets.Num(); ++SchemaSubsampleIdx)
+	{
+		CurrentElement.SubsampleIdx = SchemaSubsampleIdx;
+
+		const int32 SampleOffset = Schema->PoseSampleOffsets[SchemaSubsampleIdx];
+		const int32 SubsampleIdx = FMath::Clamp(SampleIdx + SampleOffset, 0, Context.AccumulatedRootMotion.Num() - 1);
+
+		FTransform SubsampleRoot = Context.AccumulatedRootMotion[SubsampleIdx];
+		SubsampleRoot.SetToRelativeTransform(SampleSpaceOrigin);
+
+		for (int32 SchemaBoneIndex = 0; SchemaBoneIndex != Context.NumBones; ++SchemaBoneIndex)
+		{
+			CurrentElement.SchemaBoneIdx = SchemaBoneIndex;
+
+			int32 BoneSampleIdx = Context.NumBones * (SampleIdx + SampleOffset) + SchemaBoneIndex;
+			int32 BonePrevSampleIdx = Context.NumBones * (SampleIdx - 1 + SampleOffset) + SchemaBoneIndex;
+			
+			// @@@Add extrapolation. Clamp for now
+			BoneSampleIdx = FMath::Clamp(BoneSampleIdx, 0, Context.ComponentSpacePose.Num() - 1);
+			BonePrevSampleIdx = FMath::Clamp(BonePrevSampleIdx, 0, Context.ComponentSpacePose.Num() - 1);
+
+			FTransform BoneInSampleSpace = Context.ComponentSpacePose[BoneSampleIdx] * SubsampleRoot;
+			FTransform BonePrevInSampleSpace = Context.ComponentSpacePose[BonePrevSampleIdx] * SubsampleRoot;
+
+			Builder.SetTransform(CurrentElement, BoneInSampleSpace);
+			Builder.SetTransformDerivative(CurrentElement, BoneInSampleSpace, BonePrevInSampleSpace, Schema->SamplingInterval);
+		}
+	}
+}
+
+void FSequenceIndexer::AddTrajectoryTimeFeatures(int32 SampleIdx)
+{
+	FPoseSearchFeatureDesc CurrentElement;
+	CurrentElement.Domain = EPoseSearchFeatureDomain::Time;
+	CurrentElement.SchemaBoneIdx = FPoseSearchFeatureDesc::TrajectoryBoneIndex;
+
+	FTransform SampleSpaceOrigin = Context.AccumulatedRootMotion[SampleIdx];
+
+	for (int32 SubsampleIdx = 0; SubsampleIdx != Schema->TrajectorySampleOffsets.Num(); ++SubsampleIdx)
+	{
+		CurrentElement.SubsampleIdx = SubsampleIdx;
+
+		int32 RootMotionIdx = SampleIdx + Schema->TrajectorySampleOffsets[SubsampleIdx];
+		int32 RootMotionPrevIdx = RootMotionIdx - 1;
+
+		// @@@ Add extrapolation. Clamp for now
+		RootMotionIdx = FMath::Clamp(RootMotionIdx, 0, Context.AccumulatedRootMotion.Num() - 1);
+		RootMotionPrevIdx = FMath::Clamp(RootMotionPrevIdx, 0, Context.AccumulatedRootMotion.Num() - 1);
+
+		FTransform SubsampleRoot = Context.AccumulatedRootMotion[RootMotionIdx];
+		SubsampleRoot.SetToRelativeTransform(SampleSpaceOrigin);
+
+		FTransform SubsamplePrevRoot = Context.AccumulatedRootMotion[RootMotionPrevIdx];
+		SubsamplePrevRoot.SetToRelativeTransform(SampleSpaceOrigin);
+
+		Builder.SetTransform(CurrentElement, SubsampleRoot);
+		Builder.SetTransformDerivative(CurrentElement, SubsampleRoot, SubsamplePrevRoot, Schema->SamplingInterval);
+	}
+}
+
+void FSequenceIndexer::AddTrajectoryDistanceFeatures(int32 SampleIdx)
+{
+	FPoseSearchFeatureDesc CurrentElement;
+	CurrentElement.Domain = EPoseSearchFeatureDomain::Distance;
+	CurrentElement.SchemaBoneIdx = FPoseSearchFeatureDesc::TrajectoryBoneIndex;
+
+	TArrayView<const float> AccumulatedRootDistances = Context.AccumulatedRootDistance;
+
+	FTransform SampleSpaceOrigin = Context.AccumulatedRootMotion[SampleIdx];
+
+	for (int32 SubsampleIdx = 0; SubsampleIdx != Schema->TrajectoryDistanceOffsets.Num(); ++SubsampleIdx)
+	{
+		CurrentElement.SubsampleIdx = SubsampleIdx;
+
+		const float TrajectoryDistance = Schema->TrajectoryDistanceOffsets[SubsampleIdx];
+		const float SampleAccumulatedRootDistance = TrajectoryDistance + AccumulatedRootDistances[SampleIdx];
+
+		int32 LowerBoundSampleIdx = Algo::LowerBound(AccumulatedRootDistances, SampleAccumulatedRootDistance);
+
+		// @@@ Add extrapolation. Clamp for now
+		int32 PrevSampleIdx = FMath::Clamp(LowerBoundSampleIdx - 1, 0, AccumulatedRootDistances.Num() - 1);
+		int32 NextSampleIdx = FMath::Clamp(LowerBoundSampleIdx, 0, AccumulatedRootDistances.Num() - 1);
+
+		const float PrevSampleDistance = AccumulatedRootDistances[PrevSampleIdx];
+		const float NextSampleDistance = AccumulatedRootDistances[NextSampleIdx];
+
+		FTransform PrevRootInSampleSpace = Context.AccumulatedRootMotion[PrevSampleIdx];
+		PrevRootInSampleSpace.SetToRelativeTransform(SampleSpaceOrigin);
+
+		FTransform NextRootInSampleSpace = Context.AccumulatedRootMotion[NextSampleIdx];
+		NextRootInSampleSpace.SetToRelativeTransform(SampleSpaceOrigin);
+		
+		float Alpha = FMath::GetRangePct(PrevSampleDistance, NextSampleDistance, SampleAccumulatedRootDistance);
+		FTransform BlendedRootInSampleSpace;
+		BlendedRootInSampleSpace.Blend(PrevRootInSampleSpace, NextRootInSampleSpace, Alpha);
+
+		Builder.SetTransform(CurrentElement, BlendedRootInSampleSpace);
+	}
+}
+
+
+//////////////////////////////////////////////////////////////////////////
+// PoseSearch API
+
+static void DrawTrajectoryFeatures(const FDebugDrawParams& DrawParams, const FFeatureVectorReader& Reader)
+{
+	const float LifeTime = DrawParams.DefaultLifeTime;
+	const uint8 DepthPriority = ESceneDepthPriorityGroup::SDPG_Foreground + 2;
+
+	FPoseSearchFeatureDesc Feature;
+	Feature.Domain = EPoseSearchFeatureDomain::Time;
+	Feature.SchemaBoneIdx = FPoseSearchFeatureDesc::TrajectoryBoneIndex;
+
+	const int32 NumSubsamples = DrawParams.GetSchema()->TrajectorySampleOffsets.Num();
+
+	if (NumSubsamples == 0)
+	{
+		return;
+	}
+
+	FVector TrajectoryPosPrev(0.0f);
+	for (int32 SubsampleIdx = 0; SubsampleIdx != NumSubsamples; ++SubsampleIdx)
+	{
+		Feature.SubsampleIdx = SubsampleIdx;
+
+		FVector TrajectoryPos;
+		if (Reader.GetPosition(Feature, &TrajectoryPos))
+		{	
+			Feature.Type = EPoseSearchFeatureType::Position;
+			FLinearColor LinearColor = GetColorForFeature(Feature, Reader.GetLayout());
+			FColor Color = LinearColor.ToFColor(true);
+
+			TrajectoryPos = DrawParams.RootTransform.TransformPosition(TrajectoryPos);
+			DrawDebugSphere(DrawParams.World, TrajectoryPos, DrawDebugSphereSize, DrawDebugSphereSegments, Color, false, LifeTime, DepthPriority,  DrawDebugSphereLineThickness);
+
+			if (SubsampleIdx != 0)
+			{
+				FVector Direction;
+				float Length;
+				(TrajectoryPos - TrajectoryPosPrev).ToDirectionAndLength(Direction, Length);
+				DrawDebugLine(DrawParams.World, TrajectoryPosPrev, TrajectoryPos, Color, false, LifeTime, DepthPriority, DrawDebugLineThickness);
+			}
+
+			TrajectoryPosPrev = TrajectoryPos;
+		}
+		else
+		{
+			TrajectoryPos = DrawParams.RootTransform.GetTranslation();
+		}
+
+		FVector TrajectoryVel;
+		if (Reader.GetLinearVelocity(Feature, &TrajectoryVel))
+		{
+			Feature.Type = EPoseSearchFeatureType::LinearVelocity;
+			FLinearColor LinearColor = GetColorForFeature(Feature, Reader.GetLayout());
+			FColor Color = LinearColor.ToFColor(true);
+
+			TrajectoryVel *= DrawDebugVelocityScale;
+			TrajectoryVel = DrawParams.RootTransform.TransformVector(TrajectoryVel);
+			FVector TrajectoryVelDirection = TrajectoryVel.GetSafeNormal();
+			DrawDebugDirectionalArrow(DrawParams.World, TrajectoryPos + TrajectoryVelDirection * DrawDebugSphereSize, TrajectoryPos + TrajectoryVel, DrawDebugArrowSize, Color, false, LifeTime, DepthPriority, DrawDebugLineThickness);
+		}
+	}
+}
+
+static void DrawPoseFeatures(const FDebugDrawParams& DrawParams, const FFeatureVectorReader& Reader)
+{
+	const UPoseSearchSchema* Schema = DrawParams.GetSchema();
+	check(Schema && Schema->IsValid());
+
+	const float LifeTime = DrawParams.DefaultLifeTime;
+	const uint8 DepthPriority = ESceneDepthPriorityGroup::SDPG_Foreground + 2;
+
+	FPoseSearchFeatureDesc Feature;
+	Feature.Domain = EPoseSearchFeatureDomain::Time;
+
+	const int32 NumSubsamples = Schema->PoseSampleOffsets.Num();
+	const int32 NumBones = Schema->Bones.Num();
+
+	if ((NumSubsamples * NumBones) == 0)
+	{
+		return;
+	}
+
+	for (int32 SubsampleIdx = 0; SubsampleIdx != NumSubsamples; ++SubsampleIdx)
+	{
+		Feature.SubsampleIdx = SubsampleIdx;
+		Feature.SchemaBoneIdx = 0;
+		
+		for (int32 SchemaBoneIdx = 0; SchemaBoneIdx != NumBones; ++SchemaBoneIdx)
+		{
+			Feature.SchemaBoneIdx = SchemaBoneIdx;
+
+			FVector BonePos;
+			bool bHaveBonePos = Reader.GetPosition(Feature, &BonePos);
+			if (bHaveBonePos)
+			{
+				Feature.Type = EPoseSearchFeatureType::Position;
+				FLinearColor Color = GetColorForFeature(Feature, Reader.GetLayout());
+
+				BonePos = DrawParams.RootTransform.TransformPosition(BonePos);
+				DrawDebugSphere(DrawParams.World, BonePos, DrawDebugSphereSize, DrawDebugSphereSegments, Color.ToFColor(true), false, LifeTime, DepthPriority, DrawDebugSphereLineThickness);
+			}
+
+			FVector BoneVel;
+			if (bHaveBonePos && Reader.GetLinearVelocity(Feature, &BoneVel))
+			{
+				Feature.Type = EPoseSearchFeatureType::LinearVelocity;
+				FLinearColor Color = GetColorForFeature(Feature, Reader.GetLayout());
+
+				BoneVel *= DrawDebugVelocityScale;
+				BoneVel = DrawParams.RootTransform.TransformVector(BoneVel);
+				FVector BoneVelDirection = BoneVel.GetSafeNormal();
+				DrawDebugDirectionalArrow(DrawParams.World, BonePos + BoneVelDirection * DrawDebugSphereSize, BonePos + BoneVel, DrawDebugArrowSize, Color.ToFColor(true), false, LifeTime, DepthPriority, DrawDebugLineThickness);
+			}
+		}
+	}
+}
+
+static void DrawFeatureVector(const FDebugDrawParams& DrawParams, const FFeatureVectorReader& Reader)
+{
+	DrawPoseFeatures(DrawParams, Reader);
+	DrawTrajectoryFeatures(DrawParams, Reader);
+}
+
+static void DrawSearchIndex(const FDebugDrawParams& DrawParams)
+{
+	if (!DrawParams.CanDraw())
+	{
+		return;
+	}
+
+	const UPoseSearchSchema* Schema = DrawParams.GetSchema();
+	const FPoseSearchIndex* SearchIndex = DrawParams.GetSearchIndex();
+	check(Schema);
+	check(SearchIndex);
+
+	FFeatureVectorReader Reader;
+	Reader.Init(&Schema->Layout);
+
+	int32 LastPoseIdx = SearchIndex->NumPoses;
+	int32 StartPoseIdx = 0;
+	if (!(DrawParams.Flags & EDebugDrawFlags::DrawSearchIndex))
+	{
+		StartPoseIdx = DrawParams.HighlightPoseIdx;
+		LastPoseIdx = StartPoseIdx + 1;
+	}
+
+	if (StartPoseIdx < 0)
+	{
+		return;
+	}
+
+	for (int32 PoseIdx = StartPoseIdx; PoseIdx != LastPoseIdx; ++PoseIdx)
+	{
+		Reader.SetValues(SearchIndex->GetPoseValues(PoseIdx));
+
+		DrawFeatureVector(DrawParams, Reader);
+	}
+}
+
+static void DrawQuery(const FDebugDrawParams& DrawParams)
+{
+	if (!DrawParams.CanDraw())
+	{
+		return;
+	}
+
+	const UPoseSearchSchema* Schema = DrawParams.GetSchema();
+	check(Schema);
+
+	FFeatureVectorReader Reader;
+	Reader.Init(&Schema->Layout);
+	Reader.SetValues(DrawParams.Query);
+	DrawFeatureVector(DrawParams, Reader);
+}
+
+void Draw(const FDebugDrawParams& DebugDrawParams)
+{
+	if (DebugDrawParams.CanDraw())
+	{
+		if (EnumHasAnyFlags(DebugDrawParams.Flags, EDebugDrawFlags::DrawQuery))
+		{
+			DrawQuery(DebugDrawParams);
+		}
+
+		if (EnumHasAnyFlags(DebugDrawParams.Flags, EDebugDrawFlags::DrawSearchIndex | EDebugDrawFlags::DrawBest))
+		{
+			DrawSearchIndex(DebugDrawParams);
+		}
+	}
+}
+
+bool BuildIndex(const UAnimSequence* Sequence, UPoseSearchSequenceMetaData* SequenceMetaData)
+{
+	check(Sequence);
+	check(SequenceMetaData);
+
+	if (!SequenceMetaData->IsValidForIndexing())
+	{
+		return false;
+	}
+
+	USkeleton* SeqSkeleton = Sequence->GetSkeleton();
+	if (!SeqSkeleton || !SeqSkeleton->IsCompatible(SequenceMetaData->Schema->Skeleton))
+	{
+		return false;
+	}
+
+	FSequenceIndexer Indexer;
+	FSequenceIndexer::Result Result = Indexer.Process(SequenceMetaData->Schema, Sequence, SequenceMetaData->SamplingRange);
+
+	SequenceMetaData->SearchIndex.Values = Result.Values;
+	SequenceMetaData->SearchIndex.NumPoses = Result.NumIndexedPoses;
+	SequenceMetaData->SearchIndex.Schema = SequenceMetaData->Schema;
+	return true;
+}
+
+bool BuildIndex(UPoseSearchDatabase* Database)
+{
+	check(Database);
+
+	if (!Database->IsValidForIndexing())
+	{
+		return false;
+	}
+
+	for (const FPoseSearchDatabaseSequence& DbSequence : Database->Sequences)
+	{
+		USkeleton* SeqSkeleton = DbSequence.Sequence->GetSkeleton();
+		if (!SeqSkeleton || !SeqSkeleton->IsCompatible(Database->Schema->Skeleton))
+		{
+			return false;
+		}
+	}
+
+	// Prepare animation indexing tasks
+	TArray<FSequenceIndexer> Indexers;
+	Indexers.SetNum(Database->Sequences.Num());
+
+	auto IndexerTask = [&Database, &Indexers](int32 SequenceIdx)
+	{
+		const FPoseSearchDatabaseSequence& DbSequence = Database->Sequences[SequenceIdx];
+		FSequenceIndexer& Indexer = Indexers[SequenceIdx];
+		Indexer.Process(Database->Schema, DbSequence.Sequence, DbSequence.SamplingRange);
+	};
+
+	// Index animations independently
+	ParallelFor(Database->Sequences.Num(), IndexerTask);
+
+	// Write index info to sequence and count up total poses and storage required
+	int32 TotalPoses = 0;
+	int32 TotalFloats = 0;
+	for (int32 SequenceIdx = 0; SequenceIdx != Database->Sequences.Num(); ++SequenceIdx)
+	{
+		FPoseSearchDatabaseSequence& DbSequence = Database->Sequences[SequenceIdx];
+		FSequenceIndexer::Result Result = Indexers[SequenceIdx].GetResult();
+		DbSequence.NumPoses = Result.NumIndexedPoses;
+		DbSequence.FirstPoseIdx = TotalPoses;
+		TotalPoses += Result.NumIndexedPoses;
+		TotalFloats += Result.Values.Num();
+	}
+
+	// Join animation data into a single search index
+	Database->SearchIndex.Values.Reset(TotalFloats);
+	for (const FSequenceIndexer& Indexer : Indexers)
+	{
+		FSequenceIndexer::Result Result = Indexer.GetResult();
+		Database->SearchIndex.Values.Append(Result.Values.GetData(), Result.Values.Num());
+	}
+
+	Database->SearchIndex.NumPoses = TotalPoses;
+	Database->SearchIndex.Schema = Database->Schema;
+	return true;
+}
+
+static FSearchResult Search(const FPoseSearchIndex& SearchIndex, TArrayView<const float> Query)
+{
+	FSearchResult Result;
+
+	if (!ensure(SearchIndex.IsValid()))
+	{
+		return Result;
+	}
+
+	if(!ensure(Query.Num() == SearchIndex.Schema->Layout.NumFloats))
+	{
+		return Result;
+	}
+
+	float BestPoseDissimilarity = MAX_flt;
+	int32 BestPoseIdx = INDEX_NONE;
+
+	for (int32 PoseIdx = 0; PoseIdx != SearchIndex.NumPoses; ++PoseIdx)
+	{
+		const int32 FeatureValueOffset = PoseIdx * SearchIndex.Schema->Layout.NumFloats;
+
+		float PoseDissimilarity = CompareFeatureVectors(
+			SearchIndex.Schema->Layout.NumFloats,
+			Query.GetData(),
+			&SearchIndex.Values[FeatureValueOffset]
+		);
+
+		if (PoseDissimilarity < BestPoseDissimilarity)
+		{
+			BestPoseDissimilarity = PoseDissimilarity;
+			BestPoseIdx = PoseIdx;
+		}
+	}
+
+	ensure(BestPoseIdx != INDEX_NONE);
+
+	Result.Dissimilarity = BestPoseDissimilarity;
+	Result.PoseIdx = BestPoseIdx;
+	// Result.TimeOffsetSeconds is set by caller
+
+	return Result;
+}
+
+FSearchResult Search(const UAnimSequenceBase* Sequence, TArrayView<const float> Query, FDebugDrawParams DebugDrawParams)
+{
+	const UPoseSearchSequenceMetaData* MetaData = Sequence ? Sequence->FindMetaDataByClass<UPoseSearchSequenceMetaData>() : nullptr;
+	if (!MetaData || !MetaData->IsValidForSearch())
+	{
+		return FSearchResult();
+	}
+
+	const FPoseSearchIndex& SearchIndex = MetaData->SearchIndex;
+
+	FSearchResult Result = Search(SearchIndex, Query);
+	if (!Result.IsValid())
+	{
+		return Result;
+	}
+
+	const FFloatInterval SamplingRange = GetEffectiveSamplingRange(Sequence, MetaData->SamplingRange);
+	Result.TimeOffsetSeconds = SamplingRange.Min + (SearchIndex.Schema->SamplingInterval * Result.PoseIdx);
+
+	// Do debug visualization
+	DebugDrawParams.SequenceMetaData = MetaData;
+	DebugDrawParams.Query = Query;
+	DebugDrawParams.HighlightPoseIdx = Result.PoseIdx;
+	Draw(DebugDrawParams);
+
+	return Result;
+}
+
+FDbSearchResult Search(const UPoseSearchDatabase* Database, TArrayView<const float> Query, FDebugDrawParams DebugDrawParams)
+{
+	if (!ensure(Database && Database->IsValidForSearch()))
+	{
+		return FDbSearchResult();
+	}
+
+	const FPoseSearchIndex& SearchIndex = Database->SearchIndex;
+
+	FDbSearchResult Result = Search(SearchIndex, Query);
+	if (!Result.IsValid())
+	{
+		return FDbSearchResult();
+	}
+
+	int32 DbSequenceIdx = Database->FindSequenceForPose(Result.PoseIdx);
+	if (DbSequenceIdx == INDEX_NONE)
+	{
+		return FDbSearchResult();
+	}
+	
+	const FPoseSearchDatabaseSequence& DbSequence = Database->Sequences[DbSequenceIdx];
+	const FFloatInterval SamplingRange = GetEffectiveSamplingRange(DbSequence.Sequence, DbSequence.SamplingRange);
+
+	Result.DbSequenceIdx = DbSequenceIdx;
+	Result.TimeOffsetSeconds = SamplingRange.Min + SearchIndex.Schema->SamplingInterval * (Result.PoseIdx - DbSequence.FirstPoseIdx);
+
+	// Do debug visualization
+	DebugDrawParams.Database = Database;
+	DebugDrawParams.Query = Query;
+	DebugDrawParams.HighlightPoseIdx = Result.PoseIdx;
+	Draw(DebugDrawParams);
+
+	return Result;
+}
+
+float ComparePoses(const FPoseSearchIndex& SearchIndex, int32 PoseIdx, TArrayView<const float> Query)
+{
+	TArrayView<const float> PoseValues = SearchIndex.GetPoseValues(PoseIdx);
+	check(PoseValues.Num() == Query.Num());
+	return CompareFeatureVectors(PoseValues.Num(), PoseValues.GetData(), Query.GetData());
+}
+
+
+//////////////////////////////////////////////////////////////////////////
+// FModule
+
+class FModule : public IModuleInterface, public UE::Anim::IPoseSearchProvider
+{
+public: // IModuleInterface
+	virtual void StartupModule() override;
+	virtual void ShutdownModule() override;
+
+public: // IPoseSearchProvider
+	virtual UE::Anim::IPoseSearchProvider::FSearchResult Search(const FAnimationBaseContext& GraphContext, const UAnimSequenceBase* Sequence) override;
+};
+
+void FModule::StartupModule()
+{
+	IModularFeatures::Get().RegisterModularFeature(UE::Anim::IPoseSearchProvider::ModularFeatureName, this);
+}
+
+void FModule::ShutdownModule()
+{
+	IModularFeatures::Get().UnregisterModularFeature(UE::Anim::IPoseSearchProvider::ModularFeatureName, this);
+}
+
+UE::Anim::IPoseSearchProvider::FSearchResult FModule::Search(const FAnimationBaseContext& GraphContext, const UAnimSequenceBase* Sequence)
+{
+	UE::Anim::IPoseSearchProvider::FSearchResult ProviderResult;
+
+	const UPoseSearchSequenceMetaData* MetaData = Sequence ? Sequence->FindMetaDataByClass<UPoseSearchSequenceMetaData>() : nullptr;
+	if (!MetaData || !MetaData->IsValidForSearch())
+	{
+		return ProviderResult;
+	}
+
+	IPoseHistoryProvider* PoseHistoryProvider = GraphContext.GetMessage<IPoseHistoryProvider>();
+	if (!PoseHistoryProvider)
+	{
+		return ProviderResult;
+	}
+
+	FPoseHistory& PoseHistory = PoseHistoryProvider->GetPoseHistory();
+	TArray<float>& QueryBuffer = PoseHistory.GetQueryBuffer();
+
+	QueryBuffer.SetNum(MetaData->Schema->Layout.NumFloats);
+	TArrayView<float> Query = MakeArrayView(QueryBuffer.GetData(), QueryBuffer.Num());
+
+	FFeatureVectorBuilder QueryBuilder;
+	QueryBuilder.Init(&MetaData->Schema->Layout, Query);
+	QueryBuilder.SetPoseFeatures(MetaData->Schema, &PoseHistory);
+
+	::UE::PoseSearch::FSearchResult Result = ::UE::PoseSearch::Search(Sequence, Query);
+
+	ProviderResult.Dissimilarity = Result.Dissimilarity;
+	ProviderResult.PoseIdx = Result.PoseIdx;
+	ProviderResult.TimeOffsetSeconds = Result.TimeOffsetSeconds;
+	return ProviderResult;
+}
+
+}} // namespace UE::PoseSearch
+
+IMPLEMENT_MODULE(UE::PoseSearch::FModule, PoseSearch)
\ No newline at end of file
diff --git a/Engine/Plugins/Experimental/Animation/PoseSearch/Source/Runtime/Public/PoseSearch/AnimNode_MotionMatching.h b/Engine/Plugins/Experimental/Animation/PoseSearch/Source/Runtime/Public/PoseSearch/AnimNode_MotionMatching.h
new file mode 100644
--- /dev/null
+++ b/Engine/Plugins/Experimental/Animation/PoseSearch/Source/Runtime/Public/PoseSearch/AnimNode_MotionMatching.h
@@ -0,0 +1,68 @@
+// Copyright Epic Games, Inc. All Rights Reserved.
+
+#pragma once
+
+#include "PoseSearch/PoseSearch.h"
+#include "Animation/AnimNodeBase.h"
+#include "Animation/AnimNode_SequencePlayer.h"
+#include "AnimNode_MotionMatching.generated.h"
+
+USTRUCT(BlueprintInternalUseOnly)
+struct POSESEARCH_API FAnimNode_MotionMatching : public FAnimNode_Base
+{
+	GENERATED_BODY()
+
+public:
+	UPROPERTY()
+	FPoseLink Source;
+
+	// Capsule local velocity from inputs
+	// Note this is temporary, we'll need a component to supply future trajectory according to the database schema
+	UPROPERTY(EditAnywhere, BlueprintReadWrite, Category=Settings, meta=(PinShownByDefault))
+	FVector LocalVelocity = FVector::ZeroVector;
+
+	// Collection of animations for motion matching
+	UPROPERTY(EditAnywhere, BlueprintReadWrite, Category=Settings, meta=(PinShownByDefault))
+	const UPoseSearchDatabase* Database = nullptr;
+
+	// Time in seconds to blend out to the new pose. Uses inertial blending and requires an Inertialization node after this node.
+	UPROPERTY(EditAnywhere, BlueprintReadWrite, Category=Settings, meta=(ClampMin="0"))
+	float BlendTime = 0.1f;
+
+	// Ignore pose candidates within the active animation that are less than PoseJumpThreshold seconds away from the current asset player time
+	UPROPERTY(EditAnywhere, BlueprintReadWrite, Category=Settings, meta=(ClampMin="0"))
+	float PoseJumpThreshold = 0.5f;
+
+#if WITH_EDITORONLY_DATA
+	UPROPERTY(EditAnywhere, Category=Debug)
+	bool bDebugDraw = false;
+#endif
+
+	// FAnimNode_Base interface
+	virtual void Initialize_AnyThread(const FAnimationInitializeContext& Context) override;
+	virtual void Evaluate_AnyThread(FPoseContext& Output) override;
+	virtual void Update_AnyThread(const FAnimationUpdateContext& Context) override;
+	virtual void GatherDebugData(FNodeDebugData& DebugData) override;
+	virtual bool HasPreUpdate() const override;
+	virtual void PreUpdate(const UAnimInstance* InAnimInstance) override;
+	// End of FAnimNode_Base interface
+
+private:
+
+	// Embedded sequence player node for playing animations from the motion matching database
+	FAnimNode_SequencePlayer SequencePlayerNode;
+
+	// The current pose we're playing from the database
+	int32 DbPoseIdx = INDEX_NONE;
+
+	// The current animation we're playing from the database
+	int32 DbSequenceIdx = INDEX_NONE;
+
+	// The current query feature vector used to search the database for pose candidates
+	UE::PoseSearch::FFeatureVectorBuilder QueryBuilder;
+	TArray<float> Query;
+
+	bool IsValidForSearch() const;
+	void SetTrajectoryFeatures();
+	void JumpToPose(const FAnimationUpdateContext& Context, UE::PoseSearch::FDbSearchResult Result);
+};
\ No newline at end of file
diff --git a/Engine/Plugins/Experimental/Animation/PoseSearch/Source/Runtime/Public/PoseSearch/AnimNode_PoseSearchHistoryCollector.h b/Engine/Plugins/Experimental/Animation/PoseSearch/Source/Runtime/Public/PoseSearch/AnimNode_PoseSearchHistoryCollector.h
new file mode 100644
--- /dev/null
+++ b/Engine/Plugins/Experimental/Animation/PoseSearch/Source/Runtime/Public/PoseSearch/AnimNode_PoseSearchHistoryCollector.h
@@ -0,0 +1,37 @@
+// Copyright Epic Games, Inc. All Rights Reserved.
+
+#pragma once
+
+#include "PoseSearch/PoseSearch.h"
+#include "Animation/AnimNodeBase.h"
+#include "AnimNode_PoseSearchHistoryCollector.generated.h"
+
+USTRUCT(BlueprintInternalUseOnly)
+struct POSESEARCH_API FAnimNode_PoseSearchHistoryCollector : public FAnimNode_Base
+{
+	GENERATED_BODY()
+
+public:
+	UPROPERTY(EditAnywhere, BlueprintReadWrite, Category = Links)
+	FPoseLink Source;
+
+public:
+	FAnimNode_PoseSearchHistoryCollector()
+	{}
+
+	// FAnimNode_Base interface
+	virtual void Initialize_AnyThread(const FAnimationInitializeContext& Context) override;
+	virtual void Evaluate_AnyThread(FPoseContext& Output) override;
+	virtual void Update_AnyThread(const FAnimationUpdateContext& Context);
+	virtual void GatherDebugData(FNodeDebugData& DebugData);
+	// End of FAnimNode_Base interface
+
+	UE::PoseSearch::FPoseHistory& GetPoseHistory() { return PoseHistory; }
+	const UE::PoseSearch::FPoseHistory& GetPoseHistory() const { return PoseHistory; }
+
+protected:
+
+	UE::PoseSearch::FPoseHistory PoseHistory;
+	TArray<float> Query;
+	float EvalDeltaTime = 0.0f;
+};
\ No newline at end of file
diff --git a/Engine/Plugins/Experimental/Animation/PoseSearch/Source/Runtime/Public/PoseSearch/PoseSearch.h b/Engine/Plugins/Experimental/Animation/PoseSearch/Source/Runtime/Public/PoseSearch/PoseSearch.h
new file mode 100644
--- /dev/null
+++ b/Engine/Plugins/Experimental/Animation/PoseSearch/Source/Runtime/Public/PoseSearch/PoseSearch.h
@@ -0,0 +1,521 @@
+// Copyright Epic Games, Inc. All Rights Reserved.
+
+#pragma once
+
+#include "CoreMinimal.h"
+#include "Containers/RingBuffer.h"
+#include "Engine/DataAsset.h"
+#include "Modules/ModuleManager.h"
+#include "Misc/EnumClassFlags.h"
+#include "UObject/Object.h"
+#include "UObject/ObjectMacros.h"
+#include "Animation/AnimMetaData.h"
+#include "Animation/AnimNodeMessages.h"
+#include "Animation/AnimationAsset.h"
+#include "BoneIndices.h"
+#include "Interfaces/Interface_BoneReferenceSkeletonProvider.h"
+
+#include "PoseSearch.generated.h"
+
+
+class UAnimSequence;
+struct FCompactPose;
+struct FReferenceSkeleton;
+
+UENUM()
+enum class EPoseSearchFeatureType : int32
+{
+	Invalid			= -1,
+
+	Position		= 0,	
+	Rotation		= 1,
+	LinearVelocity	= 2,
+	AngularVelocity	= 3,
+
+	Num
+};
+
+UENUM()
+enum class EPoseSearchFeatureDomain : int32
+{
+	Invalid		= -1,
+
+	Time		= 0,
+	Distance	= 1,
+
+	Num
+};
+
+/** Describes each feature of a vector, including data type, sampling options, and buffer offset. */
+USTRUCT()
+struct POSESEARCH_API FPoseSearchFeatureDesc
+{
+	GENERATED_BODY()
+
+	static constexpr int32 TrajectoryBoneIndex = -1;
+
+	UPROPERTY()
+	int32 SchemaBoneIdx;
+
+	UPROPERTY()
+	int32 SubsampleIdx;
+
+	UPROPERTY()
+	EPoseSearchFeatureType Type;
+
+	UPROPERTY()
+	EPoseSearchFeatureDomain Domain;
+
+	// Set via FPoseSearchFeatureLayout::Init() and ignored by operator==
+	UPROPERTY()
+	int32 ValueOffset;
+
+	bool operator==(const FPoseSearchFeatureDesc& Other) const;
+
+	FPoseSearchFeatureDesc()
+		: SchemaBoneIdx(MAX_int32)
+		, SubsampleIdx(MAX_int32)
+		, Type(EPoseSearchFeatureType::Invalid)
+		, Domain(EPoseSearchFeatureDomain::Invalid)
+		, ValueOffset(MAX_int32)
+	{}
+};
+
+/**
+* Explicit description of a pose feature vector.
+* Determined by options set in a UPoseSearchSchema and owned by the schema.
+* See UPoseSearchSchema::GenerateLayout().
+*/
+USTRUCT()
+struct POSESEARCH_API FPoseSearchFeatureVectorLayout
+{
+	GENERATED_BODY()
+
+	void Init();
+	void Reset();
+
+	UPROPERTY()
+	TArray<FPoseSearchFeatureDesc> Features;
+
+	UPROPERTY()
+	uint32 NumFloats = 0;
+
+	bool IsValid(int32 MaxNumBones) const;
+};
+
+
+//USTRUCT()
+//struct FPoseSearchSchemaBoneRef
+//{
+//	GENERATED_BODY()
+//
+//	/** Index into UPoseSearchSchema::Bones. */
+//	UPROPERTY()
+//	int32 BoneIdx = -1;
+//
+//	/**
+//	* Optional index into UPoseSearchSchema::Bones.
+//	* Causes BoneIdx to be sampled in the reference frame of this bone.
+//	*/
+//	UPROPERTY()
+//	int32 RelativeBoneIdx = -1;
+//};
+
+
+/**
+* Specifies the format of a pose search index. At runtime, queries are built according to the schema for searching.
+*/
+UCLASS(BlueprintType, Category = "Animation|PoseSearch")
+class POSESEARCH_API UPoseSearchSchema : public UDataAsset, public IBoneReferenceSkeletonProvider
+{
+	GENERATED_BODY()
+
+public:
+
+	static constexpr int32 DefaultSampleRate = 10;
+
+	UPROPERTY(EditAnywhere, Category = "Schema")
+	USkeleton* Skeleton;
+
+	UPROPERTY(EditAnywhere, meta = (ClampMin = "1", ClampMax = "60"), Category = "Schema")
+	int32 SampleRate;
+
+	UPROPERTY(EditAnywhere, Category = "Schema")
+	bool bUseBoneVelocities;
+
+	UPROPERTY(EditAnywhere, Category = "Schema")
+	bool bUseBonePositions;
+
+	UPROPERTY(EditAnywhere, Category = "Schema")
+	bool bUseTrajectoryVelocities;
+
+	UPROPERTY(EditAnywhere, Category = "Schema")
+	bool bUseTrajectoryPositions;
+
+	UPROPERTY(EditAnywhere, Category = "Schema")
+	TArray<FBoneReference> Bones;
+
+	/*
+	UPROPERTY(EditAnywhere, Category = "Schema")
+	TArray<FPoseSearchSchemaBoneRef> PoseSampleBones;
+	*/
+
+	UPROPERTY(EditAnywhere, Category = "Schema")
+	TArray<int32> PoseSampleOffsets;
+
+	UPROPERTY(EditAnywhere, Category = "Schema")
+	TArray<int32> TrajectorySampleOffsets;
+
+	UPROPERTY(EditAnywhere, Category = "Schema")
+	TArray<float> TrajectoryDistanceOffsets;
+
+	UPROPERTY()
+	float SamplingInterval;
+
+	UPROPERTY()
+	FPoseSearchFeatureVectorLayout Layout;
+
+	UPROPERTY(Transient)
+	TArray<uint16> BoneIndices;
+
+	UPROPERTY(Transient)
+	TArray<uint16> BoneIndicesWithParents;
+
+	UPoseSearchSchema()
+		: SampleRate(DefaultSampleRate)
+	{}
+
+	bool IsValid () const;
+
+public: // UObject
+	virtual void PreSave(const class ITargetPlatform* TargetPlatform) override;
+	virtual void PostLoad() override;
+
+public: // IBoneReferenceSkeletonProvider
+	class USkeleton* GetSkeleton(bool& bInvalidSkeletonIsError) override { bInvalidSkeletonIsError = false; return Skeleton; }
+
+private:
+	void GenerateLayout();
+	void ResolveBoneReferences();
+};
+
+/**
+* A search index for animation poses. The structure of the search index is determined by its UPoseSearchSchema.
+* May represent a single animation (see UPoseSearchSequenceMetaData) or a collection (see UPoseSearchDatabase).
+*/
+USTRUCT()
+struct POSESEARCH_API FPoseSearchIndex
+{
+	GENERATED_BODY()
+public:
+	UPROPERTY()
+	int32 NumPoses = 0;
+
+	UPROPERTY()
+	TArray<float> Values;
+
+	UPROPERTY()
+	const UPoseSearchSchema* Schema = nullptr;
+
+	bool IsValid() const;
+
+	TArrayView<const float> GetPoseValues(int32 PoseIdx) const;
+
+	void Reset();
+};
+
+/** Animation metadata object for indexing a single animation. */
+UCLASS(BlueprintType, Category = "Animation|PoseSearch")
+class POSESEARCH_API UPoseSearchSequenceMetaData : public UAnimMetaData
+{
+	GENERATED_BODY()
+public:
+
+	UPROPERTY(EditAnywhere, Category="Settings")
+	const UPoseSearchSchema* Schema;
+
+	UPROPERTY(EditAnywhere, Category="Settings")
+	FFloatInterval SamplingRange = FFloatInterval(0.0f, 0.0f);
+
+	UPROPERTY();
+	FPoseSearchIndex SearchIndex;
+
+	bool IsValidForIndexing() const;
+	bool IsValidForSearch() const;
+
+public: // UObject
+	virtual void PreSave(const class ITargetPlatform* TargetPlatform);
+};
+
+/** An entry in a UPoseSearchDatabase. */
+USTRUCT(BlueprintType, Category = "Animation|PoseSearch")
+struct POSESEARCH_API FPoseSearchDatabaseSequence
+{
+	GENERATED_BODY()
+
+	UPROPERTY(EditAnywhere, Category="Sequence")
+	UAnimSequence* Sequence = nullptr;
+
+	UPROPERTY(EditAnywhere, Category="Sequence")
+	FFloatInterval SamplingRange = FFloatInterval(0.0f, 0.0f);
+
+	UPROPERTY(EditAnywhere, Category="Sequence")
+	bool bLoopAnimation = false;
+
+	UPROPERTY()
+	int32 FirstPoseIdx = 0;
+
+	UPROPERTY()
+	int32 NumPoses = 0;
+
+	int32 GetPoseIndexFromAssetTime(float AssetTime) const;
+};
+
+/** A data asset for indexing a collection of animation sequences. */
+UCLASS(BlueprintType, Category = "Animation|PoseSearch")
+class POSESEARCH_API UPoseSearchDatabase : public UDataAsset
+{
+	GENERATED_BODY()
+public:
+
+	UPROPERTY(EditAnywhere, Category="Database")
+	const UPoseSearchSchema* Schema;
+
+	UPROPERTY(EditAnywhere, Category="Database")
+	TArray<FPoseSearchDatabaseSequence> Sequences;
+
+	UPROPERTY()
+	FPoseSearchIndex SearchIndex;
+
+	int32 FindSequenceForPose(int32 PoseIdx) const;
+	int32 GetPoseIndexFromAssetTime(int32 DbSequenceIdx, float AssetTime) const;
+
+	bool IsValidForIndexing() const;
+	bool IsValidForSearch() const;
+
+public: // UObject
+	virtual void PreSave(const class ITargetPlatform* TargetPlatform);
+};
+
+
+namespace UE { namespace PoseSearch {
+
+/** Records poses over time in a ring buffer. FFeatureVectorBuilder uses this to sample from the present or past poses according to the search schema. */
+// @@@Rename to FPoseQueryContext?
+class POSESEARCH_API FPoseHistory
+{
+public:
+	void Init(int32 InNumPoses, float InTimeHorizon);
+	void Init(const FPoseHistory& History);
+	bool Sample(float SecondsAgo, const FReferenceSkeleton& RefSkeleton, const TArray<FBoneIndexType>& RequiredBones);
+	void Update(float SecondsElapsed, const FCompactPose& EvalContext);
+	float GetSampleInterval() const;
+	TArrayView<const FTransform> GetLocalPoseSample() const { return SampledLocalPose; }
+	TArrayView<const FTransform> GetComponentPoseSample() const { return SampledComponentPose; }
+	TArrayView<const FTransform> GetPrevLocalPoseSample() const { return SampledPrevLocalPose; }
+	TArrayView<const FTransform> GetPrevComponentPoseSample() const { return SampledPrevComponentPose; }
+	float GetTimeHorizon() const { return TimeHorizon; }
+	TArray<float>& GetQueryBuffer() { return QueryBuffer; }
+
+private:
+	bool SampleLocalPose(float Time, const TArray<FBoneIndexType>& RequiredBones, TArray<FTransform>& LocalPose);
+
+	struct FPose
+	{
+		TArray<FTransform> LocalTransforms;
+	};
+
+	TRingBuffer<FPose> Poses;
+	TRingBuffer<float> Knots;
+	TArray<FTransform> SampledLocalPose;
+	TArray<FTransform> SampledComponentPose;
+	TArray<FTransform> SampledPrevLocalPose;
+	TArray<FTransform> SampledPrevComponentPose;
+	TArray<float> QueryBuffer;
+
+	float TimeHorizon = 0.0f;
+};
+
+class IPoseHistoryProvider : public UE::Anim::IGraphMessage
+{
+	DECLARE_ANIMGRAPH_MESSAGE(IPoseHistoryProvider);
+
+public:
+
+	virtual const FPoseHistory& GetPoseHistory() const = 0;
+	virtual FPoseHistory& GetPoseHistory() = 0;
+};
+
+
+/** 
+* Helper object for writing features into a float buffer according to a feature vector layout.
+* Keeps track of which features are present, allowing the feature vector to be built up piecemeal.
+* FFeatureVectorBuilder is used to build search queries at runtime and for adding samples during search index construction.
+*/
+class POSESEARCH_API FFeatureVectorBuilder
+{
+public:
+	void Init(const FPoseSearchFeatureVectorLayout* Layout, TArrayView<float> Buffer);
+	void ResetFeatures();
+
+	void SetTransform(FPoseSearchFeatureDesc Feature, const FTransform& Transform);
+	void SetTransformDerivative(FPoseSearchFeatureDesc Feature, const FTransform& Transform, const FTransform& PrevTransform, float DeltaTime);
+
+	void SetPosition(FPoseSearchFeatureDesc Feature, const FVector& Translation);
+	void SetRotation(FPoseSearchFeatureDesc Feature, const FQuat& Rotation);
+	void SetLinearVelocity(FPoseSearchFeatureDesc Feature, const FTransform& Transform, const FTransform& PrevTransform, float DeltaTime);
+	void SetAngularVelocity(FPoseSearchFeatureDesc Feature, const FTransform& Transform, const FTransform& PrevTransform, float DeltaTime);
+
+	void SetVector(FPoseSearchFeatureDesc Feature, const FVector& Vector);
+
+	bool SetPoseFeatures(const UPoseSearchSchema* Schema, FPoseHistory* History);
+
+	void Copy(TArrayView<const float> FeatureVector);
+
+	bool IsComplete() const;
+
+private:
+	const FPoseSearchFeatureVectorLayout* Layout = nullptr;
+	TArrayView<float> Values;
+	TBitArray<> FeaturesAdded;
+	int32 NumFeaturesAdded = 0;
+};
+
+/** Helper object for extracting features from a float buffer according to the feature vector layout. */
+class POSESEARCH_API FFeatureVectorReader
+{
+public:
+	void Init(const FPoseSearchFeatureVectorLayout* Layout);
+	void SetValues(TArrayView<const float> Values);
+	bool IsValid() const;
+
+	bool GetTransform(FPoseSearchFeatureDesc Feature, FTransform* OutTransform) const;
+	bool GetPosition(FPoseSearchFeatureDesc Feature, FVector* OutPosition) const;
+	bool GetRotation(FPoseSearchFeatureDesc Feature, FQuat* OutRotation) const;
+	bool GetLinearVelocity(FPoseSearchFeatureDesc Feature, FVector* OutLinearVelocity) const;
+	bool GetAngularVelocity(FPoseSearchFeatureDesc Feature, FVector* OutAngularVelocity) const;
+	bool GetVector(FPoseSearchFeatureDesc Feature, FVector* OutVector) const;
+
+	const FPoseSearchFeatureVectorLayout* GetLayout() const { return Layout; }
+
+private:
+	const FPoseSearchFeatureVectorLayout* Layout = nullptr;
+	TArrayView<const float> Values;
+};
+
+
+//////////////////////////////////////////////////////////////////////////
+// Main PoseSearch API
+
+enum class EDebugDrawFlags : uint32
+{
+	None			= 0,
+	DrawQuery		= 1 << 0,
+	DrawSearchIndex = 1 << 1,
+	DrawBest        = 1 << 2,
+	DrawAll		    = MAX_uint32
+};
+ENUM_CLASS_FLAGS(EDebugDrawFlags);
+
+struct FDebugDrawParams
+{
+	const UWorld* World = nullptr;
+	const UPoseSearchDatabase* Database = nullptr;
+	const UPoseSearchSequenceMetaData* SequenceMetaData = nullptr;
+	EDebugDrawFlags Flags = EDebugDrawFlags::None;
+	float DefaultLifeTime = 5.0f;
+	FTransform RootTransform = FTransform::Identity;
+	int32 HighlightPoseIdx = -1;
+	TArrayView<const float> Query;
+
+	bool CanDraw () const;
+	const FPoseSearchIndex* GetSearchIndex() const;
+	const UPoseSearchSchema* GetSchema() const;
+};
+
+struct FSearchResult
+{
+	int32 PoseIdx = -1;
+	float TimeOffsetSeconds = 0.0f;
+	float Dissimilarity = MAX_flt;
+
+	bool IsValid() const { return PoseIdx >= 0; }
+};
+
+struct FDbSearchResult : public FSearchResult
+{
+	FDbSearchResult() = default;
+
+	FDbSearchResult(const FSearchResult& Result)
+		: FSearchResult(Result)
+	{}
+
+	int32 DbSequenceIdx = INDEX_NONE;
+};
+
+
+/**
+* Visualize pose search debug information
+*
+* @param DrawParams		Visualization options
+*/
+POSESEARCH_API void Draw(const FDebugDrawParams& DrawParams);
+
+
+/**
+* Creates a pose search index for an animation sequence
+* 
+* @param Sequence			The input sequence create a search index for
+* @param SequenceMetaData	The input sequence indexing info and output search index
+* 
+* @return					Whether the index was built successfully
+*/
+POSESEARCH_API bool BuildIndex(const UAnimSequence* Sequence, UPoseSearchSequenceMetaData* SequenceMetaData);
+
+
+/**
+* Creates a pose search index for a collection of animations
+* 
+* @param Database	The input collection of animations and output search index
+* 
+* @return			Whether the index was built successfully
+*/
+POSESEARCH_API bool BuildIndex(UPoseSearchDatabase* Database);
+
+
+/**
+* Performs a pose search on a sequence's UPoseSearchSequenceMetaData.
+*
+* @param Sequence		The sequence to search within. Must contain a UPoseSearchSequenceMetaData object.
+* @param Query			The pose query to search for. To build a query, see FFeatureVectorBuilder. Must have been built using the same schema as the UPoseSearchSequenceMetaData.
+* @param DrawParams		Visualization options
+* 
+* @return				The pose in the sequence that most closely matches the Query.
+*/
+POSESEARCH_API FSearchResult Search(const UAnimSequenceBase* Sequence, TArrayView<const float> Query, FDebugDrawParams DrawParams = FDebugDrawParams());
+
+
+/**
+* Performs a pose search on a UPoseSearchDatabase.
+*
+* @param Database		The database to search within
+* @param Query			The pose query to search for. To build a query, see FFeatureVectorBuilder. Must have been built using the same schema as the Database.
+* @param DrawParams		Visualization options
+* 
+* @return				The pose in the database that most closely matches the Query.
+*/
+POSESEARCH_API FDbSearchResult Search(const UPoseSearchDatabase* Database, TArrayView<const float> Query, FDebugDrawParams DrawParams = FDebugDrawParams());
+
+
+/**
+* Evaluate pose comparison metric between a pose in the search index and an input query
+* 
+* @param SearchIndex	The search index containing the pose to compare to the query
+* @param PoseIdx		The index of the pose in the search index to compare to the query
+* @param Query			The query to compare against. Must have been built using the same schema as the SearchIndex.
+* 
+* @return Dissimilarity between the two poses
+*/
+POSESEARCH_API float ComparePoses(const FPoseSearchIndex& SearchIndex, int32 PoseIdx, TArrayView<const float> Query);
+
+}} // namespace UE::PoseSearch
\ No newline at end of file
