PoseSearch: Motion matching goal input, basic trajectory prediction
+ FPoseSearchFeatureVectorBuilder
  - Promoted to BlueprintType
  - Now capable of merging feature vector data with another builder
  - Now owns internal value buffer
+ UPoseSearchSchema
  - Replaced user facing sample offsets with sample times for usability
+ AnimNode_MotionMatching
  - No longer takes velocity input. Replaced with a goal feature vector.
  - Goal feature vector is merged with pose goal from Match Pose node to compose the final search query
+ AnimInput_CapsulePrediction
  - Demonstration of basic trajectory prediction and routing into motion matching node via goal feature vector

diff --git a/Engine/Plugins/Experimental/Animation/PoseSearch/Source/Runtime/Private/AnimNode_MotionMatching.cpp b/Engine/Plugins/Experimental/Animation/PoseSearch/Source/Runtime/Private/AnimNode_MotionMatching.cpp
--- a/Engine/Plugins/Experimental/Animation/PoseSearch/Source/Runtime/Private/AnimNode_MotionMatching.cpp
+++ b/Engine/Plugins/Experimental/Animation/PoseSearch/Source/Runtime/Private/AnimNode_MotionMatching.cpp
@@ -9,6 +9,7 @@
 
 #define LOCTEXT_NAMESPACE "AnimNode_MotionMatching"
 
+
 /////////////////////////////////////////////////////
 // FAnimNode_MotionMatching
 
@@ -20,16 +21,14 @@ void FAnimNode_MotionMatching::Initialize_AnyThread(const FAnimationInitializeCo
 
 	if (IsValidForSearch())
 	{
-		Query.Reset();
-		Query.SetNumZeroed(Database->Schema->Layout.NumFloats);
-		QueryBuilder.Init(&Database->Schema->Layout, Query);
+		ComposedQuery.Init(Database->Schema);
 
 		// Set initial pose features
 		UE::PoseSearch::IPoseHistoryProvider* PoseHistoryProvider = Context.GetMessage<UE::PoseSearch::IPoseHistoryProvider>();
 		if (PoseHistoryProvider)
 		{
 			UE::PoseSearch::FPoseHistory& History = PoseHistoryProvider->GetPoseHistory();
-			QueryBuilder.SetPoseFeatures(Database->Schema, &History);
+			ComposedQuery.SetPoseFeatures(&History);
 		}
 	}
 
@@ -68,23 +67,31 @@ void FAnimNode_MotionMatching::Update_AnyThread(const FAnimationUpdateContext& C
 			// Overwrite query feature vector with the stepped pose from the database
 			if (DbPoseIdx != INDEX_NONE)
 			{
-				TArrayView<const float> PoseFeatureVector = Database->SearchIndex.GetPoseValues(DbPoseIdx);
-				QueryBuilder.Copy(PoseFeatureVector);
+				TArrayView<const float> DbFeatureVector = Database->SearchIndex.GetPoseValues(DbPoseIdx);
+				ComposedQuery.Copy(DbFeatureVector);
 			}
 		}
 
-		// Update trajectory features in the query with the latest inputs
-		SetTrajectoryFeatures();
+		// Update features in the query with the latest inputs
+		ComposeQuery();
 
 		// Determine how much the updated query vector deviates from the current pose vector
 		float CurrentDissimilarity = MAX_flt;
 		if (DbPoseIdx != INDEX_NONE)
 		{
-			CurrentDissimilarity = UE::PoseSearch::ComparePoses(Database->SearchIndex, DbPoseIdx, Query);
+			CurrentDissimilarity = UE::PoseSearch::ComparePoses(Database->SearchIndex, DbPoseIdx, ComposedQuery.GetValues());
 		}
 
 		// Search the database for the nearest match to the updated query vector
-		UE::PoseSearch::FDbSearchResult Result = UE::PoseSearch::Search(Database, MakeArrayView(Query));
+
+		UE::PoseSearch::IPoseHistoryProvider* PoseHistoryProvider = Context.GetMessage<UE::PoseSearch::IPoseHistoryProvider>();
+		if (PoseHistoryProvider)
+		{
+			UE::PoseSearch::FPoseHistory& History = PoseHistoryProvider->GetPoseHistory();
+			ComposedQuery.SetPoseFeatures(&History);
+		}
+
+		UE::PoseSearch::FDbSearchResult Result = UE::PoseSearch::Search(Database, ComposedQuery.GetValues());
 		if (Result.IsValid())
 		{
 			const FPoseSearchDatabaseSequence& ResultDbSequence = Database->Sequences[Result.DbSequenceIdx];
@@ -104,11 +111,13 @@ void FAnimNode_MotionMatching::Update_AnyThread(const FAnimationUpdateContext& C
 			}
 
 			// And we won't bother to jump to another pose within the same looping sequence we're already playing
-			// (This should probably should be configurable)
+			// (This should probably be configurable)
 			bool bSameCycle = (DbSequenceIdx == Result.DbSequenceIdx) && ResultDbSequence.bLoopAnimation;
 
+			bool bSameSequence = DbSequenceIdx == Result.DbSequenceIdx;
+
 			// Start playback from the candidate pose if we determined it was a better option
-			if (bBetterPose && !bNearbyPose && !bSameCycle)
+			if (bBetterPose && !bNearbyPose && !bSameCycle && !bSameSequence)
 			{
 				JumpToPose(Context, Result);
 			}
@@ -139,7 +148,7 @@ void FAnimNode_MotionMatching::PreUpdate(const UAnimInstance* InAnimInstance)
 		DrawParams.RootTransform = SkeletalMeshComponent->GetBoneTransform(0);
 		DrawParams.Flags = UE::PoseSearch::EDebugDrawFlags::DrawQuery;
 		DrawParams.Database = Database;
-		DrawParams.Query = Query;
+		DrawParams.Query = ComposedQuery.GetValues();
 		DrawParams.World = SkeletalMeshComponent->GetWorld();
 		DrawParams.DefaultLifeTime = 0.0f;
 		UE::PoseSearch::Draw(DrawParams);
@@ -157,23 +166,12 @@ bool FAnimNode_MotionMatching::IsValidForSearch() const
 	return Database && Database->IsValidForSearch();
 }
 
-void FAnimNode_MotionMatching::SetTrajectoryFeatures()
+void FAnimNode_MotionMatching::ComposeQuery()
 {
-	// This is all placeholder-- trajectory should be provided externally rather than computed in this node.
-	// And we shouldn't be assuming the format of the database's schema!
-
-	// Add instantaneous root velocity to query
-	FPoseSearchFeatureDesc TrajectoryFeature;
-	TrajectoryFeature.Domain = EPoseSearchFeatureDomain::Time;
-	TrajectoryFeature.SchemaBoneIdx = FPoseSearchFeatureDesc::TrajectoryBoneIndex;
-	TrajectoryFeature.SubsampleIdx = 0;
-	TrajectoryFeature.Type = EPoseSearchFeatureType::LinearVelocity;
-
-	// I'm sure there is probably a less magical way to get LocalVelocity into the same space as the animation data...
-	FQuat Rotate(FVector::ZAxisVector, HALF_PI);
-
-	// Apply the updated trajectory to the query feature vector, leaving any existing pose information intact
-	QueryBuilder.SetVector(TrajectoryFeature, Rotate.RotateVector(LocalVelocity));
+	if (ComposedQuery.IsCompatible(Goal))
+	{
+		ComposedQuery.MergeReplace(Goal);
+	}
 }
 
 void FAnimNode_MotionMatching::JumpToPose(const FAnimationUpdateContext& Context, UE::PoseSearch::FDbSearchResult Result)
diff --git a/Engine/Plugins/Experimental/Animation/PoseSearch/Source/Runtime/Private/PoseSearch.cpp b/Engine/Plugins/Experimental/Animation/PoseSearch/Source/Runtime/Private/PoseSearch.cpp
--- a/Engine/Plugins/Experimental/Animation/PoseSearch/Source/Runtime/Private/PoseSearch.cpp
+++ b/Engine/Plugins/Experimental/Animation/PoseSearch/Source/Runtime/Private/PoseSearch.cpp
@@ -69,7 +69,7 @@
 	int32 FeatureIdx = Layout->Features.IndexOfByKey(Feature);
 	check(FeatureIdx != INDEX_NONE);
 	float Lerp = (float)(FeatureIdx) / (Layout->Features.Num() - 1);
-	FLinearColor ColorHSV(Lerp * 360.0f, 0.8f, 0.5f, 0.6f);
+	FLinearColor ColorHSV(Lerp * 360.0f, 0.8f, 0.5f, 1.0f);
 	return ColorHSV.HSVToLinearRGB();
 }
 
@@ -221,11 +221,15 @@
 
 void UPoseSearchSchema::PreSave(const class ITargetPlatform* TargetPlatform)
 {
+	SampleRate = FMath::Clamp(SampleRate, 1, 60);
 	SamplingInterval = 1.0f / SampleRate;
 
-	PoseSampleOffsets.Sort(TLess<>());
-	TrajectorySampleOffsets.Sort(TLess<>());
-	TrajectoryDistanceOffsets.Sort(TLess<>());
+	PoseSampleTimes.Sort(TLess<>());
+	TrajectorySampleTimes.Sort(TLess<>());
+	TrajectorySampleDistances.Sort(TLess<>());
 
+	ConvertTimesToOffsets(PoseSampleTimes, PoseSampleOffsets);
+	ConvertTimesToOffsets(TrajectorySampleTimes, TrajectorySampleOffsets);
+
 	GenerateLayout();
 	ResolveBoneReferences();
@@ -280,7 +284,7 @@
 		}
 	}
 
- 	for (int32 TrajectoryDistSubsampleIdx = 0; TrajectoryDistSubsampleIdx != TrajectoryDistanceOffsets.Num(); ++TrajectoryDistSubsampleIdx)
+ 	for (int32 TrajectoryDistSubsampleIdx = 0; TrajectoryDistSubsampleIdx != TrajectorySampleDistances.Num(); ++TrajectoryDistSubsampleIdx)
  	{
  		FPoseSearchFeatureDesc Element;
  		Element.SchemaBoneIdx = FPoseSearchFeatureDesc::TrajectoryBoneIndex;
@@ -350,6 +354,16 @@
 	}
 }
 
+void UPoseSearchSchema::ConvertTimesToOffsets(TArrayView<const float> SampleTimes, TArray<int32>& OutSampleOffsets)
+{
+	OutSampleOffsets.SetNum(SampleTimes.Num());
 
+	for (int32 Idx = 0; Idx != SampleTimes.Num(); ++Idx)
+	{
+		OutSampleOffsets[Idx] = FMath::RoundToInt(SampleTimes[Idx] * SampleRate);
+	}
+}
+
+
 //////////////////////////////////////////////////////////////////////////
 // FPoseSearchIndex
@@ -469,198 +483,52 @@
 }
 
 
-namespace UE { namespace PoseSearch {
+//////////////////////////////////////////////////////////////////////////
+// FPoseSearchFeatureVectorBuilder
 
-//////////////////////////////////////////////////////////////////////////
-// FPoseHistory
-
-/**
-* Fills skeleton transforms with evaluated compact pose transforms.
-* Bones that weren't evaluated are filled with the bone's reference pose.
-*/
-static void CopyCompactToSkeletonPose(const FCompactPose& Pose, TArray<FTransform>& OutLocalTransforms)
+void FPoseSearchFeatureVectorBuilder::Init(const UPoseSearchSchema* InSchema)
 {
-	const FBoneContainer& BoneContainer = Pose.GetBoneContainer();
-	const FReferenceSkeleton& RefSkeleton = BoneContainer.GetReferenceSkeleton();
-	TArrayView<const FTransform> RefSkeletonTransforms = MakeArrayView(RefSkeleton.GetRefBonePose());
-
-	const int32 NumSkeletonBones = BoneContainer.GetNumBones();
-	OutLocalTransforms.SetNum(NumSkeletonBones);
-
-	for (auto SkeletonBoneIdx = FSkeletonPoseBoneIndex(0); SkeletonBoneIdx != NumSkeletonBones; ++SkeletonBoneIdx)
-	{
-		FCompactPoseBoneIndex CompactBoneIdx = BoneContainer.GetCompactPoseIndexFromSkeletonIndex(SkeletonBoneIdx.GetInt());
-		OutLocalTransforms[SkeletonBoneIdx.GetInt()] = CompactBoneIdx.IsValid() ? Pose[CompactBoneIdx] : RefSkeletonTransforms[SkeletonBoneIdx.GetInt()];
-	}
-}
-
-void FPoseHistory::Init(int32 InNumPoses, float InTimeHorizon)
-{
-	Poses.Reserve(InNumPoses);
-	Knots.Reserve(InNumPoses);
-	TimeHorizon = InTimeHorizon;
-}
-
-void FPoseHistory::Init(const FPoseHistory& History)
-{
-	Poses = History.Poses;
-	Knots = History.Knots;
-	TimeHorizon = History.TimeHorizon;
-}
-
-bool FPoseHistory::SampleLocalPose(float SecondsAgo, const TArray<FBoneIndexType>& RequiredBones, TArray<FTransform>& LocalPose)
-{
-	int32 NextIdx = LowerBound(Knots.begin(), Knots.end(), SecondsAgo, TGreater<>());
-	if (NextIdx >= Knots.Num())
-	{
-		return false;
-	}
-
-	int32 PrevIdx = NextIdx - 1;
-
-	// Compute alpha between previous and next knots
-	float Alpha = FMath::GetMappedRangeValueUnclamped(
-		FVector2D(Knots[PrevIdx], Knots[NextIdx]),
-		FVector2D(0.0f, 1.0f),
-		SecondsAgo);
-
-	TArray<FTransform>& PrevPose = Poses[PrevIdx].LocalTransforms;
-	TArray<FTransform>& NextPose = Poses[NextIdx].LocalTransforms;
-
-	// We may not have accumulated enough poses yet
-	if (PrevPose.Num() != NextPose.Num())
-	{
-		return false;
-	}
-
-	if (RequiredBones.Num() > PrevPose.Num())
-	{
-		return false;
-	}
-
-	// Lerp between poses by alpha to produce output local pose at requested sample time
-	LocalPose = PrevPose;
-	FAnimationRuntime::LerpBoneTransforms(
-		LocalPose,
-		NextPose,
-		Alpha,
-		RequiredBones);
-
-	return true;
-}
-
-bool FPoseHistory::Sample(float SecondsAgo, const FReferenceSkeleton& RefSkeleton, const TArray<FBoneIndexType>& RequiredBones)
-{
-	// Compute local space pose at requested time
-	bool bSampled = SampleLocalPose(SecondsAgo, RequiredBones, SampledLocalPose);
-
-	// Compute local space pose one sample interval in the past
-	bSampled = bSampled && SampleLocalPose(SecondsAgo + GetSampleInterval(), RequiredBones, SampledPrevLocalPose);
-
-	// Convert local to component space
-	if (bSampled)
-	{
-		FAnimationRuntime::FillUpComponentSpaceTransforms(RefSkeleton, SampledLocalPose, SampledComponentPose);
-		FAnimationRuntime::FillUpComponentSpaceTransforms(RefSkeleton, SampledPrevLocalPose, SampledPrevComponentPose);
-	}
-
-	return bSampled;
-}
-
-void FPoseHistory::Update(float SecondsElapsed, const FCompactPose& Pose)
-{
-	// Age our elapsed times
-	for (float& Knot : Knots)
-	{
-		Knot += SecondsElapsed;
-	}
-
-	if (Knots.Num() != Knots.Max())
-	{
-		// Consume every pose until the queue is full
-		Knots.AddUninitialized();
-		Poses.Emplace();
-	}
-	else
-	{
-		// Exercise pose retention policy. We must guarantee there is always one additional knot
-		// at or beyond the desired time horizon H so we can fulfill sample requests at t=H. We also
-		// want to evenly distribute knots across the entire history buffer so we only push additional
-		// poses when enough time has elapsed.
-
-		const float SampleInterval = GetSampleInterval();
-
-		bool bCanEvictOldest = Knots[1] >= TimeHorizon;
-		bool bShouldPushNewest = Knots[Knots.Num() - 2] >= SampleInterval;
-
-		if (bCanEvictOldest && bShouldPushNewest)
-		{
-			FPose PoseTemp = MoveTemp(Poses.First());
-			Poses.PopFront();
-			Poses.Emplace(MoveTemp(PoseTemp));
-
-			Knots.PopFront();
-			Knots.AddUninitialized();
-		}
-	}
-
-	// Regardless of the retention policy, we always update the most recent pose
-	Knots.Last() = 0.0f;
-	CopyCompactToSkeletonPose(Pose, Poses.Last().LocalTransforms);
-}
-
-float FPoseHistory::GetSampleInterval() const
-{
-	return TimeHorizon / Knots.Max();
-}
-
-
-//////////////////////////////////////////////////////////////////////////
-// FFeatureVectorBuilder
-
-void FFeatureVectorBuilder::Init(const FPoseSearchFeatureVectorLayout* InLayout, TArrayView<float> Buffer)
-{
-	check(InLayout);
-	check(Buffer.Num() == InLayout->NumFloats);
-	Layout = InLayout;
-	Values = Buffer;
+	check(InSchema && InSchema->IsValid());
+	Schema = InSchema;
 	ResetFeatures();
 }
 
-void FFeatureVectorBuilder::ResetFeatures()
+void FPoseSearchFeatureVectorBuilder::ResetFeatures()
 {
+	Values.Reset(0);
+	Values.SetNumZeroed(Schema->Layout.NumFloats);
 	NumFeaturesAdded = 0;
-	FeaturesAdded.Init(false, Layout->Features.Num());
+	FeaturesAdded.Init(false, Schema->Layout.Features.Num());
 }
 
-void FFeatureVectorBuilder::SetTransform(FPoseSearchFeatureDesc Element, const FTransform& Transform)
+void FPoseSearchFeatureVectorBuilder::SetTransform(FPoseSearchFeatureDesc Element, const FTransform& Transform)
 {
 	SetPosition(Element, Transform.GetTranslation());
 	SetRotation(Element, Transform.GetRotation());
 }
 
-void FFeatureVectorBuilder::SetTransformDerivative(FPoseSearchFeatureDesc Element, const FTransform& Transform, const FTransform& PrevTransform, float DeltaTime)
+void FPoseSearchFeatureVectorBuilder::SetTransformDerivative(FPoseSearchFeatureDesc Element, const FTransform& Transform, const FTransform& PrevTransform, float DeltaTime)
 {
 	SetLinearVelocity(Element, Transform, PrevTransform, DeltaTime);
 	SetAngularVelocity(Element, Transform, PrevTransform, DeltaTime);
 }
 
-void FFeatureVectorBuilder::SetPosition(FPoseSearchFeatureDesc Element, const FVector& Position)
+void FPoseSearchFeatureVectorBuilder::SetPosition(FPoseSearchFeatureDesc Element, const FVector& Position)
 {
 	Element.Type = EPoseSearchFeatureType::Position;
 	SetVector(Element, Position);
 }
 
-void FFeatureVectorBuilder::SetRotation(FPoseSearchFeatureDesc Element, const FQuat& Rotation)
+void FPoseSearchFeatureVectorBuilder::SetRotation(FPoseSearchFeatureDesc Element, const FQuat& Rotation)
 {
 	Element.Type = EPoseSearchFeatureType::Rotation;
-	int32 ElementIndex = Layout->Features.Find(Element);
+	int32 ElementIndex = Schema->Layout.Features.Find(Element);
 	if (ElementIndex >= 0)
 	{
 		FVector X = Rotation.GetAxisX();
 		FVector Y = Rotation.GetAxisY();
 
-		const FPoseSearchFeatureDesc& FoundElement = Layout->Features[ElementIndex];
+		const FPoseSearchFeatureDesc& FoundElement = Schema->Layout.Features[ElementIndex];
 
 		Values[FoundElement.ValueOffset + 0] = X.X;
 		Values[FoundElement.ValueOffset + 1] = X.Y;
@@ -677,17 +545,17 @@
 	}
 }
 
-void FFeatureVectorBuilder::SetLinearVelocity(FPoseSearchFeatureDesc Element, const FTransform& Transform, const FTransform& PrevTransform, float DeltaTime)
+void FPoseSearchFeatureVectorBuilder::SetLinearVelocity(FPoseSearchFeatureDesc Element, const FTransform& Transform, const FTransform& PrevTransform, float DeltaTime)
 {
 	Element.Type = EPoseSearchFeatureType::LinearVelocity;
 	FVector LinearVelocity = (Transform.GetTranslation() - PrevTransform.GetTranslation()) / DeltaTime;
 	SetVector(Element, LinearVelocity);
 }
 
-void FFeatureVectorBuilder::SetAngularVelocity(FPoseSearchFeatureDesc Element, const FTransform& Transform, const FTransform& PrevTransform, float DeltaTime)
+void FPoseSearchFeatureVectorBuilder::SetAngularVelocity(FPoseSearchFeatureDesc Element, const FTransform& Transform, const FTransform& PrevTransform, float DeltaTime)
 {
 	Element.Type = EPoseSearchFeatureType::AngularVelocity;
-	int32 ElementIndex = Layout->Features.Find(Element);
+	int32 ElementIndex = Schema->Layout.Features.Find(Element);
 	if (ElementIndex >= 0)
 	{
 		FQuat Q0 = PrevTransform.GetRotation();
@@ -706,7 +574,7 @@
 
 		FVector AngularVelocity(W.X, W.Y, W.Z);
 
-		const FPoseSearchFeatureDesc& FoundElement = Layout->Features[ElementIndex];
+		const FPoseSearchFeatureDesc& FoundElement = Schema->Layout.Features[ElementIndex];
 
 		Values[FoundElement.ValueOffset + 0] = AngularVelocity[0];
 		Values[FoundElement.ValueOffset + 1] = AngularVelocity[1];
@@ -720,12 +588,12 @@
 	}
 }
 
-void FFeatureVectorBuilder::SetVector(FPoseSearchFeatureDesc Element, const FVector& Vector)
+void FPoseSearchFeatureVectorBuilder::SetVector(FPoseSearchFeatureDesc Element, const FVector& Vector)
 {
-	int32 ElementIndex = Layout->Features.Find(Element);
+	int32 ElementIndex = Schema->Layout.Features.Find(Element);
 	if (ElementIndex >= 0)
 	{
-		const FPoseSearchFeatureDesc& FoundElement = Layout->Features[ElementIndex];
+		const FPoseSearchFeatureDesc& FoundElement = Schema->Layout.Features[ElementIndex];
 
 		Values[FoundElement.ValueOffset + 0] = Vector[0];
 		Values[FoundElement.ValueOffset + 1] = Vector[1];
@@ -739,7 +607,7 @@
 	}
 }
 
-bool FFeatureVectorBuilder::SetPoseFeatures(const UPoseSearchSchema* Schema, FPoseHistory* History)
+bool FPoseSearchFeatureVectorBuilder::SetPoseFeatures(UE::PoseSearch::FPoseHistory* History)
 {
 	check(Schema && Schema->IsValid());
 	check(History);
@@ -776,17 +644,204 @@
 	return true;
 }
 
-void FFeatureVectorBuilder::Copy(TArrayView<const float> FeatureVector)
+void FPoseSearchFeatureVectorBuilder::Copy(TArrayView<const float> FeatureVector)
 {
 	check(FeatureVector.Num() == Values.Num());
 	FMemory::Memcpy(Values.GetData(), FeatureVector.GetData(), FeatureVector.GetTypeSize() * FeatureVector.Num());
-	NumFeaturesAdded = Layout->Features.Num();
+	NumFeaturesAdded = Schema->Layout.Features.Num();
 	FeaturesAdded.SetRange(0, FeaturesAdded.Num(), true);
 }
 
-bool FFeatureVectorBuilder::IsComplete() const
+void FPoseSearchFeatureVectorBuilder::CopyFeature(const FPoseSearchFeatureVectorBuilder& OtherBuilder, int32 FeatureIdx)
 {
-	return NumFeaturesAdded == Layout->Features.Num();
+	check(IsCompatible(OtherBuilder));
+	check(OtherBuilder.FeaturesAdded[FeatureIdx]);
+
+	const FPoseSearchFeatureDesc& FeatureDesc = Schema->Layout.Features[FeatureIdx];
+	const int32 FeatureNumFloats = UE::PoseSearch::GetFeatureTypeTraits(FeatureDesc.Type).NumFloats;
+	const int32 FeatureValueOffset = FeatureDesc.ValueOffset;
+
+	for(int32 FeatureValueIdx = FeatureValueOffset; FeatureValueIdx != FeatureValueOffset + FeatureNumFloats; ++FeatureValueIdx)
+	{
+		Values[FeatureValueIdx] = OtherBuilder.Values[FeatureValueIdx];
+	}
+
+	if (!FeaturesAdded[FeatureIdx])
+	{
+		FeaturesAdded[FeatureIdx] = true;
+		++NumFeaturesAdded;
+	}
+}
+
+void FPoseSearchFeatureVectorBuilder::MergeReplace(const FPoseSearchFeatureVectorBuilder& OtherBuilder)
+{
+	check(IsCompatible(OtherBuilder));
+
+	for (TConstSetBitIterator<> Iter(OtherBuilder.FeaturesAdded); Iter; ++Iter)
+	{
+		CopyFeature(OtherBuilder, Iter.GetIndex());
+	}
+}
+
+bool FPoseSearchFeatureVectorBuilder::IsInitialized() const
+{
+	return Schema != nullptr;
+}
+
+bool FPoseSearchFeatureVectorBuilder::IsComplete() const
+{
+	return NumFeaturesAdded == Schema->Layout.Features.Num();
+}
+
+bool FPoseSearchFeatureVectorBuilder::IsCompatible(const FPoseSearchFeatureVectorBuilder& OtherBuilder) const
+{
+	return IsInitialized() && (Schema == OtherBuilder.Schema);
 }
 
 
+namespace UE { namespace PoseSearch {
+
+//////////////////////////////////////////////////////////////////////////
+// FPoseHistory
+
+/**
+* Fills skeleton transforms with evaluated compact pose transforms.
+* Bones that weren't evaluated are filled with the bone's reference pose.
+*/
+static void CopyCompactToSkeletonPose(const FCompactPose& Pose, TArray<FTransform>& OutLocalTransforms)
+{
+	const FBoneContainer& BoneContainer = Pose.GetBoneContainer();
+	const FReferenceSkeleton& RefSkeleton = BoneContainer.GetReferenceSkeleton();
+	TArrayView<const FTransform> RefSkeletonTransforms = MakeArrayView(RefSkeleton.GetRefBonePose());
+
+	const int32 NumSkeletonBones = BoneContainer.GetNumBones();
+	OutLocalTransforms.SetNum(NumSkeletonBones);
+
+	for (auto SkeletonBoneIdx = FSkeletonPoseBoneIndex(0); SkeletonBoneIdx != NumSkeletonBones; ++SkeletonBoneIdx)
+	{
+		FCompactPoseBoneIndex CompactBoneIdx = BoneContainer.GetCompactPoseIndexFromSkeletonIndex(SkeletonBoneIdx.GetInt());
+		OutLocalTransforms[SkeletonBoneIdx.GetInt()] = CompactBoneIdx.IsValid() ? Pose[CompactBoneIdx] : RefSkeletonTransforms[SkeletonBoneIdx.GetInt()];
+	}
+}
+
+void FPoseHistory::Init(int32 InNumPoses, float InTimeHorizon)
+{
+	Poses.Reserve(InNumPoses);
+	Knots.Reserve(InNumPoses);
+	TimeHorizon = InTimeHorizon;
+}
+
+void FPoseHistory::Init(const FPoseHistory& History)
+{
+	Poses = History.Poses;
+	Knots = History.Knots;
+	TimeHorizon = History.TimeHorizon;
+}
+
+bool FPoseHistory::SampleLocalPose(float SecondsAgo, const TArray<FBoneIndexType>& RequiredBones, TArray<FTransform>& LocalPose)
+{
+	int32 NextIdx = LowerBound(Knots.begin(), Knots.end(), SecondsAgo, TGreater<>());
+	if (NextIdx <= 0 || NextIdx >= Knots.Num())
+	{
+		return false;
+	}
+
+	int32 PrevIdx = NextIdx - 1;
+
+	// Compute alpha between previous and next knots
+	float Alpha = FMath::GetMappedRangeValueUnclamped(
+		FVector2D(Knots[PrevIdx], Knots[NextIdx]),
+		FVector2D(0.0f, 1.0f),
+		SecondsAgo);
+
+	TArray<FTransform>& PrevPose = Poses[PrevIdx].LocalTransforms;
+	TArray<FTransform>& NextPose = Poses[NextIdx].LocalTransforms;
+
+	// We may not have accumulated enough poses yet
+	if (PrevPose.Num() != NextPose.Num())
+	{
+		return false;
+	}
+
+	if (RequiredBones.Num() > PrevPose.Num())
+	{
+		return false;
+	}
+
+	// Lerp between poses by alpha to produce output local pose at requested sample time
+	LocalPose = PrevPose;
+	FAnimationRuntime::LerpBoneTransforms(
+		LocalPose,
+		NextPose,
+		Alpha,
+		RequiredBones);
+
+	return true;
+}
+
+bool FPoseHistory::Sample(float SecondsAgo, const FReferenceSkeleton& RefSkeleton, const TArray<FBoneIndexType>& RequiredBones)
+{
+	// Compute local space pose at requested time
+	bool bSampled = SampleLocalPose(SecondsAgo, RequiredBones, SampledLocalPose);
+
+	// Compute local space pose one sample interval in the past
+	bSampled = bSampled && SampleLocalPose(SecondsAgo + GetSampleInterval(), RequiredBones, SampledPrevLocalPose);
+
+	// Convert local to component space
+	if (bSampled)
+	{
+		FAnimationRuntime::FillUpComponentSpaceTransforms(RefSkeleton, SampledLocalPose, SampledComponentPose);
+		FAnimationRuntime::FillUpComponentSpaceTransforms(RefSkeleton, SampledPrevLocalPose, SampledPrevComponentPose);
+	}
+
+	return bSampled;
+}
+
+void FPoseHistory::Update(float SecondsElapsed, const FCompactPose& Pose)
+{
+	// Age our elapsed times
+	for (float& Knot : Knots)
+	{
+		Knot += SecondsElapsed;
+	}
+
+	if (Knots.Num() != Knots.Max())
+	{
+		// Consume every pose until the queue is full
+		Knots.AddUninitialized();
+		Poses.Emplace();
+	}
+	else
+	{
+		// Exercise pose retention policy. We must guarantee there is always one additional knot
+		// at or beyond the desired time horizon H so we can fulfill sample requests at t=H. We also
+		// want to evenly distribute knots across the entire history buffer so we only push additional
+		// poses when enough time has elapsed.
+
+		const float SampleInterval = GetSampleInterval();
+
+		bool bCanEvictOldest = Knots[1] >= TimeHorizon;
+		bool bShouldPushNewest = Knots[Knots.Num() - 2] >= SampleInterval;
+
+		if (bCanEvictOldest && bShouldPushNewest)
+		{
+			FPose PoseTemp = MoveTemp(Poses.First());
+			Poses.PopFront();
+			Poses.Emplace(MoveTemp(PoseTemp));
+
+			Knots.PopFront();
+			Knots.AddUninitialized();
+		}
+	}
+
+	// Regardless of the retention policy, we always update the most recent pose
+	Knots.Last() = 0.0f;
+	CopyCompactToSkeletonPose(Pose, Poses.Last().LocalTransforms);
+}
+
+float FPoseHistory::GetSampleInterval() const
+{
+	return TimeHorizon / Knots.Max();
+}
+
+
@@ -895,7 +950,7 @@
 
 
 //////////////////////////////////////////////////////////////////////////
-// FSequenceIndexer
+// FDebugDrawParams
 
 bool FDebugDrawParams::CanDraw () const
 {
@@ -960,7 +1015,7 @@
 
 private:
 	void SampleBegin(int32 SampleIdx);
-	void SampleEnd();
+	void SampleEnd(int32 SampleIdx);
 	void ExtractPoses(const UAnimSequence* Sequence);
 	void ExtractRootMotion(const UAnimSequence* Sequence);
 	void AddPoseFeatures(int32 SampleIdx);
@@ -988,7 +1043,7 @@
 
 	TArray<float> Values;
 	
-	FFeatureVectorBuilder Builder;
+	FPoseSearchFeatureVectorBuilder Builder;
 	FSampleContext Context;
 };
 
@@ -1049,7 +1104,7 @@
 		AddTrajectoryTimeFeatures(SampleIdx);
 		AddTrajectoryDistanceFeatures(SampleIdx);
 
-		SampleEnd();
+		SampleEnd(SampleIdx);
 	}
 
 	return GetResult();
@@ -1065,14 +1120,20 @@
 
 void FSequenceIndexer::SampleBegin(int32 SampleIdx)
 {
-	int32 FirstValueIdx = (SampleIdx - Context.FirstIndexedSample) * Schema->Layout.NumFloats;
-	TArrayView<float> FeatureVectorValues = MakeArrayView(&Values[FirstValueIdx], Schema->Layout.NumFloats);
-	Builder.Init(&Schema->Layout, FeatureVectorValues);
+	Builder.Init(Schema);
 }
 
-void FSequenceIndexer::SampleEnd()
+void FSequenceIndexer::SampleEnd(int32 SampleIdx)
 {
 	check(Builder.IsComplete());
+
+	int32 FirstValueIdx = (SampleIdx - Context.FirstIndexedSample) * Schema->Layout.NumFloats;
+	TArrayView<float> WriteValues = MakeArrayView(&Values[FirstValueIdx], Schema->Layout.NumFloats);
+
+	TArrayView<const float> ReadValues = Builder.GetValues();
+	
+	check(WriteValues.Num() == ReadValues.Num());
+	FMemory::Memcpy(WriteValues.GetData(), ReadValues.GetData(), WriteValues.Num() * WriteValues.GetTypeSize());
 }
 
 void FSequenceIndexer::ExtractPoses(const UAnimSequence* Sequence)
@@ -1211,11 +1272,11 @@
 
 	FTransform SampleSpaceOrigin = Context.AccumulatedRootMotion[SampleIdx];
 
-	for (int32 SubsampleIdx = 0; SubsampleIdx != Schema->TrajectoryDistanceOffsets.Num(); ++SubsampleIdx)
+	for (int32 SubsampleIdx = 0; SubsampleIdx != Schema->TrajectorySampleDistances.Num(); ++SubsampleIdx)
 	{
 		CurrentElement.SubsampleIdx = SubsampleIdx;
 
-		const float TrajectoryDistance = Schema->TrajectoryDistanceOffsets[SubsampleIdx];
+		const float TrajectoryDistance = Schema->TrajectorySampleDistances[SubsampleIdx];
 		const float SampleAccumulatedRootDistance = TrajectoryDistance + AccumulatedRootDistances[SampleIdx];
 
 		int32 LowerBoundSampleIdx = Algo::LowerBound(AccumulatedRootDistances, SampleAccumulatedRootDistance);
@@ -1276,13 +1337,6 @@
 			TrajectoryPos = DrawParams.RootTransform.TransformPosition(TrajectoryPos);
 			DrawDebugSphere(DrawParams.World, TrajectoryPos, DrawDebugSphereSize, DrawDebugSphereSegments, Color, false, LifeTime, DepthPriority,  DrawDebugSphereLineThickness);
 
-			if (SubsampleIdx != 0)
-			{
-				FVector Direction;
-				float Length;
-				(TrajectoryPos - TrajectoryPosPrev).ToDirectionAndLength(Direction, Length);
-				DrawDebugLine(DrawParams.World, TrajectoryPosPrev, TrajectoryPos, Color, false, LifeTime, DepthPriority, DrawDebugLineThickness);
-			}
 
 			TrajectoryPosPrev = TrajectoryPos;
 		}
@@ -1412,6 +1466,11 @@
 	const UPoseSearchSchema* Schema = DrawParams.GetSchema();
 	check(Schema);
 
+	if (DrawParams.Query.Num() != Schema->Layout.NumFloats)
+	{
+		return;
+	}
+
 	FFeatureVectorReader Reader;
 	Reader.Init(&Schema->Layout);
 	Reader.SetValues(DrawParams.Query);
@@ -1672,16 +1731,12 @@
 	}
 
 	FPoseHistory& PoseHistory = PoseHistoryProvider->GetPoseHistory();
-	TArray<float>& QueryBuffer = PoseHistory.GetQueryBuffer();
+	FPoseSearchFeatureVectorBuilder& QueryBuilder = PoseHistory.GetQueryBuilder();
 
-	QueryBuffer.SetNum(MetaData->Schema->Layout.NumFloats);
-	TArrayView<float> Query = MakeArrayView(QueryBuffer.GetData(), QueryBuffer.Num());
+	QueryBuilder.Init(MetaData->Schema);
+	QueryBuilder.SetPoseFeatures(&PoseHistory);
 
-	FFeatureVectorBuilder QueryBuilder;
-	QueryBuilder.Init(&MetaData->Schema->Layout, Query);
-	QueryBuilder.SetPoseFeatures(MetaData->Schema, &PoseHistory);
+	::UE::PoseSearch::FSearchResult Result = ::UE::PoseSearch::Search(Sequence, QueryBuilder.GetValues());
 
-	::UE::PoseSearch::FSearchResult Result = ::UE::PoseSearch::Search(Sequence, Query);
-
 	ProviderResult.Dissimilarity = Result.Dissimilarity;
 	ProviderResult.PoseIdx = Result.PoseIdx;
diff --git a/Engine/Plugins/Experimental/Animation/PoseSearch/Source/Runtime/Public/PoseSearch/AnimNode_MotionMatching.h b/Engine/Plugins/Experimental/Animation/PoseSearch/Source/Runtime/Public/PoseSearch/AnimNode_MotionMatching.h
--- a/Engine/Plugins/Experimental/Animation/PoseSearch/Source/Runtime/Public/PoseSearch/AnimNode_MotionMatching.h
+++ b/Engine/Plugins/Experimental/Animation/PoseSearch/Source/Runtime/Public/PoseSearch/AnimNode_MotionMatching.h
@@ -16,15 +16,14 @@ public:
 	UPROPERTY()
 	FPoseLink Source;
 
-	// Capsule local velocity from inputs
-	// Note this is temporary, we'll need a component to supply future trajectory according to the database schema
-	UPROPERTY(EditAnywhere, BlueprintReadWrite, Category=Settings, meta=(PinShownByDefault))
-	FVector LocalVelocity = FVector::ZeroVector;
-
 	// Collection of animations for motion matching
 	UPROPERTY(EditAnywhere, BlueprintReadWrite, Category=Settings, meta=(PinShownByDefault))
 	const UPoseSearchDatabase* Database = nullptr;
 
+	// Motion matching goal
+	UPROPERTY(EditAnywhere, BlueprintReadWrite, Category=Settings, meta=(PinShownByDefault))
+	FPoseSearchFeatureVectorBuilder Goal;
+
 	// Time in seconds to blend out to the new pose. Uses inertial blending and requires an Inertialization node after this node.
 	UPROPERTY(EditAnywhere, BlueprintReadWrite, Category=Settings, meta=(ClampMin="0"))
 	float BlendTime = 0.1f;
@@ -59,10 +58,9 @@ private:
 	int32 DbSequenceIdx = INDEX_NONE;
 
 	// The current query feature vector used to search the database for pose candidates
-	UE::PoseSearch::FFeatureVectorBuilder QueryBuilder;
-	TArray<float> Query;
+	FPoseSearchFeatureVectorBuilder ComposedQuery;
 
 	bool IsValidForSearch() const;
-	void SetTrajectoryFeatures();
+	void ComposeQuery();
 	void JumpToPose(const FAnimationUpdateContext& Context, UE::PoseSearch::FDbSearchResult Result);
 };
\ No newline at end of file
diff --git a/Engine/Plugins/Experimental/Animation/PoseSearch/Source/Runtime/Public/PoseSearch/PoseSearch.h b/Engine/Plugins/Experimental/Animation/PoseSearch/Source/Runtime/Public/PoseSearch/PoseSearch.h
--- a/Engine/Plugins/Experimental/Animation/PoseSearch/Source/Runtime/Public/PoseSearch/PoseSearch.h
+++ b/Engine/Plugins/Experimental/Animation/PoseSearch/Source/Runtime/Public/PoseSearch/PoseSearch.h
@@ -22,6 +22,10 @@ class UAnimSequence;
 struct FCompactPose;
 struct FReferenceSkeleton;
 
+namespace UE { namespace PoseSearch {
+class FPoseHistory;
+}}
+
 UENUM()
 enum class EPoseSearchFeatureType : int32
 {
@@ -161,13 +165,19 @@ public:
 	*/
 
 	UPROPERTY(EditAnywhere, Category = "Schema")
-	TArray<int32> PoseSampleOffsets;
+	TArray<float> PoseSampleTimes;
 
 	UPROPERTY(EditAnywhere, Category = "Schema")
-	TArray<int32> TrajectorySampleOffsets;
+	TArray<float> TrajectorySampleTimes;
 
 	UPROPERTY(EditAnywhere, Category = "Schema")
-	TArray<float> TrajectoryDistanceOffsets;
+	TArray<float> TrajectorySampleDistances;
+
+	UPROPERTY()
+	TArray<int32> TrajectorySampleOffsets;
+
+	UPROPERTY()
+	TArray<int32> PoseSampleOffsets;
 
 	UPROPERTY()
 	float SamplingInterval;
@@ -187,6 +197,8 @@ public:
 
 	bool IsValid () const;
 
+	int32 NumBones () const { return BoneIndices.Num(); }
+
 public: // UObject
 	virtual void PreSave(const class ITargetPlatform* TargetPlatform) override;
 	virtual void PostLoad() override;
@@ -197,6 +209,7 @@ public: // IBoneReferenceSkeletonProvider
 private:
 	void GenerateLayout();
 	void ResolveBoneReferences();
+	void ConvertTimesToOffsets(TArrayView<const float> SampleTimes, TArray<int32>& OutSampleOffsets);
 };
 
 /**
@@ -298,6 +311,50 @@ public: // UObject
 };
 
 
+/** 
+* Helper object for writing features into a float buffer according to a feature vector layout.
+* Keeps track of which features are present, allowing the feature vector to be built up piecemeal.
+* FPoseSearchFeatureVectorBuilder is used to build search queries at runtime and during search index construction.
+*/
+
+USTRUCT(BlueprintType, Category = "Animation|PoseSearch")
+struct POSESEARCH_API FPoseSearchFeatureVectorBuilder
+{
+	GENERATED_BODY()
+public:
+	void Init(const UPoseSearchSchema* Schema);
+	void ResetFeatures();
+
+	TArrayView<const float> GetValues() const { return Values; }
+
+	void SetTransform(FPoseSearchFeatureDesc Feature, const FTransform& Transform);
+	void SetTransformDerivative(FPoseSearchFeatureDesc Feature, const FTransform& Transform, const FTransform& PrevTransform, float DeltaTime);
+	void SetPosition(FPoseSearchFeatureDesc Feature, const FVector& Translation);
+	void SetRotation(FPoseSearchFeatureDesc Feature, const FQuat& Rotation);
+	void SetLinearVelocity(FPoseSearchFeatureDesc Feature, const FTransform& Transform, const FTransform& PrevTransform, float DeltaTime);
+	void SetAngularVelocity(FPoseSearchFeatureDesc Feature, const FTransform& Transform, const FTransform& PrevTransform, float DeltaTime);
+	void SetVector(FPoseSearchFeatureDesc Feature, const FVector& Vector);
+	bool SetPoseFeatures(UE::PoseSearch::FPoseHistory* History);
+
+	void Copy(TArrayView<const float> FeatureVector);
+	void CopyFeature(const FPoseSearchFeatureVectorBuilder& OtherBuilder, int32 FeatureIdx);
+
+	void MergeReplace(const FPoseSearchFeatureVectorBuilder& OtherBuilder);
+
+	bool IsInitialized() const;
+	bool IsComplete() const;
+	bool IsCompatible(const FPoseSearchFeatureVectorBuilder& OtherBuilder) const;
+
+private:
+	UPROPERTY(Transient)
+	const UPoseSearchSchema* Schema = nullptr;
+
+	TArray<float> Values;
+	TBitArray<> FeaturesAdded;
+	int32 NumFeaturesAdded = 0;
+};
+
+
 namespace UE { namespace PoseSearch {
 
 /** Records poses over time in a ring buffer. FFeatureVectorBuilder uses this to sample from the present or past poses according to the search schema. */
@@ -315,7 +372,7 @@ public:
 	TArrayView<const FTransform> GetPrevLocalPoseSample() const { return SampledPrevLocalPose; }
 	TArrayView<const FTransform> GetPrevComponentPoseSample() const { return SampledPrevComponentPose; }
 	float GetTimeHorizon() const { return TimeHorizon; }
-	TArray<float>& GetQueryBuffer() { return QueryBuffer; }
+	FPoseSearchFeatureVectorBuilder& GetQueryBuilder() { return QueryBuilder; }
 
 private:
 	bool SampleLocalPose(float Time, const TArray<FBoneIndexType>& RequiredBones, TArray<FTransform>& LocalPose);
@@ -331,7 +388,7 @@ private:
 	TArray<FTransform> SampledComponentPose;
 	TArray<FTransform> SampledPrevLocalPose;
 	TArray<FTransform> SampledPrevComponentPose;
-	TArray<float> QueryBuffer;
+	FPoseSearchFeatureVectorBuilder QueryBuilder;
 
 	float TimeHorizon = 0.0f;
 };
@@ -347,40 +404,6 @@ public:
 };
 
 
-/** 
-* Helper object for writing features into a float buffer according to a feature vector layout.
-* Keeps track of which features are present, allowing the feature vector to be built up piecemeal.
-* FFeatureVectorBuilder is used to build search queries at runtime and for adding samples during search index construction.
-*/
-class POSESEARCH_API FFeatureVectorBuilder
-{
-public:
-	void Init(const FPoseSearchFeatureVectorLayout* Layout, TArrayView<float> Buffer);
-	void ResetFeatures();
-
-	void SetTransform(FPoseSearchFeatureDesc Feature, const FTransform& Transform);
-	void SetTransformDerivative(FPoseSearchFeatureDesc Feature, const FTransform& Transform, const FTransform& PrevTransform, float DeltaTime);
-
-	void SetPosition(FPoseSearchFeatureDesc Feature, const FVector& Translation);
-	void SetRotation(FPoseSearchFeatureDesc Feature, const FQuat& Rotation);
-	void SetLinearVelocity(FPoseSearchFeatureDesc Feature, const FTransform& Transform, const FTransform& PrevTransform, float DeltaTime);
-	void SetAngularVelocity(FPoseSearchFeatureDesc Feature, const FTransform& Transform, const FTransform& PrevTransform, float DeltaTime);
-
-	void SetVector(FPoseSearchFeatureDesc Feature, const FVector& Vector);
-
-	bool SetPoseFeatures(const UPoseSearchSchema* Schema, FPoseHistory* History);
-
-	void Copy(TArrayView<const float> FeatureVector);
-
-	bool IsComplete() const;
-
-private:
-	const FPoseSearchFeatureVectorLayout* Layout = nullptr;
-	TArrayView<float> Values;
-	TBitArray<> FeaturesAdded;
-	int32 NumFeaturesAdded = 0;
-};
-
 /** Helper object for extracting features from a float buffer according to the feature vector layout. */
 class POSESEARCH_API FFeatureVectorReader
 {
