PoseSearch: Support for lead-in and follow-up anims for database entries, and searching trajectory history
- UPoseSearchDatabase now has optional lead-in and follow-up anims allowing the sampler to index beyond the clip bounds of the main sequence for past and future pose and trajectory subsampling
- Sampling and indexing has been split into two, more distinct phases. A gather step first collects all the referenced sequences and extracts samples from them all, then indexing pulls from the sample data to build the search index.
- FSequenceIndexer can handle sampling out of the bounds of the main clip if a predecessor or successor clip is provided. Bookend clips can optionally be looping and motion accumulation is handled appropriately.
- FPoseHistory now records and can sample trajectory history along with past pose. Also fixed history retention to allow for derivative sampling at time horizon.
- FPoseSearchFeatureVectorBuilder now has a helper for setting past trajectory
- Commented out AddTrajectoryDistanceFeatures for now since it needs to be rewritten to work with this new sampling/indexing approach
- Motion matching node has optional throttling, handles ending transition clips, and past trajectory matching

diff --git a/Engine/Plugins/Experimental/Animation/PoseSearch/Source/Runtime/Private/AnimNode_MotionMatching.cpp b/Engine/Plugins/Experimental/Animation/PoseSearch/Source/Runtime/Private/AnimNode_MotionMatching.cpp
--- a/Engine/Plugins/Experimental/Animation/PoseSearch/Source/Runtime/Private/AnimNode_MotionMatching.cpp
+++ b/Engine/Plugins/Experimental/Animation/PoseSearch/Source/Runtime/Private/AnimNode_MotionMatching.cpp
@@ -19,20 +19,9 @@ void FAnimNode_MotionMatching::Initialize_AnyThread(const FAnimationInitializeCo
 
 	GetEvaluateGraphExposedInputs().Execute(Context);
 
-	if (IsValidForSearch())
-	{
-		ComposedQuery.Init(Database->Schema);
-
-		// Set initial pose features
-		UE::PoseSearch::IPoseHistoryProvider* PoseHistoryProvider = Context.GetMessage<UE::PoseSearch::IPoseHistoryProvider>();
-		if (PoseHistoryProvider)
-		{
-			UE::PoseSearch::FPoseHistory& History = PoseHistoryProvider->GetPoseHistory();
-			ComposedQuery.SetPoseFeatures(&History);
-		}
-	}
-
 	DbPoseIdx = INDEX_NONE;
+	DbSequenceIdx = INDEX_NONE;
+	ElapsedPoseJumpTime = ThrottleTime;
 
 	Source.SetLinkNode(&SequencePlayerNode);
 	Source.Initialize(Context);
@@ -52,9 +41,17 @@ void FAnimNode_MotionMatching::Update_AnyThread(const FAnimationUpdateContext& C
 	GetEvaluateGraphExposedInputs().Execute(Context);
 	// Note: What if the input database changes? That's not being handled at all!
 
+	bool bJumpedToPose = false;
+
 	if (IsValidForSearch())
 	{
-		float AssetTime = SequencePlayerNode.GetAccumulatedTime();
+		if (!ComposedQuery.IsInitialized())
+		{
+			ComposedQuery.Init(Database->Schema);
+		}
+
+		const float AssetTime = SequencePlayerNode.GetAccumulatedTime();
+		const float AssetLength = SequencePlayerNode.GetCurrentAssetLength();
 
 		// Step the pose forward
 		if (DbPoseIdx != INDEX_NONE)
@@ -72,8 +69,18 @@ void FAnimNode_MotionMatching::Update_AnyThread(const FAnimationUpdateContext& C
 			}
 		}
 
+		if (DbPoseIdx == INDEX_NONE)
+		{
+			UE::PoseSearch::IPoseHistoryProvider* PoseHistoryProvider = Context.GetMessage<UE::PoseSearch::IPoseHistoryProvider>();
+			if (PoseHistoryProvider)
+			{
+				UE::PoseSearch::FPoseHistory& History = PoseHistoryProvider->GetPoseHistory();
+				ComposedQuery.SetPoseFeatures(&History);
+			}
+		}
+
 		// Update features in the query with the latest inputs
-		ComposeQuery();
+		ComposeQuery(Context);
 
 		// Determine how much the updated query vector deviates from the current pose vector
 		float CurrentDissimilarity = MAX_flt;
@@ -83,21 +90,13 @@ void FAnimNode_MotionMatching::Update_AnyThread(const FAnimationUpdateContext& C
 		}
 
 		// Search the database for the nearest match to the updated query vector
-
-		UE::PoseSearch::IPoseHistoryProvider* PoseHistoryProvider = Context.GetMessage<UE::PoseSearch::IPoseHistoryProvider>();
-		if (PoseHistoryProvider)
-		{
-			UE::PoseSearch::FPoseHistory& History = PoseHistoryProvider->GetPoseHistory();
-			ComposedQuery.SetPoseFeatures(&History);
-		}
-
 		UE::PoseSearch::FDbSearchResult Result = UE::PoseSearch::Search(Database, ComposedQuery.GetValues());
-		if (Result.IsValid())
+		if (Result.IsValid() && (ElapsedPoseJumpTime >= ThrottleTime))
 		{
 			const FPoseSearchDatabaseSequence& ResultDbSequence = Database->Sequences[Result.DbSequenceIdx];
 
 			// Consider the search result better if it is more similar to the query than the current pose we're playing back from the database
-			bool bBetterPose = Result.Dissimilarity < CurrentDissimilarity;
+			bool bBetterPose = Result.Dissimilarity * (1.0f + SimilarityThreshold) < CurrentDissimilarity;
 
 			// We'll ignore the candidate pose if it is too near to our current pose
 			bool bNearbyPose = false;
@@ -106,24 +105,56 @@ void FAnimNode_MotionMatching::Update_AnyThread(const FAnimationUpdateContext& C
 				bNearbyPose = FMath::Abs(AssetTime - Result.TimeOffsetSeconds) < PoseJumpThreshold;
 				if (!bNearbyPose && ResultDbSequence.bLoopAnimation)
 				{
-					bNearbyPose = FMath::Abs(SequencePlayerNode.GetCurrentAssetLength() - AssetTime - Result.TimeOffsetSeconds) < PoseJumpThreshold;
+					bNearbyPose = FMath::Abs(AssetLength - AssetTime - Result.TimeOffsetSeconds) < PoseJumpThreshold;
 				}
 			}
 
 			// And we won't bother to jump to another pose within the same looping sequence we're already playing
 			// (This should probably be configurable)
-			bool bSameCycle = (DbSequenceIdx == Result.DbSequenceIdx) && ResultDbSequence.bLoopAnimation;
-
 			bool bSameSequence = DbSequenceIdx == Result.DbSequenceIdx;
+			bool bSameCycle = bSameSequence && ResultDbSequence.bLoopAnimation;
 
 			// Start playback from the candidate pose if we determined it was a better option
 			if (bBetterPose && !bNearbyPose && !bSameCycle && !bSameSequence)
 			{
 				JumpToPose(Context, Result);
+				bJumpedToPose = true;
+			}
+		}
+
+		// Continue with the follow up sequence if we're finishing a one shot anim
+		if (!bJumpedToPose && SequencePlayerNode.Sequence && !SequencePlayerNode.bLoopAnimation)
+		{
+			float AssetTimeAfterUpdate = AssetTime + Context.GetDeltaTime();
+			if (AssetTimeAfterUpdate > AssetLength)
+			{
+				const FPoseSearchDatabaseSequence& DbSequence = Database->Sequences[DbSequenceIdx];
+				int32 FollowUpDbSequenceIdx = Database->Sequences.IndexOfByPredicate([&](const FPoseSearchDatabaseSequence& Entry){ return Entry.Sequence == DbSequence.FollowUpSequence; });
+				if (FollowUpDbSequenceIdx != INDEX_NONE)
+				{
+					const FPoseSearchDatabaseSequence& FollowUpDbSequence = Database->Sequences[FollowUpDbSequenceIdx];
+					float FollowUpAssetTime = AssetTimeAfterUpdate - AssetLength;
+					int32 FollowUpPoseIdx = Database->GetPoseIndexFromAssetTime(FollowUpDbSequenceIdx, FollowUpAssetTime);
+
+					UE::PoseSearch::FDbSearchResult FollowUpResult;
+					FollowUpResult.DbSequenceIdx = FollowUpDbSequenceIdx;
+					FollowUpResult.PoseIdx = FollowUpPoseIdx;
+					JumpToPose(Context, FollowUpResult);
+					bJumpedToPose = true;
+				}
 			}
 		}
 	}
 
+	if (bJumpedToPose)
+	{
+		ElapsedPoseJumpTime = 0.0f;
+	}
+	else
+	{
+		ElapsedPoseJumpTime += Context.GetDeltaTime();
+	}
+
 	Source.Update(Context);
 }
 
@@ -145,12 +176,13 @@ void FAnimNode_MotionMatching::PreUpdate(const UAnimInstance* InAnimInstance)
 		check(SkeletalMeshComponent);
 
 		UE::PoseSearch::FDebugDrawParams DrawParams;
-		DrawParams.RootTransform = SkeletalMeshComponent->GetBoneTransform(0);
+		DrawParams.RootTransform = SkeletalMeshComponent->GetComponentTransform();
 		DrawParams.Flags = UE::PoseSearch::EDebugDrawFlags::DrawQuery;
 		DrawParams.Database = Database;
 		DrawParams.Query = ComposedQuery.GetValues();
 		DrawParams.World = SkeletalMeshComponent->GetWorld();
 		DrawParams.DefaultLifeTime = 0.0f;
+		DrawParams.HighlightPoseIdx = DbPoseIdx;
 		UE::PoseSearch::Draw(DrawParams);
 	}
 #endif
@@ -166,8 +198,17 @@ bool FAnimNode_MotionMatching::IsValidForSearch() const
 	return Database && Database->IsValidForSearch();
 }
 
-void FAnimNode_MotionMatching::ComposeQuery()
+void FAnimNode_MotionMatching::ComposeQuery(const FAnimationBaseContext& Context)
 {
+	// Set past trajectory features
+	UE::PoseSearch::IPoseHistoryProvider* PoseHistoryProvider = Context.GetMessage<UE::PoseSearch::IPoseHistoryProvider>();
+	if (PoseHistoryProvider)
+	{
+		UE::PoseSearch::FPoseHistory& History = PoseHistoryProvider->GetPoseHistory();
+		ComposedQuery.SetPastTrajectoryFeatures(&History);
+	}
+
+	// Merge goal features into the query vector
 	if (ComposedQuery.IsCompatible(Goal))
 	{
 		ComposedQuery.MergeReplace(Goal);
diff --git a/Engine/Plugins/Experimental/Animation/PoseSearch/Source/Runtime/Private/AnimNode_PoseSearchHistoryCollector.cpp b/Engine/Plugins/Experimental/Animation/PoseSearch/Source/Runtime/Private/AnimNode_PoseSearchHistoryCollector.cpp
--- a/Engine/Plugins/Experimental/Animation/PoseSearch/Source/Runtime/Private/AnimNode_PoseSearchHistoryCollector.cpp
+++ b/Engine/Plugins/Experimental/Animation/PoseSearch/Source/Runtime/Private/AnimNode_PoseSearchHistoryCollector.cpp
@@ -57,7 +57,7 @@ void FAnimNode_PoseSearchHistoryCollector::Evaluate_AnyThread(FPoseContext& Outp
 
 	Source.Evaluate(Output);
 
-	PoseHistory.Update(EvalDeltaTime, Output.Pose);
+	PoseHistory.Update(EvalDeltaTime, Output);
 
 	EvalDeltaTime = 0.0f;
 }
diff --git a/Engine/Plugins/Experimental/Animation/PoseSearch/Source/Runtime/Private/PoseSearch.cpp b/Engine/Plugins/Experimental/Animation/PoseSearch/Source/Runtime/Private/PoseSearch.cpp
--- a/Engine/Plugins/Experimental/Animation/PoseSearch/Source/Runtime/Private/PoseSearch.cpp
+++ b/Engine/Plugins/Experimental/Animation/PoseSearch/Source/Runtime/Private/PoseSearch.cpp
@@ -11,6 +11,7 @@
 #include "Animation/AnimPoseSearchProvider.h"
 #include "Animation/AnimCurveTypes.h"
 #include "Animation/AnimInstance.h"
+#include "Animation/AnimInstanceProxy.h"
 #include "Animation/AnimNodeBase.h"
 #include "Animation/AnimMetaData.h"
 #include "Animation/AnimNodeBase.h"
@@ -27,13 +28,13 @@ namespace UE { namespace PoseSearch {
 //////////////////////////////////////////////////////////////////////////
 // Constants and utilities
 
-constexpr float DrawDebugLineThickness = 1.0f;
+constexpr float DrawDebugLineThickness = 2.0f;
 constexpr float DrawDebugPointSize = 3.0f;
 constexpr float DrawDebugVelocityScale = 0.1f;
 constexpr float DrawDebugArrowSize = 5.0f;
-constexpr float DrawDebugSphereSize = 1.0f;
+constexpr float DrawDebugSphereSize = 3.0f;
 constexpr int32 DrawDebugSphereSegments = 8;
-constexpr float DrawDebugSphereLineThickness = 0.2f;
+constexpr float DrawDebugSphereLineThickness = 0.5f;
 
 static bool IsSamplingRangeValid(FFloatInterval Range)
 {
@@ -615,14 +616,14 @@ bool FPoseSearchFeatureVectorBuilder::SetPoseFeatures(UE::PoseSearch::FPoseHisto
 	FPoseSearchFeatureDesc Feature;
 	Feature.Domain = EPoseSearchFeatureDomain::Time;
 
-	for (int32 SubsampleIdx = 0; SubsampleIdx != Schema->PoseSampleOffsets.Num(); ++SubsampleIdx)
+	for (int32 SchemaSubsampleIdx = 0; SchemaSubsampleIdx != Schema->PoseSampleOffsets.Num(); ++SchemaSubsampleIdx)
 	{
-		Feature.SubsampleIdx = SubsampleIdx;
+		Feature.SubsampleIdx = SchemaSubsampleIdx;
 
-		int32 Offset = Schema->PoseSampleOffsets[SubsampleIdx];
+		int32 Offset = Schema->PoseSampleOffsets[SchemaSubsampleIdx];
 		float TimeDelta = -Offset * Schema->SamplingInterval;
 
-		if (!History->Sample(TimeDelta, Schema->Skeleton->GetReferenceSkeleton(), Schema->BoneIndicesWithParents))
+		if (!History->SamplePose(TimeDelta, Schema->Skeleton->GetReferenceSkeleton(), Schema->BoneIndicesWithParents))
 		{
 			return false;
 		}
@@ -644,6 +645,45 @@ bool FPoseSearchFeatureVectorBuilder::SetPoseFeatures(UE::PoseSearch::FPoseHisto
 	return true;
 }
 
+bool FPoseSearchFeatureVectorBuilder::SetPastTrajectoryFeatures(UE::PoseSearch::FPoseHistory* History)
+{
+	check(Schema && Schema->IsValid());
+	check(History);
+
+	FPoseSearchFeatureDesc Feature;
+	Feature.Domain = EPoseSearchFeatureDomain::Time;
+	Feature.SchemaBoneIdx = FPoseSearchFeatureDesc::TrajectoryBoneIndex;
+
+	for (int32 SchemaSubsampleIdx = 0; SchemaSubsampleIdx != Schema->TrajectorySampleOffsets.Num(); ++SchemaSubsampleIdx)
+	{
+		Feature.SubsampleIdx = SchemaSubsampleIdx;
+
+		int32 SubsampleIndex = Schema->TrajectorySampleOffsets[SchemaSubsampleIdx];
+		if (SubsampleIndex >= 0)
+		{
+			break;
+		}
+
+		float SecondsAgo = -SubsampleIndex * Schema->SamplingInterval;
+		FTransform WorldComponentTransform;
+		if (!History->SampleRoot(SecondsAgo, &WorldComponentTransform))
+		{
+			return false;
+		}
+
+		FTransform WorldPrevComponentTransform;
+		if (!History->SampleRoot(SecondsAgo + History->GetSampleInterval(), &WorldPrevComponentTransform))
+		{
+			return false;
+		}
+
+		SetTransform(Feature, WorldComponentTransform);
+		SetTransformDerivative(Feature, WorldComponentTransform, WorldPrevComponentTransform, History->GetSampleInterval());
+	}
+
+	return true;
+}
+
 void FPoseSearchFeatureVectorBuilder::Copy(TArrayView<const float> FeatureVector)
 {
 	check(FeatureVector.Num() == Values.Num());
@@ -685,7 +725,7 @@ void FPoseSearchFeatureVectorBuilder::MergeReplace(const FPoseSearchFeatureVecto
 
 bool FPoseSearchFeatureVectorBuilder::IsInitialized() const
 {
-	return Schema != nullptr;
+	return (Schema != nullptr) && (Values.Num() == Schema->Layout.NumFloats);
 }
 
 bool FPoseSearchFeatureVectorBuilder::IsComplete() const
@@ -748,38 +788,38 @@ bool FPoseHistory::SampleLocalPose(float SecondsAgo, const TArray<FBoneIndexType
 
 	int32 PrevIdx = NextIdx - 1;
 
+	const FPose& PrevPose = Poses[PrevIdx];
+	const FPose& NextPose = Poses[NextIdx];
+
 	// Compute alpha between previous and next knots
 	float Alpha = FMath::GetMappedRangeValueUnclamped(
 		FVector2D(Knots[PrevIdx], Knots[NextIdx]),
 		FVector2D(0.0f, 1.0f),
 		SecondsAgo);
 
-	TArray<FTransform>& PrevPose = Poses[PrevIdx].LocalTransforms;
-	TArray<FTransform>& NextPose = Poses[NextIdx].LocalTransforms;
-
 	// We may not have accumulated enough poses yet
-	if (PrevPose.Num() != NextPose.Num())
+	if (PrevPose.LocalTransforms.Num() != NextPose.LocalTransforms.Num())
 	{
 		return false;
 	}
 
-	if (RequiredBones.Num() > PrevPose.Num())
+	if (RequiredBones.Num() > PrevPose.LocalTransforms.Num())
 	{
 		return false;
 	}
 
 	// Lerp between poses by alpha to produce output local pose at requested sample time
-	LocalPose = PrevPose;
+	LocalPose = PrevPose.LocalTransforms;
 	FAnimationRuntime::LerpBoneTransforms(
 		LocalPose,
-		NextPose,
+		NextPose.LocalTransforms,
 		Alpha,
 		RequiredBones);
 
 	return true;
 }
 
-bool FPoseHistory::Sample(float SecondsAgo, const FReferenceSkeleton& RefSkeleton, const TArray<FBoneIndexType>& RequiredBones)
+bool FPoseHistory::SamplePose(float SecondsAgo, const FReferenceSkeleton& RefSkeleton, const TArray<FBoneIndexType>& RequiredBones)
 {
 	// Compute local space pose at requested time
 	bool bSampled = SampleLocalPose(SecondsAgo, RequiredBones, SampledLocalPose);
@@ -797,7 +837,34 @@ bool FPoseHistory::Sample(float SecondsAgo, const FReferenceSkeleton& RefSkeleto
 	return bSampled;
 }
 
-void FPoseHistory::Update(float SecondsElapsed, const FCompactPose& Pose)
+bool FPoseHistory::SampleRoot(float SecondsAgo, FTransform* OutTransform) const
+{
+	int32 NextIdx = LowerBound(Knots.begin(), Knots.end(), SecondsAgo, TGreater<>());
+	if (NextIdx <= 0 || NextIdx >= Knots.Num())
+	{
+		return false;
+	}
+
+	int32 PrevIdx = NextIdx - 1;
+
+	const FPose& PrevPose = Poses[PrevIdx];
+	const FPose& NextPose = Poses[NextIdx];
+
+	// Compute alpha between previous and next knots
+	float Alpha = FMath::GetMappedRangeValueUnclamped(
+		FVector2D(Knots[PrevIdx], Knots[NextIdx]),
+		FVector2D(0.0f, 1.0f),
+		SecondsAgo);
+
+	FTransform RootTransform;
+	RootTransform.Blend(PrevPose.WorldComponentTransform, NextPose.WorldComponentTransform, Alpha);
+	RootTransform.SetToRelativeTransform(Poses.Last().WorldComponentTransform);
+
+	*OutTransform = RootTransform;
+	return true;
+}
+
+void FPoseHistory::Update(float SecondsElapsed, const FPoseContext& PoseContext)
 {
 	// Age our elapsed times
 	for (float& Knot : Knots)
@@ -814,13 +881,13 @@ void FPoseHistory::Update(float SecondsElapsed, const FCompactPose& Pose)
 	else
 	{
 		// Exercise pose retention policy. We must guarantee there is always one additional knot
-		// at or beyond the desired time horizon H so we can fulfill sample requests at t=H. We also
+		// beyond the time horizon so we can compute derivatives at the time horizon. We also
 		// want to evenly distribute knots across the entire history buffer so we only push additional
 		// poses when enough time has elapsed.
 
 		const float SampleInterval = GetSampleInterval();
 
-		bool bCanEvictOldest = Knots[1] >= TimeHorizon;
+		bool bCanEvictOldest = Knots[1] >= TimeHorizon + SampleInterval;
 		bool bShouldPushNewest = Knots[Knots.Num() - 2] >= SampleInterval;
 
 		if (bCanEvictOldest && bShouldPushNewest)
@@ -836,12 +903,15 @@ void FPoseHistory::Update(float SecondsElapsed, const FCompactPose& Pose)
 
 	// Regardless of the retention policy, we always update the most recent pose
 	Knots.Last() = 0.0f;
-	CopyCompactToSkeletonPose(Pose, Poses.Last().LocalTransforms);
+	FPose& CurrentPose = Poses.Last();
+	CopyCompactToSkeletonPose(PoseContext.Pose, CurrentPose.LocalTransforms);
+	CurrentPose.WorldComponentTransform = PoseContext.AnimInstanceProxy->GetComponentTransform();
 }
 
 float FPoseHistory::GetSampleInterval() const
 {
-	return TimeHorizon / Knots.Max();
+	// Reserve one knot for computing derivatives at the time horizon
+	return TimeHorizon / (Knots.Max() - 1);
 }
 
 
@@ -1000,29 +1070,19 @@ const UPoseSearchSchema* FDebugDrawParams::GetSchema() const
 
 
 //////////////////////////////////////////////////////////////////////////
-// FSequenceIndexer
+// FSequenceSampler
 
-class FSequenceIndexer
+struct FSequenceSampler
 {
 public:
-	struct Result
+	struct FInput
 	{
-		int32 NumIndexedPoses = 0;
-		TArrayView<const float> Values;
-	};
-	Result Process(const UPoseSearchSchema* Schema, const UAnimSequence* Sequence, FFloatInterval SamplingRange);
-	Result GetResult() const;
-
-private:
-	void SampleBegin(int32 SampleIdx);
-	void SampleEnd(int32 SampleIdx);
-	void ExtractPoses(const UAnimSequence* Sequence);
-	void ExtractRootMotion(const UAnimSequence* Sequence);
-	void AddPoseFeatures(int32 SampleIdx);
-	void AddTrajectoryTimeFeatures(int32 SampleIdx);
-	void AddTrajectoryDistanceFeatures(int32 SampleIdx);
+		const UPoseSearchSchema* Schema = nullptr;
+		const UAnimSequence* Sequence = nullptr;
+		bool bLoopable = false;
+	} Input;
 
-	struct FSampleContext
+	struct FOutput
 	{
 		TArray<FTransform> ComponentSpacePose;		// Indexed by SampleIdx * NumBones + SchemaBoneIdx
 		TArray<FTransform> LocalRootMotion;			// Indexed by SampleIdx
@@ -1030,122 +1090,111 @@ private:
 		TArray<float> AccumulatedRootDistance;		// Indexed by SampleIdx
 
 		int32 TotalSamples = 0;
-		int32 FirstIndexedSample = 0;
-		int32 LastIndexedSample = 0;
-		int32 NumIndexedSamples = 0;
-		int32 NumBones = 0;
+	} Output;
 
-		void Reset();
-		void Reserve();
-	};
+	void Reset();
+	void Init(const FInput& Input);
+	void Process();
 
-	const UPoseSearchSchema* Schema = nullptr;
+	struct FWrappedSampleIndex
+	{
+		int32 Idx = INDEX_NONE;
+		int32 NumCycles = 0;
+		bool bClamped = false;
+	};
+	FWrappedSampleIndex WrapOrClampSubsampleIndex (int32 SampleIdx) const;
 
-	TArray<float> Values;
-	
-	FPoseSearchFeatureVectorBuilder Builder;
-	FSampleContext Context;
+private:
+	void Reserve();
+	void ExtractPoses();
+	void ExtractRootMotion();
 };
 
-void FSequenceIndexer::FSampleContext::Reset()
-{
-	TotalSamples = 0;
-	FirstIndexedSample = 0;
-	LastIndexedSample = 0;
-	NumIndexedSamples = 0;
-	NumBones = 0;
-
-	ComponentSpacePose.Reset(0);
-	LocalRootMotion.Reset(0);
-	AccumulatedRootMotion.Reset(0);
-	AccumulatedRootDistance.Reset(0);
-}
-
-void FSequenceIndexer::FSampleContext::Reserve()
-{
-	ComponentSpacePose.Reserve(NumBones * TotalSamples);
-	LocalRootMotion.Reserve(TotalSamples);
-	AccumulatedRootMotion.Reserve(TotalSamples);
-	AccumulatedRootDistance.Reserve(TotalSamples);
-}
-
-FSequenceIndexer::Result FSequenceIndexer::Process(const UPoseSearchSchema* InSchema, const UAnimSequence* Sequence, FFloatInterval RequestedSamplingRange)
+void FSequenceSampler::Init(const FInput& InInput)
 {
-	check(InSchema);
-	check(Sequence);
-
-	Schema = InSchema;
-
-	USkeleton* Skeleton = Sequence->GetSkeleton();
-	check(Skeleton);
-	check(Skeleton->IsCompatible(Schema->Skeleton));
-
-	const float SequencePlayLength = Sequence->GetPlayLength();
-	const FFloatInterval SamplingRange = GetEffectiveSamplingRange(Sequence, RequestedSamplingRange);
-
-	Context.Reset();
-	Context.NumBones = Schema->BoneIndices.Num();
-	Context.TotalSamples = FMath::FloorToInt(SequencePlayLength * Schema->SampleRate);
-	Context.FirstIndexedSample = FMath::FloorToInt(SamplingRange.Min * Schema->SampleRate);
-	Context.LastIndexedSample = FMath::Max(0, FMath::FloorToInt(SamplingRange.Max * Schema->SampleRate) - 1);
-	Context.NumIndexedSamples = Context.LastIndexedSample - Context.FirstIndexedSample + 1;
-	Context.Reserve();
+	check(InInput.Schema);
+	check(InInput.Schema->IsValid());
+	check(InInput.Sequence);
 
-	Values.SetNumZeroed(Schema->Layout.NumFloats * Context.NumIndexedSamples);
+	Reset();
 
-	ExtractPoses(Sequence);
-	ExtractRootMotion(Sequence);
+	Input = InInput;
 
-	for (int32 SampleIdx = Context.FirstIndexedSample; SampleIdx <= Context.LastIndexedSample; ++SampleIdx)
-	{
-		SampleBegin(SampleIdx);
+	const float SequencePlayLength = Input.Sequence->GetPlayLength();
+	Output.TotalSamples = FMath::FloorToInt(SequencePlayLength * Input.Schema->SampleRate);
 
-		AddPoseFeatures(SampleIdx);
-		AddTrajectoryTimeFeatures(SampleIdx);
-		AddTrajectoryDistanceFeatures(SampleIdx);
+	Reserve();
+}
 
-		SampleEnd(SampleIdx);
-	}
+void FSequenceSampler::Reset()
+{
+	Input = FInput();
 
-	return GetResult();
+	Output.TotalSamples = 0;
+	Output.ComponentSpacePose.Reset(0);
+	Output.LocalRootMotion.Reset(0);
+	Output.AccumulatedRootMotion.Reset(0);
+	Output.AccumulatedRootDistance.Reset(0);
 }
 
-FSequenceIndexer::Result FSequenceIndexer::GetResult() const
+void FSequenceSampler::Reserve()
 {
-	Result Result;
-	Result.NumIndexedPoses = Context.NumIndexedSamples;
-	Result.Values = Values;
-	return Result;
+	Output.ComponentSpacePose.Reserve(Input.Schema->NumBones() * Output.TotalSamples);
+	Output.LocalRootMotion.Reserve(Output.TotalSamples);
+	Output.AccumulatedRootMotion.Reserve(Output.TotalSamples);
+	Output.AccumulatedRootDistance.Reserve(Output.TotalSamples);
 }
 
-void FSequenceIndexer::SampleBegin(int32 SampleIdx)
+void FSequenceSampler::Process()
 {
-	Builder.Init(Schema);
+	ExtractPoses();
+	ExtractRootMotion();
 }
 
-void FSequenceIndexer::SampleEnd(int32 SampleIdx)
+
+FSequenceSampler::FWrappedSampleIndex FSequenceSampler::WrapOrClampSubsampleIndex (int32 SampleIdx) const
 {
-	check(Builder.IsComplete());
+	FWrappedSampleIndex Result;
+	Result.Idx = SampleIdx;
+	Result.NumCycles = 0;
+	Result.bClamped = false;
+
+	// Wrap the index if this is a loopable sequence
+	if (Input.bLoopable)
+	{
+		while (Result.Idx < 0)
+		{
+			Result.Idx += Output.TotalSamples;
+			++Result.NumCycles;
+		}
 
-	int32 FirstValueIdx = (SampleIdx - Context.FirstIndexedSample) * Schema->Layout.NumFloats;
-	TArrayView<float> WriteValues = MakeArrayView(&Values[FirstValueIdx], Schema->Layout.NumFloats);
+		while (Result.Idx >= Output.TotalSamples)
+		{
+			Result.Idx -= Output.TotalSamples;
+			++Result.NumCycles;
+		}
+	}
 
-	TArrayView<const float> ReadValues = Builder.GetValues();
-	
-	check(WriteValues.Num() == ReadValues.Num());
-	FMemory::Memcpy(WriteValues.GetData(), ReadValues.GetData(), WriteValues.Num() * WriteValues.GetTypeSize());
+	// Clamp if we can't loop
+	else if (SampleIdx < 0 || SampleIdx >= Output.TotalSamples)
+	{
+		Result.Idx = FMath::Clamp(SampleIdx, 0, Output.TotalSamples - 1);
+		Result.bClamped = true;
+	}
+
+	return Result;
 }
 
-void FSequenceIndexer::ExtractPoses(const UAnimSequence* Sequence)
+void FSequenceSampler::ExtractPoses()
 {
-	if (Schema->Bones.IsEmpty())
+	if (Input.Schema->Bones.IsEmpty())
 	{
 		return;
 	}
 
-	USkeleton* Skeleton = Sequence->GetSkeleton();
+	USkeleton* Skeleton = Input.Sequence->GetSkeleton();
 	FBoneContainer BoneContainer;
-	BoneContainer.InitializeTo(Schema->BoneIndicesWithParents, FCurveEvaluationOption(false), *Skeleton);
+	BoneContainer.InitializeTo(Input.Schema->BoneIndicesWithParents, FCurveEvaluationOption(false), *Skeleton);
 
 	FCompactPose Pose;
 	Pose.SetBoneContainer(&BoneContainer);
@@ -1160,148 +1209,369 @@ void FSequenceIndexer::ExtractPoses(const UAnimSequence* Sequence)
 	ExtractionCtx.bExtractRootMotion = true;
 
 	FAnimationPoseData AnimPoseData(Pose, UnusedCurve, UnusedAttributes);
-	for (int32 SampleIdx = 0; SampleIdx != Context.TotalSamples; ++SampleIdx)
+	for (int32 SampleIdx = 0; SampleIdx != Output.TotalSamples; ++SampleIdx)
 	{
-		const float CurrentTime = SampleIdx * Schema->SamplingInterval;
+		const float CurrentTime = SampleIdx * Input.Schema->SamplingInterval;
 
 		ExtractionCtx.CurrentTime = CurrentTime;
-		Sequence->GetAnimationPose(AnimPoseData, ExtractionCtx);
+		Input.Sequence->GetAnimationPose(AnimPoseData, ExtractionCtx);
 		ComponentSpacePose.InitPose(Pose);
 
-		for (int32 BoneIndex : Schema->BoneIndices)
+		for (int32 BoneIndex : Input.Schema->BoneIndices)
 		{
 			FCompactPoseBoneIndex CompactBoneIndex = BoneContainer.MakeCompactPoseIndex(FMeshPoseBoneIndex(BoneIndex));
 			const FTransform& Transform = ComponentSpacePose.GetComponentSpaceTransform(CompactBoneIndex);
-			Context.ComponentSpacePose.Add(Transform);
+			Output.ComponentSpacePose.Add(Transform);
 		}
 	}
 }
 
-void FSequenceIndexer::ExtractRootMotion(const UAnimSequence* Sequence)
+void FSequenceSampler::ExtractRootMotion()
 {
 	double AccumulatedRootDistance = 0.0;
 	FTransform AccumulatedRootMotion = FTransform::Identity;
-	for (int32 SampleIdx = 0; SampleIdx != Context.TotalSamples; ++SampleIdx)
+	for (int32 SampleIdx = 0; SampleIdx != Output.TotalSamples; ++SampleIdx)
 	{
-		const float CurrentTime = SampleIdx * Schema->SamplingInterval;
+		const float CurrentTime = SampleIdx * Input.Schema->SamplingInterval;
 
-		FTransform LocalRootMotion = Sequence->ExtractRootMotion(CurrentTime, Schema->SamplingInterval, false /*!allowLooping*/);
-		Context.LocalRootMotion.Add(LocalRootMotion);
+		FTransform LocalRootMotion = Input.Sequence->ExtractRootMotion(CurrentTime, Input.Schema->SamplingInterval, false /*!allowLooping*/);
+		Output.LocalRootMotion.Add(LocalRootMotion);
 
-		AccumulatedRootMotion = LocalRootMotion * AccumulatedRootMotion;
+		AccumulatedRootMotion = AccumulatedRootMotion * LocalRootMotion;
 		AccumulatedRootDistance += LocalRootMotion.GetTranslation().Size();
-		Context.AccumulatedRootMotion.Add(AccumulatedRootMotion);
-		Context.AccumulatedRootDistance.Add((float)AccumulatedRootDistance);
+		Output.AccumulatedRootMotion.Add(AccumulatedRootMotion);
+		Output.AccumulatedRootDistance.Add((float)AccumulatedRootDistance);
 	}
 }
 
-void FSequenceIndexer::AddPoseFeatures(int32 SampleIdx)
+
+//////////////////////////////////////////////////////////////////////////
+// FSequenceIndexer
+
+class FSequenceIndexer
 {
-	FPoseSearchFeatureDesc CurrentElement;
-	CurrentElement.Domain = EPoseSearchFeatureDomain::Time;
+public:
 
-	FTransform SampleSpaceOrigin = Context.AccumulatedRootMotion[SampleIdx];
-	
-	for (int32 SchemaSubsampleIdx = 0; SchemaSubsampleIdx != Schema->PoseSampleOffsets.Num(); ++SchemaSubsampleIdx)
+	struct FInput
 	{
-		CurrentElement.SubsampleIdx = SchemaSubsampleIdx;
+		const UPoseSearchSchema* Schema = nullptr;
+		const FSequenceSampler* MainSequence = nullptr;
+		const FSequenceSampler* LeadInSequence = nullptr;
+		const FSequenceSampler* FollowUpSequence = nullptr;
+		FFloatInterval RequestedSamplingRange = FFloatInterval(0.0f, 0.0f);
+	} Input;
 
-		const int32 SampleOffset = Schema->PoseSampleOffsets[SchemaSubsampleIdx];
-		const int32 SubsampleIdx = FMath::Clamp(SampleIdx + SampleOffset, 0, Context.AccumulatedRootMotion.Num() - 1);
+	struct FOutput
+	{
+		int32 FirstIndexedSample = 0;
+		int32 LastIndexedSample = 0;
+		int32 NumIndexedPoses = 0;
+		TArray<float> FeatureVectorTable;
+	} Output;
 
-		FTransform SubsampleRoot = Context.AccumulatedRootMotion[SubsampleIdx];
-		SubsampleRoot.SetToRelativeTransform(SampleSpaceOrigin);
+	void Reset();
+	void Init(const FInput& Input);
+	void Process();
 
-		for (int32 SchemaBoneIndex = 0; SchemaBoneIndex != Context.NumBones; ++SchemaBoneIndex)
-		{
-			CurrentElement.SchemaBoneIdx = SchemaBoneIndex;
+private:
+	FPoseSearchFeatureVectorBuilder FeatureVector;
 
-			int32 BoneSampleIdx = Context.NumBones * (SampleIdx + SampleOffset) + SchemaBoneIndex;
-			int32 BonePrevSampleIdx = Context.NumBones * (SampleIdx - 1 + SampleOffset) + SchemaBoneIndex;
-			
-			// @@@Add extrapolation. Clamp for now
-			BoneSampleIdx = FMath::Clamp(BoneSampleIdx, 0, Context.ComponentSpacePose.Num() - 1);
-			BonePrevSampleIdx = FMath::Clamp(BonePrevSampleIdx, 0, Context.ComponentSpacePose.Num() - 1);
+	struct FSubsample
+	{
+		const FSequenceSampler* Sampler = nullptr;
+		int32 AbsoluteSampleIdx = INDEX_NONE;
+		FTransform AccumulatedRootMotion;
+		float AccumulatedRootDistance = 0.0f;
+	};
 
-			FTransform BoneInSampleSpace = Context.ComponentSpacePose[BoneSampleIdx] * SubsampleRoot;
-			FTransform BonePrevInSampleSpace = Context.ComponentSpacePose[BonePrevSampleIdx] * SubsampleRoot;
+	FSubsample ResolveSubsample(int32 MainSubsampleIdx) const;
 
-			Builder.SetTransform(CurrentElement, BoneInSampleSpace);
-			Builder.SetTransformDerivative(CurrentElement, BoneInSampleSpace, BonePrevInSampleSpace, Schema->SamplingInterval);
-		}
+	void Reserve();
+	void SampleBegin(int32 SampleIdx);
+	void SampleEnd(int32 SampleIdx);
+	void AddPoseFeatures(int32 SampleIdx);
+	void AddTrajectoryTimeFeatures(int32 SampleIdx);
+	void AddTrajectoryDistanceFeatures(int32 SampleIdx);
+};
+
+void FSequenceIndexer::Reset()
+{
+	Output.FirstIndexedSample = 0;
+	Output.LastIndexedSample = 0;
+	Output.NumIndexedPoses = 0;
+
+	Output.FeatureVectorTable.Reset(0);
+}
+
+void FSequenceIndexer::Reserve()
+{
+	Output.FeatureVectorTable.SetNumZeroed(Input.Schema->Layout.NumFloats * Output.NumIndexedPoses);
+}
+
+void FSequenceIndexer::Init(const FInput& InSettings)
+{
+	check(InSettings.Schema);
+	check(InSettings.Schema->IsValid());
+	check(InSettings.MainSequence);
+
+	Input = InSettings;
+
+	const FFloatInterval SamplingRange = GetEffectiveSamplingRange(Input.MainSequence->Input.Sequence, Input.RequestedSamplingRange);
+
+	Reset();
+	Output.FirstIndexedSample = FMath::FloorToInt(SamplingRange.Min * Input.Schema->SampleRate);
+	Output.LastIndexedSample = FMath::Max(0, FMath::FloorToInt(SamplingRange.Max * Input.Schema->SampleRate) - 1);
+	Output.NumIndexedPoses = Output.LastIndexedSample - Output.FirstIndexedSample + 1;
+	Reserve();
+}
+
+void FSequenceIndexer::Process()
+{
+	for (int32 SampleIdx = Output.FirstIndexedSample; SampleIdx <= Output.LastIndexedSample; ++SampleIdx)
+	{
+		SampleBegin(SampleIdx);
+
+		AddPoseFeatures(SampleIdx);
+		AddTrajectoryTimeFeatures(SampleIdx);
+		AddTrajectoryDistanceFeatures(SampleIdx);
+
+		SampleEnd(SampleIdx);
 	}
 }
 
-void FSequenceIndexer::AddTrajectoryTimeFeatures(int32 SampleIdx)
+void FSequenceIndexer::SampleBegin(int32 SampleIdx)
+{
+	FeatureVector.Init(Input.Schema);
+}
+
+void FSequenceIndexer::SampleEnd(int32 SampleIdx)
 {
-	FPoseSearchFeatureDesc CurrentElement;
-	CurrentElement.Domain = EPoseSearchFeatureDomain::Time;
-	CurrentElement.SchemaBoneIdx = FPoseSearchFeatureDesc::TrajectoryBoneIndex;
+	check(FeatureVector.IsComplete());
 
-	FTransform SampleSpaceOrigin = Context.AccumulatedRootMotion[SampleIdx];
+	int32 FirstValueIdx = (SampleIdx - Output.FirstIndexedSample) * Input.Schema->Layout.NumFloats;
+	TArrayView<float> WriteValues = MakeArrayView(&Output.FeatureVectorTable[FirstValueIdx], Input.Schema->Layout.NumFloats);
 
-	for (int32 SubsampleIdx = 0; SubsampleIdx != Schema->TrajectorySampleOffsets.Num(); ++SubsampleIdx)
+	TArrayView<const float> ReadValues = FeatureVector.GetValues();
+	
+	check(WriteValues.Num() == ReadValues.Num());
+	FMemory::Memcpy(WriteValues.GetData(), ReadValues.GetData(), WriteValues.Num() * WriteValues.GetTypeSize());
+}
+
+FSequenceIndexer::FSubsample FSequenceIndexer::ResolveSubsample(int32 MainSubsampleIdx) const
+{
+	// MainSubsampleIdx is relative to the samples in the main sequence. With future subsampling,
+	// SampleIdx may be greater than the  number of samples in the main sequence. For past subsampling,
+	// SampleIdx may be negative. This function handles those edge cases by wrapping within the main
+	// sequence if it is loopable, or by indexing into the lead-in or follow-up sequences which themselves
+	// may or may not be loopable.
+	// The relative SampleIdx may be multiple cycles away, so this function also handles the math for
+	// accumulating multiple cycles of root motion.
+	// It returns an absolute index into the relevant sample data and root motion info.
+
+	FSubsample Subsample;
+
+	FTransform RootMotionLast = FTransform::Identity;
+	FTransform RootMotionInitial = FTransform::Identity;
+
+	float RootDistanceLast = 0.0f;
+	float RootDistanceInitial = 0.0f;
+
+	FSequenceSampler::FWrappedSampleIndex MainSample = Input.MainSequence->WrapOrClampSubsampleIndex(MainSubsampleIdx);
+	FSequenceSampler::FWrappedSampleIndex EffectiveSample;
+
+	// Use the lead in anim if we had to clamp to the beginning of the main anim
+	if (MainSample.bClamped && (MainSubsampleIdx < 0) && Input.LeadInSequence)
+	{
+		EffectiveSample = Input.LeadInSequence->WrapOrClampSubsampleIndex(MainSubsampleIdx + Input.LeadInSequence->Output.TotalSamples);
+
+		Subsample.Sampler = Input.LeadInSequence;
+		Subsample.AbsoluteSampleIdx = EffectiveSample.Idx;
+
+		RootMotionInitial = FTransform::Identity;
+		RootDistanceInitial = 0.0f;
+
+		RootMotionLast = Input.LeadInSequence->Output.AccumulatedRootMotion.Last();
+		RootDistanceLast = Input.LeadInSequence->Output.AccumulatedRootDistance.Last();
+	}
+
+	// Use the follow up anim if we had to clamp to the end of the main anim
+	if (MainSample.bClamped && (MainSubsampleIdx >= Input.MainSequence->Output.TotalSamples) && Input.FollowUpSequence)
+	{
+		EffectiveSample = Input.FollowUpSequence->WrapOrClampSubsampleIndex(MainSubsampleIdx - Input.MainSequence->Output.TotalSamples);
+
+		Subsample.Sampler = Input.FollowUpSequence;
+		Subsample.AbsoluteSampleIdx = EffectiveSample.Idx;
+
+		RootMotionInitial = Input.MainSequence->Output.AccumulatedRootMotion.Last();
+		RootDistanceInitial = Input.MainSequence->Output.AccumulatedRootDistance.Last();
+
+		RootMotionLast = Input.FollowUpSequence->Output.AccumulatedRootMotion.Last();
+		RootDistanceLast = Input.FollowUpSequence->Output.AccumulatedRootDistance.Last();
+	}
+
+	// Use the main anim if we didn't use the lead-in or follow-up anims.
+	// The main anim sample may have been wrapped or clamped
+	if (EffectiveSample.Idx == INDEX_NONE)
 	{
-		CurrentElement.SubsampleIdx = SubsampleIdx;
+		EffectiveSample = MainSample;
+
+		Subsample.Sampler = Input.MainSequence;
+		Subsample.AbsoluteSampleIdx = EffectiveSample.Idx;
 
-		int32 RootMotionIdx = SampleIdx + Schema->TrajectorySampleOffsets[SubsampleIdx];
-		int32 RootMotionPrevIdx = RootMotionIdx - 1;
+		RootMotionInitial = FTransform::Identity;
+		RootDistanceInitial = 0.0f;
 
-		// @@@ Add extrapolation. Clamp for now
-		RootMotionIdx = FMath::Clamp(RootMotionIdx, 0, Context.AccumulatedRootMotion.Num() - 1);
-		RootMotionPrevIdx = FMath::Clamp(RootMotionPrevIdx, 0, Context.AccumulatedRootMotion.Num() - 1);
+		RootMotionLast = Input.MainSequence->Output.AccumulatedRootMotion.Last();
+		RootDistanceLast = Input.MainSequence->Output.AccumulatedRootDistance.Last();
+	}
+
+	// Determine how to accumulate motion for every cycle of the anim. If the sample
+	// had to be clamped, this motion will end up not getting applied below.
+	// Also invert the accumulation direction if the requested sample was wrapped backwards.
+	FTransform RootMotionPerCycle = RootMotionLast;
+	float RootDistancePerCycle = RootDistanceLast;
+	if (MainSubsampleIdx < 0)
+	{
+		RootMotionPerCycle = RootMotionPerCycle.Inverse();
+		RootDistancePerCycle *= -1;
+	}
 
-		FTransform SubsampleRoot = Context.AccumulatedRootMotion[RootMotionIdx];
-		SubsampleRoot.SetToRelativeTransform(SampleSpaceOrigin);
+	// Find the remaining motion deltas after wrapping
+	FTransform RootMotionRemainder = Subsample.Sampler->Output.AccumulatedRootMotion[EffectiveSample.Idx];
+	float RootDistanceRemainder = Subsample.Sampler->Output.AccumulatedRootDistance[EffectiveSample.Idx];
 
-		FTransform SubsamplePrevRoot = Context.AccumulatedRootMotion[RootMotionPrevIdx];
-		SubsamplePrevRoot.SetToRelativeTransform(SampleSpaceOrigin);
+	// Invert motion deltas if we wrapped backwards
+	if (MainSubsampleIdx < 0)
+	{
+		RootMotionRemainder.SetToRelativeTransform(RootMotionLast);
+		RootDistanceRemainder = -(RootDistanceLast - RootDistanceRemainder);
+	}
 
-		Builder.SetTransform(CurrentElement, SubsampleRoot);
-		Builder.SetTransformDerivative(CurrentElement, SubsampleRoot, SubsamplePrevRoot, Schema->SamplingInterval);
+	Subsample.AccumulatedRootMotion = RootMotionInitial;
+	Subsample.AccumulatedRootDistance = RootDistanceInitial;
+
+	// Note if the sample was clamped, no motion will be applied here because NumCycles will be zero
+	int32 CyclesRemaining = EffectiveSample.NumCycles;
+	while (CyclesRemaining--)
+	{
+		Subsample.AccumulatedRootMotion *= RootMotionPerCycle;
+		Subsample.AccumulatedRootDistance += RootDistancePerCycle;
 	}
+
+	Subsample.AccumulatedRootMotion *= RootMotionRemainder;
+	Subsample.AccumulatedRootDistance += RootDistanceRemainder;
+
+	return Subsample;
 }
 
-void FSequenceIndexer::AddTrajectoryDistanceFeatures(int32 SampleIdx)
+void FSequenceIndexer::AddPoseFeatures(int32 SampleIdx)
 {
-	FPoseSearchFeatureDesc CurrentElement;
-	CurrentElement.Domain = EPoseSearchFeatureDomain::Distance;
-	CurrentElement.SchemaBoneIdx = FPoseSearchFeatureDesc::TrajectoryBoneIndex;
-
-	TArrayView<const float> AccumulatedRootDistances = Context.AccumulatedRootDistance;
+	FPoseSearchFeatureDesc Feature;
+	Feature.Domain = EPoseSearchFeatureDomain::Time;
 
-	FTransform SampleSpaceOrigin = Context.AccumulatedRootMotion[SampleIdx];
+	const int32 NumBones = Input.Schema->NumBones();
 
-	for (int32 SubsampleIdx = 0; SubsampleIdx != Schema->TrajectorySampleDistances.Num(); ++SubsampleIdx)
+	FSequenceIndexer::FSubsample OriginSample = ResolveSubsample(SampleIdx);
+	
+	for (int32 SchemaSubsampleIdx = 0; SchemaSubsampleIdx != Input.Schema->PoseSampleOffsets.Num(); ++SchemaSubsampleIdx)
 	{
-		CurrentElement.SubsampleIdx = SubsampleIdx;
+		Feature.SubsampleIdx = SchemaSubsampleIdx;
 
-		const float TrajectoryDistance = Schema->TrajectorySampleDistances[SubsampleIdx];
-		const float SampleAccumulatedRootDistance = TrajectoryDistance + AccumulatedRootDistances[SampleIdx];
+		const int32 SubsampleIdx = SampleIdx + Input.Schema->PoseSampleOffsets[SchemaSubsampleIdx];
 
-		int32 LowerBoundSampleIdx = Algo::LowerBound(AccumulatedRootDistances, SampleAccumulatedRootDistance);
+		FSequenceIndexer::FSubsample Subsample = ResolveSubsample(SubsampleIdx);
+		FSequenceIndexer::FSubsample SubsamplePrev = ResolveSubsample(SubsampleIdx - 1);
 
-		// @@@ Add extrapolation. Clamp for now
-		int32 PrevSampleIdx = FMath::Clamp(LowerBoundSampleIdx - 1, 0, AccumulatedRootDistances.Num() - 1);
-		int32 NextSampleIdx = FMath::Clamp(LowerBoundSampleIdx, 0, AccumulatedRootDistances.Num() - 1);
+		FTransform SubsampleRoot = Subsample.AccumulatedRootMotion;
+		SubsampleRoot.SetToRelativeTransform(OriginSample.AccumulatedRootMotion);
 
-		const float PrevSampleDistance = AccumulatedRootDistances[PrevSampleIdx];
-		const float NextSampleDistance = AccumulatedRootDistances[NextSampleIdx];
+		for (int32 SchemaBoneIndex = 0; SchemaBoneIndex != NumBones; ++SchemaBoneIndex)
+		{
+			Feature.SchemaBoneIdx = SchemaBoneIndex;
 
-		FTransform PrevRootInSampleSpace = Context.AccumulatedRootMotion[PrevSampleIdx];
-		PrevRootInSampleSpace.SetToRelativeTransform(SampleSpaceOrigin);
+			int32 BoneSampleIdx = NumBones * Subsample.AbsoluteSampleIdx + SchemaBoneIndex;
+			int32 BonePrevSampleIdx = NumBones * SubsamplePrev.AbsoluteSampleIdx  + SchemaBoneIndex;
 
-		FTransform NextRootInSampleSpace = Context.AccumulatedRootMotion[NextSampleIdx];
-		NextRootInSampleSpace.SetToRelativeTransform(SampleSpaceOrigin);
-		
-		float Alpha = FMath::GetRangePct(PrevSampleDistance, NextSampleDistance, SampleAccumulatedRootDistance);
-		FTransform BlendedRootInSampleSpace;
-		BlendedRootInSampleSpace.Blend(PrevRootInSampleSpace, NextRootInSampleSpace, Alpha);
+			FTransform BoneInComponentSpace = Subsample.Sampler->Output.ComponentSpacePose[BoneSampleIdx];
+			FTransform BonePrevInComponentSpace = SubsamplePrev.Sampler->Output.ComponentSpacePose[BonePrevSampleIdx];
+
+			FTransform BoneInSampleSpace = BoneInComponentSpace * SubsampleRoot;
+			FTransform BonePrevInSampleSpace = BonePrevInComponentSpace * SubsampleRoot;
+
+			FeatureVector.SetTransform(Feature, BoneInSampleSpace);
+			FeatureVector.SetTransformDerivative(Feature, BoneInSampleSpace, BonePrevInSampleSpace, Input.Schema->SamplingInterval);
+		}
+	}
+}
+
+void FSequenceIndexer::AddTrajectoryTimeFeatures(int32 SampleIdx)
+{
+	FPoseSearchFeatureDesc Feature;
+	Feature.Domain = EPoseSearchFeatureDomain::Time;
+	Feature.SchemaBoneIdx = FPoseSearchFeatureDesc::TrajectoryBoneIndex;
+
+	FSequenceIndexer::FSubsample OriginSample = ResolveSubsample(SampleIdx);
+	
+	for (int32 SchemaSubsampleIdx = 0; SchemaSubsampleIdx != Input.Schema->TrajectorySampleOffsets.Num(); ++SchemaSubsampleIdx)
+	{
+		Feature.SubsampleIdx = SchemaSubsampleIdx;
 
-		Builder.SetTransform(CurrentElement, BlendedRootInSampleSpace);
+		const int32 SubsampleIdx = SampleIdx + Input.Schema->TrajectorySampleOffsets[SchemaSubsampleIdx];
+
+		FSequenceIndexer::FSubsample Subsample = ResolveSubsample(SubsampleIdx);
+		FTransform SubsampleRoot = Subsample.AccumulatedRootMotion;
+		SubsampleRoot.SetToRelativeTransform(OriginSample.AccumulatedRootMotion);
+
+		FSequenceIndexer::FSubsample SubsamplePrev = ResolveSubsample(SubsampleIdx - 1);
+		FTransform SubsamplePrevRoot = SubsamplePrev.AccumulatedRootMotion;
+		SubsamplePrevRoot.SetToRelativeTransform(OriginSample.AccumulatedRootMotion);
+
+		FeatureVector.SetTransform(Feature, SubsampleRoot);
+		FeatureVector.SetTransformDerivative(Feature, SubsampleRoot, SubsamplePrevRoot, Input.Schema->SamplingInterval);
 	}
 }
 
+void FSequenceIndexer::AddTrajectoryDistanceFeatures(int32 SampleIdx)
+{
+	// This function needs to be rewritten to work with the updated sampler
+	// and lead-in/follow-up anims
+
+//	FPoseSearchFeatureDesc Feature;
+//	Feature.Domain = EPoseSearchFeatureDomain::Distance;
+//	Feature.SchemaBoneIdx = FPoseSearchFeatureDesc::TrajectoryBoneIndex;
+//
+//	FSequenceIndexer::FSampleRef OriginSampleRef = ResolveSampleRef(SampleIdx);
+//
+//	for (int32 SubsampleIdx = 0; SubsampleIdx != Input.Schema->TrajectorySampleDistances.Num(); ++SubsampleIdx)
+//	{
+//		Feature.SubsampleIdx = SubsampleIdx;
+//
+//		const float TrajectoryDistance = Input.Schema->TrajectorySampleDistances[SubsampleIdx];
+//		const float SampleAccumulatedRootDistance = TrajectoryDistance + AccumulatedRootDistances[SampleIdx];
+//
+//		int32 LowerBoundSampleIdx = Algo::LowerBound(AccumulatedRootDistances, SampleAccumulatedRootDistance);
+//
+//		// @@@ Add extrapolation. Clamp for now
+//		int32 PrevSampleIdx = FMath::Clamp(LowerBoundSampleIdx - 1, 0, AccumulatedRootDistances.Num() - 1);
+//		int32 NextSampleIdx = FMath::Clamp(LowerBoundSampleIdx, 0, AccumulatedRootDistances.Num() - 1);
+//
+//		const float PrevSampleDistance = AccumulatedRootDistances[PrevSampleIdx];
+//		const float NextSampleDistance = AccumulatedRootDistances[NextSampleIdx];
+//
+//		FTransform PrevRootInSampleSpace = AccumulatedRootMotion[PrevSampleIdx];
+//		PrevRootInSampleSpace.SetToRelativeTransform(SampleSpaceOrigin);
+//
+//		FTransform NextRootInSampleSpace = AccumulatedRootMotion[NextSampleIdx];
+//		NextRootInSampleSpace.SetToRelativeTransform(SampleSpaceOrigin);
+//		
+//		float Alpha = FMath::GetRangePct(PrevSampleDistance, NextSampleDistance, SampleAccumulatedRootDistance);
+//		FTransform BlendedRootInSampleSpace;
+//		BlendedRootInSampleSpace.Blend(PrevRootInSampleSpace, NextRootInSampleSpace, Alpha);
+//
+//		FeatureVector.SetTransform(Feature, BlendedRootInSampleSpace);
+//	}
+}
+
 
 //////////////////////////////////////////////////////////////////////////
 // PoseSearch API
@@ -1322,10 +1592,9 @@ static void DrawTrajectoryFeatures(const FDebugDrawParams& DrawParams, const FFe
 		return;
 	}
 
-	FVector TrajectoryPosPrev(0.0f);
-	for (int32 SubsampleIdx = 0; SubsampleIdx != NumSubsamples; ++SubsampleIdx)
+	for (int32 SchemaSubsampleIdx = 0; SchemaSubsampleIdx != NumSubsamples; ++SchemaSubsampleIdx)
 	{
-		Feature.SubsampleIdx = SubsampleIdx;
+		Feature.SubsampleIdx = SchemaSubsampleIdx;
 
 		FVector TrajectoryPos;
 		if (Reader.GetPosition(Feature, &TrajectoryPos))
@@ -1336,9 +1605,6 @@ static void DrawTrajectoryFeatures(const FDebugDrawParams& DrawParams, const FFe
 
 			TrajectoryPos = DrawParams.RootTransform.TransformPosition(TrajectoryPos);
 			DrawDebugSphere(DrawParams.World, TrajectoryPos, DrawDebugSphereSize, DrawDebugSphereSegments, Color, false, LifeTime, DepthPriority,  DrawDebugSphereLineThickness);
-
-
-			TrajectoryPosPrev = TrajectoryPos;
 		}
 		else
 		{
@@ -1379,10 +1645,9 @@ static void DrawPoseFeatures(const FDebugDrawParams& DrawParams, const FFeatureV
 		return;
 	}
 
-	for (int32 SubsampleIdx = 0; SubsampleIdx != NumSubsamples; ++SubsampleIdx)
+	for (int32 SchemaSubsampleIdx = 0; SchemaSubsampleIdx != NumSubsamples; ++SchemaSubsampleIdx)
 	{
-		Feature.SubsampleIdx = SubsampleIdx;
-		Feature.SchemaBoneIdx = 0;
+		Feature.SubsampleIdx = SchemaSubsampleIdx;
 		
 		for (int32 SchemaBoneIdx = 0; SchemaBoneIdx != NumBones; ++SchemaBoneIdx)
 		{
@@ -1509,11 +1774,24 @@ bool BuildIndex(const UAnimSequence* Sequence, UPoseSearchSequenceMetaData* Sequ
 		return false;
 	}
 
-	FSequenceIndexer Indexer;
-	FSequenceIndexer::Result Result = Indexer.Process(SequenceMetaData->Schema, Sequence, SequenceMetaData->SamplingRange);
+	FSequenceSampler Sampler;
+	FSequenceSampler::FInput SamplerInput;
+	SamplerInput.Schema = SequenceMetaData->Schema;
+	SamplerInput.Sequence = Sequence;
+	SamplerInput.bLoopable = false;
+	Sampler.Init(SamplerInput);
+	Sampler.Process();
 
-	SequenceMetaData->SearchIndex.Values = Result.Values;
-	SequenceMetaData->SearchIndex.NumPoses = Result.NumIndexedPoses;
+	FSequenceIndexer Indexer;
+	FSequenceIndexer::FInput IndexerInput;
+	IndexerInput.MainSequence = &Sampler;
+	IndexerInput.Schema = SequenceMetaData->Schema;
+	IndexerInput.RequestedSamplingRange = SequenceMetaData->SamplingRange;
+	Indexer.Init(IndexerInput);
+	Indexer.Process();
+
+	SequenceMetaData->SearchIndex.Values = Indexer.Output.FeatureVectorTable;
+	SequenceMetaData->SearchIndex.NumPoses = Indexer.Output.NumIndexedPoses;
 	SequenceMetaData->SearchIndex.Schema = SequenceMetaData->Schema;
 	return true;
 }
@@ -1536,19 +1814,72 @@ bool BuildIndex(UPoseSearchDatabase* Database)
 		}
 	}
 
+	// Prepare animation sampling tasks
+	TArray<FSequenceSampler> SequenceSamplers;
+	TMap<const UAnimSequence*, int32> SequenceSamplerMap;
+
+	auto AddSampler = [&](const UAnimSequence* Sequence, bool bLoopable)
+	{
+		if (!SequenceSamplerMap.Contains(Sequence))
+		{
+			int32 SequenceSamplerIdx = SequenceSamplers.AddDefaulted();
+			SequenceSamplerMap.Add(Sequence, SequenceSamplerIdx);
+
+			FSequenceSampler::FInput Input;
+			Input.Schema = Database->Schema;
+			Input.Sequence = Sequence;
+			Input.bLoopable = bLoopable;
+			SequenceSamplers[SequenceSamplerIdx].Init(Input);
+		}
+	};
+
+	for (const FPoseSearchDatabaseSequence& DbSequence : Database->Sequences)
+	{
+		if (DbSequence.Sequence)
+		{
+			AddSampler(DbSequence.Sequence, DbSequence.bLoopAnimation);
+		}
+
+		if (DbSequence.LeadInSequence)
+		{
+			AddSampler(DbSequence.LeadInSequence, DbSequence.bLoopLeadInAnimation);
+		}
+
+		if (DbSequence.FollowUpSequence)
+		{
+			AddSampler(DbSequence.FollowUpSequence, DbSequence.bLoopFollowUpAnimation);
+		}
+	}
+
+	// Sample animations independently
+	ParallelFor(SequenceSamplers.Num(), [&SequenceSamplers](int32 SamplerIdx){ SequenceSamplers[SamplerIdx].Process(); });
+
+
+	auto GetSampler = [&](const UAnimSequence* Sequence) -> const FSequenceSampler*
+	{
+		return Sequence ? &SequenceSamplers[SequenceSamplerMap[Sequence]] : nullptr;
+	};
+
 	// Prepare animation indexing tasks
 	TArray<FSequenceIndexer> Indexers;
 	Indexers.SetNum(Database->Sequences.Num());
-
-	auto IndexerTask = [&Database, &Indexers](int32 SequenceIdx)
+	for (int32 SequenceIdx = 0; SequenceIdx != Database->Sequences.Num(); ++SequenceIdx)
 	{
 		const FPoseSearchDatabaseSequence& DbSequence = Database->Sequences[SequenceIdx];
 		FSequenceIndexer& Indexer = Indexers[SequenceIdx];
-		Indexer.Process(Database->Schema, DbSequence.Sequence, DbSequence.SamplingRange);
-	};
+
+		FSequenceIndexer::FInput Input;
+		Input.MainSequence = GetSampler(DbSequence.Sequence);
+		Input.LeadInSequence = GetSampler(DbSequence.LeadInSequence);
+		Input.FollowUpSequence = GetSampler(DbSequence.FollowUpSequence);
+		Input.Schema = Database->Schema;
+		Input.RequestedSamplingRange = DbSequence.SamplingRange;
+		Indexer.Init(Input);
+	}
 
 	// Index animations independently
-	ParallelFor(Database->Sequences.Num(), IndexerTask);
+	ParallelFor(Indexers.Num(), [&Indexers](int32 SequenceIdx){ Indexers[SequenceIdx].Process(); });
+
 
 	// Write index info to sequence and count up total poses and storage required
 	int32 TotalPoses = 0;
@@ -1556,19 +1887,19 @@ bool BuildIndex(UPoseSearchDatabase* Database)
 	for (int32 SequenceIdx = 0; SequenceIdx != Database->Sequences.Num(); ++SequenceIdx)
 	{
 		FPoseSearchDatabaseSequence& DbSequence = Database->Sequences[SequenceIdx];
-		FSequenceIndexer::Result Result = Indexers[SequenceIdx].GetResult();
-		DbSequence.NumPoses = Result.NumIndexedPoses;
+		const FSequenceIndexer::FOutput& Output = Indexers[SequenceIdx].Output;
+		DbSequence.NumPoses = Output.NumIndexedPoses;
 		DbSequence.FirstPoseIdx = TotalPoses;
-		TotalPoses += Result.NumIndexedPoses;
-		TotalFloats += Result.Values.Num();
+		TotalPoses += Output.NumIndexedPoses;
+		TotalFloats += Output.FeatureVectorTable.Num();
 	}
 
 	// Join animation data into a single search index
 	Database->SearchIndex.Values.Reset(TotalFloats);
 	for (const FSequenceIndexer& Indexer : Indexers)
 	{
-		FSequenceIndexer::Result Result = Indexer.GetResult();
-		Database->SearchIndex.Values.Append(Result.Values.GetData(), Result.Values.Num());
+		const FSequenceIndexer::FOutput& Output = Indexer.Output;
+		Database->SearchIndex.Values.Append(Output.FeatureVectorTable.GetData(), Output.FeatureVectorTable.Num());
 	}
 
 	Database->SearchIndex.NumPoses = TotalPoses;
diff --git a/Engine/Plugins/Experimental/Animation/PoseSearch/Source/Runtime/Public/PoseSearch/AnimNode_MotionMatching.h b/Engine/Plugins/Experimental/Animation/PoseSearch/Source/Runtime/Public/PoseSearch/AnimNode_MotionMatching.h
--- a/Engine/Plugins/Experimental/Animation/PoseSearch/Source/Runtime/Public/PoseSearch/AnimNode_MotionMatching.h
+++ b/Engine/Plugins/Experimental/Animation/PoseSearch/Source/Runtime/Public/PoseSearch/AnimNode_MotionMatching.h
@@ -32,6 +32,12 @@ public:
 	UPROPERTY(EditAnywhere, BlueprintReadWrite, Category=Settings, meta=(ClampMin="0"))
 	float PoseJumpThreshold = 0.5f;
 
+	UPROPERTY(EditAnywhere, BlueprintReadWrite, Category=Settings, meta=(ClampMin="0"))
+	float SimilarityThreshold = 0.1f;
+
+	UPROPERTY(EditAnywhere, BlueprintReadWrite, Category=Settings, meta=(ClampMin="0"))
+	float ThrottleTime = 0.1f;
+
 #if WITH_EDITORONLY_DATA
 	UPROPERTY(EditAnywhere, Category=Debug)
 	bool bDebugDraw = false;
@@ -60,7 +66,10 @@ private:
 	// The current query feature vector used to search the database for pose candidates
 	FPoseSearchFeatureVectorBuilder ComposedQuery;
 
+	// Time since the last pose jump
+	float ElapsedPoseJumpTime = 0.0f;
+
 	bool IsValidForSearch() const;
-	void ComposeQuery();
+	void ComposeQuery(const FAnimationBaseContext& Context);
 	void JumpToPose(const FAnimationUpdateContext& Context, UE::PoseSearch::FDbSearchResult Result);
 };
\ No newline at end of file
diff --git a/Engine/Plugins/Experimental/Animation/PoseSearch/Source/Runtime/Public/PoseSearch/PoseSearch.h b/Engine/Plugins/Experimental/Animation/PoseSearch/Source/Runtime/Public/PoseSearch/PoseSearch.h
--- a/Engine/Plugins/Experimental/Animation/PoseSearch/Source/Runtime/Public/PoseSearch/PoseSearch.h
+++ b/Engine/Plugins/Experimental/Animation/PoseSearch/Source/Runtime/Public/PoseSearch/PoseSearch.h
@@ -275,6 +275,26 @@ struct POSESEARCH_API FPoseSearchDatabaseSequence
 	UPROPERTY(EditAnywhere, Category="Sequence")
 	bool bLoopAnimation = false;
 
+	// Used for sampling past pose information at the beginning of the main sequence.
+	// This setting is intended for transitions between cycles. It is optional and only used
+	// for one shot anims with past sampling. When past sampling is used without a lead in sequence,
+	// the sampling range of the main sequence will be clamped if necessary.
+	UPROPERTY(EditAnywhere, Category="Sequence")
+	UAnimSequence* LeadInSequence = nullptr;
+
+	UPROPERTY(EditAnywhere, Category="Sequence")
+	bool bLoopLeadInAnimation = false;
+
+	// Used for sampling future pose information at the end of the main sequence.
+	// This setting is intended for transitions between cycles. It is optional and only used
+	// for one shot anims with future sampling. When future sampling is used without a follow up sequence,
+	// the sampling range of the main sequence will be clamped if necessary.
+	UPROPERTY(EditAnywhere, Category="Sequence")
+	UAnimSequence* FollowUpSequence = nullptr;
+
+	UPROPERTY(EditAnywhere, Category="Sequence")
+	bool bLoopFollowUpAnimation = false;
+
 	UPROPERTY()
 	int32 FirstPoseIdx = 0;
 
@@ -314,7 +334,7 @@ public: // UObject
 /** 
 * Helper object for writing features into a float buffer according to a feature vector layout.
 * Keeps track of which features are present, allowing the feature vector to be built up piecemeal.
-* FPoseSearchFeatureVectorBuilder is used to build search queries at runtime and during search index construction.
+* FFeatureVectorBuilder is used to build search queries at runtime and for adding samples during search index construction.
 */
 
 USTRUCT(BlueprintType, Category = "Animation|PoseSearch")
@@ -335,6 +355,7 @@ public:
 	void SetAngularVelocity(FPoseSearchFeatureDesc Feature, const FTransform& Transform, const FTransform& PrevTransform, float DeltaTime);
 	void SetVector(FPoseSearchFeatureDesc Feature, const FVector& Vector);
 	bool SetPoseFeatures(UE::PoseSearch::FPoseHistory* History);
+	bool SetPastTrajectoryFeatures(UE::PoseSearch::FPoseHistory* History);
 
 	void Copy(TArrayView<const float> FeatureVector);
 	void CopyFeature(const FPoseSearchFeatureVectorBuilder& OtherBuilder, int32 FeatureIdx);
@@ -364,8 +385,9 @@ class POSESEARCH_API FPoseHistory
 public:
 	void Init(int32 InNumPoses, float InTimeHorizon);
 	void Init(const FPoseHistory& History);
-	bool Sample(float SecondsAgo, const FReferenceSkeleton& RefSkeleton, const TArray<FBoneIndexType>& RequiredBones);
-	void Update(float SecondsElapsed, const FCompactPose& EvalContext);
+	bool SamplePose(float SecondsAgo, const FReferenceSkeleton& RefSkeleton, const TArray<FBoneIndexType>& RequiredBones);
+	bool SampleRoot(float SecondsAgo, FTransform* OutTransform) const;
+	void Update(float SecondsElapsed, const FPoseContext& PoseContext);
 	float GetSampleInterval() const;
 	TArrayView<const FTransform> GetLocalPoseSample() const { return SampledLocalPose; }
 	TArrayView<const FTransform> GetComponentPoseSample() const { return SampledComponentPose; }
@@ -380,6 +402,7 @@ private:
 	struct FPose
 	{
 		TArray<FTransform> LocalTransforms;
+		FTransform WorldComponentTransform;
 	};
 
 	TRingBuffer<FPose> Poses;
