PoseSearch: Fixed root motion accumulation bugs, added debug vis for motion matching search result, removed motion matching checks for jumping to the same sequence, updated motion matching node options for clarity

diff --git a/Engine/Plugins/Experimental/Animation/PoseSearch/Source/Runtime/Private/AnimNode_MotionMatching.cpp b/Engine/Plugins/Experimental/Animation/PoseSearch/Source/Runtime/Private/AnimNode_MotionMatching.cpp
--- a/Engine/Plugins/Experimental/Animation/PoseSearch/Source/Runtime/Private/AnimNode_MotionMatching.cpp
+++ b/Engine/Plugins/Experimental/Animation/PoseSearch/Source/Runtime/Private/AnimNode_MotionMatching.cpp
@@ -21,7 +21,7 @@ void FAnimNode_MotionMatching::Initialize_AnyThread(const FAnimationInitializeCo
 
 	DbPoseIdx = INDEX_NONE;
 	DbSequenceIdx = INDEX_NONE;
-	ElapsedPoseJumpTime = ThrottleTime;
+	ElapsedPoseJumpTime = SearchThrottleTime;
 
 	Source.SetLinkNode(&SequencePlayerNode);
 	Source.Initialize(Context);
@@ -91,12 +91,12 @@ void FAnimNode_MotionMatching::Update_AnyThread(const FAnimationUpdateContext& C
 
 		// Search the database for the nearest match to the updated query vector
 		UE::PoseSearch::FDbSearchResult Result = UE::PoseSearch::Search(Database, ComposedQuery.GetValues());
-		if (Result.IsValid() && (ElapsedPoseJumpTime >= ThrottleTime))
+		if (Result.IsValid() && (ElapsedPoseJumpTime >= SearchThrottleTime))
 		{
 			const FPoseSearchDatabaseSequence& ResultDbSequence = Database->Sequences[Result.DbSequenceIdx];
 
 			// Consider the search result better if it is more similar to the query than the current pose we're playing back from the database
-			bool bBetterPose = Result.Dissimilarity * (1.0f + SimilarityThreshold) < CurrentDissimilarity;
+			bool bBetterPose = Result.Dissimilarity * (1.0f + (MinPercentImprovement / 100.0f)) < CurrentDissimilarity;
 
 			// We'll ignore the candidate pose if it is too near to our current pose
 			bool bNearbyPose = false;
@@ -109,13 +109,8 @@ void FAnimNode_MotionMatching::Update_AnyThread(const FAnimationUpdateContext& C
 				}
 			}
 
-			// And we won't bother to jump to another pose within the same looping sequence we're already playing
-			// (This should probably be configurable)
-			bool bSameSequence = DbSequenceIdx == Result.DbSequenceIdx;
-			bool bSameCycle = bSameSequence && ResultDbSequence.bLoopAnimation;
-
 			// Start playback from the candidate pose if we determined it was a better option
-			if (bBetterPose && !bNearbyPose && !bSameCycle && !bSameSequence)
+			if (bBetterPose && !bNearbyPose)
 			{
 				JumpToPose(Context, Result);
 				bJumpedToPose = true;
@@ -177,13 +172,16 @@ void FAnimNode_MotionMatching::PreUpdate(const UAnimInstance* InAnimInstance)
 
 		UE::PoseSearch::FDebugDrawParams DrawParams;
 		DrawParams.RootTransform = SkeletalMeshComponent->GetComponentTransform();
-		DrawParams.Flags = UE::PoseSearch::EDebugDrawFlags::DrawQuery;
+		DrawParams.Flags = UE::PoseSearch::EDebugDrawFlags::DrawBest;
 		DrawParams.Database = Database;
 		DrawParams.Query = ComposedQuery.GetValues();
 		DrawParams.World = SkeletalMeshComponent->GetWorld();
 		DrawParams.DefaultLifeTime = 0.0f;
 		DrawParams.HighlightPoseIdx = DbPoseIdx;
 		UE::PoseSearch::Draw(DrawParams);
+
+		DrawParams.Flags = UE::PoseSearch::EDebugDrawFlags::DrawQuery;
+		UE::PoseSearch::Draw(DrawParams);
 	}
 #endif
 }
diff --git a/Engine/Plugins/Experimental/Animation/PoseSearch/Source/Runtime/Private/PoseSearch.cpp b/Engine/Plugins/Experimental/Animation/PoseSearch/Source/Runtime/Private/PoseSearch.cpp
--- a/Engine/Plugins/Experimental/Animation/PoseSearch/Source/Runtime/Private/PoseSearch.cpp
+++ b/Engine/Plugins/Experimental/Animation/PoseSearch/Source/Runtime/Private/PoseSearch.cpp
@@ -1162,10 +1162,15 @@ FSequenceSampler::FWrappedSampleIndex FSequenceSampler::WrapOrClampSubsampleInde
 	// Wrap the index if this is a loopable sequence
 	if (Input.bLoopable)
 	{
-		while (Result.Idx < 0)
+		if (Result.Idx < 0)
 		{
 			Result.Idx += Output.TotalSamples;
-			++Result.NumCycles;
+
+			while (Result.Idx < 0)
+			{
+				Result.Idx += Output.TotalSamples;
+				++Result.NumCycles;
+			}
 		}
 
 		while (Result.Idx >= Output.TotalSamples)
@@ -1237,7 +1242,7 @@ void FSequenceSampler::ExtractRootMotion()
 		FTransform LocalRootMotion = Input.Sequence->ExtractRootMotion(CurrentTime, Input.Schema->SamplingInterval, false /*!allowLooping*/);
 		Output.LocalRootMotion.Add(LocalRootMotion);
 
-		AccumulatedRootMotion = AccumulatedRootMotion * LocalRootMotion;
+		AccumulatedRootMotion = LocalRootMotion * AccumulatedRootMotion;
 		AccumulatedRootDistance += LocalRootMotion.GetTranslation().Size();
 		Output.AccumulatedRootMotion.Add(AccumulatedRootMotion);
 		Output.AccumulatedRootDistance.Add((float)AccumulatedRootDistance);
@@ -1382,7 +1387,7 @@ FSequenceIndexer::FSubsample FSequenceIndexer::ResolveSubsample(int32 MainSubsam
 	// Use the lead in anim if we had to clamp to the beginning of the main anim
 	if (MainSample.bClamped && (MainSubsampleIdx < 0) && Input.LeadInSequence)
 	{
-		EffectiveSample = Input.LeadInSequence->WrapOrClampSubsampleIndex(MainSubsampleIdx + Input.LeadInSequence->Output.TotalSamples);
+		EffectiveSample = Input.LeadInSequence->WrapOrClampSubsampleIndex(MainSubsampleIdx);
 
 		Subsample.Sampler = Input.LeadInSequence;
 		Subsample.AbsoluteSampleIdx = EffectiveSample.Idx;
diff --git a/Engine/Plugins/Experimental/Animation/PoseSearch/Source/Runtime/Public/PoseSearch/AnimNode_MotionMatching.h b/Engine/Plugins/Experimental/Animation/PoseSearch/Source/Runtime/Public/PoseSearch/AnimNode_MotionMatching.h
--- a/Engine/Plugins/Experimental/Animation/PoseSearch/Source/Runtime/Public/PoseSearch/AnimNode_MotionMatching.h
+++ b/Engine/Plugins/Experimental/Animation/PoseSearch/Source/Runtime/Public/PoseSearch/AnimNode_MotionMatching.h
@@ -26,17 +26,20 @@ public:
 
 	// Time in seconds to blend out to the new pose. Uses inertial blending and requires an Inertialization node after this node.
 	UPROPERTY(EditAnywhere, BlueprintReadWrite, Category=Settings, meta=(ClampMin="0"))
-	float BlendTime = 0.1f;
+	float BlendTime = 0.2f;
 
-	// Ignore pose candidates within the active animation that are less than PoseJumpThreshold seconds away from the current asset player time
+	// Don't jump to poses that are less than this many seconds away
 	UPROPERTY(EditAnywhere, BlueprintReadWrite, Category=Settings, meta=(ClampMin="0"))
-	float PoseJumpThreshold = 0.5f;
+	float PoseJumpThreshold = 4.0f;
 
+	// Minimum amount of time to wait between pose searches
 	UPROPERTY(EditAnywhere, BlueprintReadWrite, Category=Settings, meta=(ClampMin="0"))
-	float SimilarityThreshold = 0.1f;
+	float SearchThrottleTime = 0.1f;
 
+	// How much better the search result must be compared to the current pose in order to jump to it
+	// Note: This feature won't work quite as advertised until search data rescaling is implemented
 	UPROPERTY(EditAnywhere, BlueprintReadWrite, Category=Settings, meta=(ClampMin="0"))
-	float ThrottleTime = 0.1f;
+	float MinPercentImprovement = 0.0f;
 
 #if WITH_EDITORONLY_DATA
 	UPROPERTY(EditAnywhere, Category=Debug)
