Change presave prototype: temporarily remove deprecation to find the discrepancy with the farm build vs my local build.
#rnx

diff --git a/Engine/Plugins/Enterprise/LidarPointCloud/Source/LidarPointCloudRuntime/Public/LidarPointCloud.h b/Engine/Plugins/Enterprise/LidarPointCloud/Source/LidarPointCloudRuntime/Public/LidarPointCloud.h
--- a/Engine/Plugins/Enterprise/LidarPointCloud/Source/LidarPointCloudRuntime/Public/LidarPointCloud.h
+++ b/Engine/Plugins/Enterprise/LidarPointCloud/Source/LidarPointCloudRuntime/Public/LidarPointCloud.h
@@ -160,7 +160,7 @@ public:
 	virtual void PostLoad() override;
 	virtual void GetAssetRegistryTags(TArray<FAssetRegistryTag>& OutTags) const override;
 	virtual void BeginDestroy() override;
-	UE_DEPRECATED(5.0, "Use version that takes FObjectPreSaveContext instead.")
+	//UE_DEPRECATED(5.0, "Use version that takes FObjectPreSaveContext instead.")
 	virtual void PreSave(const class ITargetPlatform* TargetPlatform) override;
 	virtual void PreSave(FObjectPreSaveContext ObjectSaveContext) override;
 #if WITH_EDITOR
diff --git a/Engine/Plugins/Experimental/Animation/ContextualAnimation/Source/ContextualAnimation/Public/ContextualAnimCompositeSceneAsset.h b/Engine/Plugins/Experimental/Animation/ContextualAnimation/Source/ContextualAnimation/Public/ContextualAnimCompositeSceneAsset.h
--- a/Engine/Plugins/Experimental/Animation/ContextualAnimation/Source/ContextualAnimation/Public/ContextualAnimCompositeSceneAsset.h
+++ b/Engine/Plugins/Experimental/Animation/ContextualAnimation/Source/ContextualAnimation/Public/ContextualAnimCompositeSceneAsset.h
@@ -82,7 +82,7 @@ public:
 
 	UContextualAnimCompositeSceneAsset(const FObjectInitializer& ObjectInitializer);
 
-	UE_DEPRECATED(5.0, "Use version that takes FObjectPreSaveContext instead.")
+	//UE_DEPRECATED(5.0, "Use version that takes FObjectPreSaveContext instead.")
 	virtual void PreSave(const class ITargetPlatform* TargetPlatform) override;
 	virtual void PreSave(FObjectPreSaveContext ObjectSaveContext) override;
 
diff --git a/Engine/Plugins/Experimental/Animation/ContextualAnimation/Source/ContextualAnimation/Public/ContextualAnimSceneAsset.h b/Engine/Plugins/Experimental/Animation/ContextualAnimation/Source/ContextualAnimation/Public/ContextualAnimSceneAsset.h
--- a/Engine/Plugins/Experimental/Animation/ContextualAnimation/Source/ContextualAnimation/Public/ContextualAnimSceneAsset.h
+++ b/Engine/Plugins/Experimental/Animation/ContextualAnimation/Source/ContextualAnimation/Public/ContextualAnimSceneAsset.h
@@ -17,7 +17,7 @@ public:
 
 	UContextualAnimSceneAsset(const FObjectInitializer& ObjectInitializer);
 
-	UE_DEPRECATED(5.0, "Use version that takes FObjectPreSaveContext instead.")
+	//UE_DEPRECATED(5.0, "Use version that takes FObjectPreSaveContext instead.")
 	virtual void PreSave(const class ITargetPlatform* TargetPlatform) override;
 	virtual void PreSave(FObjectPreSaveContext ObjectSaveContext) override;
 
diff --git a/Engine/Plugins/Experimental/Animation/ContextualAnimation/Source/ContextualAnimation/Public/ContextualAnimSceneAssetBase.h b/Engine/Plugins/Experimental/Animation/ContextualAnimation/Source/ContextualAnimation/Public/ContextualAnimSceneAssetBase.h
--- a/Engine/Plugins/Experimental/Animation/ContextualAnimation/Source/ContextualAnimation/Public/ContextualAnimSceneAssetBase.h
+++ b/Engine/Plugins/Experimental/Animation/ContextualAnimation/Source/ContextualAnimation/Public/ContextualAnimSceneAssetBase.h
@@ -44,7 +44,7 @@ public:
 
 	UContextualAnimSceneAssetBase(const FObjectInitializer& ObjectInitializer);
 
-	UE_DEPRECATED(5.0, "Use version that takes FObjectPreSaveContext instead.")
+	//UE_DEPRECATED(5.0, "Use version that takes FObjectPreSaveContext instead.")
 	virtual void PreSave(const class ITargetPlatform* TargetPlatform) override;
 	virtual void PreSave(FObjectPreSaveContext ObjectSaveContext) override;
 
diff --git a/Engine/Plugins/Experimental/Animation/PoseSearch/Source/Runtime/Public/PoseSearch/PoseSearch.h b/Engine/Plugins/Experimental/Animation/PoseSearch/Source/Runtime/Public/PoseSearch/PoseSearch.h
--- a/Engine/Plugins/Experimental/Animation/PoseSearch/Source/Runtime/Public/PoseSearch/PoseSearch.h
+++ b/Engine/Plugins/Experimental/Animation/PoseSearch/Source/Runtime/Public/PoseSearch/PoseSearch.h
@@ -201,7 +201,7 @@ public:
 	int32 NumBones () const { return BoneIndices.Num(); }
 
 public: // UObject
-	UE_DEPRECATED(5.0, "Use version that takes FObjectPreSaveContext instead.")
+	//UE_DEPRECATED(5.0, "Use version that takes FObjectPreSaveContext instead.")
 	virtual void PreSave(const class ITargetPlatform* TargetPlatform) override;
 	virtual void PreSave(FObjectPreSaveContext ObjectSaveContext) override;
 	virtual void PostLoad() override;
@@ -260,7 +260,7 @@ public:
 	bool IsValidForSearch() const;
 
 public: // UObject
-	UE_DEPRECATED(5.0, "Use version that takes FObjectPreSaveContext instead.")
+	//UE_DEPRECATED(5.0, "Use version that takes FObjectPreSaveContext instead.")
 	virtual void PreSave(const class ITargetPlatform* TargetPlatform) override;
 	virtual void PreSave(FObjectPreSaveContext ObjectSaveContext) override;
 };
@@ -332,7 +332,7 @@ public:
 	bool IsValidForSearch() const;
 
 public: // UObject
-	UE_DEPRECATED(5.0, "Use version that takes FObjectPreSaveContext instead.")
+	//UE_DEPRECATED(5.0, "Use version that takes FObjectPreSaveContext instead.")
 	virtual void PreSave(const class ITargetPlatform* TargetPlatform) override;
 	virtual void PreSave(FObjectPreSaveContext ObjectSaveContext) override;
 };
diff --git a/Engine/Plugins/Experimental/CommonConversation/Source/CommonConversationRuntime/Public/ConversationDatabase.h b/Engine/Plugins/Experimental/CommonConversation/Source/CommonConversationRuntime/Public/ConversationDatabase.h
--- a/Engine/Plugins/Experimental/CommonConversation/Source/CommonConversationRuntime/Public/ConversationDatabase.h
+++ b/Engine/Plugins/Experimental/CommonConversation/Source/CommonConversationRuntime/Public/ConversationDatabase.h
@@ -77,7 +77,7 @@ public:
 #if WITH_EDITOR
 	virtual void PostDuplicate(bool bDuplicateForPIE) override;
 	virtual EDataValidationResult IsDataValid(TArray<FText>& ValidationErrors) override;
-	UE_DEPRECATED(5.0, "Use version that takes FObjectPreSaveContext instead.")
+	//UE_DEPRECATED(5.0, "Use version that takes FObjectPreSaveContext instead.")
 	virtual void PreSave(const class ITargetPlatform* TargetPlatform) override;
 	virtual void PreSave(FObjectPreSaveContext ObjectSaveContext) override;
 #endif
diff --git a/Engine/Plugins/Experimental/ControlRig/Source/ControlRigDeveloper/Public/ControlRigBlueprint.h b/Engine/Plugins/Experimental/ControlRig/Source/ControlRigDeveloper/Public/ControlRigBlueprint.h
--- a/Engine/Plugins/Experimental/ControlRig/Source/ControlRigDeveloper/Public/ControlRigBlueprint.h
+++ b/Engine/Plugins/Experimental/ControlRig/Source/ControlRigDeveloper/Public/ControlRigBlueprint.h
@@ -161,7 +161,7 @@ public:
 	virtual void GetTypeActions(FBlueprintActionDatabaseRegistrar& ActionRegistrar) const override;
 	virtual void GetInstanceActions(FBlueprintActionDatabaseRegistrar& ActionRegistrar) const override;
 	virtual void SetObjectBeingDebugged(UObject* NewObject) override;
-	UE_DEPRECATED(5.0, "Use version that takes FObjectPreSaveContext instead.")
+	//UE_DEPRECATED(5.0, "Use version that takes FObjectPreSaveContext instead.")
 	virtual void PreSave(const class ITargetPlatform* TargetPlatform) override;
 	virtual void PreSave(FObjectPreSaveContext ObjectSaveContext) override;
 	virtual void PostLoad() override;
diff --git a/Engine/Plugins/Experimental/VirtualProductionUtilities/Source/VPUtilitiesEditor/Private/VPCameraBase.h b/Engine/Plugins/Experimental/VirtualProductionUtilities/Source/VPUtilitiesEditor/Private/VPCameraBase.h
--- a/Engine/Plugins/Experimental/VirtualProductionUtilities/Source/VPUtilitiesEditor/Private/VPCameraBase.h
+++ b/Engine/Plugins/Experimental/VirtualProductionUtilities/Source/VPUtilitiesEditor/Private/VPCameraBase.h
@@ -27,7 +27,7 @@ public:
 	virtual bool ShouldTickIfViewportsOnly() const override;
 	virtual void Tick(float DeltaSeconds) override;
 	virtual void BeginDestroy() override;
-	UE_DEPRECATED(5.0, "Use version that takes FObjectPreSaveContext instead.")
+	//UE_DEPRECATED(5.0, "Use version that takes FObjectPreSaveContext instead.")
 	virtual void PreSave(const class ITargetPlatform* TargetPlatform) override;
 	virtual void PreSave(FObjectPreSaveContext ObjectSaveContext) override;
 
diff --git a/Engine/Plugins/FX/Niagara/Source/Niagara/Classes/NiagaraScript.h b/Engine/Plugins/FX/Niagara/Source/Niagara/Classes/NiagaraScript.h
--- a/Engine/Plugins/FX/Niagara/Source/Niagara/Classes/NiagaraScript.h
+++ b/Engine/Plugins/FX/Niagara/Source/Niagara/Classes/NiagaraScript.h
@@ -610,7 +610,7 @@ public:
 #endif
 
 	//~ Begin UObject interface
-	UE_DEPRECATED(5.0, "Use version that takes FObjectPreSaveContext instead.")
+	//UE_DEPRECATED(5.0, "Use version that takes FObjectPreSaveContext instead.")
 	virtual void PreSave(const class ITargetPlatform* TargetPlatform) override;
 	virtual void PreSave(FObjectPreSaveContext ObjectSaveContext) override;
 	virtual void Serialize(FArchive& Ar)override;
diff --git a/Engine/Plugins/FX/Niagara/Source/Niagara/Classes/NiagaraSystem.h b/Engine/Plugins/FX/Niagara/Source/Niagara/Classes/NiagaraSystem.h
--- a/Engine/Plugins/FX/Niagara/Source/Niagara/Classes/NiagaraSystem.h
+++ b/Engine/Plugins/FX/Niagara/Source/Niagara/Classes/NiagaraSystem.h
@@ -217,7 +217,7 @@ public:
 	virtual void Serialize(FArchive& Ar) override;
 	virtual void PostLoad() override; 
 	virtual void BeginDestroy() override;
-	UE_DEPRECATED(5.0, "Use version that takes FObjectPreSaveContext instead.")
+	//UE_DEPRECATED(5.0, "Use version that takes FObjectPreSaveContext instead.")
 	virtual void PreSave(const class ITargetPlatform* TargetPlatform) override;
 	virtual void PreSave(FObjectPreSaveContext ObjectSaveContext) override;
 #if WITH_EDITOR
diff --git a/Engine/Plugins/Runtime/DataRegistry/Source/DataRegistry/Public/DataRegistrySource_CurveTable.h b/Engine/Plugins/Runtime/DataRegistry/Source/DataRegistry/Public/DataRegistrySource_CurveTable.h
--- a/Engine/Plugins/Runtime/DataRegistry/Source/DataRegistry/Public/DataRegistrySource_CurveTable.h
+++ b/Engine/Plugins/Runtime/DataRegistry/Source/DataRegistry/Public/DataRegistrySource_CurveTable.h
@@ -72,7 +72,7 @@ protected:
 
 #if WITH_EDITOR
 	virtual void EditorRefreshSource();
-	UE_DEPRECATED(5.0, "Use version that takes FObjectPreSaveContext instead.")
+	//UE_DEPRECATED(5.0, "Use version that takes FObjectPreSaveContext instead.")
 	virtual void PreSave(const class ITargetPlatform* TargetPlatform) override;
 	virtual void PreSave(FObjectPreSaveContext ObjectSaveContext) override;
 
diff --git a/Engine/Plugins/Runtime/DataRegistry/Source/DataRegistry/Public/DataRegistrySource_DataTable.h b/Engine/Plugins/Runtime/DataRegistry/Source/DataRegistry/Public/DataRegistrySource_DataTable.h
--- a/Engine/Plugins/Runtime/DataRegistry/Source/DataRegistry/Public/DataRegistrySource_DataTable.h
+++ b/Engine/Plugins/Runtime/DataRegistry/Source/DataRegistry/Public/DataRegistrySource_DataTable.h
@@ -86,7 +86,7 @@ protected:
 
 #if WITH_EDITOR
 	virtual void EditorRefreshSource();
-	UE_DEPRECATED(5.0, "Use version that takes FObjectPreSaveContext instead.")
+	//UE_DEPRECATED(5.0, "Use version that takes FObjectPreSaveContext instead.")
 	virtual void PreSave(const class ITargetPlatform* TargetPlatform) override;
 	virtual void PreSave(FObjectPreSaveContext ObjectSaveContext) override;
 
diff --git a/Engine/Plugins/Runtime/GameplayAbilities/Source/GameplayAbilities/Public/GameplayEffect.h b/Engine/Plugins/Runtime/GameplayAbilities/Source/GameplayAbilities/Public/GameplayEffect.h
--- a/Engine/Plugins/Runtime/GameplayAbilities/Source/GameplayAbilities/Public/GameplayEffect.h
+++ b/Engine/Plugins/Runtime/GameplayAbilities/Source/GameplayAbilities/Public/GameplayEffect.h
@@ -1912,7 +1912,7 @@ public:
 	virtual void GetOwnedGameplayTags(FGameplayTagContainer& TagContainer) const override;
 	virtual void PostInitProperties() override;
 	virtual void PostLoad() override;
-	UE_DEPRECATED(5.0, "Use version that takes FObjectPreSaveContext instead.")
+	//UE_DEPRECATED(5.0, "Use version that takes FObjectPreSaveContext instead.")
 	virtual void PreSave(const class ITargetPlatform* TargetPlatform) override;
 	virtual void PreSave(FObjectPreSaveContext ObjectSaveContext) override;
 #if WITH_EDITOR
diff --git a/Engine/Plugins/Runtime/HairStrands/Source/HairStrandsCore/Public/GroomAsset.h b/Engine/Plugins/Runtime/HairStrands/Source/HairStrandsCore/Public/GroomAsset.h
--- a/Engine/Plugins/Runtime/HairStrands/Source/HairStrandsCore/Public/GroomAsset.h
+++ b/Engine/Plugins/Runtime/HairStrands/Source/HairStrandsCore/Public/GroomAsset.h
@@ -338,7 +338,7 @@ public:
 	virtual void BeginDestroy() override;
 	virtual void Serialize(FArchive& Ar) override;
 
-	UE_DEPRECATED(5.0, "Use version that takes FObjectPreSaveContext instead.")
+	//UE_DEPRECATED(5.0, "Use version that takes FObjectPreSaveContext instead.")
 	virtual void PreSave(const class ITargetPlatform* TargetPlatform) override;
 	virtual void PreSave(FObjectPreSaveContext ObjectSaveContext) override;
 
diff --git a/Engine/Plugins/Runtime/HairStrands/Source/HairStrandsCore/Public/GroomBindingAsset.h b/Engine/Plugins/Runtime/HairStrands/Source/HairStrandsCore/Public/GroomBindingAsset.h
--- a/Engine/Plugins/Runtime/HairStrands/Source/HairStrandsCore/Public/GroomBindingAsset.h
+++ b/Engine/Plugins/Runtime/HairStrands/Source/HairStrandsCore/Public/GroomBindingAsset.h
@@ -121,10 +121,10 @@ public:
 	//~ Begin UObject Interface.
 	virtual void GetResourceSizeEx(FResourceSizeEx& CumulativeResourceSize) override;
 	virtual void PostLoad() override;
-	UE_DEPRECATED(5.0, "Use version that takes FObjectPreSaveContext instead.")
+	//UE_DEPRECATED(5.0, "Use version that takes FObjectPreSaveContext instead.")
 	virtual void PreSave(const class ITargetPlatform* TargetPlatform) override;
 	virtual void PreSave(FObjectPreSaveContext ObjectSaveContext) override;
-	UE_DEPRECATED(5.0, "Use version that takes FObjectPostSaveContext instead.")
+	//UE_DEPRECATED(5.0, "Use version that takes FObjectPostSaveContext instead.")
 	virtual void PostSaveRoot(bool bCleanupIsRequired) override;
 	virtual void PostSaveRoot(FObjectPostSaveRootContext ObjectSaveContext) override;
 	virtual void BeginDestroy() override;
diff --git a/Engine/Source/Editor/BlueprintGraph/Classes/K2Node_FunctionEntry.h b/Engine/Source/Editor/BlueprintGraph/Classes/K2Node_FunctionEntry.h
--- a/Engine/Source/Editor/BlueprintGraph/Classes/K2Node_FunctionEntry.h
+++ b/Engine/Source/Editor/BlueprintGraph/Classes/K2Node_FunctionEntry.h
@@ -34,7 +34,7 @@ class UK2Node_FunctionEntry : public UK2Node_FunctionTerminator
 
 	//~ Begin UObject Interface
 	virtual void Serialize(FArchive& Ar) override;
-	UE_DEPRECATED(5.0, "Use version that takes FObjectPreSaveContext instead.")
+	//UE_DEPRECATED(5.0, "Use version that takes FObjectPreSaveContext instead.")
 	virtual void PreSave(const class ITargetPlatform* TargetPlatform) override;
 	virtual void PreSave(FObjectPreSaveContext ObjectSaveContext) override;
 	virtual void PostLoad() override;
diff --git a/Engine/Source/Editor/UMGEditor/Public/WidgetBlueprint.h b/Engine/Source/Editor/UMGEditor/Public/WidgetBlueprint.h
--- a/Engine/Source/Editor/UMGEditor/Public/WidgetBlueprint.h
+++ b/Engine/Source/Editor/UMGEditor/Public/WidgetBlueprint.h
@@ -257,7 +257,7 @@ public:
 	virtual void PostDuplicate(bool bDuplicateForPIE) override;
 
 #if WITH_EDITORONLY_DATA
-	UE_DEPRECATED(5.0, "Use version that takes FObjectPreSaveContext instead.")
+	//UE_DEPRECATED(5.0, "Use version that takes FObjectPreSaveContext instead.")
 	virtual void PreSave(const class ITargetPlatform* TargetPlatform) override;
 	virtual void PreSave(FObjectPreSaveContext ObjectSaveContext) override;
 #endif // WITH_EDITORONLY_DATA
diff --git a/Engine/Source/Editor/UnrealEd/Classes/Editor/EditorEngine.h b/Engine/Source/Editor/UnrealEd/Classes/Editor/EditorEngine.h
--- a/Engine/Source/Editor/UnrealEd/Classes/Editor/EditorEngine.h
+++ b/Engine/Source/Editor/UnrealEd/Classes/Editor/EditorEngine.h
@@ -2201,9 +2201,9 @@ public:
 	virtual bool InitializePhysicsSceneForSaveIfNecessary(UWorld* World, bool &bOutForceInitialized);
 	void CleanupPhysicsSceneThatWasInitializedForSave(UWorld* World, bool bForceInitialized);
 
-	UE_DEPRECATED(5.0, "Use version that takes FObjectPreSaveContext instead.")
+	//UE_DEPRECATED(5.0, "Use version that takes FObjectPreSaveContext instead.")
 	virtual void OnPreSaveWorld(uint32 SaveFlags, UWorld* World);
-	UE_DEPRECATED(5.0, "Use version that takes FObjectPostSaveContext instead.")
+	//UE_DEPRECATED(5.0, "Use version that takes FObjectPostSaveContext instead.")
 	virtual void OnPostSaveWorld(uint32 SaveFlags, UWorld* World, uint32 OriginalPackageFlags, bool bSuccess);
 
 	/** Invoked before a UWorld is saved to update editor systems */
diff --git a/Engine/Source/Runtime/AIModule/Classes/EnvironmentQuery/Tests/EnvQueryTest_GameplayTags.h b/Engine/Source/Runtime/AIModule/Classes/EnvironmentQuery/Tests/EnvQueryTest_GameplayTags.h
--- a/Engine/Source/Runtime/AIModule/Classes/EnvironmentQuery/Tests/EnvQueryTest_GameplayTags.h
+++ b/Engine/Source/Runtime/AIModule/Classes/EnvironmentQuery/Tests/EnvQueryTest_GameplayTags.h
@@ -31,7 +31,7 @@ protected:
 
 	bool SatisfiesTest(IGameplayTagAssetInterface* ItemGameplayTagAssetInterface) const;
 
-	UE_DEPRECATED(5.0, "Use version that takes FObjectPreSaveContext instead.")
+	//UE_DEPRECATED(5.0, "Use version that takes FObjectPreSaveContext instead.")
 	virtual void PreSave(const class ITargetPlatform* TargetPlatform) override;
 	/**
 	 * Presave function. Gets called once before an object gets serialized for saving. This function is necessary
diff --git a/Engine/Source/Runtime/CoreUObject/Public/UObject/ObjectRedirector.h b/Engine/Source/Runtime/CoreUObject/Public/UObject/ObjectRedirector.h
--- a/Engine/Source/Runtime/CoreUObject/Public/UObject/ObjectRedirector.h
+++ b/Engine/Source/Runtime/CoreUObject/Public/UObject/ObjectRedirector.h
@@ -21,7 +21,7 @@ class UObjectRedirector : public UObject
 	// Variables.
 	UObject*		DestinationObject;
 	// UObject interface.
-	UE_DEPRECATED(5.0, "Use version that takes FObjectPreSaveContext instead.")
+	//UE_DEPRECATED(5.0, "Use version that takes FObjectPreSaveContext instead.")
 	virtual void PreSave(const class ITargetPlatform* TargetPlatform) override;
 	virtual void PreSave(FObjectPreSaveContext ObjectSaveContext) override;
 	void Serialize(FArchive& Ar) override;
diff --git a/Engine/Source/Runtime/Engine/Classes/Animation/AnimMontage.h b/Engine/Source/Runtime/Engine/Classes/Animation/AnimMontage.h
--- a/Engine/Source/Runtime/Engine/Classes/Animation/AnimMontage.h
+++ b/Engine/Source/Runtime/Engine/Classes/Animation/AnimMontage.h
@@ -718,7 +718,7 @@ class UAnimMontage : public UAnimCompositeBase
 public:
 	//~ Begin UObject Interface
 	virtual void PostLoad() override;
-	UE_DEPRECATED(5.0, "Use version that takes FObjectPreSaveContext instead.")
+	//UE_DEPRECATED(5.0, "Use version that takes FObjectPreSaveContext instead.")
 	virtual void PreSave(const class ITargetPlatform* TargetPlatform) override;
 	virtual void PreSave(FObjectPreSaveContext ObjectSaveContext) override;
 
diff --git a/Engine/Source/Runtime/Engine/Classes/Animation/AnimNotifies/AnimNotify.h b/Engine/Source/Runtime/Engine/Classes/Animation/AnimNotifies/AnimNotify.h
--- a/Engine/Source/Runtime/Engine/Classes/Animation/AnimNotifies/AnimNotify.h
+++ b/Engine/Source/Runtime/Engine/Classes/Animation/AnimNotifies/AnimNotify.h
@@ -99,7 +99,7 @@ class ENGINE_API UAnimNotify : public UObject
 
 	/** UObject Interface */
 	virtual void PostLoad() override;
-	UE_DEPRECATED(5.0, "Use version that takes FObjectPreSaveContext instead.")
+	//UE_DEPRECATED(5.0, "Use version that takes FObjectPreSaveContext instead.")
 	virtual void PreSave(const class ITargetPlatform* TargetPlatform) override;
 	virtual void PreSave(FObjectPreSaveContext ObjectSaveContext) override;
 	/** End UObject Interface */
diff --git a/Engine/Source/Runtime/Engine/Classes/Animation/AnimSequence.h b/Engine/Source/Runtime/Engine/Classes/Animation/AnimSequence.h
--- a/Engine/Source/Runtime/Engine/Classes/Animation/AnimSequence.h
+++ b/Engine/Source/Runtime/Engine/Classes/Animation/AnimSequence.h
@@ -401,7 +401,7 @@ public:
 	virtual void Serialize(FArchive& Ar) override;
 	virtual void PostInitProperties() override;
 	virtual void PostLoad() override;
-	UE_DEPRECATED(5.0, "Use version that takes FObjectPreSaveContext instead.")
+	//UE_DEPRECATED(5.0, "Use version that takes FObjectPreSaveContext instead.")
 	virtual void PreSave(const class ITargetPlatform* TargetPlatform) override;
 	virtual void PreSave(FObjectPreSaveContext ObjectSaveContext) override;
 	virtual void GetPreloadDependencies(TArray<UObject*>& OutDeps) override;
diff --git a/Engine/Source/Runtime/Engine/Classes/Animation/AnimStreamable.h b/Engine/Source/Runtime/Engine/Classes/Animation/AnimStreamable.h
--- a/Engine/Source/Runtime/Engine/Classes/Animation/AnimStreamable.h
+++ b/Engine/Source/Runtime/Engine/Classes/Animation/AnimStreamable.h
@@ -200,7 +200,7 @@ public:
 	bool bUseNormalizedRootMotionScale;
 
 	//~ Begin UObject Interface
-	UE_DEPRECATED(5.0, "Use version that takes FObjectPreSaveContext instead.")
+	//UE_DEPRECATED(5.0, "Use version that takes FObjectPreSaveContext instead.")
 	virtual void PreSave(const class ITargetPlatform* TargetPlatform) override;
 	virtual void PreSave(FObjectPreSaveContext ObjectSaveContext) override;
 	virtual void Serialize(FArchive& Ar) override;
diff --git a/Engine/Source/Runtime/Engine/Classes/Animation/PoseAsset.h b/Engine/Source/Runtime/Engine/Classes/Animation/PoseAsset.h
--- a/Engine/Source/Runtime/Engine/Classes/Animation/PoseAsset.h
+++ b/Engine/Source/Runtime/Engine/Classes/Animation/PoseAsset.h
@@ -216,7 +216,7 @@ public:
 	virtual void PostLoad() override;
 	virtual bool IsPostLoadThreadSafe() const override;
 	virtual void Serialize(FArchive& Ar) override;
-	UE_DEPRECATED(5.0, "Use version that takes FObjectPreSaveContext instead.")
+	//UE_DEPRECATED(5.0, "Use version that takes FObjectPreSaveContext instead.")
 	virtual void PreSave(const class ITargetPlatform* TargetPlatform) override;
 	virtual void PreSave(FObjectPreSaveContext ObjectSaveContext) override;
 	virtual void GetAssetRegistryTags(TArray<FAssetRegistryTag>& OutTags) const override;
diff --git a/Engine/Source/Runtime/Engine/Classes/Camera/CameraAnim.h b/Engine/Source/Runtime/Engine/Classes/Camera/CameraAnim.h
--- a/Engine/Source/Runtime/Engine/Classes/Camera/CameraAnim.h
+++ b/Engine/Source/Runtime/Engine/Classes/Camera/CameraAnim.h
@@ -64,7 +64,7 @@ public:
 	float BasePostProcessBlendWeight;
 
 	//~ Begin UObject Interface
-	UE_DEPRECATED(5.0, "Use version that takes FObjectPreSaveContext instead.")
+	//UE_DEPRECATED(5.0, "Use version that takes FObjectPreSaveContext instead.")
 	virtual void PreSave(const class ITargetPlatform* TargetPlatform) override;
 	virtual void PreSave(FObjectPreSaveContext ObjectSaveContext) override;
 	virtual void PostLoad() override;
diff --git a/Engine/Source/Runtime/Engine/Classes/Components/LightComponent.h b/Engine/Source/Runtime/Engine/Classes/Components/LightComponent.h
--- a/Engine/Source/Runtime/Engine/Classes/Components/LightComponent.h
+++ b/Engine/Source/Runtime/Engine/Classes/Components/LightComponent.h
@@ -393,7 +393,7 @@ public:
 	virtual void Serialize(FArchive& Ar) override;
 	virtual void PostLoad() override;
 #if WITH_EDITOR
-	UE_DEPRECATED(5.0, "Use version that takes FObjectPreSaveContext instead.")
+	//UE_DEPRECATED(5.0, "Use version that takes FObjectPreSaveContext instead.")
 	virtual void PreSave(const class ITargetPlatform* TargetPlatform) override;
 	virtual void PreSave(FObjectPreSaveContext ObjectSaveContext) override;
 	virtual bool CanEditChange(const FProperty* InProperty) const override;
diff --git a/Engine/Source/Runtime/Engine/Classes/Components/StaticMeshComponent.h b/Engine/Source/Runtime/Engine/Classes/Components/StaticMeshComponent.h
--- a/Engine/Source/Runtime/Engine/Classes/Components/StaticMeshComponent.h
+++ b/Engine/Source/Runtime/Engine/Classes/Components/StaticMeshComponent.h
@@ -441,7 +441,7 @@ public:
 	virtual void UpdateBounds() override;
 #endif // WITH_EDITOR
 #if WITH_EDITORONLY_DATA
-	UE_DEPRECATED(5.0, "Use version that takes FObjectPreSaveContext instead.")
+	//UE_DEPRECATED(5.0, "Use version that takes FObjectPreSaveContext instead.")
 	virtual void PreSave(const class ITargetPlatform* TargetPlatform) override;
 	virtual void PreSave(FObjectPreSaveContext ObjectSaveContext) override;
 #endif // WITH_EDITORONLY_DATA
diff --git a/Engine/Source/Runtime/Engine/Classes/EdGraph/EdGraphNode.h b/Engine/Source/Runtime/Engine/Classes/EdGraph/EdGraphNode.h
--- a/Engine/Source/Runtime/Engine/Classes/EdGraph/EdGraphNode.h
+++ b/Engine/Source/Runtime/Engine/Classes/EdGraph/EdGraphNode.h
@@ -457,7 +457,7 @@ private:
 public:
 	// UObject interface
 	static void AddReferencedObjects(UObject* InThis, FReferenceCollector& Collector);
-	UE_DEPRECATED(5.0, "Use version that takes FObjectPreSaveContext instead.")
+	//UE_DEPRECATED(5.0, "Use version that takes FObjectPreSaveContext instead.")
 	virtual void PreSave(const class ITargetPlatform* TargetPlatform) override;
 	virtual void PreSave(FObjectPreSaveContext ObjectSaveContext) override;
 	virtual void PostLoad() override;
diff --git a/Engine/Source/Runtime/Engine/Classes/Engine/Blueprint.h b/Engine/Source/Runtime/Engine/Classes/Engine/Blueprint.h
--- a/Engine/Source/Runtime/Engine/Classes/Engine/Blueprint.h
+++ b/Engine/Source/Runtime/Engine/Classes/Engine/Blueprint.h
@@ -906,7 +906,7 @@ public:
 
 	//~ Begin UObject Interface
 #if WITH_EDITORONLY_DATA
-	UE_DEPRECATED(5.0, "Use version that takes FObjectPreSaveContext instead.")
+	//UE_DEPRECATED(5.0, "Use version that takes FObjectPreSaveContext instead.")
 	virtual void PreSave(const class ITargetPlatform* TargetPlatform) override;
 	virtual void PreSave(FObjectPreSaveContext ObjectSaveContext) override;
 #endif // WITH_EDITORONLY_DATA
diff --git a/Engine/Source/Runtime/Engine/Classes/Engine/DataAsset.h b/Engine/Source/Runtime/Engine/Classes/Engine/DataAsset.h
--- a/Engine/Source/Runtime/Engine/Classes/Engine/DataAsset.h
+++ b/Engine/Source/Runtime/Engine/Classes/Engine/DataAsset.h
@@ -50,7 +50,7 @@ public:
 	ENGINE_API virtual void UpdateAssetBundleData();
 
 	/** Updates AssetBundleData */
-	UE_DEPRECATED(5.0, "Use version that takes FObjectPreSaveContext instead.")
+	//UE_DEPRECATED(5.0, "Use version that takes FObjectPreSaveContext instead.")
 	ENGINE_API virtual void PreSave(const class ITargetPlatform* TargetPlatform) override;
 	ENGINE_API virtual void PreSave(FObjectPreSaveContext ObjectSaveContext) override;
 
diff --git a/Engine/Source/Runtime/Engine/Classes/Engine/HLODProxy.h b/Engine/Source/Runtime/Engine/Classes/Engine/HLODProxy.h
--- a/Engine/Source/Runtime/Engine/Classes/Engine/HLODProxy.h
+++ b/Engine/Source/Runtime/Engine/Classes/Engine/HLODProxy.h
@@ -58,7 +58,7 @@ public:
 	static uint32 GetCRC(UStaticMeshComponent* InComponent, uint32 InCRC = 0, const FTransform& TransformComponents = FTransform::Identity);
 
 	virtual void PostLoad() override;
-	UE_DEPRECATED(5.0, "Use version that takes FObjectPreSaveContext instead.")
+	//UE_DEPRECATED(5.0, "Use version that takes FObjectPreSaveContext instead.")
 	virtual void PreSave(const class ITargetPlatform* TargetPlatform) override;
 	virtual void PreSave(FObjectPreSaveContext ObjectSaveContext) override;
 
diff --git a/Engine/Source/Runtime/Engine/Classes/Engine/LODActor.h b/Engine/Source/Runtime/Engine/Classes/Engine/LODActor.h
--- a/Engine/Source/Runtime/Engine/Classes/Engine/LODActor.h
+++ b/Engine/Source/Runtime/Engine/Classes/Engine/LODActor.h
@@ -269,7 +269,7 @@ public:
 	virtual void PreEditChange(FProperty* PropertyThatWillChange) override;
 	virtual void PostEditChangeProperty(FPropertyChangedEvent& PropertyChangedEvent) override;
 
-	UE_DEPRECATED(5.0, "Use version that takes FObjectPreSaveContext instead.")
+	//UE_DEPRECATED(5.0, "Use version that takes FObjectPreSaveContext instead.")
 	virtual void PreSave(const class ITargetPlatform* TargetPlatform) override;
 	virtual void PreSave(FObjectPreSaveContext ObjectSaveContext) override;
 #endif // WITH_EDITOR	
diff --git a/Engine/Source/Runtime/Engine/Classes/Engine/Level.h b/Engine/Source/Runtime/Engine/Classes/Engine/Level.h
--- a/Engine/Source/Runtime/Engine/Classes/Engine/Level.h
+++ b/Engine/Source/Runtime/Engine/Classes/Engine/Level.h
@@ -788,7 +788,7 @@ public:
 	virtual bool LoadSubobject(const TCHAR* SubObjectPath, UObject*& OutObject, bool bOnlyTestExistence) override;
 #endif // WITH_EDITOR
 	virtual void PostLoad() override;
-	UE_DEPRECATED(5.0, "Use version that takes FObjectPreSaveContext instead.")
+	//UE_DEPRECATED(5.0, "Use version that takes FObjectPreSaveContext instead.")
 	virtual void PreSave(const class ITargetPlatform* TargetPlatform) override;
 	virtual void PreSave(FObjectPreSaveContext ObjectSaveContext) override;
 	virtual void PreDuplicate(FObjectDuplicationParameters& DupParams) override;
diff --git a/Engine/Source/Runtime/Engine/Classes/Engine/SkeletalMesh.h b/Engine/Source/Runtime/Engine/Classes/Engine/SkeletalMesh.h
--- a/Engine/Source/Runtime/Engine/Classes/Engine/SkeletalMesh.h
+++ b/Engine/Source/Runtime/Engine/Classes/Engine/SkeletalMesh.h
@@ -2517,7 +2517,7 @@ public:
 #endif // WITH_EDITOR
 	virtual void BeginDestroy() override;
 	virtual bool IsReadyForFinishDestroy() override;
-	UE_DEPRECATED(5.0, "Use version that takes FObjectPreSaveContext instead.")
+	//UE_DEPRECATED(5.0, "Use version that takes FObjectPreSaveContext instead.")
 	virtual void PreSave(const class ITargetPlatform* TargetPlatform) override;
 	virtual void PreSave(FObjectPreSaveContext ObjectSaveContext) override;
 	virtual void Serialize(FArchive& Ar) override;
diff --git a/Engine/Source/Runtime/Engine/Classes/Engine/Texture2D.h b/Engine/Source/Runtime/Engine/Classes/Engine/Texture2D.h
--- a/Engine/Source/Runtime/Engine/Classes/Engine/Texture2D.h
+++ b/Engine/Source/Runtime/Engine/Classes/Engine/Texture2D.h
@@ -108,7 +108,7 @@ public:
 	virtual void BeginDestroy() override;
 	virtual bool IsReadyForAsyncPostLoad() const override;
 	virtual void PostLoad() override;
-	UE_DEPRECATED(5.0, "Use version that takes FObjectPreSaveContext instead.")
+	//UE_DEPRECATED(5.0, "Use version that takes FObjectPreSaveContext instead.")
 	virtual void PreSave(const class ITargetPlatform* TargetPlatform) override;
 	virtual void PreSave(FObjectPreSaveContext ObjectSaveContext) override;
 	virtual void GetAssetRegistryTags(TArray<FAssetRegistryTag>& OutTags) const override;
diff --git a/Engine/Source/Runtime/Engine/Classes/Engine/World.h b/Engine/Source/Runtime/Engine/Classes/Engine/World.h
--- a/Engine/Source/Runtime/Engine/Classes/Engine/World.h
+++ b/Engine/Source/Runtime/Engine/Classes/Engine/World.h
@@ -2605,10 +2605,10 @@ public:
 	virtual bool IsReadyForFinishDestroy() override;
 	virtual void PostLoad() override;
 	virtual void PreDuplicate(FObjectDuplicationParameters& DupParams) override;
-	UE_DEPRECATED(5.0, "Use version that takes FObjectPreSaveContext instead.")
+	//UE_DEPRECATED(5.0, "Use version that takes FObjectPreSaveContext instead.")
 	virtual bool PreSaveRoot(const TCHAR* InFilename) override;
 	virtual void PreSaveRoot(FObjectPreSaveRootContext ObjectSaveContext) override;
-	UE_DEPRECATED(5.0, "Use version that takes FObjectPostSaveContext instead.")
+	//UE_DEPRECATED(5.0, "Use version that takes FObjectPostSaveContext instead.")
 	virtual void PostSaveRoot(bool bCleanupIsRequired) override;
 	virtual void PostSaveRoot(FObjectPostSaveRootContext ObjectSaveContext) override;
 	virtual UWorld* GetWorld() const override;
diff --git a/Engine/Source/Runtime/Engine/Classes/GameFramework/Actor.h b/Engine/Source/Runtime/Engine/Classes/GameFramework/Actor.h
--- a/Engine/Source/Runtime/Engine/Classes/GameFramework/Actor.h
+++ b/Engine/Source/Runtime/Engine/Classes/GameFramework/Actor.h
@@ -1843,13 +1843,13 @@ public:
 	virtual bool IsEditorOnly() const override;
 	virtual bool IsAsset() const override;
 
-	UE_DEPRECATED(5.0, "Use version that takes FObjectPreSaveContext instead.")
+	//UE_DEPRECATED(5.0, "Use version that takes FObjectPreSaveContext instead.")
 	virtual bool PreSaveRoot(const TCHAR* InFilename) override;
 	virtual void PreSaveRoot(FObjectPreSaveRootContext ObjectSaveContext) override;
-	UE_DEPRECATED(5.0, "Use version that takes FObjectPostSaveContext instead.")
+	//UE_DEPRECATED(5.0, "Use version that takes FObjectPostSaveContext instead.")
 	virtual void PostSaveRoot(bool bCleanupIsRequired) override;
 	virtual void PostSaveRoot(FObjectPostSaveRootContext ObjectSaveContext) override;
-	UE_DEPRECATED(5.0, "Use version that takes FObjectPreSaveContext instead.")
+	//UE_DEPRECATED(5.0, "Use version that takes FObjectPreSaveContext instead.")
 	virtual void PreSave(const ITargetPlatform* TargetPlatform) override;
 	virtual void PreSave(FObjectPreSaveContext ObjectSaveContext) override;
 
diff --git a/Engine/Source/Runtime/Engine/Classes/Materials/Material.h b/Engine/Source/Runtime/Engine/Classes/Materials/Material.h
--- a/Engine/Source/Runtime/Engine/Classes/Materials/Material.h
+++ b/Engine/Source/Runtime/Engine/Classes/Materials/Material.h
@@ -1087,7 +1087,7 @@ public:
 	//~ End UMaterialInterface Interface.
 
 	//~ Begin UObject Interface
-	UE_DEPRECATED(5.0, "Use version that takes FObjectPreSaveContext instead.")
+	//UE_DEPRECATED(5.0, "Use version that takes FObjectPreSaveContext instead.")
 	ENGINE_API virtual void PreSave(const class ITargetPlatform* TargetPlatform) override;
 	ENGINE_API virtual void PreSave(FObjectPreSaveContext ObjectSaveContext) override;
 	ENGINE_API virtual void PostInitProperties() override;
diff --git a/Engine/Source/Runtime/Engine/Classes/Materials/MaterialInstance.h b/Engine/Source/Runtime/Engine/Classes/Materials/MaterialInstance.h
--- a/Engine/Source/Runtime/Engine/Classes/Materials/MaterialInstance.h
+++ b/Engine/Source/Runtime/Engine/Classes/Materials/MaterialInstance.h
@@ -425,7 +425,7 @@ class UMaterialInstance : public UMaterialInterface
 	 */
 	class FMaterialInstanceResource* Resource;
 
-	UE_DEPRECATED(5.0, "Use version that takes FObjectPreSaveContext instead.")
+	//UE_DEPRECATED(5.0, "Use version that takes FObjectPreSaveContext instead.")
 	ENGINE_API virtual void PreSave(const class ITargetPlatform* TargetPlatform) override;
 	ENGINE_API virtual void PreSave(FObjectPreSaveContext ObjectSaveContext) override;
 	ENGINE_API virtual float GetTextureDensity(FName TextureName, const struct FMeshUVChannelInfo& UVChannelData) const override;
diff --git a/Engine/Source/Runtime/Engine/Classes/Particles/ParticleSystem.h b/Engine/Source/Runtime/Engine/Classes/Particles/ParticleSystem.h
--- a/Engine/Source/Runtime/Engine/Classes/Particles/ParticleSystem.h
+++ b/Engine/Source/Runtime/Engine/Classes/Particles/ParticleSystem.h
@@ -404,7 +404,7 @@ public:
 #if WITH_EDITOR
 	virtual void PostEditChangeProperty(FPropertyChangedEvent& PropertyChangedEvent) override;
 #endif // WITH_EDITOR
-	UE_DEPRECATED(5.0, "Use version that takes FObjectPreSaveContext instead.")
+	//UE_DEPRECATED(5.0, "Use version that takes FObjectPreSaveContext instead.")
 	virtual void PreSave(const class ITargetPlatform* TargetPlatform) override;
 	virtual void PreSave(FObjectPreSaveContext ObjectSaveContext) override;
 	virtual void PostLoad() override;
diff --git a/Engine/Source/Runtime/Engine/Public/WorldPartition/WorldPartitionRuntimeSpatialHash.h b/Engine/Source/Runtime/Engine/Public/WorldPartition/WorldPartitionRuntimeSpatialHash.h
--- a/Engine/Source/Runtime/Engine/Public/WorldPartition/WorldPartitionRuntimeSpatialHash.h
+++ b/Engine/Source/Runtime/Engine/Public/WorldPartition/WorldPartitionRuntimeSpatialHash.h
@@ -188,7 +188,7 @@ class ENGINE_API UWorldPartitionRuntimeSpatialHash : public UWorldPartitionRunti
 	GENERATED_UCLASS_BODY()
 
 public:
-	UE_DEPRECATED(5.0, "Use version that takes FObjectPreSaveContext instead.")
+	//UE_DEPRECATED(5.0, "Use version that takes FObjectPreSaveContext instead.")
 	virtual void PreSave(const class ITargetPlatform* TargetPlatform) override;
 	virtual void PreSave(FObjectPreSaveContext ObjectSaveContext) override;
 
diff --git a/Engine/Source/Runtime/Landscape/Classes/Landscape.h b/Engine/Source/Runtime/Landscape/Classes/Landscape.h
--- a/Engine/Source/Runtime/Landscape/Classes/Landscape.h
+++ b/Engine/Source/Runtime/Landscape/Classes/Landscape.h
@@ -220,7 +220,7 @@ public:
 	LANDSCAPE_API static void SplitHeightmap(ULandscapeComponent* Comp, ALandscapeProxy* TargetProxy = nullptr, class FMaterialUpdateContext* InOutUpdateContext = nullptr, TArray<class FComponentRecreateRenderStateContext>* InOutRecreateRenderStateContext = nullptr, bool InReregisterComponent = true);
 	
 	//~ Begin UObject Interface.
-	UE_DEPRECATED(5.0, "Use version that takes FObjectPreSaveContext instead.")
+	//UE_DEPRECATED(5.0, "Use version that takes FObjectPreSaveContext instead.")
 	virtual void PreSave(const class ITargetPlatform* TargetPlatform) override;
 	virtual void PreSave(FObjectPreSaveContext ObjectSaveContext) override;
 	virtual void PreEditChange(FProperty* PropertyThatWillChange) override;
diff --git a/Engine/Source/Runtime/Landscape/Classes/LandscapeHeightfieldCollisionComponent.h b/Engine/Source/Runtime/Landscape/Classes/LandscapeHeightfieldCollisionComponent.h
--- a/Engine/Source/Runtime/Landscape/Classes/LandscapeHeightfieldCollisionComponent.h
+++ b/Engine/Source/Runtime/Landscape/Classes/LandscapeHeightfieldCollisionComponent.h
@@ -222,7 +222,7 @@ public:
 	virtual void Serialize(FArchive& Ar) override;
 	virtual void BeginDestroy() override;
 	virtual void PostLoad() override;
-	UE_DEPRECATED(5.0, "Use version that takes FObjectPreSaveContext instead.")
+	//UE_DEPRECATED(5.0, "Use version that takes FObjectPreSaveContext instead.")
 	virtual void PreSave(const class ITargetPlatform* TargetPlatform) override;
 	virtual void PreSave(FObjectPreSaveContext ObjectSaveContext) override;
 #if WITH_EDITOR
diff --git a/Engine/Source/Runtime/Landscape/Classes/LandscapeProxy.h b/Engine/Source/Runtime/Landscape/Classes/LandscapeProxy.h
--- a/Engine/Source/Runtime/Landscape/Classes/LandscapeProxy.h
+++ b/Engine/Source/Runtime/Landscape/Classes/LandscapeProxy.h
@@ -897,7 +897,7 @@ public:
 #endif
 
 	//~ Begin UObject Interface.
-	UE_DEPRECATED(5.0, "Use version that takes FObjectPreSaveContext instead.")
+	//UE_DEPRECATED(5.0, "Use version that takes FObjectPreSaveContext instead.")
 	virtual void PreSave(const class ITargetPlatform* TargetPlatform) override;
 	virtual void PreSave(FObjectPreSaveContext ObjectSaveContext) override;
 	virtual void Serialize(FArchive& Ar) override;
diff --git a/Engine/Source/Runtime/MediaAssets/Public/BaseMediaSource.h b/Engine/Source/Runtime/MediaAssets/Public/BaseMediaSource.h
--- a/Engine/Source/Runtime/MediaAssets/Public/BaseMediaSource.h
+++ b/Engine/Source/Runtime/MediaAssets/Public/BaseMediaSource.h
@@ -38,7 +38,7 @@ private:
 public:
 
 	//~ UObject interface
-	UE_DEPRECATED(5.0, "Use version that takes FObjectPreSaveContext instead.")
+	//UE_DEPRECATED(5.0, "Use version that takes FObjectPreSaveContext instead.")
 	virtual void PreSave(const class ITargetPlatform* TargetPlatform);
 	virtual void PreSave(FObjectPreSaveContext ObjectSaveContext);
 	virtual void GetAssetRegistryTags(TArray<FAssetRegistryTag>& OutTags) const override;
diff --git a/Engine/Source/Runtime/MediaAssets/Public/PlatformMediaSource.h b/Engine/Source/Runtime/MediaAssets/Public/PlatformMediaSource.h
--- a/Engine/Source/Runtime/MediaAssets/Public/PlatformMediaSource.h
+++ b/Engine/Source/Runtime/MediaAssets/Public/PlatformMediaSource.h
@@ -38,7 +38,7 @@ private:
 
 public:
 	//~ UObject interface
-	UE_DEPRECATED(5.0, "Use version that takes FObjectPreSaveContext instead.")
+	//UE_DEPRECATED(5.0, "Use version that takes FObjectPreSaveContext instead.")
 	virtual void PreSave(const class ITargetPlatform* TargetPlatform);
 	virtual void PreSave(FObjectPreSaveContext ObjectSaveContext);
 	virtual void Serialize(FArchive& Ar) override;
diff --git a/Engine/Source/Runtime/MovieScene/Public/MovieScene.h b/Engine/Source/Runtime/MovieScene/Public/MovieScene.h
--- a/Engine/Source/Runtime/MovieScene/Public/MovieScene.h
+++ b/Engine/Source/Runtime/MovieScene/Public/MovieScene.h
@@ -1124,7 +1124,7 @@ protected:
 
 protected:
 
-	UE_DEPRECATED(5.0, "Use version that takes FObjectPreSaveContext instead.")
+	//UE_DEPRECATED(5.0, "Use version that takes FObjectPreSaveContext instead.")
 	virtual void PreSave(const class ITargetPlatform* TargetPlatform) override;
 	/** Called before this object is being deserialized. */
 	virtual void PreSave(FObjectPreSaveContext ObjectSaveContext) override;
diff --git a/Engine/Source/Runtime/MovieScene/Public/MovieSceneSequence.h b/Engine/Source/Runtime/MovieScene/Public/MovieSceneSequence.h
--- a/Engine/Source/Runtime/MovieScene/Public/MovieSceneSequence.h
+++ b/Engine/Source/Runtime/MovieScene/Public/MovieSceneSequence.h
@@ -209,7 +209,7 @@ public:
 public:
 
 	MOVIESCENE_API virtual void PostLoad() override;
-	UE_DEPRECATED(5.0, "Use version that takes FObjectPreSaveContext instead.")
+	//UE_DEPRECATED(5.0, "Use version that takes FObjectPreSaveContext instead.")
 	MOVIESCENE_API virtual void PreSave(const ITargetPlatform* TargetPlatform) override;
 	MOVIESCENE_API virtual void PreSave(FObjectPreSaveContext ObjectSaveContext) override;
 	MOVIESCENE_API virtual void BeginDestroy() override;
diff --git a/Engine/Source/Runtime/UMG/Public/Blueprint/UserWidget.h b/Engine/Source/Runtime/UMG/Public/Blueprint/UserWidget.h
--- a/Engine/Source/Runtime/UMG/Public/Blueprint/UserWidget.h
+++ b/Engine/Source/Runtime/UMG/Public/Blueprint/UserWidget.h
@@ -1092,7 +1092,7 @@ public:
 
 	//~ Begin UObject Interface
 	virtual bool IsAsset() const;
-	UE_DEPRECATED(5.0, "Use version that takes FObjectPreSaveContext instead.")
+	//UE_DEPRECATED(5.0, "Use version that takes FObjectPreSaveContext instead.")
 	virtual void PreSave(const class ITargetPlatform* TargetPlatform) override;
 	virtual void PreSave(FObjectPreSaveContext ObjectSaveContext) override;
 	//~ End UObject Interface
diff --git a/Engine/Source/Runtime/UMG/Public/Blueprint/WidgetTree.h b/Engine/Source/Runtime/UMG/Public/Blueprint/WidgetTree.h
--- a/Engine/Source/Runtime/UMG/Public/Blueprint/WidgetTree.h
+++ b/Engine/Source/Runtime/UMG/Public/Blueprint/WidgetTree.h
@@ -101,7 +101,7 @@ public:
 	}
 
 	// UObject interface
-	UE_DEPRECATED(5.0, "Use version that takes FObjectPreSaveContext instead.")
+	//UE_DEPRECATED(5.0, "Use version that takes FObjectPreSaveContext instead.")
 	virtual void PreSave(const class ITargetPlatform* TargetPlatform) override;
 	virtual void PreSave(FObjectPreSaveContext ObjectSaveContext) override;
 	virtual void PostLoad() override;
diff --git a/Engine/Source/Runtime/UMG/Public/Components/Widget.h b/Engine/Source/Runtime/UMG/Public/Components/Widget.h
--- a/Engine/Source/Runtime/UMG/Public/Components/Widget.h
+++ b/Engine/Source/Runtime/UMG/Public/Components/Widget.h
@@ -892,7 +892,7 @@ public:
 	// Begin UObject
 	virtual UWorld* GetWorld() const override;
 	virtual void FinishDestroy() override;
-	UE_DEPRECATED(5.0, "Use version that takes FObjectPreSaveContext instead.")
+	//UE_DEPRECATED(5.0, "Use version that takes FObjectPreSaveContext instead.")
 	virtual void PreSave(const class ITargetPlatform* TargetPlatform) override;
 	virtual void PreSave(FObjectPreSaveContext ObjectSaveContext) override;
 	// End UObject
diff --git a/Engine/Source/Runtime/UMG/Public/Slate/SlateVectorArtData.h b/Engine/Source/Runtime/UMG/Public/Slate/SlateVectorArtData.h
--- a/Engine/Source/Runtime/UMG/Public/Slate/SlateVectorArtData.h
+++ b/Engine/Source/Runtime/UMG/Public/Slate/SlateVectorArtData.h
@@ -109,7 +109,7 @@ public:
 
 private:
 	// ~ UObject Interface
-	UE_DEPRECATED(5.0, "Use version that takes FObjectPreSaveContext instead.")
+	//UE_DEPRECATED(5.0, "Use version that takes FObjectPreSaveContext instead.")
 	virtual void PreSave(const class ITargetPlatform* TargetPlatform) override;
 	virtual void PreSave(FObjectPreSaveContext ObjectSaveContext) override;
 	// ~ UObject Interface
