Change presave prototype:
Add PRAGMA_(DIS/EN)ABLE_DEPRECATION_WARNINGS around the deprecated function overrides and restore the deprecation attribute on the deprecated PreSave functions.
#rnx

diff --git a/Engine/Plugins/Enterprise/LidarPointCloud/Source/LidarPointCloudRuntime/Public/LidarPointCloud.h b/Engine/Plugins/Enterprise/LidarPointCloud/Source/LidarPointCloudRuntime/Public/LidarPointCloud.h
--- a/Engine/Plugins/Enterprise/LidarPointCloud/Source/LidarPointCloudRuntime/Public/LidarPointCloud.h
+++ b/Engine/Plugins/Enterprise/LidarPointCloud/Source/LidarPointCloudRuntime/Public/LidarPointCloud.h
@@ -160,8 +160,10 @@ public:
 	virtual void PostLoad() override;
 	virtual void GetAssetRegistryTags(TArray<FAssetRegistryTag>& OutTags) const override;
 	virtual void BeginDestroy() override;
-	//UE_DEPRECATED(5.0, "Use version that takes FObjectPreSaveContext instead.")
+	PRAGMA_DISABLE_DEPRECATION_WARNINGS // Suppress compiler warning on override of deprecated function
+	UE_DEPRECATED(5.0, "Use version that takes FObjectPreSaveContext instead.")
 	virtual void PreSave(const class ITargetPlatform* TargetPlatform) override;
+	PRAGMA_ENABLE_DEPRECATION_WARNINGS
 	virtual void PreSave(FObjectPreSaveContext ObjectSaveContext) override;
 #if WITH_EDITOR
 	virtual void ClearAllCachedCookedPlatformData() override;
diff --git a/Engine/Plugins/Experimental/Animation/ContextualAnimation/Source/ContextualAnimation/Public/ContextualAnimCompositeSceneAsset.h b/Engine/Plugins/Experimental/Animation/ContextualAnimation/Source/ContextualAnimation/Public/ContextualAnimCompositeSceneAsset.h
--- a/Engine/Plugins/Experimental/Animation/ContextualAnimation/Source/ContextualAnimation/Public/ContextualAnimCompositeSceneAsset.h
+++ b/Engine/Plugins/Experimental/Animation/ContextualAnimation/Source/ContextualAnimation/Public/ContextualAnimCompositeSceneAsset.h
@@ -82,8 +82,10 @@ public:
 
 	UContextualAnimCompositeSceneAsset(const FObjectInitializer& ObjectInitializer);
 
-	//UE_DEPRECATED(5.0, "Use version that takes FObjectPreSaveContext instead.")
+	PRAGMA_DISABLE_DEPRECATION_WARNINGS // Suppress compiler warning on override of deprecated function
+	UE_DEPRECATED(5.0, "Use version that takes FObjectPreSaveContext instead.")
 	virtual void PreSave(const class ITargetPlatform* TargetPlatform) override;
+	PRAGMA_ENABLE_DEPRECATION_WARNINGS
 	virtual void PreSave(FObjectPreSaveContext ObjectSaveContext) override;
 
 	bool QueryData(FContextualAnimQueryResult& OutResult, const FContextualAnimQueryParams& QueryParams, const FTransform& ToWorldTransform) const;
diff --git a/Engine/Plugins/Experimental/Animation/ContextualAnimation/Source/ContextualAnimation/Public/ContextualAnimSceneAsset.h b/Engine/Plugins/Experimental/Animation/ContextualAnimation/Source/ContextualAnimation/Public/ContextualAnimSceneAsset.h
--- a/Engine/Plugins/Experimental/Animation/ContextualAnimation/Source/ContextualAnimation/Public/ContextualAnimSceneAsset.h
+++ b/Engine/Plugins/Experimental/Animation/ContextualAnimation/Source/ContextualAnimation/Public/ContextualAnimSceneAsset.h
@@ -17,8 +17,10 @@ public:
 
 	UContextualAnimSceneAsset(const FObjectInitializer& ObjectInitializer);
 
-	//UE_DEPRECATED(5.0, "Use version that takes FObjectPreSaveContext instead.")
+	PRAGMA_DISABLE_DEPRECATION_WARNINGS // Suppress compiler warning on override of deprecated function
+	UE_DEPRECATED(5.0, "Use version that takes FObjectPreSaveContext instead.")
 	virtual void PreSave(const class ITargetPlatform* TargetPlatform) override;
+	PRAGMA_ENABLE_DEPRECATION_WARNINGS
 	virtual void PreSave(FObjectPreSaveContext ObjectSaveContext) override;
 
 	virtual UClass* GetPreviewActorClassForRole(const FName& Role) const override;
diff --git a/Engine/Plugins/Experimental/Animation/ContextualAnimation/Source/ContextualAnimation/Public/ContextualAnimSceneAssetBase.h b/Engine/Plugins/Experimental/Animation/ContextualAnimation/Source/ContextualAnimation/Public/ContextualAnimSceneAssetBase.h
--- a/Engine/Plugins/Experimental/Animation/ContextualAnimation/Source/ContextualAnimation/Public/ContextualAnimSceneAssetBase.h
+++ b/Engine/Plugins/Experimental/Animation/ContextualAnimation/Source/ContextualAnimation/Public/ContextualAnimSceneAssetBase.h
@@ -44,8 +44,10 @@ public:
 
 	UContextualAnimSceneAssetBase(const FObjectInitializer& ObjectInitializer);
 
-	//UE_DEPRECATED(5.0, "Use version that takes FObjectPreSaveContext instead.")
+	PRAGMA_DISABLE_DEPRECATION_WARNINGS // Suppress compiler warning on override of deprecated function
+	UE_DEPRECATED(5.0, "Use version that takes FObjectPreSaveContext instead.")
 	virtual void PreSave(const class ITargetPlatform* TargetPlatform) override;
+	PRAGMA_ENABLE_DEPRECATION_WARNINGS
 	virtual void PreSave(FObjectPreSaveContext ObjectSaveContext) override;
 
 	//@TODO: Make pure virtual
diff --git a/Engine/Plugins/Experimental/Animation/PoseSearch/Source/Runtime/Public/PoseSearch/PoseSearch.h b/Engine/Plugins/Experimental/Animation/PoseSearch/Source/Runtime/Public/PoseSearch/PoseSearch.h
--- a/Engine/Plugins/Experimental/Animation/PoseSearch/Source/Runtime/Public/PoseSearch/PoseSearch.h
+++ b/Engine/Plugins/Experimental/Animation/PoseSearch/Source/Runtime/Public/PoseSearch/PoseSearch.h
@@ -201,8 +201,10 @@ public:
 	int32 NumBones () const { return BoneIndices.Num(); }
 
 public: // UObject
-	//UE_DEPRECATED(5.0, "Use version that takes FObjectPreSaveContext instead.")
+	PRAGMA_DISABLE_DEPRECATION_WARNINGS // Suppress compiler warning on override of deprecated function
+	UE_DEPRECATED(5.0, "Use version that takes FObjectPreSaveContext instead.")
 	virtual void PreSave(const class ITargetPlatform* TargetPlatform) override;
+	PRAGMA_ENABLE_DEPRECATION_WARNINGS
 	virtual void PreSave(FObjectPreSaveContext ObjectSaveContext) override;
 	virtual void PostLoad() override;
 
@@ -260,8 +262,10 @@ public:
 	bool IsValidForSearch() const;
 
 public: // UObject
-	//UE_DEPRECATED(5.0, "Use version that takes FObjectPreSaveContext instead.")
+	PRAGMA_DISABLE_DEPRECATION_WARNINGS // Suppress compiler warning on override of deprecated function
+	UE_DEPRECATED(5.0, "Use version that takes FObjectPreSaveContext instead.")
 	virtual void PreSave(const class ITargetPlatform* TargetPlatform) override;
+	PRAGMA_ENABLE_DEPRECATION_WARNINGS
 	virtual void PreSave(FObjectPreSaveContext ObjectSaveContext) override;
 };
 
@@ -332,8 +336,10 @@ public:
 	bool IsValidForSearch() const;
 
 public: // UObject
-	//UE_DEPRECATED(5.0, "Use version that takes FObjectPreSaveContext instead.")
+	PRAGMA_DISABLE_DEPRECATION_WARNINGS // Suppress compiler warning on override of deprecated function
+	UE_DEPRECATED(5.0, "Use version that takes FObjectPreSaveContext instead.")
 	virtual void PreSave(const class ITargetPlatform* TargetPlatform) override;
+	PRAGMA_ENABLE_DEPRECATION_WARNINGS
 	virtual void PreSave(FObjectPreSaveContext ObjectSaveContext) override;
 };
 
diff --git a/Engine/Plugins/Experimental/CommonConversation/Source/CommonConversationRuntime/Public/ConversationDatabase.h b/Engine/Plugins/Experimental/CommonConversation/Source/CommonConversationRuntime/Public/ConversationDatabase.h
--- a/Engine/Plugins/Experimental/CommonConversation/Source/CommonConversationRuntime/Public/ConversationDatabase.h
+++ b/Engine/Plugins/Experimental/CommonConversation/Source/CommonConversationRuntime/Public/ConversationDatabase.h
@@ -77,8 +77,10 @@ public:
 #if WITH_EDITOR
 	virtual void PostDuplicate(bool bDuplicateForPIE) override;
 	virtual EDataValidationResult IsDataValid(TArray<FText>& ValidationErrors) override;
-	//UE_DEPRECATED(5.0, "Use version that takes FObjectPreSaveContext instead.")
+	PRAGMA_DISABLE_DEPRECATION_WARNINGS // Suppress compiler warning on override of deprecated function
+	UE_DEPRECATED(5.0, "Use version that takes FObjectPreSaveContext instead.")
 	virtual void PreSave(const class ITargetPlatform* TargetPlatform) override;
+	PRAGMA_ENABLE_DEPRECATION_WARNINGS
 	virtual void PreSave(FObjectPreSaveContext ObjectSaveContext) override;
 #endif
 
diff --git a/Engine/Plugins/Experimental/ControlRig/Source/ControlRigDeveloper/Public/ControlRigBlueprint.h b/Engine/Plugins/Experimental/ControlRig/Source/ControlRigDeveloper/Public/ControlRigBlueprint.h
--- a/Engine/Plugins/Experimental/ControlRig/Source/ControlRigDeveloper/Public/ControlRigBlueprint.h
+++ b/Engine/Plugins/Experimental/ControlRig/Source/ControlRigDeveloper/Public/ControlRigBlueprint.h
@@ -161,8 +161,10 @@ public:
 	virtual void GetTypeActions(FBlueprintActionDatabaseRegistrar& ActionRegistrar) const override;
 	virtual void GetInstanceActions(FBlueprintActionDatabaseRegistrar& ActionRegistrar) const override;
 	virtual void SetObjectBeingDebugged(UObject* NewObject) override;
-	//UE_DEPRECATED(5.0, "Use version that takes FObjectPreSaveContext instead.")
+	PRAGMA_DISABLE_DEPRECATION_WARNINGS // Suppress compiler warning on override of deprecated function
+	UE_DEPRECATED(5.0, "Use version that takes FObjectPreSaveContext instead.")
 	virtual void PreSave(const class ITargetPlatform* TargetPlatform) override;
+	PRAGMA_ENABLE_DEPRECATION_WARNINGS
 	virtual void PreSave(FObjectPreSaveContext ObjectSaveContext) override;
 	virtual void PostLoad() override;
 	virtual bool IsPostLoadThreadSafe() const override { return false; }
diff --git a/Engine/Plugins/Experimental/VirtualProductionUtilities/Source/VPUtilitiesEditor/Private/VPCameraBase.h b/Engine/Plugins/Experimental/VirtualProductionUtilities/Source/VPUtilitiesEditor/Private/VPCameraBase.h
--- a/Engine/Plugins/Experimental/VirtualProductionUtilities/Source/VPUtilitiesEditor/Private/VPCameraBase.h
+++ b/Engine/Plugins/Experimental/VirtualProductionUtilities/Source/VPUtilitiesEditor/Private/VPCameraBase.h
@@ -27,8 +27,10 @@ public:
 	virtual bool ShouldTickIfViewportsOnly() const override;
 	virtual void Tick(float DeltaSeconds) override;
 	virtual void BeginDestroy() override;
-	//UE_DEPRECATED(5.0, "Use version that takes FObjectPreSaveContext instead.")
+	PRAGMA_DISABLE_DEPRECATION_WARNINGS // Suppress compiler warning on override of deprecated function
+	UE_DEPRECATED(5.0, "Use version that takes FObjectPreSaveContext instead.")
 	virtual void PreSave(const class ITargetPlatform* TargetPlatform) override;
+	PRAGMA_ENABLE_DEPRECATION_WARNINGS
 	virtual void PreSave(FObjectPreSaveContext ObjectSaveContext) override;
 
 private:
diff --git a/Engine/Plugins/FX/Niagara/Source/Niagara/Classes/NiagaraScript.h b/Engine/Plugins/FX/Niagara/Source/Niagara/Classes/NiagaraScript.h
--- a/Engine/Plugins/FX/Niagara/Source/Niagara/Classes/NiagaraScript.h
+++ b/Engine/Plugins/FX/Niagara/Source/Niagara/Classes/NiagaraScript.h
@@ -610,8 +610,10 @@ public:
 #endif
 
 	//~ Begin UObject interface
-	//UE_DEPRECATED(5.0, "Use version that takes FObjectPreSaveContext instead.")
+	PRAGMA_DISABLE_DEPRECATION_WARNINGS // Suppress compiler warning on override of deprecated function
+	UE_DEPRECATED(5.0, "Use version that takes FObjectPreSaveContext instead.")
 	virtual void PreSave(const class ITargetPlatform* TargetPlatform) override;
+	PRAGMA_ENABLE_DEPRECATION_WARNINGS
 	virtual void PreSave(FObjectPreSaveContext ObjectSaveContext) override;
 	virtual void Serialize(FArchive& Ar)override;
 	virtual void PostLoad() override;
diff --git a/Engine/Plugins/FX/Niagara/Source/Niagara/Classes/NiagaraSystem.h b/Engine/Plugins/FX/Niagara/Source/Niagara/Classes/NiagaraSystem.h
--- a/Engine/Plugins/FX/Niagara/Source/Niagara/Classes/NiagaraSystem.h
+++ b/Engine/Plugins/FX/Niagara/Source/Niagara/Classes/NiagaraSystem.h
@@ -217,8 +217,10 @@ public:
 	virtual void Serialize(FArchive& Ar) override;
 	virtual void PostLoad() override; 
 	virtual void BeginDestroy() override;
-	//UE_DEPRECATED(5.0, "Use version that takes FObjectPreSaveContext instead.")
+	PRAGMA_DISABLE_DEPRECATION_WARNINGS // Suppress compiler warning on override of deprecated function
+	UE_DEPRECATED(5.0, "Use version that takes FObjectPreSaveContext instead.")
 	virtual void PreSave(const class ITargetPlatform* TargetPlatform) override;
+	PRAGMA_ENABLE_DEPRECATION_WARNINGS
 	virtual void PreSave(FObjectPreSaveContext ObjectSaveContext) override;
 #if WITH_EDITOR
 	virtual void PreEditChange(FProperty* PropertyThatWillChange)override;
diff --git a/Engine/Plugins/Runtime/DataRegistry/Source/DataRegistry/Public/DataRegistrySource_CurveTable.h b/Engine/Plugins/Runtime/DataRegistry/Source/DataRegistry/Public/DataRegistrySource_CurveTable.h
--- a/Engine/Plugins/Runtime/DataRegistry/Source/DataRegistry/Public/DataRegistrySource_CurveTable.h
+++ b/Engine/Plugins/Runtime/DataRegistry/Source/DataRegistry/Public/DataRegistrySource_CurveTable.h
@@ -72,8 +72,10 @@ protected:
 
 #if WITH_EDITOR
 	virtual void EditorRefreshSource();
-	//UE_DEPRECATED(5.0, "Use version that takes FObjectPreSaveContext instead.")
+	PRAGMA_DISABLE_DEPRECATION_WARNINGS // Suppress compiler warning on override of deprecated function
+	UE_DEPRECATED(5.0, "Use version that takes FObjectPreSaveContext instead.")
 	virtual void PreSave(const class ITargetPlatform* TargetPlatform) override;
+	PRAGMA_ENABLE_DEPRECATION_WARNINGS
 	virtual void PreSave(FObjectPreSaveContext ObjectSaveContext) override;
 
 #endif
diff --git a/Engine/Plugins/Runtime/DataRegistry/Source/DataRegistry/Public/DataRegistrySource_DataTable.h b/Engine/Plugins/Runtime/DataRegistry/Source/DataRegistry/Public/DataRegistrySource_DataTable.h
--- a/Engine/Plugins/Runtime/DataRegistry/Source/DataRegistry/Public/DataRegistrySource_DataTable.h
+++ b/Engine/Plugins/Runtime/DataRegistry/Source/DataRegistry/Public/DataRegistrySource_DataTable.h
@@ -86,8 +86,10 @@ protected:
 
 #if WITH_EDITOR
 	virtual void EditorRefreshSource();
-	//UE_DEPRECATED(5.0, "Use version that takes FObjectPreSaveContext instead.")
+	PRAGMA_DISABLE_DEPRECATION_WARNINGS // Suppress compiler warning on override of deprecated function
+	UE_DEPRECATED(5.0, "Use version that takes FObjectPreSaveContext instead.")
 	virtual void PreSave(const class ITargetPlatform* TargetPlatform) override;
+	PRAGMA_ENABLE_DEPRECATION_WARNINGS
 	virtual void PreSave(FObjectPreSaveContext ObjectSaveContext) override;
 
 #endif
diff --git a/Engine/Plugins/Runtime/GameplayAbilities/Source/GameplayAbilities/Public/GameplayEffect.h b/Engine/Plugins/Runtime/GameplayAbilities/Source/GameplayAbilities/Public/GameplayEffect.h
--- a/Engine/Plugins/Runtime/GameplayAbilities/Source/GameplayAbilities/Public/GameplayEffect.h
+++ b/Engine/Plugins/Runtime/GameplayAbilities/Source/GameplayAbilities/Public/GameplayEffect.h
@@ -1912,8 +1912,10 @@ public:
 	virtual void GetOwnedGameplayTags(FGameplayTagContainer& TagContainer) const override;
 	virtual void PostInitProperties() override;
 	virtual void PostLoad() override;
-	//UE_DEPRECATED(5.0, "Use version that takes FObjectPreSaveContext instead.")
+	PRAGMA_DISABLE_DEPRECATION_WARNINGS // Suppress compiler warning on override of deprecated function
+	UE_DEPRECATED(5.0, "Use version that takes FObjectPreSaveContext instead.")
 	virtual void PreSave(const class ITargetPlatform* TargetPlatform) override;
+	PRAGMA_ENABLE_DEPRECATION_WARNINGS
 	virtual void PreSave(FObjectPreSaveContext ObjectSaveContext) override;
 #if WITH_EDITOR
 	void PostEditChangeProperty(FPropertyChangedEvent& PropertyChangedEvent) override;
diff --git a/Engine/Plugins/Runtime/HairStrands/Source/HairStrandsCore/Public/GroomAsset.h b/Engine/Plugins/Runtime/HairStrands/Source/HairStrandsCore/Public/GroomAsset.h
--- a/Engine/Plugins/Runtime/HairStrands/Source/HairStrandsCore/Public/GroomAsset.h
+++ b/Engine/Plugins/Runtime/HairStrands/Source/HairStrandsCore/Public/GroomAsset.h
@@ -338,8 +338,10 @@ public:
 	virtual void BeginDestroy() override;
 	virtual void Serialize(FArchive& Ar) override;
 
-	//UE_DEPRECATED(5.0, "Use version that takes FObjectPreSaveContext instead.")
+	PRAGMA_DISABLE_DEPRECATION_WARNINGS // Suppress compiler warning on override of deprecated function
+	UE_DEPRECATED(5.0, "Use version that takes FObjectPreSaveContext instead.")
 	virtual void PreSave(const class ITargetPlatform* TargetPlatform) override;
+	PRAGMA_ENABLE_DEPRECATION_WARNINGS
 	virtual void PreSave(FObjectPreSaveContext ObjectSaveContext) override;
 
 #if WITH_EDITOR
diff --git a/Engine/Plugins/Runtime/HairStrands/Source/HairStrandsCore/Public/GroomBindingAsset.h b/Engine/Plugins/Runtime/HairStrands/Source/HairStrandsCore/Public/GroomBindingAsset.h
--- a/Engine/Plugins/Runtime/HairStrands/Source/HairStrandsCore/Public/GroomBindingAsset.h
+++ b/Engine/Plugins/Runtime/HairStrands/Source/HairStrandsCore/Public/GroomBindingAsset.h
@@ -121,11 +121,13 @@ public:
 	//~ Begin UObject Interface.
 	virtual void GetResourceSizeEx(FResourceSizeEx& CumulativeResourceSize) override;
 	virtual void PostLoad() override;
-	//UE_DEPRECATED(5.0, "Use version that takes FObjectPreSaveContext instead.")
+	PRAGMA_DISABLE_DEPRECATION_WARNINGS // Suppress compiler warning on override of deprecated function
+	UE_DEPRECATED(5.0, "Use version that takes FObjectPreSaveContext instead.")
 	virtual void PreSave(const class ITargetPlatform* TargetPlatform) override;
-	virtual void PreSave(FObjectPreSaveContext ObjectSaveContext) override;
-	//UE_DEPRECATED(5.0, "Use version that takes FObjectPostSaveContext instead.")
+	UE_DEPRECATED(5.0, "Use version that takes FObjectPostSaveRootContext instead.")
 	virtual void PostSaveRoot(bool bCleanupIsRequired) override;
+	PRAGMA_ENABLE_DEPRECATION_WARNINGS
+	virtual void PreSave(FObjectPreSaveContext ObjectSaveContext) override;
 	virtual void PostSaveRoot(FObjectPostSaveRootContext ObjectSaveContext) override;
 	virtual void BeginDestroy() override;
 	virtual void Serialize(FArchive& Ar) override;
diff --git a/Engine/Source/Editor/BlueprintGraph/Classes/K2Node_FunctionEntry.h b/Engine/Source/Editor/BlueprintGraph/Classes/K2Node_FunctionEntry.h
--- a/Engine/Source/Editor/BlueprintGraph/Classes/K2Node_FunctionEntry.h
+++ b/Engine/Source/Editor/BlueprintGraph/Classes/K2Node_FunctionEntry.h
@@ -34,8 +34,10 @@ class UK2Node_FunctionEntry : public UK2Node_FunctionTerminator
 
 	//~ Begin UObject Interface
 	virtual void Serialize(FArchive& Ar) override;
-	//UE_DEPRECATED(5.0, "Use version that takes FObjectPreSaveContext instead.")
+	PRAGMA_DISABLE_DEPRECATION_WARNINGS // Suppress compiler warning on override of deprecated function
+	UE_DEPRECATED(5.0, "Use version that takes FObjectPreSaveContext instead.")
 	virtual void PreSave(const class ITargetPlatform* TargetPlatform) override;
+	PRAGMA_ENABLE_DEPRECATION_WARNINGS
 	virtual void PreSave(FObjectPreSaveContext ObjectSaveContext) override;
 	virtual void PostLoad() override;
 	//~ End UObject Interface
diff --git a/Engine/Source/Editor/UMGEditor/Public/WidgetBlueprint.h b/Engine/Source/Editor/UMGEditor/Public/WidgetBlueprint.h
--- a/Engine/Source/Editor/UMGEditor/Public/WidgetBlueprint.h
+++ b/Engine/Source/Editor/UMGEditor/Public/WidgetBlueprint.h
@@ -257,8 +257,10 @@ public:
 	virtual void PostDuplicate(bool bDuplicateForPIE) override;
 
 #if WITH_EDITORONLY_DATA
-	//UE_DEPRECATED(5.0, "Use version that takes FObjectPreSaveContext instead.")
+	PRAGMA_DISABLE_DEPRECATION_WARNINGS // Suppress compiler warning on override of deprecated function
+	UE_DEPRECATED(5.0, "Use version that takes FObjectPreSaveContext instead.")
 	virtual void PreSave(const class ITargetPlatform* TargetPlatform) override;
+	PRAGMA_ENABLE_DEPRECATION_WARNINGS
 	virtual void PreSave(FObjectPreSaveContext ObjectSaveContext) override;
 #endif // WITH_EDITORONLY_DATA
 
diff --git a/Engine/Source/Editor/UnrealEd/Classes/Editor/EditorEngine.h b/Engine/Source/Editor/UnrealEd/Classes/Editor/EditorEngine.h
--- a/Engine/Source/Editor/UnrealEd/Classes/Editor/EditorEngine.h
+++ b/Engine/Source/Editor/UnrealEd/Classes/Editor/EditorEngine.h
@@ -2201,9 +2201,9 @@ public:
 	virtual bool InitializePhysicsSceneForSaveIfNecessary(UWorld* World, bool &bOutForceInitialized);
 	void CleanupPhysicsSceneThatWasInitializedForSave(UWorld* World, bool bForceInitialized);
 
-	//UE_DEPRECATED(5.0, "Use version that takes FObjectPreSaveContext instead.")
+	UE_DEPRECATED(5.0, "Use version that takes FObjectPreSaveContext instead.")
 	virtual void OnPreSaveWorld(uint32 SaveFlags, UWorld* World);
-	//UE_DEPRECATED(5.0, "Use version that takes FObjectPostSaveContext instead.")
+	UE_DEPRECATED(5.0, "Use version that takes FObjectPostSaveContext instead.")
 	virtual void OnPostSaveWorld(uint32 SaveFlags, UWorld* World, uint32 OriginalPackageFlags, bool bSuccess);
 
 	/** Invoked before a UWorld is saved to update editor systems */
diff --git a/Engine/Source/Runtime/AIModule/Classes/EnvironmentQuery/Tests/EnvQueryTest_GameplayTags.h b/Engine/Source/Runtime/AIModule/Classes/EnvironmentQuery/Tests/EnvQueryTest_GameplayTags.h
--- a/Engine/Source/Runtime/AIModule/Classes/EnvironmentQuery/Tests/EnvQueryTest_GameplayTags.h
+++ b/Engine/Source/Runtime/AIModule/Classes/EnvironmentQuery/Tests/EnvQueryTest_GameplayTags.h
@@ -31,8 +31,10 @@ protected:
 
 	bool SatisfiesTest(IGameplayTagAssetInterface* ItemGameplayTagAssetInterface) const;
 
-	//UE_DEPRECATED(5.0, "Use version that takes FObjectPreSaveContext instead.")
+	PRAGMA_DISABLE_DEPRECATION_WARNINGS // Suppress compiler warning on override of deprecated function
+	UE_DEPRECATED(5.0, "Use version that takes FObjectPreSaveContext instead.")
 	virtual void PreSave(const class ITargetPlatform* TargetPlatform) override;
+	PRAGMA_ENABLE_DEPRECATION_WARNINGS
 	/**
 	 * Presave function. Gets called once before an object gets serialized for saving. This function is necessary
 	 * for save time computation as Serialize gets called three times per object from within SavePackage.
diff --git a/Engine/Source/Runtime/CoreUObject/Public/UObject/Object.h b/Engine/Source/Runtime/CoreUObject/Public/UObject/Object.h
--- a/Engine/Source/Runtime/CoreUObject/Public/UObject/Object.h
+++ b/Engine/Source/Runtime/CoreUObject/Public/UObject/Object.h
@@ -211,7 +211,7 @@ public:
 	{
 	}
 
-	//UE_DEPRECATED(5.0, "Use version that takes FObjectPreSaveContext instead.")
+	UE_DEPRECATED(5.0, "Use version that takes FObjectPreSaveContext instead.")
 	virtual bool PreSaveRoot(const TCHAR* Filename);
 
 	/**
@@ -223,7 +223,7 @@ public:
 	 */
 	virtual void PreSaveRoot(FObjectPreSaveRootContext ObjectSaveContext);
 
-	//UE_DEPRECATED(5.0, "Use version that takes FObjectPostSaveContext instead.")
+	UE_DEPRECATED(5.0, "Use version that takes FObjectPostSaveContext instead.")
 	virtual void PostSaveRoot(bool bCleanupIsRequired);
 
 	/**
@@ -234,7 +234,7 @@ public:
 	 */
 	virtual void PostSaveRoot(FObjectPostSaveRootContext ObjectSaveContext);
 
-	//UE_DEPRECATED(5.0, "Use version that takes FObjectPreSaveContext instead.")
+	UE_DEPRECATED(5.0, "Use version that takes FObjectPreSaveContext instead.")
 	virtual void PreSave(const class ITargetPlatform* TargetPlatform);
 
 	/**
diff --git a/Engine/Source/Runtime/CoreUObject/Public/UObject/ObjectRedirector.h b/Engine/Source/Runtime/CoreUObject/Public/UObject/ObjectRedirector.h
--- a/Engine/Source/Runtime/CoreUObject/Public/UObject/ObjectRedirector.h
+++ b/Engine/Source/Runtime/CoreUObject/Public/UObject/ObjectRedirector.h
@@ -21,8 +21,10 @@ class UObjectRedirector : public UObject
 	// Variables.
 	UObject*		DestinationObject;
 	// UObject interface.
-	//UE_DEPRECATED(5.0, "Use version that takes FObjectPreSaveContext instead.")
+	PRAGMA_DISABLE_DEPRECATION_WARNINGS // Suppress compiler warning on override of deprecated function
+	UE_DEPRECATED(5.0, "Use version that takes FObjectPreSaveContext instead.")
 	virtual void PreSave(const class ITargetPlatform* TargetPlatform) override;
+	PRAGMA_ENABLE_DEPRECATION_WARNINGS
 	virtual void PreSave(FObjectPreSaveContext ObjectSaveContext) override;
 	void Serialize(FArchive& Ar) override;
 	void Serialize(FStructuredArchive::FRecord Record) override;
diff --git a/Engine/Source/Runtime/Engine/Classes/Animation/AnimMontage.h b/Engine/Source/Runtime/Engine/Classes/Animation/AnimMontage.h
--- a/Engine/Source/Runtime/Engine/Classes/Animation/AnimMontage.h
+++ b/Engine/Source/Runtime/Engine/Classes/Animation/AnimMontage.h
@@ -718,8 +718,10 @@ class UAnimMontage : public UAnimCompositeBase
 public:
 	//~ Begin UObject Interface
 	virtual void PostLoad() override;
-	//UE_DEPRECATED(5.0, "Use version that takes FObjectPreSaveContext instead.")
+	PRAGMA_DISABLE_DEPRECATION_WARNINGS // Suppress compiler warning on override of deprecated function
+	UE_DEPRECATED(5.0, "Use version that takes FObjectPreSaveContext instead.")
 	virtual void PreSave(const class ITargetPlatform* TargetPlatform) override;
+	PRAGMA_ENABLE_DEPRECATION_WARNINGS
 	virtual void PreSave(FObjectPreSaveContext ObjectSaveContext) override;
 
 	// Gets the sequence length of the montage by calculating it from the lengths of the segments in the montage
diff --git a/Engine/Source/Runtime/Engine/Classes/Animation/AnimNotifies/AnimNotify.h b/Engine/Source/Runtime/Engine/Classes/Animation/AnimNotifies/AnimNotify.h
--- a/Engine/Source/Runtime/Engine/Classes/Animation/AnimNotifies/AnimNotify.h
+++ b/Engine/Source/Runtime/Engine/Classes/Animation/AnimNotifies/AnimNotify.h
@@ -99,8 +99,10 @@ class ENGINE_API UAnimNotify : public UObject
 
 	/** UObject Interface */
 	virtual void PostLoad() override;
-	//UE_DEPRECATED(5.0, "Use version that takes FObjectPreSaveContext instead.")
+	PRAGMA_DISABLE_DEPRECATION_WARNINGS // Suppress compiler warning on override of deprecated function
+	UE_DEPRECATED(5.0, "Use version that takes FObjectPreSaveContext instead.")
 	virtual void PreSave(const class ITargetPlatform* TargetPlatform) override;
+	PRAGMA_ENABLE_DEPRECATION_WARNINGS
 	virtual void PreSave(FObjectPreSaveContext ObjectSaveContext) override;
 	/** End UObject Interface */
 
diff --git a/Engine/Source/Runtime/Engine/Classes/Animation/AnimSequence.h b/Engine/Source/Runtime/Engine/Classes/Animation/AnimSequence.h
--- a/Engine/Source/Runtime/Engine/Classes/Animation/AnimSequence.h
+++ b/Engine/Source/Runtime/Engine/Classes/Animation/AnimSequence.h
@@ -400,8 +400,10 @@ public:
 	virtual void Serialize(FArchive& Ar) override;
 	virtual void PostInitProperties() override;
 	virtual void PostLoad() override;
-	//UE_DEPRECATED(5.0, "Use version that takes FObjectPreSaveContext instead.")
+	PRAGMA_DISABLE_DEPRECATION_WARNINGS // Suppress compiler warning on override of deprecated function
+	UE_DEPRECATED(5.0, "Use version that takes FObjectPreSaveContext instead.")
 	virtual void PreSave(const class ITargetPlatform* TargetPlatform) override;
+	PRAGMA_ENABLE_DEPRECATION_WARNINGS
 	virtual void PreSave(FObjectPreSaveContext ObjectSaveContext) override;
 	virtual void GetPreloadDependencies(TArray<UObject*>& OutDeps) override;
 #if WITH_EDITOR
diff --git a/Engine/Source/Runtime/Engine/Classes/Animation/AnimStreamable.h b/Engine/Source/Runtime/Engine/Classes/Animation/AnimStreamable.h
--- a/Engine/Source/Runtime/Engine/Classes/Animation/AnimStreamable.h
+++ b/Engine/Source/Runtime/Engine/Classes/Animation/AnimStreamable.h
@@ -200,8 +200,10 @@ public:
 	bool bUseNormalizedRootMotionScale;
 
 	//~ Begin UObject Interface
-	//UE_DEPRECATED(5.0, "Use version that takes FObjectPreSaveContext instead.")
+	PRAGMA_DISABLE_DEPRECATION_WARNINGS // Suppress compiler warning on override of deprecated function
+	UE_DEPRECATED(5.0, "Use version that takes FObjectPreSaveContext instead.")
 	virtual void PreSave(const class ITargetPlatform* TargetPlatform) override;
+	PRAGMA_ENABLE_DEPRECATION_WARNINGS
 	virtual void PreSave(FObjectPreSaveContext ObjectSaveContext) override;
 	virtual void Serialize(FArchive& Ar) override;
 	virtual void PostLoad() override;
diff --git a/Engine/Source/Runtime/Engine/Classes/Animation/PoseAsset.h b/Engine/Source/Runtime/Engine/Classes/Animation/PoseAsset.h
--- a/Engine/Source/Runtime/Engine/Classes/Animation/PoseAsset.h
+++ b/Engine/Source/Runtime/Engine/Classes/Animation/PoseAsset.h
@@ -216,8 +216,10 @@ public:
 	virtual void PostLoad() override;
 	virtual bool IsPostLoadThreadSafe() const override;
 	virtual void Serialize(FArchive& Ar) override;
-	//UE_DEPRECATED(5.0, "Use version that takes FObjectPreSaveContext instead.")
+	PRAGMA_DISABLE_DEPRECATION_WARNINGS // Suppress compiler warning on override of deprecated function
+	UE_DEPRECATED(5.0, "Use version that takes FObjectPreSaveContext instead.")
 	virtual void PreSave(const class ITargetPlatform* TargetPlatform) override;
+	PRAGMA_ENABLE_DEPRECATION_WARNINGS
 	virtual void PreSave(FObjectPreSaveContext ObjectSaveContext) override;
 	virtual void GetAssetRegistryTags(TArray<FAssetRegistryTag>& OutTags) const override;
 	//End UObject Interface
diff --git a/Engine/Source/Runtime/Engine/Classes/Camera/CameraAnim.h b/Engine/Source/Runtime/Engine/Classes/Camera/CameraAnim.h
--- a/Engine/Source/Runtime/Engine/Classes/Camera/CameraAnim.h
+++ b/Engine/Source/Runtime/Engine/Classes/Camera/CameraAnim.h
@@ -64,8 +64,10 @@ public:
 	float BasePostProcessBlendWeight;
 
 	//~ Begin UObject Interface
-	//UE_DEPRECATED(5.0, "Use version that takes FObjectPreSaveContext instead.")
+	PRAGMA_DISABLE_DEPRECATION_WARNINGS // Suppress compiler warning on override of deprecated function
+	UE_DEPRECATED(5.0, "Use version that takes FObjectPreSaveContext instead.")
 	virtual void PreSave(const class ITargetPlatform* TargetPlatform) override;
+	PRAGMA_ENABLE_DEPRECATION_WARNINGS
 	virtual void PreSave(FObjectPreSaveContext ObjectSaveContext) override;
 	virtual void PostLoad() override;
 	//~ End UObject Interface
diff --git a/Engine/Source/Runtime/Engine/Classes/Components/LightComponent.h b/Engine/Source/Runtime/Engine/Classes/Components/LightComponent.h
--- a/Engine/Source/Runtime/Engine/Classes/Components/LightComponent.h
+++ b/Engine/Source/Runtime/Engine/Classes/Components/LightComponent.h
@@ -393,8 +393,10 @@ public:
 	virtual void Serialize(FArchive& Ar) override;
 	virtual void PostLoad() override;
 #if WITH_EDITOR
-	//UE_DEPRECATED(5.0, "Use version that takes FObjectPreSaveContext instead.")
+	PRAGMA_DISABLE_DEPRECATION_WARNINGS // Suppress compiler warning on override of deprecated function
+	UE_DEPRECATED(5.0, "Use version that takes FObjectPreSaveContext instead.")
 	virtual void PreSave(const class ITargetPlatform* TargetPlatform) override;
+	PRAGMA_ENABLE_DEPRECATION_WARNINGS
 	virtual void PreSave(FObjectPreSaveContext ObjectSaveContext) override;
 	virtual bool CanEditChange(const FProperty* InProperty) const override;
 	virtual void PostEditChangeProperty(FPropertyChangedEvent& PropertyChangedEvent) override;
diff --git a/Engine/Source/Runtime/Engine/Classes/Components/StaticMeshComponent.h b/Engine/Source/Runtime/Engine/Classes/Components/StaticMeshComponent.h
--- a/Engine/Source/Runtime/Engine/Classes/Components/StaticMeshComponent.h
+++ b/Engine/Source/Runtime/Engine/Classes/Components/StaticMeshComponent.h
@@ -441,8 +441,10 @@ public:
 	virtual void UpdateBounds() override;
 #endif // WITH_EDITOR
 #if WITH_EDITORONLY_DATA
-	//UE_DEPRECATED(5.0, "Use version that takes FObjectPreSaveContext instead.")
+	PRAGMA_DISABLE_DEPRECATION_WARNINGS // Suppress compiler warning on override of deprecated function
+	UE_DEPRECATED(5.0, "Use version that takes FObjectPreSaveContext instead.")
 	virtual void PreSave(const class ITargetPlatform* TargetPlatform) override;
+	PRAGMA_ENABLE_DEPRECATION_WARNINGS
 	virtual void PreSave(FObjectPreSaveContext ObjectSaveContext) override;
 #endif // WITH_EDITORONLY_DATA
 	virtual void PostLoad() override;
diff --git a/Engine/Source/Runtime/Engine/Classes/EdGraph/EdGraphNode.h b/Engine/Source/Runtime/Engine/Classes/EdGraph/EdGraphNode.h
--- a/Engine/Source/Runtime/Engine/Classes/EdGraph/EdGraphNode.h
+++ b/Engine/Source/Runtime/Engine/Classes/EdGraph/EdGraphNode.h
@@ -457,8 +457,10 @@ private:
 public:
 	// UObject interface
 	static void AddReferencedObjects(UObject* InThis, FReferenceCollector& Collector);
-	//UE_DEPRECATED(5.0, "Use version that takes FObjectPreSaveContext instead.")
+	PRAGMA_DISABLE_DEPRECATION_WARNINGS // Suppress compiler warning on override of deprecated function
+	UE_DEPRECATED(5.0, "Use version that takes FObjectPreSaveContext instead.")
 	virtual void PreSave(const class ITargetPlatform* TargetPlatform) override;
+	PRAGMA_ENABLE_DEPRECATION_WARNINGS
 	virtual void PreSave(FObjectPreSaveContext ObjectSaveContext) override;
 	virtual void PostLoad() override;
 	virtual void PostEditChangeProperty(FPropertyChangedEvent& PropertyChangedEvent) override;
diff --git a/Engine/Source/Runtime/Engine/Classes/Engine/Blueprint.h b/Engine/Source/Runtime/Engine/Classes/Engine/Blueprint.h
--- a/Engine/Source/Runtime/Engine/Classes/Engine/Blueprint.h
+++ b/Engine/Source/Runtime/Engine/Classes/Engine/Blueprint.h
@@ -906,8 +906,10 @@ public:
 
 	//~ Begin UObject Interface
 #if WITH_EDITORONLY_DATA
-	//UE_DEPRECATED(5.0, "Use version that takes FObjectPreSaveContext instead.")
+	PRAGMA_DISABLE_DEPRECATION_WARNINGS // Suppress compiler warning on override of deprecated function
+	UE_DEPRECATED(5.0, "Use version that takes FObjectPreSaveContext instead.")
 	virtual void PreSave(const class ITargetPlatform* TargetPlatform) override;
+	PRAGMA_ENABLE_DEPRECATION_WARNINGS
 	virtual void PreSave(FObjectPreSaveContext ObjectSaveContext) override;
 #endif // WITH_EDITORONLY_DATA
 	virtual void Serialize(FArchive& Ar) override;
diff --git a/Engine/Source/Runtime/Engine/Classes/Engine/DataAsset.h b/Engine/Source/Runtime/Engine/Classes/Engine/DataAsset.h
--- a/Engine/Source/Runtime/Engine/Classes/Engine/DataAsset.h
+++ b/Engine/Source/Runtime/Engine/Classes/Engine/DataAsset.h
@@ -50,8 +50,10 @@ public:
 	ENGINE_API virtual void UpdateAssetBundleData();
 
 	/** Updates AssetBundleData */
-	//UE_DEPRECATED(5.0, "Use version that takes FObjectPreSaveContext instead.")
+	PRAGMA_DISABLE_DEPRECATION_WARNINGS // Suppress compiler warning on override of deprecated function
+	UE_DEPRECATED(5.0, "Use version that takes FObjectPreSaveContext instead.")
 	ENGINE_API virtual void PreSave(const class ITargetPlatform* TargetPlatform) override;
+	PRAGMA_ENABLE_DEPRECATION_WARNINGS
 	ENGINE_API virtual void PreSave(FObjectPreSaveContext ObjectSaveContext) override;
 
 protected:
diff --git a/Engine/Source/Runtime/Engine/Classes/Engine/HLODProxy.h b/Engine/Source/Runtime/Engine/Classes/Engine/HLODProxy.h
--- a/Engine/Source/Runtime/Engine/Classes/Engine/HLODProxy.h
+++ b/Engine/Source/Runtime/Engine/Classes/Engine/HLODProxy.h
@@ -58,8 +58,10 @@ public:
 	static uint32 GetCRC(UStaticMeshComponent* InComponent, uint32 InCRC = 0, const FTransform& TransformComponents = FTransform::Identity);
 
 	virtual void PostLoad() override;
-	//UE_DEPRECATED(5.0, "Use version that takes FObjectPreSaveContext instead.")
+	PRAGMA_DISABLE_DEPRECATION_WARNINGS // Suppress compiler warning on override of deprecated function
+	UE_DEPRECATED(5.0, "Use version that takes FObjectPreSaveContext instead.")
 	virtual void PreSave(const class ITargetPlatform* TargetPlatform) override;
+	PRAGMA_ENABLE_DEPRECATION_WARNINGS
 	virtual void PreSave(FObjectPreSaveContext ObjectSaveContext) override;
 
 	/** Returns true if proxy doesn't contain any mesh entry. */
diff --git a/Engine/Source/Runtime/Engine/Classes/Engine/LODActor.h b/Engine/Source/Runtime/Engine/Classes/Engine/LODActor.h
--- a/Engine/Source/Runtime/Engine/Classes/Engine/LODActor.h
+++ b/Engine/Source/Runtime/Engine/Classes/Engine/LODActor.h
@@ -269,8 +269,10 @@ public:
 	virtual void PreEditChange(FProperty* PropertyThatWillChange) override;
 	virtual void PostEditChangeProperty(FPropertyChangedEvent& PropertyChangedEvent) override;
 
-	//UE_DEPRECATED(5.0, "Use version that takes FObjectPreSaveContext instead.")
+	PRAGMA_DISABLE_DEPRECATION_WARNINGS // Suppress compiler warning on override of deprecated function
+	UE_DEPRECATED(5.0, "Use version that takes FObjectPreSaveContext instead.")
 	virtual void PreSave(const class ITargetPlatform* TargetPlatform) override;
+	PRAGMA_ENABLE_DEPRECATION_WARNINGS
 	virtual void PreSave(FObjectPreSaveContext ObjectSaveContext) override;
 #endif // WITH_EDITOR	
 
diff --git a/Engine/Source/Runtime/Engine/Classes/Engine/Level.h b/Engine/Source/Runtime/Engine/Classes/Engine/Level.h
--- a/Engine/Source/Runtime/Engine/Classes/Engine/Level.h
+++ b/Engine/Source/Runtime/Engine/Classes/Engine/Level.h
@@ -788,8 +788,10 @@ public:
 #endif // WITH_EDITOR
 	virtual bool ResolveSubobject(const TCHAR* SubObjectPath, UObject*& OutObject, bool bLoadIfExists) override;
 	virtual void PostLoad() override;
-	//UE_DEPRECATED(5.0, "Use version that takes FObjectPreSaveContext instead.")
+	PRAGMA_DISABLE_DEPRECATION_WARNINGS // Suppress compiler warning on override of deprecated function
+	UE_DEPRECATED(5.0, "Use version that takes FObjectPreSaveContext instead.")
 	virtual void PreSave(const class ITargetPlatform* TargetPlatform) override;
+	PRAGMA_ENABLE_DEPRECATION_WARNINGS
 	virtual void PreSave(FObjectPreSaveContext ObjectSaveContext) override;
 	virtual void PreDuplicate(FObjectDuplicationParameters& DupParams) override;
 	virtual void PostDuplicate(bool bDuplicateForPIE) override;
diff --git a/Engine/Source/Runtime/Engine/Classes/Engine/SkeletalMesh.h b/Engine/Source/Runtime/Engine/Classes/Engine/SkeletalMesh.h
--- a/Engine/Source/Runtime/Engine/Classes/Engine/SkeletalMesh.h
+++ b/Engine/Source/Runtime/Engine/Classes/Engine/SkeletalMesh.h
@@ -2522,8 +2522,10 @@ public:
 #endif // WITH_EDITOR
 	virtual void BeginDestroy() override;
 	virtual bool IsReadyForFinishDestroy() override;
-	//UE_DEPRECATED(5.0, "Use version that takes FObjectPreSaveContext instead.")
+	PRAGMA_DISABLE_DEPRECATION_WARNINGS // Suppress compiler warning on override of deprecated function
+	UE_DEPRECATED(5.0, "Use version that takes FObjectPreSaveContext instead.")
 	virtual void PreSave(const class ITargetPlatform* TargetPlatform) override;
+	PRAGMA_ENABLE_DEPRECATION_WARNINGS
 	virtual void PreSave(FObjectPreSaveContext ObjectSaveContext) override;
 	virtual void Serialize(FArchive& Ar) override;
 	virtual void PostInitProperties() override;
diff --git a/Engine/Source/Runtime/Engine/Classes/Engine/Texture.h b/Engine/Source/Runtime/Engine/Classes/Engine/Texture.h
--- a/Engine/Source/Runtime/Engine/Classes/Engine/Texture.h
+++ b/Engine/Source/Runtime/Engine/Classes/Engine/Texture.h
@@ -1084,8 +1084,10 @@ public:
 	ENGINE_API virtual void Serialize(FArchive& Ar) override;
 	ENGINE_API virtual void PostInitProperties() override;
 	ENGINE_API virtual void PostLoad() override;
-	//UE_DEPRECATED(5.0, "Use version that takes FObjectPreSaveContext instead.")
+	PRAGMA_DISABLE_DEPRECATION_WARNINGS // Suppress compiler warning on override of deprecated function
+	UE_DEPRECATED(5.0, "Use version that takes FObjectPreSaveContext instead.")
 	ENGINE_API virtual void PreSave(const class ITargetPlatform* TargetPlatform) override;
+	PRAGMA_ENABLE_DEPRECATION_WARNINGS
 	ENGINE_API virtual void PreSave(FObjectPreSaveContext ObjectSaveContext) override;
 	ENGINE_API virtual void BeginDestroy() override;
 	ENGINE_API virtual bool IsReadyForFinishDestroy() override;
diff --git a/Engine/Source/Runtime/Engine/Classes/Engine/Texture2D.h b/Engine/Source/Runtime/Engine/Classes/Engine/Texture2D.h
--- a/Engine/Source/Runtime/Engine/Classes/Engine/Texture2D.h
+++ b/Engine/Source/Runtime/Engine/Classes/Engine/Texture2D.h
@@ -108,8 +108,10 @@ public:
 	virtual void BeginDestroy() override;
 	virtual bool IsReadyForAsyncPostLoad() const override;
 	virtual void PostLoad() override;
-	//UE_DEPRECATED(5.0, "Use version that takes FObjectPreSaveContext instead.")
+	PRAGMA_DISABLE_DEPRECATION_WARNINGS // Suppress compiler warning on override of deprecated function
+	UE_DEPRECATED(5.0, "Use version that takes FObjectPreSaveContext instead.")
 	virtual void PreSave(const class ITargetPlatform* TargetPlatform) override;
+	PRAGMA_ENABLE_DEPRECATION_WARNINGS
 	virtual void PreSave(FObjectPreSaveContext ObjectSaveContext) override;
 	virtual void GetAssetRegistryTags(TArray<FAssetRegistryTag>& OutTags) const override;
 	virtual FString GetDesc() override;
diff --git a/Engine/Source/Runtime/Engine/Classes/Engine/World.h b/Engine/Source/Runtime/Engine/Classes/Engine/World.h
--- a/Engine/Source/Runtime/Engine/Classes/Engine/World.h
+++ b/Engine/Source/Runtime/Engine/Classes/Engine/World.h
@@ -2605,11 +2605,13 @@ public:
 	virtual bool IsReadyForFinishDestroy() override;
 	virtual void PostLoad() override;
 	virtual void PreDuplicate(FObjectDuplicationParameters& DupParams) override;
-	//UE_DEPRECATED(5.0, "Use version that takes FObjectPreSaveContext instead.")
+	PRAGMA_DISABLE_DEPRECATION_WARNINGS // Suppress compiler warning on override of deprecated function
+	UE_DEPRECATED(5.0, "Use version that takes FObjectPreSaveRootContext instead.")
 	virtual bool PreSaveRoot(const TCHAR* InFilename) override;
-	virtual void PreSaveRoot(FObjectPreSaveRootContext ObjectSaveContext) override;
-	//UE_DEPRECATED(5.0, "Use version that takes FObjectPostSaveContext instead.")
+	UE_DEPRECATED(5.0, "Use version that takes FObjectPostSaveRootContext instead.")
 	virtual void PostSaveRoot(bool bCleanupIsRequired) override;
+	PRAGMA_ENABLE_DEPRECATION_WARNINGS
+	virtual void PreSaveRoot(FObjectPreSaveRootContext ObjectSaveContext) override;
 	virtual void PostSaveRoot(FObjectPostSaveRootContext ObjectSaveContext) override;
 	virtual UWorld* GetWorld() const override;
 	virtual FPrimaryAssetId GetPrimaryAssetId() const override;
diff --git a/Engine/Source/Runtime/Engine/Classes/GameFramework/Actor.h b/Engine/Source/Runtime/Engine/Classes/GameFramework/Actor.h
--- a/Engine/Source/Runtime/Engine/Classes/GameFramework/Actor.h
+++ b/Engine/Source/Runtime/Engine/Classes/GameFramework/Actor.h
@@ -1843,14 +1843,16 @@ public:
 	virtual bool IsEditorOnly() const override;
 	virtual bool IsAsset() const override;
 
-	//UE_DEPRECATED(5.0, "Use version that takes FObjectPreSaveContext instead.")
+	PRAGMA_DISABLE_DEPRECATION_WARNINGS // Suppress compiler warning on override of deprecated function
+	UE_DEPRECATED(5.0, "Use version that takes FObjectPreSaveRootContext instead.")
 	virtual bool PreSaveRoot(const TCHAR* InFilename) override;
-	virtual void PreSaveRoot(FObjectPreSaveRootContext ObjectSaveContext) override;
-	//UE_DEPRECATED(5.0, "Use version that takes FObjectPostSaveContext instead.")
+	UE_DEPRECATED(5.0, "Use version that takes FObjectPostSaveRootContext instead.")
 	virtual void PostSaveRoot(bool bCleanupIsRequired) override;
-	virtual void PostSaveRoot(FObjectPostSaveRootContext ObjectSaveContext) override;
-	//UE_DEPRECATED(5.0, "Use version that takes FObjectPreSaveContext instead.")
+	UE_DEPRECATED(5.0, "Use version that takes FObjectPreSaveContext instead.")
 	virtual void PreSave(const ITargetPlatform* TargetPlatform) override;
+	PRAGMA_ENABLE_DEPRECATION_WARNINGS
+	virtual void PreSaveRoot(FObjectPreSaveRootContext ObjectSaveContext) override;
+	virtual void PostSaveRoot(FObjectPostSaveRootContext ObjectSaveContext) override;
 	virtual void PreSave(FObjectPreSaveContext ObjectSaveContext) override;
 
 	/** Used to check if Actor is the main actor of a package (currently Child Actors are not) */
diff --git a/Engine/Source/Runtime/Engine/Classes/Materials/Material.h b/Engine/Source/Runtime/Engine/Classes/Materials/Material.h
--- a/Engine/Source/Runtime/Engine/Classes/Materials/Material.h
+++ b/Engine/Source/Runtime/Engine/Classes/Materials/Material.h
@@ -1087,8 +1087,10 @@ public:
 	//~ End UMaterialInterface Interface.
 
 	//~ Begin UObject Interface
-	//UE_DEPRECATED(5.0, "Use version that takes FObjectPreSaveContext instead.")
+	PRAGMA_DISABLE_DEPRECATION_WARNINGS // Suppress compiler warning on override of deprecated function
+	UE_DEPRECATED(5.0, "Use version that takes FObjectPreSaveContext instead.")
 	ENGINE_API virtual void PreSave(const class ITargetPlatform* TargetPlatform) override;
+	PRAGMA_ENABLE_DEPRECATION_WARNINGS
 	ENGINE_API virtual void PreSave(FObjectPreSaveContext ObjectSaveContext) override;
 	ENGINE_API virtual void PostInitProperties() override;
 	ENGINE_API virtual void Serialize(FArchive& Ar) override;
diff --git a/Engine/Source/Runtime/Engine/Classes/Materials/MaterialInstance.h b/Engine/Source/Runtime/Engine/Classes/Materials/MaterialInstance.h
--- a/Engine/Source/Runtime/Engine/Classes/Materials/MaterialInstance.h
+++ b/Engine/Source/Runtime/Engine/Classes/Materials/MaterialInstance.h
@@ -425,8 +425,10 @@ class UMaterialInstance : public UMaterialInterface
 	 */
 	class FMaterialInstanceResource* Resource;
 
-	//UE_DEPRECATED(5.0, "Use version that takes FObjectPreSaveContext instead.")
+	PRAGMA_DISABLE_DEPRECATION_WARNINGS // Suppress compiler warning on override of deprecated function
+	UE_DEPRECATED(5.0, "Use version that takes FObjectPreSaveContext instead.")
 	ENGINE_API virtual void PreSave(const class ITargetPlatform* TargetPlatform) override;
+	PRAGMA_ENABLE_DEPRECATION_WARNINGS
 	ENGINE_API virtual void PreSave(FObjectPreSaveContext ObjectSaveContext) override;
 	ENGINE_API virtual float GetTextureDensity(FName TextureName, const struct FMeshUVChannelInfo& UVChannelData) const override;
 
diff --git a/Engine/Source/Runtime/Engine/Classes/Materials/MaterialInterface.h b/Engine/Source/Runtime/Engine/Classes/Materials/MaterialInterface.h
--- a/Engine/Source/Runtime/Engine/Classes/Materials/MaterialInterface.h
+++ b/Engine/Source/Runtime/Engine/Classes/Materials/MaterialInterface.h
@@ -943,8 +943,10 @@ public:
 	*/
 	ENGINE_API virtual float GetTextureDensity(FName TextureName, const struct FMeshUVChannelInfo& UVChannelData) const;
 
-	//UE_DEPRECATED(5.0, "Use version that takes FObjectPreSaveContext instead.")
+	PRAGMA_DISABLE_DEPRECATION_WARNINGS // Suppress compiler warning on override of deprecated function
+	UE_DEPRECATED(5.0, "Use version that takes FObjectPreSaveContext instead.")
 	ENGINE_API virtual void PreSave(const class ITargetPlatform* TargetPlatform) override;
+	PRAGMA_ENABLE_DEPRECATION_WARNINGS
 	ENGINE_API virtual void PreSave(FObjectPreSaveContext ObjectSaveContext) override;
 
 	/**
diff --git a/Engine/Source/Runtime/Engine/Classes/Particles/ParticleSystem.h b/Engine/Source/Runtime/Engine/Classes/Particles/ParticleSystem.h
--- a/Engine/Source/Runtime/Engine/Classes/Particles/ParticleSystem.h
+++ b/Engine/Source/Runtime/Engine/Classes/Particles/ParticleSystem.h
@@ -404,8 +404,10 @@ public:
 #if WITH_EDITOR
 	virtual void PostEditChangeProperty(FPropertyChangedEvent& PropertyChangedEvent) override;
 #endif // WITH_EDITOR
-	//UE_DEPRECATED(5.0, "Use version that takes FObjectPreSaveContext instead.")
+	PRAGMA_DISABLE_DEPRECATION_WARNINGS // Suppress compiler warning on override of deprecated function
+	UE_DEPRECATED(5.0, "Use version that takes FObjectPreSaveContext instead.")
 	virtual void PreSave(const class ITargetPlatform* TargetPlatform) override;
+	PRAGMA_ENABLE_DEPRECATION_WARNINGS
 	virtual void PreSave(FObjectPreSaveContext ObjectSaveContext) override;
 	virtual void PostLoad() override;
 	virtual bool IsPostLoadThreadSafe() const override;
diff --git a/Engine/Source/Runtime/Engine/Public/WorldPartition/WorldPartitionRuntimeSpatialHash.h b/Engine/Source/Runtime/Engine/Public/WorldPartition/WorldPartitionRuntimeSpatialHash.h
--- a/Engine/Source/Runtime/Engine/Public/WorldPartition/WorldPartitionRuntimeSpatialHash.h
+++ b/Engine/Source/Runtime/Engine/Public/WorldPartition/WorldPartitionRuntimeSpatialHash.h
@@ -188,8 +188,10 @@ class ENGINE_API UWorldPartitionRuntimeSpatialHash : public UWorldPartitionRunti
 	GENERATED_UCLASS_BODY()
 
 public:
-	//UE_DEPRECATED(5.0, "Use version that takes FObjectPreSaveContext instead.")
+	PRAGMA_DISABLE_DEPRECATION_WARNINGS // Suppress compiler warning on override of deprecated function
+	UE_DEPRECATED(5.0, "Use version that takes FObjectPreSaveContext instead.")
 	virtual void PreSave(const class ITargetPlatform* TargetPlatform) override;
+	PRAGMA_ENABLE_DEPRECATION_WARNINGS
 	virtual void PreSave(FObjectPreSaveContext ObjectSaveContext) override;
 
 #if WITH_EDITOR
diff --git a/Engine/Source/Runtime/Landscape/Classes/Landscape.h b/Engine/Source/Runtime/Landscape/Classes/Landscape.h
--- a/Engine/Source/Runtime/Landscape/Classes/Landscape.h
+++ b/Engine/Source/Runtime/Landscape/Classes/Landscape.h
@@ -220,8 +220,10 @@ public:
 	LANDSCAPE_API static void SplitHeightmap(ULandscapeComponent* Comp, ALandscapeProxy* TargetProxy = nullptr, class FMaterialUpdateContext* InOutUpdateContext = nullptr, TArray<class FComponentRecreateRenderStateContext>* InOutRecreateRenderStateContext = nullptr, bool InReregisterComponent = true);
 	
 	//~ Begin UObject Interface.
-	//UE_DEPRECATED(5.0, "Use version that takes FObjectPreSaveContext instead.")
+	PRAGMA_DISABLE_DEPRECATION_WARNINGS // Suppress compiler warning on override of deprecated function
+	UE_DEPRECATED(5.0, "Use version that takes FObjectPreSaveContext instead.")
 	virtual void PreSave(const class ITargetPlatform* TargetPlatform) override;
+	PRAGMA_ENABLE_DEPRECATION_WARNINGS
 	virtual void PreSave(FObjectPreSaveContext ObjectSaveContext) override;
 	virtual void PreEditChange(FProperty* PropertyThatWillChange) override;
 	virtual void PostEditChangeProperty(FPropertyChangedEvent& PropertyChangedEvent) override;
diff --git a/Engine/Source/Runtime/Landscape/Classes/LandscapeHeightfieldCollisionComponent.h b/Engine/Source/Runtime/Landscape/Classes/LandscapeHeightfieldCollisionComponent.h
--- a/Engine/Source/Runtime/Landscape/Classes/LandscapeHeightfieldCollisionComponent.h
+++ b/Engine/Source/Runtime/Landscape/Classes/LandscapeHeightfieldCollisionComponent.h
@@ -222,8 +222,10 @@ public:
 	virtual void Serialize(FArchive& Ar) override;
 	virtual void BeginDestroy() override;
 	virtual void PostLoad() override;
-	//UE_DEPRECATED(5.0, "Use version that takes FObjectPreSaveContext instead.")
+	PRAGMA_DISABLE_DEPRECATION_WARNINGS // Suppress compiler warning on override of deprecated function
+	UE_DEPRECATED(5.0, "Use version that takes FObjectPreSaveContext instead.")
 	virtual void PreSave(const class ITargetPlatform* TargetPlatform) override;
+	PRAGMA_ENABLE_DEPRECATION_WARNINGS
 	virtual void PreSave(FObjectPreSaveContext ObjectSaveContext) override;
 #if WITH_EDITOR
 	virtual void ExportCustomProperties(FOutputDevice& Out, uint32 Indent) override;
diff --git a/Engine/Source/Runtime/Landscape/Classes/LandscapeProxy.h b/Engine/Source/Runtime/Landscape/Classes/LandscapeProxy.h
--- a/Engine/Source/Runtime/Landscape/Classes/LandscapeProxy.h
+++ b/Engine/Source/Runtime/Landscape/Classes/LandscapeProxy.h
@@ -897,8 +897,10 @@ public:
 #endif
 
 	//~ Begin UObject Interface.
-	//UE_DEPRECATED(5.0, "Use version that takes FObjectPreSaveContext instead.")
+	PRAGMA_DISABLE_DEPRECATION_WARNINGS // Suppress compiler warning on override of deprecated function
+	UE_DEPRECATED(5.0, "Use version that takes FObjectPreSaveContext instead.")
 	virtual void PreSave(const class ITargetPlatform* TargetPlatform) override;
+	PRAGMA_ENABLE_DEPRECATION_WARNINGS
 	virtual void PreSave(FObjectPreSaveContext ObjectSaveContext) override;
 	virtual void Serialize(FArchive& Ar) override;
 	static void AddReferencedObjects(UObject* InThis, FReferenceCollector& Collector);
diff --git a/Engine/Source/Runtime/MediaAssets/Public/BaseMediaSource.h b/Engine/Source/Runtime/MediaAssets/Public/BaseMediaSource.h
--- a/Engine/Source/Runtime/MediaAssets/Public/BaseMediaSource.h
+++ b/Engine/Source/Runtime/MediaAssets/Public/BaseMediaSource.h
@@ -38,8 +38,10 @@ private:
 public:
 
 	//~ UObject interface
-	//UE_DEPRECATED(5.0, "Use version that takes FObjectPreSaveContext instead.")
+	PRAGMA_DISABLE_DEPRECATION_WARNINGS // Suppress compiler warning on override of deprecated function
+	UE_DEPRECATED(5.0, "Use version that takes FObjectPreSaveContext instead.")
 	virtual void PreSave(const class ITargetPlatform* TargetPlatform);
+	PRAGMA_ENABLE_DEPRECATION_WARNINGS
 	virtual void PreSave(FObjectPreSaveContext ObjectSaveContext);
 	virtual void GetAssetRegistryTags(TArray<FAssetRegistryTag>& OutTags) const override;
 	virtual void Serialize(FArchive& Ar) override;
diff --git a/Engine/Source/Runtime/MediaAssets/Public/PlatformMediaSource.h b/Engine/Source/Runtime/MediaAssets/Public/PlatformMediaSource.h
--- a/Engine/Source/Runtime/MediaAssets/Public/PlatformMediaSource.h
+++ b/Engine/Source/Runtime/MediaAssets/Public/PlatformMediaSource.h
@@ -38,8 +38,10 @@ private:
 
 public:
 	//~ UObject interface
-	//UE_DEPRECATED(5.0, "Use version that takes FObjectPreSaveContext instead.")
+	PRAGMA_DISABLE_DEPRECATION_WARNINGS // Suppress compiler warning on override of deprecated function
+	UE_DEPRECATED(5.0, "Use version that takes FObjectPreSaveContext instead.")
 	virtual void PreSave(const class ITargetPlatform* TargetPlatform);
+	PRAGMA_ENABLE_DEPRECATION_WARNINGS
 	virtual void PreSave(FObjectPreSaveContext ObjectSaveContext);
 	virtual void Serialize(FArchive& Ar) override;
 
diff --git a/Engine/Source/Runtime/MovieScene/Public/MovieScene.h b/Engine/Source/Runtime/MovieScene/Public/MovieScene.h
--- a/Engine/Source/Runtime/MovieScene/Public/MovieScene.h
+++ b/Engine/Source/Runtime/MovieScene/Public/MovieScene.h
@@ -1124,8 +1124,10 @@ protected:
 
 protected:
 
-	//UE_DEPRECATED(5.0, "Use version that takes FObjectPreSaveContext instead.")
+	PRAGMA_DISABLE_DEPRECATION_WARNINGS // Suppress compiler warning on override of deprecated function
+	UE_DEPRECATED(5.0, "Use version that takes FObjectPreSaveContext instead.")
 	virtual void PreSave(const class ITargetPlatform* TargetPlatform) override;
+	PRAGMA_ENABLE_DEPRECATION_WARNINGS
 	/** Called before this object is being deserialized. */
 	virtual void PreSave(FObjectPreSaveContext ObjectSaveContext) override;
 
diff --git a/Engine/Source/Runtime/MovieScene/Public/MovieSceneSequence.h b/Engine/Source/Runtime/MovieScene/Public/MovieSceneSequence.h
--- a/Engine/Source/Runtime/MovieScene/Public/MovieSceneSequence.h
+++ b/Engine/Source/Runtime/MovieScene/Public/MovieSceneSequence.h
@@ -209,8 +209,10 @@ public:
 public:
 
 	MOVIESCENE_API virtual void PostLoad() override;
-	//UE_DEPRECATED(5.0, "Use version that takes FObjectPreSaveContext instead.")
+	PRAGMA_DISABLE_DEPRECATION_WARNINGS // Suppress compiler warning on override of deprecated function
+	UE_DEPRECATED(5.0, "Use version that takes FObjectPreSaveContext instead.")
 	MOVIESCENE_API virtual void PreSave(const ITargetPlatform* TargetPlatform) override;
+	PRAGMA_ENABLE_DEPRECATION_WARNINGS
 	MOVIESCENE_API virtual void PreSave(FObjectPreSaveContext ObjectSaveContext) override;
 	MOVIESCENE_API virtual void BeginDestroy() override;
 	MOVIESCENE_API virtual void PostDuplicate(bool bDuplicateForPIE) override;
diff --git a/Engine/Source/Runtime/UMG/Public/Blueprint/UserWidget.h b/Engine/Source/Runtime/UMG/Public/Blueprint/UserWidget.h
--- a/Engine/Source/Runtime/UMG/Public/Blueprint/UserWidget.h
+++ b/Engine/Source/Runtime/UMG/Public/Blueprint/UserWidget.h
@@ -1092,8 +1092,10 @@ public:
 
 	//~ Begin UObject Interface
 	virtual bool IsAsset() const;
-	//UE_DEPRECATED(5.0, "Use version that takes FObjectPreSaveContext instead.")
+	PRAGMA_DISABLE_DEPRECATION_WARNINGS // Suppress compiler warning on override of deprecated function
+	UE_DEPRECATED(5.0, "Use version that takes FObjectPreSaveContext instead.")
 	virtual void PreSave(const class ITargetPlatform* TargetPlatform) override;
+	PRAGMA_ENABLE_DEPRECATION_WARNINGS
 	virtual void PreSave(FObjectPreSaveContext ObjectSaveContext) override;
 	//~ End UObject Interface
 
diff --git a/Engine/Source/Runtime/UMG/Public/Blueprint/WidgetTree.h b/Engine/Source/Runtime/UMG/Public/Blueprint/WidgetTree.h
--- a/Engine/Source/Runtime/UMG/Public/Blueprint/WidgetTree.h
+++ b/Engine/Source/Runtime/UMG/Public/Blueprint/WidgetTree.h
@@ -101,8 +101,10 @@ public:
 	}
 
 	// UObject interface
-	//UE_DEPRECATED(5.0, "Use version that takes FObjectPreSaveContext instead.")
+	PRAGMA_DISABLE_DEPRECATION_WARNINGS // Suppress compiler warning on override of deprecated function
+	UE_DEPRECATED(5.0, "Use version that takes FObjectPreSaveContext instead.")
 	virtual void PreSave(const class ITargetPlatform* TargetPlatform) override;
+	PRAGMA_ENABLE_DEPRECATION_WARNINGS
 	virtual void PreSave(FObjectPreSaveContext ObjectSaveContext) override;
 	virtual void PostLoad() override;
 	// End of UObject interface
diff --git a/Engine/Source/Runtime/UMG/Public/Components/Widget.h b/Engine/Source/Runtime/UMG/Public/Components/Widget.h
--- a/Engine/Source/Runtime/UMG/Public/Components/Widget.h
+++ b/Engine/Source/Runtime/UMG/Public/Components/Widget.h
@@ -892,8 +892,10 @@ public:
 	// Begin UObject
 	virtual UWorld* GetWorld() const override;
 	virtual void FinishDestroy() override;
-	//UE_DEPRECATED(5.0, "Use version that takes FObjectPreSaveContext instead.")
+	PRAGMA_DISABLE_DEPRECATION_WARNINGS // Suppress compiler warning on override of deprecated function
+	UE_DEPRECATED(5.0, "Use version that takes FObjectPreSaveContext instead.")
 	virtual void PreSave(const class ITargetPlatform* TargetPlatform) override;
+	PRAGMA_ENABLE_DEPRECATION_WARNINGS
 	virtual void PreSave(FObjectPreSaveContext ObjectSaveContext) override;
 	// End UObject
 
diff --git a/Engine/Source/Runtime/UMG/Public/Slate/SlateVectorArtData.h b/Engine/Source/Runtime/UMG/Public/Slate/SlateVectorArtData.h
--- a/Engine/Source/Runtime/UMG/Public/Slate/SlateVectorArtData.h
+++ b/Engine/Source/Runtime/UMG/Public/Slate/SlateVectorArtData.h
@@ -109,8 +109,10 @@ public:
 
 private:
 	// ~ UObject Interface
-	//UE_DEPRECATED(5.0, "Use version that takes FObjectPreSaveContext instead.")
+	PRAGMA_DISABLE_DEPRECATION_WARNINGS // Suppress compiler warning on override of deprecated function
+	UE_DEPRECATED(5.0, "Use version that takes FObjectPreSaveContext instead.")
 	virtual void PreSave(const class ITargetPlatform* TargetPlatform) override;
+	PRAGMA_ENABLE_DEPRECATION_WARNINGS
 	virtual void PreSave(FObjectPreSaveContext ObjectSaveContext) override;
 	// ~ UObject Interface
 
