PoseSearch: Added a normalization preprocessing step for pose search indices.

Pose data is preprocessed at save time so that each feature is unitless and
contributes proportionally to the distance metric. There are multiple
preprocessors, but each one applies some transformation to the original pose
data and outputs a transformation matrix, its inverse, and a mean vector.
Search queries are transformed into the normalized data space at runtime, and
normalized data can be inverse transformed.

Three preprocessors are implemented. The "None" preprocessor does not modify
the data and emits identity matrices and a zero mean vector. The "Normalize"
preprocessor scales each feature such that its mean absolute deviation about
the sample mean is one. The "Sphere" preprocessor normalizes and decorrelates
the data. Sphering needs more work and is currently hidden as an option.

Also misc cleanup and some prep for future score weighting, and fixed a bug
in motion matching node where time offset wasn't set on followup anims.

#rnx

diff --git a/Engine/Plugins/Experimental/Animation/PoseSearch/Source/Runtime/PoseSearch.Build.cs b/Engine/Plugins/Experimental/Animation/PoseSearch/Source/Runtime/PoseSearch.Build.cs
--- a/Engine/Plugins/Experimental/Animation/PoseSearch/Source/Runtime/PoseSearch.Build.cs
+++ b/Engine/Plugins/Experimental/Animation/PoseSearch/Source/Runtime/PoseSearch.Build.cs
@@ -8,6 +8,10 @@ public class PoseSearch : ModuleRules
 	{
 		PCHUsage = ModuleRules.PCHUsageMode.UseExplicitOrSharedPCHs;
 
+		AddEngineThirdPartyPrivateStaticDependencies(Target,
+			"Eigen"
+		);
+
 		PublicIncludePaths.AddRange(
 			new string[]
 			{
diff --git a/Engine/Plugins/Experimental/Animation/PoseSearch/Source/Runtime/Private/AnimNode_MotionMatching.cpp b/Engine/Plugins/Experimental/Animation/PoseSearch/Source/Runtime/Private/AnimNode_MotionMatching.cpp
--- a/Engine/Plugins/Experimental/Animation/PoseSearch/Source/Runtime/Private/AnimNode_MotionMatching.cpp
+++ b/Engine/Plugins/Experimental/Animation/PoseSearch/Source/Runtime/Private/AnimNode_MotionMatching.cpp
@@ -64,8 +64,7 @@ void FAnimNode_MotionMatching::Update_AnyThread(const FAnimationUpdateContext& C
 			// Overwrite query feature vector with the stepped pose from the database
 			if (DbPoseIdx != INDEX_NONE)
 			{
-				TArrayView<const float> DbFeatureVector = Database->SearchIndex.GetPoseValues(DbPoseIdx);
-				ComposedQuery.Copy(DbFeatureVector);
+				ComposedQuery.CopyFromSearchIndex(Database->SearchIndex, DbPoseIdx);
 			}
 		}
 
@@ -86,11 +85,11 @@ void FAnimNode_MotionMatching::Update_AnyThread(const FAnimationUpdateContext& C
 		float CurrentDissimilarity = MAX_flt;
 		if (DbPoseIdx != INDEX_NONE)
 		{
-			CurrentDissimilarity = UE::PoseSearch::ComparePoses(Database->SearchIndex, DbPoseIdx, ComposedQuery.GetValues());
+			CurrentDissimilarity = UE::PoseSearch::ComparePoses(Database->SearchIndex, DbPoseIdx, ComposedQuery.GetNormalizedValues());
 		}
 
 		// Search the database for the nearest match to the updated query vector
-		UE::PoseSearch::FDbSearchResult Result = UE::PoseSearch::Search(Database, ComposedQuery.GetValues());
+		UE::PoseSearch::FDbSearchResult Result = UE::PoseSearch::Search(Database, ComposedQuery.GetNormalizedValues());
 		if (Result.IsValid() && (ElapsedPoseJumpTime >= SearchThrottleTime))
 		{
 			const FPoseSearchDatabaseSequence& ResultDbSequence = Database->Sequences[Result.DbSequenceIdx];
@@ -134,6 +133,7 @@ void FAnimNode_MotionMatching::Update_AnyThread(const FAnimationUpdateContext& C
 					UE::PoseSearch::FDbSearchResult FollowUpResult;
 					FollowUpResult.DbSequenceIdx = FollowUpDbSequenceIdx;
 					FollowUpResult.PoseIdx = FollowUpPoseIdx;
+					FollowUpResult.TimeOffsetSeconds = FollowUpAssetTime;
 					JumpToPose(Context, FollowUpResult);
 					bJumpedToPose = true;
 				}
@@ -211,6 +211,8 @@ void FAnimNode_MotionMatching::ComposeQuery(const FAnimationBaseContext& Context
 	{
 		ComposedQuery.MergeReplace(Goal);
 	}
+
+	ComposedQuery.Normalize(Database->SearchIndex);
 }
 
 void FAnimNode_MotionMatching::JumpToPose(const FAnimationUpdateContext& Context, UE::PoseSearch::FDbSearchResult Result)
diff --git a/Engine/Plugins/Experimental/Animation/PoseSearch/Source/Runtime/Private/AnimNode_PoseSearchHistoryCollector.cpp b/Engine/Plugins/Experimental/Animation/PoseSearch/Source/Runtime/Private/AnimNode_PoseSearchHistoryCollector.cpp
--- a/Engine/Plugins/Experimental/Animation/PoseSearch/Source/Runtime/Private/AnimNode_PoseSearchHistoryCollector.cpp
+++ b/Engine/Plugins/Experimental/Animation/PoseSearch/Source/Runtime/Private/AnimNode_PoseSearchHistoryCollector.cpp
@@ -47,8 +47,6 @@ void FAnimNode_PoseSearchHistoryCollector::Initialize_AnyThread(const FAnimation
 	UE::Anim::TScopedGraphMessage<UE::PoseSearch::Private::FPoseHistoryProvider> ScopedMessage(Context, this);
 
 	Source.Initialize(Context);
-
-	EvalDeltaTime = 0.0f;
 }
 
 void FAnimNode_PoseSearchHistoryCollector::Evaluate_AnyThread(FPoseContext& Output)
@@ -57,9 +55,7 @@ void FAnimNode_PoseSearchHistoryCollector::Evaluate_AnyThread(FPoseContext& Outp
 
 	Source.Evaluate(Output);
 
-	PoseHistory.Update(EvalDeltaTime, Output);
-
-	EvalDeltaTime = 0.0f;
+	PoseHistory.Update(Output.AnimInstanceProxy->GetDeltaSeconds(), Output);
 }
 
 void FAnimNode_PoseSearchHistoryCollector::Update_AnyThread(const FAnimationUpdateContext& Context)
@@ -69,8 +65,6 @@ void FAnimNode_PoseSearchHistoryCollector::Update_AnyThread(const FAnimationUpda
 	UE::Anim::TScopedGraphMessage<UE::PoseSearch::Private::FPoseHistoryProvider> ScopedMessage(Context, this);
 
 	Source.Update(Context);
-
-	EvalDeltaTime += Context.GetDeltaTime();
 }
 
 void FAnimNode_PoseSearchHistoryCollector::GatherDebugData(FNodeDebugData& DebugData)
diff --git a/Engine/Plugins/Experimental/Animation/PoseSearch/Source/Runtime/Private/PoseSearch.cpp b/Engine/Plugins/Experimental/Animation/PoseSearch/Source/Runtime/Private/PoseSearch.cpp
--- a/Engine/Plugins/Experimental/Animation/PoseSearch/Source/Runtime/Private/PoseSearch.cpp
+++ b/Engine/Plugins/Experimental/Animation/PoseSearch/Source/Runtime/Private/PoseSearch.cpp
@@ -1,6 +1,7 @@
 // Copyright Epic Games, Inc. All Rights Reserved.
 
 #include "PoseSearch/PoseSearch.h"
+#include "PoseSearchEigenHelper.h"
 
 #include "Async/ParallelFor.h"
 #include "Templates/IdentityFunctor.h"
@@ -176,6 +177,14 @@ bool FPoseSearchFeatureDesc::operator==(const FPoseSearchFeatureDesc& Other) con
 		(Domain == Other.Domain);
 }
 
+bool FPoseSearchFeatureDesc::IsSubsampleOfSameFeature(const FPoseSearchFeatureDesc& Other) const
+{
+	return
+		(SchemaBoneIdx == Other.SchemaBoneIdx) &&
+		(Type == Other.Type) &&
+		(Domain == Other.Domain);
+}
+
 
 //////////////////////////////////////////////////////////////////////////
 // FPoseSearchFeatureVectorLayout
@@ -187,7 +196,28 @@ void FPoseSearchFeatureVectorLayout::Init()
 	for (FPoseSearchFeatureDesc& Feature : Features)
 	{
 		Feature.ValueOffset = FloatCount;
-		FloatCount += UE::PoseSearch::GetFeatureTypeTraits(Feature.Type).NumFloats;
+
+		uint32 FeatureNumFloats = UE::PoseSearch::GetFeatureTypeTraits(Feature.Type).NumFloats;
+		FloatCount += FeatureNumFloats;
+
+		if (Feature.SchemaBoneIdx == FPoseSearchFeatureDesc::TrajectoryBoneIndex)
+		{
+			if (FirstTrajectoryValueOffset == -1)
+			{
+				FirstTrajectoryValueOffset = Feature.ValueOffset;
+			}
+
+			NumTrajectoryValues += FeatureNumFloats;
+		}
+		else
+		{
+			if (FirstPoseValueOffset == -1)
+			{
+				FirstPoseValueOffset = Feature.ValueOffset;
+			}
+
+			NumPoseValues += FeatureNumFloats;
+		}
 	}
 
 	NumFloats = FloatCount;
@@ -197,11 +227,15 @@ void FPoseSearchFeatureVectorLayout::Reset()
 {
 	Features.Reset();
 	NumFloats = 0;
+	NumTrajectoryValues = 0;
+	NumPoseValues = 0;
+	FirstTrajectoryValueOffset = -1;
+	FirstPoseValueOffset = -1;
 }
 
 bool FPoseSearchFeatureVectorLayout::IsValid(int32 MaxNumBones) const
 {
-	if (NumFloats == 0.0f)
+	if (NumFloats == 0)
 	{
 		return false;
 	}
@@ -221,13 +255,6 @@ bool FPoseSearchFeatureVectorLayout::IsValid(int32 MaxNumBones) const
 //////////////////////////////////////////////////////////////////////////
 // UPoseSearchSchema
 
-void UPoseSearchSchema::PreSave(const class ITargetPlatform* TargetPlatform)
-{
-	PRAGMA_DISABLE_DEPRECATION_WARNINGS;
-	Super::PreSave(TargetPlatform);
-	PRAGMA_ENABLE_DEPRECATION_WARNINGS;
-}
-
 void UPoseSearchSchema::PreSave(FObjectPreSaveContext ObjectSaveContext)
 {
 	SampleRate = FMath::Clamp(SampleRate, 1, 60);
@@ -243,6 +270,12 @@ void UPoseSearchSchema::PreSave(FObjectPreSaveContext ObjectSaveContext)
 	GenerateLayout();
 	ResolveBoneReferences();
 
+	EffectiveDataPreprocessor = DataPreprocessor;
+	if (EffectiveDataPreprocessor == EPoseSearchDataPreprocessor::Automatic)
+	{
+		EffectiveDataPreprocessor = EPoseSearchDataPreprocessor::Normalize;
+	}
+
 	Super::PreSave(ObjectSaveContext);
 }
 
@@ -273,65 +306,84 @@ void UPoseSearchSchema::GenerateLayout()
 {
 	Layout.Reset();
 
-	for (int32 TrajectoryTimeSubsampleIdx = 0; TrajectoryTimeSubsampleIdx != TrajectorySampleOffsets.Num(); ++TrajectoryTimeSubsampleIdx)
+	// Time domain trajectory positions
+	if (bUseTrajectoryPositions && TrajectorySampleOffsets.Num())
 	{
-		FPoseSearchFeatureDesc Element;
-		Element.SchemaBoneIdx = FPoseSearchFeatureDesc::TrajectoryBoneIndex;
-		Element.SubsampleIdx = TrajectoryTimeSubsampleIdx;
-		Element.Domain = EPoseSearchFeatureDomain::Time;
-
-		if (bUseTrajectoryPositions)
+		FPoseSearchFeatureDesc Feature;
+		Feature.SchemaBoneIdx = FPoseSearchFeatureDesc::TrajectoryBoneIndex;
+		Feature.Domain = EPoseSearchFeatureDomain::Time;
+		Feature.Type = EPoseSearchFeatureType::Position;
+		for (Feature.SubsampleIdx = 0; Feature.SubsampleIdx != TrajectorySampleOffsets.Num(); ++Feature.SubsampleIdx)
 		{
-			Element.Type = EPoseSearchFeatureType::Position;
-			Layout.Features.Add(Element);
+			Layout.Features.Add(Feature);
 		}
+	}
 
-		if (bUseTrajectoryVelocities)
+	// Time domain trajectory linear velocities
+	if (bUseTrajectoryVelocities && TrajectorySampleOffsets.Num())
+	{
+		FPoseSearchFeatureDesc Feature;
+		Feature.SchemaBoneIdx = FPoseSearchFeatureDesc::TrajectoryBoneIndex;
+		Feature.Domain = EPoseSearchFeatureDomain::Time;
+		Feature.Type = EPoseSearchFeatureType::LinearVelocity;
+		for (Feature.SubsampleIdx = 0; Feature.SubsampleIdx != TrajectorySampleOffsets.Num(); ++Feature.SubsampleIdx)
 		{
-			Element.Type = EPoseSearchFeatureType::LinearVelocity;
-			Layout.Features.Add(Element);
+			Layout.Features.Add(Feature);
 		}
 	}
 
- 	for (int32 TrajectoryDistSubsampleIdx = 0; TrajectoryDistSubsampleIdx != TrajectorySampleDistances.Num(); ++TrajectoryDistSubsampleIdx)
- 	{
- 		FPoseSearchFeatureDesc Element;
- 		Element.SchemaBoneIdx = FPoseSearchFeatureDesc::TrajectoryBoneIndex;
- 		Element.SubsampleIdx = TrajectoryDistSubsampleIdx;
- 		Element.Domain = EPoseSearchFeatureDomain::Distance;
-
-		if (bUseTrajectoryPositions)
+	// Distance domain trajectory positions
+	if (bUseTrajectoryPositions && TrajectorySampleDistances.Num())
+	{
+		FPoseSearchFeatureDesc Feature;
+		Feature.SchemaBoneIdx = FPoseSearchFeatureDesc::TrajectoryBoneIndex;
+		Feature.Domain = EPoseSearchFeatureDomain::Distance;
+		Feature.Type = EPoseSearchFeatureType::Position;
+		for (Feature.SubsampleIdx = 0; Feature.SubsampleIdx != TrajectorySampleDistances.Num(); ++Feature.SubsampleIdx)
 		{
-			Element.Type = EPoseSearchFeatureType::Position;
-			Layout.Features.Add(Element);
+			Layout.Features.Add(Feature);
 		}
+	}
 
-		if (bUseTrajectoryVelocities)
+	// Distance domain trajectory linear velocities
+	if (bUseTrajectoryVelocities && TrajectorySampleDistances.Num())
+	{
+		FPoseSearchFeatureDesc Feature;
+		Feature.SchemaBoneIdx = FPoseSearchFeatureDesc::TrajectoryBoneIndex;
+		Feature.Domain = EPoseSearchFeatureDomain::Distance;
+		Feature.Type = EPoseSearchFeatureType::LinearVelocity;
+		for (Feature.SubsampleIdx = 0; Feature.SubsampleIdx != TrajectorySampleDistances.Num(); ++Feature.SubsampleIdx)
 		{
-			Element.Type = EPoseSearchFeatureType::LinearVelocity;
-			Layout.Features.Add(Element);
+			Layout.Features.Add(Feature);
 		}
- 	}
+	}
 
-	for (int32 PoseSubsampleIdx = 0; PoseSubsampleIdx != PoseSampleOffsets.Num(); ++PoseSubsampleIdx)
+	// Time domain bone positions
+	if (bUseBonePositions && PoseSampleOffsets.Num())
 	{
-		FPoseSearchFeatureDesc Element;
-		Element.SubsampleIdx = PoseSubsampleIdx;
-		Element.Domain = EPoseSearchFeatureDomain::Time;
-
-		for (int32 SchemaBoneIdx = 0; SchemaBoneIdx != Bones.Num(); ++SchemaBoneIdx)
+		FPoseSearchFeatureDesc Feature;
+		Feature.Domain = EPoseSearchFeatureDomain::Time;
+		Feature.Type = EPoseSearchFeatureType::Position;
+		for (Feature.SubsampleIdx = 0; Feature.SubsampleIdx != PoseSampleOffsets.Num(); ++Feature.SubsampleIdx)
 		{
-			Element.SchemaBoneIdx = SchemaBoneIdx;
-			if (bUseBonePositions)
+			for (Feature.SchemaBoneIdx = 0; Feature.SchemaBoneIdx != Bones.Num(); ++Feature.SchemaBoneIdx)
 			{
-				Element.Type = EPoseSearchFeatureType::Position;
-				Layout.Features.Add(Element);
+				Layout.Features.Add(Feature);
 			}
+		}
+	}
 
-			if (bUseBoneVelocities)
+	// Time domain bone linear velocities
+	if (bUseBoneVelocities && PoseSampleOffsets.Num())
+	{
+		FPoseSearchFeatureDesc Feature;
+		Feature.Domain = EPoseSearchFeatureDomain::Time;
+		Feature.Type = EPoseSearchFeatureType::LinearVelocity;
+		for (Feature.SubsampleIdx = 0; Feature.SubsampleIdx != PoseSampleOffsets.Num(); ++Feature.SubsampleIdx)
+		{
+			for (Feature.SchemaBoneIdx = 0; Feature.SchemaBoneIdx != Bones.Num(); ++Feature.SchemaBoneIdx)
 			{
-				Element.Type = EPoseSearchFeatureType::LinearVelocity;
-				Layout.Features.Add(Element);
+				Layout.Features.Add(Feature);
 			}
 		}
 	}
@@ -399,29 +451,81 @@ void FPoseSearchIndex::Reset()
 	Schema = nullptr;
 }
 
+void FPoseSearchIndex::Normalize (TArrayView<float> InOutPoseVector) const
+{
+	using namespace Eigen;
 
-//////////////////////////////////////////////////////////////////////////
-// UPoseSearchSequenceMetaData
+	auto TransformationMtx = Map<const Matrix<float, Dynamic, Dynamic, ColMajor>>
+	(
+		PreprocessInfo.TransformationMatrix.GetData(),
+		PreprocessInfo.NumDimensions,
+		PreprocessInfo.NumDimensions
+	);
+	auto SampleMean = Map<const Matrix<float, Dynamic, 1, ColMajor>>
+	(
+		PreprocessInfo.SampleMean.GetData(),
+		PreprocessInfo.NumDimensions
+	);
+
+	checkSlow(InOutPoseVector.Num() == PreprocessInfo.NumDimensions);
 
-void UPoseSearchSequenceMetaData::PreSave(const class ITargetPlatform* TargetPlatform)
+	auto PoseVector = Map<Matrix<float, Dynamic, 1, ColMajor>>
+	(
+		InOutPoseVector.GetData(),
+		InOutPoseVector.Num()
+	);
+
+	PoseVector = TransformationMtx * (PoseVector - SampleMean);
+}
+
+void FPoseSearchIndex::InverseNormalize (TArrayView<float> InOutNormalizedPoseVector) const
 {
-	PRAGMA_DISABLE_DEPRECATION_WARNINGS;
-	Super::PreSave(TargetPlatform);
-	PRAGMA_ENABLE_DEPRECATION_WARNINGS;
+	using namespace Eigen;
+
+	auto InverseTransformationMtx = Map<const Matrix<float, Dynamic, Dynamic, ColMajor>>
+	(
+		PreprocessInfo.InverseTransformationMatrix.GetData(),
+		PreprocessInfo.NumDimensions,
+		PreprocessInfo.NumDimensions
+	);
+	auto SampleMean = Map<const Matrix<float, Dynamic, 1, ColMajor>>
+	(
+		PreprocessInfo.SampleMean.GetData(),
+		PreprocessInfo.NumDimensions
+	);
+
+	checkSlow(InOutNormalizedPoseVector.Num() == PreprocessInfo.NumDimensions);
+
+	auto NormalizedPoseVector = Map<Matrix<float, Dynamic, 1, ColMajor>>
+	(
+		InOutNormalizedPoseVector.GetData(),
+		InOutNormalizedPoseVector.Num()
+	);
+
+	NormalizedPoseVector = (InverseTransformationMtx * NormalizedPoseVector) + SampleMean;
 }
 
+
+//////////////////////////////////////////////////////////////////////////
+// UPoseSearchSequenceMetaData
+
 void UPoseSearchSequenceMetaData::PreSave(FObjectPreSaveContext ObjectSaveContext)
 {
 	SearchIndex.Reset();
 
-	if (IsValidForIndexing())
+#if WITH_EDITOR
+	if (!IsTemplate())
 	{
-		UObject* Outer = GetOuter();
-		if (UAnimSequence* Sequence = Cast<UAnimSequence>(Outer))
+		if (IsValidForIndexing())
 		{
-			UE::PoseSearch::BuildIndex(Sequence, this);
+			UObject* Outer = GetOuter();
+			if (UAnimSequence* Sequence = Cast<UAnimSequence>(Outer))
+			{
+				UE::PoseSearch::BuildIndex(Sequence, this);
+			}
 		}
 	}
+#endif
 
 	Super::PreSave(ObjectSaveContext);
 }
@@ -478,7 +582,31 @@ int32 UPoseSearchDatabase::GetPoseIndexFromAssetTime(int32 DbSequenceIdx, float
 
 bool UPoseSearchDatabase::IsValidForIndexing() const
 {
-	return Schema && Schema->IsValid() && !Sequences.IsEmpty();
+	bool bValid = Schema && Schema->IsValid() && !Sequences.IsEmpty();
+
+	if (bValid)
+	{
+		bool bSequencesValid = true;
+		for (const FPoseSearchDatabaseSequence& DbSequence : Sequences)
+		{
+			if (!DbSequence.Sequence)
+			{
+				bSequencesValid = false;
+				break;
+			}
+
+			USkeleton* SeqSkeleton = DbSequence.Sequence->GetSkeleton();
+			if (!SeqSkeleton || !SeqSkeleton->IsCompatible(Schema->Skeleton))
+			{
+				bSequencesValid = false;
+				break;
+			}
+		}
+
+		bValid = bSequencesValid;
+	}
+
+	return bValid;
 }
 
 bool UPoseSearchDatabase::IsValidForSearch() const
@@ -486,21 +614,19 @@ bool UPoseSearchDatabase::IsValidForSearch() const
 	return IsValidForIndexing() && SearchIndex.IsValid();
 }
 
-void UPoseSearchDatabase::PreSave(const class ITargetPlatform* TargetPlatform)
-{
-	PRAGMA_DISABLE_DEPRECATION_WARNINGS;
-	Super::PreSave(TargetPlatform);
-	PRAGMA_ENABLE_DEPRECATION_WARNINGS;
-}
-
 void UPoseSearchDatabase::PreSave(FObjectPreSaveContext ObjectSaveContext)
 {
 	SearchIndex.Reset();
 
-	if (IsValidForIndexing())
+#if WITH_EDITOR
+	if (!IsTemplate())
 	{
-		UE::PoseSearch::BuildIndex(this);
+		if (IsValidForIndexing())
+		{
+			UE::PoseSearch::BuildIndex(this);
+		}
 	}
+#endif
 
 	Super::PreSave(ObjectSaveContext);
 }
@@ -520,6 +646,8 @@ void FPoseSearchFeatureVectorBuilder::ResetFeatures()
 {
 	Values.Reset(0);
 	Values.SetNumZeroed(Schema->Layout.NumFloats);
+	ValuesNormalized.Reset(0);
+	ValuesNormalized.SetNumZeroed(Schema->Layout.NumFloats);
 	NumFeaturesAdded = 0;
 	FeaturesAdded.Init(false, Schema->Layout.Features.Num());
 }
@@ -706,10 +834,16 @@ bool FPoseSearchFeatureVectorBuilder::SetPastTrajectoryFeatures(UE::PoseSearch::
 	return true;
 }
 
-void FPoseSearchFeatureVectorBuilder::Copy(TArrayView<const float> FeatureVector)
+void FPoseSearchFeatureVectorBuilder::CopyFromSearchIndex(const FPoseSearchIndex& SearchIndex, int32 PoseIdx)
 {
-	check(FeatureVector.Num() == Values.Num());
-	FMemory::Memcpy(Values.GetData(), FeatureVector.GetData(), FeatureVector.GetTypeSize() * FeatureVector.Num());
+	check(Schema == SearchIndex.Schema);
+
+	TArrayView<const float> FeatureVector = SearchIndex.GetPoseValues(PoseIdx);
+
+	ValuesNormalized = FeatureVector;
+	Values = FeatureVector;
+	SearchIndex.InverseNormalize(Values);
+
 	NumFeaturesAdded = Schema->Layout.Features.Num();
 	FeaturesAdded.SetRange(0, FeaturesAdded.Num(), true);
 }
@@ -760,6 +894,11 @@ bool FPoseSearchFeatureVectorBuilder::IsCompatible(const FPoseSearchFeatureVecto
 	return IsInitialized() && (Schema == OtherBuilder.Schema);
 }
 
+void FPoseSearchFeatureVectorBuilder::Normalize(const FPoseSearchIndex& ForSearchIndex)
+{
+	ValuesNormalized = Values;
+	ForSearchIndex.Normalize(ValuesNormalized);
+}
 
 namespace UE { namespace PoseSearch {
 
@@ -1740,9 +1879,12 @@ static void DrawSearchIndex(const FDebugDrawParams& DrawParams)
 		return;
 	}
 
+	TArray<float> PoseVector;
 	for (int32 PoseIdx = StartPoseIdx; PoseIdx != LastPoseIdx; ++PoseIdx)
 	{
-		Reader.SetValues(SearchIndex->GetPoseValues(PoseIdx));
+		PoseVector = SearchIndex->GetPoseValues(PoseIdx);
+		SearchIndex->InverseNormalize(PoseVector);
+		Reader.SetValues(PoseVector);
 
 		DrawFeatureVector(DrawParams, Reader);
 	}
@@ -1785,6 +1927,416 @@ void Draw(const FDebugDrawParams& DebugDrawParams)
 	}
 }
 
+static void PreprocessSearchIndexNone(FPoseSearchIndex* SearchIndex)
+{
+	// This function leaves the data unmodified and simply outputs the transformation
+	// and inverse transformation matrices as the identity matrix and the sample mean
+	// as the zero vector.
+
+	using namespace Eigen;
+
+	check(SearchIndex->IsValid());
+
+	FPoseSearchIndexPreprocessInfo& Info = SearchIndex->PreprocessInfo;
+	Info.Reset();
+
+	const FPoseSearchFeatureVectorLayout& Layout = SearchIndex->Schema->Layout;
+
+	const int32 NumPoses = SearchIndex->NumPoses;
+	const int32 NumDimensions = Layout.NumFloats;
+
+	Info.NumDimensions = NumDimensions;
+	Info.TransformationMatrix.SetNumZeroed(NumDimensions * NumPoses);
+	Info.InverseTransformationMatrix.SetNumZeroed(NumDimensions * NumPoses);
+	Info.SampleMean.SetNumZeroed(NumDimensions);
+
+	// Map output transformation matrix
+	auto TransformMap = Map<Matrix<float, Dynamic, Dynamic, ColMajor>>(
+		Info.TransformationMatrix.GetData(),
+		NumDimensions, NumPoses
+	);
+
+	// Map output inverse transformation matrix
+	auto InverseTransformMap = Map<Matrix<float, Dynamic, Dynamic, ColMajor>>(
+		Info.InverseTransformationMatrix.GetData(),
+		NumDimensions, NumPoses
+	);
+
+	// Map output sample mean vector
+	auto SampleMeanMap = Map<VectorXf>(Info.SampleMean.GetData(), NumDimensions);
+
+	// Write the transformation matrices and sample mean
+	TransformMap = MatrixXf::Identity(NumDimensions, NumPoses);
+	InverseTransformMap = MatrixXf::Identity(NumDimensions, NumPoses);
+	SampleMeanMap = VectorXf::Zero(NumDimensions);
+}
+
+static Eigen::VectorXd ComputeFeatureMeanDeviations (const Eigen::MatrixXd& CenteredPoseMatrix, const FPoseSearchFeatureVectorLayout& Layout) {
+	using namespace Eigen;
+
+	const int32 NumPoses = CenteredPoseMatrix.cols();
+	const int32 NumDimensions = CenteredPoseMatrix.rows();
+
+	VectorXd MeanDeviations(NumDimensions);
+	MeanDeviations.setConstant(0);
+	for (const FPoseSearchFeatureDesc& Feature : Layout.Features)
+	{
+		int32 FeatureDims = GetFeatureTypeTraits(Feature.Type).NumFloats;
+
+		// Construct a submatrix for the feature and find the average distance to the feature's centroid.
+		// Since we've already mean centered the data, the average distance to the mean is simply the average norm.
+		double FeatureMeanDeviation = CenteredPoseMatrix.block(Feature.ValueOffset, 0, FeatureDims, NumPoses).colwise().norm().mean();
+
+		// Fill the feature's corresponding scaling axes with the average distance
+		MeanDeviations.segment(Feature.ValueOffset, FeatureDims).setConstant(FeatureMeanDeviation);
+	}
+
+	return MeanDeviations;
+}
+
+static void PreprocessSearchIndexNormalize(FPoseSearchIndex* SearchIndex)
+{
+	// This function performs a modified z-score normalization where features are normalized
+	// by mean absolute deviation rather than standard deviation. Both methods are preferable
+	// here to min-max scaling because they preserve outliers.
+	// 
+	// Mean absolute deviation is preferred here over standard deviation because the latter
+	// emphasizes outliers since squaring the distance from the mean increases variance 
+	// exponentially rather than additively and square rooting the sum of squares does not 
+	// remove that bias. [1]
+	//
+	// The pose matrix is transformed in place and the tranformation matrix, its inverse,
+	// and data mean vector are computed and stored along with it.
+	//
+	// N:	number of dimensions for input column vectors
+	// P:	number of input column vectors
+	// X:	NxP input matrix
+	// x_p:	pth column vector of input matrix
+	// u:   mean column vector of X
+	//
+	// S:	mean absolute deviations of X, as diagonal NxN matrix with average distances replicated for each feature's axes
+	// s_n:	nth deviation
+	//
+	// Normalization by mean absolute deviation algorithm:
+	//
+	// 1) mean-center X
+	//    x_p := x_p - u
+	// 2) rescale X by inverse mean absolute deviation
+	//    x_p := x_p * s_n^(-1)
+	// 
+	// Let S^(-1) be the inverse of S where the nth diagonal element is s_n^(-1)
+	// then step 2 can be expressed as matrix multiplication:
+	// X := S^(-1) * X
+	//
+	// By persisting the mean vector u and linear transform S, we can bring an input vector q
+	// into the same space as the mean centered and scaled data matrix X:
+	// q := S^(-1) * (q - u)
+	//
+	// This operation is invertible, a normalized data vector x can be unscaled via:
+	// x := (S * x) + u
+	//
+	// References:
+	// [1] Gorard, S. (2005), “Revisiting a 90-Year-Old Debate: The Advantages of the Mean Deviation.”
+	//     British Journal of Educational Studies, 53: 417-430.
+
+	using namespace Eigen;
+
+	check(SearchIndex->IsValid());
+
+	FPoseSearchIndexPreprocessInfo& Info = SearchIndex->PreprocessInfo;
+	Info.Reset();
+
+	const FPoseSearchFeatureVectorLayout& Layout = SearchIndex->Schema->Layout;
+
+	const int32 NumPoses = SearchIndex->NumPoses;
+	const int32 NumDimensions = Layout.NumFloats;
+
+	// Map input buffer
+	auto PoseMatrixSourceMap = Map<Matrix<float, Dynamic, Dynamic, RowMajor>>(
+		SearchIndex->Values.GetData(),
+		NumPoses,		// rows
+		NumDimensions	// cols
+	);
+
+	// Copy row major float matrix to column major double matrix
+	MatrixXd PoseMatrix = PoseMatrixSourceMap.transpose().cast<double>();
+	checkSlow(PoseMatrix.rows() == NumDimensions);
+	checkSlow(PoseMatrix.cols() == NumPoses);
+
+	// Mean center
+	VectorXd SampleMean = PoseMatrix.rowwise().mean();
+	PoseMatrix = PoseMatrix.colwise() - SampleMean;
+
+	// Compute per-feature average distances
+	VectorXd MeanDeviations = ComputeFeatureMeanDeviations(PoseMatrix, Layout);
+
+	// Construct a scaling matrix that uniformly scales each feature by its average distance from the mean
+	MatrixXd ScalingMatrix = MeanDeviations.cwiseInverse().asDiagonal();
+
+	// Construct the inverse scaling matrix
+	MatrixXd InverseScalingMatrix = MeanDeviations.asDiagonal();
+
+	// Rescale data by transforming it with the scaling matrix
+	// Now each feature has an average Euclidean length = 1.
+	PoseMatrix = ScalingMatrix * PoseMatrix;
+
+	// Write normalized data back to source buffer, converting from column data back to row data
+	PoseMatrixSourceMap = PoseMatrix.transpose().cast<float>();
+
+	// Output preprocessing info
+	Info.NumDimensions = NumDimensions;
+	Info.TransformationMatrix.SetNumZeroed(ScalingMatrix.size());
+	Info.InverseTransformationMatrix.SetNumZeroed(InverseScalingMatrix.size());
+	Info.SampleMean.SetNumZeroed(SampleMean.size());
+
+	auto TransformMap = Map<Matrix<float, Dynamic, Dynamic, ColMajor>>(
+		Info.TransformationMatrix.GetData(),
+		ScalingMatrix.rows(), ScalingMatrix.cols()
+	);
+
+	auto InverseTransformMap = Map<Matrix<float, Dynamic, Dynamic, ColMajor>>(
+		Info.InverseTransformationMatrix.GetData(),
+		InverseScalingMatrix.rows(), InverseScalingMatrix.cols()
+	);
+
+	auto SampleMeanMap = Map<VectorXf>(Info.SampleMean.GetData(), SampleMean.size());
+
+	// Output scaling matrix, inverse scaling matrix, and mean vector
+	TransformMap = ScalingMatrix.cast<float>();
+	InverseTransformMap = InverseScalingMatrix.cast<float>();
+	SampleMeanMap = SampleMean.cast<float>();
+
+#if UE_POSE_SEARCH_EIGEN_DEBUG
+	FString PoseMtxStr = EigenMatrixToString(PoseMatrix);
+	FString TransformationStr = EigenMatrixToString(TransformMap);
+	FString InverseTransformationStr = EigenMatrixToString(InverseTransformMap);
+	FString SampleMeanStr = EigenMatrixToString(SampleMeanMap);
+#endif // UE_POSE_SEARCH_EIGEN_DEBUG
+}
+
+static void PreprocessSearchIndexSphere(FPoseSearchIndex* SearchIndex)
+{
+	// This function performs correlation based zero-phase component analysis sphering (ZCA-cor sphering)
+	// The pose matrix is transformed in place and the tranformation matrix, its inverse,
+	// and data mean vector are computed and stored along with it.
+	//
+	// N:	number of dimensions for input column vectors
+	// P:	number of input column vectors
+	// X:	NxP input matrix
+	// x_p:	pth column vector of input matrix
+	// u:   mean column vector of X
+	//
+	// Eigendecomposition of correlation matrix of X:
+	// cor(X) = (1/P) * X * X^T = V * D * V^T
+	//
+	// V:	eigenvectors of cor(X), stacked as columns in an orthogonal NxN matrix
+	// D:	eigenvalues of cor(X), as diagonal NxN matrix
+	// d_n:	nth eigenvalue
+	// s_n: nth standard deviation
+	// s_n^2 = d_n, the variance along the nth eigenvector
+	// s_n   = d_n^(1/2)
+	//
+	// ZCA sphering algorithm:
+	//
+	// 1) mean-center X
+	//    x_p := x_p - u
+	// 2) align largest orthogonal directions of variance in X to coordinate axes (PCA rotate)
+	//    x_p := V^T * x_p
+	// 3) rescale X by inverse standard deviation
+	//    x_p := x_p * d_n^(-1/2)
+	// 4) return now rescaled X back to original rotation (inverse PCA rotate)
+	//    x_p := V * x_p
+	// 
+	// Let D^(-1/2) be the inverse square root of D where the nth diagonal element is d_n^(-1/2)
+	// then steps 2-4 can be expressed as a series of matrix multiplications:
+	// Z = V * D^(-1/2) * V^T
+	// X := Z * X
+	//
+	// By persisting the mean vector u and linear transform Z, we can bring an input vector q
+	// into the same space as the sphered data matrix X:
+	// q := Z * (q - u)
+	//
+	// This operation is invertible, a sphere standardized data vector x can be unscaled via:
+	// Z^(-1) = V * D^(1/2) * V^T
+	// x := (Z^(-1) * x) + u
+	//
+	// The sphering processs allows nearest neighbor queries to use the Mahalonobis metric
+	// which is unitless, scale-invariant, and uncorrelated. The Mahalanobis distance between
+	// two random vectors x and y in data matrix X is:
+	// d(x,y) = ((x-y)^T * cov(X)^(-1) * (x-y))^(1/2)
+	//
+	// Because sphering transforms X into a new matrix with identity covariance, the Mahalonobis
+	// distance equation above reduces to Eucledean distance since cov(X)^(-1) = I:
+	// d(x,y) = ((x-y)^T * (x-y))^(1/2)
+	// 
+	// References:
+	// Watt, Jeremy, et al. Machine Learning Refined: Foundations, Algorithms, and Applications.
+	// 2nd ed., Cambridge University Press, 2020.
+	// 
+	// Kessy, Agnan, Alex Lewin, and Korbinian Strimmer. "Optimal whitening and decorrelation."
+	// The American Statistician 72.4 (2018): 309-314.
+	// 
+	// https://en.wikipedia.org/wiki/Whitening_transformation
+	// 
+	// https://en.wikipedia.org/wiki/Mahalanobis_distance
+	//
+	// Note this sphering preprocessor needs more work and isn't yet exposed in the editor as an option.
+	// Todo:
+	// - Try singular value decomposition in place of eigendecomposition
+	// - Remove zero variance feature axes from data and search queries
+	// - Support weighted Mahalanobis metric. User supplied weights need to be transformed to data's new basis.
+
+#if UE_POSE_SEARCH_EIGEN_DEBUG
+	double StartTime = FPlatformTime::Seconds();
+#endif
+
+	using namespace Eigen;
+
+	check(SearchIndex->IsValid());
+
+	FPoseSearchIndexPreprocessInfo& Info = SearchIndex->PreprocessInfo;
+	Info.Reset();
+
+	const FPoseSearchFeatureVectorLayout& Layout = SearchIndex->Schema->Layout;
+
+	const int32 NumPoses = SearchIndex->NumPoses;
+	const int32 NumDimensions = Layout.NumFloats;
+
+	// Map input buffer
+	auto PoseMatrixSourceMap = Map<Matrix<float, Dynamic, Dynamic, RowMajor>>(
+		SearchIndex->Values.GetData(),
+		NumPoses,		// rows
+		NumDimensions	// cols
+	);
+
+	// Copy row major float matrix to column major double matrix
+	MatrixXd PoseMatrix = PoseMatrixSourceMap.transpose().cast<double>();
+	checkSlow(PoseMatrix.rows() == NumDimensions);
+	checkSlow(PoseMatrix.cols() == NumPoses);
+
+	// Mean center
+	VectorXd SampleMean = PoseMatrix.rowwise().mean();
+	PoseMatrix = PoseMatrix.colwise() - SampleMean;
+
+	// Compute per-feature average distances
+	VectorXd MeanDeviations = ComputeFeatureMeanDeviations(PoseMatrix, Layout);
+
+	// Rescale data by transforming it with the scaling matrix
+	// Now each feature has an average Euclidean length = 1.
+	MatrixXd PoseMatrixNormalized = MeanDeviations.cwiseInverse().asDiagonal() * PoseMatrix;
+
+	// Compute sample covariance
+	MatrixXd Covariance = ((1.0 / NumPoses) * (PoseMatrixNormalized * PoseMatrixNormalized.transpose())) + 1e-7 * MatrixXd::Identity(NumDimensions, NumDimensions);
+
+	VectorXd StdDev = Covariance.diagonal().cwiseSqrt();
+	VectorXd InvStdDev = StdDev.cwiseInverse();
+	MatrixXd Correlation = InvStdDev.asDiagonal() * Covariance * InvStdDev.asDiagonal();
+
+	// Compute eigenvalues and eigenvectors of correlation matrix
+	SelfAdjointEigenSolver<MatrixXd> EigenDecomposition(Correlation, ComputeEigenvectors);
+
+	VectorXd EigenValues = EigenDecomposition.eigenvalues();
+	MatrixXd EigenVectors = EigenDecomposition.eigenvectors();
+
+	// Sort eigenpairs by descending eigenvalue
+	{
+		const Eigen::Index n = EigenValues.size();
+		for (Eigen::Index i = 0; i < n-1; ++i)
+		{
+			Index k;
+			EigenValues.segment(i,n-i).maxCoeff(&k);
+			if (k > 0)
+			{
+				std::swap(EigenValues[i], EigenValues[k+i]);
+				EigenVectors.col(i).swap(EigenVectors.col(k+i));
+			}
+		}
+	}
+
+	// Regularize eigenvalues
+	EigenValues = EigenValues.array() + 1e-7;
+
+	// Compute ZCA-cor and ZCA-cor^(-1)
+	MatrixXd ZCA = EigenVectors * EigenValues.cwiseInverse().cwiseSqrt().asDiagonal() * EigenVectors.transpose() * MeanDeviations.cwiseInverse().asDiagonal();
+	MatrixXd ZCAInverse = MeanDeviations.asDiagonal() * EigenVectors * EigenValues.cwiseSqrt().asDiagonal() * EigenVectors.transpose();
+
+	// Apply sphering transform to the data matrix
+	PoseMatrix = ZCA * PoseMatrix;
+	checkSlow(PoseMatrix.rows() == NumDimensions);
+	checkSlow(PoseMatrix.cols() == NumPoses);
+
+	// Write data back to source buffer, converting from column data back to row data
+	PoseMatrixSourceMap = PoseMatrix.transpose().cast<float>();
+
+	// Output preprocessing info
+	Info.NumDimensions = NumDimensions;
+	Info.TransformationMatrix.SetNumZeroed(ZCA.size());
+	Info.InverseTransformationMatrix.SetNumZeroed(ZCAInverse.size());
+	Info.SampleMean.SetNumZeroed(SampleMean.size());
+
+	auto TransformMap = Map<Matrix<float, Dynamic, Dynamic, ColMajor>>(
+		Info.TransformationMatrix.GetData(),
+		ZCA.rows(), ZCA.cols()
+	);
+
+	auto InverseTransformMap = Map<Matrix<float, Dynamic, Dynamic, ColMajor>>(
+		Info.InverseTransformationMatrix.GetData(),
+		ZCAInverse.rows(), ZCAInverse.cols()
+	);
+
+	auto SampleMeanMap = Map<VectorXf>(Info.SampleMean.GetData(), SampleMean.size());
+
+	// Output sphering matrix, inverse sphering matrix, and mean vector
+	TransformMap = ZCA.cast<float>();
+	InverseTransformMap = ZCAInverse.cast<float>();
+	SampleMeanMap = SampleMean.cast<float>();
+
+#if UE_POSE_SEARCH_EIGEN_DEBUG
+	double ElapsedTime = FPlatformTime::Seconds() - StartTime;
+
+	FString EigenValuesStr = EigenMatrixToString(EigenValues);
+	FString EigenVectorsStr = EigenMatrixToString(EigenVectors);
+
+	FString CovarianceStr = EigenMatrixToString(Covariance);
+	FString CorrelationStr = EigenMatrixToString(Correlation);
+
+	FString ZCAStr = EigenMatrixToString(ZCA);
+	FString ZCAInverseStr = EigenMatrixToString(ZCAInverse);
+
+	FString PoseMatrixSphereStr = EigenMatrixToString(PoseMatrix);
+	MatrixXd PoseMatrixUnsphered = ZCAInverse * PoseMatrix;
+	FString PoseMatrixUnspheredStr = EigenMatrixToString(PoseMatrixUnsphered);
+
+	FString OutputPoseMatrixStr = EigenMatrixToString(PoseMatrixSourceMap);
+
+	FString TransformStr = EigenMatrixToString(TransformMap);
+	FString InverseTransformStr = EigenMatrixToString(InverseTransformMap);
+	FString SampleMeanStr = EigenMatrixToString(SampleMeanMap);
+#endif // UE_POSE_SEARCH_EIGEN_DEBUG
+}
+
+static void PreprocessSearchIndex(FPoseSearchIndex* SearchIndex)
+{
+	switch (SearchIndex->Schema->EffectiveDataPreprocessor)
+	{
+		case EPoseSearchDataPreprocessor::Normalize:
+			PreprocessSearchIndexNormalize(SearchIndex);
+		break;
+
+		case EPoseSearchDataPreprocessor::Sphere:
+			PreprocessSearchIndexSphere(SearchIndex);
+		break;
+
+		case EPoseSearchDataPreprocessor::None:
+			PreprocessSearchIndexNone(SearchIndex);
+		break;
+
+		case EPoseSearchDataPreprocessor::Invalid:
+			checkNoEntry();
+		break;
+	}
+}
+
 bool BuildIndex(const UAnimSequence* Sequence, UPoseSearchSequenceMetaData* SequenceMetaData)
 {
 	check(Sequence);
@@ -1820,6 +2372,9 @@ bool BuildIndex(const UAnimSequence* Sequence, UPoseSearchSequenceMetaData* Sequ
 	SequenceMetaData->SearchIndex.Values = Indexer.Output.FeatureVectorTable;
 	SequenceMetaData->SearchIndex.NumPoses = Indexer.Output.NumIndexedPoses;
 	SequenceMetaData->SearchIndex.Schema = SequenceMetaData->Schema;
+
+	PreprocessSearchIndex(&SequenceMetaData->SearchIndex);
+
 	return true;
 }
 
@@ -1832,15 +2387,6 @@ bool BuildIndex(UPoseSearchDatabase* Database)
 		return false;
 	}
 
-	for (const FPoseSearchDatabaseSequence& DbSequence : Database->Sequences)
-	{
-		USkeleton* SeqSkeleton = DbSequence.Sequence->GetSkeleton();
-		if (!SeqSkeleton || !SeqSkeleton->IsCompatible(Database->Schema->Skeleton))
-		{
-			return false;
-		}
-	}
-
 	// Prepare animation sampling tasks
 	TArray<FSequenceSampler> SequenceSamplers;
 	TMap<const UAnimSequence*, int32> SequenceSamplerMap;
@@ -1931,6 +2477,9 @@ bool BuildIndex(UPoseSearchDatabase* Database)
 
 	Database->SearchIndex.NumPoses = TotalPoses;
 	Database->SearchIndex.Schema = Database->Schema;
+
+	PreprocessSearchIndex(&Database->SearchIndex);
+
 	return true;
 }
 
diff --git a/Engine/Plugins/Experimental/Animation/PoseSearch/Source/Runtime/Private/PoseSearchEigenHelper.h b/Engine/Plugins/Experimental/Animation/PoseSearch/Source/Runtime/Private/PoseSearchEigenHelper.h
new file mode 100644
--- /dev/null
+++ b/Engine/Plugins/Experimental/Animation/PoseSearch/Source/Runtime/Private/PoseSearchEigenHelper.h
@@ -0,0 +1,60 @@
+// Copyright Epic Games, Inc. All Rights Reserved.
+
+#if UE_BUILD_DEBUG && WITH_EDITOR
+#define UE_POSE_SEARCH_EIGEN_DEBUG 1
+#endif
+
+#ifndef UE_POSE_SEARCH_EIGEN_DEBUG
+#define UE_POSE_SEARCH_EIGEN_DEBUG 0
+#endif
+
+// Eigen debugging configuration
+#if UE_POSE_SEARCH_EIGEN_DEBUG
+#define EIGEN_INITIALIZE_MATRICES_BY_NAN
+#define EIGEN_NO_AUTOMATIC_RESIZING
+#ifdef EIGEN_NO_DEBUG
+#undef EIGEN_NO_DEBUG
+#endif
+#ifdef EIGEN_NO_STATIC_ASSERT
+#undef EIGEN_NO_STATIC_ASSERT
+#endif
+#endif // UE_POSE_SEARCH_EIGEN_DEBUG
+
+// Disbable static analysis warnings
+#if defined(_MSC_VER) && USING_CODE_ANALYSIS
+#pragma warning(push)
+#pragma warning(disable:6294) // Ill-defined for-loop:  initial condition does not satisfy test.  Loop body not executed.
+#endif
+
+#ifndef EIGEN_MPL2_ONLY
+#define EIGEN_MPL2_ONLY
+#endif
+
+PRAGMA_DEFAULT_VISIBILITY_START
+THIRD_PARTY_INCLUDES_START
+#include <Eigen/Core>
+#include <Eigen/Eigenvalues>
+THIRD_PARTY_INCLUDES_END
+PRAGMA_DEFAULT_VISIBILITY_END
+
+// Restore static analysis warnings
+#if defined(_MSC_VER) && USING_CODE_ANALYSIS
+#pragma warning(pop)
+#endif
+
+#if UE_POSE_SEARCH_EIGEN_DEBUG
+#include <sstream>
+#endif
+
+namespace UE { namespace PoseSearch {
+
+#if UE_POSE_SEARCH_EIGEN_DEBUG
+template<typename EigenDenseBaseDerivedType>
+FString EigenMatrixToString(const Eigen::DenseBase<EigenDenseBaseDerivedType>& Matrix){
+	std::stringstream StringStream;
+    StringStream << Matrix;
+	return StringStream.str().c_str();
+}
+#endif
+
+}}
\ No newline at end of file
diff --git a/Engine/Plugins/Experimental/Animation/PoseSearch/Source/Runtime/Public/PoseSearch/AnimNode_PoseSearchHistoryCollector.h b/Engine/Plugins/Experimental/Animation/PoseSearch/Source/Runtime/Public/PoseSearch/AnimNode_PoseSearchHistoryCollector.h
--- a/Engine/Plugins/Experimental/Animation/PoseSearch/Source/Runtime/Public/PoseSearch/AnimNode_PoseSearchHistoryCollector.h
+++ b/Engine/Plugins/Experimental/Animation/PoseSearch/Source/Runtime/Public/PoseSearch/AnimNode_PoseSearchHistoryCollector.h
@@ -33,5 +33,4 @@ protected:
 
 	UE::PoseSearch::FPoseHistory PoseHistory;
 	TArray<float> Query;
-	float EvalDeltaTime = 0.0f;
 };
\ No newline at end of file
diff --git a/Engine/Plugins/Experimental/Animation/PoseSearch/Source/Runtime/Public/PoseSearch/PoseSearch.h b/Engine/Plugins/Experimental/Animation/PoseSearch/Source/Runtime/Public/PoseSearch/PoseSearch.h
--- a/Engine/Plugins/Experimental/Animation/PoseSearch/Source/Runtime/Public/PoseSearch/PoseSearch.h
+++ b/Engine/Plugins/Experimental/Animation/PoseSearch/Source/Runtime/Public/PoseSearch/PoseSearch.h
@@ -77,6 +77,8 @@ struct POSESEARCH_API FPoseSearchFeatureDesc
 
 	bool operator==(const FPoseSearchFeatureDesc& Other) const;
 
+	bool IsSubsampleOfSameFeature(const FPoseSearchFeatureDesc& Other) const;
+
 	FPoseSearchFeatureDesc()
 		: SchemaBoneIdx(MAX_int32)
 		, SubsampleIdx(MAX_int32)
@@ -105,6 +107,18 @@ struct POSESEARCH_API FPoseSearchFeatureVectorLayout
 	UPROPERTY()
 	uint32 NumFloats = 0;
 
+	UPROPERTY()
+	uint32 NumTrajectoryValues = 0;
+
+	UPROPERTY()
+	int32 FirstTrajectoryValueOffset = -1;
+
+	UPROPERTY()
+	uint32 NumPoseValues = 0;
+
+	UPROPERTY()
+	int32 FirstPoseValueOffset = -1;
+
 	bool IsValid(int32 MaxNumBones) const;
 };
 
@@ -126,6 +140,16 @@ struct POSESEARCH_API FPoseSearchFeatureVectorLayout
 //	int32 RelativeBoneIdx = -1;
 //};
 
+UENUM()
+enum class EPoseSearchDataPreprocessor : int32
+{
+	None			= 0,
+	Automatic		= 1,
+	Normalize    	= 2,
+	Sphere			= 3 UMETA(Hidden),
+
+	Invalid = -1 UMETA(Hidden)
+};
 
 /**
 * Specifies the format of a pose search index. At runtime, queries are built according to the schema for searching.
@@ -140,22 +164,22 @@ public:
 	static constexpr int32 DefaultSampleRate = 10;
 
 	UPROPERTY(EditAnywhere, Category = "Schema")
-	USkeleton* Skeleton;
+	USkeleton* Skeleton = nullptr;
 
 	UPROPERTY(EditAnywhere, meta = (ClampMin = "1", ClampMax = "60"), Category = "Schema")
-	int32 SampleRate;
+	int32 SampleRate = DefaultSampleRate;
 
 	UPROPERTY(EditAnywhere, Category = "Schema")
-	bool bUseBoneVelocities;
+	bool bUseBoneVelocities = true;
 
 	UPROPERTY(EditAnywhere, Category = "Schema")
-	bool bUseBonePositions;
+	bool bUseBonePositions = true;
 
 	UPROPERTY(EditAnywhere, Category = "Schema")
-	bool bUseTrajectoryVelocities;
+	bool bUseTrajectoryVelocities = true;
 
 	UPROPERTY(EditAnywhere, Category = "Schema")
-	bool bUseTrajectoryPositions;
+	bool bUseTrajectoryPositions = true;
 
 	UPROPERTY(EditAnywhere, Category = "Schema")
 	TArray<FBoneReference> Bones;
@@ -174,6 +198,9 @@ public:
 	UPROPERTY(EditAnywhere, Category = "Schema")
 	TArray<float> TrajectorySampleDistances;
 
+	UPROPERTY(EditAnywhere, Category = "Schema")
+	EPoseSearchDataPreprocessor DataPreprocessor = EPoseSearchDataPreprocessor::Automatic;
+
 	UPROPERTY()
 	TArray<int32> TrajectorySampleOffsets;
 
@@ -181,7 +208,10 @@ public:
 	TArray<int32> PoseSampleOffsets;
 
 	UPROPERTY()
-	float SamplingInterval;
+	EPoseSearchDataPreprocessor EffectiveDataPreprocessor = EPoseSearchDataPreprocessor::Invalid;
+
+	UPROPERTY()
+	float SamplingInterval = 1.0f / DefaultSampleRate;
 
 	UPROPERTY()
 	FPoseSearchFeatureVectorLayout Layout;
@@ -192,19 +222,12 @@ public:
 	UPROPERTY(Transient)
 	TArray<uint16> BoneIndicesWithParents;
 
-	UPoseSearchSchema()
-		: SampleRate(DefaultSampleRate)
-	{}
 
 	bool IsValid () const;
 
 	int32 NumBones () const { return BoneIndices.Num(); }
 
 public: // UObject
-	PRAGMA_DISABLE_DEPRECATION_WARNINGS // Suppress compiler warning on override of deprecated function
-	UE_DEPRECATED(5.0, "Use version that takes FObjectPreSaveContext instead.")
-	virtual void PreSave(const class ITargetPlatform* TargetPlatform) override;
-	PRAGMA_ENABLE_DEPRECATION_WARNINGS
 	virtual void PreSave(FObjectPreSaveContext ObjectSaveContext) override;
 	virtual void PostLoad() override;
 
@@ -217,6 +240,32 @@ private:
 	void ConvertTimesToOffsets(TArrayView<const float> SampleTimes, TArray<int32>& OutSampleOffsets);
 };
 
+USTRUCT()
+struct POSESEARCH_API FPoseSearchIndexPreprocessInfo
+{
+	GENERATED_BODY()
+
+	UPROPERTY()
+	uint32 NumDimensions = 0;
+
+	UPROPERTY()
+	TArray<float> TransformationMatrix;
+
+	UPROPERTY()
+	TArray<float> InverseTransformationMatrix;
+
+	UPROPERTY()
+	TArray<float> SampleMean;
+
+	void Reset()
+	{
+		NumDimensions = 0;
+		TransformationMatrix.Reset();
+		InverseTransformationMatrix.Reset();
+		SampleMean.Reset();
+	}
+};
+
 /**
 * A search index for animation poses. The structure of the search index is determined by its UPoseSearchSchema.
 * May represent a single animation (see UPoseSearchSequenceMetaData) or a collection (see UPoseSearchDatabase).
@@ -225,7 +274,7 @@ USTRUCT()
 struct POSESEARCH_API FPoseSearchIndex
 {
 	GENERATED_BODY()
-public:
+
 	UPROPERTY()
 	int32 NumPoses = 0;
 
@@ -235,11 +284,17 @@ public:
 	UPROPERTY()
 	const UPoseSearchSchema* Schema = nullptr;
 
+	UPROPERTY()
+	FPoseSearchIndexPreprocessInfo PreprocessInfo;
+
 	bool IsValid() const;
 
 	TArrayView<const float> GetPoseValues(int32 PoseIdx) const;
 
 	void Reset();
+
+	void Normalize (TArrayView<float> PoseVector) const;
+	void InverseNormalize (TArrayView<float> PoseVector) const;
 };
 
 /** Animation metadata object for indexing a single animation. */
@@ -250,7 +305,7 @@ class POSESEARCH_API UPoseSearchSequenceMetaData : public UAnimMetaData
 public:
 
 	UPROPERTY(EditAnywhere, Category="Settings")
-	const UPoseSearchSchema* Schema;
+	const UPoseSearchSchema* Schema = nullptr;
 
 	UPROPERTY(EditAnywhere, Category="Settings")
 	FFloatInterval SamplingRange = FFloatInterval(0.0f, 0.0f);
@@ -262,10 +317,6 @@ public:
 	bool IsValidForSearch() const;
 
 public: // UObject
-	PRAGMA_DISABLE_DEPRECATION_WARNINGS // Suppress compiler warning on override of deprecated function
-	UE_DEPRECATED(5.0, "Use version that takes FObjectPreSaveContext instead.")
-	virtual void PreSave(const class ITargetPlatform* TargetPlatform) override;
-	PRAGMA_ENABLE_DEPRECATION_WARNINGS
 	virtual void PreSave(FObjectPreSaveContext ObjectSaveContext) override;
 };
 
@@ -336,10 +387,6 @@ public:
 	bool IsValidForSearch() const;
 
 public: // UObject
-	PRAGMA_DISABLE_DEPRECATION_WARNINGS // Suppress compiler warning on override of deprecated function
-	UE_DEPRECATED(5.0, "Use version that takes FObjectPreSaveContext instead.")
-	virtual void PreSave(const class ITargetPlatform* TargetPlatform) override;
-	PRAGMA_ENABLE_DEPRECATION_WARNINGS
 	virtual void PreSave(FObjectPreSaveContext ObjectSaveContext) override;
 };
 
@@ -359,6 +406,7 @@ public:
 	void ResetFeatures();
 
 	TArrayView<const float> GetValues() const { return Values; }
+	TArrayView<const float> GetNormalizedValues() const { return ValuesNormalized; }
 
 	void SetTransform(FPoseSearchFeatureDesc Feature, const FTransform& Transform);
 	void SetTransformDerivative(FPoseSearchFeatureDesc Feature, const FTransform& Transform, const FTransform& PrevTransform, float DeltaTime);
@@ -370,7 +418,7 @@ public:
 	bool SetPoseFeatures(UE::PoseSearch::FPoseHistory* History);
 	bool SetPastTrajectoryFeatures(UE::PoseSearch::FPoseHistory* History);
 
-	void Copy(TArrayView<const float> FeatureVector);
+	void CopyFromSearchIndex(const FPoseSearchIndex& SearchIndex, int32 PoseIdx);
 	void CopyFeature(const FPoseSearchFeatureVectorBuilder& OtherBuilder, int32 FeatureIdx);
 
 	void MergeReplace(const FPoseSearchFeatureVectorBuilder& OtherBuilder);
@@ -379,11 +427,14 @@ public:
 	bool IsComplete() const;
 	bool IsCompatible(const FPoseSearchFeatureVectorBuilder& OtherBuilder) const;
 
+	void Normalize(const FPoseSearchIndex& ForSearchIndex);
+
 private:
 	UPROPERTY(Transient)
 	const UPoseSearchSchema* Schema = nullptr;
 
 	TArray<float> Values;
+	TArray<float> ValuesNormalized;
 	TBitArray<> FeaturesAdded;
 	int32 NumFeaturesAdded = 0;
 };
