PoseSearch: Fixed clang complaint about misleading indentation in Eigen library
#rnx

diff --git a/Engine/Plugins/Experimental/Animation/PoseSearch/Source/Runtime/Private/PoseSearchEigenHelper.h b/Engine/Plugins/Experimental/Animation/PoseSearch/Source/Runtime/Private/PoseSearchEigenHelper.h
--- a/Engine/Plugins/Experimental/Animation/PoseSearch/Source/Runtime/Private/PoseSearchEigenHelper.h
+++ b/Engine/Plugins/Experimental/Animation/PoseSearch/Source/Runtime/Private/PoseSearchEigenHelper.h
@@ -12,7 +12,7 @@
 
 // Eigen debugging configuration
 #if UE_POSE_SEARCH_EIGEN_DEBUG
-#define EIGEN_INITIALIZE_MATRICES_BY_NAN
+//#define EIGEN_INITIALIZE_MATRICES_BY_NAN
 #define EIGEN_NO_AUTOMATIC_RESIZING
 #ifdef EIGEN_NO_DEBUG
 #undef EIGEN_NO_DEBUG
