Pose search weighting biases, primarily for motion matching

diff --git a/Engine/Plugins/Experimental/Animation/PoseSearch/Source/Runtime/Private/AnimNode_MotionMatching.cpp b/Engine/Plugins/Experimental/Animation/PoseSearch/Source/Runtime/Private/AnimNode_MotionMatching.cpp
--- a/Engine/Plugins/Experimental/Animation/PoseSearch/Source/Runtime/Private/AnimNode_MotionMatching.cpp
+++ b/Engine/Plugins/Experimental/Animation/PoseSearch/Source/Runtime/Private/AnimNode_MotionMatching.cpp
@@ -81,15 +81,21 @@ void FAnimNode_MotionMatching::Update_AnyThread(const FAnimationUpdateContext& C
 		// Update features in the query with the latest inputs
 		ComposeQuery(Context);
 
+		// Initialize the pose search query bias weights context
+		FPoseSearchBiasWeights QueryBiasWeights;
+		QueryBiasWeights.Init(BiasWeights, Database->Schema->Layout);
+
+		const FPoseSearchBiasWeightsContext BiasWeightsContext = { &QueryBiasWeights, Database };
+
 		// Determine how much the updated query vector deviates from the current pose vector
 		float CurrentDissimilarity = MAX_flt;
 		if (DbPoseIdx != INDEX_NONE)
 		{
-			CurrentDissimilarity = UE::PoseSearch::ComparePoses(Database->SearchIndex, DbPoseIdx, ComposedQuery.GetNormalizedValues());
+			CurrentDissimilarity = UE::PoseSearch::ComparePoses(Database->SearchIndex, DbPoseIdx, ComposedQuery.GetNormalizedValues(), &BiasWeightsContext);
 		}
 
 		// Search the database for the nearest match to the updated query vector
-		UE::PoseSearch::FDbSearchResult Result = UE::PoseSearch::Search(Database, ComposedQuery.GetNormalizedValues());
+		UE::PoseSearch::FDbSearchResult Result = UE::PoseSearch::Search(Database, ComposedQuery.GetNormalizedValues(), &BiasWeightsContext);
 		if (Result.IsValid() && (ElapsedPoseJumpTime >= SearchThrottleTime))
 		{
 			const FPoseSearchDatabaseSequence& ResultDbSequence = Database->Sequences[Result.DbSequenceIdx];
diff --git a/Engine/Plugins/Experimental/Animation/PoseSearch/Source/Runtime/Private/PoseSearch.cpp b/Engine/Plugins/Experimental/Animation/PoseSearch/Source/Runtime/Private/PoseSearch.cpp
--- a/Engine/Plugins/Experimental/Animation/PoseSearch/Source/Runtime/Private/PoseSearch.cpp
+++ b/Engine/Plugins/Experimental/Animation/PoseSearch/Source/Runtime/Private/PoseSearch.cpp
@@ -54,13 +54,13 @@ static FFloatInterval GetEffectiveSamplingRange(const UAnimSequenceBase* Sequenc
 	return Range;
 }
 
-static float CompareFeatureVectors(int32 NumValues, const float* A, const float* B)
+static inline float CompareFeatureVectors(int32 NumValues, const float* A, const float* B, const float* Weights)
 {
-	float Dissimilarity = 0.0f;
+	float Dissimilarity = 0.f;
 
 	for (int32 ValueIdx = 0; ValueIdx != NumValues; ++ValueIdx)
 	{
-		float Diff = A[ValueIdx] - B[ValueIdx];
+		const float Diff = Weights[ValueIdx] * (A[ValueIdx] - B[ValueIdx]);
 		Dissimilarity += Diff * Diff;
 	}
 
@@ -251,6 +251,24 @@ bool FPoseSearchFeatureVectorLayout::IsValid(int32 MaxNumBones) const
 	return true;
 }
 
+bool FPoseSearchFeatureVectorLayout::EnumerateFeature(EPoseSearchFeatureType FeatureType, bool bTrajectory, int32& OutFeatureIdx) const
+{
+	// This function behaves similar to a generator
+	// OutFeatureIdx represents the 'next' Pose Feature Description that matches the inner criteria
+	// OutFeatureIdx can then be used again as a starting index to begin a subsequent search
+	for (int32 Idx = OutFeatureIdx + 1, Size = Features.Num(); Idx < Size; ++Idx)
+	{
+		// A trajectory feature match will result when bTrajectory = True and SchemaBoneIdx = -1
+		// A pose feature match will result when bTrajectory = False and SchemaBoneIdx != -1
+		if (Features[Idx].Type == FeatureType && (bTrajectory == (Features[Idx].SchemaBoneIdx == FPoseSearchFeatureDesc::TrajectoryBoneIndex)))
+		{
+			OutFeatureIdx = Idx;
+			return true;
+		}
+	}
+
+	return false;
+}
 
 //////////////////////////////////////////////////////////////////////////
 // UPoseSearchSchema
@@ -426,6 +444,37 @@ void UPoseSearchSchema::ConvertTimesToOffsets(TArrayView<const float> SampleTime
 }
 
 
+//////////////////////////////////////////////////////////////////////////
+// FPoseSearchBiasWeights
+
+void FPoseSearchBiasWeights::Init(const FPoseSearchBiasWeightParams& WeightParams, const FPoseSearchFeatureVectorLayout& Layout)
+{
+	// Initialize all weights to a default value of 1, and subsequently override all bound weights to their assigned value
+	Weights.Init(1.f, Layout.NumFloats);
+	BindSemanticWeight(WeightParams.TrajectoryPositionWeight, Layout, EPoseSearchFeatureType::Position, true);
+	BindSemanticWeight(WeightParams.TrajectoryLinearVelocityWeight, Layout, EPoseSearchFeatureType::LinearVelocity, true);
+	BindSemanticWeight(WeightParams.PosePositionWeight, Layout, EPoseSearchFeatureType::Position, false);
+	BindSemanticWeight(WeightParams.PoseLinearVelocityWeight, Layout, EPoseSearchFeatureType::LinearVelocity, false);
+}
+
+void FPoseSearchBiasWeights::BindSemanticWeight(float Weight, const FPoseSearchFeatureVectorLayout& Layout, EPoseSearchFeatureType FeatureType, bool bTrajectory)
+{
+	// The Weight parameter will be bound to a specific feature described by the FPoseSearchFeatureVectorLayout
+	int32 FeatureIdx = -1;
+	while (Layout.EnumerateFeature(FeatureType, bTrajectory, FeatureIdx))
+	{
+		const FPoseSearchFeatureDesc& Feature = Layout.Features[FeatureIdx];
+		const int32 FirstValueIdx = Feature.ValueOffset;
+		const int32 NumValues = UE::PoseSearch::GetFeatureTypeTraits(FeatureType).NumFloats;
+
+		for (int32 Idx = 0; Idx < NumValues; ++Idx)
+		{
+			Weights[FirstValueIdx + Idx] = Weight;
+		}
+	}
+}
+
+
 //////////////////////////////////////////////////////////////////////////
 // FPoseSearchIndex
 
@@ -2427,7 +2476,6 @@ bool BuildIndex(UPoseSearchDatabase* Database)
 	// Sample animations independently
 	ParallelFor(SequenceSamplers.Num(), [&SequenceSamplers](int32 SamplerIdx){ SequenceSamplers[SamplerIdx].Process(); });
 
-
 	auto GetSampler = [&](const UAnimSequence* Sequence) -> const FSequenceSampler*
 	{
 		return Sequence ? &SequenceSamplers[SequenceSamplerMap[Sequence]] : nullptr;
@@ -2467,6 +2515,20 @@ bool BuildIndex(UPoseSearchDatabase* Database)
 		TotalFloats += Output.FeatureVectorTable.Num();
 	}
 
+	// Establish per-sequence pose search bias weights if metadata is present
+	for (int32 SequenceIdx = 0; SequenceIdx != Database->Sequences.Num(); ++SequenceIdx)
+	{
+		FPoseSearchDatabaseSequence& DbSequence = Database->Sequences[SequenceIdx];
+		if (DbSequence.Sequence)
+		{
+			const UPoseSearchSequenceBiasWeightMetaData* BiasWeightMetadata = DbSequence.Sequence->FindMetaDataByClass<UPoseSearchSequenceBiasWeightMetaData>();
+			if (BiasWeightMetadata)
+			{
+				DbSequence.BiasWeights.Init(BiasWeightMetadata->BiasWeights, Database->Schema->Layout);
+			}
+		}
+	}
+
 	// Join animation data into a single search index
 	Database->SearchIndex.Values.Reset(TotalFloats);
 	for (const FSequenceIndexer& Indexer : Indexers)
@@ -2483,10 +2545,23 @@ bool BuildIndex(UPoseSearchDatabase* Database)
 	return true;
 }
 
-static FSearchResult Search(const FPoseSearchIndex& SearchIndex, TArrayView<const float> Query)
+static inline void DefaultInitializeWeights(const FPoseSearchBiasWeightsContext* BiasWeightsContext, int32 Size, TArray<float>& Weights, bool& bBiasWeightContextAvailable)
 {
-	FSearchResult Result;
+	bBiasWeightContextAvailable = BiasWeightsContext && BiasWeightsContext->HasBiasWeights();
+
+	if (bBiasWeightContextAvailable)
+	{
+		Weights = BiasWeightsContext->BiasWeights->Weights;
+	}
+	else
+	{
+		Weights.Init(1.f, Size);
+	}
+}
 
+static FSearchResult Search(const FPoseSearchIndex& SearchIndex, TArrayView<const float> Query, const FPoseSearchBiasWeightsContext* BiasWeightsContext = nullptr)
+{
+	FSearchResult Result;
 	if (!ensure(SearchIndex.IsValid()))
 	{
 		return Result;
@@ -2497,19 +2572,61 @@ static FSearchResult Search(const FPoseSearchIndex& SearchIndex, TArrayView<cons
 		return Result;
 	}
 
+	bool bBiasWeightContextAvailable = false;
+
+	// Initial weights by default are set to 1, but may be independently set by an external system such as motion matching
+	TArray<float> InitialWeights;
+	DefaultInitializeWeights(BiasWeightsContext, Query.Num(), InitialWeights, bBiasWeightContextAvailable);
+	const auto InitialWeightsMap = Eigen::Map<const Eigen::ArrayXf>(InitialWeights.GetData(), InitialWeights.Num());
+
+	// Accumulated weights will contain the per-pose final weights, optionally including per-sequence and/or other external values
+	TArray<float> AccumulatedWeights;
+	AccumulatedWeights = InitialWeights;
+	auto AccumulatedWeightsMap = Eigen::Map<Eigen::ArrayXf>(AccumulatedWeights.GetData(), AccumulatedWeights.Num());
+
 	float BestPoseDissimilarity = MAX_flt;
 	int32 BestPoseIdx = INDEX_NONE;
 
-	for (int32 PoseIdx = 0; PoseIdx != SearchIndex.NumPoses; ++PoseIdx)
+	for (int32 PoseIdx = 0, PrevSequenceIdx = -1; PoseIdx != SearchIndex.NumPoses; ++PoseIdx)
 	{
+		// Sequence index and metadata tracking are done within this loop in order to optimize
+		// and elide unnecessary recomputing of the weight buffer.
+		bool bSequenceWeightsAvailable = false;
+		int32_t SequenceIdx = -1;
+
+		if (bBiasWeightContextAvailable)
+		{
+			// Apply the per-sequence bias weights if they are present within the sequence metadata
+			SequenceIdx = BiasWeightsContext->Database->FindSequenceForPose(PoseIdx);
+			if (SequenceIdx != PrevSequenceIdx)
+			{
+				const FPoseSearchDatabaseSequence& SequenceEntry = BiasWeightsContext->Database->Sequences[SequenceIdx];
+				bSequenceWeightsAvailable = SequenceEntry.BiasWeights.IsInitialized();
+
+				if (bSequenceWeightsAvailable)
+				{
+					const auto& SequenceBiasWeights = SequenceEntry.BiasWeights.Weights;
+					AccumulatedWeightsMap *= Eigen::Map<const Eigen::ArrayXf>(SequenceBiasWeights.GetData(), SequenceBiasWeights.Num());
+				}
+			}
+		}
+
 		const int32 FeatureValueOffset = PoseIdx * SearchIndex.Schema->Layout.NumFloats;
 
 		float PoseDissimilarity = CompareFeatureVectors(
 			SearchIndex.Schema->Layout.NumFloats,
 			Query.GetData(),
-			&SearchIndex.Values[FeatureValueOffset]
+			&SearchIndex.Values[FeatureValueOffset],
+			AccumulatedWeights.GetData()
 		);
 
+		if (bSequenceWeightsAvailable && SequenceIdx != PrevSequenceIdx)
+		{
+			// Reset pose weights to remove any extraneous sequence contributions for the next iteration
+			AccumulatedWeightsMap = InitialWeightsMap;
+			PrevSequenceIdx = SequenceIdx;
+		}
+		
 		if (PoseDissimilarity < BestPoseDissimilarity)
 		{
 			BestPoseDissimilarity = PoseDissimilarity;
@@ -2554,7 +2671,7 @@ FSearchResult Search(const UAnimSequenceBase* Sequence, TArrayView<const float>
 	return Result;
 }
 
-FDbSearchResult Search(const UPoseSearchDatabase* Database, TArrayView<const float> Query, FDebugDrawParams DebugDrawParams)
+FDbSearchResult Search(const UPoseSearchDatabase* Database, TArrayView<const float> Query, const FPoseSearchBiasWeightsContext* BiasWeightsContext, FDebugDrawParams DebugDrawParams)
 {
 	if (!ensure(Database && Database->IsValidForSearch()))
 	{
@@ -2563,7 +2680,7 @@ FDbSearchResult Search(const UPoseSearchDatabase* Database, TArrayView<const flo
 
 	const FPoseSearchIndex& SearchIndex = Database->SearchIndex;
 
-	FDbSearchResult Result = Search(SearchIndex, Query);
+	FDbSearchResult Result = Search(SearchIndex, Query, BiasWeightsContext);
 	if (!Result.IsValid())
 	{
 		return FDbSearchResult();
@@ -2590,11 +2707,30 @@ FDbSearchResult Search(const UPoseSearchDatabase* Database, TArrayView<const flo
 	return Result;
 }
 
-float ComparePoses(const FPoseSearchIndex& SearchIndex, int32 PoseIdx, TArrayView<const float> Query)
+float ComparePoses(const FPoseSearchIndex& SearchIndex, int32 PoseIdx, TArrayView<const float> Query, const FPoseSearchBiasWeightsContext* BiasWeightsContext)
 {
 	TArrayView<const float> PoseValues = SearchIndex.GetPoseValues(PoseIdx);
 	check(PoseValues.Num() == Query.Num());
-	return CompareFeatureVectors(PoseValues.Num(), PoseValues.GetData(), Query.GetData());
+
+	bool bBiasWeightContextAvailable = false;
+
+	TArray<float> Weights;
+	DefaultInitializeWeights(BiasWeightsContext, Query.Num(), Weights, bBiasWeightContextAvailable);
+
+	if (bBiasWeightContextAvailable)
+	{
+		// Apply the per-sequence bias weights if present on the sequence metadata
+		const int32 SequenceIdx = BiasWeightsContext->Database->FindSequenceForPose(PoseIdx);
+		const FPoseSearchDatabaseSequence& SequenceEntry = BiasWeightsContext->Database->Sequences[SequenceIdx];
+
+		if (SequenceEntry.BiasWeights.IsInitialized())
+		{
+			const auto& SequenceWeights = SequenceEntry.BiasWeights.Weights;
+			Eigen::Map<Eigen::ArrayXf>(Weights.GetData(), Weights.Num()) *= Eigen::Map<const Eigen::ArrayXf>(SequenceWeights.GetData(), SequenceWeights.Num());
+		}
+	}
+
+	return CompareFeatureVectors(PoseValues.Num(), PoseValues.GetData(), Query.GetData(), Weights.GetData());
 }
 
 
diff --git a/Engine/Plugins/Experimental/Animation/PoseSearch/Source/Runtime/Public/PoseSearch/AnimNode_MotionMatching.h b/Engine/Plugins/Experimental/Animation/PoseSearch/Source/Runtime/Public/PoseSearch/AnimNode_MotionMatching.h
--- a/Engine/Plugins/Experimental/Animation/PoseSearch/Source/Runtime/Public/PoseSearch/AnimNode_MotionMatching.h
+++ b/Engine/Plugins/Experimental/Animation/PoseSearch/Source/Runtime/Public/PoseSearch/AnimNode_MotionMatching.h
@@ -20,6 +20,9 @@ public:
 	UPROPERTY(EditAnywhere, BlueprintReadWrite, Category=Settings, meta=(PinHiddenByDefault))
 	const UPoseSearchDatabase* Database = nullptr;
 
+	UPROPERTY(EditAnywhere, BlueprintReadWrite, Category = Settings, meta = (PinHiddenByDefault))
+	FPoseSearchBiasWeightParams BiasWeights;
+
 	// Motion matching goal
 	UPROPERTY(EditAnywhere, BlueprintReadWrite, Category=Settings, meta=(PinShownByDefault))
 	FPoseSearchFeatureVectorBuilder Goal;
diff --git a/Engine/Plugins/Experimental/Animation/PoseSearch/Source/Runtime/Public/PoseSearch/PoseSearch.h b/Engine/Plugins/Experimental/Animation/PoseSearch/Source/Runtime/Public/PoseSearch/PoseSearch.h
--- a/Engine/Plugins/Experimental/Animation/PoseSearch/Source/Runtime/Public/PoseSearch/PoseSearch.h
+++ b/Engine/Plugins/Experimental/Animation/PoseSearch/Source/Runtime/Public/PoseSearch/PoseSearch.h
@@ -120,6 +120,8 @@ struct POSESEARCH_API FPoseSearchFeatureVectorLayout
 	int32 FirstPoseValueOffset = -1;
 
 	bool IsValid(int32 MaxNumBones) const;
+
+	bool EnumerateFeature(EPoseSearchFeatureType FeatureType, bool bTrajectory, int32& OutFeatureIdx) const;
 };
 
 
@@ -320,6 +322,65 @@ public: // UObject
 	virtual void PreSave(FObjectPreSaveContext ObjectSaveContext) override;
 };
 
+/** Bias weights for influencing the results of a pose search query */
+USTRUCT(BlueprintType, Category = "Animation|PoseSearch")
+struct POSESEARCH_API FPoseSearchBiasWeightParams
+{
+	GENERATED_BODY()
+
+	UPROPERTY(BlueprintReadWrite, EditAnywhere, Category = "Settings", meta = (ClampMin = "0.0"));
+	float TrajectoryLinearVelocityWeight = 1.f;
+
+	UPROPERTY(BlueprintReadWrite, EditAnywhere, Category = "Settings", meta = (ClampMin = "0.0"));
+	float TrajectoryPositionWeight = 1.f;
+
+	UPROPERTY(BlueprintReadWrite, EditAnywhere, Category = "Settings", meta = (ClampMin = "0.0"));
+	float PoseLinearVelocityWeight = 1.f;
+
+	UPROPERTY(BlueprintReadWrite, EditAnywhere, Category = "Settings", meta = (ClampMin = "0.0"));
+	float PosePositionWeight = 1.f;
+};
+
+/** Expanded bias weights from FPoseSearchBiasWeightParams which match the binary layout described by FPoseSearchFeatureVectorLayout */
+USTRUCT()
+struct POSESEARCH_API FPoseSearchBiasWeights
+{
+	GENERATED_BODY()
+
+	UPROPERTY()
+	TArray<float> Weights;
+
+	bool IsInitialized() const { return !Weights.IsEmpty(); }
+	void Init(const FPoseSearchBiasWeightParams& WeightParams, const FPoseSearchFeatureVectorLayout& Layout);
+
+private:
+	void BindSemanticWeight(float Weight, const FPoseSearchFeatureVectorLayout& Layout, EPoseSearchFeatureType FeatureType, bool bTrajectory);
+};
+
+/** 
+* Bias weight context contains a reference to global/external weights and a relevant pose search database
+* The database parameter is primarily used for locating per-sequence bias weights
+*/
+struct FPoseSearchBiasWeightsContext
+{
+	const FPoseSearchBiasWeights* BiasWeights = nullptr;
+	const UPoseSearchDatabase* Database = nullptr;
+
+	bool IsInitialized() const { return !!Database; }
+	bool HasBiasWeights() const { return !!BiasWeights && BiasWeights->IsInitialized(); }
+};
+
+/** Per-sequence bias weights for influencing the results of a pose search query */
+UCLASS(BlueprintType, Category = "Animation|PoseSearch")
+class POSESEARCH_API UPoseSearchSequenceBiasWeightMetaData : public UAnimMetaData
+{
+	GENERATED_BODY()
+public:
+
+	UPROPERTY(EditAnywhere, Category = "Settings")
+	FPoseSearchBiasWeightParams BiasWeights;
+};
+
 /** An entry in a UPoseSearchDatabase. */
 USTRUCT(BlueprintType, Category = "Animation|PoseSearch")
 struct POSESEARCH_API FPoseSearchDatabaseSequence
@@ -361,7 +422,8 @@ struct POSESEARCH_API FPoseSearchDatabaseSequence
 	UPROPERTY()
 	int32 NumPoses = 0;
 
-	int32 GetPoseIndexFromAssetTime(float AssetTime) const;
+	UPROPERTY()
+	FPoseSearchBiasWeights BiasWeights;
 };
 
 /** A data asset for indexing a collection of animation sequences. */
@@ -596,11 +658,12 @@ POSESEARCH_API bool BuildIndex(UPoseSearchDatabase* Database);
 /**
 * Performs a pose search on a sequence's UPoseSearchSequenceMetaData.
 *
-* @param Sequence		The sequence to search within. Must contain a UPoseSearchSequenceMetaData object.
-* @param Query			The pose query to search for. To build a query, see FFeatureVectorBuilder. Must have been built using the same schema as the UPoseSearchSequenceMetaData.
-* @param DrawParams		Visualization options
+* @param Sequence			The sequence to search within. Must contain a UPoseSearchSequenceMetaData object.
+* @param Query				The pose query to search for. To build a query, see FFeatureVectorBuilder. Must have been built using the same schema as the UPoseSearchSequenceMetaData.
+* @param BiasWeightsContext	Optional bias weight context used to influence pose search query results
+* @param DrawParams			Visualization options
 * 
-* @return				The pose in the sequence that most closely matches the Query.
+* @return					The pose in the sequence that most closely matches the Query.
 */
 POSESEARCH_API FSearchResult Search(const UAnimSequenceBase* Sequence, TArrayView<const float> Query, FDebugDrawParams DrawParams = FDebugDrawParams());
 
@@ -608,24 +671,26 @@ POSESEARCH_API FSearchResult Search(const UAnimSequenceBase* Sequence, TArrayVie
 /**
 * Performs a pose search on a UPoseSearchDatabase.
 *
-* @param Database		The database to search within
-* @param Query			The pose query to search for. To build a query, see FFeatureVectorBuilder. Must have been built using the same schema as the Database.
-* @param DrawParams		Visualization options
+* @param Database			The database to search within
+* @param Query				The pose query to search for. To build a query, see FFeatureVectorBuilder. Must have been built using the same schema as the Database.
+* @param BiasWeightsContext	Optional bias weight context used to influence pose search query results
+* @param DrawParams			Visualization options
 * 
-* @return				The pose in the database that most closely matches the Query.
+* @return					The pose in the database that most closely matches the Query.
 */
-POSESEARCH_API FDbSearchResult Search(const UPoseSearchDatabase* Database, TArrayView<const float> Query, FDebugDrawParams DrawParams = FDebugDrawParams());
+POSESEARCH_API FDbSearchResult Search(const UPoseSearchDatabase* Database, TArrayView<const float> Query, const FPoseSearchBiasWeightsContext* BiasWeightsContext = nullptr, FDebugDrawParams DrawParams = FDebugDrawParams());
 
 
 /**
 * Evaluate pose comparison metric between a pose in the search index and an input query
 * 
-* @param SearchIndex	The search index containing the pose to compare to the query
-* @param PoseIdx		The index of the pose in the search index to compare to the query
-* @param Query			The query to compare against. Must have been built using the same schema as the SearchIndex.
+* @param SearchIndex		The search index containing the pose to compare to the query
+* @param PoseIdx			The index of the pose in the search index to compare to the query
+* @param Query				The query to compare against. Must have been built using the same schema as the SearchIndex.
+* @param BiasWeightsContext Optional bias weight context used to influence pose comparison
 * 
-* @return Dissimilarity between the two poses
+* @return					Dissimilarity between the two poses
 */
-POSESEARCH_API float ComparePoses(const FPoseSearchIndex& SearchIndex, int32 PoseIdx, TArrayView<const float> Query);
+POSESEARCH_API float ComparePoses(const FPoseSearchIndex& SearchIndex, int32 PoseIdx, TArrayView<const float> Query, const FPoseSearchBiasWeightsContext* BiasWeightsContext = nullptr);
 
 }} // namespace UE::PoseSearch
\ No newline at end of file
