Pose Search weight biasing code cleanup
#trivial

diff --git a/Engine/Plugins/Experimental/Animation/PoseSearch/Source/Runtime/Private/PoseSearch.cpp b/Engine/Plugins/Experimental/Animation/PoseSearch/Source/Runtime/Private/PoseSearch.cpp
--- a/Engine/Plugins/Experimental/Animation/PoseSearch/Source/Runtime/Private/PoseSearch.cpp
+++ b/Engine/Plugins/Experimental/Animation/PoseSearch/Source/Runtime/Private/PoseSearch.cpp
@@ -460,7 +460,7 @@ void FPoseSearchBiasWeights::Init(const FPoseSearchBiasWeightParams& WeightParam
 void FPoseSearchBiasWeights::BindSemanticWeight(float Weight, const FPoseSearchFeatureVectorLayout& Layout, EPoseSearchFeatureType FeatureType, bool bTrajectory)
 {
 	// The Weight parameter will be bound to a specific feature described by the FPoseSearchFeatureVectorLayout
-	int32 FeatureIdx = -1;
+	int32 FeatureIdx = INDEX_NONE;
 	while (Layout.EnumerateFeature(FeatureType, bTrajectory, FeatureIdx))
 	{
 		const FPoseSearchFeatureDesc& Feature = Layout.Features[FeatureIdx];
@@ -2580,19 +2580,18 @@ static FSearchResult Search(const FPoseSearchIndex& SearchIndex, TArrayView<cons
 	const auto InitialWeightsMap = Eigen::Map<const Eigen::ArrayXf>(InitialWeights.GetData(), InitialWeights.Num());
 
 	// Accumulated weights will contain the per-pose final weights, optionally including per-sequence and/or other external values
-	TArray<float> AccumulatedWeights;
-	AccumulatedWeights = InitialWeights;
+	TArray<float> AccumulatedWeights(InitialWeights);
 	auto AccumulatedWeightsMap = Eigen::Map<Eigen::ArrayXf>(AccumulatedWeights.GetData(), AccumulatedWeights.Num());
 
 	float BestPoseDissimilarity = MAX_flt;
 	int32 BestPoseIdx = INDEX_NONE;
 
-	for (int32 PoseIdx = 0, PrevSequenceIdx = -1; PoseIdx != SearchIndex.NumPoses; ++PoseIdx)
+	for (int32 PoseIdx = 0, PrevSequenceIdx = INDEX_NONE; PoseIdx != SearchIndex.NumPoses; ++PoseIdx)
 	{
 		// Sequence index and metadata tracking are done within this loop in order to optimize
 		// and elide unnecessary recomputing of the weight buffer.
 		bool bSequenceWeightsAvailable = false;
-		int32_t SequenceIdx = -1;
+		int32_t SequenceIdx = INDEX_NONE;
 
 		if (bBiasWeightContextAvailable)
 		{
@@ -2613,7 +2612,7 @@ static FSearchResult Search(const FPoseSearchIndex& SearchIndex, TArrayView<cons
 
 		const int32 FeatureValueOffset = PoseIdx * SearchIndex.Schema->Layout.NumFloats;
 
-		float PoseDissimilarity = CompareFeatureVectors(
+		const float PoseDissimilarity = CompareFeatureVectors(
 			SearchIndex.Schema->Layout.NumFloats,
 			Query.GetData(),
 			&SearchIndex.Values[FeatureValueOffset],
