PoseSearch: Distance-based trajectory indexing
- Added support for distanced based trajectory indexing. This feature uses
  distance offsets to mark where samples are taken, rather than the usual time
  offsets.
- Refactored sampling and indexing process. Previously, poses and root motion
  were extracted at the known fixed sampling rate, then the indexing process
  looked up sampled data based on integer offsets. Now the indexing process
  extracts poses and root motion at time offsets directly from the source data.
  This made distance sampling simpler to implement, improves the accuracy of
  velocity sampling, and paves the way for acceleration sampling.
- Velocity and acceleration sampling are implemented as symmetric difference
  quotients. This means for every sample at time t, we also take samples at t-h
  and t+h. So three pose and root motion extractions are taken at each sample.
- Fixed subtle off-by-one sample bugs. Previously the sampling pre-process phase
  did not guarantee that the end of a clip would be sampled. Now this is
  guaranteed and index to time conversions clamp to play length.
- Fixed bug where a zero trajectory time offset would cause the motion matching
  goal to not be filled with the current root motion properties.
- Fixed normalization divide by zero bug when the variance of a feature is zero.
- Fixed pose matching not using normalized search queries.
- Added motion matching debug vis options to view goal and match separately.

diff --git a/Engine/Plugins/Experimental/Animation/AnimationWarping/Source/Runtime/Private/AnimationWarpingRuntimeModule.cpp b/Engine/Plugins/Experimental/Animation/AnimationWarping/Source/Runtime/Private/AnimationWarpingRuntimeModule.cpp
--- a/Engine/Plugins/Experimental/Animation/AnimationWarping/Source/Runtime/Private/AnimationWarpingRuntimeModule.cpp
+++ b/Engine/Plugins/Experimental/Animation/AnimationWarping/Source/Runtime/Private/AnimationWarpingRuntimeModule.cpp
@@ -42,7 +42,8 @@ void FModule::ShutdownModule()
 void FModule::SampleRootMotion(const FDeltaTimeRecord& SampleRange, const UAnimSequence& Sequence, bool bLoopingSequence, UE::Anim::FStackAttributeContainer& OutAttributes) const
 {
 	const FTransform RootMotionTransform = Sequence.ExtractRootMotion(SampleRange.Previous, SampleRange.Delta, bLoopingSequence);
-	OutAttributes.Add<FTransformAnimationAttribute>(RootMotionAttributeId, { RootMotionTransform });
+	FTransformAnimationAttribute* RootMotionAttribute = OutAttributes.FindOrAdd<FTransformAnimationAttribute>(RootMotionAttributeId);
+	RootMotionAttribute->Value = RootMotionTransform;
 }
 
 bool FModule::OverrideRootMotion(const FTransform& RootMotionDelta, UE::Anim::FStackAttributeContainer& OutAttributes) const
diff --git a/Engine/Plugins/Experimental/Animation/PoseSearch/Source/Runtime/Private/AnimNode_MotionMatching.cpp b/Engine/Plugins/Experimental/Animation/PoseSearch/Source/Runtime/Private/AnimNode_MotionMatching.cpp
--- a/Engine/Plugins/Experimental/Animation/PoseSearch/Source/Runtime/Private/AnimNode_MotionMatching.cpp
+++ b/Engine/Plugins/Experimental/Animation/PoseSearch/Source/Runtime/Private/AnimNode_MotionMatching.cpp
@@ -177,16 +177,24 @@ void FAnimNode_MotionMatching::PreUpdate(const UAnimInstance* InAnimInstance)
 		check(SkeletalMeshComponent);
 
 		UE::PoseSearch::FDebugDrawParams DrawParams;
+		DrawParams.Flags = UE::PoseSearch::EDebugDrawFlags::None;
 		DrawParams.RootTransform = SkeletalMeshComponent->GetComponentTransform();
-		DrawParams.Flags = UE::PoseSearch::EDebugDrawFlags::DrawBest;
 		DrawParams.Database = Database;
 		DrawParams.Query = ComposedQuery.GetValues();
 		DrawParams.World = SkeletalMeshComponent->GetWorld();
 		DrawParams.DefaultLifeTime = 0.0f;
 		DrawParams.HighlightPoseIdx = DbPoseIdx;
-		UE::PoseSearch::Draw(DrawParams);
 
-		DrawParams.Flags = UE::PoseSearch::EDebugDrawFlags::DrawQuery;
+		if (bDebugDrawMatch)
+		{
+			DrawParams.Flags |= UE::PoseSearch::EDebugDrawFlags::DrawBest;
+		}
+
+		if (bDebugDrawGoal)
+		{
+			DrawParams.Flags |= UE::PoseSearch::EDebugDrawFlags::DrawQuery;
+		}
+
 		UE::PoseSearch::Draw(DrawParams);
 	}
 #endif
diff --git a/Engine/Plugins/Experimental/Animation/PoseSearch/Source/Runtime/Private/PoseSearch.cpp b/Engine/Plugins/Experimental/Animation/PoseSearch/Source/Runtime/Private/PoseSearch.cpp
--- a/Engine/Plugins/Experimental/Animation/PoseSearch/Source/Runtime/Private/PoseSearch.cpp
+++ b/Engine/Plugins/Experimental/Animation/PoseSearch/Source/Runtime/Private/PoseSearch.cpp
@@ -500,7 +500,7 @@ void FPoseSearchIndex::Reset()
 	Schema = nullptr;
 }
 
-void FPoseSearchIndex::Normalize (TArrayView<float> InOutPoseVector) const
+void FPoseSearchIndex::Normalize(TArrayView<float> InOutPoseVector) const
 {
 	using namespace Eigen;
 
@@ -527,7 +527,7 @@ void FPoseSearchIndex::Normalize (TArrayView<float> InOutPoseVector) const
 	PoseVector = TransformationMtx * (PoseVector - SampleMean);
 }
 
-void FPoseSearchIndex::InverseNormalize (TArrayView<float> InOutNormalizedPoseVector) const
+void FPoseSearchIndex::InverseNormalize(TArrayView<float> InOutNormalizedPoseVector) const
 {
 	using namespace Eigen;
 
@@ -707,7 +707,7 @@ void FPoseSearchFeatureVectorBuilder::SetTransform(FPoseSearchFeatureDesc Elemen
 	SetRotation(Element, Transform.GetRotation());
 }
 
-void FPoseSearchFeatureVectorBuilder::SetTransformDerivative(FPoseSearchFeatureDesc Element, const FTransform& Transform, const FTransform& PrevTransform, float DeltaTime)
+void FPoseSearchFeatureVectorBuilder::SetTransformVelocity(FPoseSearchFeatureDesc Element, const FTransform& Transform, const FTransform& PrevTransform, float DeltaTime)
 {
 	SetLinearVelocity(Element, Transform, PrevTransform, DeltaTime);
 	SetAngularVelocity(Element, Transform, PrevTransform, DeltaTime);
@@ -837,7 +837,7 @@ bool FPoseSearchFeatureVectorBuilder::SetPoseFeatures(UE::PoseSearch::FPoseHisto
 			const FTransform& Transform = ComponentPose[SkeletonBoneIndex];
 			const FTransform& PrevTransform = ComponentPrevPose[SkeletonBoneIndex];
 			SetTransform(Feature, Transform);
-			SetTransformDerivative(Feature, Transform, PrevTransform, History->GetSampleInterval());
+			SetTransformVelocity(Feature, Transform, PrevTransform, History->GetSampleInterval());
 		}
 	}
 
@@ -858,7 +858,7 @@ bool FPoseSearchFeatureVectorBuilder::SetPastTrajectoryFeatures(UE::PoseSearch::
 		Feature.SubsampleIdx = SchemaSubsampleIdx;
 
 		int32 SubsampleIndex = Schema->TrajectorySampleOffsets[SchemaSubsampleIdx];
-		if (SubsampleIndex >= 0)
+		if (SubsampleIndex > 0)
 		{
 			break;
 		}
@@ -877,7 +877,7 @@ bool FPoseSearchFeatureVectorBuilder::SetPastTrajectoryFeatures(UE::PoseSearch::
 		}
 
 		SetTransform(Feature, WorldComponentTransform);
-		SetTransformDerivative(Feature, WorldComponentTransform, WorldPrevComponentTransform, History->GetSampleInterval());
+		SetTransformVelocity(Feature, WorldComponentTransform, WorldPrevComponentTransform, History->GetSampleInterval());
 	}
 
 	return true;
@@ -1232,7 +1232,7 @@ bool FFeatureVectorReader::GetVector(FPoseSearchFeatureDesc Element, FVector* Ou
 //////////////////////////////////////////////////////////////////////////
 // FDebugDrawParams
 
-bool FDebugDrawParams::CanDraw () const
+bool FDebugDrawParams::CanDraw() const
 {
 	if (!World || !EnumHasAnyFlags(Flags, EDebugDrawFlags::DrawAll))
 	{
@@ -1290,34 +1290,29 @@ public:
 		const UPoseSearchSchema* Schema = nullptr;
 		const UAnimSequence* Sequence = nullptr;
 		bool bLoopable = false;
+		int32 DistanceSamplingRate = 30;
 	} Input;
 
 	struct FOutput
 	{
-		TArray<FTransform> ComponentSpacePose;		// Indexed by SampleIdx * NumBones + SchemaBoneIdx
-		TArray<FTransform> LocalRootMotion;			// Indexed by SampleIdx
-		TArray<FTransform> AccumulatedRootMotion;	// Indexed by SampleIdx
-		TArray<float> AccumulatedRootDistance;		// Indexed by SampleIdx
+		TArray<float> AccumulatedRootDistance;
 
-		int32 TotalSamples = 0;
+		int32 NumDistanceSamples = 0;
+		float PlayLength = 0.0f;
+		float TotalRootDistance = 0.0f;
+		FTransform TotalRootMotion = FTransform::Identity;
 	} Output;
 
 	void Reset();
 	void Init(const FInput& Input);
 	void Process();
 
-	struct FWrappedSampleIndex
-	{
-		int32 Idx = INDEX_NONE;
-		int32 NumCycles = 0;
-		bool bClamped = false;
-	};
-	FWrappedSampleIndex WrapOrClampSubsampleIndex (int32 SampleIdx) const;
+	FTransform ExtractRootTransform(float Time) const;
+	float ExtractRootDistance(float Time) const;
 
 private:
 	void Reserve();
-	void ExtractPoses();
-	void ExtractRootMotion();
+	void ProcessRootMotion();
 };
 
 void FSequenceSampler::Init(const FInput& InInput)
@@ -1330,8 +1325,8 @@ void FSequenceSampler::Init(const FInput& InInput)
 
 	Input = InInput;
 
-	const float SequencePlayLength = Input.Sequence->GetPlayLength();
-	Output.TotalSamples = FMath::FloorToInt(SequencePlayLength * Input.Schema->SampleRate);
+	Output.PlayLength = Input.Sequence->GetPlayLength();
+	Output.NumDistanceSamples = FMath::CeilToInt(Output.PlayLength * Input.DistanceSamplingRate) + 1;
 
 	Reserve();
 }
@@ -1340,123 +1335,139 @@ void FSequenceSampler::Reset()
 {
 	Input = FInput();
 
-	Output.TotalSamples = 0;
-	Output.ComponentSpacePose.Reset(0);
-	Output.LocalRootMotion.Reset(0);
-	Output.AccumulatedRootMotion.Reset(0);
+	Output.NumDistanceSamples = 0;
+	Output.PlayLength = 0.0f;
+	Output.TotalRootDistance = 0.0f;
 	Output.AccumulatedRootDistance.Reset(0);
 }
 
 void FSequenceSampler::Reserve()
 {
-	Output.ComponentSpacePose.Reserve(Input.Schema->NumBones() * Output.TotalSamples);
-	Output.LocalRootMotion.Reserve(Output.TotalSamples);
-	Output.AccumulatedRootMotion.Reserve(Output.TotalSamples);
-	Output.AccumulatedRootDistance.Reserve(Output.TotalSamples);
+	Output.AccumulatedRootDistance.Reserve(Output.NumDistanceSamples);
 }
 
 void FSequenceSampler::Process()
 {
-	ExtractPoses();
-	ExtractRootMotion();
+	ProcessRootMotion();
 }
 
-
-FSequenceSampler::FWrappedSampleIndex FSequenceSampler::WrapOrClampSubsampleIndex (int32 SampleIdx) const
+FTransform FSequenceSampler::ExtractRootTransform(float Time) const
 {
-	FWrappedSampleIndex Result;
-	Result.Idx = SampleIdx;
-	Result.NumCycles = 0;
-	Result.bClamped = false;
-
-	// Wrap the index if this is a loopable sequence
-	if (Input.bLoopable)
-	{
-		if (Result.Idx < 0)
-		{
-			Result.Idx += Output.TotalSamples;
+	return Input.Sequence->ExtractRootMotionFromRange(0.0f, Time);
+}
 
-			while (Result.Idx < 0)
-			{
-				Result.Idx += Output.TotalSamples;
-				++Result.NumCycles;
-			}
-		}
+float FSequenceSampler::ExtractRootDistance(float Time) const
+{
+	check(Time <= Output.PlayLength);
 
-		while (Result.Idx >= Output.TotalSamples)
-		{
-			Result.Idx -= Output.TotalSamples;
-			++Result.NumCycles;
-		}
-	}
+	// Find the distance sample that corresponds with the time and split into whole and partial parts
+	float IntegralDistanceSample;
+	float DistanceAlpha = FMath::Modf(Time * Input.DistanceSamplingRate, &IntegralDistanceSample);
+	float DistanceIdx = (int32)IntegralDistanceSample;
 
-	// Clamp if we can't loop
-	else if (SampleIdx < 0 || SampleIdx >= Output.TotalSamples)
+	// Verify the distance offset and any residual portion would be in bounds
+	check(DistanceIdx + (DistanceAlpha > 0.0f ? 1 : 0) < Output.AccumulatedRootDistance.Num());
+	
+	// Look up the distance and interpolate between distance samples if necessary
+	float Distance = Output.AccumulatedRootDistance[DistanceIdx];
+	if (DistanceAlpha > 0.0f)
 	{
-		Result.Idx = FMath::Clamp(SampleIdx, 0, Output.TotalSamples - 1);
-		Result.bClamped = true;
+		float NextDistance = Output.AccumulatedRootDistance[DistanceIdx + 1];
+		Distance = FMath::Lerp(Distance, NextDistance, DistanceAlpha);
 	}
 
-	return Result;
+	return Distance;
 }
 
-void FSequenceSampler::ExtractPoses()
+void FSequenceSampler::ProcessRootMotion()
 {
-	if (Input.Schema->Bones.IsEmpty())
+	// Note the distance sampling interval is independent of the schema's sampling interval
+	const float DistanceSamplingInterval = 1.0f / Input.DistanceSamplingRate;
+
+	const FTransform InitialRootTransform = Input.Sequence->ExtractRootTrackTransform(0.0f, nullptr);
+
+	// Build a distance lookup table by sampling root motion at a fixed rate and accumulating
+	// absolute translation deltas. During indexing we'll bsearch this table and interpolate
+	// between samples in order to convert distance offsets to time offsets.
+	// See also FSequenceIndexer::AddTrajectoryDistanceFeatures().
+	double AccumulatedRootDistance = 0.0;
+	FTransform LastRootTransform = InitialRootTransform;
+	float SampleTime = 0.0f;
+	for (int32 SampleIdx = 0; SampleIdx != Output.NumDistanceSamples; ++SampleIdx)
 	{
-		return;
+		SampleTime = FMath::Min(SampleIdx * DistanceSamplingInterval, Output.PlayLength);
+
+		FTransform RootTransform = Input.Sequence->ExtractRootTrackTransform(SampleTime, nullptr);
+		FTransform LocalRootMotion = RootTransform.GetRelativeTransform(LastRootTransform);
+		LastRootTransform = RootTransform;
+
+		AccumulatedRootDistance += LocalRootMotion.GetTranslation().Size();
+		Output.AccumulatedRootDistance.Add((float)AccumulatedRootDistance);
 	}
 
-	USkeleton* Skeleton = Input.Sequence->GetSkeleton();
-	FBoneContainer BoneContainer;
-	BoneContainer.InitializeTo(Input.Schema->BoneIndicesWithParents, FCurveEvaluationOption(false), *Skeleton);
+	// Verify we sampled the final frame of the clip
+	check(SampleTime == Input.Sequence->GetPlayLength());
+
+	// Also emit root motion summary info to help with sample wrapping in 
+	// FSequenceIndexer::GetSampleTimeFromDistance() and FSequenceIndexer::GetSampleInfo()
+	Output.TotalRootMotion = LastRootTransform.GetRelativeTransform(InitialRootTransform);
+	Output.TotalRootDistance = Output.AccumulatedRootDistance.Last();
+}
+
+struct FSamplingParam
+{
+	float WrappedParam = 0.0f;
+	int32 NumCycles = 0;
+	bool bClamped = false;
+};
 
-	FCompactPose Pose;
-	Pose.SetBoneContainer(&BoneContainer);
-	FCSPose<FCompactPose> ComponentSpacePose;
+static FSamplingParam WrapOrClampSamplingParam(bool bCanWrap, float SamplingParamExtent, float SamplingParam)
+{
+	// This is a helper function used by both time and distance sampling. A schema may specify time or distance
+	// offsets that are multiple cycles of a clip away from the current pose being sampled.
+	// And that time or distance offset may before the beginning of the clip (SamplingParam < 0.0f)
+	// or after the end of the clip (SamplingParam > SamplingParamExtent). So this function
+	// helps determine how many cycles need to be applied and what the wrapped value should be, clamping
+	// if necessary.
 
-	FBlendedCurve UnusedCurve;
-	UE::Anim::FStackAttributeContainer UnusedAttributes;
+	FSamplingParam Result;
 
-	FAnimExtractContext ExtractionCtx;
-	// ExtractionCtx.PoseCurves is intentionally left empty
-	// ExtractionCtx.BonesRequired is unused by UAnimSequence::GetAnimationPose
-	ExtractionCtx.bExtractRootMotion = true;
+	Result.WrappedParam = SamplingParam;
 
-	FAnimationPoseData AnimPoseData(Pose, UnusedCurve, UnusedAttributes);
-	for (int32 SampleIdx = 0; SampleIdx != Output.TotalSamples; ++SampleIdx)
+	if (bCanWrap)
 	{
-		const float CurrentTime = SampleIdx * Input.Schema->SamplingInterval;
-
-		ExtractionCtx.CurrentTime = CurrentTime;
-		Input.Sequence->GetAnimationPose(AnimPoseData, ExtractionCtx);
-		ComponentSpacePose.InitPose(Pose);
+		if (SamplingParam < 0.0f)
+		{
+			// Note a full reverse cycle isn't completed merely when SamplingParam < 0.0f,
+			// rather when SamplingParam < -SamplingParamExtent
+			Result.WrappedParam += SamplingParamExtent;
+			while (Result.WrappedParam < 0.0f)
+			{
+				Result.WrappedParam += SamplingParamExtent;
+				++Result.NumCycles;
+			}
+		}
 
-		for (int32 BoneIndex : Input.Schema->BoneIndices)
+		else
 		{
-			FCompactPoseBoneIndex CompactBoneIndex = BoneContainer.MakeCompactPoseIndex(FMeshPoseBoneIndex(BoneIndex));
-			const FTransform& Transform = ComponentSpacePose.GetComponentSpaceTransform(CompactBoneIndex);
-			Output.ComponentSpacePose.Add(Transform);
+			// On the other hand a full forward cycle is completed each time we exceed SamplingParamExtent
+			while (Result.WrappedParam > SamplingParamExtent)
+			{
+				Result.WrappedParam -= SamplingParamExtent;
+				++Result.NumCycles;
+			}
 		}
 	}
-}
 
-void FSequenceSampler::ExtractRootMotion()
-{
-	double AccumulatedRootDistance = 0.0;
-	FTransform AccumulatedRootMotion = FTransform::Identity;
-	for (int32 SampleIdx = 0; SampleIdx != Output.TotalSamples; ++SampleIdx)
+	float ParamClamped = FMath::Clamp(Result.WrappedParam, 0.0f, SamplingParamExtent);
+	if (ParamClamped != Result.WrappedParam)
 	{
-		const float CurrentTime = SampleIdx * Input.Schema->SamplingInterval;
-
-		FTransform LocalRootMotion = Input.Sequence->ExtractRootMotion(CurrentTime, Input.Schema->SamplingInterval, false /*!allowLooping*/);
-		Output.LocalRootMotion.Add(LocalRootMotion);
-
-		AccumulatedRootMotion = LocalRootMotion * AccumulatedRootMotion;
-		AccumulatedRootDistance += LocalRootMotion.GetTranslation().Size();
-		Output.AccumulatedRootMotion.Add(AccumulatedRootMotion);
-		Output.AccumulatedRootDistance.Add((float)AccumulatedRootDistance);
+		check(!bCanWrap);
+		Result.WrappedParam = ParamClamped;
+		Result.bClamped = true;
 	}
+	
+	return Result;
 }
 
 
@@ -1469,6 +1480,7 @@ public:
 
 	struct FInput
 	{
+		const FBoneContainer* BoneContainer = nullptr;
 		const UPoseSearchSchema* Schema = nullptr;
 		const FSequenceSampler* MainSequence = nullptr;
 		const FSequenceSampler* LeadInSequence = nullptr;
@@ -1491,15 +1503,19 @@ public:
 private:
 	FPoseSearchFeatureVectorBuilder FeatureVector;
 
-	struct FSubsample
+	struct FSampleInfo
 	{
-		const FSequenceSampler* Sampler = nullptr;
-		int32 AbsoluteSampleIdx = INDEX_NONE;
-		FTransform AccumulatedRootMotion;
-		float AccumulatedRootDistance = 0.0f;
+		const FSequenceSampler* Clip = nullptr;
+		FTransform RootTransform;
+		float ClipTime = 0.0f;
+		float RootDistance = 0.0f;
+
+		bool IsValid() const { return Clip != nullptr; }
 	};
 
-	FSubsample ResolveSubsample(int32 MainSubsampleIdx) const;
+	FSampleInfo GetSampleInfo(float SampleTime) const;
+	FSampleInfo GetSampleInfoRelative(float SampleTime, const FSampleInfo& Origin) const;
+	const float GetSampleTimeFromDistance(float Distance) const;
 
 	void Reserve();
 	void SampleBegin(int32 SampleIdx);
@@ -1535,7 +1551,7 @@ void FSequenceIndexer::Init(const FInput& InSettings)
 
 	Reset();
 	Output.FirstIndexedSample = FMath::FloorToInt(SamplingRange.Min * Input.Schema->SampleRate);
-	Output.LastIndexedSample = FMath::Max(0, FMath::FloorToInt(SamplingRange.Max * Input.Schema->SampleRate) - 1);
+	Output.LastIndexedSample = FMath::Max(0, FMath::CeilToInt(SamplingRange.Max * Input.Schema->SampleRate));
 	Output.NumIndexedPoses = Output.LastIndexedSample - Output.FirstIndexedSample + 1;
 	Reserve();
 }
@@ -1572,18 +1588,96 @@ void FSequenceIndexer::SampleEnd(int32 SampleIdx)
 	FMemory::Memcpy(WriteValues.GetData(), ReadValues.GetData(), WriteValues.Num() * WriteValues.GetTypeSize());
 }
 
-FSequenceIndexer::FSubsample FSequenceIndexer::ResolveSubsample(int32 MainSubsampleIdx) const
+const float FSequenceIndexer::GetSampleTimeFromDistance(float SampleDistance) const
 {
-	// MainSubsampleIdx is relative to the samples in the main sequence. With future subsampling,
-	// SampleIdx may be greater than the  number of samples in the main sequence. For past subsampling,
-	// SampleIdx may be negative. This function handles those edge cases by wrapping within the main
-	// sequence if it is loopable, or by indexing into the lead-in or follow-up sequences which themselves
-	// may or may not be loopable.
-	// The relative SampleIdx may be multiple cycles away, so this function also handles the math for
-	// accumulating multiple cycles of root motion.
-	// It returns an absolute index into the relevant sample data and root motion info.
+	auto CanWrapDistanceSamples = [](const FSequenceSampler* Sampler) -> bool
+	{
+		constexpr float SMALL_ROOT_DISTANCE = 1.0f;
+		return Sampler->Input.bLoopable && Sampler->Output.TotalRootDistance > SMALL_ROOT_DISTANCE;
+	};
+
+	auto ClipTimeFromDistance = [](const FSequenceSampler* Sampler, float ClipDistance) -> float
+	{
+		// Search for the distance value
+		int32 NextSampleIdx = Algo::LowerBound(Sampler->Output.AccumulatedRootDistance, ClipDistance);
+		check(NextSampleIdx < Sampler->Output.AccumulatedRootDistance.Num());
+	
+		// Compute distance interpolation amount
+		int32 PrevSampleIdx = FMath::Max(0, NextSampleIdx - 1);
+		float NextDistance = Sampler->Output.AccumulatedRootDistance[NextSampleIdx];
+		float PrevDistance = PrevSampleIdx < NextSampleIdx ? Sampler->Output.AccumulatedRootDistance[PrevSampleIdx] : 0.0f;
+		float DistanceSampleAlpha = FMath::GetRangePct(PrevDistance, NextDistance, ClipDistance);
+
+		// Convert to time
+		float ClipTime = (float(NextSampleIdx) - (1.0f - DistanceSampleAlpha)) / Sampler->Input.DistanceSamplingRate;
+		float ClipTimeClamped = FMath::Clamp(ClipTime, 0.0f, Sampler->Output.PlayLength);
+
+		return ClipTimeClamped;
+	};
+
+	float MainTotalDistance = Input.MainSequence->Output.TotalRootDistance;
+	bool bMainCanWrap = CanWrapDistanceSamples(Input.MainSequence);
+
+	float SampleTime = MAX_flt;
+
+	if (!bMainCanWrap)
+	{
+		// Use the lead in anim if we would have to clamp to the beginning of the main anim
+		if (Input.LeadInSequence && (SampleDistance < 0.0f))
+		{
+			const FSequenceSampler::FOutput& ClipData = Input.LeadInSequence->Output;
+
+			bool bLeadInCanWrap = CanWrapDistanceSamples(Input.LeadInSequence);
+			FSamplingParam SamplingParam = WrapOrClampSamplingParam(bLeadInCanWrap, ClipData.TotalRootDistance, SampleDistance);
 
-	FSubsample Subsample;
+			float ClipTime = ClipTimeFromDistance(Input.LeadInSequence, SamplingParam.WrappedParam);
+
+			// Make the lead in clip time relative to the main sequence again and unwrap
+			SampleTime = -((SamplingParam.NumCycles * ClipData.PlayLength) + (ClipData.PlayLength - ClipTime));
+		}
+
+		// Use the follow up anim if we would have clamp to the end of the main anim
+		else if (Input.FollowUpSequence && (SampleDistance > MainTotalDistance))
+		{
+			const FSequenceSampler::FOutput& ClipData = Input.FollowUpSequence->Output;
+
+			bool bFollowUpCanWrap = CanWrapDistanceSamples(Input.FollowUpSequence);
+			float FollowRelativeDistance = SampleDistance - MainTotalDistance;
+			FSamplingParam SamplingParam = WrapOrClampSamplingParam(bFollowUpCanWrap, ClipData.TotalRootDistance, FollowRelativeDistance);
+
+			float ClipTime = ClipTimeFromDistance(Input.FollowUpSequence, SamplingParam.WrappedParam);
+
+			// Make the follow up clip time relative to the main sequence again and unwrap
+			SampleTime = Input.MainSequence->Output.PlayLength + SamplingParam.NumCycles * ClipData.PlayLength + ClipTime;
+		}
+	}
+
+	// Use the main anim if we didn't use the lead-in or follow-up anims.
+	// The main anim sample may have been wrapped or clamped
+	if (SampleTime == MAX_flt)
+	{
+		const FSequenceSampler::FOutput& ClipData = Input.MainSequence->Output;
+
+		FSamplingParam SamplingParam = WrapOrClampSamplingParam(bMainCanWrap, MainTotalDistance, SampleDistance);
+		float ClipTime = ClipTimeFromDistance(Input.MainSequence, SamplingParam.WrappedParam);
+
+		// Unwrap the main clip time
+		if (SampleDistance < 0.0f)
+		{
+			SampleTime = -((SamplingParam.NumCycles * ClipData.PlayLength) + (ClipData.PlayLength - ClipTime));
+		}
+		else
+		{
+			SampleTime = SamplingParam.NumCycles * ClipData.PlayLength + ClipTime;
+		}
+	}
+
+	return SampleTime;
+}
+
+FSequenceIndexer::FSampleInfo FSequenceIndexer::GetSampleInfo(float SampleTime) const
+{
+	FSampleInfo Sample;
 
 	FTransform RootMotionLast = FTransform::Identity;
 	FTransform RootMotionInitial = FTransform::Identity;
@@ -1591,200 +1685,276 @@ FSequenceIndexer::FSubsample FSequenceIndexer::ResolveSubsample(int32 MainSubsam
 	float RootDistanceLast = 0.0f;
 	float RootDistanceInitial = 0.0f;
 
-	FSequenceSampler::FWrappedSampleIndex MainSample = Input.MainSequence->WrapOrClampSubsampleIndex(MainSubsampleIdx);
-	FSequenceSampler::FWrappedSampleIndex EffectiveSample;
+	auto CanWrapTimeSamples = [](const FSequenceSampler* Sampler) -> bool
+	{
+		return Sampler->Input.bLoopable;
+	};
+
+	float MainPlayLength = Input.MainSequence->Output.PlayLength;
+	bool bMainCanWrap = CanWrapTimeSamples(Input.MainSequence);
 
-	// Use the lead in anim if we had to clamp to the beginning of the main anim
-	if (MainSample.bClamped && (MainSubsampleIdx < 0) && Input.LeadInSequence)
+	FSamplingParam SamplingParam;
+	if (!bMainCanWrap)
 	{
-		EffectiveSample = Input.LeadInSequence->WrapOrClampSubsampleIndex(MainSubsampleIdx);
+		// Use the lead in anim if we would have to clamp to the beginning of the main anim
+		if (Input.LeadInSequence && (SampleTime < 0.0f))
+		{
+			const FSequenceSampler::FOutput& ClipData = Input.LeadInSequence->Output;
 
-		Subsample.Sampler = Input.LeadInSequence;
-		Subsample.AbsoluteSampleIdx = EffectiveSample.Idx;
+			bool bLeadInCanWrap = CanWrapTimeSamples(Input.LeadInSequence);
+			SamplingParam = WrapOrClampSamplingParam(bLeadInCanWrap, ClipData.PlayLength, SampleTime);
 
-		RootMotionInitial = FTransform::Identity;
-		RootDistanceInitial = 0.0f;
+			Sample.Clip = Input.LeadInSequence;
 
-		RootMotionLast = Input.LeadInSequence->Output.AccumulatedRootMotion.Last();
-		RootDistanceLast = Input.LeadInSequence->Output.AccumulatedRootDistance.Last();
-	}
+			RootMotionInitial = FTransform::Identity;
+			RootDistanceInitial = 0.0f;
 
-	// Use the follow up anim if we had to clamp to the end of the main anim
-	if (MainSample.bClamped && (MainSubsampleIdx >= Input.MainSequence->Output.TotalSamples) && Input.FollowUpSequence)
-	{
-		EffectiveSample = Input.FollowUpSequence->WrapOrClampSubsampleIndex(MainSubsampleIdx - Input.MainSequence->Output.TotalSamples);
+			RootMotionLast = Input.LeadInSequence->Output.TotalRootMotion;
+			RootDistanceLast = Input.LeadInSequence->Output.TotalRootDistance;
+		}
+
+		// Use the follow up anim if we would have clamp to the end of the main anim
+		else if (Input.FollowUpSequence && (SampleTime > MainPlayLength))
+		{
+			const FSequenceSampler::FOutput& ClipData = Input.FollowUpSequence->Output;
+
+			bool bFollowUpCanWrap = CanWrapTimeSamples(Input.FollowUpSequence);
+			float FollowRelativeTime = SampleTime - MainPlayLength;
+			SamplingParam = WrapOrClampSamplingParam(bFollowUpCanWrap, ClipData.PlayLength, FollowRelativeTime);
 
-		Subsample.Sampler = Input.FollowUpSequence;
-		Subsample.AbsoluteSampleIdx = EffectiveSample.Idx;
+			Sample.Clip = Input.FollowUpSequence;
 
-		RootMotionInitial = Input.MainSequence->Output.AccumulatedRootMotion.Last();
-		RootDistanceInitial = Input.MainSequence->Output.AccumulatedRootDistance.Last();
+			RootMotionInitial = Input.MainSequence->Output.TotalRootMotion;
+			RootDistanceInitial = Input.MainSequence->Output.TotalRootDistance;
 
-		RootMotionLast = Input.FollowUpSequence->Output.AccumulatedRootMotion.Last();
-		RootDistanceLast = Input.FollowUpSequence->Output.AccumulatedRootDistance.Last();
+			RootMotionLast = Input.FollowUpSequence->Output.TotalRootMotion;
+			RootDistanceLast = Input.FollowUpSequence->Output.TotalRootDistance;
+		}
 	}
 
 	// Use the main anim if we didn't use the lead-in or follow-up anims.
 	// The main anim sample may have been wrapped or clamped
-	if (EffectiveSample.Idx == INDEX_NONE)
+	if (!Sample.IsValid())
 	{
-		EffectiveSample = MainSample;
+		SamplingParam = WrapOrClampSamplingParam(bMainCanWrap, MainPlayLength, SampleTime);
 
-		Subsample.Sampler = Input.MainSequence;
-		Subsample.AbsoluteSampleIdx = EffectiveSample.Idx;
+		Sample.Clip = Input.MainSequence;
 
 		RootMotionInitial = FTransform::Identity;
 		RootDistanceInitial = 0.0f;
 
-		RootMotionLast = Input.MainSequence->Output.AccumulatedRootMotion.Last();
-		RootDistanceLast = Input.MainSequence->Output.AccumulatedRootDistance.Last();
+		RootMotionLast = Input.MainSequence->Output.TotalRootMotion;
+		RootDistanceLast = Input.MainSequence->Output.TotalRootDistance;
 	}
 
+	Sample.ClipTime = SamplingParam.WrappedParam;
+
 	// Determine how to accumulate motion for every cycle of the anim. If the sample
 	// had to be clamped, this motion will end up not getting applied below.
 	// Also invert the accumulation direction if the requested sample was wrapped backwards.
 	FTransform RootMotionPerCycle = RootMotionLast;
 	float RootDistancePerCycle = RootDistanceLast;
-	if (MainSubsampleIdx < 0)
+	if (SampleTime < 0.0f)
 	{
 		RootMotionPerCycle = RootMotionPerCycle.Inverse();
 		RootDistancePerCycle *= -1;
 	}
 
 	// Find the remaining motion deltas after wrapping
-	FTransform RootMotionRemainder = Subsample.Sampler->Output.AccumulatedRootMotion[EffectiveSample.Idx];
-	float RootDistanceRemainder = Subsample.Sampler->Output.AccumulatedRootDistance[EffectiveSample.Idx];
+	FTransform RootMotionRemainder = Sample.Clip->ExtractRootTransform(Sample.ClipTime);
+	float RootDistanceRemainder = Sample.Clip->ExtractRootDistance(Sample.ClipTime);
 
 	// Invert motion deltas if we wrapped backwards
-	if (MainSubsampleIdx < 0)
+	if (SampleTime < 0.0f)
 	{
 		RootMotionRemainder.SetToRelativeTransform(RootMotionLast);
 		RootDistanceRemainder = -(RootDistanceLast - RootDistanceRemainder);
 	}
 
-	Subsample.AccumulatedRootMotion = RootMotionInitial;
-	Subsample.AccumulatedRootDistance = RootDistanceInitial;
+	Sample.RootTransform = RootMotionInitial;
+	Sample.RootDistance = RootDistanceInitial;
 
 	// Note if the sample was clamped, no motion will be applied here because NumCycles will be zero
-	int32 CyclesRemaining = EffectiveSample.NumCycles;
+	int32 CyclesRemaining = SamplingParam.NumCycles;
 	while (CyclesRemaining--)
 	{
-		Subsample.AccumulatedRootMotion *= RootMotionPerCycle;
-		Subsample.AccumulatedRootDistance += RootDistancePerCycle;
+		Sample.RootTransform *= RootMotionPerCycle;
+		Sample.RootDistance += RootDistancePerCycle;
 	}
 
-	Subsample.AccumulatedRootMotion *= RootMotionRemainder;
-	Subsample.AccumulatedRootDistance += RootDistanceRemainder;
+	Sample.RootTransform *= RootMotionRemainder;
+	Sample.RootDistance += RootDistanceRemainder;
+
+	return Sample;
+}
 
-	return Subsample;
+FSequenceIndexer::FSampleInfo FSequenceIndexer::GetSampleInfoRelative(float SampleTime, const FSampleInfo& Origin) const
+{
+	FSampleInfo Sample = GetSampleInfo(SampleTime);
+	Sample.RootTransform.SetToRelativeTransform(Origin.RootTransform);
+	Sample.RootDistance = Origin.RootDistance - Sample.RootDistance;
+	return Sample;
 }
 
 void FSequenceIndexer::AddPoseFeatures(int32 SampleIdx)
 {
+	// This function samples the instantaneous pose at time t as well as the pose's velocity and acceleration at time t.
+	// Symmetric finite differences are used to approximate derivatives:
+	//	First symmetric derivative:   f'(t) ~ (f(t+h) - f(t-h)) / 2h
+	//	Second symmetric derivative: f''(t) ~ (f(t+h) - 2f(t) + f(t-h)) / h^2
+	// Where h is a constant time delta
+	// So this means three pose extractions are taken at time t-h, t, and t+h
+	constexpr float FiniteDelta = 1 / 60.0f;
+	constexpr int32 NumFiniteDiffTerms = 3;
+
+	if (Input.Schema->Bones.IsEmpty() || Input.Schema->PoseSampleOffsets.IsEmpty())
+	{
+		return;
+	}
+
+	FCompactPose Poses[NumFiniteDiffTerms];
+	FCSPose<FCompactPose> ComponentSpacePoses[NumFiniteDiffTerms];
+	FBlendedCurve Curves[NumFiniteDiffTerms];
+	UE::Anim::FStackAttributeContainer Atrributes[NumFiniteDiffTerms];
+
+	for(FCompactPose& Pose: Poses)
+	{
+		Pose.SetBoneContainer(Input.BoneContainer);
+	}
+
+	FAnimationPoseData AnimPoseData[NumFiniteDiffTerms] = 
+	{
+		{Poses[0], Curves[0], Atrributes[0]},
+		{Poses[1], Curves[1], Atrributes[1]},
+		{Poses[2], Curves[2], Atrributes[2]},
+	};
+
+	FAnimExtractContext ExtractionCtx;
+	ExtractionCtx.bExtractRootMotion = true;
+
 	FPoseSearchFeatureDesc Feature;
 	Feature.Domain = EPoseSearchFeatureDomain::Time;
 
-	const int32 NumBones = Input.Schema->NumBones();
-
-	FSequenceIndexer::FSubsample OriginSample = ResolveSubsample(SampleIdx);
+	float SampleTime = FMath::Min(SampleIdx * Input.Schema->SamplingInterval, Input.MainSequence->Output.PlayLength);
+	FSampleInfo Origin = GetSampleInfo(SampleTime);
 	
 	for (int32 SchemaSubsampleIdx = 0; SchemaSubsampleIdx != Input.Schema->PoseSampleOffsets.Num(); ++SchemaSubsampleIdx)
 	{
 		Feature.SubsampleIdx = SchemaSubsampleIdx;
 
-		const int32 SubsampleIdx = SampleIdx + Input.Schema->PoseSampleOffsets[SchemaSubsampleIdx];
+		float SubsampleTime = SampleTime + Input.Schema->PoseSampleTimes[SchemaSubsampleIdx];
 
-		FSequenceIndexer::FSubsample Subsample = ResolveSubsample(SubsampleIdx);
-		FSequenceIndexer::FSubsample SubsamplePrev = ResolveSubsample(SubsampleIdx - 1);
+		// For each pose subsample term, get the corresponding clip, accumulated root motion,
+		// and wrap the time parameter based on the clip's length.
+		FSampleInfo Samples[NumFiniteDiffTerms];
+		Samples[0] = GetSampleInfoRelative(SubsampleTime - FiniteDelta, Origin);
+		Samples[1] = GetSampleInfoRelative(SubsampleTime, Origin);
+		Samples[2] = GetSampleInfoRelative(SubsampleTime + FiniteDelta, Origin);
 
-		FTransform SubsampleRoot = Subsample.AccumulatedRootMotion;
-		SubsampleRoot.SetToRelativeTransform(OriginSample.AccumulatedRootMotion);
+		// Get pose samples
+		for (int32 Term = 0; Term != NumFiniteDiffTerms; ++Term)
+		{
+			ExtractionCtx.CurrentTime = Samples[Term].ClipTime;
+			Samples[Term].Clip->Input.Sequence->GetAnimationPose(AnimPoseData[Term], ExtractionCtx);
+			ComponentSpacePoses[Term].InitPose(Poses[Term]);
+		}
 
-		for (int32 SchemaBoneIndex = 0; SchemaBoneIndex != NumBones; ++SchemaBoneIndex)
+		// Get each bone's component transform, velocity, and acceleration and add accumulated root motion at this time offset
+		// Think of this process as freezing the character in place (at SampleTime) and then tracing the paths of their joints
+		// as they move through space from past to present to future (at times indicated by PoseSampleOffsets).
+		for (int32 SchemaBoneIndex = 0; SchemaBoneIndex != Input.Schema->NumBones(); ++SchemaBoneIndex)
 		{
 			Feature.SchemaBoneIdx = SchemaBoneIndex;
 
-			int32 BoneSampleIdx = NumBones * Subsample.AbsoluteSampleIdx + SchemaBoneIndex;
-			int32 BonePrevSampleIdx = NumBones * SubsamplePrev.AbsoluteSampleIdx  + SchemaBoneIndex;
+			FCompactPoseBoneIndex CompactBoneIndex = Input.BoneContainer->MakeCompactPoseIndex(FMeshPoseBoneIndex(Input.Schema->BoneIndices[SchemaBoneIndex]));
 
-			FTransform BoneInComponentSpace = Subsample.Sampler->Output.ComponentSpacePose[BoneSampleIdx];
-			FTransform BonePrevInComponentSpace = SubsamplePrev.Sampler->Output.ComponentSpacePose[BonePrevSampleIdx];
-
-			FTransform BoneInSampleSpace = BoneInComponentSpace * SubsampleRoot;
-			FTransform BonePrevInSampleSpace = BonePrevInComponentSpace * SubsampleRoot;
+			FTransform BoneTransforms[NumFiniteDiffTerms];
+			for (int32 Term = 0; Term != NumFiniteDiffTerms; ++Term)
+			{
+				BoneTransforms[Term] = ComponentSpacePoses[Term].GetComponentSpaceTransform(CompactBoneIndex);
+				BoneTransforms[Term] = BoneTransforms[Term] * Samples[Term].RootTransform;
+			}
 
-			FeatureVector.SetTransform(Feature, BoneInSampleSpace);
-			FeatureVector.SetTransformDerivative(Feature, BoneInSampleSpace, BonePrevInSampleSpace, Input.Schema->SamplingInterval);
+			// Add properties to the feature vector for the pose at SampleIdx
+			FeatureVector.SetTransform(Feature, BoneTransforms[1]);
+			FeatureVector.SetTransformVelocity(Feature, BoneTransforms[2], BoneTransforms[0], 2.0f * FiniteDelta);
+			//FeatureVector.SetTransformAccleration(Feature, BoneTransforms[2], BoneTransforms[1], BoneTransforms[0], FiniteDelta * FiniteDelta);
 		}
 	}
 }
 
 void FSequenceIndexer::AddTrajectoryTimeFeatures(int32 SampleIdx)
 {
+	// This function samples the instantaneous trajectory at time t as well as the trajectory's velocity and acceleration at time t.
+	// Symmetric finite differences are used to approximate derivatives:
+	//	First symmetric derivative:   f'(t) ~ (f(t+h) - f(t-h)) / 2h
+	//	Second symmetric derivative: f''(t) ~ (f(t+h) - 2f(t) + f(t-h)) / h^2
+	// Where h is a constant time delta
+	// So this means three root motion extractions are taken at time t-h, t, and t+h
+	constexpr float FiniteDelta = 1 / 60.0f;
+
 	FPoseSearchFeatureDesc Feature;
 	Feature.Domain = EPoseSearchFeatureDomain::Time;
 	Feature.SchemaBoneIdx = FPoseSearchFeatureDesc::TrajectoryBoneIndex;
 
-	FSequenceIndexer::FSubsample OriginSample = ResolveSubsample(SampleIdx);
+	float SampleTime = FMath::Min(SampleIdx * Input.Schema->SamplingInterval, Input.MainSequence->Output.PlayLength);
+	FSampleInfo Origin = GetSampleInfo(SampleTime);
 	
-	for (int32 SchemaSubsampleIdx = 0; SchemaSubsampleIdx != Input.Schema->TrajectorySampleOffsets.Num(); ++SchemaSubsampleIdx)
+	for (int32 SchemaSubsampleIdx = 0; SchemaSubsampleIdx != Input.Schema->TrajectorySampleTimes.Num(); ++SchemaSubsampleIdx)
 	{
 		Feature.SubsampleIdx = SchemaSubsampleIdx;
 
-		const int32 SubsampleIdx = SampleIdx + Input.Schema->TrajectorySampleOffsets[SchemaSubsampleIdx];
-
-		FSequenceIndexer::FSubsample Subsample = ResolveSubsample(SubsampleIdx);
-		FTransform SubsampleRoot = Subsample.AccumulatedRootMotion;
-		SubsampleRoot.SetToRelativeTransform(OriginSample.AccumulatedRootMotion);
+		float SubsampleTime = SampleTime + Input.Schema->TrajectorySampleTimes[SchemaSubsampleIdx];
 
-		FSequenceIndexer::FSubsample SubsamplePrev = ResolveSubsample(SubsampleIdx - 1);
-		FTransform SubsamplePrevRoot = SubsamplePrev.AccumulatedRootMotion;
-		SubsamplePrevRoot.SetToRelativeTransform(OriginSample.AccumulatedRootMotion);
+		// For each pose subsample term, get the corresponding clip, accumulated root motion,
+		// and wrap the time parameter based on the clip's length.
+		FSampleInfo Samples[3];
+		Samples[0] = GetSampleInfoRelative(SubsampleTime - FiniteDelta, Origin);
+		Samples[1] = GetSampleInfoRelative(SubsampleTime, Origin);
+		Samples[2] = GetSampleInfoRelative(SubsampleTime + FiniteDelta, Origin);
 
-		FeatureVector.SetTransform(Feature, SubsampleRoot);
-		FeatureVector.SetTransformDerivative(Feature, SubsampleRoot, SubsamplePrevRoot, Input.Schema->SamplingInterval);
+		// Add properties to the feature vector for the pose at SampleIdx
+		FeatureVector.SetTransform(Feature, Samples[1].RootTransform);
+		FeatureVector.SetTransformVelocity(Feature, Samples[2].RootTransform, Samples[0].RootTransform, 2.0f * FiniteDelta);
+		//FeatureVector.SetTransformAcceleration(Feature, Samples[2].RootTransform, Samples[1].RootTransform, Samples[0].RootTransform, FiniteDelta * FiniteDelta);
 	}
 }
 
 void FSequenceIndexer::AddTrajectoryDistanceFeatures(int32 SampleIdx)
 {
-	// This function needs to be rewritten to work with the updated sampler
-	// and lead-in/follow-up anims
-
-//	FPoseSearchFeatureDesc Feature;
-//	Feature.Domain = EPoseSearchFeatureDomain::Distance;
-//	Feature.SchemaBoneIdx = FPoseSearchFeatureDesc::TrajectoryBoneIndex;
-//
-//	FSequenceIndexer::FSampleRef OriginSampleRef = ResolveSampleRef(SampleIdx);
-//
-//	for (int32 SubsampleIdx = 0; SubsampleIdx != Input.Schema->TrajectorySampleDistances.Num(); ++SubsampleIdx)
-//	{
-//		Feature.SubsampleIdx = SubsampleIdx;
-//
-//		const float TrajectoryDistance = Input.Schema->TrajectorySampleDistances[SubsampleIdx];
-//		const float SampleAccumulatedRootDistance = TrajectoryDistance + AccumulatedRootDistances[SampleIdx];
-//
-//		int32 LowerBoundSampleIdx = Algo::LowerBound(AccumulatedRootDistances, SampleAccumulatedRootDistance);
-//
-//		// @@@ Add extrapolation. Clamp for now
-//		int32 PrevSampleIdx = FMath::Clamp(LowerBoundSampleIdx - 1, 0, AccumulatedRootDistances.Num() - 1);
-//		int32 NextSampleIdx = FMath::Clamp(LowerBoundSampleIdx, 0, AccumulatedRootDistances.Num() - 1);
-//
-//		const float PrevSampleDistance = AccumulatedRootDistances[PrevSampleIdx];
-//		const float NextSampleDistance = AccumulatedRootDistances[NextSampleIdx];
-//
-//		FTransform PrevRootInSampleSpace = AccumulatedRootMotion[PrevSampleIdx];
-//		PrevRootInSampleSpace.SetToRelativeTransform(SampleSpaceOrigin);
-//
-//		FTransform NextRootInSampleSpace = AccumulatedRootMotion[NextSampleIdx];
-//		NextRootInSampleSpace.SetToRelativeTransform(SampleSpaceOrigin);
-//		
-//		float Alpha = FMath::GetRangePct(PrevSampleDistance, NextSampleDistance, SampleAccumulatedRootDistance);
-//		FTransform BlendedRootInSampleSpace;
-//		BlendedRootInSampleSpace.Blend(PrevRootInSampleSpace, NextRootInSampleSpace, Alpha);
-//
-//		FeatureVector.SetTransform(Feature, BlendedRootInSampleSpace);
-//	}
+	// This function is very similar to AddTrajectoryTimeFeatures, but samples are taken in the distance domain
+	// instead of time domain.
+	constexpr float FiniteDelta = 1 / 60.0f;
+
+	FPoseSearchFeatureDesc Feature;
+	Feature.Domain = EPoseSearchFeatureDomain::Distance;
+	Feature.SchemaBoneIdx = FPoseSearchFeatureDesc::TrajectoryBoneIndex;
+
+	float SampleTime = FMath::Min(SampleIdx * Input.Schema->SamplingInterval, Input.MainSequence->Output.PlayLength);
+	FSampleInfo Origin = GetSampleInfo(SampleTime);
+	
+	for (int32 SchemaSubsampleIdx = 0; SchemaSubsampleIdx != Input.Schema->TrajectorySampleDistances.Num(); ++SchemaSubsampleIdx)
+	{
+		Feature.SubsampleIdx = SchemaSubsampleIdx;
+
+		// For distance based sampling of the trajectory, we first have to look up the time value
+		// we're sampling given the desired travel distance of the root for this distance offset.
+		// Once we know the time, we can then carry on just like time-based sampling.
+		const float SubsampleDistance = Origin.RootDistance + Input.Schema->TrajectorySampleDistances[SchemaSubsampleIdx];
+		float SubsampleTime = GetSampleTimeFromDistance(SubsampleDistance);
+
+		// For each pose subsample term, get the corresponding clip, accumulated root motion,
+		// and wrap the time parameter based on the clip's length.
+		FSampleInfo Samples[3];
+		Samples[0] = GetSampleInfoRelative(SubsampleTime - FiniteDelta, Origin);
+		Samples[1] = GetSampleInfoRelative(SubsampleTime, Origin);
+		Samples[2] = GetSampleInfoRelative(SubsampleTime + FiniteDelta, Origin);
+
+		// Add properties to the feature vector for the pose at SampleIdx
+		FeatureVector.SetTransform(Feature, Samples[1].RootTransform);
+		FeatureVector.SetTransformVelocity(Feature, Samples[2].RootTransform, Samples[0].RootTransform, 2.0f * FiniteDelta);
+		//FeatureVector.SetTransformAcceleration(Feature, Samples[0].RootTransform, Samples[1].RootTransform, Samples[2].RootTransform, FiniteDelta * FiniteDelta);
+	}
 }
 
 
@@ -2020,24 +2190,28 @@ static void PreprocessSearchIndexNone(FPoseSearchIndex* SearchIndex)
 	SampleMeanMap = VectorXf::Zero(NumDimensions);
 }
 
-static Eigen::VectorXd ComputeFeatureMeanDeviations (const Eigen::MatrixXd& CenteredPoseMatrix, const FPoseSearchFeatureVectorLayout& Layout) {
+inline Eigen::VectorXd ComputeFeatureMeanDeviations(const Eigen::MatrixXd& CenteredPoseMatrix, const FPoseSearchFeatureVectorLayout& Layout) {
 	using namespace Eigen;
 
 	const int32 NumPoses = CenteredPoseMatrix.cols();
 	const int32 NumDimensions = CenteredPoseMatrix.rows();
 
 	VectorXd MeanDeviations(NumDimensions);
-	MeanDeviations.setConstant(0);
+	MeanDeviations.setConstant(1.0);
 	for (const FPoseSearchFeatureDesc& Feature : Layout.Features)
 	{
 		int32 FeatureDims = GetFeatureTypeTraits(Feature.Type).NumFloats;
 
 		// Construct a submatrix for the feature and find the average distance to the feature's centroid.
-		// Since we've already mean centered the data, the average distance to the mean is simply the average norm.
+		// Since we've already mean centered the data, the average distance to the centroid is simply the average norm.
 		double FeatureMeanDeviation = CenteredPoseMatrix.block(Feature.ValueOffset, 0, FeatureDims, NumPoses).colwise().norm().mean();
 
 		// Fill the feature's corresponding scaling axes with the average distance
-		MeanDeviations.segment(Feature.ValueOffset, FeatureDims).setConstant(FeatureMeanDeviation);
+		// Avoid scaling by zero by leaving near-zero deviations as 1.0
+		if (FeatureMeanDeviation > KINDA_SMALL_NUMBER)
+		{
+			MeanDeviations.segment(Feature.ValueOffset, FeatureDims).setConstant(FeatureMeanDeviation);
+		}
 	}
 
 	return MeanDeviations;
@@ -2085,7 +2259,7 @@ static void PreprocessSearchIndexNormalize(FPoseSearchIndex* SearchIndex)
 	// x := (S * x) + u
 	//
 	// References:
-	// [1] Gorard, S. (2005), “Revisiting a 90-Year-Old Debate: The Advantages of the Mean Deviation.”
+	// [1] Gorard, S. (2005), "Revisiting a 90-Year-Old Debate: The Advantages of the Mean Deviation."
 	//     British Journal of Educational Studies, 53: 417-430.
 
 	using namespace Eigen;
@@ -2112,6 +2286,10 @@ static void PreprocessSearchIndexNormalize(FPoseSearchIndex* SearchIndex)
 	checkSlow(PoseMatrix.rows() == NumDimensions);
 	checkSlow(PoseMatrix.cols() == NumPoses);
 
+#if UE_POSE_SEARCH_EIGEN_DEBUG
+	MatrixXd PoseMatrixOriginal = PoseMatrix;
+#endif
+
 	// Mean center
 	VectorXd SampleMean = PoseMatrix.rowwise().mean();
 	PoseMatrix = PoseMatrix.colwise() - SampleMean;
@@ -2156,6 +2334,7 @@ static void PreprocessSearchIndexNormalize(FPoseSearchIndex* SearchIndex)
 	SampleMeanMap = SampleMean.cast<float>();
 
 #if UE_POSE_SEARCH_EIGEN_DEBUG
+	FString PoseMtxOriginalStr = EigenMatrixToString(PoseMatrixOriginal);
 	FString PoseMtxStr = EigenMatrixToString(PoseMatrix);
 	FString TransformationStr = EigenMatrixToString(TransformMap);
 	FString InverseTransformationStr = EigenMatrixToString(InverseTransformMap);
@@ -2210,8 +2389,8 @@ static void PreprocessSearchIndexSphere(FPoseSearchIndex* SearchIndex)
 	// x := (Z^(-1) * x) + u
 	//
 	// The sphering processs allows nearest neighbor queries to use the Mahalonobis metric
-	// which is unitless, scale-invariant, and uncorrelated. The Mahalanobis distance between
-	// two random vectors x and y in data matrix X is:
+	// which is unitless, scale-invariant, and accounts for feature correlation.
+	// The Mahalanobis distance between two random vectors x and y in data matrix X is:
 	// d(x,y) = ((x-y)^T * cov(X)^(-1) * (x-y))^(1/2)
 	//
 	// Because sphering transforms X into a new matrix with identity covariance, the Mahalonobis
@@ -2263,6 +2442,10 @@ static void PreprocessSearchIndexSphere(FPoseSearchIndex* SearchIndex)
 	checkSlow(PoseMatrix.rows() == NumDimensions);
 	checkSlow(PoseMatrix.cols() == NumPoses);
 
+#if UE_POSE_SEARCH_EIGEN_DEBUG
+	MatrixXd PoseMatrixOriginal = PoseMatrix;
+#endif
+
 	// Mean center
 	VectorXd SampleMean = PoseMatrix.rowwise().mean();
 	PoseMatrix = PoseMatrix.colwise() - SampleMean;
@@ -2354,7 +2537,9 @@ static void PreprocessSearchIndexSphere(FPoseSearchIndex* SearchIndex)
 
 	FString PoseMatrixSphereStr = EigenMatrixToString(PoseMatrix);
 	MatrixXd PoseMatrixUnsphered = ZCAInverse * PoseMatrix;
+	PoseMatrixUnsphered = PoseMatrixUnsphered.colwise() + SampleMean;
 	FString PoseMatrixUnspheredStr = EigenMatrixToString(PoseMatrixUnsphered);
+	FString PoseMatrixOriginalStr = EigenMatrixToString(PoseMatrixOriginal);
 
 	FString OutputPoseMatrixStr = EigenMatrixToString(PoseMatrixSourceMap);
 
@@ -2402,6 +2587,9 @@ bool BuildIndex(const UAnimSequence* Sequence, UPoseSearchSequenceMetaData* Sequ
 		return false;
 	}
 
+	FBoneContainer BoneContainer;
+	BoneContainer.InitializeTo(SequenceMetaData->Schema->BoneIndicesWithParents, FCurveEvaluationOption(false), *SeqSkeleton);
+
 	FSequenceSampler Sampler;
 	FSequenceSampler::FInput SamplerInput;
 	SamplerInput.Schema = SequenceMetaData->Schema;
@@ -2415,6 +2603,7 @@ bool BuildIndex(const UAnimSequence* Sequence, UPoseSearchSequenceMetaData* Sequ
 	IndexerInput.MainSequence = &Sampler;
 	IndexerInput.Schema = SequenceMetaData->Schema;
 	IndexerInput.RequestedSamplingRange = SequenceMetaData->SamplingRange;
+	IndexerInput.BoneContainer = &BoneContainer;
 	Indexer.Init(IndexerInput);
 	Indexer.Process();
 
@@ -2436,7 +2625,10 @@ bool BuildIndex(UPoseSearchDatabase* Database)
 		return false;
 	}
 
-	// Prepare animation sampling tasks
+	FBoneContainer BoneContainer;
+	BoneContainer.InitializeTo(Database->Schema->BoneIndicesWithParents, FCurveEvaluationOption(false), *Database->Schema->Skeleton);
+
+	// Prepare animation preprocessing tasks
 	TArray<FSequenceSampler> SequenceSamplers;
 	TMap<const UAnimSequence*, int32> SequenceSamplerMap;
 
@@ -2473,9 +2665,10 @@ bool BuildIndex(UPoseSearchDatabase* Database)
 		}
 	}
 
-	// Sample animations independently
+	// Preprocess animations independently
 	ParallelFor(SequenceSamplers.Num(), [&SequenceSamplers](int32 SamplerIdx){ SequenceSamplers[SamplerIdx].Process(); });
 
+
 	auto GetSampler = [&](const UAnimSequence* Sequence) -> const FSequenceSampler*
 	{
 		return Sequence ? &SequenceSamplers[SequenceSamplerMap[Sequence]] : nullptr;
@@ -2490,6 +2683,7 @@ bool BuildIndex(UPoseSearchDatabase* Database)
 		FSequenceIndexer& Indexer = Indexers[SequenceIdx];
 
 		FSequenceIndexer::FInput Input;
+		Input.BoneContainer = &BoneContainer;
 		Input.MainSequence = GetSampler(DbSequence.Sequence);
 		Input.LeadInSequence = GetSampler(DbSequence.LeadInSequence);
 		Input.FollowUpSequence = GetSampler(DbSequence.FollowUpSequence);
@@ -2777,8 +2971,9 @@ UE::Anim::IPoseSearchProvider::FSearchResult FModule::Search(const FAnimationBas
 
 	QueryBuilder.Init(MetaData->Schema);
 	QueryBuilder.SetPoseFeatures(&PoseHistory);
+	QueryBuilder.Normalize(MetaData->SearchIndex);
 
-	::UE::PoseSearch::FSearchResult Result = ::UE::PoseSearch::Search(Sequence, QueryBuilder.GetValues());
+	::UE::PoseSearch::FSearchResult Result = ::UE::PoseSearch::Search(Sequence, QueryBuilder.GetNormalizedValues());
 
 	ProviderResult.Dissimilarity = Result.Dissimilarity;
 	ProviderResult.PoseIdx = Result.PoseIdx;
diff --git a/Engine/Plugins/Experimental/Animation/PoseSearch/Source/Runtime/Public/PoseSearch/AnimNode_MotionMatching.h b/Engine/Plugins/Experimental/Animation/PoseSearch/Source/Runtime/Public/PoseSearch/AnimNode_MotionMatching.h
--- a/Engine/Plugins/Experimental/Animation/PoseSearch/Source/Runtime/Public/PoseSearch/AnimNode_MotionMatching.h
+++ b/Engine/Plugins/Experimental/Animation/PoseSearch/Source/Runtime/Public/PoseSearch/AnimNode_MotionMatching.h
@@ -47,6 +47,12 @@ public:
 #if WITH_EDITORONLY_DATA
 	UPROPERTY(EditAnywhere, Category=Debug)
 	bool bDebugDraw = false;
+
+	UPROPERTY(EditAnywhere, Category=Debug)
+	bool bDebugDrawGoal = true;
+
+	UPROPERTY(EditAnywhere, Category=Debug)
+	bool bDebugDrawMatch = true;
 #endif
 
 	// FAnimNode_Base interface
diff --git a/Engine/Plugins/Experimental/Animation/PoseSearch/Source/Runtime/Public/PoseSearch/PoseSearch.h b/Engine/Plugins/Experimental/Animation/PoseSearch/Source/Runtime/Public/PoseSearch/PoseSearch.h
--- a/Engine/Plugins/Experimental/Animation/PoseSearch/Source/Runtime/Public/PoseSearch/PoseSearch.h
+++ b/Engine/Plugins/Experimental/Animation/PoseSearch/Source/Runtime/Public/PoseSearch/PoseSearch.h
@@ -471,7 +471,7 @@ public:
 	TArrayView<const float> GetNormalizedValues() const { return ValuesNormalized; }
 
 	void SetTransform(FPoseSearchFeatureDesc Feature, const FTransform& Transform);
-	void SetTransformDerivative(FPoseSearchFeatureDesc Feature, const FTransform& Transform, const FTransform& PrevTransform, float DeltaTime);
+	void SetTransformVelocity(FPoseSearchFeatureDesc Feature, const FTransform& Transform, const FTransform& PrevTransform, float DeltaTime);
 	void SetPosition(FPoseSearchFeatureDesc Feature, const FVector& Translation);
 	void SetRotation(FPoseSearchFeatureDesc Feature, const FQuat& Rotation);
 	void SetLinearVelocity(FPoseSearchFeatureDesc Feature, const FTransform& Transform, const FTransform& PrevTransform, float DeltaTime);
