PoseSearch: Distance based trajectory prediction and history for motion matching

Refactored FAnimInput_CapsulePrediction a bit to make room for distance
based predictions. Prediction math has been moved into step functions that
accept configuration parameters and manipulate prediction state. States can
be interpolated for sampling in between steps. Distance based prediction runs
the stepper forward until all the distances have been hit or a maximum step
count is reached. This could be more refined in the future, perhaps with
increasing step intervals when the requested distance offsets are very far away.

Capsule prediction has separate properties for a distanced based goal and
schema. This enabled side-by-side time based and distanced based motion
matching comparison.

PoseHistory can now fulfill requests for distance based samples by searching
for the distance marker, converting to time, and then taking a time sample.
History retention is still time based regardless of distance matching. We could
add a distance based retention mechanic in the future, but for now I've just
increased the duration of the recording buffer.

FPoseSearchFeatureVectorBuilder can now add past trajectory distance features.

Debug drawing can render distanced based trajectories

Cleanup:
- Removed quantized time offsets from the schema. They are no longer needed
  after the previous indexing refactor.
- Renamed some functions that return bools to Try*

diff --git a/Engine/Plugins/Experimental/Animation/PoseSearch/Source/Runtime/Private/AnimNode_MotionMatching.cpp b/Engine/Plugins/Experimental/Animation/PoseSearch/Source/Runtime/Private/AnimNode_MotionMatching.cpp
--- a/Engine/Plugins/Experimental/Animation/PoseSearch/Source/Runtime/Private/AnimNode_MotionMatching.cpp
+++ b/Engine/Plugins/Experimental/Animation/PoseSearch/Source/Runtime/Private/AnimNode_MotionMatching.cpp
@@ -74,7 +74,7 @@ void FAnimNode_MotionMatching::Update_AnyThread(const FAnimationUpdateContext& C
 			if (PoseHistoryProvider)
 			{
 				UE::PoseSearch::FPoseHistory& History = PoseHistoryProvider->GetPoseHistory();
-				ComposedQuery.SetPoseFeatures(&History);
+				ComposedQuery.TrySetPoseFeatures(&History);
 			}
 		}
 
@@ -217,7 +217,7 @@ void FAnimNode_MotionMatching::ComposeQuery(const FAnimationBaseContext& Context
 	if (PoseHistoryProvider)
 	{
 		UE::PoseSearch::FPoseHistory& History = PoseHistoryProvider->GetPoseHistory();
-		ComposedQuery.SetPastTrajectoryFeatures(&History);
+		ComposedQuery.TrySetPastTrajectoryFeatures(&History);
 	}
 
 	// Merge goal features into the query vector
diff --git a/Engine/Plugins/Experimental/Animation/PoseSearch/Source/Runtime/Private/AnimNode_PoseSearchHistoryCollector.cpp b/Engine/Plugins/Experimental/Animation/PoseSearch/Source/Runtime/Private/AnimNode_PoseSearchHistoryCollector.cpp
--- a/Engine/Plugins/Experimental/Animation/PoseSearch/Source/Runtime/Private/AnimNode_PoseSearchHistoryCollector.cpp
+++ b/Engine/Plugins/Experimental/Animation/PoseSearch/Source/Runtime/Private/AnimNode_PoseSearchHistoryCollector.cpp
@@ -42,7 +42,7 @@ void FAnimNode_PoseSearchHistoryCollector::Initialize_AnyThread(const FAnimation
 	Super::Initialize_AnyThread(Context);
 
 	// @@ need to do this once based on descendant node's (or input param?) search schema, not every node init
-	PoseHistory.Init(32, 1.0f);
+	PoseHistory.Init(128, 4.0f);
 
 	UE::Anim::TScopedGraphMessage<UE::PoseSearch::Private::FPoseHistoryProvider> ScopedMessage(Context, this);
 
diff --git a/Engine/Plugins/Experimental/Animation/PoseSearch/Source/Runtime/Private/PoseSearch.cpp b/Engine/Plugins/Experimental/Animation/PoseSearch/Source/Runtime/Private/PoseSearch.cpp
--- a/Engine/Plugins/Experimental/Animation/PoseSearch/Source/Runtime/Private/PoseSearch.cpp
+++ b/Engine/Plugins/Experimental/Animation/PoseSearch/Source/Runtime/Private/PoseSearch.cpp
@@ -27,6 +27,8 @@ IMPLEMENT_ANIMGRAPH_MESSAGE(UE::PoseSearch::IPoseHistoryProvider);
 
 namespace UE { namespace PoseSearch {
 
+DEFINE_LOG_CATEGORY(LogPoseSearch);
+
 //////////////////////////////////////////////////////////////////////////
 // Constants and utilities
 
@@ -282,9 +284,6 @@ void UPoseSearchSchema::PreSave(FObjectPreSaveContext ObjectSaveContext)
 	TrajectorySampleTimes.Sort(TLess<>());
 	TrajectorySampleDistances.Sort(TLess<>());
 
-	ConvertTimesToOffsets(PoseSampleTimes, PoseSampleOffsets);
-	ConvertTimesToOffsets(TrajectorySampleTimes, TrajectorySampleOffsets);
-
 	GenerateLayout();
 	ResolveBoneReferences();
 
@@ -325,26 +324,26 @@ void UPoseSearchSchema::GenerateLayout()
 	Layout.Reset();
 
 	// Time domain trajectory positions
-	if (bUseTrajectoryPositions && TrajectorySampleOffsets.Num())
+	if (bUseTrajectoryPositions && TrajectorySampleTimes.Num())
 	{
 		FPoseSearchFeatureDesc Feature;
 		Feature.SchemaBoneIdx = FPoseSearchFeatureDesc::TrajectoryBoneIndex;
 		Feature.Domain = EPoseSearchFeatureDomain::Time;
 		Feature.Type = EPoseSearchFeatureType::Position;
-		for (Feature.SubsampleIdx = 0; Feature.SubsampleIdx != TrajectorySampleOffsets.Num(); ++Feature.SubsampleIdx)
+		for (Feature.SubsampleIdx = 0; Feature.SubsampleIdx != TrajectorySampleTimes.Num(); ++Feature.SubsampleIdx)
 		{
 			Layout.Features.Add(Feature);
 		}
 	}
 
 	// Time domain trajectory linear velocities
-	if (bUseTrajectoryVelocities && TrajectorySampleOffsets.Num())
+	if (bUseTrajectoryVelocities && TrajectorySampleTimes.Num())
 	{
 		FPoseSearchFeatureDesc Feature;
 		Feature.SchemaBoneIdx = FPoseSearchFeatureDesc::TrajectoryBoneIndex;
 		Feature.Domain = EPoseSearchFeatureDomain::Time;
 		Feature.Type = EPoseSearchFeatureType::LinearVelocity;
-		for (Feature.SubsampleIdx = 0; Feature.SubsampleIdx != TrajectorySampleOffsets.Num(); ++Feature.SubsampleIdx)
+		for (Feature.SubsampleIdx = 0; Feature.SubsampleIdx != TrajectorySampleTimes.Num(); ++Feature.SubsampleIdx)
 		{
 			Layout.Features.Add(Feature);
 		}
@@ -377,12 +376,12 @@ void UPoseSearchSchema::GenerateLayout()
 	}
 
 	// Time domain bone positions
-	if (bUseBonePositions && PoseSampleOffsets.Num())
+	if (bUseBonePositions && PoseSampleTimes.Num())
 	{
 		FPoseSearchFeatureDesc Feature;
 		Feature.Domain = EPoseSearchFeatureDomain::Time;
 		Feature.Type = EPoseSearchFeatureType::Position;
-		for (Feature.SubsampleIdx = 0; Feature.SubsampleIdx != PoseSampleOffsets.Num(); ++Feature.SubsampleIdx)
+		for (Feature.SubsampleIdx = 0; Feature.SubsampleIdx != PoseSampleTimes.Num(); ++Feature.SubsampleIdx)
 		{
 			for (Feature.SchemaBoneIdx = 0; Feature.SchemaBoneIdx != Bones.Num(); ++Feature.SchemaBoneIdx)
 			{
@@ -392,12 +391,12 @@ void UPoseSearchSchema::GenerateLayout()
 	}
 
 	// Time domain bone linear velocities
-	if (bUseBoneVelocities && PoseSampleOffsets.Num())
+	if (bUseBoneVelocities && PoseSampleTimes.Num())
 	{
 		FPoseSearchFeatureDesc Feature;
 		Feature.Domain = EPoseSearchFeatureDomain::Time;
 		Feature.Type = EPoseSearchFeatureType::LinearVelocity;
-		for (Feature.SubsampleIdx = 0; Feature.SubsampleIdx != PoseSampleOffsets.Num(); ++Feature.SubsampleIdx)
+		for (Feature.SubsampleIdx = 0; Feature.SubsampleIdx != PoseSampleTimes.Num(); ++Feature.SubsampleIdx)
 		{
 			for (Feature.SchemaBoneIdx = 0; Feature.SchemaBoneIdx != Bones.Num(); ++Feature.SchemaBoneIdx)
 			{
@@ -433,16 +432,6 @@ void UPoseSearchSchema::ResolveBoneReferences()
 	}
 }
 
-void UPoseSearchSchema::ConvertTimesToOffsets(TArrayView<const float> SampleTimes, TArray<int32>& OutSampleOffsets)
-{
-	OutSampleOffsets.SetNum(SampleTimes.Num());
-
-	for (int32 Idx = 0; Idx != SampleTimes.Num(); ++Idx)
-	{
-		OutSampleOffsets[Idx] = FMath::RoundToInt(SampleTimes[Idx] * SampleRate);
-	}
-}
-
 
 //////////////////////////////////////////////////////////////////////////
 // FPoseSearchBiasWeights
@@ -807,7 +796,7 @@ void FPoseSearchFeatureVectorBuilder::SetVector(FPoseSearchFeatureDesc Element,
 	}
 }
 
-bool FPoseSearchFeatureVectorBuilder::SetPoseFeatures(UE::PoseSearch::FPoseHistory* History)
+bool FPoseSearchFeatureVectorBuilder::TrySetPoseFeatures(UE::PoseSearch::FPoseHistory* History)
 {
 	check(Schema && Schema->IsValid());
 	check(History);
@@ -815,14 +804,19 @@ bool FPoseSearchFeatureVectorBuilder::SetPoseFeatures(UE::PoseSearch::FPoseHisto
 	FPoseSearchFeatureDesc Feature;
 	Feature.Domain = EPoseSearchFeatureDomain::Time;
 
-	for (int32 SchemaSubsampleIdx = 0; SchemaSubsampleIdx != Schema->PoseSampleOffsets.Num(); ++SchemaSubsampleIdx)
+	for (int32 SchemaSubsampleIdx = 0; SchemaSubsampleIdx != Schema->PoseSampleTimes.Num(); ++SchemaSubsampleIdx)
 	{
 		Feature.SubsampleIdx = SchemaSubsampleIdx;
 
-		int32 Offset = Schema->PoseSampleOffsets[SchemaSubsampleIdx];
-		float TimeDelta = -Offset * Schema->SamplingInterval;
+		// Stop when we've reached future samples
+		float SampleTime = Schema->PoseSampleTimes[SchemaSubsampleIdx];
+		if (SampleTime > 0.0f)
+		{
+			break;
+		}
 
-		if (!History->SamplePose(TimeDelta, Schema->Skeleton->GetReferenceSkeleton(), Schema->BoneIndicesWithParents))
+		float SecondsAgo = -SampleTime;
+		if (!History->TrySamplePose(SecondsAgo, Schema->Skeleton->GetReferenceSkeleton(), Schema->BoneIndicesWithParents))
 		{
 			return false;
 		}
@@ -837,47 +831,94 @@ bool FPoseSearchFeatureVectorBuilder::SetPoseFeatures(UE::PoseSearch::FPoseHisto
 			const FTransform& Transform = ComponentPose[SkeletonBoneIndex];
 			const FTransform& PrevTransform = ComponentPrevPose[SkeletonBoneIndex];
 			SetTransform(Feature, Transform);
-			SetTransformVelocity(Feature, Transform, PrevTransform, History->GetSampleInterval());
+			SetTransformVelocity(Feature, Transform, PrevTransform, History->GetSampleTimeInterval());
 		}
 	}
 
 	return true;
 }
 
-bool FPoseSearchFeatureVectorBuilder::SetPastTrajectoryFeatures(UE::PoseSearch::FPoseHistory* History)
+bool FPoseSearchFeatureVectorBuilder::TrySetPastTrajectoryFeatures(UE::PoseSearch::FPoseHistory* History)
 {
 	check(Schema && Schema->IsValid());
 	check(History);
 
+	bool bSuccess = TrySetPastTrajectoryTimeFeatures(History);
+	bSuccess &= TrySetPastTrajectoryDistanceFeatures(History);
+
+	return bSuccess;
+}
+
+bool FPoseSearchFeatureVectorBuilder::TrySetPastTrajectoryTimeFeatures(UE::PoseSearch::FPoseHistory* History)
+{
 	FPoseSearchFeatureDesc Feature;
 	Feature.Domain = EPoseSearchFeatureDomain::Time;
 	Feature.SchemaBoneIdx = FPoseSearchFeatureDesc::TrajectoryBoneIndex;
 
-	for (int32 SchemaSubsampleIdx = 0; SchemaSubsampleIdx != Schema->TrajectorySampleOffsets.Num(); ++SchemaSubsampleIdx)
+	for (int32 SchemaSubsampleIdx = 0; SchemaSubsampleIdx != Schema->TrajectorySampleTimes.Num(); ++SchemaSubsampleIdx)
 	{
 		Feature.SubsampleIdx = SchemaSubsampleIdx;
+		
+		// Stop when we've reached future samples
+		float SampleTime = Schema->TrajectorySampleTimes[SchemaSubsampleIdx];
+		if (SampleTime > 0.0f)
+		{
+			break;
+		}
+
+		float SecondsAgo = -SampleTime;
+		FTransform WorldComponentTransform;
+		if (!History->TrySampleRootTimeBased(SecondsAgo, &WorldComponentTransform))
+		{
+			return false;
+		}
+
+		FTransform WorldPrevComponentTransform;
+		if (!History->TrySampleRootTimeBased(SecondsAgo + History->GetSampleTimeInterval(), &WorldPrevComponentTransform))
+		{
+			return false;
+		}
+
+		SetTransform(Feature, WorldComponentTransform);
+		SetTransformVelocity(Feature, WorldComponentTransform, WorldPrevComponentTransform, History->GetSampleTimeInterval());
+	}
+
+	return true;
+}
 
-		int32 SubsampleIndex = Schema->TrajectorySampleOffsets[SchemaSubsampleIdx];
-		if (SubsampleIndex > 0)
+bool FPoseSearchFeatureVectorBuilder::TrySetPastTrajectoryDistanceFeatures(UE::PoseSearch::FPoseHistory* History)
+{
+	FPoseSearchFeatureDesc Feature;
+	Feature.Domain = EPoseSearchFeatureDomain::Distance;
+	Feature.SchemaBoneIdx = FPoseSearchFeatureDesc::TrajectoryBoneIndex;
+
+	for (int32 SchemaSubsampleIdx = 0; SchemaSubsampleIdx != Schema->TrajectorySampleDistances.Num(); ++SchemaSubsampleIdx)
+	{
+		Feature.SubsampleIdx = SchemaSubsampleIdx;
+		
+		// Stop when we've reached future samples
+		float SampleDistance = Schema->TrajectorySampleDistances[SchemaSubsampleIdx];
+		if (SampleDistance > 0.0f)
 		{
 			break;
 		}
 
-		float SecondsAgo = -SubsampleIndex * Schema->SamplingInterval;
+		float SampleTime = 0.0f;
+		float PastDistance = -SampleDistance;
 		FTransform WorldComponentTransform;
-		if (!History->SampleRoot(SecondsAgo, &WorldComponentTransform))
+		if (!History->TrySampleRootDistanceBased(PastDistance, &WorldComponentTransform, &SampleTime))
 		{
 			return false;
 		}
 
 		FTransform WorldPrevComponentTransform;
-		if (!History->SampleRoot(SecondsAgo + History->GetSampleInterval(), &WorldPrevComponentTransform))
+		if (!History->TrySampleRootTimeBased(SampleTime + History->GetSampleTimeInterval(), &WorldPrevComponentTransform))
 		{
 			return false;
 		}
 
 		SetTransform(Feature, WorldComponentTransform);
-		SetTransformVelocity(Feature, WorldComponentTransform, WorldPrevComponentTransform, History->GetSampleInterval());
+		SetTransformVelocity(Feature, WorldComponentTransform, WorldPrevComponentTransform, History->GetSampleTimeInterval());
 	}
 
 	return true;
@@ -933,6 +974,11 @@ bool FPoseSearchFeatureVectorBuilder::IsInitialized() const
 	return (Schema != nullptr) && (Values.Num() == Schema->Layout.NumFloats);
 }
 
+bool FPoseSearchFeatureVectorBuilder::IsInitializedForSchema(const UPoseSearchSchema* InSchema) const
+{
+	return (Schema == InSchema) && IsInitialized();
+}
+
 bool FPoseSearchFeatureVectorBuilder::IsComplete() const
 {
 	return NumFeaturesAdded == Schema->Layout.Features.Num();
@@ -988,7 +1034,7 @@ void FPoseHistory::Init(const FPoseHistory& History)
 	TimeHorizon = History.TimeHorizon;
 }
 
-bool FPoseHistory::SampleLocalPose(float SecondsAgo, const TArray<FBoneIndexType>& RequiredBones, TArray<FTransform>& LocalPose)
+bool FPoseHistory::TrySampleLocalPose(float SecondsAgo, const TArray<FBoneIndexType>& RequiredBones, TArray<FTransform>& LocalPose)
 {
 	int32 NextIdx = LowerBound(Knots.begin(), Knots.end(), SecondsAgo, TGreater<>());
 	if (NextIdx <= 0 || NextIdx >= Knots.Num())
@@ -1029,13 +1075,13 @@ bool FPoseHistory::SampleLocalPose(float SecondsAgo, const TArray<FBoneIndexType
 	return true;
 }
 
-bool FPoseHistory::SamplePose(float SecondsAgo, const FReferenceSkeleton& RefSkeleton, const TArray<FBoneIndexType>& RequiredBones)
+bool FPoseHistory::TrySamplePose(float SecondsAgo, const FReferenceSkeleton& RefSkeleton, const TArray<FBoneIndexType>& RequiredBones)
 {
 	// Compute local space pose at requested time
-	bool bSampled = SampleLocalPose(SecondsAgo, RequiredBones, SampledLocalPose);
+	bool bSampled = TrySampleLocalPose(SecondsAgo, RequiredBones, SampledLocalPose);
 
 	// Compute local space pose one sample interval in the past
-	bSampled = bSampled && SampleLocalPose(SecondsAgo + GetSampleInterval(), RequiredBones, SampledPrevLocalPose);
+	bSampled = bSampled && TrySampleLocalPose(SecondsAgo + GetSampleTimeInterval(), RequiredBones, SampledPrevLocalPose);
 
 	// Convert local to component space
 	if (bSampled)
@@ -1047,7 +1093,7 @@ bool FPoseHistory::SamplePose(float SecondsAgo, const FReferenceSkeleton& RefSke
 	return bSampled;
 }
 
-bool FPoseHistory::SampleRoot(float SecondsAgo, FTransform* OutTransform) const
+bool FPoseHistory::TrySampleRootTimeBased(float SecondsAgo, FTransform* OutTransform) const
 {
 	int32 NextIdx = LowerBound(Knots.begin(), Knots.end(), SecondsAgo, TGreater<>());
 	if (NextIdx <= 0 || NextIdx >= Knots.Num())
@@ -1074,6 +1120,34 @@ bool FPoseHistory::SampleRoot(float SecondsAgo, FTransform* OutTransform) const
 	return true;
 }
 
+bool FPoseHistory::TrySampleRootDistanceBased(float SampleDistance, FTransform* OutTransform, float* OutSampleTime) const
+{
+	float SampleTime = 0.0f;
+
+	// Find the relevant distance samples and convert to time
+	double NextDistance = 0.0;
+	double PrevDistance = 0.0;
+	for (int32 NextIdx = Poses.Num() - 2; NextIdx >= 0; --NextIdx)
+	{
+		FTransform PrevTransform = Poses[NextIdx+1].WorldComponentTransform;
+		FTransform NextTransform = Poses[NextIdx].WorldComponentTransform;
+
+		NextDistance += FVector::Distance(PrevTransform.GetTranslation(), NextTransform.GetTranslation());
+		if (NextDistance >= SampleDistance)
+		{
+			float DistanceSampleAlpha = (float)FMath::GetRangePct(PrevDistance, NextDistance, (double)SampleDistance);
+			SampleTime = FMath::Lerp(Knots[NextIdx+1], Knots[NextIdx], DistanceSampleAlpha);
+			break;
+		}
+
+		PrevDistance = NextDistance;
+	}
+
+	bool bResult = TrySampleRootTimeBased(SampleTime, OutTransform);
+	*OutSampleTime = SampleTime;
+	return bResult;
+}
+
 void FPoseHistory::Update(float SecondsElapsed, const FPoseContext& PoseContext)
 {
 	// Age our elapsed times
@@ -1095,7 +1169,7 @@ void FPoseHistory::Update(float SecondsElapsed, const FPoseContext& PoseContext)
 		// want to evenly distribute knots across the entire history buffer so we only push additional
 		// poses when enough time has elapsed.
 
-		const float SampleInterval = GetSampleInterval();
+		const float SampleInterval = GetSampleTimeInterval();
 
 		bool bCanEvictOldest = Knots[1] >= TimeHorizon + SampleInterval;
 		bool bShouldPushNewest = Knots[Knots.Num() - 2] >= SampleInterval;
@@ -1118,7 +1192,7 @@ void FPoseHistory::Update(float SecondsElapsed, const FPoseContext& PoseContext)
 	CurrentPose.WorldComponentTransform = PoseContext.AnimInstanceProxy->GetComponentTransform();
 }
 
-float FPoseHistory::GetSampleInterval() const
+float FPoseHistory::GetSampleTimeInterval() const
 {
 	// Reserve one knot for computing derivatives at the time horizon
 	return TimeHorizon / (Knots.Max() - 1);
@@ -1290,7 +1364,7 @@ public:
 		const UPoseSearchSchema* Schema = nullptr;
 		const UAnimSequence* Sequence = nullptr;
 		bool bLoopable = false;
-		int32 DistanceSamplingRate = 30;
+		int32 DistanceSamplingRate = 60;
 	} Input;
 
 	struct FOutput
@@ -1560,6 +1634,8 @@ void FSequenceIndexer::Process()
 {
 	for (int32 SampleIdx = Output.FirstIndexedSample; SampleIdx <= Output.LastIndexedSample; ++SampleIdx)
 	{
+		FMemMark Mark(FMemStack::Get());
+
 		SampleBegin(SampleIdx);
 
 		AddPoseFeatures(SampleIdx);
@@ -1807,7 +1883,7 @@ void FSequenceIndexer::AddPoseFeatures(int32 SampleIdx)
 	constexpr float FiniteDelta = 1 / 60.0f;
 	constexpr int32 NumFiniteDiffTerms = 3;
 
-	if (Input.Schema->Bones.IsEmpty() || Input.Schema->PoseSampleOffsets.IsEmpty())
+	if (Input.Schema->Bones.IsEmpty() || Input.Schema->PoseSampleTimes.IsEmpty())
 	{
 		return;
 	}
@@ -1838,7 +1914,7 @@ void FSequenceIndexer::AddPoseFeatures(int32 SampleIdx)
 	float SampleTime = FMath::Min(SampleIdx * Input.Schema->SamplingInterval, Input.MainSequence->Output.PlayLength);
 	FSampleInfo Origin = GetSampleInfo(SampleTime);
 	
-	for (int32 SchemaSubsampleIdx = 0; SchemaSubsampleIdx != Input.Schema->PoseSampleOffsets.Num(); ++SchemaSubsampleIdx)
+	for (int32 SchemaSubsampleIdx = 0; SchemaSubsampleIdx != Input.Schema->PoseSampleTimes.Num(); ++SchemaSubsampleIdx)
 	{
 		Feature.SubsampleIdx = SchemaSubsampleIdx;
 
@@ -1861,7 +1937,7 @@ void FSequenceIndexer::AddPoseFeatures(int32 SampleIdx)
 
 		// Get each bone's component transform, velocity, and acceleration and add accumulated root motion at this time offset
 		// Think of this process as freezing the character in place (at SampleTime) and then tracing the paths of their joints
-		// as they move through space from past to present to future (at times indicated by PoseSampleOffsets).
+		// as they move through space from past to present to future (at times indicated by PoseSampleTimes).
 		for (int32 SchemaBoneIndex = 0; SchemaBoneIndex != Input.Schema->NumBones(); ++SchemaBoneIndex)
 		{
 			Feature.SchemaBoneIdx = SchemaBoneIndex;
@@ -1961,16 +2037,18 @@ void FSequenceIndexer::AddTrajectoryDistanceFeatures(int32 SampleIdx)
 //////////////////////////////////////////////////////////////////////////
 // PoseSearch API
 
-static void DrawTrajectoryFeatures(const FDebugDrawParams& DrawParams, const FFeatureVectorReader& Reader)
+static void DrawTrajectoryFeatures(const FDebugDrawParams& DrawParams, const FFeatureVectorReader& Reader, EPoseSearchFeatureDomain Domain)
 {
 	const float LifeTime = DrawParams.DefaultLifeTime;
 	const uint8 DepthPriority = ESceneDepthPriorityGroup::SDPG_Foreground + 2;
 
 	FPoseSearchFeatureDesc Feature;
-	Feature.Domain = EPoseSearchFeatureDomain::Time;
+	Feature.Domain = Domain;
 	Feature.SchemaBoneIdx = FPoseSearchFeatureDesc::TrajectoryBoneIndex;
 
-	const int32 NumSubsamples = DrawParams.GetSchema()->TrajectorySampleOffsets.Num();
+	const int32 NumSubsamples = Domain == EPoseSearchFeatureDomain::Time ?
+		DrawParams.GetSchema()->TrajectorySampleTimes.Num() :
+		DrawParams.GetSchema()->TrajectorySampleDistances.Num();
 
 	if (NumSubsamples == 0)
 	{
@@ -2022,7 +2100,7 @@ static void DrawPoseFeatures(const FDebugDrawParams& DrawParams, const FFeatureV
 	FPoseSearchFeatureDesc Feature;
 	Feature.Domain = EPoseSearchFeatureDomain::Time;
 
-	const int32 NumSubsamples = Schema->PoseSampleOffsets.Num();
+	const int32 NumSubsamples = Schema->PoseSampleTimes.Num();
 	const int32 NumBones = Schema->Bones.Num();
 
 	if ((NumSubsamples * NumBones) == 0)
@@ -2067,7 +2145,8 @@ static void DrawPoseFeatures(const FDebugDrawParams& DrawParams, const FFeatureV
 static void DrawFeatureVector(const FDebugDrawParams& DrawParams, const FFeatureVectorReader& Reader)
 {
 	DrawPoseFeatures(DrawParams, Reader);
-	DrawTrajectoryFeatures(DrawParams, Reader);
+	DrawTrajectoryFeatures(DrawParams, Reader, EPoseSearchFeatureDomain::Time);
+	DrawTrajectoryFeatures(DrawParams, Reader, EPoseSearchFeatureDomain::Distance);
 }
 
 static void DrawSearchIndex(const FDebugDrawParams& DrawParams)
@@ -2970,7 +3049,11 @@ UE::Anim::IPoseSearchProvider::FSearchResult FModule::Search(const FAnimationBas
 	FPoseSearchFeatureVectorBuilder& QueryBuilder = PoseHistory.GetQueryBuilder();
 
 	QueryBuilder.Init(MetaData->Schema);
-	QueryBuilder.SetPoseFeatures(&PoseHistory);
+	if (!QueryBuilder.TrySetPoseFeatures(&PoseHistory))
+	{
+		return ProviderResult;
+	}
+
 	QueryBuilder.Normalize(MetaData->SearchIndex);
 
 	::UE::PoseSearch::FSearchResult Result = ::UE::PoseSearch::Search(Sequence, QueryBuilder.GetNormalizedValues());
diff --git a/Engine/Plugins/Experimental/Animation/PoseSearch/Source/Runtime/Public/PoseSearch/PoseSearch.h b/Engine/Plugins/Experimental/Animation/PoseSearch/Source/Runtime/Public/PoseSearch/PoseSearch.h
--- a/Engine/Plugins/Experimental/Animation/PoseSearch/Source/Runtime/Public/PoseSearch/PoseSearch.h
+++ b/Engine/Plugins/Experimental/Animation/PoseSearch/Source/Runtime/Public/PoseSearch/PoseSearch.h
@@ -24,7 +24,11 @@ struct FPoseContext;
 struct FReferenceSkeleton;
 
 namespace UE { namespace PoseSearch {
+
 class FPoseHistory;
+
+DECLARE_LOG_CATEGORY_EXTERN(LogPoseSearch, Log, All);
+
 }}
 
 UENUM()
@@ -203,12 +207,6 @@ public:
 	UPROPERTY(EditAnywhere, Category = "Schema")
 	EPoseSearchDataPreprocessor DataPreprocessor = EPoseSearchDataPreprocessor::Automatic;
 
-	UPROPERTY()
-	TArray<int32> TrajectorySampleOffsets;
-
-	UPROPERTY()
-	TArray<int32> PoseSampleOffsets;
-
 	UPROPERTY()
 	EPoseSearchDataPreprocessor EffectiveDataPreprocessor = EPoseSearchDataPreprocessor::Invalid;
 
@@ -239,7 +237,6 @@ public: // IBoneReferenceSkeletonProvider
 private:
 	void GenerateLayout();
 	void ResolveBoneReferences();
-	void ConvertTimesToOffsets(TArrayView<const float> SampleTimes, TArray<int32>& OutSampleOffsets);
 };
 
 USTRUCT()
@@ -467,6 +464,8 @@ public:
 	void Init(const UPoseSearchSchema* Schema);
 	void ResetFeatures();
 
+	const UPoseSearchSchema* GetSchema() const { return Schema; }
+
 	TArrayView<const float> GetValues() const { return Values; }
 	TArrayView<const float> GetNormalizedValues() const { return ValuesNormalized; }
 
@@ -477,8 +476,8 @@ public:
 	void SetLinearVelocity(FPoseSearchFeatureDesc Feature, const FTransform& Transform, const FTransform& PrevTransform, float DeltaTime);
 	void SetAngularVelocity(FPoseSearchFeatureDesc Feature, const FTransform& Transform, const FTransform& PrevTransform, float DeltaTime);
 	void SetVector(FPoseSearchFeatureDesc Feature, const FVector& Vector);
-	bool SetPoseFeatures(UE::PoseSearch::FPoseHistory* History);
-	bool SetPastTrajectoryFeatures(UE::PoseSearch::FPoseHistory* History);
+	bool TrySetPoseFeatures(UE::PoseSearch::FPoseHistory* History);
+	bool TrySetPastTrajectoryFeatures(UE::PoseSearch::FPoseHistory* History);
 
 	void CopyFromSearchIndex(const FPoseSearchIndex& SearchIndex, int32 PoseIdx);
 	void CopyFeature(const FPoseSearchFeatureVectorBuilder& OtherBuilder, int32 FeatureIdx);
@@ -486,12 +485,16 @@ public:
 	void MergeReplace(const FPoseSearchFeatureVectorBuilder& OtherBuilder);
 
 	bool IsInitialized() const;
+	bool IsInitializedForSchema(const UPoseSearchSchema* Schema) const;
 	bool IsComplete() const;
 	bool IsCompatible(const FPoseSearchFeatureVectorBuilder& OtherBuilder) const;
 
 	void Normalize(const FPoseSearchIndex& ForSearchIndex);
 
 private:
+	bool TrySetPastTrajectoryTimeFeatures(UE::PoseSearch::FPoseHistory* History);
+	bool TrySetPastTrajectoryDistanceFeatures(UE::PoseSearch::FPoseHistory* History);
+
 	UPROPERTY(Transient)
 	const UPoseSearchSchema* Schema = nullptr;
 
@@ -511,10 +514,12 @@ class POSESEARCH_API FPoseHistory
 public:
 	void Init(int32 InNumPoses, float InTimeHorizon);
 	void Init(const FPoseHistory& History);
-	bool SamplePose(float SecondsAgo, const FReferenceSkeleton& RefSkeleton, const TArray<FBoneIndexType>& RequiredBones);
-	bool SampleRoot(float SecondsAgo, FTransform* OutTransform) const;
+	bool TrySamplePose(float SecondsAgo, const FReferenceSkeleton& RefSkeleton, const TArray<FBoneIndexType>& RequiredBones);
+	bool TrySampleRootTimeBased(float SecondsAgo, FTransform* OutTransform) const;
+	bool TrySampleRootDistanceBased(float PastDistance, FTransform* OutTransform, float* OutSampleTime) const;
+
 	void Update(float SecondsElapsed, const FPoseContext& PoseContext);
-	float GetSampleInterval() const;
+	float GetSampleTimeInterval() const;
 	TArrayView<const FTransform> GetLocalPoseSample() const { return SampledLocalPose; }
 	TArrayView<const FTransform> GetComponentPoseSample() const { return SampledComponentPose; }
 	TArrayView<const FTransform> GetPrevLocalPoseSample() const { return SampledPrevLocalPose; }
@@ -523,7 +528,7 @@ public:
 	FPoseSearchFeatureVectorBuilder& GetQueryBuilder() { return QueryBuilder; }
 
 private:
-	bool SampleLocalPose(float Time, const TArray<FBoneIndexType>& RequiredBones, TArray<FTransform>& LocalPose);
+	bool TrySampleLocalPose(float Time, const TArray<FBoneIndexType>& RequiredBones, TArray<FTransform>& LocalPose);
 
 	struct FPose
 	{
