Makes PoseSearchDatabase changes work during runtime in FAnimNode_MotionMatching. Now when the database changes:
- Previous search results are reset
- If the schema changed, the ComposedQuery is reinitialized
- Exposed databse schemas to blueprint in different spots to better support dynamically changing the pose search database

diff --git a/Engine/Plugins/Experimental/Animation/PoseSearch/Source/Runtime/Private/AnimNode_MotionMatching.cpp b/Engine/Plugins/Experimental/Animation/PoseSearch/Source/Runtime/Private/AnimNode_MotionMatching.cpp
--- a/Engine/Plugins/Experimental/Animation/PoseSearch/Source/Runtime/Private/AnimNode_MotionMatching.cpp
+++ b/Engine/Plugins/Experimental/Animation/PoseSearch/Source/Runtime/Private/AnimNode_MotionMatching.cpp
@@ -19,9 +19,7 @@ void FAnimNode_MotionMatching::Initialize_AnyThread(const FAnimationInitializeCo
 
 	GetEvaluateGraphExposedInputs().Execute(Context);
 
-	DbPoseIdx = INDEX_NONE;
-	DbSequenceIdx = INDEX_NONE;
-	ElapsedPoseJumpTime = SearchThrottleTime;
+	InitNewDatabaseSearch();
 
 	Source.SetLinkNode(&SequencePlayerNode);
 	Source.Initialize(Context);
@@ -45,7 +43,12 @@ void FAnimNode_MotionMatching::Update_AnyThread(const FAnimationUpdateContext& C
 
 	if (IsValidForSearch())
 	{
-		if (!ComposedQuery.IsInitialized())
+		if (PreviousDatabase != Database)
+		{
+			InitNewDatabaseSearch();
+		}
+
+		if (!ComposedQuery.IsInitializedForSchema(Database->Schema))
 		{
 			ComposedQuery.Init(Database->Schema);
 		}
@@ -253,4 +256,12 @@ void FAnimNode_MotionMatching::JumpToPose(const FAnimationUpdateContext& Context
 	}
 }
 
+void FAnimNode_MotionMatching::InitNewDatabaseSearch()
+{
+	DbPoseIdx = INDEX_NONE;
+	DbSequenceIdx = INDEX_NONE;
+	ElapsedPoseJumpTime = SearchThrottleTime;
+	PreviousDatabase = Database;
+}
+
 #undef LOCTEXT_NAMESPACE
\ No newline at end of file
diff --git a/Engine/Plugins/Experimental/Animation/PoseSearch/Source/Runtime/Public/PoseSearch/AnimNode_MotionMatching.h b/Engine/Plugins/Experimental/Animation/PoseSearch/Source/Runtime/Public/PoseSearch/AnimNode_MotionMatching.h
--- a/Engine/Plugins/Experimental/Animation/PoseSearch/Source/Runtime/Public/PoseSearch/AnimNode_MotionMatching.h
+++ b/Engine/Plugins/Experimental/Animation/PoseSearch/Source/Runtime/Public/PoseSearch/AnimNode_MotionMatching.h
@@ -18,7 +18,7 @@ public:
 
 	// Collection of animations for motion matching
 	UPROPERTY(EditAnywhere, BlueprintReadWrite, Category=Settings, meta=(PinHiddenByDefault))
-	const UPoseSearchDatabase* Database = nullptr;
+	TObjectPtr<const UPoseSearchDatabase> Database;
 
 	UPROPERTY(EditAnywhere, BlueprintReadWrite, Category = Settings, meta = (PinHiddenByDefault))
 	FPoseSearchBiasWeightParams BiasWeights;
@@ -78,10 +78,15 @@ private:
 	// The current query feature vector used to search the database for pose candidates
 	FPoseSearchFeatureVectorBuilder ComposedQuery;
 
+	// When the current database is different from the previous one, the search parameters are reset
+	UPROPERTY(Transient)
+	TObjectPtr<const UPoseSearchDatabase> PreviousDatabase;
+
 	// Time since the last pose jump
 	float ElapsedPoseJumpTime = 0.0f;
 
 	bool IsValidForSearch() const;
 	void ComposeQuery(const FAnimationBaseContext& Context);
 	void JumpToPose(const FAnimationUpdateContext& Context, UE::PoseSearch::FDbSearchResult Result);
+	void InitNewDatabaseSearch();
 };
\ No newline at end of file
diff --git a/Engine/Plugins/Experimental/Animation/PoseSearch/Source/Runtime/Public/PoseSearch/PoseSearch.h b/Engine/Plugins/Experimental/Animation/PoseSearch/Source/Runtime/Public/PoseSearch/PoseSearch.h
--- a/Engine/Plugins/Experimental/Animation/PoseSearch/Source/Runtime/Public/PoseSearch/PoseSearch.h
+++ b/Engine/Plugins/Experimental/Animation/PoseSearch/Source/Runtime/Public/PoseSearch/PoseSearch.h
@@ -430,7 +430,7 @@ class POSESEARCH_API UPoseSearchDatabase : public UDataAsset
 	GENERATED_BODY()
 public:
 
-	UPROPERTY(EditAnywhere, Category="Database")
+	UPROPERTY(EditAnywhere, BlueprintReadOnly, Category="Database")
 	const UPoseSearchSchema* Schema;
 
 	UPROPERTY(EditAnywhere, Category="Database")
