Converts Motion Matching anim node into an asset player node to enable automatic transition rules. The node implements the interface by calling the same functions from the embedded sequence player node. In the anim graph node, GetAnimationAsset() returns nullptr since the PoseSearchDatabase is not an animation asset (yet, maybe later?).

diff --git a/Engine/Plugins/Experimental/Animation/PoseSearch/Source/Editor/Private/AnimGraphNode_MotionMatching.cpp b/Engine/Plugins/Experimental/Animation/PoseSearch/Source/Editor/Private/AnimGraphNode_MotionMatching.cpp
--- a/Engine/Plugins/Experimental/Animation/PoseSearch/Source/Editor/Private/AnimGraphNode_MotionMatching.cpp
+++ b/Engine/Plugins/Experimental/Animation/PoseSearch/Source/Editor/Private/AnimGraphNode_MotionMatching.cpp
@@ -34,4 +34,30 @@ void UAnimGraphNode_MotionMatching::GetOutputLinkAttributes(FNodeAttributeArray&
 	}
 }
 
+void UAnimGraphNode_MotionMatching::BakeDataDuringCompilation(class FCompilerResultsLog& MessageLog)
+{
+	UAnimBlueprint* AnimBlueprint = GetAnimBlueprint();
+	AnimBlueprint->FindOrAddGroup(Node.GetGroupName());
+}
+
+bool UAnimGraphNode_MotionMatching::DoesSupportTimeForTransitionGetter() const
+{
+	return true;
+}
+
+UAnimationAsset* UAnimGraphNode_MotionMatching::GetAnimationAsset() const
+{
+	return nullptr;
+}
+
+const TCHAR* UAnimGraphNode_MotionMatching::GetTimePropertyName() const
+{
+	return TEXT("InternalTimeAccumulator");
+}
+
+UScriptStruct* UAnimGraphNode_MotionMatching::GetTimePropertyStruct() const
+{
+	return FAnimNode_MotionMatching::StaticStruct();
+}
+
 #undef LOCTEXT_NAMESPACE
diff --git a/Engine/Plugins/Experimental/Animation/PoseSearch/Source/Editor/Public/AnimGraphNode_MotionMatching.h b/Engine/Plugins/Experimental/Animation/PoseSearch/Source/Editor/Public/AnimGraphNode_MotionMatching.h
--- a/Engine/Plugins/Experimental/Animation/PoseSearch/Source/Editor/Public/AnimGraphNode_MotionMatching.h
+++ b/Engine/Plugins/Experimental/Animation/PoseSearch/Source/Editor/Public/AnimGraphNode_MotionMatching.h
@@ -4,13 +4,13 @@
 
 #include "CoreMinimal.h"
 #include "UObject/ObjectMacros.h"
-#include "AnimGraphNode_Base.h"
+#include "AnimGraphNode_AssetPlayerBase.h"
 #include "PoseSearch/AnimNode_MotionMatching.h"
 #include "AnimGraphNode_MotionMatching.generated.h"
 
 
 UCLASS(MinimalAPI)
-class UAnimGraphNode_MotionMatching : public UAnimGraphNode_Base
+class UAnimGraphNode_MotionMatching : public UAnimGraphNode_AssetPlayerBase
 {
 	GENERATED_BODY()
 
@@ -22,4 +22,10 @@ class UAnimGraphNode_MotionMatching : public UAnimGraphNode_Base
 	virtual FText GetNodeTitle(ENodeTitleType::Type TitleType) const override;
 	virtual FText GetMenuCategory() const override;
 	virtual void GetOutputLinkAttributes(FNodeAttributeArray& OutAttributes) const override;
+
+	virtual void BakeDataDuringCompilation(class FCompilerResultsLog& MessageLog) override;
+	virtual bool DoesSupportTimeForTransitionGetter() const override;
+	virtual UAnimationAsset* GetAnimationAsset() const override;
+	virtual const TCHAR* GetTimePropertyName() const override;
+	virtual UScriptStruct* GetTimePropertyStruct() const override;
 };
diff --git a/Engine/Plugins/Experimental/Animation/PoseSearch/Source/Runtime/Private/AnimNode_MotionMatching.cpp b/Engine/Plugins/Experimental/Animation/PoseSearch/Source/Runtime/Private/AnimNode_MotionMatching.cpp
--- a/Engine/Plugins/Experimental/Animation/PoseSearch/Source/Runtime/Private/AnimNode_MotionMatching.cpp
+++ b/Engine/Plugins/Experimental/Animation/PoseSearch/Source/Runtime/Private/AnimNode_MotionMatching.cpp
@@ -32,9 +32,9 @@ void FAnimNode_MotionMatching::Evaluate_AnyThread(FPoseContext& Output)
 	Source.Evaluate(Output);
 }
 
-void FAnimNode_MotionMatching::Update_AnyThread(const FAnimationUpdateContext& Context)
+void FAnimNode_MotionMatching::UpdateAssetPlayer(const FAnimationUpdateContext& Context)
 {
-	DECLARE_SCOPE_HIERARCHICAL_COUNTER_ANIMNODE(Update_AnyThread);
+	DECLARE_SCOPE_HIERARCHICAL_COUNTER_ANIMNODE(UpdateAssetPlayer);
 
 	GetEvaluateGraphExposedInputs().Execute(Context);
 	// Note: What if the input database changes? That's not being handled at all!
@@ -264,4 +264,42 @@ void FAnimNode_MotionMatching::InitNewDatabaseSearch()
 	PreviousDatabase = Database;
 }
 
+// FAnimNode_AssetPlayerBase interface
+float FAnimNode_MotionMatching::GetAccumulatedTime() const
+{
+	return SequencePlayerNode.GetAccumulatedTime();
+}
+
+UAnimationAsset* FAnimNode_MotionMatching::GetAnimAsset() const
+{
+	return SequencePlayerNode.GetAnimAsset();
+}
+
+float FAnimNode_MotionMatching::GetCurrentAssetLength() const
+{
+	return SequencePlayerNode.GetCurrentAssetLength();
+}
+
+float FAnimNode_MotionMatching::GetCurrentAssetTime() const
+{
+	return SequencePlayerNode.GetCurrentAssetTime();
+}
+
+float FAnimNode_MotionMatching::GetCurrentAssetTimePlayRateAdjusted() const
+{
+	return SequencePlayerNode.GetCurrentAssetTimePlayRateAdjusted();
+}
+
+bool FAnimNode_MotionMatching::GetIgnoreForRelevancyTest() const
+{
+	return GET_ANIM_NODE_DATA(bool, bIgnoreForRelevancyTest);
+}
+
+void FAnimNode_MotionMatching::SetIgnoreForRelevancyTest(bool bInIgnoreForRelevancyTest)
+{
+#if WITH_EDITORONLY_DATA
+	bIgnoreForRelevancyTest = bInIgnoreForRelevancyTest;
+#endif
+}
+
 #undef LOCTEXT_NAMESPACE
\ No newline at end of file
diff --git a/Engine/Plugins/Experimental/Animation/PoseSearch/Source/Runtime/Public/PoseSearch/AnimNode_MotionMatching.h b/Engine/Plugins/Experimental/Animation/PoseSearch/Source/Runtime/Public/PoseSearch/AnimNode_MotionMatching.h
--- a/Engine/Plugins/Experimental/Animation/PoseSearch/Source/Runtime/Public/PoseSearch/AnimNode_MotionMatching.h
+++ b/Engine/Plugins/Experimental/Animation/PoseSearch/Source/Runtime/Public/PoseSearch/AnimNode_MotionMatching.h
@@ -3,12 +3,12 @@
 #pragma once
 
 #include "PoseSearch/PoseSearch.h"
-#include "Animation/AnimNodeBase.h"
+#include "Animation/AnimNode_AssetPlayerBase.h"
 #include "Animation/AnimNode_SequencePlayer.h"
 #include "AnimNode_MotionMatching.generated.h"
 
 USTRUCT(BlueprintInternalUseOnly)
-struct POSESEARCH_API FAnimNode_MotionMatching : public FAnimNode_Base
+struct POSESEARCH_API FAnimNode_MotionMatching : public FAnimNode_AssetPlayerBase
 {
 	GENERATED_BODY()
 
@@ -58,7 +58,6 @@ public:
 	// FAnimNode_Base interface
 	virtual void Initialize_AnyThread(const FAnimationInitializeContext& Context) override;
 	virtual void Evaluate_AnyThread(FPoseContext& Output) override;
-	virtual void Update_AnyThread(const FAnimationUpdateContext& Context) override;
 	virtual void GatherDebugData(FNodeDebugData& DebugData) override;
 	virtual bool HasPreUpdate() const override;
 	virtual void PreUpdate(const UAnimInstance* InAnimInstance) override;
@@ -89,4 +88,25 @@ private:
 	void ComposeQuery(const FAnimationBaseContext& Context);
 	void JumpToPose(const FAnimationUpdateContext& Context, UE::PoseSearch::FDbSearchResult Result);
 	void InitNewDatabaseSearch();
+	
+	// FAnimNode_AssetPlayerBase
+protected:
+	// FAnimNode_AssetPlayerBase interface
+	virtual float GetAccumulatedTime() const;
+	virtual UAnimationAsset* GetAnimAsset() const override;
+	virtual void UpdateAssetPlayer(const FAnimationUpdateContext& Context) override;
+	virtual float GetCurrentAssetLength() const override;
+	virtual float GetCurrentAssetTime() const override;
+	virtual float GetCurrentAssetTimePlayRateAdjusted() const override;
+	virtual bool GetIgnoreForRelevancyTest() const override;
+	virtual void SetIgnoreForRelevancyTest(bool bInIgnoreForRelevancyTest) override;
+	// End of FAnimNode_AssetPlayerBase interface
+
+private:
+
+#if WITH_EDITORONLY_DATA
+	// If true, "Relevant anim" nodes that look for the highest weighted animation in a state will ignore this node
+	UPROPERTY(EditAnywhere, Category=Relevancy, meta=(FoldProperty, PinHiddenByDefault))
+	bool bIgnoreForRelevancyTest = false;
+#endif // WITH_EDITORONLY_DATA
 };
\ No newline at end of file
