Fixed initialization warnings in FAnimNode_MotionMatching

diff --git a/Engine/Plugins/Experimental/Animation/PoseSearch/Source/Runtime/Public/PoseSearch/AnimNode_MotionMatching.h b/Engine/Plugins/Experimental/Animation/PoseSearch/Source/Runtime/Public/PoseSearch/AnimNode_MotionMatching.h
--- a/Engine/Plugins/Experimental/Animation/PoseSearch/Source/Runtime/Public/PoseSearch/AnimNode_MotionMatching.h
+++ b/Engine/Plugins/Experimental/Animation/PoseSearch/Source/Runtime/Public/PoseSearch/AnimNode_MotionMatching.h
@@ -18,7 +18,7 @@ public:
 
 	// Collection of animations for motion matching
 	UPROPERTY(EditAnywhere, BlueprintReadWrite, Category=Settings, meta=(PinHiddenByDefault))
-	TObjectPtr<const UPoseSearchDatabase> Database;
+	TObjectPtr<const UPoseSearchDatabase> Database = nullptr;
 
 	UPROPERTY(EditAnywhere, BlueprintReadWrite, Category = Settings, meta = (PinHiddenByDefault))
 	FPoseSearchBiasWeightParams BiasWeights;
@@ -79,7 +79,7 @@ private:
 
 	// When the current database is different from the previous one, the search parameters are reset
 	UPROPERTY(Transient)
-	TObjectPtr<const UPoseSearchDatabase> PreviousDatabase;
+	TObjectPtr<const UPoseSearchDatabase> PreviousDatabase = nullptr;
 
 	// Time since the last pose jump
 	float ElapsedPoseJumpTime = 0.0f;
