AnimInputs: Snap predictions to cardinal directions and walk/jog max speeds
- Added direction and max speed snapping. Direction snapping uses configurable sector angles.
- Refactored prediction to decouple from pose search. Now prediction produces a state array that is processed in a separate phase to construct pose vectors.
- Added support for time, distance, and step count stop conditions.
- Removed secondary schema and goal and replaced with generic array of schemas and goals
- Added pose search schema helpers to get prediction horizons

diff --git a/Engine/Plugins/Experimental/Animation/PoseSearch/Source/Runtime/Private/PoseSearch.cpp b/Engine/Plugins/Experimental/Animation/PoseSearch/Source/Runtime/Private/PoseSearch.cpp
--- a/Engine/Plugins/Experimental/Animation/PoseSearch/Source/Runtime/Private/PoseSearch.cpp
+++ b/Engine/Plugins/Experimental/Animation/PoseSearch/Source/Runtime/Private/PoseSearch.cpp
@@ -319,6 +319,26 @@ bool UPoseSearchSchema::IsValid() const
 	return bValid;
 }
 
+float UPoseSearchSchema::GetTrajectoryFutureTimeHorizon () const
+{
+	return TrajectorySampleTimes.Num() ? TrajectorySampleTimes.Last() : -1.0f;
+}
+
+float UPoseSearchSchema::GetTrajectoryPastTimeHorizon () const
+{
+	return TrajectorySampleTimes.Num() ? TrajectorySampleTimes[0] : 1.0f;
+}
+
+float UPoseSearchSchema::GetTrajectoryFutureDistanceHorizon () const
+{
+	return TrajectorySampleDistances.Num() ? TrajectorySampleDistances.Last() : -1.0f;
+}
+
+float UPoseSearchSchema::GetTrajectoryPastDistanceHorizon () const
+{
+	return TrajectorySampleDistances.Num() ? TrajectorySampleDistances[0] : 1.0f;
+}
+
 void UPoseSearchSchema::GenerateLayout()
 {
 	Layout.Reset();
@@ -680,6 +700,15 @@ void FPoseSearchFeatureVectorBuilder::Init(const UPoseSearchSchema* InSchema)
 	ResetFeatures();
 }
 
+void FPoseSearchFeatureVectorBuilder::Reset()
+{
+	Schema = nullptr;
+	Values.Reset(0);
+	ValuesNormalized.Reset(0);
+	NumFeaturesAdded = 0;
+	FeaturesAdded.Reset();
+}
+
 void FPoseSearchFeatureVectorBuilder::ResetFeatures()
 {
 	Values.Reset(0);
diff --git a/Engine/Plugins/Experimental/Animation/PoseSearch/Source/Runtime/Public/PoseSearch/PoseSearch.h b/Engine/Plugins/Experimental/Animation/PoseSearch/Source/Runtime/Public/PoseSearch/PoseSearch.h
--- a/Engine/Plugins/Experimental/Animation/PoseSearch/Source/Runtime/Public/PoseSearch/PoseSearch.h
+++ b/Engine/Plugins/Experimental/Animation/PoseSearch/Source/Runtime/Public/PoseSearch/PoseSearch.h
@@ -227,6 +227,22 @@ public:
 
 	int32 NumBones () const { return BoneIndices.Num(); }
 
+	// Returns farthest future sample time >= 0.0f.
+	// Returns a negative value when there are no future sample times.
+	float GetTrajectoryFutureTimeHorizon () const;
+
+	// Returns farthest past sample time <= 0.0f.
+	// Returns a positive value when there are no past sample times.
+	float GetTrajectoryPastTimeHorizon () const;
+
+	// Returns farthest future sample distance >= 0.0f.
+	// Returns a negative value when there are no future sample distances.
+	float GetTrajectoryFutureDistanceHorizon () const;
+
+	// Returns farthest path sample distance <= 0.0f.
+	// Returns a positive value when there are no past sample distances.
+	float GetTrajectoryPastDistanceHorizon () const;
+
 public: // UObject
 	virtual void PreSave(FObjectPreSaveContext ObjectSaveContext) override;
 	virtual void PostLoad() override;
@@ -462,6 +478,7 @@ struct POSESEARCH_API FPoseSearchFeatureVectorBuilder
 	GENERATED_BODY()
 public:
 	void Init(const UPoseSearchSchema* Schema);
+	void Reset();
 	void ResetFeatures();
 
 	const UPoseSearchSchema* GetSchema() const { return Schema; }
