Trajectory prediction based play rate scaling for Motion Matching.

diff --git a/Engine/Plugins/Experimental/Animation/PoseSearch/Source/Runtime/Private/AnimNode_MotionMatching.cpp b/Engine/Plugins/Experimental/Animation/PoseSearch/Source/Runtime/Private/AnimNode_MotionMatching.cpp
--- a/Engine/Plugins/Experimental/Animation/PoseSearch/Source/Runtime/Private/AnimNode_MotionMatching.cpp
+++ b/Engine/Plugins/Experimental/Animation/PoseSearch/Source/Runtime/Private/AnimNode_MotionMatching.cpp
@@ -2,6 +2,7 @@
 
 #include "PoseSearch/AnimNode_MotionMatching.h"
 #include "PoseSearch/PoseSearch.h"
+#include "PoseSearch/PoseSearchPredictionLibrary.h"
 #include "PoseSearch/AnimNode_PoseSearchHistoryCollector.h"
 #include "Animation/AnimInstanceProxy.h"
 #include "Animation/AnimNode_Inertialization.h"
@@ -41,6 +42,8 @@ void FAnimNode_MotionMatching::UpdateAssetPlayer(const FAnimationUpdateContext&
 
 	bool bJumpedToPose = false;
 
+	const float DeltaTime = Context.GetDeltaTime();
+
 	if (IsValidForSearch())
 	{
 		if (PreviousDatabase != Database)
@@ -128,7 +131,7 @@ void FAnimNode_MotionMatching::UpdateAssetPlayer(const FAnimationUpdateContext&
 		// Continue with the follow up sequence if we're finishing a one shot anim
 		if (!bJumpedToPose && SequencePlayerNode.GetSequence() && !SequencePlayerNode.GetLoopAnimation())
 		{
-			float AssetTimeAfterUpdate = AssetTime + Context.GetDeltaTime();
+			float AssetTimeAfterUpdate = AssetTime + DeltaTime;
 			if (AssetTimeAfterUpdate > AssetLength)
 			{
 				const FPoseSearchDatabaseSequence& DbSequence = Database->Sequences[DbSequenceIdx];
@@ -156,8 +159,12 @@ void FAnimNode_MotionMatching::UpdateAssetPlayer(const FAnimationUpdateContext&
 	}
 	else
 	{
-		ElapsedPoseJumpTime += Context.GetDeltaTime();
+		ElapsedPoseJumpTime += DeltaTime;
 	}
+	 
+	const FPredictionSequenceState SequenceState = { SequencePlayerNode.GetSequence() , SequencePlayerNode.GetAccumulatedTime(), SequencePlayerNode.GetLoopAnimation() };
+	const float PlayRate = UPoseSearchPredictionDistanceMatching::ComputePlayRate(Context, Prediction, PredictionSettings, SequenceState);
+	SequencePlayerNode.SetPlayRate(PlayRate);
 
 	Source.Update(Context);
 }
diff --git a/Engine/Plugins/Experimental/Animation/PoseSearch/Source/Runtime/Public/PoseSearch/AnimNode_MotionMatching.h b/Engine/Plugins/Experimental/Animation/PoseSearch/Source/Runtime/Public/PoseSearch/AnimNode_MotionMatching.h
--- a/Engine/Plugins/Experimental/Animation/PoseSearch/Source/Runtime/Public/PoseSearch/AnimNode_MotionMatching.h
+++ b/Engine/Plugins/Experimental/Animation/PoseSearch/Source/Runtime/Public/PoseSearch/AnimNode_MotionMatching.h
@@ -3,6 +3,7 @@
 #pragma once
 
 #include "PoseSearch/PoseSearch.h"
+#include "PoseSearch/PoseSearchPredictionTypes.h"
 #include "Animation/AnimNode_AssetPlayerBase.h"
 #include "Animation/AnimNode_SequencePlayer.h"
 #include "AnimNode_MotionMatching.generated.h"
@@ -20,6 +21,7 @@ public:
 	UPROPERTY(EditAnywhere, BlueprintReadWrite, Category=Settings, meta=(PinHiddenByDefault))
 	TObjectPtr<const UPoseSearchDatabase> Database = nullptr;
 
+	// Dynamic biasing weights during pose search queries
 	UPROPERTY(EditAnywhere, BlueprintReadWrite, Category = Settings, meta = (PinHiddenByDefault))
 	FPoseSearchBiasWeightParams BiasWeights;
 
@@ -27,6 +29,14 @@ public:
 	UPROPERTY(EditAnywhere, BlueprintReadWrite, Category=Settings, meta=(PinShownByDefault))
 	FPoseSearchFeatureVectorBuilder Goal;
 
+	// Trajectory prediction samples
+	UPROPERTY(EditAnywhere, BlueprintReadWrite, Category = Settings, meta = (PinHiddenByDefault))
+	FPredictionTrajectoryRange Prediction;
+
+	// Configuration for prediction play rate adjustment, error tolerances, and sampling resolution
+	UPROPERTY(EditAnywhere, BlueprintReadWrite, Category = Settings, meta = (PinHiddenByDefault))
+	FPredictionTrajectorySettings PredictionSettings;
+
 	// Time in seconds to blend out to the new pose. Uses inertial blending and requires an Inertialization node after this node.
 	UPROPERTY(EditAnywhere, BlueprintReadWrite, Category=Settings, meta=(ClampMin="0"))
 	float BlendTime = 0.2f;
diff --git a/Engine/Plugins/Experimental/Animation/PoseSearch/Source/Runtime/Public/PoseSearch/PoseSearchPredictionLibrary.cpp b/Engine/Plugins/Experimental/Animation/PoseSearch/Source/Runtime/Public/PoseSearch/PoseSearchPredictionLibrary.cpp
new file mode 100644
--- /dev/null
+++ b/Engine/Plugins/Experimental/Animation/PoseSearch/Source/Runtime/Public/PoseSearch/PoseSearchPredictionLibrary.cpp
@@ -0,0 +1,170 @@
+// Copyright Epic Games, Inc. All Rights Reserved.
+
+#include "PoseSearchPredictionLibrary.h"
+#include "Animation/AnimInstanceProxy.h"
+#include "Algo/MinElement.h"
+
+
+float UPoseSearchPredictionDistanceMatching::ComputePlayRate(const FAnimationUpdateContext& Context
+	, const FPredictionTrajectoryRange& TrajectoryRange
+	, const FPredictionTrajectorySettings& Settings
+	, const FPredictionSequenceState& SequenceState)
+{
+	const float DeltaTime = Context.GetDeltaTime();
+	float PlayRate = 1.f;
+
+	// Delta time is not progressing
+	if (FMath::IsNearlyZero(DeltaTime, SMALL_NUMBER))
+	{
+		return PlayRate;
+	}
+
+	// Prediction range isn't being updated
+	if (!TrajectoryRange.HasSamples())
+	{
+		return PlayRate;
+	}
+
+	// Prediction range contains only zeroed samples
+	if (TrajectoryRange.HasOnlyZeroSamples())
+	{
+		return PlayRate;
+	}
+
+	// No sequence available to play rate scale
+	if (!SequenceState.HasSequence())
+	{
+		return PlayRate;
+	}
+
+	// Find the minima prediction trajectory velocity
+	// Approximate zeroed values may indicate the synchronization point for a stop or pivot
+	const FPredictionTrajectoryState* MinimaSample = Algo::MinElementBy(TrajectoryRange.Samples, [](const FPredictionTrajectoryState& Sample)
+	{
+		return Sample.LocalLinearVelocity.SizeSquared();
+	});
+
+	check(MinimaSample);
+	const UAnimSequence* Sequence = Cast<const UAnimSequence>(SequenceState.SequenceBase);
+	check(Sequence);
+
+	// Given a high-resolution time step, walk the current animation sequence to find a corresponding minima root motion delta
+	// Changes in direction are considered extreme minima events, meaning no other subsequent minima is significant enough to take precedence
+	float MinimaSampleTime = FLT_MAX;
+	bool bPivotDetected = false;
+
+	for (auto [SampleTime, PreviousDirection, MinimaDisplacement, CosPivotAngleThreshold] = std::tuple{ (double)SequenceState.AccumulatedTime, FVector::ZeroVector, FLT_MAX, FMath::Cos(Settings.RootMotionPivotAngleThreshold) };
+		SampleTime <= Sequence->GetPlayLength(); 
+		SampleTime += Settings.RootMotionSampleStep)
+	{
+		const FVector RootMotion = Sequence->ExtractRootMotion(SampleTime, Settings.RootMotionSampleStep, SequenceState.bLooping).GetTranslation();
+		
+		FVector RootMotionDirection;
+		float RootMotionDisplacement;
+		RootMotion.ToDirectionAndLength(RootMotionDirection, RootMotionDisplacement);
+
+		// Found a smaller displacement in the root motion track
+		if (RootMotionDisplacement <= MinimaDisplacement)
+		{
+			MinimaDisplacement = RootMotionDisplacement;
+			MinimaSampleTime = SampleTime;
+		}
+
+		const float PotentialPivotAngle = RootMotionDirection.Dot(PreviousDirection);
+		PreviousDirection = RootMotionDirection;
+
+		// Significant changes in direction will be defined as a pivot
+		// Hack: Unfortunately in practice, we may erroneously identify animations as a pivot when in fact they have malformed root motion tracks
+		if (PotentialPivotAngle < CosPivotAngleThreshold && RootMotionDisplacement > Settings.RootMotionPivotDisplacementError)
+		{
+			// Bias the minima sample time for pivots to favor the pre-pivot phase (moment prior to direction change)
+			MinimaSampleTime = SampleTime - Settings.RootMotionSampleStep;
+			bPivotDetected = true;
+			break;
+		}
+	}
+
+	// We should always be able to find a minima, however a few situations could lead to this firing
+	// 1) If the root motion sampling time step is not high enough resolution, we may not be able to sample the
+	//    track near the end of a non-looping sequence
+	// 2) We may be sampling at time 'sequence length' of a non-looping sequence
+	if (MinimaSampleTime == FLT_MAX)
+	{
+		return PlayRate;
+	}
+
+	// Extrapolate the minima forward in time to detect a potential complete loss in velocity
+	const FVector MinimaRootMotionDelta = Sequence->ExtractRootMotion(MinimaSampleTime, DeltaTime, SequenceState.bLooping).GetTranslation();
+
+	// Play rate scaling is minima driven when a near zero root motion delta or pivot has been detected
+	// Otherwise we are locomotion driven, which is reflected in the numerator of the divisor
+	bool bMinimaDrivenPlayRate = MinimaRootMotionDelta.IsNearlyZero(KINDA_SMALL_NUMBER) || bPivotDetected;
+#if WITH_EDITORONLY_DATA
+	FColor SynchronizationColor = FColor::Red;
+#endif
+	FVector RemainingRootMotionDelta = FVector::ZeroVector;
+
+	// This loop allows for play rate scaling to apply correction in the event that we discover the animation and trajectory prediction minima mismatch
+	// Example: 
+	//	If both the trajectory prediction and chosen animation are decelerating to zero, minima driven play rate scaling can be correctly applied
+	//  However if the chosen animation is not decelerating to zero -- such as Jog_Right vs Jog_Right_Stop -- then a mismatch has been detected
+	//	and locomotion/instantaneous driven play rate scaling is attempted instead
+	for (int MaxIterations = 1; MaxIterations >= 0; MaxIterations--)
+	{
+		// Minima driven play rate scaling synchronizes using the remaining displacement to near zero with: animation / locomotion
+		// Locomotion driven play rate scaling synchronizes using the per-frame instantaneous displacement with: locomotion / animation
+		const float RemainingSequenceDelta = bMinimaDrivenPlayRate ? MinimaSampleTime - SequenceState.AccumulatedTime : DeltaTime;
+		const float PredictionDisplacement = bMinimaDrivenPlayRate ? MinimaSample->AccumulatedDistance : TrajectoryRange.Samples[0].LocalLinearVelocity.Size2D() * DeltaTime;
+
+		RemainingRootMotionDelta = Sequence->ExtractRootMotion(SequenceState.AccumulatedTime, RemainingSequenceDelta, SequenceState.bLooping).GetTranslation();
+		const float RemainingRootMotionDeltaDisplacement = RemainingRootMotionDelta.Size2D();
+
+		// Zero displacement is left in the animation, which may result in sliding if the prediction minima has non-zero displacement
+		const bool bZeroRootMotion = FMath::IsNearlyZero(RemainingRootMotionDeltaDisplacement, KINDA_SMALL_NUMBER);
+
+		// Zero displacement is left in the prediction minima, which may result in a pop or a pose break if the animation has non-zero displacement
+		const bool bZeroPrediction = FMath::IsNearlyZero(PredictionDisplacement, KINDA_SMALL_NUMBER);
+
+		// Play rate scaling isn't require since no trajectory prediction motion or root motion is present
+		if (bZeroRootMotion && bZeroPrediction)
+		{
+			break;
+		}
+		// The computed minima in the root motion and prediction mismatch, so flip the synchronization behavior as an attempt to correctly play rate scale
+		// If this fails, the algorithm will automatically fall back to a default play rate of 1.f, which may introduce sliding.
+		else if (bZeroPrediction || bZeroRootMotion)
+		{
+			bMinimaDrivenPlayRate = !bMinimaDrivenPlayRate;
+			continue;
+		}
+		// Play rate scaling succeeded
+		else
+		{
+#if WITH_EDITORONLY_DATA
+			SynchronizationColor = bMinimaDrivenPlayRate ? FColor::Purple : FColor::Blue;
+#endif
+			PlayRate = bMinimaDrivenPlayRate
+				? RemainingRootMotionDeltaDisplacement / PredictionDisplacement 
+				: PredictionDisplacement / RemainingRootMotionDeltaDisplacement;
+
+			break;
+		}
+	}
+
+#if WITH_EDITORONLY_DATA
+	// Render the starting and ending trajectory prediction positions for the distance matching play rate synchronization
+	if (Settings.bDebugDraw)
+	{
+		Context.AnimInstanceProxy->AnimDrawDebugSphere(Context.AnimInstanceProxy->GetComponentTransform().GetLocation(), 10.f, 10, FColor::Green);
+		Context.AnimInstanceProxy->AnimDrawDebugSphere(Context.AnimInstanceProxy->GetComponentTransform().TransformPosition(RemainingRootMotionDelta), 10.f, 10, SynchronizationColor);
+
+		if (bMinimaDrivenPlayRate)
+		{
+			Context.AnimInstanceProxy->AnimDrawDebugSphere(Context.AnimInstanceProxy->GetComponentTransform().TransformPosition(MinimaSample->Position), 10.f, 10, FColor::Yellow);
+		}
+	}
+#endif
+
+	// Optionally remap the computed play rate against a curve
+	return Settings.PlayRateAdjustment.ComputePlayRate(PlayRate, DeltaTime);
+}
\ No newline at end of file
diff --git a/Engine/Plugins/Experimental/Animation/PoseSearch/Source/Runtime/Public/PoseSearch/PoseSearchPredictionLibrary.h b/Engine/Plugins/Experimental/Animation/PoseSearch/Source/Runtime/Public/PoseSearch/PoseSearchPredictionLibrary.h
new file mode 100644
--- /dev/null
+++ b/Engine/Plugins/Experimental/Animation/PoseSearch/Source/Runtime/Public/PoseSearch/PoseSearchPredictionLibrary.h
@@ -0,0 +1,31 @@
+// Copyright Epic Games, Inc. All Rights Reserved.
+
+#pragma once
+
+// Note: Various includes, class and function annotations have been commented out as reference
+// for using this functionality in the context of the run-time "Encapsulation" initiative
+
+#include "CoreMinimal.h"
+#include "PoseSearchPredictionTypes.h"
+// #include "UObject/ObjectMacros.h"
+// #include "PoseSearchPredictionLibrary.generated.h"
+
+// UCLASS()
+class UPoseSearchPredictionDistanceMatching // : public UBlueprintFunctionLibrary
+{
+	// GENERATED_BODY()
+
+public:
+	// Computes the effective play rate of a sequence by modeling and analyzing the ratio of the capsule trajectory prediction 
+	// vs sequence root motion, or vice versa.
+	//
+	// Currently this has been tested against a variety of locomotion features such as:
+	// 1) Starts, Stops, Pivots, and Cycles
+	// 2) Differing, similar, or divergent prediction and root motion velocity/acceleration models
+
+	// UFUNCTION(BlueprintCallable, Category = "Distance Matching")
+	static float ComputePlayRate(const FAnimationUpdateContext& Context
+		, const FPredictionTrajectoryRange& TrajectoryRange
+		, const FPredictionTrajectorySettings& Settings
+		, const FPredictionSequenceState& SequenceState);
+};
diff --git a/Engine/Plugins/Experimental/Animation/PoseSearch/Source/Runtime/Public/PoseSearch/PoseSearchPredictionTypes.cpp b/Engine/Plugins/Experimental/Animation/PoseSearch/Source/Runtime/Public/PoseSearch/PoseSearchPredictionTypes.cpp
new file mode 100644
--- /dev/null
+++ b/Engine/Plugins/Experimental/Animation/PoseSearch/Source/Runtime/Public/PoseSearch/PoseSearchPredictionTypes.cpp
@@ -0,0 +1,42 @@
+// Copyright Epic Games, Inc. All Rights Reserved.
+
+#include "PoseSearchPredictionTypes.h"
+#include "Algo/AllOf.h"
+
+
+float FPredictionPlayRateAdjustment::ComputePlayRate(float PlayRate, float DeltaTime) const
+{
+	PlayRate = RemappingCurve ? RemappingCurve->GetFloatValue(PlayRate) : PlayRate;
+	PlayRate = ScaleBiasClamp.ApplyTo(PlayRate, DeltaTime);
+	return PlayRate;
+}
+
+FPredictionTrajectoryState FPredictionTrajectoryState::Lerp(const FPredictionTrajectoryState& A, const FPredictionTrajectoryState& B, float Alpha)
+{
+	FPredictionTrajectoryState State;
+	State.AccumulatedDistance = FMath::Lerp(A.AccumulatedDistance, B.AccumulatedDistance, Alpha);
+	State.LocalLinearAcceleration = FMath::Lerp(A.LocalLinearAcceleration, B.LocalLinearAcceleration, Alpha);
+	State.LocalLinearVelocity = FMath::Lerp(A.LocalLinearVelocity, B.LocalLinearVelocity, Alpha);
+	State.Position = FMath::Lerp(A.Position, B.Position, Alpha);
+	return State;
+}
+
+bool FPredictionTrajectoryRange::HasSamples() const
+{
+	return !Samples.IsEmpty();
+}
+
+bool FPredictionTrajectoryRange::HasOnlyZeroSamples() const
+{
+	return Algo::AllOf(Samples, [](const FPredictionTrajectoryState& Sample) {
+		return Sample.LocalLinearVelocity == FVector::ZeroVector
+			&& Sample.LocalLinearAcceleration == FVector::ZeroVector
+			&& Sample.Position == FVector::ZeroVector
+			&& Sample.AccumulatedDistance == 0.f;
+		});
+}
+
+bool FPredictionSequenceState::HasSequence() const
+{
+	return !SequenceBase.IsNull() && SequenceBase->IsA<UAnimSequence>();
+}
\ No newline at end of file
diff --git a/Engine/Plugins/Experimental/Animation/PoseSearch/Source/Runtime/Public/PoseSearch/PoseSearchPredictionTypes.h b/Engine/Plugins/Experimental/Animation/PoseSearch/Source/Runtime/Public/PoseSearch/PoseSearchPredictionTypes.h
new file mode 100644
--- /dev/null
+++ b/Engine/Plugins/Experimental/Animation/PoseSearch/Source/Runtime/Public/PoseSearch/PoseSearchPredictionTypes.h
@@ -0,0 +1,109 @@
+// Copyright Epic Games, Inc. All Rights Reserved.
+
+#pragma once
+
+#include "CoreMinimal.h"
+#include "Curves/CurveFloat.h"
+#include "UObject/ObjectMacros.h"
+#include "Animation/AnimSequenceBase.h"
+#include "Animation/InputScaleBias.h"
+#include "PoseSearchPredictionTypes.generated.h"
+
+USTRUCT(BlueprintType, Category = "Animation|Prediction")
+struct POSESEARCH_API FPredictionPlayRateAdjustment
+{
+	GENERATED_BODY()
+
+	// Optional scale, bias, and clamp play rate value adjustment
+	UPROPERTY(EditAnywhere, BlueprintReadWrite, Category = Settings, meta = (PinShownByDefault))
+	FInputScaleBiasClamp ScaleBiasClamp;
+
+	// Optional play rate value remapping curve
+	UPROPERTY(EditAnywhere, BlueprintReadWrite, Category = Settings, meta = (PinHiddenByDefault))
+	TObjectPtr<const UCurveFloat> RemappingCurve = nullptr;
+
+	// Computes a new play rate by first applying the remapping curve and then applying the scale/bias/clamp adjustment
+	float ComputePlayRate(float PlayRate, float DeltaTime) const;
+};
+
+USTRUCT(BlueprintType, Category = "Animation|Prediction")
+struct POSESEARCH_API FPredictionTrajectoryState
+{
+	GENERATED_BODY()
+
+	UPROPERTY(BlueprintReadOnly, Category = Prediction)
+	FVector Position = FVector::ZeroVector;
+
+	UPROPERTY(BlueprintReadOnly, Category = Prediction)
+	FVector LocalLinearVelocity = FVector::ZeroVector;
+
+	UPROPERTY(BlueprintReadOnly, Category = Prediction)
+	FVector LocalLinearAcceleration = FVector::ZeroVector;
+
+	UPROPERTY(BlueprintReadOnly, Category = Prediction)
+	float AccumulatedDistance = 0.0;
+
+	static FPredictionTrajectoryState Lerp(const FPredictionTrajectoryState& A, const FPredictionTrajectoryState& B, float Alpha);
+};
+
+USTRUCT(BlueprintType, Category = "Animation|Prediction")
+struct POSESEARCH_API FPredictionTrajectoryRange
+{
+	GENERATED_BODY()
+
+	// Contains the per-frame range of predicted trajectory states
+	UPROPERTY(BlueprintReadOnly, Category = Prediction)
+	TArray<FPredictionTrajectoryState> Samples;
+
+	bool HasSamples() const;
+
+	// Useful in detecting idle motion states
+	bool HasOnlyZeroSamples() const;
+};
+
+USTRUCT(BlueprintType, Category = "Animation|Prediction")
+struct POSESEARCH_API FPredictionTrajectorySettings
+{
+	GENERATED_BODY()
+	
+	// Optional play rate scaling adjustment applied as a post process
+	UPROPERTY(EditAnywhere, BlueprintReadWrite, Category = Settings, meta = (PinShownByDefault))
+	FPredictionPlayRateAdjustment PlayRateAdjustment;
+
+	// Root motion time step used for analyzing and determining future velocity minima
+	UPROPERTY(EditAnywhere, BlueprintReadWrite, Category = Settings, meta = (PinHiddenByDefault))
+	float RootMotionSampleStep = 1.f / 240.f;
+
+	// Root motion angle threshold for determining significant changes in direction (ie: pivots)
+	UPROPERTY(EditAnywhere, BlueprintReadWrite, Category = Settings, meta = (PinHiddenByDefault))
+	float RootMotionPivotAngleThreshold = 90.f;
+
+	// Root motion displacement error tolerance for identifying pivot false positives
+	UPROPERTY(EditAnywhere, BlueprintReadWrite, Category = Settings, meta = (PinHiddenByDefault))
+	float RootMotionPivotDisplacementError = 0.0004f;
+	
+#if WITH_EDITORONLY_DATA
+	UPROPERTY(EditAnywhere, Category = Debug)
+	bool bDebugDraw = false;
+#endif
+};
+
+USTRUCT(BlueprintInternalUseOnly, Category = "Animation|Prediction")
+struct POSESEARCH_API FPredictionSequenceState
+{
+	GENERATED_BODY()
+
+	// Current evaluating sequence (currently for use in Motion Matching)
+	UPROPERTY(EditAnywhere, BlueprintReadWrite, Category = State, meta = (PinShownByDefault))
+	TObjectPtr<const UAnimSequenceBase> SequenceBase = nullptr;
+
+	// Internal sequence accumulated time
+	UPROPERTY(EditAnywhere, BlueprintReadWrite, Category = State, meta = (PinShownByDefault))
+	float AccumulatedTime = 0.f;
+
+	// Looping or non-looping sequence
+	UPROPERTY(EditAnywhere, BlueprintReadWrite, Category = State, meta = (PinShownByDefault))
+	bool bLooping = false;
+
+	bool HasSequence() const;
+};
\ No newline at end of file
