Improvements/updates to the Pose Warping demo project. Integrated trajectory predicition play rate scaling into the motion matching movement mode.

#trivial

diff --git a/Engine/Plugins/Experimental/Animation/PoseSearch/Source/Runtime/Public/PoseSearch/PoseSearchPredictionLibrary.cpp b/Engine/Plugins/Experimental/Animation/PoseSearch/Source/Runtime/Public/PoseSearch/PoseSearchPredictionLibrary.cpp
--- a/Engine/Plugins/Experimental/Animation/PoseSearch/Source/Runtime/Public/PoseSearch/PoseSearchPredictionLibrary.cpp
+++ b/Engine/Plugins/Experimental/Animation/PoseSearch/Source/Runtime/Public/PoseSearch/PoseSearchPredictionLibrary.cpp
@@ -53,14 +53,14 @@ float UPoseSearchPredictionDistanceMatching::ComputePlayRate(const FAnimationUpd
 	float MinimaSampleTime = FLT_MAX;
 	bool bPivotDetected = false;
 
-	for (auto [SampleTime, PreviousDirection, MinimaDisplacement, CosPivotAngleThreshold] = std::tuple{ (double)SequenceState.AccumulatedTime, FVector::ZeroVector, FLT_MAX, FMath::Cos(Settings.RootMotionPivotAngleThreshold) };
+	for (auto [SampleTime, RootMotionSampleStep, PreviousDirection, MinimaDisplacement, CosPivotAngleThreshold] = std::tuple{ (double)SequenceState.AccumulatedTime, 1.f / Settings.RootMotionSampleStepPerSecond, FVector::ZeroVector, FLT_MAX, FMath::Cos(Settings.RootMotionPivotAngleThreshold) };
 		SampleTime <= Sequence->GetPlayLength(); 
-		SampleTime += Settings.RootMotionSampleStep)
+		SampleTime += RootMotionSampleStep)
 	{
-		const FVector RootMotion = Sequence->ExtractRootMotion(SampleTime, Settings.RootMotionSampleStep, SequenceState.bLooping).GetTranslation();
-		
 		FVector RootMotionDirection;
 		float RootMotionDisplacement;
+
+		const FVector RootMotion = Sequence->ExtractRootMotion(SampleTime, RootMotionSampleStep, SequenceState.bLooping).GetTranslation();
 		RootMotion.ToDirectionAndLength(RootMotionDirection, RootMotionDisplacement);
 
 		// Found a smaller displacement in the root motion track
@@ -78,7 +78,7 @@ float UPoseSearchPredictionDistanceMatching::ComputePlayRate(const FAnimationUpd
 		if (PotentialPivotAngle < CosPivotAngleThreshold && RootMotionDisplacement > Settings.RootMotionPivotDisplacementError)
 		{
 			// Bias the minima sample time for pivots to favor the pre-pivot phase (moment prior to direction change)
-			MinimaSampleTime = SampleTime - Settings.RootMotionSampleStep;
+			MinimaSampleTime = SampleTime - RootMotionSampleStep;
 			bPivotDetected = true;
 			break;
 		}
diff --git a/Engine/Plugins/Experimental/Animation/PoseSearch/Source/Runtime/Public/PoseSearch/PoseSearchPredictionTypes.h b/Engine/Plugins/Experimental/Animation/PoseSearch/Source/Runtime/Public/PoseSearch/PoseSearchPredictionTypes.h
--- a/Engine/Plugins/Experimental/Animation/PoseSearch/Source/Runtime/Public/PoseSearch/PoseSearchPredictionTypes.h
+++ b/Engine/Plugins/Experimental/Animation/PoseSearch/Source/Runtime/Public/PoseSearch/PoseSearchPredictionTypes.h
@@ -71,11 +71,11 @@ struct POSESEARCH_API FPredictionTrajectorySettings
 	FPredictionPlayRateAdjustment PlayRateAdjustment;
 
 	// Root motion time step used for analyzing and determining future velocity minima
-	UPROPERTY(EditAnywhere, BlueprintReadWrite, Category = Settings, meta = (PinHiddenByDefault))
-	float RootMotionSampleStep = 1.f / 240.f;
+	UPROPERTY(EditAnywhere, BlueprintReadWrite, Category = Settings, meta = (PinHiddenByDefault, ClampMin = "15", ClampMax = "240"))
+	float RootMotionSampleStepPerSecond = 120.f;
 
-	// Root motion angle threshold for determining significant changes in direction (ie: pivots)
-	UPROPERTY(EditAnywhere, BlueprintReadWrite, Category = Settings, meta = (PinHiddenByDefault))
+	// Root motion angle threshold for determining significant changes in direction (pivots)
+	UPROPERTY(EditAnywhere, BlueprintReadWrite, Category = Settings, meta = (PinHiddenByDefault, ClampMin = "0", ClampMax = "180"))
 	float RootMotionPivotAngleThreshold = 90.f;
 
 	// Root motion displacement error tolerance for identifying pivot false positives
