More improvements/updates to the Pose Warping demo project. Implemented debugging cvars and improved the motion matching movement mode a bit further with play rate clamping.

#trivial

diff --git a/Engine/Plugins/Experimental/Animation/PoseSearch/Source/Runtime/Public/PoseSearch/PoseSearchPredictionLibrary.cpp b/Engine/Plugins/Experimental/Animation/PoseSearch/Source/Runtime/Public/PoseSearch/PoseSearchPredictionLibrary.cpp
--- a/Engine/Plugins/Experimental/Animation/PoseSearch/Source/Runtime/Public/PoseSearch/PoseSearchPredictionLibrary.cpp
+++ b/Engine/Plugins/Experimental/Animation/PoseSearch/Source/Runtime/Public/PoseSearch/PoseSearchPredictionLibrary.cpp
@@ -5,6 +5,11 @@
 #include "Algo/MinElement.h"
 
 
+#if ENABLE_ANIM_DEBUG
+TAutoConsoleVariable<int32> CVarAnimNodePredictionPlayRateDebug(TEXT("a.AnimNode.PredictionPlayRate.Debug"), 0, TEXT("Turn on debug for trajectory prediction play rate scaling"));
+TAutoConsoleVariable<int32> CVarAnimNodePredictionPlayRateEnable(TEXT("a.AnimNode.PredictionPlayRate.Enable"), 1, TEXT("Toggle trajectory prediction play rate scaling"));
+#endif
+
 float UPoseSearchPredictionDistanceMatching::ComputePlayRate(const FAnimationUpdateContext& Context
 	, const FPredictionTrajectoryRange& TrajectoryRange
 	, const FPredictionTrajectorySettings& Settings
@@ -13,6 +18,21 @@ float UPoseSearchPredictionDistanceMatching::ComputePlayRate(const FAnimationUpd
 	const float DeltaTime = Context.GetDeltaTime();
 	float PlayRate = 1.f;
 
+#if ENABLE_ANIM_DEBUG
+	// Debug enable/disable toggle for play rate scaling
+	if (!CVarAnimNodePredictionPlayRateEnable.GetValueOnAnyThread())
+	{
+		return PlayRate;
+	}
+#if WITH_EDITORONLY_DATA
+	bool bDebugDraw = Settings.bDebugDraw || CVarAnimNodePredictionPlayRateDebug.GetValueOnAnyThread();
+#endif
+#else
+#if WITH_EDITORONLY_DATA
+	bool bDebugDraw = Settings.bDebugDraw;
+#endif
+#endif
+
 	// Delta time is not progressing
 	if (FMath::IsNearlyZero(DeltaTime, SMALL_NUMBER))
 	{
@@ -153,14 +173,14 @@ float UPoseSearchPredictionDistanceMatching::ComputePlayRate(const FAnimationUpd
 
 #if WITH_EDITORONLY_DATA
 	// Render the starting and ending trajectory prediction positions for the distance matching play rate synchronization
-	if (Settings.bDebugDraw)
+	if (bDebugDraw)
 	{
-		Context.AnimInstanceProxy->AnimDrawDebugSphere(Context.AnimInstanceProxy->GetComponentTransform().GetLocation(), 10.f, 10, FColor::Green);
-		Context.AnimInstanceProxy->AnimDrawDebugSphere(Context.AnimInstanceProxy->GetComponentTransform().TransformPosition(RemainingRootMotionDelta), 10.f, 10, SynchronizationColor);
+		Context.AnimInstanceProxy->AnimDrawDebugSphere(Context.AnimInstanceProxy->GetComponentTransform().GetLocation(), 8.f, 16, FColor::Green);
+		Context.AnimInstanceProxy->AnimDrawDebugSphere(Context.AnimInstanceProxy->GetComponentTransform().TransformPosition(RemainingRootMotionDelta), 8.f, 16, SynchronizationColor);
 
 		if (bMinimaDrivenPlayRate)
 		{
-			Context.AnimInstanceProxy->AnimDrawDebugSphere(Context.AnimInstanceProxy->GetComponentTransform().TransformPosition(MinimaSample->Position), 10.f, 10, FColor::Yellow);
+			Context.AnimInstanceProxy->AnimDrawDebugSphere(Context.AnimInstanceProxy->GetComponentTransform().TransformPosition(MinimaSample->Position), 8.f, 16, FColor::Yellow);
 		}
 	}
 #endif
