AnimInputs:
- Added minimum prediction time/distance so schema horizon doesn't need to change for playrate scaling
- Project acceleration onto velocity so they are in alignment after direction snapping
- Disable walk speed snapping
PoseSearch:
- Use past trajectory from db for continuous matching, only use recorded trajectory on intialize

#rnx

diff --git a/Engine/Plugins/Experimental/Animation/PoseSearch/Source/Runtime/Private/AnimNode_MotionMatching.cpp b/Engine/Plugins/Experimental/Animation/PoseSearch/Source/Runtime/Private/AnimNode_MotionMatching.cpp
--- a/Engine/Plugins/Experimental/Animation/PoseSearch/Source/Runtime/Private/AnimNode_MotionMatching.cpp
+++ b/Engine/Plugins/Experimental/Animation/PoseSearch/Source/Runtime/Private/AnimNode_MotionMatching.cpp
@@ -81,6 +81,7 @@ void FAnimNode_MotionMatching::UpdateAssetPlayer(const FAnimationUpdateContext&
 			{
 				UE::PoseSearch::FPoseHistory& History = PoseHistoryProvider->GetPoseHistory();
 				ComposedQuery.TrySetPoseFeatures(&History);
+				ComposedQuery.TrySetPastTrajectoryFeatures(&History);
 			}
 		}
 
@@ -222,14 +223,6 @@ bool FAnimNode_MotionMatching::IsValidForSearch() const
 
 void FAnimNode_MotionMatching::ComposeQuery(const FAnimationBaseContext& Context)
 {
-	// Set past trajectory features
-	UE::PoseSearch::IPoseHistoryProvider* PoseHistoryProvider = Context.GetMessage<UE::PoseSearch::IPoseHistoryProvider>();
-	if (PoseHistoryProvider)
-	{
-		UE::PoseSearch::FPoseHistory& History = PoseHistoryProvider->GetPoseHistory();
-		ComposedQuery.TrySetPastTrajectoryFeatures(&History);
-	}
-
 	// Merge goal features into the query vector
 	if (ComposedQuery.IsCompatible(Goal))
 	{
