Infrastructure for Trace functionality in the Pose Search plugin
- Includes the module for startup/shutdown, the trace logger files for the user-facing macros, analyzer to receive events, and provider to append to and read from timelines
- The trace-module currently resides inside the PoseSearch.cpp, and starts off the startup chain in OnStartupModule
- Does not include any editor code, this should be entirely runtime stuff
- utrace files confirmed not to crash in standalone insights tool despite no dependencies on these features, so plugin compatibility should not be an issue

#rnx

diff --git a/Engine/Plugins/Experimental/Animation/PoseSearch/Source/Runtime/PoseSearch.Build.cs b/Engine/Plugins/Experimental/Animation/PoseSearch/Source/Runtime/PoseSearch.Build.cs
--- a/Engine/Plugins/Experimental/Animation/PoseSearch/Source/Runtime/PoseSearch.Build.cs
+++ b/Engine/Plugins/Experimental/Animation/PoseSearch/Source/Runtime/PoseSearch.Build.cs
@@ -32,7 +32,12 @@ public class PoseSearch : ModuleRules
 				"Core",
 				"CoreUObject",
 				"Engine",
-				"AnimationCore"
+				"AnimationCore",
+				
+				// Trace dependencies
+				"TraceLog",
+				"TraceAnalysis",
+				"TraceServices"
 			}
 		);
 
diff --git a/Engine/Plugins/Experimental/Animation/PoseSearch/Source/Runtime/Private/AnimNode_MotionMatching.cpp b/Engine/Plugins/Experimental/Animation/PoseSearch/Source/Runtime/Private/AnimNode_MotionMatching.cpp
--- a/Engine/Plugins/Experimental/Animation/PoseSearch/Source/Runtime/Private/AnimNode_MotionMatching.cpp
+++ b/Engine/Plugins/Experimental/Animation/PoseSearch/Source/Runtime/Private/AnimNode_MotionMatching.cpp
@@ -7,10 +7,10 @@
 #include "Animation/AnimInstanceProxy.h"
 #include "Animation/AnimNode_Inertialization.h"
 #include "Animation/AnimSequence.h"
+#include "Trace/PoseSearchTraceLogger.h"
 
 #define LOCTEXT_NAMESPACE "AnimNode_MotionMatching"
 
-
 /////////////////////////////////////////////////////
 // FAnimNode_MotionMatching
 
@@ -166,8 +166,18 @@ void FAnimNode_MotionMatching::UpdateAssetPlayer(const FAnimationUpdateContext&
 	const FPredictionSequenceState SequenceState = { SequencePlayerNode.GetSequence() , SequencePlayerNode.GetAccumulatedTime(), SequencePlayerNode.GetLoopAnimation() };
 	const float PlayRate = UPoseSearchPredictionDistanceMatching::ComputePlayRate(Context, Prediction, PredictionSettings, SequenceState);
 	SequencePlayerNode.SetPlayRate(PlayRate);
-
 	Source.Update(Context);
+	
+#if UE_POSE_SEARCH_TRACE_ENABLED
+	// Gather our node state and trace it
+	UE::PoseSearch::FTraceMotionMatchingState State;
+	State.ElapsedPoseJumpTime = ElapsedPoseJumpTime;
+	State.Flags = UE::PoseSearch::FTraceMotionMatchingState::EFlags::None;
+	State.DbPoseIdx = DbPoseIdx;
+	State.DatabaseId = FObjectTrace::GetObjectId(Database.Get());
+	State.QueryVector = ComposedQuery.GetValues();
+	UE_TRACE_POSE_SEARCH_MOTION_MATCHING_STATE(Context, State)
+#endif
 }
 
 bool FAnimNode_MotionMatching::HasPreUpdate() const
diff --git a/Engine/Plugins/Experimental/Animation/PoseSearch/Source/Runtime/Private/PoseSearch.cpp b/Engine/Plugins/Experimental/Animation/PoseSearch/Source/Runtime/Private/PoseSearch.cpp
--- a/Engine/Plugins/Experimental/Animation/PoseSearch/Source/Runtime/Private/PoseSearch.cpp
+++ b/Engine/Plugins/Experimental/Animation/PoseSearch/Source/Runtime/Private/PoseSearch.cpp
@@ -21,6 +21,9 @@
 #include "Animation/AnimationPoseData.h"
 #include "AnimationRuntime.h"
 #include "BonePose.h"
+#include "Trace/PoseSearchTraceLogger.h"
+#include "Trace/PoseSearchTraceModule.h"
+#include "TraceServices/ModuleService.h"
 #include "UObject/ObjectSaveContext.h"
 
 IMPLEMENT_ANIMGRAPH_MESSAGE(UE::PoseSearch::IPoseHistoryProvider);
@@ -3044,6 +3047,9 @@ public: // IModuleInterface
 	virtual void StartupModule() override;
 	virtual void ShutdownModule() override;
 
+	/** Host our trace module here, which registers the PoseSearch trace features */
+	FPoseSearchTraceModule PoseSearchTraceModule;
+
 public: // IPoseSearchProvider
 	virtual UE::Anim::IPoseSearchProvider::FSearchResult Search(const FAnimationBaseContext& GraphContext, const UAnimSequenceBase* Sequence) override;
 };
@@ -3051,6 +3057,10 @@ public: // IPoseSearchProvider
 void FModule::StartupModule()
 {
 	IModularFeatures::Get().RegisterModularFeature(UE::Anim::IPoseSearchProvider::ModularFeatureName, this);
+
+	// Enable the PoseSearch trace channel
+	UE::Trace::ToggleChannel(*FTraceLogger::Name.ToString(), true);
+	IModularFeatures::Get().RegisterModularFeature(TraceServices::ModuleFeatureName, &PoseSearchTraceModule);
 }
 
 void FModule::ShutdownModule()
diff --git a/Engine/Plugins/Experimental/Animation/PoseSearch/Source/Runtime/Private/Trace/PoseSearchTraceAnalyzer.cpp b/Engine/Plugins/Experimental/Animation/PoseSearch/Source/Runtime/Private/Trace/PoseSearchTraceAnalyzer.cpp
new file mode 100644
--- /dev/null
+++ b/Engine/Plugins/Experimental/Animation/PoseSearch/Source/Runtime/Private/Trace/PoseSearchTraceAnalyzer.cpp
@@ -0,0 +1,67 @@
+﻿// Copyright Epic Games, Inc. All Rights Reserved.
+
+#include "PoseSearchTraceAnalyzer.h"
+#include "PoseSearchTraceProvider.h"
+#include "TraceServices/Model/AnalysisSession.h"
+
+namespace UE { namespace PoseSearch {
+
+FTraceAnalyzer::FTraceAnalyzer(TraceServices::IAnalysisSession& InSession, FTraceProvider& InTraceProvider) : Session(InSession), TraceProvider(InTraceProvider)
+{
+}
+
+void FTraceAnalyzer::OnAnalysisBegin(const FOnAnalysisContext& Context)
+{
+	FInterfaceBuilder& Builder = Context.InterfaceBuilder;
+
+	ANSICHAR LoggerName[NAME_SIZE];
+	ANSICHAR MotionMatchingStateName[NAME_SIZE];
+
+	FTraceLogger::Name.GetPlainANSIString(LoggerName);
+	FTraceMotionMatchingState::Name.GetPlainANSIString(MotionMatchingStateName);
+
+	Builder.RouteEvent(RouteId_MotionMatchingState, LoggerName, MotionMatchingStateName);
+}
+
+bool FTraceAnalyzer::OnEvent(uint16 RouteId, EStyle Style, const FOnEventContext& Context)
+{
+	TraceServices::FAnalysisSessionEditScope Scope(Session);
+	const FEventData& EventData = Context.EventData;
+
+	// Gather event data values
+	const double Time = Context.EventTime.AsSeconds(EventData.GetValue<uint64>("Cycle"));
+	const uint16 FrameCounter = EventData.GetValue<uint16>("FrameCounter");
+	const uint64 AnimInstanceId = EventData.GetValue<uint64>("AnimInstanceId");
+	const int32 NodeId = EventData.GetValue<int32>("NodeId");
+
+	switch (RouteId)
+	{
+		case RouteId_MotionMatchingState:
+			{
+				FTraceMotionMatchingStateMessage Message;
+				Message.ElapsedPoseJumpTime = EventData.GetValue<float>("ElapsedPoseJumpTime");
+				// Cast back to our flag type
+				Message.Flags = static_cast<FTraceMotionMatchingState::EFlags>(EventData.GetValue<uint32>("Flags"));
+				Message.DbPoseIdx = EventData.GetValue<int32>("DbPoseIdx");
+				const TArrayView<const float> View = EventData.GetArrayView<float>("QueryVector");
+				Message.QueryVector = TArray<float>(View.GetData(), View.Num());
+				Message.DatabaseId = EventData.GetValue<uint64>("DatabaseId");
+
+				// Common data
+				Message.NodeId = NodeId;
+				Message.AnimInstanceId = AnimInstanceId;
+				Message.FrameCounter = FrameCounter;
+
+				TraceProvider.AppendMotionMatchingState(Message, Time);
+				break;
+			}
+		default:
+			{
+				// Should not happen
+				checkNoEntry();
+			}
+	}
+
+	return true;
+}
+}}
diff --git a/Engine/Plugins/Experimental/Animation/PoseSearch/Source/Runtime/Private/Trace/PoseSearchTraceAnalyzer.h b/Engine/Plugins/Experimental/Animation/PoseSearch/Source/Runtime/Private/Trace/PoseSearchTraceAnalyzer.h
new file mode 100644
--- /dev/null
+++ b/Engine/Plugins/Experimental/Animation/PoseSearch/Source/Runtime/Private/Trace/PoseSearchTraceAnalyzer.h
@@ -0,0 +1,40 @@
+﻿// Copyright Epic Games, Inc. All Rights Reserved.
+
+#pragma once
+
+#include "CoreMinimal.h"
+#include "Trace/Analyzer.h"
+
+namespace TraceServices
+{
+class IAnalysisSession;
+}
+
+namespace UE { namespace PoseSearch { 
+
+class FTraceProvider;
+
+/**
+* Used to relay trace information to the pose search provider
+*/
+class FTraceAnalyzer : public Trace::IAnalyzer
+{
+public:
+	FTraceAnalyzer(TraceServices::IAnalysisSession& InSession, FTraceProvider& InTraceProvider);
+
+protected:
+	virtual void OnAnalysisBegin(const FOnAnalysisContext& Context) override;
+	virtual bool OnEvent(uint16 RouteId, EStyle Style, const FOnEventContext& Context) override;
+
+private:
+	/** Event types for routing to the provider */
+	enum : uint16
+	{
+		RouteId_MotionMatchingState,
+	};
+
+	TraceServices::IAnalysisSession& Session;
+	FTraceProvider& TraceProvider;
+};
+
+}}
diff --git a/Engine/Plugins/Experimental/Animation/PoseSearch/Source/Runtime/Private/Trace/PoseSearchTraceLogger.cpp b/Engine/Plugins/Experimental/Animation/PoseSearch/Source/Runtime/Private/Trace/PoseSearchTraceLogger.cpp
new file mode 100644
--- /dev/null
+++ b/Engine/Plugins/Experimental/Animation/PoseSearch/Source/Runtime/Private/Trace/PoseSearchTraceLogger.cpp
@@ -0,0 +1,69 @@
+﻿// Copyright Epic Games, Inc. All Rights Reserved.
+
+#include "PoseSearchTraceLogger.h"
+#include "Animation/AnimInstanceProxy.h"
+#include "Trace/Trace.inl"
+
+#if UE_POSE_SEARCH_TRACE_ENABLED
+
+UE_TRACE_CHANNEL_DEFINE(PoseSearchChannel);
+
+UE_TRACE_EVENT_BEGIN(PoseSearch, MotionMatchingState)
+	UE_TRACE_EVENT_FIELD(uint64, Cycle)
+	UE_TRACE_EVENT_FIELD(uint16, FrameCounter)
+	UE_TRACE_EVENT_FIELD(uint64, AnimInstanceId)
+	UE_TRACE_EVENT_FIELD(int32, NodeId)
+	UE_TRACE_EVENT_FIELD(float, ElapsedPoseJumpTime)
+	UE_TRACE_EVENT_FIELD(uint32, Flags)
+	UE_TRACE_EVENT_FIELD(uint64, DatabaseId)
+	UE_TRACE_EVENT_FIELD(int32, DbPoseIdx)
+	UE_TRACE_EVENT_FIELD(float[], QueryVector)
+UE_TRACE_EVENT_END()
+
+namespace UE { namespace PoseSearch {
+
+const FName FTraceLogger::Name("PoseSearch");
+const FName FTraceMotionMatchingState::Name("MotionMatchingState");
+
+static bool ShouldTrace(const FAnimationBaseContext& InContext)
+{
+	const bool bChannelEnabled = UE_TRACE_CHANNELEXPR_IS_ENABLED(PoseSearchChannel);
+	if (!bChannelEnabled)
+	{
+		return false;
+	}
+
+	if (InContext.GetCurrentNodeId() == INDEX_NONE)
+	{
+		return false;
+	}
+
+	check(InContext.AnimInstanceProxy);
+	return !CANNOT_TRACE_OBJECT(InContext.AnimInstanceProxy->GetSkelMeshComponent());
+}
+
+void FTraceMotionMatchingState::Output(const FAnimationBaseContext& InContext, const FTraceMotionMatchingState& State)
+{
+	if (!ShouldTrace(InContext))
+	{
+		return;
+	}
+
+	UObject* AnimInstance = InContext.AnimInstanceProxy->GetAnimInstanceObject();
+	TRACE_OBJECT(AnimInstance);
+
+	UE_TRACE_LOG(PoseSearch, MotionMatchingState, PoseSearchChannel)
+		<< MotionMatchingState.Cycle(FPlatformTime::Cycles64())
+		<< MotionMatchingState.FrameCounter(FObjectTrace::GetObjectWorldTickCounter(AnimInstance))
+		<< MotionMatchingState.AnimInstanceId(FObjectTrace::GetObjectId(AnimInstance))
+		<< MotionMatchingState.NodeId(InContext.GetCurrentNodeId())
+		<< MotionMatchingState.ElapsedPoseJumpTime(State.ElapsedPoseJumpTime)
+		// Cast to uint32 for event message
+		<< MotionMatchingState.Flags(static_cast<uint32>(State.Flags))
+		<< MotionMatchingState.QueryVector(State.QueryVector.GetData(), State.QueryVector.Num())
+		<< MotionMatchingState.DbPoseIdx(State.DbPoseIdx)
+		<< MotionMatchingState.DatabaseId(State.DatabaseId);
+}
+
+}}
+#endif
\ No newline at end of file
diff --git a/Engine/Plugins/Experimental/Animation/PoseSearch/Source/Runtime/Private/Trace/PoseSearchTraceLogger.h b/Engine/Plugins/Experimental/Animation/PoseSearch/Source/Runtime/Private/Trace/PoseSearchTraceLogger.h
new file mode 100644
--- /dev/null
+++ b/Engine/Plugins/Experimental/Animation/PoseSearch/Source/Runtime/Private/Trace/PoseSearchTraceLogger.h
@@ -0,0 +1,71 @@
+﻿// Copyright Epic Games, Inc. All Rights Reserved.
+
+#pragma once
+
+#include "ObjectTrace.h"
+#include "PoseSearch/PoseSearch.h"
+
+// Enable this if object tracing is enabled, mimics animation tracing
+#define UE_POSE_SEARCH_TRACE_ENABLED OBJECT_TRACE_ENABLED
+
+#if UE_POSE_SEARCH_TRACE_ENABLED
+
+UE_TRACE_CHANNEL_EXTERN(PoseSearchChannel, POSESEARCH_API);
+
+namespace UE { namespace PoseSearch {
+
+struct FTraceLogger
+{
+	/** Used for reading trace data */
+    static const FName Name;
+};
+
+/**
+ * Used to trace motion matching state data via the logger, which is then placed into a timeline
+ */
+struct POSESEARCH_API FTraceMotionMatchingState
+{
+	/** Bitfield for various state booleans */
+	enum class EFlags : uint32
+	{
+		None = 0u,
+
+		/** Whether the last animation was a forced follow-up animation due to expended animation runway */
+		FollowupAnimation = 1u << 0
+	};
+	
+	/** Amount of time since the last pose switch */
+	float ElapsedPoseJumpTime = 0.0f;
+
+	/** Storage container for state booleans */
+	EFlags Flags = EFlags::None;
+
+	/** Feature value array of the query data */
+	TArrayView<const float> QueryVector;
+
+	/** Index of the pose in our database */
+	int32 DbPoseIdx = 0;
+
+	/** Object Id of the database asset */
+	uint64 DatabaseId = 0;
+
+	/** Output the current state info to the logger */
+	static void Output(const FAnimationBaseContext& InContext, const FTraceMotionMatchingState& State);
+	
+	static const FName Name;
+};
+ENUM_CLASS_FLAGS(FTraceMotionMatchingState::EFlags)
+
+
+}}
+
+#define UE_TRACE_POSE_SEARCH_MOTION_MATCHING_STATE(Context, MotionMatchingState) \
+	UE::PoseSearch::FTraceMotionMatchingState::Output(Context, State);
+
+#else
+
+// Empty declarations if not enabled
+#define UE_TRACE_POSE_SEARCH_MOTION_MATCHING_STATE(Context, State)
+
+#endif
+
diff --git a/Engine/Plugins/Experimental/Animation/PoseSearch/Source/Runtime/Private/Trace/PoseSearchTraceModule.cpp b/Engine/Plugins/Experimental/Animation/PoseSearch/Source/Runtime/Private/Trace/PoseSearchTraceModule.cpp
new file mode 100644
--- /dev/null
+++ b/Engine/Plugins/Experimental/Animation/PoseSearch/Source/Runtime/Private/Trace/PoseSearchTraceModule.cpp
@@ -0,0 +1,34 @@
+﻿// Copyright Epic Games, Inc. All Rights Reserved.
+
+#include "PoseSearchTraceModule.h"
+#include "PoseSearchTraceAnalyzer.h"
+#include "PoseSearchTraceProvider.h"
+
+namespace UE { namespace PoseSearch {
+
+const FName FPoseSearchTraceModule::ModuleName("PoseSearchTrace");
+
+void FPoseSearchTraceModule::GetModuleInfo(TraceServices::FModuleInfo& OutModuleInfo)
+{
+	OutModuleInfo.Name = ModuleName;
+	OutModuleInfo.DisplayName = TEXT("PoseSearch");
+}
+
+void FPoseSearchTraceModule::OnAnalysisBegin(TraceServices::IAnalysisSession& InSession)
+{
+	// Add our provider and analyzer, starting our systems
+	FTraceProvider* PoseSearchProvider = new FTraceProvider(InSession);
+	InSession.AddProvider(FTraceProvider::ProviderName, PoseSearchProvider);
+	InSession.AddAnalyzer(new FTraceAnalyzer(InSession, *PoseSearchProvider));
+}
+
+void FPoseSearchTraceModule::GetLoggers(TArray<const TCHAR*>& OutLoggers)
+{
+	OutLoggers.Add(TEXT("PoseSearch"));
+}
+
+void FPoseSearchTraceModule::GenerateReports(const TraceServices::IAnalysisSession& Session, const TCHAR* CmdLine, const TCHAR* OutputDirectory)
+{
+}
+
+}}
diff --git a/Engine/Plugins/Experimental/Animation/PoseSearch/Source/Runtime/Private/Trace/PoseSearchTraceModule.h b/Engine/Plugins/Experimental/Animation/PoseSearch/Source/Runtime/Private/Trace/PoseSearchTraceModule.h
new file mode 100644
--- /dev/null
+++ b/Engine/Plugins/Experimental/Animation/PoseSearch/Source/Runtime/Private/Trace/PoseSearchTraceModule.h
@@ -0,0 +1,29 @@
+﻿// Copyright Epic Games, Inc. All Rights Reserved.
+
+#pragma once
+
+#include "CoreMinimal.h"
+#include "TraceServices/ModuleService.h"
+
+namespace UE { namespace PoseSearch {
+
+/**
+* Module used for loading our PoseSearch tracing systems
+*/
+class FPoseSearchTraceModule : public TraceServices::IModule
+{
+public:
+	FPoseSearchTraceModule() = default;
+	virtual ~FPoseSearchTraceModule() = default;
+	
+	const static FName ModuleName;
+
+protected:
+	virtual void GetModuleInfo(TraceServices::FModuleInfo& OutModuleInfo) override;
+	virtual void OnAnalysisBegin(TraceServices::IAnalysisSession& InSession) override;
+	virtual void GetLoggers(TArray<const TCHAR*>& OutLoggers) override;
+	virtual void GenerateReports(const TraceServices::IAnalysisSession& Session, const TCHAR* CmdLine, const TCHAR* OutputDirectory) override;
+	virtual const TCHAR* GetCommandLineArgument() override { return TEXT("posesearchtrace"); }
+};
+
+}}
\ No newline at end of file
diff --git a/Engine/Plugins/Experimental/Animation/PoseSearch/Source/Runtime/Private/Trace/PoseSearchTraceProvider.cpp b/Engine/Plugins/Experimental/Animation/PoseSearch/Source/Runtime/Private/Trace/PoseSearchTraceProvider.cpp
new file mode 100644
--- /dev/null
+++ b/Engine/Plugins/Experimental/Animation/PoseSearch/Source/Runtime/Private/Trace/PoseSearchTraceProvider.cpp
@@ -0,0 +1,32 @@
+﻿// Copyright Epic Games, Inc. All Rights Reserved.
+
+#include "PoseSearchTraceProvider.h"
+#include "Model/PointTimeline.h"
+
+namespace UE { namespace PoseSearch {
+
+const FName FTraceProvider::ProviderName("PoseSearchTraceProvider");
+
+FTraceProvider::FTraceProvider(TraceServices::IAnalysisSession& InSession) : Session(InSession)
+{
+}
+
+bool FTraceProvider::ReadMotionMatchingStateTimeline(uint64 InObjectId, int32 InNodeId, TFunctionRef<void(const FMotionMatchingStateTimeline&)> Callback) const
+{
+	Session.ReadAccessCheck();
+	return MotionMatchingStateTimelineStorage.ReadTimeline(InObjectId, InNodeId, Callback);
+}
+
+
+void FTraceProvider::AppendMotionMatchingState(const FTraceMotionMatchingStateMessage& InMessage, double InTime)
+{
+	Session.WriteAccessCheck();
+
+	TSharedRef<TraceServices::TPointTimeline<FTraceMotionMatchingStateMessage>> Timeline = MotionMatchingStateTimelineStorage.GetTimeline(Session, InMessage.AnimInstanceId, InMessage.NodeId);
+	Timeline->AppendEvent(InTime, InMessage);
+	
+	Session.UpdateDurationSeconds(InTime);
+}
+
+
+}}
diff --git a/Engine/Plugins/Experimental/Animation/PoseSearch/Source/Runtime/Private/Trace/PoseSearchTraceProvider.h b/Engine/Plugins/Experimental/Animation/PoseSearch/Source/Runtime/Private/Trace/PoseSearchTraceProvider.h
new file mode 100644
--- /dev/null
+++ b/Engine/Plugins/Experimental/Animation/PoseSearch/Source/Runtime/Private/Trace/PoseSearchTraceProvider.h
@@ -0,0 +1,126 @@
+﻿// Copyright Epic Games, Inc. All Rights Reserved.
+
+#pragma once
+
+#include "CoreMinimal.h"
+#include "Model/PointTimeline.h"
+#include "PoseSearchTraceLogger.h"
+#include "Model/IntervalTimeline.h"
+#include "TraceServices/Model/AnalysisSession.h"
+
+
+namespace UE { namespace PoseSearch {
+
+// Message types for appending / reading to the timeline
+/** Base message type for common data */
+struct FTraceMessage
+{
+	uint64 AnimInstanceId;
+	int32 NodeId;
+	uint16 FrameCounter;
+};
+
+/** Motion matching state message container */
+struct FTraceMotionMatchingStateMessage : FTraceMessage
+{
+	FTraceMotionMatchingState::EFlags Flags;
+	float ElapsedPoseJumpTime;
+	TArray<float> QueryVector;
+	int32 DbPoseIdx;
+	uint64 DatabaseId;
+};
+
+
+/**
+ * Provider to the widgets for pose search functionality, largely mimicking FAnimationProvider
+ */
+class POSESEARCH_API FTraceProvider : public TraceServices::IProvider
+{
+public:
+	FTraceProvider(TraceServices::IAnalysisSession& InSession);
+
+	/** Read the object-relevant timeline info and provide it via callback */
+	using FMotionMatchingStateTimeline = TraceServices::ITimeline<FTraceMotionMatchingStateMessage>;
+	bool ReadMotionMatchingStateTimeline(uint64 InAnimInstanceId, int32 InNodeId, TFunctionRef<void(const FMotionMatchingStateTimeline&)> Callback) const;
+
+	/** Append to the timeline on our object */
+	void AppendMotionMatchingState(const FTraceMotionMatchingStateMessage& InMessage, double InTime);
+
+	static const FName ProviderName;
+
+private:
+	/**
+	 * Convenience struct for creating timelines per message type.
+	 * We store an array of timelines for every possible AnimInstance Id that gets appended.
+	 */
+	template <class MessageType, class TimelineType = TraceServices::TPointTimeline<MessageType>>
+	struct TTimelineStorage
+	{
+		// Map the timelines as AnimInstanceId -> NodeId -> Timeline
+		using FNodeToTimelineMap = TMap<int32, uint32>;
+		using FAnimInstanceToNodeMap = TMap<uint64, FNodeToTimelineMap>;
+
+		/** Retrieves the timeline for internal use, creating it if it does not exist */
+		TSharedRef<TimelineType> GetTimeline(TraceServices::IAnalysisSession& InSession, uint64 InAnimInstanceId, int32 InNodeId)
+		{
+			const uint32* TimelineIndex = nullptr;
+			
+			FNodeToTimelineMap* NodeToTimelineMap = AnimInstanceIdToTimelines.Find(InAnimInstanceId);
+			if (NodeToTimelineMap == nullptr)
+			{
+				// Create the timeline map if this is the first time accessing with this anim instance
+				NodeToTimelineMap = &AnimInstanceIdToTimelines.Add(InAnimInstanceId, {});
+			}
+			else
+			{
+				// Anim instance already used, attempt to find the NodeId's timeline
+				TimelineIndex = NodeToTimelineMap->Find(InNodeId);
+			}
+
+			if (TimelineIndex == nullptr)
+			{
+				// Append our timeline to the storage and our object + the storage index to the map
+				TSharedRef<TimelineType> Timeline = MakeShared<TimelineType>(InSession.GetLinearAllocator());
+				const uint32 NewIndex = Timelines.Add(Timeline);
+				NodeToTimelineMap->Add(InNodeId, NewIndex);
+				return Timeline;
+			}
+
+			return Timelines[*TimelineIndex];
+		}
+
+		/** Retrieve a timeline from an anim instance + node and execute the callback */
+		bool ReadTimeline(uint64 InAnimInstanceId, int32 InNodeId, TFunctionRef<void(const TraceServices::ITimeline<MessageType>&)> Callback) const
+		{
+			const FNodeToTimelineMap* NodeToTimelineMap = AnimInstanceIdToTimelines.Find(InAnimInstanceId);
+			if (NodeToTimelineMap == nullptr)
+			{
+				return false;
+			}
+
+			const uint32* TimelineIndex = NodeToTimelineMap->Find(InNodeId);
+			if (TimelineIndex == nullptr || !Timelines.IsValidIndex(*TimelineIndex))
+			{
+				return false;
+			}
+
+			Callback(*Timelines[*TimelineIndex]);
+			return true;
+		}
+
+		/** Maps AnimInstanceIds to a map of NodeIds to timelines. AnimInstanceId -> NodeId -> Timeline */
+		TMap<uint64, TMap<int32, uint32>> AnimInstanceIdToTimelines;
+
+		/** Timelines per node */
+		TArray<TSharedRef<TimelineType>> Timelines;
+	};
+
+	// Storage for each message type
+	struct FMotionMatchingStateTimelineStorage : TTimelineStorage<FTraceMotionMatchingStateMessage>
+	{
+	} MotionMatchingStateTimelineStorage;
+
+
+	TraceServices::IAnalysisSession& Session;
+};
+}}
