Undo //UE5/Main/Engine/Plugins/Experimental/Animation/PoseSearch/Source/Runtime/... changelist 16759388

diff --git a/Engine/Plugins/Experimental/Animation/PoseSearch/Source/Runtime/PoseSearch.Build.cs b/Engine/Plugins/Experimental/Animation/PoseSearch/Source/Runtime/PoseSearch.Build.cs
--- a/Engine/Plugins/Experimental/Animation/PoseSearch/Source/Runtime/PoseSearch.Build.cs
+++ b/Engine/Plugins/Experimental/Animation/PoseSearch/Source/Runtime/PoseSearch.Build.cs
@@ -32,12 +32,7 @@ public class PoseSearch : ModuleRules
 				"Core",
 				"CoreUObject",
 				"Engine",
-				"AnimationCore",
-				
-				// Trace dependencies
-				"TraceLog",
-				"TraceAnalysis",
-				"TraceServices"
+				"AnimationCore"
 			}
 		);
 
diff --git a/Engine/Plugins/Experimental/Animation/PoseSearch/Source/Runtime/Private/AnimNode_MotionMatching.cpp b/Engine/Plugins/Experimental/Animation/PoseSearch/Source/Runtime/Private/AnimNode_MotionMatching.cpp
--- a/Engine/Plugins/Experimental/Animation/PoseSearch/Source/Runtime/Private/AnimNode_MotionMatching.cpp
+++ b/Engine/Plugins/Experimental/Animation/PoseSearch/Source/Runtime/Private/AnimNode_MotionMatching.cpp
@@ -7,10 +7,10 @@
 #include "Animation/AnimInstanceProxy.h"
 #include "Animation/AnimNode_Inertialization.h"
 #include "Animation/AnimSequence.h"
-#include "Trace/PoseSearchTraceLogger.h"
 
 #define LOCTEXT_NAMESPACE "AnimNode_MotionMatching"
 
+
 /////////////////////////////////////////////////////
 // FAnimNode_MotionMatching
 
@@ -166,18 +166,8 @@ void FAnimNode_MotionMatching::UpdateAssetPlayer(const FAnimationUpdateContext&
 	const FPredictionSequenceState SequenceState = { SequencePlayerNode.GetSequence() , SequencePlayerNode.GetAccumulatedTime(), SequencePlayerNode.GetLoopAnimation() };
 	const float PlayRate = UPoseSearchPredictionDistanceMatching::ComputePlayRate(Context, Prediction, PredictionSettings, SequenceState);
 	SequencePlayerNode.SetPlayRate(PlayRate);
+
 	Source.Update(Context);
-	
-#if UE_POSE_SEARCH_TRACE_ENABLED
-	// Gather our node state and trace it
-	UE::PoseSearch::FTraceMotionMatchingState State;
-	State.ElapsedPoseJumpTime = ElapsedPoseJumpTime;
-	State.Flags = UE::PoseSearch::FTraceMotionMatchingState::EFlags::None;
-	State.DbPoseIdx = DbPoseIdx;
-	State.DatabaseId = FObjectTrace::GetObjectId(Database.Get());
-	State.QueryVector = ComposedQuery.GetValues();
-	UE_TRACE_POSE_SEARCH_MOTION_MATCHING_STATE(Context, State)
-#endif
 }
 
 bool FAnimNode_MotionMatching::HasPreUpdate() const
diff --git a/Engine/Plugins/Experimental/Animation/PoseSearch/Source/Runtime/Private/PoseSearch.cpp b/Engine/Plugins/Experimental/Animation/PoseSearch/Source/Runtime/Private/PoseSearch.cpp
--- a/Engine/Plugins/Experimental/Animation/PoseSearch/Source/Runtime/Private/PoseSearch.cpp
+++ b/Engine/Plugins/Experimental/Animation/PoseSearch/Source/Runtime/Private/PoseSearch.cpp
@@ -21,9 +21,6 @@
 #include "Animation/AnimationPoseData.h"
 #include "AnimationRuntime.h"
 #include "BonePose.h"
-#include "Trace/PoseSearchTraceLogger.h"
-#include "Trace/PoseSearchTraceModule.h"
-#include "TraceServices/ModuleService.h"
 #include "UObject/ObjectSaveContext.h"
 
 IMPLEMENT_ANIMGRAPH_MESSAGE(UE::PoseSearch::IPoseHistoryProvider);
@@ -3047,9 +3044,6 @@ public: // IModuleInterface
 	virtual void StartupModule() override;
 	virtual void ShutdownModule() override;
 
-	/** Host our trace module here, which registers the PoseSearch trace features */
-	FPoseSearchTraceModule PoseSearchTraceModule;
-
 public: // IPoseSearchProvider
 	virtual UE::Anim::IPoseSearchProvider::FSearchResult Search(const FAnimationBaseContext& GraphContext, const UAnimSequenceBase* Sequence) override;
 };
@@ -3057,10 +3051,6 @@ public: // IPoseSearchProvider
 void FModule::StartupModule()
 {
 	IModularFeatures::Get().RegisterModularFeature(UE::Anim::IPoseSearchProvider::ModularFeatureName, this);
-
-	// Enable the PoseSearch trace channel
-	UE::Trace::ToggleChannel(*FTraceLogger::Name.ToString(), true);
-	IModularFeatures::Get().RegisterModularFeature(TraceServices::ModuleFeatureName, &PoseSearchTraceModule);
 }
 
 void FModule::ShutdownModule()
diff --git a/Engine/Plugins/Experimental/Animation/PoseSearch/Source/Runtime/Private/Trace/PoseSearchTraceAnalyzer.cpp b/Engine/Plugins/Experimental/Animation/PoseSearch/Source/Runtime/Private/Trace/PoseSearchTraceAnalyzer.cpp
deleted file mode 100644
diff --git a/Engine/Plugins/Experimental/Animation/PoseSearch/Source/Runtime/Private/Trace/PoseSearchTraceAnalyzer.h b/Engine/Plugins/Experimental/Animation/PoseSearch/Source/Runtime/Private/Trace/PoseSearchTraceAnalyzer.h
deleted file mode 100644
diff --git a/Engine/Plugins/Experimental/Animation/PoseSearch/Source/Runtime/Private/Trace/PoseSearchTraceLogger.cpp b/Engine/Plugins/Experimental/Animation/PoseSearch/Source/Runtime/Private/Trace/PoseSearchTraceLogger.cpp
deleted file mode 100644
diff --git a/Engine/Plugins/Experimental/Animation/PoseSearch/Source/Runtime/Private/Trace/PoseSearchTraceLogger.h b/Engine/Plugins/Experimental/Animation/PoseSearch/Source/Runtime/Private/Trace/PoseSearchTraceLogger.h
deleted file mode 100644
diff --git a/Engine/Plugins/Experimental/Animation/PoseSearch/Source/Runtime/Private/Trace/PoseSearchTraceModule.cpp b/Engine/Plugins/Experimental/Animation/PoseSearch/Source/Runtime/Private/Trace/PoseSearchTraceModule.cpp
deleted file mode 100644
diff --git a/Engine/Plugins/Experimental/Animation/PoseSearch/Source/Runtime/Private/Trace/PoseSearchTraceModule.h b/Engine/Plugins/Experimental/Animation/PoseSearch/Source/Runtime/Private/Trace/PoseSearchTraceModule.h
deleted file mode 100644
diff --git a/Engine/Plugins/Experimental/Animation/PoseSearch/Source/Runtime/Private/Trace/PoseSearchTraceProvider.cpp b/Engine/Plugins/Experimental/Animation/PoseSearch/Source/Runtime/Private/Trace/PoseSearchTraceProvider.cpp
deleted file mode 100644
diff --git a/Engine/Plugins/Experimental/Animation/PoseSearch/Source/Runtime/Private/Trace/PoseSearchTraceProvider.h b/Engine/Plugins/Experimental/Animation/PoseSearch/Source/Runtime/Private/Trace/PoseSearchTraceProvider.h
deleted file mode 100644
