Infrastructure for Trace functionality in the PoseSearch plugin
- Includes tracing macros for motion matching data
- Does not include any editor code, this should be entirely runtime stuff
- utrace files confirmed not to crash in standalone insights tool despite no dependencies on these features, so plugin compatibility should not be an issue

#rnx

diff --git a/Engine/Plugins/Experimental/Animation/PoseSearch/Source/Runtime/PoseSearch.Build.cs b/Engine/Plugins/Experimental/Animation/PoseSearch/Source/Runtime/PoseSearch.Build.cs
--- a/Engine/Plugins/Experimental/Animation/PoseSearch/Source/Runtime/PoseSearch.Build.cs
+++ b/Engine/Plugins/Experimental/Animation/PoseSearch/Source/Runtime/PoseSearch.Build.cs
@@ -32,7 +32,8 @@ public class PoseSearch : ModuleRules
 				"Core",
 				"CoreUObject",
 				"Engine",
-				"AnimationCore"
+				"AnimationCore",
+				"TraceLog"
 			}
 		);
 
diff --git a/Engine/Plugins/Experimental/Animation/PoseSearch/Source/Runtime/Private/AnimNode_MotionMatching.cpp b/Engine/Plugins/Experimental/Animation/PoseSearch/Source/Runtime/Private/AnimNode_MotionMatching.cpp
--- a/Engine/Plugins/Experimental/Animation/PoseSearch/Source/Runtime/Private/AnimNode_MotionMatching.cpp
+++ b/Engine/Plugins/Experimental/Animation/PoseSearch/Source/Runtime/Private/AnimNode_MotionMatching.cpp
@@ -7,6 +7,7 @@
 #include "Animation/AnimInstanceProxy.h"
 #include "Animation/AnimNode_Inertialization.h"
 #include "Animation/AnimSequence.h"
+#include "Trace/PoseSearchTraceLogger.h"
 
 #define LOCTEXT_NAMESPACE "AnimNode_MotionMatching"
 
@@ -168,6 +169,15 @@ void FAnimNode_MotionMatching::UpdateAssetPlayer(const FAnimationUpdateContext&
 	SequencePlayerNode.SetPlayRate(PlayRate);
 
 	Source.Update(Context);
+	
+#if UE_POSE_SEARCH_TRACE_ENABLED
+	UE::PoseSearch::FTraceMotionMatchingState State;
+	State.ElapsedPoseJumpTime = ElapsedPoseJumpTime;
+	State.QueryVector = ComposedQuery.GetValues();
+	State.DbPoseIdx = DbPoseIdx;
+	State.DatabaseId = FObjectTrace::GetObjectId(Database.Get());
+	UE_TRACE_POSE_SEARCH_MOTION_MATCHING_STATE(Context, State)
+#endif
 }
 
 bool FAnimNode_MotionMatching::HasPreUpdate() const
diff --git a/Engine/Plugins/Experimental/Animation/PoseSearch/Source/Runtime/Private/PoseSearch.cpp b/Engine/Plugins/Experimental/Animation/PoseSearch/Source/Runtime/Private/PoseSearch.cpp
--- a/Engine/Plugins/Experimental/Animation/PoseSearch/Source/Runtime/Private/PoseSearch.cpp
+++ b/Engine/Plugins/Experimental/Animation/PoseSearch/Source/Runtime/Private/PoseSearch.cpp
@@ -21,6 +21,7 @@
 #include "Animation/AnimationPoseData.h"
 #include "AnimationRuntime.h"
 #include "BonePose.h"
+#include "Trace/PoseSearchTraceLogger.h"
 #include "UObject/ObjectSaveContext.h"
 
 IMPLEMENT_ANIMGRAPH_MESSAGE(UE::PoseSearch::IPoseHistoryProvider);
@@ -3051,6 +3052,11 @@ public: // IPoseSearchProvider
 void FModule::StartupModule()
 {
 	IModularFeatures::Get().RegisterModularFeature(UE::Anim::IPoseSearchProvider::ModularFeatureName, this);
+
+#if UE_POSE_SEARCH_TRACE_ENABLED
+	// Enable the PoseSearch trace channel
+	UE::Trace::ToggleChannel(*FTraceLogger::Name.ToString(), true);
+#endif
 }
 
 void FModule::ShutdownModule()
diff --git a/Engine/Plugins/Experimental/Animation/PoseSearch/Source/Runtime/Private/Trace/PoseSearchTraceLogger.cpp b/Engine/Plugins/Experimental/Animation/PoseSearch/Source/Runtime/Private/Trace/PoseSearchTraceLogger.cpp
new file mode 100644
--- /dev/null
+++ b/Engine/Plugins/Experimental/Animation/PoseSearch/Source/Runtime/Private/Trace/PoseSearchTraceLogger.cpp
@@ -0,0 +1,69 @@
+﻿// Copyright Epic Games, Inc. All Rights Reserved.
+
+#include "PoseSearchTraceLogger.h"
+#include "Animation/AnimInstanceProxy.h"
+#include "Trace/Trace.inl"
+
+#if UE_POSE_SEARCH_TRACE_ENABLED
+
+UE_TRACE_CHANNEL_DEFINE(PoseSearchChannel);
+
+UE_TRACE_EVENT_BEGIN(PoseSearch, MotionMatchingState)
+	UE_TRACE_EVENT_FIELD(uint64, Cycle)
+	UE_TRACE_EVENT_FIELD(uint16, FrameCounter)
+	UE_TRACE_EVENT_FIELD(uint64, AnimInstanceId)
+	UE_TRACE_EVENT_FIELD(int32, NodeId)
+	UE_TRACE_EVENT_FIELD(float, ElapsedPoseJumpTime)
+	UE_TRACE_EVENT_FIELD(uint32, Flags)
+	UE_TRACE_EVENT_FIELD(uint64, DatabaseId)
+	UE_TRACE_EVENT_FIELD(int32, DbPoseIdx)
+	UE_TRACE_EVENT_FIELD(float[], QueryVector)
+UE_TRACE_EVENT_END()
+
+namespace UE { namespace PoseSearch {
+
+const FName FTraceLogger::Name("PoseSearch");
+const FName FTraceMotionMatchingState::Name("MotionMatchingState");
+
+static bool ShouldTrace(const FAnimationBaseContext& InContext)
+{
+	const bool bChannelEnabled = UE_TRACE_CHANNELEXPR_IS_ENABLED(PoseSearchChannel);
+	if (!bChannelEnabled)
+	{
+		return false;
+	}
+
+	if (InContext.GetCurrentNodeId() == INDEX_NONE)
+	{
+		return false;
+	}
+
+	check(InContext.AnimInstanceProxy);
+	return !CANNOT_TRACE_OBJECT(InContext.AnimInstanceProxy->GetSkelMeshComponent());
+}
+
+void FTraceMotionMatchingState::Output(const FAnimationBaseContext& InContext, const FTraceMotionMatchingState& State)
+{
+	if (!ShouldTrace(InContext))
+	{
+		return;
+	}
+
+	UObject* AnimInstance = InContext.AnimInstanceProxy->GetAnimInstanceObject();
+	TRACE_OBJECT(AnimInstance);
+
+	UE_TRACE_LOG(PoseSearch, MotionMatchingState, PoseSearchChannel)
+		<< MotionMatchingState.Cycle(FPlatformTime::Cycles64())
+		<< MotionMatchingState.FrameCounter(FObjectTrace::GetObjectWorldTickCounter(AnimInstance))
+		<< MotionMatchingState.AnimInstanceId(FObjectTrace::GetObjectId(AnimInstance))
+		<< MotionMatchingState.NodeId(InContext.GetCurrentNodeId())
+		<< MotionMatchingState.ElapsedPoseJumpTime(State.ElapsedPoseJumpTime)
+		// Cast to uint32 for event message
+		<< MotionMatchingState.Flags(static_cast<uint32>(State.Flags))
+		<< MotionMatchingState.QueryVector(State.QueryVector.GetData(), State.QueryVector.Num())
+		<< MotionMatchingState.DbPoseIdx(State.DbPoseIdx)
+		<< MotionMatchingState.DatabaseId(State.DatabaseId);
+}
+
+}}
+#endif // UE_POSE_SEARCH_TRACE_ENABLED
\ No newline at end of file
diff --git a/Engine/Plugins/Experimental/Animation/PoseSearch/Source/Runtime/Private/Trace/PoseSearchTraceLogger.h b/Engine/Plugins/Experimental/Animation/PoseSearch/Source/Runtime/Private/Trace/PoseSearchTraceLogger.h
new file mode 100644
--- /dev/null
+++ b/Engine/Plugins/Experimental/Animation/PoseSearch/Source/Runtime/Private/Trace/PoseSearchTraceLogger.h
@@ -0,0 +1,71 @@
+﻿// Copyright Epic Games, Inc. All Rights Reserved.
+
+#pragma once
+
+#include "ObjectTrace.h"
+#include "PoseSearch/PoseSearch.h"
+
+// Enable this if object tracing is enabled, mimics animation tracing
+#define UE_POSE_SEARCH_TRACE_ENABLED OBJECT_TRACE_ENABLED
+
+#if UE_POSE_SEARCH_TRACE_ENABLED
+
+UE_TRACE_CHANNEL_EXTERN(PoseSearchChannel, POSESEARCH_API);
+
+namespace UE { namespace PoseSearch {
+
+struct POSESEARCH_API FTraceLogger
+{
+	/** Used for reading trace data */
+    static const FName Name;
+};
+
+/**
+ * Used to trace motion matching state data via the logger, which is then placed into a timeline
+ */
+struct POSESEARCH_API FTraceMotionMatchingState
+{
+	/** Bitfield for various state booleans */
+	enum class EFlags : uint32
+	{
+		None = 0u,
+
+		/** Whether the last animation was a forced follow-up animation due to expended animation runway */
+		FollowupAnimation = 1u << 0
+	};
+	
+	/** Amount of time since the last pose switch */
+	float ElapsedPoseJumpTime = 0.0f;
+
+	/** Storage container for state booleans */
+	EFlags Flags = EFlags::None;
+
+	/** Feature value array of the query data */
+	TArrayView<const float> QueryVector;
+
+	/** Index of the pose in our database */
+	int32 DbPoseIdx = 0;
+
+	/** Object Id of the database asset */
+	uint64 DatabaseId = 0;
+
+	/** Output the current state info to the logger */
+	static void Output(const FAnimationBaseContext& InContext, const FTraceMotionMatchingState& State);
+	
+	static const FName Name;
+};
+ENUM_CLASS_FLAGS(FTraceMotionMatchingState::EFlags)
+
+
+}}
+
+#define UE_TRACE_POSE_SEARCH_MOTION_MATCHING_STATE(Context, MotionMatchingState) \
+	UE::PoseSearch::FTraceMotionMatchingState::Output(Context, State);
+
+#else
+
+// Empty declarations if not enabled
+#define UE_TRACE_POSE_SEARCH_MOTION_MATCHING_STATE(Context, State)
+
+#endif // UE_POSE_SEARCH_TRACE_ENABLED
+
