Gameplay Anim Gym and Pose Warping project cleanup. Tuning and improvements for distance-based play rate adjustment in motion matching.

diff --git a/Engine/Plugins/Experimental/Animation/AnimationWarping/Source/Editor/Private/AnimGraph/AnimGraphNode_OrientationWarping.cpp b/Engine/Plugins/Experimental/Animation/AnimationWarping/Source/Editor/Private/AnimGraph/AnimGraphNode_OrientationWarping.cpp
--- a/Engine/Plugins/Experimental/Animation/AnimationWarping/Source/Editor/Private/AnimGraph/AnimGraphNode_OrientationWarping.cpp
+++ b/Engine/Plugins/Experimental/Animation/AnimationWarping/Source/Editor/Private/AnimGraph/AnimGraphNode_OrientationWarping.cpp
@@ -67,7 +67,7 @@ void UAnimGraphNode_OrientationWarping::GetInputLinkAttributes(FNodeAttributeArr
 {
 	if (Node.Mode == EWarpingEvaluationMode::Graph)
 	{
-		OutAttributes.Add(UE::Anim::IAnimRootMotionProvider::RootMotionDeltaAttributeName);
+		OutAttributes.Add(UE::Anim::IAnimRootMotionProvider::AttributeName);
 	}
 }
 
@@ -75,7 +75,7 @@ void UAnimGraphNode_OrientationWarping::GetOutputLinkAttributes(FNodeAttributeAr
 {
 	if (Node.Mode == EWarpingEvaluationMode::Graph)
 	{
-		OutAttributes.Add(UE::Anim::IAnimRootMotionProvider::RootMotionDeltaAttributeName);
+		OutAttributes.Add(UE::Anim::IAnimRootMotionProvider::AttributeName);
 	}
 }
 
diff --git a/Engine/Plugins/Experimental/Animation/AnimationWarping/Source/Editor/Private/AnimGraph/AnimGraphNode_SlopeWarping.cpp b/Engine/Plugins/Experimental/Animation/AnimationWarping/Source/Editor/Private/AnimGraph/AnimGraphNode_SlopeWarping.cpp
--- a/Engine/Plugins/Experimental/Animation/AnimationWarping/Source/Editor/Private/AnimGraph/AnimGraphNode_SlopeWarping.cpp
+++ b/Engine/Plugins/Experimental/Animation/AnimationWarping/Source/Editor/Private/AnimGraph/AnimGraphNode_SlopeWarping.cpp
@@ -29,7 +29,7 @@ void UAnimGraphNode_SlopeWarping::GetInputLinkAttributes(FNodeAttributeArray& Ou
 {
 	if (Node.Mode == EWarpingEvaluationMode::Graph)
 	{
-		OutAttributes.Add(UE::Anim::IAnimRootMotionProvider::RootMotionDeltaAttributeName);
+		OutAttributes.Add(UE::Anim::IAnimRootMotionProvider::AttributeName);
 	}
 }
 
@@ -37,7 +37,7 @@ void UAnimGraphNode_SlopeWarping::GetOutputLinkAttributes(FNodeAttributeArray& O
 {
 	if (Node.Mode == EWarpingEvaluationMode::Graph)
 	{
-		OutAttributes.Add(UE::Anim::IAnimRootMotionProvider::RootMotionDeltaAttributeName);
+		OutAttributes.Add(UE::Anim::IAnimRootMotionProvider::AttributeName);
 	}
 }
 
diff --git a/Engine/Plugins/Experimental/Animation/AnimationWarping/Source/Editor/Private/AnimGraph/AnimGraphNode_StrideWarping.cpp b/Engine/Plugins/Experimental/Animation/AnimationWarping/Source/Editor/Private/AnimGraph/AnimGraphNode_StrideWarping.cpp
--- a/Engine/Plugins/Experimental/Animation/AnimationWarping/Source/Editor/Private/AnimGraph/AnimGraphNode_StrideWarping.cpp
+++ b/Engine/Plugins/Experimental/Animation/AnimationWarping/Source/Editor/Private/AnimGraph/AnimGraphNode_StrideWarping.cpp
@@ -50,6 +50,11 @@ void UAnimGraphNode_StrideWarping::CustomizePinData(UEdGraphPin* Pin, FName Sour
 	{
 		Pin->bHidden = (Node.Mode == EWarpingEvaluationMode::Manual);
 	}
+
+	if (Pin->PinName == GET_MEMBER_NAME_STRING_CHECKED(FAnimNode_StrideWarping, MinSpeedTolerance))
+	{
+		Pin->bHidden = (Node.Mode == EWarpingEvaluationMode::Manual);
+	}
 }
 
 void UAnimGraphNode_StrideWarping::CustomizeDetails(IDetailLayoutBuilder& DetailBuilder)
@@ -68,6 +73,7 @@ void UAnimGraphNode_StrideWarping::CustomizeDetails(IDetailLayoutBuilder& Detail
 	if (Node.Mode == EWarpingEvaluationMode::Manual)
 	{
 		DetailBuilder.HideProperty(NodeHandle->GetChildHandle(GET_MEMBER_NAME_CHECKED(FAnimNode_StrideWarping, LocomotionSpeed)));
+		DetailBuilder.HideProperty(NodeHandle->GetChildHandle(GET_MEMBER_NAME_CHECKED(FAnimNode_StrideWarping, MinSpeedTolerance)));
 	}
 }
 
@@ -126,6 +132,13 @@ void UAnimGraphNode_StrideWarping::PostEditChangeProperty(struct FPropertyChange
 						Pin->BreakAllPinLinks();
 					}
 				}
+				else if (Pin->PinName == GET_MEMBER_NAME_STRING_CHECKED(FAnimNode_StrideWarping, MinSpeedTolerance))
+				{
+					if (Node.Mode == EWarpingEvaluationMode::Manual)
+					{
+						Pin->BreakAllPinLinks();
+					}
+				}
 			}
 
 			bRequiresNodeReconstruct = true;
@@ -143,7 +156,7 @@ void UAnimGraphNode_StrideWarping::GetInputLinkAttributes(FNodeAttributeArray& O
 {
 	if (Node.Mode == EWarpingEvaluationMode::Graph)
 	{
-		OutAttributes.Add(UE::Anim::IAnimRootMotionProvider::RootMotionDeltaAttributeName);
+		OutAttributes.Add(UE::Anim::IAnimRootMotionProvider::AttributeName);
 	}
 }
 
@@ -151,7 +164,7 @@ void UAnimGraphNode_StrideWarping::GetOutputLinkAttributes(FNodeAttributeArray&
 {
 	if (Node.Mode == EWarpingEvaluationMode::Graph)
 	{
-		OutAttributes.Add(UE::Anim::IAnimRootMotionProvider::RootMotionDeltaAttributeName);
+		OutAttributes.Add(UE::Anim::IAnimRootMotionProvider::AttributeName);
 	}
 }
 
diff --git a/Engine/Plugins/Experimental/Animation/AnimationWarping/Source/Runtime/Private/AnimationWarpingRuntimeModule.cpp b/Engine/Plugins/Experimental/Animation/AnimationWarping/Source/Runtime/Private/AnimationWarpingRuntimeModule.cpp
--- a/Engine/Plugins/Experimental/Animation/AnimationWarping/Source/Runtime/Private/AnimationWarpingRuntimeModule.cpp
+++ b/Engine/Plugins/Experimental/Animation/AnimationWarping/Source/Runtime/Private/AnimationWarpingRuntimeModule.cpp
@@ -14,7 +14,7 @@ class FAnimationWarpingRuntimeModule : public IAnimationWarpingRuntimeModule
 namespace UE { namespace AnimationWarping {
 
 // Root motion animation attribute data will always be accessible on the root bone's attribute container
-static const UE::Anim::FAttributeId RootMotionAttributeId = { UE::Anim::IAnimRootMotionProvider::RootMotionDeltaAttributeName , FCompactPoseBoneIndex(0) };
+static const UE::Anim::FAttributeId RootMotionAttributeId = { UE::Anim::IAnimRootMotionProvider::AttributeName , FCompactPoseBoneIndex(0) };
 
 class FModule : public IModuleInterface, public UE::Anim::IAnimRootMotionProvider
 {
diff --git a/Engine/Plugins/Experimental/Animation/AnimationWarping/Source/Runtime/Private/BoneControllers/AnimNode_StrideWarping.cpp b/Engine/Plugins/Experimental/Animation/AnimationWarping/Source/Runtime/Private/BoneControllers/AnimNode_StrideWarping.cpp
--- a/Engine/Plugins/Experimental/Animation/AnimationWarping/Source/Runtime/Private/BoneControllers/AnimNode_StrideWarping.cpp
+++ b/Engine/Plugins/Experimental/Animation/AnimationWarping/Source/Runtime/Private/BoneControllers/AnimNode_StrideWarping.cpp
@@ -19,6 +19,7 @@ FAnimNode_StrideWarping::FAnimNode_StrideWarping()
 	, ManualStrideWarpingDir(ForceInitToZero)
 	, ManualStrideScaling(1.f)
 	, LocomotionSpeed(0.f)
+	, MinSpeedTolerance(10.f)
 	, StrideWarpingAxisMode(EStrideWarpingAxisMode::ActorSpaceVectorInput)
 	, FloorNormalAxisMode(EStrideWarpingAxisMode::IKFootRootLocalZ)
 	, GravityDirAxisMode(EStrideWarpingAxisMode::ComponentSpaceVectorInput)
diff --git a/Engine/Plugins/Experimental/Animation/AnimationWarping/Source/Runtime/Public/BoneControllers/AnimNode_StrideWarping.h b/Engine/Plugins/Experimental/Animation/AnimationWarping/Source/Runtime/Public/BoneControllers/AnimNode_StrideWarping.h
--- a/Engine/Plugins/Experimental/Animation/AnimationWarping/Source/Runtime/Public/BoneControllers/AnimNode_StrideWarping.h
+++ b/Engine/Plugins/Experimental/Animation/AnimationWarping/Source/Runtime/Public/BoneControllers/AnimNode_StrideWarping.h
@@ -79,6 +79,9 @@ struct ANIMATIONWARPINGRUNTIME_API FAnimNode_StrideWarping : public FAnimNode_Sk
 	UPROPERTY(EditAnywhere, BlueprintReadWrite, Category = Evaluation, meta = (PinHiddenByDefault))
 	float LocomotionSpeed;
 
+	UPROPERTY(EditAnywhere, BlueprintReadWrite, Category = Evaluation, meta = (PinHiddenByDefault))
+	float MinSpeedTolerance;
+
 	UPROPERTY(EditAnywhere, Category = Settings)
 	FBoneReference IKFootRootBone;
 
@@ -133,9 +136,6 @@ struct ANIMATIONWARPINGRUNTIME_API FAnimNode_StrideWarping : public FAnimNode_Sk
 	UPROPERTY(Transient)
 	float CachedDeltaTime;
 
-	UPROPERTY(EditAnywhere, BlueprintReadWrite, Category = Settings, meta = (PinHiddenByDefault, EditCondition = "Mode == EWarpingEvaluationMode::Graph"))
-	float MinSpeedTolerance = 10.f;
-
 	FVector GetAxisModeValue(const EStrideWarpingAxisMode& AxisMode, const FTransform& IKFootRootCSTransform, const FVector& UserSuppliedVector) const;
 	FCompactPoseBoneIndex FindHipBoneIndex(const FCompactPoseBoneIndex& InFootBoneIndex, const int32& NumBonesInLimb, const FBoneContainer& RequiredBones) const;
 
diff --git a/Engine/Plugins/Experimental/Animation/PoseSearch/Source/Editor/Private/AnimGraphNode_MotionMatching.cpp b/Engine/Plugins/Experimental/Animation/PoseSearch/Source/Editor/Private/AnimGraphNode_MotionMatching.cpp
--- a/Engine/Plugins/Experimental/Animation/PoseSearch/Source/Editor/Private/AnimGraphNode_MotionMatching.cpp
+++ b/Engine/Plugins/Experimental/Animation/PoseSearch/Source/Editor/Private/AnimGraphNode_MotionMatching.cpp
@@ -30,7 +30,7 @@ void UAnimGraphNode_MotionMatching::GetOutputLinkAttributes(FNodeAttributeArray&
 {
 	if (UE::Anim::IAnimRootMotionProvider::Get())
 	{
-		OutAttributes.Add(UE::Anim::IAnimRootMotionProvider::RootMotionDeltaAttributeName);
+		OutAttributes.Add(UE::Anim::IAnimRootMotionProvider::AttributeName);
 	}
 }
 
diff --git a/Engine/Plugins/Experimental/Animation/PoseSearch/Source/Runtime/Private/AnimNode_MotionMatching.cpp b/Engine/Plugins/Experimental/Animation/PoseSearch/Source/Runtime/Private/AnimNode_MotionMatching.cpp
--- a/Engine/Plugins/Experimental/Animation/PoseSearch/Source/Runtime/Private/AnimNode_MotionMatching.cpp
+++ b/Engine/Plugins/Experimental/Animation/PoseSearch/Source/Runtime/Private/AnimNode_MotionMatching.cpp
@@ -164,7 +164,7 @@ void FAnimNode_MotionMatching::UpdateAssetPlayer(const FAnimationUpdateContext&
 		ElapsedPoseJumpTime += DeltaTime;
 	}
 	 
-	const FPredictionSequenceState SequenceState = { SequencePlayerNode.GetSequence() , SequencePlayerNode.GetAccumulatedTime(), SequencePlayerNode.GetLoopAnimation() };
+	const FPredictionSequenceState SequenceState = { SequencePlayerNode.GetSequence() , SequencePlayerNode.GetAccumulatedTime(), SequencePlayerNode.GetPlayRate(), SequencePlayerNode.GetLoopAnimation() };
 	const float PlayRate = UPoseSearchPredictionDistanceMatching::ComputePlayRate(Context, Prediction, PredictionSettings, SequenceState);
 	SequencePlayerNode.SetPlayRate(PlayRate);
 
@@ -253,6 +253,7 @@ void FAnimNode_MotionMatching::JumpToPose(const FAnimationUpdateContext& Context
 	SequencePlayerNode.SetSequence(ResultDbSequence.Sequence);
 	SequencePlayerNode.SetAccumulatedTime(Result.TimeOffsetSeconds);
 	SequencePlayerNode.SetLoopAnimation(ResultDbSequence.bLoopAnimation);
+	SequencePlayerNode.SetPlayRate(1.f);
 
 	// Use inertial blending to smooth over the transition
 	// It would be cool in the future to adjust the blend time by amount of dissimilarity, but we'll need a standardized distance metric first.
diff --git a/Engine/Plugins/Experimental/Animation/PoseSearch/Source/Runtime/Public/PoseSearch/PoseSearchPredictionLibrary.cpp b/Engine/Plugins/Experimental/Animation/PoseSearch/Source/Runtime/Public/PoseSearch/PoseSearchPredictionLibrary.cpp
--- a/Engine/Plugins/Experimental/Animation/PoseSearch/Source/Runtime/Public/PoseSearch/PoseSearchPredictionLibrary.cpp
+++ b/Engine/Plugins/Experimental/Animation/PoseSearch/Source/Runtime/Public/PoseSearch/PoseSearchPredictionLibrary.cpp
@@ -2,9 +2,9 @@
 
 #include "PoseSearchPredictionLibrary.h"
 #include "Animation/AnimInstanceProxy.h"
+#include "Animation/AnimSequence.h"
 #include "Algo/MinElement.h"
 
-
 #if ENABLE_ANIM_DEBUG
 TAutoConsoleVariable<int32> CVarAnimNodePredictionPlayRateDebug(TEXT("a.AnimNode.PredictionPlayRate.Debug"), 0, TEXT("Turn on debug for trajectory prediction play rate scaling"));
 TAutoConsoleVariable<int32> CVarAnimNodePredictionPlayRateEnable(TEXT("a.AnimNode.PredictionPlayRate.Enable"), 1, TEXT("Toggle trajectory prediction play rate scaling"));
@@ -32,7 +32,12 @@ float UPoseSearchPredictionDistanceMatching::ComputePlayRate(const FAnimationUpd
 	bool bDebugDraw = Settings.bDebugDraw;
 #endif
 #endif
-
+#if WITH_EDITORONLY_DATA
+	if (bDebugDraw)
+	{
+		Context.AnimInstanceProxy->AnimDrawDebugSphere(Context.AnimInstanceProxy->GetComponentTransform().GetLocation(), 8.f, 16, FColor::Green);
+	}
+#endif
 	// Delta time is not progressing
 	if (FMath::IsNearlyZero(DeltaTime, SMALL_NUMBER))
 	{
@@ -57,6 +62,8 @@ float UPoseSearchPredictionDistanceMatching::ComputePlayRate(const FAnimationUpd
 		return PlayRate;
 	}
 
+	PlayRate = SequenceState.PlayRate;
+
 	// Find the minima prediction trajectory velocity
 	// Approximate zeroed values may indicate the synchronization point for a stop or pivot
 	const FPredictionTrajectoryState* MinimaSample = Algo::MinElementBy(TrajectoryRange.Samples, [](const FPredictionTrajectoryState& Sample)
@@ -73,7 +80,7 @@ float UPoseSearchPredictionDistanceMatching::ComputePlayRate(const FAnimationUpd
 	float MinimaSampleTime = FLT_MAX;
 	bool bPivotDetected = false;
 
-	for (auto [SampleTime, RootMotionSampleStep, PreviousDirection, MinimaDisplacement, CosPivotAngleThreshold] = std::tuple{ (double)SequenceState.AccumulatedTime, 1.f / Settings.RootMotionSampleStepPerSecond, FVector::ZeroVector, FLT_MAX, FMath::Cos(Settings.RootMotionPivotAngleThreshold) };
+	for (auto [SampleTime, RootMotionSampleStep, PreviousDirection, MinimaDisplacement, CosPivotAngleThreshold] = std::tuple{ (double)SequenceState.AccumulatedTime, 1.f / Settings.RootMotionSampleStepPerSecond, FVector::ZeroVector, FLT_MAX, FMath::Cos(Settings.ZeroRootMotionAngleThreshold) };
 		SampleTime <= Sequence->GetPlayLength(); 
 		SampleTime += RootMotionSampleStep)
 	{
@@ -95,7 +102,7 @@ float UPoseSearchPredictionDistanceMatching::ComputePlayRate(const FAnimationUpd
 
 		// Significant changes in direction will be defined as a pivot
 		// Hack: Unfortunately in practice, we may erroneously identify animations as a pivot when in fact they have malformed root motion tracks
-		if (PotentialPivotAngle < CosPivotAngleThreshold && RootMotionDisplacement > Settings.RootMotionPivotDisplacementError)
+		if (PotentialPivotAngle < CosPivotAngleThreshold && RootMotionDisplacement > Settings.ZeroRootMotionDisplacementError)
 		{
 			// Bias the minima sample time for pivots to favor the pre-pivot phase (moment prior to direction change)
 			MinimaSampleTime = SampleTime - RootMotionSampleStep;
@@ -116,13 +123,13 @@ float UPoseSearchPredictionDistanceMatching::ComputePlayRate(const FAnimationUpd
 	// Extrapolate the minima forward in time to detect a potential complete loss in velocity
 	const FVector MinimaRootMotionDelta = Sequence->ExtractRootMotion(MinimaSampleTime, DeltaTime, SequenceState.bLooping).GetTranslation();
 
-	// Play rate scaling is minima driven when a near zero root motion delta or pivot has been detected
+	// Play rate scaling is root motion driven when a near zero root motion delta or pivot has been detected
 	// Otherwise we are locomotion driven, which is reflected in the numerator of the divisor
-	bool bMinimaDrivenPlayRate = MinimaRootMotionDelta.IsNearlyZero(KINDA_SMALL_NUMBER) || bPivotDetected;
+	bool bRootMotionDrivenPlayRate = MinimaRootMotionDelta.IsNearlyZero(KINDA_SMALL_NUMBER) || bPivotDetected;
 #if WITH_EDITORONLY_DATA
 	FColor SynchronizationColor = FColor::Red;
 #endif
-	FVector RemainingRootMotionDelta = FVector::ZeroVector;
+	FVector RootMotionDelta = FVector::ZeroVector;
 
 	// This loop allows for play rate scaling to apply correction in the event that we discover the animation and trajectory prediction minima mismatch
 	// Example: 
@@ -133,14 +140,14 @@ float UPoseSearchPredictionDistanceMatching::ComputePlayRate(const FAnimationUpd
 	{
 		// Minima driven play rate scaling synchronizes using the remaining displacement to near zero with: animation / locomotion
 		// Locomotion driven play rate scaling synchronizes using the per-frame instantaneous displacement with: locomotion / animation
-		const float RemainingSequenceDelta = bMinimaDrivenPlayRate ? MinimaSampleTime - SequenceState.AccumulatedTime : DeltaTime;
-		const float PredictionDisplacement = bMinimaDrivenPlayRate ? MinimaSample->AccumulatedDistance : TrajectoryRange.Samples[0].LocalLinearVelocity.Size2D() * DeltaTime;
+		const float SequenceDelta = bRootMotionDrivenPlayRate ? MinimaSampleTime - SequenceState.AccumulatedTime : DeltaTime;
+		const float PredictionDisplacement = bRootMotionDrivenPlayRate ? MinimaSample->AccumulatedDistance : TrajectoryRange.Samples[0].LocalLinearVelocity.Size2D() * DeltaTime;
 
-		RemainingRootMotionDelta = Sequence->ExtractRootMotion(SequenceState.AccumulatedTime, RemainingSequenceDelta, SequenceState.bLooping).GetTranslation();
-		const float RemainingRootMotionDeltaDisplacement = RemainingRootMotionDelta.Size2D();
+		RootMotionDelta = Sequence->ExtractRootMotion(SequenceState.AccumulatedTime, SequenceDelta, SequenceState.bLooping).GetTranslation();
+		const float RootMotionDeltaDisplacement = RootMotionDelta.Size2D();
 
 		// Zero displacement is left in the animation, which may result in sliding if the prediction minima has non-zero displacement
-		const bool bZeroRootMotion = FMath::IsNearlyZero(RemainingRootMotionDeltaDisplacement, KINDA_SMALL_NUMBER);
+		const bool bZeroRootMotion = FMath::IsNearlyZero(RootMotionDeltaDisplacement, KINDA_SMALL_NUMBER);
 
 		// Zero displacement is left in the prediction minima, which may result in a pop or a pose break if the animation has non-zero displacement
 		const bool bZeroPrediction = FMath::IsNearlyZero(PredictionDisplacement, KINDA_SMALL_NUMBER);
@@ -148,24 +155,33 @@ float UPoseSearchPredictionDistanceMatching::ComputePlayRate(const FAnimationUpd
 		// Play rate scaling isn't require since no trajectory prediction motion or root motion is present
 		if (bZeroRootMotion && bZeroPrediction)
 		{
+			PlayRate = FMath::Clamp(PlayRate, 0.f, 1.f);
 			break;
 		}
 		// The computed minima in the root motion and prediction mismatch, so flip the synchronization behavior as an attempt to correctly play rate scale
 		// If this fails, the algorithm will automatically fall back to a default play rate of 1.f, which may introduce sliding.
 		else if (bZeroPrediction || bZeroRootMotion)
 		{
-			bMinimaDrivenPlayRate = !bMinimaDrivenPlayRate;
+			// There are cases where root motion may be available in the absence of predicted motion. We need to guarantee that the play rate 
+			// will not superfluously reach extreme values due to bRootMotionDrivenPlayRate being enabled.
+			if (bRootMotionDrivenPlayRate && MinimaSample->IsZeroSample())
+			{
+				PlayRate = FMath::Clamp(PlayRate, 0.f, 1.f);
+				break;
+			}
+
+			bRootMotionDrivenPlayRate = !bRootMotionDrivenPlayRate;
 			continue;
 		}
 		// Play rate scaling succeeded
 		else
 		{
 #if WITH_EDITORONLY_DATA
-			SynchronizationColor = bMinimaDrivenPlayRate ? FColor::Purple : FColor::Blue;
+			SynchronizationColor = bRootMotionDrivenPlayRate ? FColor::Purple : FColor::Blue;
 #endif
-			PlayRate = bMinimaDrivenPlayRate
-				? RemainingRootMotionDeltaDisplacement / PredictionDisplacement 
-				: PredictionDisplacement / RemainingRootMotionDeltaDisplacement;
+			PlayRate = bRootMotionDrivenPlayRate
+				? RootMotionDeltaDisplacement / PredictionDisplacement 
+				: PredictionDisplacement / RootMotionDeltaDisplacement;
 
 			break;
 		}
@@ -175,10 +191,9 @@ float UPoseSearchPredictionDistanceMatching::ComputePlayRate(const FAnimationUpd
 	// Render the starting and ending trajectory prediction positions for the distance matching play rate synchronization
 	if (bDebugDraw)
 	{
-		Context.AnimInstanceProxy->AnimDrawDebugSphere(Context.AnimInstanceProxy->GetComponentTransform().GetLocation(), 8.f, 16, FColor::Green);
-		Context.AnimInstanceProxy->AnimDrawDebugSphere(Context.AnimInstanceProxy->GetComponentTransform().TransformPosition(RemainingRootMotionDelta), 8.f, 16, SynchronizationColor);
+		Context.AnimInstanceProxy->AnimDrawDebugSphere(Context.AnimInstanceProxy->GetComponentTransform().TransformPosition(RootMotionDelta), 8.f, 16, SynchronizationColor);
 
-		if (bMinimaDrivenPlayRate)
+		if (bRootMotionDrivenPlayRate)
 		{
 			Context.AnimInstanceProxy->AnimDrawDebugSphere(Context.AnimInstanceProxy->GetComponentTransform().TransformPosition(MinimaSample->Position), 8.f, 16, FColor::Yellow);
 		}
diff --git a/Engine/Plugins/Experimental/Animation/PoseSearch/Source/Runtime/Public/PoseSearch/PoseSearchPredictionTypes.cpp b/Engine/Plugins/Experimental/Animation/PoseSearch/Source/Runtime/Public/PoseSearch/PoseSearchPredictionTypes.cpp
--- a/Engine/Plugins/Experimental/Animation/PoseSearch/Source/Runtime/Public/PoseSearch/PoseSearchPredictionTypes.cpp
+++ b/Engine/Plugins/Experimental/Animation/PoseSearch/Source/Runtime/Public/PoseSearch/PoseSearchPredictionTypes.cpp
@@ -1,6 +1,7 @@
 // Copyright Epic Games, Inc. All Rights Reserved.
 
 #include "PoseSearchPredictionTypes.h"
+#include "Animation/AnimSequence.h"
 #include "Algo/AllOf.h"
 
 
@@ -11,6 +12,11 @@ float FPredictionPlayRateAdjustment::ComputePlayRate(float PlayRate, float Delta
 	return PlayRate;
 }
 
+bool FPredictionTrajectoryState::IsZeroSample() const
+{
+	return LocalLinearVelocity == FVector::ZeroVector && LocalLinearAcceleration == FVector::ZeroVector && Position == FVector::ZeroVector && AccumulatedDistance == 0.f;
+}
+
 FPredictionTrajectoryState FPredictionTrajectoryState::Lerp(const FPredictionTrajectoryState& A, const FPredictionTrajectoryState& B, float Alpha)
 {
 	FPredictionTrajectoryState State;
@@ -28,12 +34,10 @@ bool FPredictionTrajectoryRange::HasSamples() const
 
 bool FPredictionTrajectoryRange::HasOnlyZeroSamples() const
 {
-	return Algo::AllOf(Samples, [](const FPredictionTrajectoryState& Sample) {
-		return Sample.LocalLinearVelocity == FVector::ZeroVector
-			&& Sample.LocalLinearAcceleration == FVector::ZeroVector
-			&& Sample.Position == FVector::ZeroVector
-			&& Sample.AccumulatedDistance == 0.f;
-		});
+	return Algo::AllOf(Samples, [](const FPredictionTrajectoryState& Sample) 
+	{
+		return Sample.IsZeroSample();
+	});
 }
 
 bool FPredictionSequenceState::HasSequence() const
diff --git a/Engine/Plugins/Experimental/Animation/PoseSearch/Source/Runtime/Public/PoseSearch/PoseSearchPredictionTypes.h b/Engine/Plugins/Experimental/Animation/PoseSearch/Source/Runtime/Public/PoseSearch/PoseSearchPredictionTypes.h
--- a/Engine/Plugins/Experimental/Animation/PoseSearch/Source/Runtime/Public/PoseSearch/PoseSearchPredictionTypes.h
+++ b/Engine/Plugins/Experimental/Animation/PoseSearch/Source/Runtime/Public/PoseSearch/PoseSearchPredictionTypes.h
@@ -43,6 +43,8 @@ struct POSESEARCH_API FPredictionTrajectoryState
 	UPROPERTY(BlueprintReadOnly, Category = Prediction)
 	float AccumulatedDistance = 0.0;
 
+	bool IsZeroSample() const;
+
 	static FPredictionTrajectoryState Lerp(const FPredictionTrajectoryState& A, const FPredictionTrajectoryState& B, float Alpha);
 };
 
@@ -76,12 +78,12 @@ struct POSESEARCH_API FPredictionTrajectorySettings
 
 	// Root motion angle threshold for determining significant changes in direction (pivots)
 	UPROPERTY(EditAnywhere, BlueprintReadWrite, Category = Settings, meta = (PinHiddenByDefault, ClampMin = "0", ClampMax = "180"))
-	float RootMotionPivotAngleThreshold = 90.f;
+	float ZeroRootMotionAngleThreshold = 90.f;
 
 	// Root motion displacement error tolerance for identifying pivot false positives
 	UPROPERTY(EditAnywhere, BlueprintReadWrite, Category = Settings, meta = (PinHiddenByDefault))
-	float RootMotionPivotDisplacementError = 0.0004f;
-	
+	float ZeroRootMotionDisplacementError = 0.0004f;
+
 #if WITH_EDITORONLY_DATA
 	UPROPERTY(EditAnywhere, Category = Debug)
 	bool bDebugDraw = false;
@@ -101,6 +103,11 @@ struct POSESEARCH_API FPredictionSequenceState
 	UPROPERTY(EditAnywhere, BlueprintReadWrite, Category = State, meta = (PinShownByDefault))
 	float AccumulatedTime = 0.f;
 
+	// Internal sequence play rate
+	UPROPERTY(EditAnywhere, BlueprintReadWrite, Category = State, meta = (PinShownByDefault))
+	float PlayRate = 1.f;
+
+
 	// Looping or non-looping sequence
 	UPROPERTY(EditAnywhere, BlueprintReadWrite, Category = State, meta = (PinShownByDefault))
 	bool bLooping = false;
diff --git a/Engine/Source/Editor/AnimGraph/Private/AnimGraphModule.cpp b/Engine/Source/Editor/AnimGraph/Private/AnimGraphModule.cpp
--- a/Engine/Source/Editor/AnimGraph/Private/AnimGraphModule.cpp
+++ b/Engine/Source/Editor/AnimGraph/Private/AnimGraphModule.cpp
@@ -128,7 +128,7 @@ void FAnimGraphModule::StartupModule()
 
 		AnimGraphAttributes->Register(
 			FAnimGraphAttributeDesc(
-				UE::Anim::IAnimRootMotionProvider::RootMotionDeltaAttributeName,
+				UE::Anim::IAnimRootMotionProvider::AttributeName,
 				EAnimGraphAttributeBlend::Blendable,
 				*FEditorStyle::Get().GetBrush("AnimGraph.Attribute.RootMotionDelta.Icon"),
 				LOCTEXT("RootMotionDelta", "Root Motion Delta"),
diff --git a/Engine/Source/Editor/AnimGraph/Private/AnimGraphNode_BlendSpaceBase.cpp b/Engine/Source/Editor/AnimGraph/Private/AnimGraphNode_BlendSpaceBase.cpp
--- a/Engine/Source/Editor/AnimGraph/Private/AnimGraphNode_BlendSpaceBase.cpp
+++ b/Engine/Source/Editor/AnimGraph/Private/AnimGraphNode_BlendSpaceBase.cpp
@@ -91,7 +91,7 @@ void UAnimGraphNode_BlendSpaceBase::GetOutputLinkAttributes(FNodeAttributeArray&
 
 	if (UE::Anim::IAnimRootMotionProvider::Get())
 	{
-		OutAttributes.Add(UE::Anim::IAnimRootMotionProvider::RootMotionDeltaAttributeName);
+		OutAttributes.Add(UE::Anim::IAnimRootMotionProvider::AttributeName);
 	}
 }
 
diff --git a/Engine/Source/Editor/AnimGraph/Private/AnimGraphNode_RandomPlayer.cpp b/Engine/Source/Editor/AnimGraph/Private/AnimGraphNode_RandomPlayer.cpp
--- a/Engine/Source/Editor/AnimGraph/Private/AnimGraphNode_RandomPlayer.cpp
+++ b/Engine/Source/Editor/AnimGraph/Private/AnimGraphNode_RandomPlayer.cpp
@@ -35,7 +35,7 @@ void UAnimGraphNode_RandomPlayer::GetOutputLinkAttributes(FNodeAttributeArray& O
 
 	if (UE::Anim::IAnimRootMotionProvider::Get())
 	{
-		OutAttributes.Add(UE::Anim::IAnimRootMotionProvider::RootMotionDeltaAttributeName);
+		OutAttributes.Add(UE::Anim::IAnimRootMotionProvider::AttributeName);
 	}
 }
 
diff --git a/Engine/Source/Editor/AnimGraph/Private/AnimGraphNode_SequencePlayer.cpp b/Engine/Source/Editor/AnimGraph/Private/AnimGraphNode_SequencePlayer.cpp
--- a/Engine/Source/Editor/AnimGraph/Private/AnimGraphNode_SequencePlayer.cpp
+++ b/Engine/Source/Editor/AnimGraph/Private/AnimGraphNode_SequencePlayer.cpp
@@ -154,7 +154,7 @@ void UAnimGraphNode_SequencePlayer::GetOutputLinkAttributes(FNodeAttributeArray&
 
 	if (UE::Anim::IAnimRootMotionProvider::Get())
 	{
-		OutAttributes.Add(UE::Anim::IAnimRootMotionProvider::RootMotionDeltaAttributeName);
+		OutAttributes.Add(UE::Anim::IAnimRootMotionProvider::AttributeName);
 	}
 }
 
@@ -217,7 +217,6 @@ void UAnimGraphNode_SequencePlayer::ReplaceReferredAnimations(const TMap<UAnimat
 	HandleAnimReferenceReplacement(Node.Sequence, AnimAssetReplacementMap);
 }
 
-
 bool UAnimGraphNode_SequencePlayer::DoesSupportTimeForTransitionGetter() const
 {
 	return true;
diff --git a/Engine/Source/Editor/EditorStyle/Private/StarshipStyle.cpp b/Engine/Source/Editor/EditorStyle/Private/StarshipStyle.cpp
--- a/Engine/Source/Editor/EditorStyle/Private/StarshipStyle.cpp
+++ b/Engine/Source/Editor/EditorStyle/Private/StarshipStyle.cpp
@@ -5743,7 +5743,7 @@ void FStarshipEditorStyle::FStyle::SetupPersonaStyle()
 		Set("AnimGraph.Attribute.Attributes.Color", FStyleColors::AccentYellow);
 		Set("AnimGraph.Attribute.Sync.Icon", new IMAGE_BRUSH_SVG("Starship/AnimationBlueprintEditor/AnimationGraphSync", Icon16x16));
 		Set("AnimGraph.Attribute.Sync.Color", FStyleColors::AccentBlue);
-		Set("AnimGraph.Attribute.RootMotionDelta.Icon", new IMAGE_BRUSH("Icons/icon_axis_local_16px", Icon16x16));
+		Set("AnimGraph.Attribute.RootMotionDelta.Icon", new IMAGE_BRUSH_SVG("Starship/AnimationBlueprintEditor/AnimationGraphAttributes", Icon16x16));
 		Set("AnimGraph.Attribute.RootMotionDelta.Color", FStyleColors::AccentGreen);
 		Set("AnimGraph.Attribute.InertialBlending.Icon", new IMAGE_BRUSH_SVG("Starship/AnimationBlueprintEditor/AnimationGraphInertialBlending", Icon16x16));
 		Set("AnimGraph.Attribute.InertialBlending.Color", FStyleColors::AccentOrange);
diff --git a/Engine/Source/Runtime/Engine/Private/Animation/AnimRootMotionProvider.cpp b/Engine/Source/Runtime/Engine/Private/Animation/AnimRootMotionProvider.cpp
--- a/Engine/Source/Runtime/Engine/Private/Animation/AnimRootMotionProvider.cpp
+++ b/Engine/Source/Runtime/Engine/Private/Animation/AnimRootMotionProvider.cpp
@@ -8,7 +8,7 @@ namespace UE { namespace Anim {
 
 const FName IAnimRootMotionProvider::ModularFeatureName(TEXT("AnimationWarping"));
 
-const FName IAnimRootMotionProvider::RootMotionDeltaAttributeName(TEXT("RootMotionDelta"));
+const FName IAnimRootMotionProvider::AttributeName(TEXT("RootMotionDelta"));
 
 const IAnimRootMotionProvider* IAnimRootMotionProvider::Get()
 {
diff --git a/Engine/Source/Runtime/Engine/Public/Animation/AnimRootMotionProvider.h b/Engine/Source/Runtime/Engine/Public/Animation/AnimRootMotionProvider.h
--- a/Engine/Source/Runtime/Engine/Public/Animation/AnimRootMotionProvider.h
+++ b/Engine/Source/Runtime/Engine/Public/Animation/AnimRootMotionProvider.h
@@ -20,7 +20,7 @@ class ENGINE_API IAnimRootMotionProvider : public IModularFeature
 {
 public:
 	static const FName ModularFeatureName; // "AnimationWarping"
-	static const FName RootMotionDeltaAttributeName;
+	static const FName AttributeName;
 
 	virtual ~IAnimRootMotionProvider() {}
 
