PoseSearch Debugger integration into the Rewind Debugger

- Includes PoseSearch Trace logger/analyzer/provider to record Motion Matching data
- Currently supports debugging multiple Motion-Matching nodes, including both database and state introspection

#rnx

diff --git a/Engine/Plugins/Experimental/Animation/PoseSearch/Source/Editor/PoseSearchEditor.Build.cs b/Engine/Plugins/Experimental/Animation/PoseSearch/Source/Editor/PoseSearchEditor.Build.cs
--- a/Engine/Plugins/Experimental/Animation/PoseSearch/Source/Editor/PoseSearchEditor.Build.cs
+++ b/Engine/Plugins/Experimental/Animation/PoseSearch/Source/Editor/PoseSearchEditor.Build.cs
@@ -17,16 +17,28 @@ public class PoseSearchEditor : ModuleRules
 				"CoreUObject",
 				"Engine",
 				"PoseSearch",
-				"ToolMenus",
-				"EditorFramework",
+				
+				// Trace-related dependencies
+				"TraceLog",
+				"TraceAnalysis",
+				"TraceServices",
+				"TraceInsights",
+				
+				// UI 
+				"PropertyEditor",
+				"Slate",
+				"SlateCore",
+				"RewindDebuggerInterface",
+				"EditorStyle",
 				"UnrealEd",
+				"InputCore"
 			}
 		);
 
 		PrivateDependencyModuleNames.AddRange(
 			new string[]
 			{
-				"BlueprintGraph",
+				"BlueprintGraph"
 			}
 		);
 	}
diff --git a/Engine/Plugins/Experimental/Animation/PoseSearch/Source/Editor/Private/PoseSearchDebugger.cpp b/Engine/Plugins/Experimental/Animation/PoseSearch/Source/Editor/Private/PoseSearchDebugger.cpp
new file mode 100644
--- /dev/null
+++ b/Engine/Plugins/Experimental/Animation/PoseSearch/Source/Editor/Private/PoseSearchDebugger.cpp
@@ -0,0 +1,1340 @@
+﻿// Copyright Epic Games, Inc. All Rights Reserved.
+
+#include "PoseSearchDebugger.h"
+#include "PoseSearch/PoseSearch.h"
+#include "Animation/AnimSequence.h"
+#include "RewindDebuggerInterface/Public/IRewindDebugger.h"
+#include "ObjectTrace.h"
+#include "TraceServices/Model/AnalysisSession.h"
+#include "Trace/PoseSearchTraceProvider.h"
+#include "Widgets/Layout/SScrollBox.h"
+#include "Widgets/Layout/SWidgetSwitcher.h"
+#include "Widgets/SBoxPanel.h"
+#include "Widgets/Input/SButton.h"
+#include "Widgets/Images/SImage.h"
+#include "Styling/SlateIconFinder.h"
+
+#define LOCTEXT_NAMESPACE "PoseSearchDebugger"
+
+
+UPoseSearchDebuggerReflection::UPoseSearchDebuggerReflection(const FObjectInitializer& ObjectInitializer) : Super(ObjectInitializer)
+{
+}
+
+void FPoseSearchDebuggerFeatureReflection::EmptyAll()
+{
+	Positions.Empty();
+	LinearVelocities.Empty();
+	AngularVelocities.Empty();
+}
+
+namespace UE { namespace PoseSearch {
+
+class FDebuggerDatabaseRowData : public TSharedFromThis<FDebuggerDatabaseRowData>
+{
+public:
+	FDebuggerDatabaseRowData() = default;
+	FDebuggerDatabaseRowData(
+	    uint32 InPoseIdx, 
+	    const FString& InAnimSequenceName,
+	    float InTime,
+	    float InLength,
+	    float InScore,
+	    float InPoseScore,
+	    float InTrajectoryScore
+	)
+		: PoseIdx(InPoseIdx)
+		, AnimSequenceName(InAnimSequenceName)
+		, Time(InTime)
+		, Length(InLength)
+		, Score(InScore)
+		, PoseScore(InPoseScore)
+		, TrajectoryScore(InTrajectoryScore)
+	{
+	}
+	
+	uint32 PoseIdx = 0;
+	FString AnimSequenceName = "";
+	float Time = 0.0f;
+	float Length = 0.0f;
+	float Score = 0.0f;
+	float PoseScore = 0.0f;
+	float TrajectoryScore = 0.0f;
+};
+
+namespace DebuggerDatabaseColumns
+{
+	/** Column struct to represent each column in the debugger database */
+	struct IColumn : TSharedFromThis<IColumn>
+	{
+		IColumn(int32 InSortIndex, bool InEnabled = true)
+			: SortIndex(InSortIndex)
+			, bEnabled(InEnabled)
+		{
+		}
+		
+		virtual ~IColumn() = default;
+		
+		/** Sorted left to right based on this index */
+		int32 SortIndex = 0;
+		/** Current width, starts at 1 to be evenly spaced between all columns */
+		float Width = 1.0f;
+		/** Disabled selectively with view options */
+		bool bEnabled = false;
+
+		virtual FName GetName() const = 0;
+		
+		using FRowDataRef = TSharedRef<FDebuggerDatabaseRowData>;
+		using FSortPredicate = TFunctionRef<bool(const FRowDataRef&, const FRowDataRef&)>;
+		
+		/** Sort predicate to sort list in ascending order by this column */
+		virtual FSortPredicate GetSortPredicateAscending() = 0;
+		/** Sort predicate to sort list in descending order by this column */
+		virtual FSortPredicate GetSortPredicateDescending() = 0;
+		/** Text to display associated with this column from given row data  */
+		virtual TSharedRef<SWidget> GenerateTextWidget(const FRowDataRef& RowData) const { return SNullWidget::NullWidget; }
+
+	protected:
+		TSharedRef<STextBlock> GenerateTextWidgetBase() const
+		{
+			static FSlateFontInfo RowFont = FEditorStyle::Get().GetFontStyle("DetailsView.CategoryTextStyle");
+			
+			return SNew(STextBlock)
+            	.Font(RowFont)
+            	.Justification(ETextJustify::Center);
+		}
+	};
+
+	struct FPoseIdx : IColumn
+	{
+		using IColumn::IColumn;
+		static const FName Name;
+		virtual FName GetName() const override { return Name; }
+
+		virtual FSortPredicate GetSortPredicateAscending() override
+		{
+			static FSortPredicate Predicate = [](const FRowDataRef& Row0, const FRowDataRef& Row1) -> bool { return Row0->PoseIdx < Row1->PoseIdx; };
+			return Predicate;
+		}
+		
+		virtual FSortPredicate GetSortPredicateDescending() override
+		{
+			static FSortPredicate Predicate = [](const FRowDataRef& Row0, const FRowDataRef& Row1) -> bool { return Row0->PoseIdx >= Row1->PoseIdx; };
+			return Predicate;
+		}
+
+		virtual TSharedRef<SWidget> GenerateTextWidget(const FRowDataRef& RowData) const override
+		{
+			TSharedRef<STextBlock> Text = GenerateTextWidgetBase();
+			Text->SetText(TAttribute<FText>::Create([RowData](){ return FText::AsNumber(RowData->PoseIdx); }));
+			return Text;
+		}
+	};
+	const FName FPoseIdx::Name = "PoseIdx";
+
+	struct FAnimSequenceName : IColumn
+	{
+		using IColumn::IColumn;
+		static const FName Name;
+		virtual FName GetName() const override { return Name; }
+		
+		virtual FSortPredicate GetSortPredicateAscending() override
+		{
+			static FSortPredicate Predicate = [](const FRowDataRef& Row0, const FRowDataRef& Row1) -> bool { return Row0->AnimSequenceName < Row1->AnimSequenceName; };
+			return Predicate;
+		}
+
+		virtual FSortPredicate GetSortPredicateDescending() override
+		{
+			static FSortPredicate Predicate = [](const FRowDataRef& Row0, const FRowDataRef& Row1) -> bool { return Row0->AnimSequenceName >= Row1->AnimSequenceName; };
+			return Predicate;
+		}
+
+		virtual TSharedRef<SWidget> GenerateTextWidget(const FRowDataRef& RowData) const override
+		{
+			TSharedRef<STextBlock> Text = GenerateTextWidgetBase();
+			Text->SetText(TAttribute<FText>::Create([RowData](){ return FText::FromString(RowData->AnimSequenceName); }));
+			return Text;
+		}
+	};
+	const FName FAnimSequenceName::Name = "AnimSequence";
+
+	struct FTime : IColumn
+	{
+		using IColumn::IColumn;
+		static const FName Name;
+		virtual FName GetName() const override { return Name; }
+		
+		virtual FSortPredicate GetSortPredicateAscending() override
+		{
+			static FSortPredicate Predicate = [](const FRowDataRef& Row0, const FRowDataRef& Row1) -> bool { return Row0->Time < Row1->Time; };
+			return Predicate;
+		}
+		
+		virtual FSortPredicate GetSortPredicateDescending() override
+		{
+			static FSortPredicate Predicate = [](const FRowDataRef& Row0, const FRowDataRef& Row1) -> bool { return Row0->Time >= Row1->Time; };
+			return Predicate;
+		}
+
+		virtual TSharedRef<SWidget> GenerateTextWidget(const FRowDataRef& RowData) const override
+		{
+			TSharedRef<STextBlock> Text = GenerateTextWidgetBase();
+			Text->SetText(TAttribute<FText>::Create([RowData](){ return FText::AsNumber(RowData->Time); }));
+			return Text;
+		}
+	};
+	const FName FTime::Name = "Time";
+
+	struct FLength : IColumn
+	{
+		using IColumn::IColumn;
+		static const FName Name;
+		virtual FName GetName() const override { return Name; }
+		
+		virtual FSortPredicate GetSortPredicateAscending() override
+		{
+			static FSortPredicate Predicate = [](const FRowDataRef& Row0, const FRowDataRef& Row1) -> bool { return Row0->Length < Row1->Length; };
+			return Predicate;
+		}
+
+		virtual FSortPredicate GetSortPredicateDescending() override
+		{
+			static FSortPredicate Predicate = [](const FRowDataRef& Row0, const FRowDataRef& Row1) -> bool { return Row0->Length >= Row1->Length; };
+			return Predicate;
+		}
+
+		virtual TSharedRef<SWidget> GenerateTextWidget(const FRowDataRef& RowData) const override
+		{
+			TSharedRef<STextBlock> Text = GenerateTextWidgetBase();
+			Text->SetText(TAttribute<FText>::Create([RowData](){ return FText::AsNumber(RowData->Length); }));
+			return Text;
+		}
+	};
+	const FName FLength::Name = "Length";
+
+	struct FScore : IColumn
+	{
+		using IColumn::IColumn;
+		static const FName Name;
+		virtual FName GetName() const override { return Name; }
+		
+		virtual FSortPredicate GetSortPredicateAscending() override
+		{
+			static FSortPredicate Predicate = [](const FRowDataRef& Row0, const FRowDataRef& Row1) -> bool { return Row0->Score < Row1->Score; };
+			return Predicate;
+		}
+
+		virtual FSortPredicate GetSortPredicateDescending() override
+		{
+			static FSortPredicate Predicate = [](const FRowDataRef& Row0, const FRowDataRef& Row1) -> bool { return Row0->Score >= Row1->Score; };
+			return Predicate;
+		}
+		
+		virtual TSharedRef<SWidget> GenerateTextWidget(const FRowDataRef& RowData) const override
+		{
+			TSharedRef<STextBlock> Text = GenerateTextWidgetBase();
+			Text->SetText(TAttribute<FText>::Create([RowData](){ return FText::AsNumber(RowData->Score); }));
+			return Text;
+		}
+	};
+	const FName FScore::Name = "Score";
+
+
+	struct FPoseScore : IColumn
+	{
+		using IColumn::IColumn;
+		static const FName Name;
+		virtual FName GetName() const override { return Name; }
+		
+		virtual FSortPredicate GetSortPredicateAscending() override
+		{
+			static FSortPredicate Predicate = [](const FRowDataRef& Row0, const FRowDataRef& Row1) -> bool { return Row0->PoseScore < Row1->PoseScore; };
+			return Predicate;
+		}
+
+		virtual FSortPredicate GetSortPredicateDescending() override
+		{
+			static FSortPredicate Predicate = [](const FRowDataRef& Row0, const FRowDataRef& Row1) -> bool { return Row0->PoseScore >= Row1->PoseScore; };
+			return Predicate;
+		}
+		
+		virtual TSharedRef<SWidget> GenerateTextWidget(const FRowDataRef& RowData) const override
+		{
+			TSharedRef<STextBlock> Text = GenerateTextWidgetBase();
+			Text->SetText(TAttribute<FText>::Create([RowData](){ return FText::AsNumber(RowData->PoseScore); }));
+			return Text;
+		}
+	};
+	const FName FPoseScore::Name = "Pose Score";
+
+
+	struct FTrajectoryScore : IColumn
+	{
+		using IColumn::IColumn;
+		static const FName Name;
+		virtual FName GetName() const override { return Name; }
+		
+		virtual FSortPredicate GetSortPredicateAscending() override
+		{
+			static FSortPredicate Predicate = [](const FRowDataRef& Row0, const FRowDataRef& Row1) -> bool { return Row0->TrajectoryScore < Row1->TrajectoryScore; };
+			return Predicate;
+		}
+
+		virtual FSortPredicate GetSortPredicateDescending() override
+		{
+			static FSortPredicate Predicate = [](const FRowDataRef& Row0, const FRowDataRef& Row1) -> bool { return Row0->TrajectoryScore >= Row1->TrajectoryScore; };
+			return Predicate;
+		}
+		
+		virtual TSharedRef<SWidget> GenerateTextWidget(const FRowDataRef& RowData) const override
+		{
+			TSharedRef<STextBlock> Text = GenerateTextWidgetBase();
+			Text->SetText(TAttribute<FText>::Create([RowData](){ return FText::AsNumber(RowData->TrajectoryScore); }));
+			return Text;
+		}
+	};
+	const FName FTrajectoryScore::Name = "Trajectory Score";
+}
+
+
+DECLARE_DELEGATE_RetVal_OneParam(TSharedRef<DebuggerDatabaseColumns::IColumn>&, FGetColumnDelegate, FName);
+
+/**
+ * Widget representing a single row of the database view
+ */
+class SDebuggerDatabaseRow : public SMultiColumnTableRow<TSharedRef<FDebuggerDatabaseRowData>>
+{
+	SLATE_BEGIN_ARGS(SDebuggerDatabaseRow) {}
+	SLATE_END_ARGS()
+
+	void Construct(
+		const FArguments& InArgs,
+		const TSharedRef<STableViewBase>& InOwnerTable,
+		TSharedRef<FDebuggerDatabaseRowData> InRow,
+		const FTableRowStyle& InRowStyle,
+		const FSlateBrush* InRowBrush,
+		FMargin InPaddingMargin
+	)
+	{
+		Row = InRow;
+		RowBrush = InRowBrush;
+
+		SMultiColumnTableRow<TSharedRef<FDebuggerDatabaseRowData>>::Construct(
+			FSuperRowType::FArguments()
+			.Padding(InPaddingMargin)
+			.Style(&InRowStyle),
+			InOwnerTable
+		);
+	}
+
+	virtual TSharedRef<SWidget> GenerateWidgetForColumn(const FName& InColumnName) override
+	{
+		check(GetColumn.IsBound());
+		check(RowBrush);
+		
+		// Get column
+		TSharedRef<DebuggerDatabaseColumns::IColumn>& Column = GetColumn.Execute(InColumnName);
+		
+		static FSlateFontInfo NormalFont = FEditorStyle::Get().GetFontStyle("DetailsView.CategoryTextStyle");
+		const TSharedRef<SWidget> Widget = Column->GenerateTextWidget(Row.ToSharedRef());
+		
+		return
+			SNew(SBorder)
+			.HAlign(HAlign_Center)
+			.VAlign(VAlign_Center)
+			.BorderImage(RowBrush)
+			.Padding(0.0f)
+			[
+				SNew(SHorizontalBox)
+				+ SHorizontalBox::Slot()
+				.VAlign(VAlign_Center)
+				.HAlign(HAlign_Center)
+				.Padding(0.0f, 3.0f)
+				[
+					Widget
+				]
+			];
+	}
+
+	/** Row data associated with this widget */
+	TSharedPtr<FDebuggerDatabaseRowData> Row;
+
+	/** Used for cell styles (active vs database row) */
+	const FSlateBrush* RowBrush = nullptr;
+
+	/** Used to grab the column struct given a column name */
+	static FGetColumnDelegate GetColumn;
+};
+FGetColumnDelegate SDebuggerDatabaseRow::GetColumn;
+
+void SDebuggerDatabaseView::Update(const FTraceMotionMatchingStateMessage* State, const UPoseSearchDatabase* Database)
+{
+	PopulateRows(State, Database);
+}
+
+void SDebuggerDatabaseView::RefreshColumns()
+{
+	using namespace DebuggerDatabaseColumns;
+	
+	ActiveView.HeaderRow->ClearColumns();
+	DatabaseView.HeaderRow->ClearColumns();
+	
+	// Sort columns by index
+	Columns.ValueSort([](const TSharedRef<IColumn> Column0, const TSharedRef<IColumn> Column1)
+	{
+		return Column0->SortIndex < Column1->SortIndex;
+	});
+
+	// Add columns from map to header row
+	for(TPair<FName, TSharedRef<IColumn>>& ColumnPair : Columns)
+	{
+		IColumn& Column = ColumnPair.Value.Get();
+		const FName& ColumnName = Column.GetName();
+		if(ColumnPair.Value->bEnabled)
+		{
+			SHeaderRow::FColumn::FArguments ColumnArgs = SHeaderRow::FColumn::FArguments()
+				.ColumnId(ColumnName)
+				.DefaultLabel(FText::FromName(ColumnName))
+				.SortMode(this, &SDebuggerDatabaseView::GetColumnSortMode, ColumnName)
+				.OnSort(this, &SDebuggerDatabaseView::OnColumnSortModeChanged)
+				.FillWidth(this, &SDebuggerDatabaseView::GetColumnWidth, ColumnName)
+				.VAlignCell(VAlign_Center)
+				.VAlignHeader(VAlign_Center)
+				.HAlignHeader(HAlign_Center)
+				.HAlignCell(HAlign_Fill);
+
+			DatabaseView.HeaderRow->AddColumn(ColumnArgs);
+			
+			// Every time the active column is changed, update the database column
+			ActiveView.HeaderRow->AddColumn(ColumnArgs.OnWidthChanged(this, &SDebuggerDatabaseView::OnColumnWidthChanged, ColumnName));
+		}
+	}
+}
+
+void SDebuggerDatabaseView::AddColumn(TSharedRef<DebuggerDatabaseColumns::IColumn>&& Column)
+{
+	Columns.Add(Column->GetName(), Column);
+}
+
+EColumnSortMode::Type SDebuggerDatabaseView::GetColumnSortMode(const FName ColumnId) const
+{
+	if (ColumnId == SortColumn)
+	{
+		return SortMode;
+	}
+	
+	return EColumnSortMode::None;
+}
+
+float SDebuggerDatabaseView::GetColumnWidth(const FName ColumnId) const
+{
+	check(Columns.Find(ColumnId));
+
+	return Columns[ColumnId]->Width;
+}
+
+
+void SDebuggerDatabaseView::OnColumnSortModeChanged(const EColumnSortPriority::Type SortPriority, const FName & ColumnId, const EColumnSortMode::Type InSortMode)
+{
+	check(Columns.Find(ColumnId));
+	SortColumn = ColumnId;
+	SortMode = InSortMode;
+	SortDatabaseRows();
+}
+
+void SDebuggerDatabaseView::OnColumnWidthChanged(const float NewWidth, FName ColumnId) const
+{
+	check(Columns.Find(ColumnId));
+	
+	Columns[ColumnId]->Width = NewWidth;
+}
+
+
+void SDebuggerDatabaseView::SortDatabaseRows()
+{
+	if (SortMode == EColumnSortMode::Ascending)
+	{
+		DatabaseView.Rows.Sort(Columns[SortColumn]->GetSortPredicateAscending());
+	}
+	else if (SortMode == EColumnSortMode::Descending)
+	{
+		DatabaseView.Rows.Sort(Columns[SortColumn]->GetSortPredicateDescending());
+	}
+
+	DatabaseView.ListView->RequestListRefresh();
+	DatabaseView.HeaderRow->RefreshColumns();
+}
+
+void SDebuggerDatabaseView::CreateRows(const UPoseSearchDatabase* Database)
+{
+	const int32 NumPoses = Database->SearchIndex.NumPoses;
+	DatabaseView.Rows.Reserve(NumPoses);
+
+	// Build database rows
+	for(const FPoseSearchDatabaseSequence& DbSequence : Database->Sequences)
+	{
+		const int32 LastPoseIdx = DbSequence.FirstPoseIdx + DbSequence.NumPoses;
+		for (int32 PoseIdx = DbSequence.FirstPoseIdx; PoseIdx != LastPoseIdx; ++PoseIdx)
+		{
+			TSharedRef<FDebuggerDatabaseRowData> Row = DatabaseView.Rows.Add_GetRef(MakeShared<FDebuggerDatabaseRowData>());
+			Row->PoseIdx = PoseIdx;
+
+			const FString SequenceName = DbSequence.Sequence->GetName();
+			const float SequenceLength = DbSequence.Sequence->GetPlayLength();
+			FFloatInterval Range = DbSequence.SamplingRange;
+
+			// @TODO: Update this when range is computed natively as part of the sequence class
+			const bool bSampleAll = (Range.Min == 0.0f) && (Range.Max == 0.0f);
+			const float SequencePlayLength = DbSequence.Sequence->GetPlayLength();
+			Range.Min = bSampleAll ? 0.0f : Range.Min;
+			Range.Max = bSampleAll ? SequencePlayLength : FMath::Min(SequencePlayLength, Range.Max);
+			// ---
+
+			Row->AnimSequenceName = SequenceName;
+			// Cap time in sequence to end of range
+			Row->Time = FGenericPlatformMath::Min(Range.Min + (PoseIdx - DbSequence.FirstPoseIdx) * Database->Schema->SamplingInterval, Range.Max);
+			Row->Length = SequenceLength;
+		}
+	}
+
+	ActiveView.Rows.Add(MakeShared<FDebuggerDatabaseRowData>());
+}
+
+void SDebuggerDatabaseView::PopulateRows(const FTraceMotionMatchingStateMessage* State, const UPoseSearchDatabase* Database)
+{
+	if (DatabaseView.Rows.IsEmpty())
+	{
+		check(ActiveView.Rows.IsEmpty());
+		CreateRows(Database);
+	}
+	check(ActiveView.Rows.Num() == 1);
+
+	// Active bias weights pulled from the MM node
+	FPoseSearchBiasWeights BiasWeights;
+    BiasWeights.Weights = State->BiasWeights;
+
+	const FPoseSearchFeatureVectorLayout& Layout = Database->Schema->Layout;
+
+	// Reverse engineer weight params from weights array
+	auto ExtractWeight = [](
+		const FPoseSearchFeatureVectorLayout& InLayout,
+		const FPoseSearchBiasWeights& InBiasWeights,
+		const EPoseSearchFeatureType InFeatureType,
+		const bool bTrajectory
+		) -> float
+		{
+			int32 FeatureIdx = INDEX_NONE;
+			if (InLayout.EnumerateFeature(InFeatureType, bTrajectory, FeatureIdx))
+			{
+				const FPoseSearchFeatureDesc& Feature = InLayout.Features[FeatureIdx];
+            	// Return first weight found associated with the feature
+				// as the same weight is applied to all features in the buffer (for now?)
+            	return InBiasWeights.Weights[Feature.ValueOffset];	
+			}
+			return 0.0f;
+		};
+
+	// @TODO: Compute alternate scores based on column visibility in view options
+	// Zeroed trajectory for pose score exclusively
+	FPoseSearchBiasWeightParams Params;
+	Params.PosePositionWeight = ExtractWeight(Layout, BiasWeights, EPoseSearchFeatureType::Position, false);
+	Params.PoseLinearVelocityWeight = ExtractWeight(Layout, BiasWeights, EPoseSearchFeatureType::LinearVelocity, false);
+	Params.TrajectoryPositionWeight = 0.0f;
+	Params.TrajectoryLinearVelocityWeight = 0.0f;
+	
+	FPoseSearchBiasWeights PoseBiasWeights;
+	PoseBiasWeights.Init(Params, Database->Schema->Layout);
+
+	// Zeroed pose for trajectory score
+	Params.PosePositionWeight = 0.0f;
+	Params.PoseLinearVelocityWeight = 0.0f;
+	Params.TrajectoryPositionWeight = ExtractWeight(Layout, BiasWeights, EPoseSearchFeatureType::Position, true);
+	Params.TrajectoryLinearVelocityWeight = ExtractWeight(Layout, BiasWeights, EPoseSearchFeatureType::LinearVelocity, true);
+	
+	FPoseSearchBiasWeights TrajectoryBiasWeights;
+	TrajectoryBiasWeights.Init(Params, Database->Schema->Layout);
+	
+	const FPoseSearchBiasWeightsContext BiasWeightsContext { &BiasWeights, Database };
+	const FPoseSearchBiasWeightsContext PoseBiasWeightsContext { &PoseBiasWeights, Database };
+	const FPoseSearchBiasWeightsContext TrajectoryBiasWeightsContext { &TrajectoryBiasWeights, Database };
+	
+	for(const TSharedRef<FDebuggerDatabaseRowData>& Row : DatabaseView.Rows)
+	{
+		const uint32 PoseIdx = Row->PoseIdx;
+		Row->Score = ComparePoses(Database->SearchIndex, PoseIdx, State->QueryVectorNormalized, &BiasWeightsContext);
+		Row->PoseScore = ComparePoses(Database->SearchIndex, PoseIdx, State->QueryVectorNormalized, &PoseBiasWeightsContext);
+		Row->TrajectoryScore = ComparePoses(Database->SearchIndex, PoseIdx, State->QueryVectorNormalized, &TrajectoryBiasWeightsContext);
+
+		// If we are on the active pose for the frame
+		if (PoseIdx == State->DbPoseIdx)
+		{
+			*ActiveView.Rows[0] = Row.Get();
+		}
+	}
+
+	SortDatabaseRows();
+}
+
+
+
+TSharedRef<ITableRow> SDebuggerDatabaseView::HandleGenerateDatabaseRow(TSharedRef<FDebuggerDatabaseRowData> Item, const TSharedRef<STableViewBase>& OwnerTable) const
+{
+	return SNew(SDebuggerDatabaseRow, OwnerTable, Item, DatabaseView.RowStyle, &DatabaseView.RowBrush, FMargin(0.0f, 2.0f, 6.0f, 2.0f));
+}
+
+TSharedRef<ITableRow> SDebuggerDatabaseView::HandleGenerateActiveRow(TSharedRef<FDebuggerDatabaseRowData> Item, const TSharedRef<STableViewBase>& OwnerTable) const
+{
+	return SNew(SDebuggerDatabaseRow, OwnerTable, Item, ActiveView.RowStyle, &ActiveView.RowBrush, FMargin(0.0f, 2.0f, 6.0f, 4.0f));
+}
+
+void SDebuggerDatabaseView::Construct(const FArguments& InArgs)
+{
+	using namespace DebuggerDatabaseColumns;
+	
+	// @TODO: Support runtime reordering of these indices
+	// Construct all column types
+	AddColumn(MakeShared<FAnimSequenceName>(0));
+	AddColumn(MakeShared<FPoseIdx>(1));
+	AddColumn(MakeShared<FTime>(2));
+	AddColumn(MakeShared<FLength>(3));
+	AddColumn(MakeShared<FScore>(4));
+	AddColumn(MakeShared<FPoseScore>(5));
+	AddColumn(MakeShared<FTrajectoryScore>(6));
+
+	// Active Row
+
+	// Used for spacing
+	ActiveView.HeaderRow = SNew(SHeaderRow);
+
+	ActiveView.ScrollBar =
+		SNew(SScrollBar)
+		.Orientation(Orient_Vertical)
+		.HideWhenNotInUse(false)
+		.AlwaysShowScrollbar(true)
+		.AlwaysShowScrollbarTrack(true);
+
+	ActiveView.ListView = SNew(SListView<TSharedRef<FDebuggerDatabaseRowData>>)
+		.ListItemsSource(&ActiveView.Rows)
+		.HeaderRow(ActiveView.HeaderRow.ToSharedRef())
+		.OnGenerateRow(this, &SDebuggerDatabaseView::HandleGenerateActiveRow)
+		.ExternalScrollbar(ActiveView.ScrollBar)
+		.SelectionMode(ESelectionMode::SingleToggle)
+		.ConsumeMouseWheel(EConsumeMouseWheel::Never)
+		.ItemHeight(1);
+
+	ActiveView.RowStyle = FEditorStyle::GetWidgetStyle<FTableRowStyle>("TableView.Row");
+	ActiveView.RowBrush = *FEditorStyle::GetBrush("DetailsView.CategoryTop");
+
+	// Database 
+	DatabaseView.ScrollBar =
+		SNew(SScrollBar)
+		.Orientation(Orient_Vertical)
+		.HideWhenNotInUse(false)
+		.AlwaysShowScrollbar(true)
+		.AlwaysShowScrollbarTrack(true);
+	DatabaseView.HeaderRow = SNew(SHeaderRow).Visibility(EVisibility::Collapsed);
+	
+	DatabaseView.ListView = SNew(SListView<TSharedRef<FDebuggerDatabaseRowData>>)
+			.ListItemsSource(&DatabaseView.Rows)
+			.HeaderRow(DatabaseView.HeaderRow.ToSharedRef())
+			.OnGenerateRow(this, &SDebuggerDatabaseView::HandleGenerateDatabaseRow)
+			.ExternalScrollbar(DatabaseView.ScrollBar)
+			.ScrollbarVisibility(EVisibility::Visible)
+			.SelectionMode(ESelectionMode::Multi)
+			.ConsumeMouseWheel(EConsumeMouseWheel::WhenScrollingPossible)
+			.ItemHeight(24);
+
+	DatabaseView.RowStyle = FEditorStyle::GetWidgetStyle<FTableRowStyle>("TableView.Row");
+	// Set selected color to white to retain visibility when multi-selecting
+	DatabaseView.RowStyle.SetSelectedTextColor(FLinearColor(FVector3f(0.8f)));
+	DatabaseView.RowBrush = *FEditorStyle::GetBrush("ToolPanel.GroupBorder");
+	
+	ChildSlot
+	[
+		SNew(SWidgetSwitcher)
+		+ SWidgetSwitcher::Slot()
+		[
+			SNew(SVerticalBox)
+			+ SVerticalBox::Slot()
+			// Side and top margins, ignore bottom handled by the color border below
+			.Padding(0.0f, 5.0f, 0.0f, 0.0f)
+			.AutoHeight()
+			[
+				// Active Row text tab
+				SNew(SVerticalBox)
+				+ SVerticalBox::Slot()
+				.Padding(0.0f)
+				.AutoHeight()
+				[
+					SNew(SHorizontalBox)
+					+ SHorizontalBox::Slot()
+					.HAlign(HAlign_Center)
+					.VAlign(VAlign_Fill)
+					.Padding(0.0f)
+					.AutoWidth()
+					[
+						SNew(SBorder)
+						.BorderImage(FEditorStyle::GetBrush("DetailsView.CategoryTop"))
+						.Padding(FMargin(30.0f, 3.0f, 30.0f, 0.0f))
+						.HAlign(HAlign_Center)
+						.VAlign(VAlign_Fill)
+						[
+							SNew(STextBlock)
+							.Text(FText::FromString("Active Pose"))	
+						]
+					]
+				]
+
+				// Active row list view with scroll bar
+				+ SVerticalBox::Slot()
+				.Padding(0.0f)
+				.AutoHeight()
+				[
+					SNew(SHorizontalBox)
+					
+					+ SHorizontalBox::Slot()
+					.HAlign(HAlign_Fill)
+					.VAlign(VAlign_Fill)
+					.Padding(0.0f)
+					[
+						SNew(SBorder)
+						.Padding(0.0f)
+						.BorderImage(FEditorStyle::GetBrush("NoBorder"))
+						[
+							ActiveView.ListView.ToSharedRef()
+						]
+					]
+
+					+ SHorizontalBox::Slot()
+					.AutoWidth()
+					[
+						ActiveView.ScrollBar.ToSharedRef()
+					]
+				]	
+			]
+
+			+ SVerticalBox::Slot()
+			.Padding(0.0f, 0.0f, 0.0f, 5.0f)
+			[
+				// Database view text tab
+				SNew(SVerticalBox)
+				+ SVerticalBox::Slot()
+				.Padding(0.0f)
+				[
+					SNew(SHorizontalBox)
+					+ SHorizontalBox::Slot()
+					.HAlign(HAlign_Center)
+					.VAlign(VAlign_Fill)
+					.Padding(0.0f)
+					[
+						SNew(SBorder)
+						.BorderImage(FEditorStyle::GetBrush("DetailsView.CategoryTop"))
+						.Padding(FMargin(30.0f, 3.0f, 30.0f, 0.0f))
+						.HAlign(HAlign_Center)
+						.VAlign(VAlign_Fill)
+						[
+							SNew(STextBlock)
+							.Text(FText::FromString("Selected Poses"))	
+						]
+					]
+					.AutoWidth()
+					
+					+ SHorizontalBox::Slot()
+					.HAlign(HAlign_Fill)
+					.Padding(0.0f)
+					[
+						SNew(SBorder)
+						.Padding(0.0f)
+						.BorderImage(&DatabaseView.RowStyle.EvenRowBackgroundBrush)
+					]
+				]
+				.AutoHeight()
+
+				// Gray line below the tab 
+				+ SVerticalBox::Slot()
+				.Padding(0.0f)
+				.AutoHeight()
+				[
+					SNew(SBorder)
+					.BorderImage(FEditorStyle::GetBrush("DetailsView.CategoryTop"))
+					.Padding(FMargin(0.0f, 3.0f, 0.0f, 3.0f))
+					.HAlign(HAlign_Fill)
+					.VAlign(VAlign_Fill)
+				]
+				
+			
+				+ SVerticalBox::Slot()
+				.Padding(0.0f)
+				[
+					SNew(SHorizontalBox)
+					+ SHorizontalBox::Slot()
+					.Padding(0.0f)
+					[
+						SNew(SScrollBox)
+						.Orientation(Orient_Vertical)
+						.ExternalScrollbar(DatabaseView.ScrollBar)
+						.ScrollBarPadding(0.0f)
+						.NavigationScrollPadding(0.0f)
+						.ScrollBarVisibility(EVisibility::Hidden)
+						.AllowOverscroll(EAllowOverscroll::No)
+						+ SScrollBox::Slot()
+						[
+							DatabaseView.ListView.ToSharedRef()
+						]
+					]
+					
+					+ SHorizontalBox::Slot()
+					.AutoWidth()
+					.Padding(0.0f)
+					[
+						DatabaseView.ScrollBar.ToSharedRef()
+					]
+				]
+			]
+		]
+	];
+	
+	SortColumn = FPoseIdx::Name;
+	SortMode = EColumnSortMode::Ascending;
+
+	ActiveView.ScrollBar->SetVisibility(EVisibility::Hidden);
+
+	// Assign the get column function on the database rows
+	SDebuggerDatabaseRow::GetColumn = FGetColumnDelegate::CreateLambda([this](const FName& Name) -> TSharedRef<IColumn>&
+	{
+		return Columns[Name];
+	});
+	
+	RefreshColumns();
+}
+
+void SDebuggerDetailsView::Construct(const FArguments& InArgs, UPoseSearchDebuggerReflection* Reflection)
+{
+	/** Construct this panel using the property editor */
+	FPropertyEditorModule& PropPlugin = FModuleManager::LoadModuleChecked<FPropertyEditorModule>("PropertyEditor");
+	FDetailsViewArgs DetailsViewArgs;
+    DetailsViewArgs.NameAreaSettings = FDetailsViewArgs::HideNameArea;
+
+	// @TODO: Hide arrays with zero elements in the detail view, if possible
+	Details = PropPlugin.CreateDetailView(DetailsViewArgs);
+
+	Details->SetObject(Reflection);
+	
+	ChildSlot
+	[
+		Details.ToSharedRef()
+	];
+}
+
+
+void SDebuggerView::Construct(const FArguments& InArgs, uint64 InAnimInstanceId)
+{
+	MotionMatchingNodeIds = InArgs._MotionMatchingNodeIds;
+	MotionMatchingState = InArgs._MotionMatchingState;
+	Reflection = InArgs._Reflection;
+	PoseSearchDatabase = InArgs._PoseSearchDatabase;
+	OnSelectionChanged = InArgs._OnSelectionChanged;
+	IsPIESimulating = InArgs._IsPIESimulating;
+	AnimInstanceId = InAnimInstanceId;
+	SelectedNode = -1;
+
+	ChildSlot
+	[
+		SAssignNew(DebuggerView, SVerticalBox)
+		+ SVerticalBox::Slot()
+		.FillHeight(1.0f)
+		.VAlign(VAlign_Center)
+		.HAlign(HAlign_Fill)
+		[
+			SAssignNew(Switcher, SWidgetSwitcher)
+
+			// [0] Box that covers everything when recording
+			+ SWidgetSwitcher::Slot()
+			.Padding(0.0f)
+			[
+				SAssignNew(SimulatingView, SVerticalBox)
+				
+				+ SVerticalBox::Slot()
+				.HAlign(HAlign_Center)
+                .VAlign(VAlign_Center)
+				[
+					SNew(STextBlock)
+					.Text(FText::FromString("Select a frame to continue..."))
+					.Font(FEditorStyle::Get().GetFontStyle("DetailsView.CategoryFontStyle"))
+				]
+			]
+
+			// [1] Selection view before node selection is made
+			+ SWidgetSwitcher::Slot()
+			.Padding(40.0f)
+			.HAlign(HAlign_Fill)
+            .VAlign(VAlign_Center)
+			[
+				SAssignNew(SelectionView, SVerticalBox)
+			]
+
+			// [2] Node selected; node debugger view
+			+ SWidgetSwitcher::Slot()
+			.Padding(0.0f)
+			[
+				GenerateNodeDebuggerView()
+			]
+		]
+	];
+	
+}
+
+void SDebuggerView::SetTimeMarker(double InTimeMarker)
+{
+	Switcher->SetActiveWidgetIndex(SwitcherViewType);
+	if (IsPIESimulating.Get())
+	{
+		SwitcherViewType = Waiting;
+		return;
+	}
+
+	const bool SameTime = FMath::Abs(InTimeMarker - TimeMarker) < DOUBLE_SMALL_NUMBER;
+	if (SameTime)
+	{
+		return;
+	}
+	TimeMarker = InTimeMarker;
+
+	UpdateViews();
+}
+
+void SDebuggerView::UpdateViews()
+{
+	const TSet<int32> NodeIds = MotionMatchingNodeIds.Get();
+	if (NodeIds.IsEmpty())
+	{
+		return;
+	}
+	
+	// Update selection view if no node selected
+	if (SelectedNode == -1)
+	{
+		SwitcherViewType = Selection;
+
+		// If we have a new set of nodes
+		if (!NodeIds.Difference(ActiveNodes).IsEmpty())
+		{
+			// Only one node active, bypass selection view
+			if (NodeIds.Num() == 1)
+			{
+				SelectedNode = *NodeIds.begin();
+				OnSelectionChanged.Execute(AnimInstanceId, SelectedNode);
+				UpdateViews();
+			}
+			// Create selection view with buttons for each node, displaying the database name
+			else
+			{
+				SelectionView->ClearChildren();
+				for (int32 NodeId : NodeIds)
+				{
+					OnSelectionChanged.Execute(AnimInstanceId, NodeId);
+					SelectionView->AddSlot()
+					.HAlign(HAlign_Fill)
+					.VAlign(VAlign_Center)
+					.Padding(10.0f)
+					[
+						SNew(SButton)
+						.Text(FText::FromString(PoseSearchDatabase.Get()->GetName()))
+						.HAlign(HAlign_Center)
+						.VAlign(VAlign_Center)
+						.ContentPadding(10.0f)
+						.OnClicked_Lambda([this, NodeId]() -> FReply
+						{
+							SelectedNode = NodeId;
+							UpdateViews();
+							return FReply::Handled();
+						})
+					];
+				}
+			}
+		}
+	}
+	else
+	{
+		check(Reflection.Get());
+    	OnSelectionChanged.Execute(AnimInstanceId, SelectedNode);
+		
+    	const FTraceMotionMatchingStateMessage* State = MotionMatchingState.Get();
+    	const UPoseSearchDatabase* Database = PoseSearchDatabase.Get();
+    	if (!State || !Database)
+    	{
+    		return;
+    	}
+    
+    	SwitcherViewType = Debugger;
+    	DatabaseView->Update(State, Database);	
+	}
+	ActiveNodes = NodeIds;
+}
+
+TSharedRef<SWidget> SDebuggerView::GenerateReturnButtonView()
+{
+	return
+		SAssignNew(ReturnButtonView, SHorizontalBox)
+		.Visibility_Lambda([this]() -> EVisibility
+		{
+			// Collapse this view if we have don't have more than 1 node
+			return ActiveNodes.Num() > 1 ? EVisibility::Visible : EVisibility::Collapsed;
+		})
+
+		+ SHorizontalBox::Slot()
+		.Padding(10, 5, 0, 0)
+		.AutoWidth()
+		[
+			SNew(SButton)
+			.VAlign(VAlign_Center)
+			.ButtonStyle(FEditorStyle::Get(), "SimpleButton")
+			.ContentPadding( FMargin(1, 0) )
+			.OnClicked_Lambda([this]() -> FReply
+			{
+				// Clicking back backtracks selected node to invalid
+				SelectedNode = -1;
+				UpdateViews();
+				return FReply::Handled();
+			})
+			// Contents of button, icon then text
+			[
+				SNew(SHorizontalBox)
+
+				+ SHorizontalBox::Slot()
+				.VAlign(VAlign_Center)
+				[
+					SNew(SImage)
+					.Image(FAppStyle::Get().GetBrush("Icons.CircleArrowLeft"))
+					.ColorAndOpacity(FSlateColor::UseForeground())
+				]
+
+				+ SHorizontalBox::Slot()
+				.AutoWidth()
+				.Padding(FMargin(8.0f, 0.0f, 0.0f, 0.0f))
+				.VAlign(VAlign_Center)
+				[
+					SNew(STextBlock)
+					.Text(FText::FromString("Return to Database Selection"))
+					.Justification(ETextJustify::Center)
+				]
+			]
+		];
+}
+
+TSharedRef<SWidget> SDebuggerView::GenerateNodeDebuggerView()
+{
+	return 
+		SNew(SSplitter)
+		.Orientation(Orient_Horizontal)
+		.ResizeMode(ESplitterResizeMode::Fill)
+	
+		// Database view
+		+ SSplitter::Slot()
+		.Value(0.65f)
+		[
+			SNew(SVerticalBox)
+			
+			+ SVerticalBox::Slot()
+			.AutoHeight()
+			[
+				GenerateReturnButtonView()
+			]
+			
+			+ SVerticalBox::Slot()
+			[
+				SAssignNew(DatabaseView, SDebuggerDatabaseView)
+			]
+		]
+
+		// Details panel view
+		+ SSplitter::Slot()
+		.Value(0.35f)
+		[
+			SAssignNew(DetailsView, SDebuggerDetailsView, Reflection.Get())
+		];
+}
+
+FName SDebuggerView::GetName() const
+{
+	static const FName DebuggerName("PoseSearchDebugger");
+	return DebuggerName;
+}
+
+uint64 SDebuggerView::GetObjectId() const
+{
+	return AnimInstanceId;
+}
+
+
+FDebugger* FDebugger::InternalInstance;
+void FDebugger::Initialize()
+{
+	InternalInstance = new FDebugger;
+	IModularFeatures::Get().RegisterModularFeature(IRewindDebuggerExtension::ModularFeatureName, InternalInstance);
+
+	// Add property editor (detail view) UObject to world root so that it persists when PIE is stopped
+	InternalInstance->Reflection = NewObject<UPoseSearchDebuggerReflection>();
+	InternalInstance->Reflection->AddToRoot();
+	
+	check(IsValid(InternalInstance->Reflection));
+}
+
+void FDebugger::Shutdown()
+{
+	FPropertyEditorModule& PropPlugin = FModuleManager::LoadModuleChecked<FPropertyEditorModule>("PropertyEditor");
+	PropPlugin.UnregisterCustomClassLayout(UPoseSearchDebuggerReflection::StaticClass()->GetFName());
+
+	// Our previously instantiated object attached to root may be cleaned up at this point
+	if (UObjectInitialized())
+	{
+		InternalInstance->Reflection->RemoveFromRoot();
+	}
+	
+	IModularFeatures::Get().UnregisterModularFeature(IRewindDebuggerExtension::ModularFeatureName, InternalInstance);
+	delete InternalInstance;
+}
+
+void FDebugger::Update(float DeltaTime, IRewindDebugger* InRewindDebugger)
+{
+	// Update active rewind debugger in use
+	RewindDebugger = InRewindDebugger;
+}
+
+void FDebugger::UpdateReflection() const
+{
+	check(Reflection);
+
+	if (!MotionMatchingState)
+	{
+		return;
+	}
+	
+	Reflection->ElapsedPoseJumpTime = MotionMatchingState->ElapsedPoseJumpTime;
+	
+	const UPoseSearchDatabase* Database = GetPoseSearchDatabase();
+	if (Database == nullptr)
+	{
+		return;
+	}
+	// Validated by the get above
+	const UPoseSearchSchema* Schema = Database->Schema;
+	
+	Reflection->PoseFeatures.EmptyAll();
+	Reflection->DistanceTrajectoryFeatures.EmptyAll();
+	Reflection->TimeTrajectoryFeatures.EmptyAll();
+
+	FFeatureVectorReader Reader;
+	Reader.Init(&Schema->Layout);
+	// Ensure parity between Layout and QueryVector
+	Reader.SetValues(MotionMatchingState->QueryVector);
+	check(Reader.IsValid());
+
+	int32 NumSubsamples = Schema->PoseSampleTimes.Num();
+	const int32 NumBones = Schema->Bones.Num();
+
+	FPoseSearchFeatureDesc Feature;
+	
+	// Aggregate all features and place into the reflection struct
+	auto Extract = [](
+		const FFeatureVectorReader& InReader,
+		const FPoseSearchFeatureDesc& InFeature,
+		FPoseSearchDebuggerFeatureReflection& ReflectionRef
+	)
+	{
+		FVector OutputVec;
+		if (InReader.GetPosition(InFeature, &OutputVec))
+		{
+			ReflectionRef.Positions.Add(OutputVec);
+		}
+		if (InReader.GetLinearVelocity(InFeature, &OutputVec))
+		{
+			ReflectionRef.LinearVelocities.Add(OutputVec);
+		}
+		if (InReader.GetAngularVelocity(InFeature, &OutputVec))
+        {
+        	ReflectionRef.AngularVelocities.Add(OutputVec);
+        }
+	};
+
+	// Pose samples
+	Feature.Domain = EPoseSearchFeatureDomain::Time;
+	for (int32 SchemaSubsampleIndex = 0; SchemaSubsampleIndex < NumSubsamples; ++SchemaSubsampleIndex)
+	{
+		Feature.SubsampleIdx = SchemaSubsampleIndex;
+
+		for(int32 SchemaBoneIdx = 0; SchemaBoneIdx < NumBones; ++SchemaBoneIdx)
+		{
+			Feature.SchemaBoneIdx = SchemaBoneIdx;
+			Extract(Reader, Feature, Reflection->PoseFeatures);
+		}
+	}
+
+	// Used for classifying trajectories instead of bones, special index
+	Feature.SchemaBoneIdx = FPoseSearchFeatureDesc::TrajectoryBoneIndex;
+	
+	// Trajectory time samples
+	NumSubsamples = Schema->TrajectorySampleTimes.Num();
+	for (int32 SchemaSubsampleIndex = 0; SchemaSubsampleIndex < NumSubsamples; ++SchemaSubsampleIndex)
+	{
+		Feature.SubsampleIdx = SchemaSubsampleIndex;
+		Extract(Reader, Feature, Reflection->TimeTrajectoryFeatures);
+	}
+	// Trajectory distance samples
+	NumSubsamples = Schema->TrajectorySampleDistances.Num();
+	Feature.Domain = EPoseSearchFeatureDomain::Distance;
+	for (int32 SchemaSubsampleIndex = 0; SchemaSubsampleIndex < NumSubsamples; ++SchemaSubsampleIndex)
+	{
+		Feature.SubsampleIdx = SchemaSubsampleIndex;
+		Extract(Reader, Feature, Reflection->DistanceTrajectoryFeatures);
+	}
+}
+
+
+const FTraceMotionMatchingStateMessage* FDebugger::GetMotionMatchingState()
+{
+	return Instance()->MotionMatchingState;
+}
+
+const UPoseSearchDatabase* FDebugger::GetPoseSearchDatabase() 
+{
+	if (InternalInstance->MotionMatchingState == nullptr)
+	{
+		return nullptr;
+	}
+
+	const uint64 DatabaseId = InternalInstance->MotionMatchingState->DatabaseId;
+	if (DatabaseId == 0)
+	{
+	    return nullptr;
+	}
+	
+	// @TODO: Load the database if not currently loaded
+	UObject* DatabaseObject = FObjectTrace::GetObjectFromId(DatabaseId);
+	if (DatabaseObject == nullptr)
+	{
+		return nullptr;
+	}
+	
+	check(DatabaseObject->IsA<UPoseSearchDatabase>());
+
+	const UPoseSearchDatabase* Database = Cast<UPoseSearchDatabase>(DatabaseObject);
+	const UPoseSearchSchema* Schema = Database->Schema;
+	if (Schema == nullptr || !Schema->IsValid())
+	{
+		return nullptr;
+	}
+	return Database;
+}
+
+UPoseSearchDebuggerReflection* FDebugger::GetReflection()
+{
+	return InternalInstance->Reflection;
+}
+
+bool FDebugger::GetIsPIESimulating()
+{
+	return InternalInstance->RewindDebugger->IsPIESimulating();
+}
+
+TSet<int32> FDebugger::GetNodeIds(uint64 AnimInstanceId)
+{
+	const TraceServices::IAnalysisSession* Session = Instance()->RewindDebugger->GetAnalysisSession();
+	TraceServices::FAnalysisSessionReadScope SessionReadScope(*Session);
+
+	// Get provider and validate its existence in the session
+	const FTraceProvider* TraceProvider = Session->ReadProvider<FTraceProvider>(FTraceProvider::ProviderName);
+	if (!TraceProvider)
+	{
+		return TSet<int32>();
+	}
+	
+	return TraceProvider->GetMotionMatchingNodeIds(AnimInstanceId);
+}
+
+void FDebugger::OnSelectionChanged(uint64 AnimInstanceId, int32 NodeId)
+{
+	const TraceServices::IAnalysisSession* Session = RewindDebugger->GetAnalysisSession();
+	TraceServices::FAnalysisSessionReadScope SessionReadScope(*Session);
+
+	// Get provider and validate
+	const FTraceProvider* TraceProvider = Session->ReadProvider<FTraceProvider>(FTraceProvider::ProviderName);
+	if (!TraceProvider)
+	{
+		return;
+	}
+
+	const double TraceTime = RewindDebugger->CurrentTraceTime();
+	TraceProvider->ReadMotionMatchingStateTimeline(AnimInstanceId, NodeId, [this, TraceTime](const FTraceProvider::FMotionMatchingStateTimeline& TimelineData)
+	{
+		double PrecedingValueTime;
+		double FollowingValueTime;
+		const FTraceMotionMatchingStateMessage* PrecedingValue;
+		const FTraceMotionMatchingStateMessage* FollowingValue;
+
+		// Get event data at the trace time
+		TimelineData.FindNearestEvents(
+			TraceTime,
+			PrecedingValue,
+			PrecedingValueTime,
+			FollowingValue,
+			FollowingValueTime
+		);
+
+		// Grab following event if it exists, otherwise get the preceding (mimicking skeletal mesh updates in the Rewind Debugger)
+		if (PrecedingValue || FollowingValue)
+		{
+			MotionMatchingState = FollowingValue ? FollowingValue : PrecedingValue;
+		}
+	});
+
+
+	UpdateReflection();
+}
+
+TSharedPtr<SDebuggerView> FDebugger::GenerateView(uint64 InAnimInstanceId)
+{
+	return
+		SNew(SDebuggerView, InAnimInstanceId)
+		.MotionMatchingNodeIds_Static(&FDebugger::GetNodeIds, InAnimInstanceId)
+		.MotionMatchingState_Static(&FDebugger::GetMotionMatchingState)
+		.Reflection_Static(&FDebugger::GetReflection)
+		.PoseSearchDatabase_Static(&FDebugger::GetPoseSearchDatabase)
+		.IsPIESimulating_Static(&FDebugger::GetIsPIESimulating)
+		.OnSelectionChanged_Raw(Instance(), &FDebugger::OnSelectionChanged);
+}
+
+
+FText FDebuggerViewCreator::GetTitle() const
+{
+	return LOCTEXT("PoseSearchDebuggerTabTitle", "Pose Search");
+}
+
+FSlateIcon FDebuggerViewCreator::GetIcon() const
+{
+#if WITH_EDITOR
+	return FSlateIconFinder::FindIconForClass(UAnimInstance::StaticClass());
+#else
+	return FSlateIcon();
+#endif
+}
+
+FName FDebuggerViewCreator::GetTargetTypeName() const
+{
+	static FName TargetTypeName = "AnimInstance";
+	return TargetTypeName;
+}
+
+TSharedPtr<IRewindDebuggerView> FDebuggerViewCreator::CreateDebugView(uint64 ObjectId, double CurrentTime, const TraceServices::IAnalysisSession& InAnalysisSession) const
+{
+	return FDebugger::GenerateView(ObjectId);
+}
+
+FName FDebuggerViewCreator::GetName() const
+{
+	static const FName Name("PoseSearchDebugger");
+	return Name;
+}
+
+}}
diff --git a/Engine/Plugins/Experimental/Animation/PoseSearch/Source/Editor/Private/PoseSearchDebugger.h b/Engine/Plugins/Experimental/Animation/PoseSearch/Source/Editor/Private/PoseSearchDebugger.h
new file mode 100644
--- /dev/null
+++ b/Engine/Plugins/Experimental/Animation/PoseSearch/Source/Editor/Private/PoseSearchDebugger.h
@@ -0,0 +1,300 @@
+﻿// Copyright Epic Games, Inc. All Rights Reserved.
+
+#pragma once
+
+#include "CoreMinimal.h"
+#include "RewindDebuggerInterface/Public/IRewindDebuggerExtension.h"
+#include "RewindDebuggerInterface/Public/IRewindDebuggerView.h"
+#include "RewindDebuggerInterface/Public/IRewindDebuggerViewCreator.h"
+#include "Widgets/Views/SHeaderRow.h"
+#include "PoseSearchDebugger.generated.h"
+
+namespace TraceServices { class IAnalysisSession; }
+class IUnrealInsightsModule;
+class UPoseSearchDatabase;
+class ITableRow;
+class STableViewBase;
+class IDetailsView;
+class SVerticalBox;
+class SHorizontalBox;
+class SScrollBox;
+class SWidgetSwitcher;
+template <typename ItemType> class SListView;
+
+/**
+ * Used by the reflection UObject to encompass a set of features
+ */
+USTRUCT()
+struct FPoseSearchDebuggerFeatureReflection
+{
+	GENERATED_USTRUCT_BODY()
+
+	// @TODO: Should be ideally enumerated based on all possible schema features
+    UPROPERTY(VisibleAnywhere, Category="Query Data")
+	TArray<FVector> Positions;
+    UPROPERTY(VisibleAnywhere, Category="Query Data")
+	TArray<FVector> LinearVelocities;
+    UPROPERTY(VisibleAnywhere, Category="Query Data")
+	TArray<FVector> AngularVelocities;
+
+	/** Empty contents of above arrays */
+	void EmptyAll();
+};
+
+/**
+ * Reflection UObject being observed in the details view panel of the debugger
+ */
+UCLASS()
+class POSESEARCHEDITOR_API UPoseSearchDebuggerReflection : public UObject
+{
+	GENERATED_UCLASS_BODY()
+
+	/** Time since last PoseSearch jump */
+	UPROPERTY(VisibleAnywhere, Category="Motion Matching State")
+	float ElapsedPoseJumpTime = 0.0f;
+
+	/** Whether it is playing the loop following the expended animation runway */
+	UPROPERTY(VisibleAnywhere, Category="Motion Matching State")
+	bool bFollowUpAnimation = false;
+
+	/** Pose features of the current query vector */
+    UPROPERTY(EditAnywhere, Category="Query Data")
+	FPoseSearchDebuggerFeatureReflection PoseFeatures;
+    	
+	/** Time-based trajectory features of the current query vector */
+    UPROPERTY(EditAnywhere, Category="Query Data")
+	FPoseSearchDebuggerFeatureReflection TimeTrajectoryFeatures;
+
+	/** Distance-based trajectory features of the current query vector */
+	UPROPERTY(EditAnywhere, Category="Query Data")
+    FPoseSearchDebuggerFeatureReflection DistanceTrajectoryFeatures;
+};
+
+
+namespace UE { namespace PoseSearch {
+
+namespace DebuggerDatabaseColumns { struct IColumn; }
+struct FTraceMotionMatchingStateMessage;
+class FDebuggerDatabaseRowData;
+
+
+/**
+ * Database panel view widget of the PoseSearch debugger
+ */
+class SDebuggerDatabaseView : public SCompoundWidget
+{
+	SLATE_BEGIN_ARGS(SDebuggerDatabaseView) {}
+	SLATE_END_ARGS()
+	
+	void Construct(const FArguments& InArgs);
+	void Update(const FTraceMotionMatchingStateMessage* State, const UPoseSearchDatabase* Database);
+
+private:
+	void RefreshColumns();
+	void AddColumn(TSharedRef<DebuggerDatabaseColumns::IColumn>&& Column);
+	void CreateRows(const UPoseSearchDatabase* Database);
+	void SortDatabaseRows();
+	void PopulateRows(const FTraceMotionMatchingStateMessage* State, const UPoseSearchDatabase* Database);
+	EColumnSortMode::Type GetColumnSortMode(const FName ColumnId) const;
+	float GetColumnWidth(const FName ColumnId) const;
+	void OnColumnSortModeChanged(const EColumnSortPriority::Type SortPriority, const FName & ColumnId, const EColumnSortMode::Type InSortMode);
+	void OnColumnWidthChanged(const float NewWidth, FName ColumnId) const;
+	TSharedRef<ITableRow> HandleGenerateDatabaseRow(TSharedRef<FDebuggerDatabaseRowData> Item, const TSharedRef<STableViewBase>& OwnerTable) const;
+	TSharedRef<ITableRow> HandleGenerateActiveRow(TSharedRef<FDebuggerDatabaseRowData> Item, const TSharedRef<STableViewBase>& OwnerTable) const;
+
+	/** Current column to sort by */
+	FName SortColumn = "";
+	/** Current sorting mode */
+	EColumnSortMode::Type SortMode = EColumnSortMode::Ascending;
+	
+	/** Column data container, used to emplace defined column structures of various types */
+    TMap<FName, TSharedRef<DebuggerDatabaseColumns::IColumn>> Columns;
+
+	struct FTable
+	{
+		/** Header row*/
+		TSharedPtr<SHeaderRow> HeaderRow;
+
+		/** Widget for displaying the list of row objects */
+		TSharedPtr<SListView<TSharedRef<FDebuggerDatabaseRowData>>> ListView;
+
+		// @TODO: Explore options for active row other than displaying array of 1 element
+		/** List of row objects */
+		TArray<TSharedRef<FDebuggerDatabaseRowData>> Rows;
+
+		/** Background style for the list view */
+		FTableRowStyle RowStyle;
+
+		/** Row color */
+		FSlateBrush RowBrush;
+
+		/** Scroll bar for the data table */
+		TSharedPtr<SScrollBar> ScrollBar;
+	};
+
+	/** Active row at the top of the view */
+	FTable ActiveView;
+	/** Database listings for all poses */
+	FTable DatabaseView;
+};
+
+/**
+ * Details panel view widget of the PoseSearch debugger
+ */
+class SDebuggerDetailsView : public SCompoundWidget
+{
+	SLATE_BEGIN_ARGS(SDebuggerDetailsView) {}
+	SLATE_END_ARGS()
+
+	void Construct(const FArguments& InArgs, UPoseSearchDebuggerReflection* Reflection);
+private:
+	/** Details widget constructed for the MM node */
+	TSharedPtr<IDetailsView> Details;
+};
+
+
+/**
+ * Callback to update the debugger when node selection is changed
+ */
+DECLARE_DELEGATE_TwoParams(FOnSelectionChanged, uint64 AnimInstanceId, int32 NodeId);
+
+/**
+ * Entire view of the PoseSearch debugger, containing all sub-widgets
+ */
+class SDebuggerView : public IRewindDebuggerView
+{
+public:
+	SDebuggerView() = default;
+
+	SLATE_BEGIN_ARGS(SDebuggerView){}
+		SLATE_ATTRIBUTE(const FTraceMotionMatchingStateMessage*, MotionMatchingState)
+		SLATE_ATTRIBUTE(UPoseSearchDebuggerReflection*, Reflection)
+		SLATE_ATTRIBUTE(const UPoseSearchDatabase*, PoseSearchDatabase)
+		SLATE_ATTRIBUTE(TSet<int32>, MotionMatchingNodeIds)
+		SLATE_ATTRIBUTE(bool, IsPIESimulating)
+		SLATE_EVENT(FOnSelectionChanged, OnSelectionChanged)
+	SLATE_END_ARGS()
+
+	void Construct(const FArguments& InArgs, uint64 InAnimInstanceId);
+	virtual void SetTimeMarker(double InTimeMarker) override;
+	virtual FName GetName() const override;
+	virtual uint64 GetObjectId() const override;
+
+private:
+	void UpdateViews();
+	TSharedRef<SWidget> GenerateReturnButtonView();
+	TSharedRef<SWidget> GenerateNodeDebuggerView();
+
+	/** Gets all MM nodes being traced in this session */
+	TAttribute<TSet<int32>> MotionMatchingNodeIds;
+	/** Retrieves the reflection UObject from the debugger */
+	TAttribute<UPoseSearchDebuggerReflection*> Reflection;
+	/** Retrieves the MM state from the debugger */
+	TAttribute<const FTraceMotionMatchingStateMessage*> MotionMatchingState;
+	/** Retrieves the PoseSearch Database from the debugger */
+	TAttribute<const UPoseSearchDatabase*> PoseSearchDatabase;
+	/** Whether the game is unpaused and currently simulating */
+	TAttribute<bool> IsPIESimulating;
+	/** Update current debugger data when node selection is changed */
+	FOnSelectionChanged OnSelectionChanged;
+
+	/** Active node being debugged */
+	int32 SelectedNode = -1;
+
+	/** List of all nodes being traced (stored for selection) */
+	TSet<int32> ActiveNodes;
+
+	/** Database view of the motion matching node */
+	TSharedPtr<SDebuggerDatabaseView> DatabaseView;
+	/** Details panel for introspecting the motion matching node */
+	TSharedPtr<SDebuggerDetailsView> DetailsView;
+	/** Node debugger view hosts the above two views */
+	TSharedPtr<SSplitter> NodeDebuggerView;
+
+	/** Contains the return to node selection (picked by database name) */
+	TSharedPtr<SHorizontalBox> ReturnButtonView;
+	/** Selection view before node is selected */
+	TSharedPtr<SVerticalBox> SelectionView;
+
+	
+	/** Gray box occluding the debugger view when simulating */
+	TSharedPtr<SVerticalBox> SimulatingView;
+	
+	/** Used to switch between views in the switcher */
+	enum ESwitcherViewType : uint16
+	{
+		Waiting = 0,
+		Selection = 1,
+		Debugger = 2
+	} SwitcherViewType = Waiting;
+	/** Contains all the above, switches between them depending on context */
+	// @TODO: Make this resolve to Waiting when PIE is stopped
+	TSharedPtr<SWidgetSwitcher> Switcher;
+
+	/** Contains the switcher, the entire debugger view */
+	TSharedPtr<SVerticalBox> DebuggerView;
+
+	/** AnimInstance this view was created for */
+	uint64 AnimInstanceId = 0;
+
+	/** Current position of the time marker */
+	double TimeMarker = 0.0;
+};
+
+
+/**
+ * PoseSearch debugger, containing the data to be acquired and relayed to the view
+ */
+class FDebugger : public IRewindDebuggerExtension, public TSharedFromThis<FDebugger>
+{
+public:
+	virtual void Update(float DeltaTime, IRewindDebugger* InRewindDebugger) override;
+	virtual ~FDebugger() = default;
+
+	static FDebugger* Instance() { return InternalInstance; }
+	static void Initialize();
+	static void Shutdown();
+	static const FName ModularFeatureName;
+
+	/** Generates the slate debugger view widget */
+	static TSharedPtr<SDebuggerView> GenerateView(uint64 InAnimInstanceId);
+
+private:
+	// Used for view callbacks
+	static const FTraceMotionMatchingStateMessage* GetMotionMatchingState();
+	static const UPoseSearchDatabase* GetPoseSearchDatabase();
+	static UPoseSearchDebuggerReflection* GetReflection();
+	static bool GetIsPIESimulating();
+	static TSet<int32> GetNodeIds(uint64 AnimInstanceId);
+	void OnSelectionChanged(uint64 AnimInstanceId, int32 NodeId);
+
+	/** Updates the current reflection data relative to the MM state */
+	void UpdateReflection() const;
+	
+	/** Last stored Rewind Debugger */
+	const IRewindDebugger* RewindDebugger = nullptr;
+	/** Last stored MM state (updated from OnSelectionChanged) */
+    const FTraceMotionMatchingStateMessage* MotionMatchingState = nullptr;
+	/** Last updated reflection data relative to MM state */
+	TObjectPtr<UPoseSearchDebuggerReflection> Reflection = nullptr;
+
+	static FDebugger* InternalInstance;
+};
+
+/**
+ * Creates the slate widgets associated with the PoseSearch debugger
+ * when prompted by the Rewind Debugger
+ */
+class FDebuggerViewCreator : public IRewindDebuggerViewCreator
+{
+public:
+	virtual ~FDebuggerViewCreator() = default;
+	virtual FName GetName() const override;
+	virtual FText GetTitle() const override;
+	virtual FSlateIcon GetIcon() const override;
+	virtual FName GetTargetTypeName() const override;
+	virtual TSharedPtr<IRewindDebuggerView> CreateDebugView(uint64 ObjectId, double CurrentTime, const TraceServices::IAnalysisSession& InAnalysisSession) const override;
+};
+
+
+}}
diff --git a/Engine/Plugins/Experimental/Animation/PoseSearch/Source/Editor/Private/PoseSearchEditor.cpp b/Engine/Plugins/Experimental/Animation/PoseSearch/Source/Editor/Private/PoseSearchEditor.cpp
--- a/Engine/Plugins/Experimental/Animation/PoseSearch/Source/Editor/Private/PoseSearchEditor.cpp
+++ b/Engine/Plugins/Experimental/Animation/PoseSearch/Source/Editor/Private/PoseSearchEditor.cpp
@@ -8,6 +8,9 @@
 #include "IAnimationEditor.h"
 #include "IPersonaToolkit.h"
 #include "IPersonaPreviewScene.h"
+#include "PoseSearchDebugger.h"
+#include "Trace/PoseSearchTraceAnalyzer.h"
+#include "Trace/PoseSearchTraceModule.h"
 
 
 //////////////////////////////////////////////////////////////////////////
@@ -62,12 +65,24 @@ public:
 
 private:
 	TArray<class IConsoleObject*> ConsoleCommands;
+	
+	/** Creates the view for the Rewind Debugger */
+	TSharedPtr<FDebuggerViewCreator> DebuggerViewCreator;
+	/** Enables dedicated PoseSearch trace module */
+	TSharedPtr<FTraceModule> TraceModule;
 };
 
 void FEditorModule::StartupModule()
 {
 	if (GIsEditor && !IsRunningCommandlet())
 	{
+		FDebugger::Initialize();
+		TraceModule = MakeShared<FTraceModule>();
+		DebuggerViewCreator = MakeShared<FDebuggerViewCreator>();
+
+		IModularFeatures::Get().RegisterModularFeature("RewindDebuggerViewCreator", DebuggerViewCreator.Get());
+		IModularFeatures::Get().RegisterModularFeature(TraceServices::ModuleFeatureName, TraceModule.Get());
+		
 		ConsoleCommands.Add(IConsoleManager::Get().RegisterConsoleCommand(
 			TEXT("a.PoseSearch.DrawSearchIndex"),
 			TEXT("Draw the search index for the selected asset"),
@@ -84,6 +99,9 @@ void FEditorModule::ShutdownModule()
 		IConsoleManager::Get().UnregisterConsoleObject(ConsoleCmd);
 	}
 	ConsoleCommands.Empty();
+	
+	IModularFeatures::Get().UnregisterModularFeature(TraceServices::ModuleFeatureName, TraceModule.Get());
+	FDebugger::Shutdown();
 }
 
 }} // namespace UE::PoseSearch
diff --git a/Engine/Plugins/Experimental/Animation/PoseSearch/Source/Editor/Private/Trace/PoseSearchTraceAnalyzer.cpp b/Engine/Plugins/Experimental/Animation/PoseSearch/Source/Editor/Private/Trace/PoseSearchTraceAnalyzer.cpp
new file mode 100644
--- /dev/null
+++ b/Engine/Plugins/Experimental/Animation/PoseSearch/Source/Editor/Private/Trace/PoseSearchTraceAnalyzer.cpp
@@ -0,0 +1,77 @@
+﻿// Copyright Epic Games, Inc. All Rights Reserved.
+
+#include "PoseSearchTraceAnalyzer.h"
+#include "PoseSearchTraceProvider.h"
+#include "Runtime/Private/Trace/PoseSearchTraceLogger.h"
+#include "TraceServices/Model/AnalysisSession.h"
+
+namespace UE { namespace PoseSearch {
+
+FTraceAnalyzer::FTraceAnalyzer(TraceServices::IAnalysisSession& InSession, FTraceProvider& InTraceProvider) : Session(InSession), TraceProvider(InTraceProvider)
+{
+}
+
+void FTraceAnalyzer::OnAnalysisBegin(const FOnAnalysisContext& Context)
+{
+	FInterfaceBuilder& Builder = Context.InterfaceBuilder;
+
+	ANSICHAR LoggerName[NAME_SIZE];
+	ANSICHAR MotionMatchingStateName[NAME_SIZE];
+
+	FTraceLogger::Name.GetPlainANSIString(LoggerName);
+	FTraceMotionMatchingState::Name.GetPlainANSIString(MotionMatchingStateName);
+
+	Builder.RouteEvent(RouteId_MotionMatchingState, LoggerName, MotionMatchingStateName);
+}
+
+bool FTraceAnalyzer::OnEvent(uint16 RouteId, EStyle Style, const FOnEventContext& Context)
+{
+	TraceServices::FAnalysisSessionEditScope Scope(Session);
+	const FEventData& EventData = Context.EventData;
+
+	// Gather event data values
+	const double Time = Context.EventTime.AsSeconds(EventData.GetValue<uint64>("Cycle"));
+	const uint16 FrameCounter = EventData.GetValue<uint16>("FrameCounter");
+	const uint64 AnimInstanceId = EventData.GetValue<uint64>("AnimInstanceId");
+	const int32 NodeId = EventData.GetValue<int32>("NodeId");
+
+	switch (RouteId)
+	{
+		case RouteId_MotionMatchingState:
+			{
+				FTraceMotionMatchingStateMessage Message;
+				Message.ElapsedPoseJumpTime = EventData.GetValue<float>("ElapsedPoseJumpTime");
+				// Cast back to our flag type
+				Message.Flags = static_cast<FTraceMotionMatchingState::EFlags>(EventData.GetValue<uint32>("Flags"));
+				Message.DbPoseIdx = EventData.GetValue<int32>("DbPoseIdx");
+
+				// Gather array types
+				TArrayView<const float> View = EventData.GetArrayView<float>("QueryVector");
+				Message.QueryVector = View;
+				
+				View = EventData.GetArrayView<float>("QueryVectorNormalized");
+                Message.QueryVectorNormalized = View;
+				
+				View = EventData.GetArrayView<float>("BiasWeights");
+				Message.BiasWeights = View;
+				
+				Message.DatabaseId = EventData.GetValue<uint64>("DatabaseId");
+
+				// Common data
+				Message.NodeId = NodeId;
+				Message.AnimInstanceId = AnimInstanceId;
+				Message.FrameCounter = FrameCounter;
+
+				TraceProvider.AppendMotionMatchingState(Message, Time);
+				break;
+			}
+		default:
+			{
+				// Should not happen
+				checkNoEntry();
+			}
+	}
+
+	return true;
+}
+}}
diff --git a/Engine/Plugins/Experimental/Animation/PoseSearch/Source/Editor/Private/Trace/PoseSearchTraceAnalyzer.h b/Engine/Plugins/Experimental/Animation/PoseSearch/Source/Editor/Private/Trace/PoseSearchTraceAnalyzer.h
new file mode 100644
--- /dev/null
+++ b/Engine/Plugins/Experimental/Animation/PoseSearch/Source/Editor/Private/Trace/PoseSearchTraceAnalyzer.h
@@ -0,0 +1,40 @@
+﻿// Copyright Epic Games, Inc. All Rights Reserved.
+
+#pragma once
+
+#include "CoreMinimal.h"
+#include "Trace/Analyzer.h"
+
+namespace TraceServices
+{
+class IAnalysisSession;
+}
+
+namespace UE { namespace PoseSearch { 
+
+class FTraceProvider;
+
+/**
+* Used to relay trace information to the pose search provider
+*/
+class FTraceAnalyzer : public Trace::IAnalyzer
+{
+public:
+	FTraceAnalyzer(TraceServices::IAnalysisSession& InSession, FTraceProvider& InTraceProvider);
+
+protected:
+	virtual void OnAnalysisBegin(const FOnAnalysisContext& Context) override;
+	virtual bool OnEvent(uint16 RouteId, EStyle Style, const FOnEventContext& Context) override;
+
+private:
+	/** Event types for routing to the provider */
+	enum : uint16
+	{
+		RouteId_MotionMatchingState,
+	};
+
+	TraceServices::IAnalysisSession& Session;
+	FTraceProvider& TraceProvider;
+};
+
+}}
diff --git a/Engine/Plugins/Experimental/Animation/PoseSearch/Source/Editor/Private/Trace/PoseSearchTraceModule.cpp b/Engine/Plugins/Experimental/Animation/PoseSearch/Source/Editor/Private/Trace/PoseSearchTraceModule.cpp
new file mode 100644
--- /dev/null
+++ b/Engine/Plugins/Experimental/Animation/PoseSearch/Source/Editor/Private/Trace/PoseSearchTraceModule.cpp
@@ -0,0 +1,34 @@
+﻿// Copyright Epic Games, Inc. All Rights Reserved.
+
+#include "PoseSearchTraceModule.h"
+#include "PoseSearchTraceAnalyzer.h"
+#include "PoseSearchTraceProvider.h"
+
+namespace UE { namespace PoseSearch {
+
+const FName FTraceModule::ModuleName("PoseSearchTrace");
+
+void FTraceModule::GetModuleInfo(TraceServices::FModuleInfo& OutModuleInfo)
+{
+	OutModuleInfo.Name = ModuleName;
+	OutModuleInfo.DisplayName = TEXT("PoseSearch");
+}
+
+void FTraceModule::OnAnalysisBegin(TraceServices::IAnalysisSession& InSession)
+{
+	// Add our provider and analyzer, starting our systems
+	FTraceProvider* PoseSearchProvider = new FTraceProvider(InSession);
+	InSession.AddProvider(FTraceProvider::ProviderName, PoseSearchProvider);
+	InSession.AddAnalyzer(new FTraceAnalyzer(InSession, *PoseSearchProvider));
+}
+
+void FTraceModule::GetLoggers(TArray<const TCHAR*>& OutLoggers)
+{
+	OutLoggers.Add(TEXT("PoseSearch"));
+}
+
+void FTraceModule::GenerateReports(const TraceServices::IAnalysisSession& Session, const TCHAR* CmdLine, const TCHAR* OutputDirectory)
+{
+}
+
+}}
diff --git a/Engine/Plugins/Experimental/Animation/PoseSearch/Source/Editor/Private/Trace/PoseSearchTraceModule.h b/Engine/Plugins/Experimental/Animation/PoseSearch/Source/Editor/Private/Trace/PoseSearchTraceModule.h
new file mode 100644
--- /dev/null
+++ b/Engine/Plugins/Experimental/Animation/PoseSearch/Source/Editor/Private/Trace/PoseSearchTraceModule.h
@@ -0,0 +1,29 @@
+﻿// Copyright Epic Games, Inc. All Rights Reserved.
+
+#pragma once
+
+#include "CoreMinimal.h"
+#include "TraceServices/ModuleService.h"
+
+namespace UE { namespace PoseSearch {
+
+/**
+* Module used for loading our PoseSearch tracing systems
+*/
+class FTraceModule : public TraceServices::IModule
+{
+public:
+	FTraceModule() = default;
+	virtual ~FTraceModule() = default;
+	
+	const static FName ModuleName;
+
+protected:
+	virtual void GetModuleInfo(TraceServices::FModuleInfo& OutModuleInfo) override;
+	virtual void OnAnalysisBegin(TraceServices::IAnalysisSession& InSession) override;
+	virtual void GetLoggers(TArray<const TCHAR*>& OutLoggers) override;
+	virtual void GenerateReports(const TraceServices::IAnalysisSession& Session, const TCHAR* CmdLine, const TCHAR* OutputDirectory) override;
+	virtual const TCHAR* GetCommandLineArgument() override { return TEXT("posesearchtrace"); }
+};
+
+}}
\ No newline at end of file
diff --git a/Engine/Plugins/Experimental/Animation/PoseSearch/Source/Editor/Private/Trace/PoseSearchTraceProvider.cpp b/Engine/Plugins/Experimental/Animation/PoseSearch/Source/Editor/Private/Trace/PoseSearchTraceProvider.cpp
new file mode 100644
--- /dev/null
+++ b/Engine/Plugins/Experimental/Animation/PoseSearch/Source/Editor/Private/Trace/PoseSearchTraceProvider.cpp
@@ -0,0 +1,36 @@
+﻿// Copyright Epic Games, Inc. All Rights Reserved.
+
+#include "PoseSearchTraceProvider.h"
+
+namespace UE { namespace PoseSearch {
+
+const FName FTraceProvider::ProviderName("PoseSearchTraceProvider");
+
+FTraceProvider::FTraceProvider(TraceServices::IAnalysisSession& InSession) : Session(InSession)
+{
+}
+
+TSet<int32> FTraceProvider::GetMotionMatchingNodeIds(uint64 InAnimInstanceId) const
+{
+	return MotionMatchingStateTimelineStorage.GetNodeIds(InAnimInstanceId);
+}
+
+bool FTraceProvider::ReadMotionMatchingStateTimeline(uint64 InAnimInstanceId, int32 InNodeId, TFunctionRef<void(const FMotionMatchingStateTimeline&)> Callback) const
+{
+	Session.ReadAccessCheck();
+	return MotionMatchingStateTimelineStorage.ReadTimeline(InAnimInstanceId, InNodeId, Callback);
+}
+
+
+void FTraceProvider::AppendMotionMatchingState(const FTraceMotionMatchingStateMessage& InMessage, double InTime)
+{
+	Session.WriteAccessCheck();
+
+	TSharedRef<TraceServices::TPointTimeline<FTraceMotionMatchingStateMessage>> Timeline = MotionMatchingStateTimelineStorage.GetTimeline(Session, InMessage.AnimInstanceId, InMessage.NodeId);
+	Timeline->AppendEvent(InTime, InMessage);
+	
+	Session.UpdateDurationSeconds(InTime);
+}
+
+
+}}
diff --git a/Engine/Plugins/Experimental/Animation/PoseSearch/Source/Editor/Private/Trace/PoseSearchTraceProvider.h b/Engine/Plugins/Experimental/Animation/PoseSearch/Source/Editor/Private/Trace/PoseSearchTraceProvider.h
new file mode 100644
--- /dev/null
+++ b/Engine/Plugins/Experimental/Animation/PoseSearch/Source/Editor/Private/Trace/PoseSearchTraceProvider.h
@@ -0,0 +1,146 @@
+﻿// Copyright Epic Games, Inc. All Rights Reserved.
+
+#pragma once
+
+#include "CoreMinimal.h"
+#include "Model/PointTimeline.h"
+#include "Runtime/Private/Trace/PoseSearchTraceLogger.h"
+#include "TraceServices/Model/AnalysisSession.h"
+
+
+namespace UE { namespace PoseSearch {
+
+// Message types for appending / reading to the timeline
+/** Base message type for common data */
+struct FTraceMessage
+{
+	uint64 AnimInstanceId;
+	int32 NodeId;
+	uint16 FrameCounter;
+};
+
+/** Motion matching state message container */
+struct FTraceMotionMatchingStateMessage : FTraceMessage
+{
+	FTraceMotionMatchingState::EFlags Flags;
+	float ElapsedPoseJumpTime;
+	TArray<float> QueryVector;
+	TArray<float> QueryVectorNormalized;
+	TArray<float> BiasWeights;
+	int32 DbPoseIdx;
+	uint64 DatabaseId;
+};
+
+
+/**
+ * Provider to the widgets for pose search functionality, largely mimicking FAnimationProvider
+ */
+class FTraceProvider : public TraceServices::IProvider
+{
+public:
+	explicit FTraceProvider(TraceServices::IAnalysisSession& InSession);
+
+	/** Get all node ids of motion matching nodes appended to the object's timeline thus far */
+	TSet<int32> GetMotionMatchingNodeIds(uint64 InAnimInstanceId) const;
+
+	/** Read the object-relevant timeline info and provide it via callback */
+	using FMotionMatchingStateTimeline = TraceServices::ITimeline<FTraceMotionMatchingStateMessage>;
+	bool ReadMotionMatchingStateTimeline(uint64 InAnimInstanceId, int32 InNodeId, TFunctionRef<void(const FMotionMatchingStateTimeline&)> Callback) const;
+
+	/** Append to the timeline on our object */
+	void AppendMotionMatchingState(const FTraceMotionMatchingStateMessage& InMessage, double InTime);
+
+	static const FName ProviderName;
+
+private:
+	/**
+	 * Convenience struct for creating timelines per message type.
+	 * We store an array of timelines for every possible AnimInstance Id that gets appended.
+	 */
+	template <class MessageType, class TimelineType = TraceServices::TPointTimeline<MessageType>>
+	struct TTimelineStorage
+	{
+		// Map the timelines as AnimInstanceId -> NodeId -> Timeline
+		using FNodeToTimelineMap = TMap<int32, uint32>;
+		using FAnimInstanceToTimelineMap = TMap<uint64, FNodeToTimelineMap>;
+
+		/** Retrieves the timeline for internal use, creating it if it does not exist */
+		TSharedRef<TimelineType> GetTimeline(TraceServices::IAnalysisSession& InSession, uint64 InAnimInstanceId, int32 InNodeId)
+		{
+			const uint32* TimelineIndex = nullptr;
+			
+			FNodeToTimelineMap* NodeToTimelineMap = AnimInstanceIdToTimelines.Find(InAnimInstanceId);
+			if (NodeToTimelineMap == nullptr)
+			{
+				// Create the timeline map if this is the first time accessing with this anim instance
+				NodeToTimelineMap = &AnimInstanceIdToTimelines.Add(InAnimInstanceId, {});
+			}
+			else
+			{
+				// Anim instance already used, attempt to find the NodeId's timeline
+				TimelineIndex = NodeToTimelineMap->Find(InNodeId);
+			}
+
+			if (TimelineIndex == nullptr)
+			{
+				// Append our timeline to the storage and our object + the storage index to the map
+				TSharedRef<TimelineType> Timeline = MakeShared<TimelineType>(InSession.GetLinearAllocator());
+				const uint32 NewIndex = Timelines.Add(Timeline);
+				NodeToTimelineMap->Add(InNodeId, NewIndex);
+				return Timeline;
+			}
+
+			return Timelines[*TimelineIndex];
+		}
+
+		/** Retrieve a timeline from an anim instance + node and execute the callback */
+		bool ReadTimeline(uint64 InAnimInstanceId, int32 InNodeId, TFunctionRef<void(const TraceServices::ITimeline<MessageType>&)> Callback) const
+		{
+			const FNodeToTimelineMap* NodeToTimelineMap = AnimInstanceIdToTimelines.Find(InAnimInstanceId);
+			if (NodeToTimelineMap == nullptr)
+			{
+				return false;
+			}
+
+			const uint32* TimelineIndex = NodeToTimelineMap->Find(InNodeId);
+			if (TimelineIndex == nullptr || !Timelines.IsValidIndex(*TimelineIndex))
+			{
+				return false;
+			}
+
+			Callback(*Timelines[*TimelineIndex]);
+			return true;
+		}
+
+		TSet<int32> GetNodeIds(uint64 InAnimInstanceId) const
+		{
+			check(AnimInstanceIdToTimelines.Find(InAnimInstanceId));
+			
+			TSet<int32> NodeIds;
+			const FNodeToTimelineMap& NodeToTimelineMap = AnimInstanceIdToTimelines[InAnimInstanceId];
+
+			// For each node that has a timeline, add the node to the set
+			for(const TTuple<int32, uint32>& Pair : NodeToTimelineMap)
+			{
+				NodeIds.Add(Pair.Key);
+			}
+			
+			return NodeIds;
+		}
+
+		/** Maps AnimInstanceIds to a map of NodeIds to timelines. AnimInstanceId -> NodeId -> Timeline */
+		FAnimInstanceToTimelineMap AnimInstanceIdToTimelines;
+
+		/** Timelines per node */
+		TArray<TSharedRef<TimelineType>> Timelines;
+	};
+
+	// Storage for each message type
+	struct FMotionMatchingStateTimelineStorage : TTimelineStorage<FTraceMotionMatchingStateMessage>
+	{
+	} MotionMatchingStateTimelineStorage;
+
+
+	TraceServices::IAnalysisSession& Session;
+};
+}}
diff --git a/Engine/Plugins/Experimental/Animation/PoseSearch/Source/Runtime/Private/AnimNode_MotionMatching.cpp b/Engine/Plugins/Experimental/Animation/PoseSearch/Source/Runtime/Private/AnimNode_MotionMatching.cpp
--- a/Engine/Plugins/Experimental/Animation/PoseSearch/Source/Runtime/Private/AnimNode_MotionMatching.cpp
+++ b/Engine/Plugins/Experimental/Animation/PoseSearch/Source/Runtime/Private/AnimNode_MotionMatching.cpp
@@ -38,10 +38,16 @@ void FAnimNode_MotionMatching::UpdateAssetPlayer(const FAnimationUpdateContext&
 {
 	DECLARE_SCOPE_HIERARCHICAL_COUNTER_ANIMNODE(UpdateAssetPlayer);
 
+
 	GetEvaluateGraphExposedInputs().Execute(Context);
 	// Note: What if the input database changes? That's not being handled at all!
 
 	bool bJumpedToPose = false;
+	FPoseSearchBiasWeights QueryBiasWeights;
+	
+#if UE_POSE_SEARCH_TRACE_ENABLED
+	UE::PoseSearch::FTraceMotionMatchingState State;
+#endif
 
 	const float DeltaTime = Context.GetDeltaTime();
 
@@ -90,7 +96,6 @@ void FAnimNode_MotionMatching::UpdateAssetPlayer(const FAnimationUpdateContext&
 		ComposeQuery(Context);
 
 		// Initialize the pose search query bias weights context
-		FPoseSearchBiasWeights QueryBiasWeights;
 		QueryBiasWeights.Init(BiasWeights, Database->Schema->Layout);
 
 		const FPoseSearchBiasWeightsContext BiasWeightsContext = { &QueryBiasWeights, Database };
@@ -150,6 +155,9 @@ void FAnimNode_MotionMatching::UpdateAssetPlayer(const FAnimationUpdateContext&
 					FollowUpResult.TimeOffsetSeconds = FollowUpAssetTime;
 					JumpToPose(Context, FollowUpResult);
 					bJumpedToPose = true;
+#if UE_POSE_SEARCH_TRACE_ENABLED
+					State.Flags |= UE::PoseSearch::FTraceMotionMatchingState::EFlags::FollowupAnimation;
+#endif
 				}
 			}
 		}
@@ -171,9 +179,15 @@ void FAnimNode_MotionMatching::UpdateAssetPlayer(const FAnimationUpdateContext&
 	Source.Update(Context);
 	
 #if UE_POSE_SEARCH_TRACE_ENABLED
-	UE::PoseSearch::FTraceMotionMatchingState State;
+	if (DbPoseIdx == INDEX_NONE)
+	{
+		return;
+	}
 	State.ElapsedPoseJumpTime = ElapsedPoseJumpTime;
+	// @TODO: Change this to only be the previous query, not persistently updated (i.e. if throttled)?
 	State.QueryVector = ComposedQuery.GetValues();
+	State.QueryVectorNormalized = ComposedQuery.GetNormalizedValues();
+	State.BiasWeights = QueryBiasWeights.Weights;
 	State.DbPoseIdx = DbPoseIdx;
 	State.DatabaseId = FObjectTrace::GetObjectId(Database.Get());
 	UE_TRACE_POSE_SEARCH_MOTION_MATCHING_STATE(Context, State)
diff --git a/Engine/Plugins/Experimental/Animation/PoseSearch/Source/Runtime/Private/PoseSearch.cpp b/Engine/Plugins/Experimental/Animation/PoseSearch/Source/Runtime/Private/PoseSearch.cpp
--- a/Engine/Plugins/Experimental/Animation/PoseSearch/Source/Runtime/Private/PoseSearch.cpp
+++ b/Engine/Plugins/Experimental/Animation/PoseSearch/Source/Runtime/Private/PoseSearch.cpp
@@ -477,6 +477,7 @@ void FPoseSearchBiasWeights::BindSemanticWeight(float Weight, const FPoseSearchF
 		const int32 FirstValueIdx = Feature.ValueOffset;
 		const int32 NumValues = UE::PoseSearch::GetFeatureTypeTraits(FeatureType).NumFloats;
 
+		// NOTE: This enumeration method is currently assumed by the PoseSearch Debugger
 		for (int32 Idx = 0; Idx < NumValues; ++Idx)
 		{
 			Weights[FirstValueIdx + Idx] = Weight;
diff --git a/Engine/Plugins/Experimental/Animation/PoseSearch/Source/Runtime/Private/Trace/PoseSearchTraceLogger.cpp b/Engine/Plugins/Experimental/Animation/PoseSearch/Source/Runtime/Private/Trace/PoseSearchTraceLogger.cpp
--- a/Engine/Plugins/Experimental/Animation/PoseSearch/Source/Runtime/Private/Trace/PoseSearchTraceLogger.cpp
+++ b/Engine/Plugins/Experimental/Animation/PoseSearch/Source/Runtime/Private/Trace/PoseSearchTraceLogger.cpp
@@ -18,6 +18,8 @@ UE_TRACE_EVENT_BEGIN(PoseSearch, MotionMatchingState)
 	UE_TRACE_EVENT_FIELD(uint64, DatabaseId)
 	UE_TRACE_EVENT_FIELD(int32, DbPoseIdx)
 	UE_TRACE_EVENT_FIELD(float[], QueryVector)
+	UE_TRACE_EVENT_FIELD(float[], QueryVectorNormalized)
+	UE_TRACE_EVENT_FIELD(float[], BiasWeights)
 UE_TRACE_EVENT_END()
 
 namespace UE { namespace PoseSearch {
@@ -61,6 +63,8 @@ void FTraceMotionMatchingState::Output(const FAnimationBaseContext& InContext, c
 		// Cast to uint32 for event message
 		<< MotionMatchingState.Flags(static_cast<uint32>(State.Flags))
 		<< MotionMatchingState.QueryVector(State.QueryVector.GetData(), State.QueryVector.Num())
+		<< MotionMatchingState.QueryVectorNormalized(State.QueryVectorNormalized.GetData(), State.QueryVectorNormalized.Num())
+		<< MotionMatchingState.BiasWeights(State.BiasWeights.GetData(), State.BiasWeights.Num())
 		<< MotionMatchingState.DbPoseIdx(State.DbPoseIdx)
 		<< MotionMatchingState.DatabaseId(State.DatabaseId);
 }
diff --git a/Engine/Plugins/Experimental/Animation/PoseSearch/Source/Runtime/Private/Trace/PoseSearchTraceLogger.h b/Engine/Plugins/Experimental/Animation/PoseSearch/Source/Runtime/Private/Trace/PoseSearchTraceLogger.h
--- a/Engine/Plugins/Experimental/Animation/PoseSearch/Source/Runtime/Private/Trace/PoseSearchTraceLogger.h
+++ b/Engine/Plugins/Experimental/Animation/PoseSearch/Source/Runtime/Private/Trace/PoseSearchTraceLogger.h
@@ -3,7 +3,6 @@
 #pragma once
 
 #include "ObjectTrace.h"
-#include "PoseSearch/PoseSearch.h"
 
 // Enable this if object tracing is enabled, mimics animation tracing
 #define UE_POSE_SEARCH_TRACE_ENABLED OBJECT_TRACE_ENABLED
@@ -42,6 +41,10 @@ struct POSESEARCH_API FTraceMotionMatchingState
 
 	/** Feature value array of the query data */
 	TArrayView<const float> QueryVector;
+	TArrayView<const float> QueryVectorNormalized;
+
+	/** Weight bias params for motion matching queries */
+	TArrayView<const float> BiasWeights;
 
 	/** Index of the pose in our database */
 	int32 DbPoseIdx = 0;
@@ -49,6 +52,9 @@ struct POSESEARCH_API FTraceMotionMatchingState
 	/** Object Id of the database asset */
 	uint64 DatabaseId = 0;
 
+	/** Node Id of the motion matching node associated with this message */
+	int32 NodeId = 0;
+
 	/** Output the current state info to the logger */
 	static void Output(const FAnimationBaseContext& InContext, const FTraceMotionMatchingState& State);
 	
