Fix for Pose Search debugger to not break on empty Motion Matching nodes

#rnx

diff --git a/Engine/Plugins/Experimental/Animation/PoseSearch/Source/Editor/Private/Trace/PoseSearchTraceProvider.h b/Engine/Plugins/Experimental/Animation/PoseSearch/Source/Editor/Private/Trace/PoseSearchTraceProvider.h
--- a/Engine/Plugins/Experimental/Animation/PoseSearch/Source/Editor/Private/Trace/PoseSearchTraceProvider.h
+++ b/Engine/Plugins/Experimental/Animation/PoseSearch/Source/Editor/Private/Trace/PoseSearchTraceProvider.h
@@ -114,9 +114,11 @@ private:
 
 		TSet<int32> GetNodeIds(uint64 InAnimInstanceId) const
 		{
-			check(AnimInstanceIdToTimelines.Find(InAnimInstanceId));
-			
 			TSet<int32> NodeIds;
+			if (!AnimInstanceIdToTimelines.Find(InAnimInstanceId))
+			{
+				return NodeIds;
+			}
 			const FNodeToTimelineMap& NodeToTimelineMap = AnimInstanceIdToTimelines[InAnimInstanceId];
 
 			// For each node that has a timeline, add the node to the set
