PoseSearch Debugger additions / fixes

- Mutliple motion-matching nodes now fully supported, including at various frames. This includes a database selection view, and the ability to return to that view to make your selection.
- State reflection of the debugger itself as overlays, i.e. messages such as "Recording..." or "No data available for this frame", etc
- Massive performance increases statically & while scrubbing

#rnx

diff --git a/Engine/Plugins/Experimental/Animation/PoseSearch/Source/Editor/Private/PoseSearchDebugger.cpp b/Engine/Plugins/Experimental/Animation/PoseSearch/Source/Editor/Private/PoseSearchDebugger.cpp
--- a/Engine/Plugins/Experimental/Animation/PoseSearch/Source/Editor/Private/PoseSearchDebugger.cpp
+++ b/Engine/Plugins/Experimental/Animation/PoseSearch/Source/Editor/Private/PoseSearchDebugger.cpp
@@ -6,7 +6,6 @@
 #include "RewindDebuggerInterface/Public/IRewindDebugger.h"
 #include "ObjectTrace.h"
 #include "TraceServices/Model/AnalysisSession.h"
-#include "TraceServices/Model/Frames.h"
 #include "Trace/PoseSearchTraceProvider.h"
 #include "Widgets/Layout/SScrollBox.h"
 #include "Widgets/Layout/SWidgetSwitcher.h"
@@ -14,6 +13,7 @@
 #include "Widgets/Input/SButton.h"
 #include "Widgets/Images/SImage.h"
 #include "Styling/SlateIconFinder.h"
+#include "TraceServices/Model/Frames.h"
 
 #define LOCTEXT_NAMESPACE "PoseSearchDebugger"
 
@@ -36,7 +36,7 @@ class FDebuggerDatabaseRowData : public TSharedFromThis<FDebuggerDatabaseRowData
 public:
 	FDebuggerDatabaseRowData() = default;
 	FDebuggerDatabaseRowData(
-	    uint32 InPoseIdx, 
+	    int32 InPoseIdx, 
 	    const FString& InAnimSequenceName,
 	    float InTime,
 	    float InLength,
@@ -53,8 +53,8 @@ public:
 		, TrajectoryScore(InTrajectoryScore)
 	{
 	}
-	
-	uint32 PoseIdx = 0;
+
+	int32 PoseIdx = 0;
 	FString AnimSequenceName = "";
 	float Time = 0.0f;
 	float Length = 0.0f;
@@ -68,7 +68,7 @@ namespace DebuggerDatabaseColumns
 	/** Column struct to represent each column in the debugger database */
 	struct IColumn : TSharedFromThis<IColumn>
 	{
-		IColumn(int32 InSortIndex, bool InEnabled = true)
+		explicit IColumn(int32 InSortIndex, bool InEnabled = true)
 			: SortIndex(InSortIndex)
 			, bEnabled(InEnabled)
 		{
@@ -89,21 +89,22 @@ namespace DebuggerDatabaseColumns
 		using FSortPredicate = TFunctionRef<bool(const FRowDataRef&, const FRowDataRef&)>;
 		
 		/** Sort predicate to sort list in ascending order by this column */
-		virtual FSortPredicate GetSortPredicateAscending() = 0;
+		virtual const FSortPredicate& GetSortPredicateAscending() const = 0;
 		/** Sort predicate to sort list in descending order by this column */
-		virtual FSortPredicate GetSortPredicateDescending() = 0;
-		/** Text to display associated with this column from given row data  */
-		virtual TSharedRef<SWidget> GenerateTextWidget(const FRowDataRef& RowData) const { return SNullWidget::NullWidget; }
-
-	protected:
-		TSharedRef<STextBlock> GenerateTextWidgetBase() const
+		virtual const FSortPredicate& GetSortPredicateDescending() const = 0;
+		
+		TSharedRef<STextBlock> GenerateTextWidget(const FRowDataRef& RowData) const
 		{
 			static FSlateFontInfo RowFont = FEditorStyle::Get().GetFontStyle("DetailsView.CategoryTextStyle");
 			
 			return SNew(STextBlock)
             	.Font(RowFont)
+				.Text_Lambda([this, RowData] { return GetRowText(RowData); })
             	.Justification(ETextJustify::Center);
 		}
+
+	protected:
+		virtual FText GetRowText(const FRowDataRef& Row) const = 0;
 	};
 
 	struct FPoseIdx : IColumn
@@ -112,23 +113,21 @@ namespace DebuggerDatabaseColumns
 		static const FName Name;
 		virtual FName GetName() const override { return Name; }
 
-		virtual FSortPredicate GetSortPredicateAscending() override
+		virtual const FSortPredicate& GetSortPredicateAscending() const override
 		{
 			static FSortPredicate Predicate = [](const FRowDataRef& Row0, const FRowDataRef& Row1) -> bool { return Row0->PoseIdx < Row1->PoseIdx; };
 			return Predicate;
 		}
 		
-		virtual FSortPredicate GetSortPredicateDescending() override
+		virtual const FSortPredicate& GetSortPredicateDescending() const override
 		{
 			static FSortPredicate Predicate = [](const FRowDataRef& Row0, const FRowDataRef& Row1) -> bool { return Row0->PoseIdx >= Row1->PoseIdx; };
 			return Predicate;
 		}
 
-		virtual TSharedRef<SWidget> GenerateTextWidget(const FRowDataRef& RowData) const override
+		virtual FText GetRowText(const FRowDataRef& Row) const override
 		{
-			TSharedRef<STextBlock> Text = GenerateTextWidgetBase();
-			Text->SetText(TAttribute<FText>::Create([RowData](){ return FText::AsNumber(RowData->PoseIdx); }));
-			return Text;
+			return FText::AsNumber(Row->PoseIdx);
 		}
 	};
 	const FName FPoseIdx::Name = "PoseIdx";
@@ -139,23 +138,21 @@ namespace DebuggerDatabaseColumns
 		static const FName Name;
 		virtual FName GetName() const override { return Name; }
 		
-		virtual FSortPredicate GetSortPredicateAscending() override
+		virtual const FSortPredicate& GetSortPredicateAscending() const override
 		{
 			static FSortPredicate Predicate = [](const FRowDataRef& Row0, const FRowDataRef& Row1) -> bool { return Row0->AnimSequenceName < Row1->AnimSequenceName; };
 			return Predicate;
 		}
 
-		virtual FSortPredicate GetSortPredicateDescending() override
+		virtual const FSortPredicate& GetSortPredicateDescending() const override
 		{
 			static FSortPredicate Predicate = [](const FRowDataRef& Row0, const FRowDataRef& Row1) -> bool { return Row0->AnimSequenceName >= Row1->AnimSequenceName; };
 			return Predicate;
 		}
 
-		virtual TSharedRef<SWidget> GenerateTextWidget(const FRowDataRef& RowData) const override
+		virtual FText GetRowText(const FRowDataRef& Row) const override
 		{
-			TSharedRef<STextBlock> Text = GenerateTextWidgetBase();
-			Text->SetText(TAttribute<FText>::Create([RowData](){ return FText::FromString(RowData->AnimSequenceName); }));
-			return Text;
+			return FText::FromString(Row->AnimSequenceName);
 		}
 	};
 	const FName FAnimSequenceName::Name = "AnimSequence";
@@ -166,23 +163,21 @@ namespace DebuggerDatabaseColumns
 		static const FName Name;
 		virtual FName GetName() const override { return Name; }
 		
-		virtual FSortPredicate GetSortPredicateAscending() override
+		virtual const FSortPredicate& GetSortPredicateAscending() const override
 		{
 			static FSortPredicate Predicate = [](const FRowDataRef& Row0, const FRowDataRef& Row1) -> bool { return Row0->Time < Row1->Time; };
 			return Predicate;
 		}
 		
-		virtual FSortPredicate GetSortPredicateDescending() override
+		virtual const FSortPredicate& GetSortPredicateDescending() const override
 		{
 			static FSortPredicate Predicate = [](const FRowDataRef& Row0, const FRowDataRef& Row1) -> bool { return Row0->Time >= Row1->Time; };
 			return Predicate;
 		}
 
-		virtual TSharedRef<SWidget> GenerateTextWidget(const FRowDataRef& RowData) const override
+		virtual FText GetRowText(const FRowDataRef& Row) const override
 		{
-			TSharedRef<STextBlock> Text = GenerateTextWidgetBase();
-			Text->SetText(TAttribute<FText>::Create([RowData](){ return FText::AsNumber(RowData->Time); }));
-			return Text;
+			return FText::AsNumber(Row->Time);
 		}
 	};
 	const FName FTime::Name = "Time";
@@ -193,23 +188,21 @@ namespace DebuggerDatabaseColumns
 		static const FName Name;
 		virtual FName GetName() const override { return Name; }
 		
-		virtual FSortPredicate GetSortPredicateAscending() override
+		virtual const FSortPredicate& GetSortPredicateAscending() const override
 		{
 			static FSortPredicate Predicate = [](const FRowDataRef& Row0, const FRowDataRef& Row1) -> bool { return Row0->Length < Row1->Length; };
 			return Predicate;
 		}
 
-		virtual FSortPredicate GetSortPredicateDescending() override
+		virtual const FSortPredicate& GetSortPredicateDescending() const override
 		{
 			static FSortPredicate Predicate = [](const FRowDataRef& Row0, const FRowDataRef& Row1) -> bool { return Row0->Length >= Row1->Length; };
 			return Predicate;
 		}
 
-		virtual TSharedRef<SWidget> GenerateTextWidget(const FRowDataRef& RowData) const override
+		virtual FText GetRowText(const FRowDataRef& Row) const override
 		{
-			TSharedRef<STextBlock> Text = GenerateTextWidgetBase();
-			Text->SetText(TAttribute<FText>::Create([RowData](){ return FText::AsNumber(RowData->Length); }));
-			return Text;
+			return FText::AsNumber(Row->Length);
 		}
 	};
 	const FName FLength::Name = "Length";
@@ -220,24 +213,22 @@ namespace DebuggerDatabaseColumns
 		static const FName Name;
 		virtual FName GetName() const override { return Name; }
 		
-		virtual FSortPredicate GetSortPredicateAscending() override
+		virtual const FSortPredicate& GetSortPredicateAscending() const override
 		{
 			static FSortPredicate Predicate = [](const FRowDataRef& Row0, const FRowDataRef& Row1) -> bool { return Row0->Score < Row1->Score; };
 			return Predicate;
 		}
 
-		virtual FSortPredicate GetSortPredicateDescending() override
+		virtual const FSortPredicate& GetSortPredicateDescending() const override
 		{
 			static FSortPredicate Predicate = [](const FRowDataRef& Row0, const FRowDataRef& Row1) -> bool { return Row0->Score >= Row1->Score; };
 			return Predicate;
 		}
 		
-		virtual TSharedRef<SWidget> GenerateTextWidget(const FRowDataRef& RowData) const override
-		{
-			TSharedRef<STextBlock> Text = GenerateTextWidgetBase();
-			Text->SetText(TAttribute<FText>::Create([RowData](){ return FText::AsNumber(RowData->Score); }));
-			return Text;
-		}
+		virtual FText GetRowText(const FRowDataRef& Row) const override
+        {
+        	return FText::AsNumber(Row->Score);
+        }
 	};
 	const FName FScore::Name = "Score";
 
@@ -248,23 +239,21 @@ namespace DebuggerDatabaseColumns
 		static const FName Name;
 		virtual FName GetName() const override { return Name; }
 		
-		virtual FSortPredicate GetSortPredicateAscending() override
+		virtual const FSortPredicate& GetSortPredicateAscending() const override
 		{
 			static FSortPredicate Predicate = [](const FRowDataRef& Row0, const FRowDataRef& Row1) -> bool { return Row0->PoseScore < Row1->PoseScore; };
 			return Predicate;
 		}
 
-		virtual FSortPredicate GetSortPredicateDescending() override
+		virtual const FSortPredicate& GetSortPredicateDescending() const override
 		{
 			static FSortPredicate Predicate = [](const FRowDataRef& Row0, const FRowDataRef& Row1) -> bool { return Row0->PoseScore >= Row1->PoseScore; };
 			return Predicate;
 		}
 		
-		virtual TSharedRef<SWidget> GenerateTextWidget(const FRowDataRef& RowData) const override
+		virtual FText GetRowText(const FRowDataRef& Row) const override
 		{
-			TSharedRef<STextBlock> Text = GenerateTextWidgetBase();
-			Text->SetText(TAttribute<FText>::Create([RowData](){ return FText::AsNumber(RowData->PoseScore); }));
-			return Text;
+			return FText::AsNumber(Row->PoseScore);
 		}
 	};
 	const FName FPoseScore::Name = "Pose Score";
@@ -276,24 +265,22 @@ namespace DebuggerDatabaseColumns
 		static const FName Name;
 		virtual FName GetName() const override { return Name; }
 		
-		virtual FSortPredicate GetSortPredicateAscending() override
+		virtual const FSortPredicate& GetSortPredicateAscending() const override
 		{
 			static FSortPredicate Predicate = [](const FRowDataRef& Row0, const FRowDataRef& Row1) -> bool { return Row0->TrajectoryScore < Row1->TrajectoryScore; };
 			return Predicate;
 		}
 
-		virtual FSortPredicate GetSortPredicateDescending() override
+		virtual const FSortPredicate& GetSortPredicateDescending() const override
 		{
 			static FSortPredicate Predicate = [](const FRowDataRef& Row0, const FRowDataRef& Row1) -> bool { return Row0->TrajectoryScore >= Row1->TrajectoryScore; };
 			return Predicate;
 		}
 		
-		virtual TSharedRef<SWidget> GenerateTextWidget(const FRowDataRef& RowData) const override
+		virtual FText GetRowText(const FRowDataRef& Row) const override
 		{
-			TSharedRef<STextBlock> Text = GenerateTextWidgetBase();
-			Text->SetText(TAttribute<FText>::Create([RowData](){ return FText::AsNumber(RowData->TrajectoryScore); }));
-			return Text;
-		}
+			return FText::AsNumber(Row->TrajectoryScore);
+        }
 	};
 	const FName FTrajectoryScore::Name = "Trajectory Score";
 }
@@ -345,7 +332,6 @@ class SDebuggerDatabaseRow : public SMultiColumnTableRow<TSharedRef<FDebuggerDat
 			.HAlign(HAlign_Center)
 			.VAlign(VAlign_Center)
 			.BorderImage(RowBrush)
-			.Padding(0.0f)
 			[
 				SNew(SHorizontalBox)
 				+ SHorizontalBox::Slot()
@@ -369,9 +355,31 @@ class SDebuggerDatabaseRow : public SMultiColumnTableRow<TSharedRef<FDebuggerDat
 };
 FGetColumnDelegate SDebuggerDatabaseRow::GetColumn;
 
+class SDebuggerMessageBox : public SCompoundWidget
+{
+	SLATE_BEGIN_ARGS(SDebuggerMessageBox) {}
+	SLATE_END_ARGS()
+
+	void Construct(const FArguments& InArgs, const FString& Message)
+	{
+		ChildSlot
+		[
+			SNew(SVerticalBox)
+			+ SVerticalBox::Slot()
+			.HAlign(HAlign_Center)
+			.VAlign(VAlign_Center)
+			[
+				SNew(STextBlock)
+				.Text(FText::FromString(Message))
+				.Font(FEditorStyle::Get().GetFontStyle("DetailsView.CategoryFontStyle"))
+			]
+		];
+	}
+};
+
 void SDebuggerDatabaseView::Update(const FTraceMotionMatchingStateMessage* State, const UPoseSearchDatabase* Database)
 {
-	PopulateRows(State, Database);
+	UpdateRows(State, Database);
 }
 
 void SDebuggerDatabaseView::RefreshColumns()
@@ -502,7 +510,7 @@ void SDebuggerDatabaseView::CreateRows(const UPoseSearchDatabase* Database)
 	ActiveView.Rows.Add(MakeShared<FDebuggerDatabaseRowData>());
 }
 
-void SDebuggerDatabaseView::PopulateRows(const FTraceMotionMatchingStateMessage* State, const UPoseSearchDatabase* Database)
+void SDebuggerDatabaseView::UpdateRows(const FTraceMotionMatchingStateMessage* State, const UPoseSearchDatabase* Database)
 {
 	if (DatabaseView.Rows.IsEmpty())
 	{
@@ -516,6 +524,7 @@ void SDebuggerDatabaseView::PopulateRows(const FTraceMotionMatchingStateMessage*
     BiasWeights.Weights = State->BiasWeights;
 
 	const FPoseSearchFeatureVectorLayout& Layout = Database->Schema->Layout;
+	const FPoseSearchIndex& SearchIndex = Database->SearchIndex;
 
 	// Reverse engineer weight params from weights array
 	auto ExtractWeight = [](
@@ -545,7 +554,7 @@ void SDebuggerDatabaseView::PopulateRows(const FTraceMotionMatchingStateMessage*
 	Params.TrajectoryLinearVelocityWeight = 0.0f;
 	
 	FPoseSearchBiasWeights PoseBiasWeights;
-	PoseBiasWeights.Init(Params, Database->Schema->Layout);
+	PoseBiasWeights.Init(Params, Layout);
 
 	// Zeroed pose for trajectory score
 	Params.PosePositionWeight = 0.0f;
@@ -554,7 +563,7 @@ void SDebuggerDatabaseView::PopulateRows(const FTraceMotionMatchingStateMessage*
 	Params.TrajectoryLinearVelocityWeight = ExtractWeight(Layout, BiasWeights, EPoseSearchFeatureType::LinearVelocity, true);
 	
 	FPoseSearchBiasWeights TrajectoryBiasWeights;
-	TrajectoryBiasWeights.Init(Params, Database->Schema->Layout);
+	TrajectoryBiasWeights.Init(Params, Layout);
 	
 	const FPoseSearchBiasWeightsContext BiasWeightsContext { &BiasWeights, Database };
 	const FPoseSearchBiasWeightsContext PoseBiasWeightsContext { &PoseBiasWeights, Database };
@@ -562,10 +571,10 @@ void SDebuggerDatabaseView::PopulateRows(const FTraceMotionMatchingStateMessage*
 	
 	for(const TSharedRef<FDebuggerDatabaseRowData>& Row : DatabaseView.Rows)
 	{
-		const uint32 PoseIdx = Row->PoseIdx;
-		Row->Score = ComparePoses(Database->SearchIndex, PoseIdx, State->QueryVectorNormalized, &BiasWeightsContext);
-		Row->PoseScore = ComparePoses(Database->SearchIndex, PoseIdx, State->QueryVectorNormalized, &PoseBiasWeightsContext);
-		Row->TrajectoryScore = ComparePoses(Database->SearchIndex, PoseIdx, State->QueryVectorNormalized, &TrajectoryBiasWeightsContext);
+		const int32 PoseIdx = Row->PoseIdx;
+		Row->Score = ComparePoses(SearchIndex, PoseIdx, State->QueryVectorNormalized, &BiasWeightsContext);
+		Row->PoseScore = ComparePoses(SearchIndex, PoseIdx, State->QueryVectorNormalized, &PoseBiasWeightsContext);
+		Row->TrajectoryScore = ComparePoses(SearchIndex, PoseIdx, State->QueryVectorNormalized, &TrajectoryBiasWeightsContext);
 
 		// If we are on the active pose for the frame
 		if (PoseIdx == State->DbPoseIdx)
@@ -577,8 +586,6 @@ void SDebuggerDatabaseView::PopulateRows(const FTraceMotionMatchingStateMessage*
 	SortDatabaseRows();
 }
 
-
-
 TSharedRef<ITableRow> SDebuggerDatabaseView::HandleGenerateDatabaseRow(TSharedRef<FDebuggerDatabaseRowData> Item, const TSharedRef<STableViewBase>& OwnerTable) const
 {
 	return SNew(SDebuggerDatabaseRow, OwnerTable, Item, DatabaseView.RowStyle, &DatabaseView.RowBrush, FMargin(0.0f, 2.0f, 6.0f, 2.0f));
@@ -604,10 +611,9 @@ void SDebuggerDatabaseView::Construct(const FArguments& InArgs)
 	AddColumn(MakeShared<FTrajectoryScore>(6));
 
 	// Active Row
-
-	// Used for spacing
 	ActiveView.HeaderRow = SNew(SHeaderRow);
 
+	// Used for spacing
 	ActiveView.ScrollBar =
 		SNew(SScrollBar)
 		.Orientation(Orient_Vertical)
@@ -621,8 +627,7 @@ void SDebuggerDatabaseView::Construct(const FArguments& InArgs)
 		.OnGenerateRow(this, &SDebuggerDatabaseView::HandleGenerateActiveRow)
 		.ExternalScrollbar(ActiveView.ScrollBar)
 		.SelectionMode(ESelectionMode::SingleToggle)
-		.ConsumeMouseWheel(EConsumeMouseWheel::Never)
-		.ItemHeight(1);
+		.ConsumeMouseWheel(EConsumeMouseWheel::Never);
 
 	ActiveView.RowStyle = FEditorStyle::GetWidgetStyle<FTableRowStyle>("TableView.Row");
 	ActiveView.RowBrush = *FEditorStyle::GetBrush("DetailsView.CategoryTop");
@@ -641,10 +646,8 @@ void SDebuggerDatabaseView::Construct(const FArguments& InArgs)
 			.HeaderRow(DatabaseView.HeaderRow.ToSharedRef())
 			.OnGenerateRow(this, &SDebuggerDatabaseView::HandleGenerateDatabaseRow)
 			.ExternalScrollbar(DatabaseView.ScrollBar)
-			.ScrollbarVisibility(EVisibility::Visible)
 			.SelectionMode(ESelectionMode::Multi)
-			.ConsumeMouseWheel(EConsumeMouseWheel::WhenScrollingPossible)
-			.ItemHeight(24);
+			.ConsumeMouseWheel(EConsumeMouseWheel::WhenScrollingPossible);
 
 	DatabaseView.RowStyle = FEditorStyle::GetWidgetStyle<FTableRowStyle>("TableView.Row");
 	// Set selected color to white to retain visibility when multi-selecting
@@ -665,8 +668,8 @@ void SDebuggerDatabaseView::Construct(const FArguments& InArgs)
 				// Active Row text tab
 				SNew(SVerticalBox)
 				+ SVerticalBox::Slot()
-				.Padding(0.0f)
 				.AutoHeight()
+				.Padding(0.0f)
 				[
 					SNew(SHorizontalBox)
 					+ SHorizontalBox::Slot()
@@ -689,7 +692,7 @@ void SDebuggerDatabaseView::Construct(const FArguments& InArgs)
 
 				// Active row list view with scroll bar
 				+ SVerticalBox::Slot()
-				.Padding(0.0f)
+				
 				.AutoHeight()
 				[
 					SNew(SHorizontalBox)
@@ -700,8 +703,9 @@ void SDebuggerDatabaseView::Construct(const FArguments& InArgs)
 					.Padding(0.0f)
 					[
 						SNew(SBorder)
-						.Padding(0.0f)
+						
 						.BorderImage(FEditorStyle::GetBrush("NoBorder"))
+						.Padding(0.0f)
 						[
 							ActiveView.ListView.ToSharedRef()
 						]
@@ -721,7 +725,6 @@ void SDebuggerDatabaseView::Construct(const FArguments& InArgs)
 				// Database view text tab
 				SNew(SVerticalBox)
 				+ SVerticalBox::Slot()
-				.Padding(0.0f)
 				[
 					SNew(SHorizontalBox)
 					+ SHorizontalBox::Slot()
@@ -743,10 +746,8 @@ void SDebuggerDatabaseView::Construct(const FArguments& InArgs)
 					
 					+ SHorizontalBox::Slot()
 					.HAlign(HAlign_Fill)
-					.Padding(0.0f)
 					[
 						SNew(SBorder)
-						.Padding(0.0f)
 						.BorderImage(&DatabaseView.RowStyle.EvenRowBackgroundBrush)
 					]
 				]
@@ -754,8 +755,8 @@ void SDebuggerDatabaseView::Construct(const FArguments& InArgs)
 
 				// Gray line below the tab 
 				+ SVerticalBox::Slot()
-				.Padding(0.0f)
 				.AutoHeight()
+				.Padding(0.0f)
 				[
 					SNew(SBorder)
 					.BorderImage(FEditorStyle::GetBrush("DetailsView.CategoryTop"))
@@ -766,20 +767,14 @@ void SDebuggerDatabaseView::Construct(const FArguments& InArgs)
 				
 			
 				+ SVerticalBox::Slot()
-				.Padding(0.0f)
 				[
 					SNew(SHorizontalBox)
 					+ SHorizontalBox::Slot()
 					.Padding(0.0f)
 					[
-						SNew(SScrollBox)
-						.Orientation(Orient_Vertical)
-						.ExternalScrollbar(DatabaseView.ScrollBar)
-						.ScrollBarPadding(0.0f)
-						.NavigationScrollPadding(0.0f)
-						.ScrollBarVisibility(EVisibility::Hidden)
-						.AllowOverscroll(EAllowOverscroll::No)
-						+ SScrollBox::Slot()
+						SNew(SBorder)
+						.BorderImage(FEditorStyle::GetBrush("NoBorder"))
+						.Padding(0.0f)
 						[
 							DatabaseView.ListView.ToSharedRef()
 						]
@@ -787,7 +782,6 @@ void SDebuggerDatabaseView::Construct(const FArguments& InArgs)
 					
 					+ SHorizontalBox::Slot()
 					.AutoWidth()
-					.Padding(0.0f)
 					[
 						DatabaseView.ScrollBar.ToSharedRef()
 					]
@@ -837,6 +831,23 @@ void SDebuggerView::Construct(const FArguments& InArgs, uint64 InAnimInstanceId)
 	PoseSearchDatabase = InArgs._PoseSearchDatabase;
 	OnSelectionChanged = InArgs._OnSelectionChanged;
 	IsPIESimulating = InArgs._IsPIESimulating;
+	IsRecording = InArgs._IsRecording;
+	RecordingDuration = InArgs._RecordingDuration;
+	ActiveNodesNum = InArgs._ActiveNodesNum;
+	OnUpdate = InArgs._OnUpdate;
+
+	// Validate the existence of the passed getters
+	check(MotionMatchingNodeIds.IsBound());
+    check(MotionMatchingState.IsBound());
+	check(Reflection.IsBound());
+	check(PoseSearchDatabase.IsBound());
+	check(OnSelectionChanged.IsBound());
+	check(IsPIESimulating.IsBound())
+    check(IsRecording.IsBound());
+    check(RecordingDuration.IsBound());
+	check(ActiveNodesNum.IsBound());
+	check(OnUpdate.IsBound());
+	
 	AnimInstanceId = InAnimInstanceId;
 	SelectedNode = -1;
 
@@ -849,24 +860,9 @@ void SDebuggerView::Construct(const FArguments& InArgs, uint64 InAnimInstanceId)
 		.HAlign(HAlign_Fill)
 		[
 			SAssignNew(Switcher, SWidgetSwitcher)
+			.WidgetIndex_Raw(this, &SDebuggerView::SelectView)
 
-			// [0] Box that covers everything when recording
-			+ SWidgetSwitcher::Slot()
-			.Padding(0.0f)
-			[
-				SAssignNew(SimulatingView, SVerticalBox)
-				
-				+ SVerticalBox::Slot()
-				.HAlign(HAlign_Center)
-                .VAlign(VAlign_Center)
-				[
-					SNew(STextBlock)
-					.Text(FText::FromString("Select a frame to continue..."))
-					.Font(FEditorStyle::Get().GetFontStyle("DetailsView.CategoryFontStyle"))
-				]
-			]
-
-			// [1] Selection view before node selection is made
+			// [0] Selection view before node selection is made
 			+ SWidgetSwitcher::Slot()
 			.Padding(40.0f)
 			.HAlign(HAlign_Fill)
@@ -875,57 +871,85 @@ void SDebuggerView::Construct(const FArguments& InArgs, uint64 InAnimInstanceId)
 				SAssignNew(SelectionView, SVerticalBox)
 			]
 
-			// [2] Node selected; node debugger view
+			// [1] Node selected; node debugger view
 			+ SWidgetSwitcher::Slot()
-			.Padding(0.0f)
 			[
 				GenerateNodeDebuggerView()
 			]
+
+			// [2] Occluding message box when stopped (no recording)
+			+ SWidgetSwitcher::Slot()
+			[
+				SNew(SDebuggerMessageBox, "Record gameplay to begin debugging")
+			]
+
+			// [3] Occluding message box when recording
+			+ SWidgetSwitcher::Slot()
+			[
+				SNew(SDebuggerMessageBox, "Recording...")
+			]
+			
+			// [4] Occluding message box when there is no data for the selected MM node
+			+ SWidgetSwitcher::Slot()
+			[
+				GenerateNoDataMessageView()
+			]
 		]
 	];
-	
 }
 
 void SDebuggerView::SetTimeMarker(double InTimeMarker)
 {
-	Switcher->SetActiveWidgetIndex(SwitcherViewType);
 	if (IsPIESimulating.Get())
 	{
-		SwitcherViewType = Waiting;
 		return;
 	}
 
 	const bool SameTime = FMath::Abs(InTimeMarker - TimeMarker) < DOUBLE_SMALL_NUMBER;
-	if (SameTime)
+	TimeMarker = InTimeMarker;
+
+	// We haven't reached the update point yet
+	if (CurrentConsecutiveFrames < ConsecutiveFramesUpdateThreshold)
 	{
+		// If we're on the same time marker, it is consecutive
+		if (SameTime)
+		{
+			++CurrentConsecutiveFrames;
+		}
+		return;
+	}
+	// New frame after having updated, reset consecutive frames count and start counting again
+	if (!SameTime)
+	{
+		CurrentConsecutiveFrames = 0;
+		bUpdated = false;
 		return;
 	}
-	TimeMarker = InTimeMarker;
 
-	UpdateViews();
+	// Haven't updated since passing through frame gate, update once
+	if (!bUpdated)
+	{
+		UpdateViews();
+    	bUpdated = true;	
+	}
 }
 
 void SDebuggerView::UpdateViews()
 {
-	const TSet<int32> NodeIds = MotionMatchingNodeIds.Get();
-	if (NodeIds.IsEmpty())
-	{
-		return;
-	}
+	OnUpdate.Execute(AnimInstanceId);
 	
 	// Update selection view if no node selected
 	if (SelectedNode == -1)
 	{
-		SwitcherViewType = Selection;
+		TSet<int32> NodeIds = TSet<int32>(*MotionMatchingNodeIds.Get());
 
 		// If we have a new set of nodes
-		if (!NodeIds.Difference(ActiveNodes).IsEmpty())
+		if (!NodeIds.Difference(StoredNodes).IsEmpty())
 		{
 			// Only one node active, bypass selection view
 			if (NodeIds.Num() == 1)
 			{
 				SelectedNode = *NodeIds.begin();
-				OnSelectionChanged.Execute(AnimInstanceId, SelectedNode);
 				UpdateViews();
 			}
 			// Create selection view with buttons for each node, displaying the database name
@@ -934,7 +958,7 @@ void SDebuggerView::UpdateViews()
 				SelectionView->ClearChildren();
 				for (int32 NodeId : NodeIds)
 				{
-					OnSelectionChanged.Execute(AnimInstanceId, NodeId);
+					OnSelectionChanged.Execute(NodeId);
 					SelectionView->AddSlot()
 					.HAlign(HAlign_Fill)
 					.VAlign(VAlign_Center)
@@ -945,21 +969,19 @@ void SDebuggerView::UpdateViews()
 						.HAlign(HAlign_Center)
 						.VAlign(VAlign_Center)
 						.ContentPadding(10.0f)
-						.OnClicked_Lambda([this, NodeId]() -> FReply
-						{
-							SelectedNode = NodeId;
-							UpdateViews();
-							return FReply::Handled();
-						})
+						.OnClicked_Raw(this, &SDebuggerView::UpdateNodeSelection, NodeId)
 					];
 				}
 			}
 		}
+		
+		// Update active node list
+		StoredNodes = MoveTemp(NodeIds);
 	}
 	else
 	{
 		check(Reflection.Get());
-    	OnSelectionChanged.Execute(AnimInstanceId, SelectedNode);
+    	OnSelectionChanged.Execute(SelectedNode);
 		
     	const FTraceMotionMatchingStateMessage* State = MotionMatchingState.Get();
     	const UPoseSearchDatabase* Database = PoseSearchDatabase.Get();
@@ -971,18 +993,78 @@ void SDebuggerView::UpdateViews()
     	SwitcherViewType = Debugger;
     	DatabaseView->Update(State, Database);	
 	}
-	ActiveNodes = NodeIds;
+}
+
+int32 SDebuggerView::SelectView() const
+{
+	if (IsPIESimulating.Get())
+	{
+		if (IsRecording.Get())
+		{
+			return RecordingMsg;
+		}
+	}
+
+	if (RecordingDuration.Get() < DOUBLE_SMALL_NUMBER)
+	{
+		return StoppedMsg;
+	}
+
+	const bool bNoActiveNodes = ActiveNodesNum.Get() == 0;
+	const bool bNodeSelectedWithoutData = SelectedNode != INDEX_NONE && MotionMatchingState.Get() == nullptr;
+	
+	if (bNoActiveNodes || bNodeSelectedWithoutData)
+    {
+    	return NoDataMsg;
+    }
+	
+	if (SelectedNode == INDEX_NONE)
+	{
+		return Selection;
+	}
+
+	return Debugger;
+}
+
+FReply SDebuggerView::UpdateNodeSelection(int32 InSelectedNode)
+{
+	// -1 will backtrack to selection view
+	SelectedNode = InSelectedNode;
+	UpdateViews();
+	return FReply::Handled();
+}
+
+TSharedRef<SWidget> SDebuggerView::GenerateNoDataMessageView()
+{
+	TSharedRef<SWidget> ReturnButtonView = GenerateReturnButtonView();
+	ReturnButtonView->SetVisibility(TAttribute<EVisibility>::CreateLambda([this]() -> EVisibility
+	{
+		// Hide the return button for the no data message if we have no nodes at all
+		return ActiveNodesNum.Get() > 0 ? EVisibility::Visible : EVisibility::Hidden;
+	}));
+	
+	return 
+		SNew(SVerticalBox)
+
+		+ SVerticalBox::Slot()
+		[
+			SNew(SDebuggerMessageBox, "No recorded data available for the selected frame")
+		]
+
+		+ SVerticalBox::Slot()
+		.AutoHeight()
+		.Padding(20.0f)
+		.HAlign(HAlign_Center)
+		.VAlign(VAlign_Center)
+		[
+			ReturnButtonView
+		];
 }
 
 TSharedRef<SWidget> SDebuggerView::GenerateReturnButtonView()
 {
 	return
-		SAssignNew(ReturnButtonView, SHorizontalBox)
-		.Visibility_Lambda([this]() -> EVisibility
-		{
-			// Collapse this view if we have don't have more than 1 node
-			return ActiveNodes.Num() > 1 ? EVisibility::Visible : EVisibility::Collapsed;
-		})
+		SNew(SHorizontalBox)
 
 		+ SHorizontalBox::Slot()
 		.Padding(10, 5, 0, 0)
@@ -992,13 +1074,7 @@ TSharedRef<SWidget> SDebuggerView::GenerateReturnButtonView()
 			.VAlign(VAlign_Center)
 			.ButtonStyle(FEditorStyle::Get(), "SimpleButton")
 			.ContentPadding( FMargin(1, 0) )
-			.OnClicked_Lambda([this]() -> FReply
-			{
-				// Clicking back backtracks selected node to invalid
-				SelectedNode = -1;
-				UpdateViews();
-				return FReply::Handled();
-			})
+			.OnClicked_Raw(this, &SDebuggerView::UpdateNodeSelection, static_cast<int32>(INDEX_NONE))
 			// Contents of button, icon then text
 			[
 				SNew(SHorizontalBox)
@@ -1026,6 +1102,14 @@ TSharedRef<SWidget> SDebuggerView::GenerateReturnButtonView()
 
 TSharedRef<SWidget> SDebuggerView::GenerateNodeDebuggerView()
 {
+	TSharedRef<SWidget> ReturnButtonView = GenerateReturnButtonView();
+	ReturnButtonView->SetVisibility(TAttribute<EVisibility>::CreateLambda([this]() -> EVisibility
+		{
+			// Collapse this view if we have don't have more than 1 node
+			return ActiveNodesNum.Get() > 1 ? EVisibility::Visible : EVisibility::Collapsed;
+		}
+	));
+	
 	return 
 		SNew(SSplitter)
 		.Orientation(Orient_Horizontal)
@@ -1040,7 +1124,7 @@ TSharedRef<SWidget> SDebuggerView::GenerateNodeDebuggerView()
 			+ SVerticalBox::Slot()
 			.AutoHeight()
 			[
-				GenerateReturnButtonView()
+				ReturnButtonView
 			]
 			
 			+ SVerticalBox::Slot()
@@ -1106,14 +1190,13 @@ void FDebugger::Update(float DeltaTime, IRewindDebugger* InRewindDebugger)
 void FDebugger::UpdateReflection() const
 {
 	check(Reflection);
-
-	if (!MotionMatchingState)
+	if (!ActiveMotionMatchingState)
 	{
 		return;
 	}
-	
-	Reflection->ElapsedPoseJumpTime = MotionMatchingState->ElapsedPoseJumpTime;
-	
+
+	Reflection->ElapsedPoseJumpTime = ActiveMotionMatchingState->ElapsedPoseJumpTime;
+
 	const UPoseSearchDatabase* Database = GetPoseSearchDatabase();
 	if (Database == nullptr)
 	{
@@ -1127,10 +1210,9 @@ void FDebugger::UpdateReflection() const
 	Reflection->TimeTrajectoryFeatures.EmptyAll();
 
 	FFeatureVectorReader Reader;
-	Reader.Init(&Schema->Layout);
 	// Ensure parity between Layout and QueryVector
-	Reader.SetValues(MotionMatchingState->QueryVector);
-	check(Reader.IsValid());
+	Reader.Init(&Schema->Layout);
+	Reader.SetValues(ActiveMotionMatchingState->QueryVector);
 
 	int32 NumSubsamples = Schema->PoseSampleTimes.Num();
 	const int32 NumBones = Schema->Bones.Num();
@@ -1192,35 +1274,71 @@ void FDebugger::UpdateReflection() const
 	}
 }
 
+void FDebugger::UpdateMotionMatchingStates(uint64 InAnimInstanceId)
+{
+	NodeIds.Empty();
+	MotionMatchingStates.Empty();
+	// Get provider and validate
+	const TraceServices::IAnalysisSession* Session = RewindDebugger->GetAnalysisSession();
+	TraceServices::FAnalysisSessionReadScope SessionReadScope(*Session);
+	
+	const FTraceProvider* TraceProvider = Session->ReadProvider<FTraceProvider>(FTraceProvider::ProviderName);
+	if (!TraceProvider)
+	{
+		return;
+	}
 
-const FTraceMotionMatchingStateMessage* FDebugger::GetMotionMatchingState()
+	const double TraceTime = RewindDebugger->CurrentTraceTime();
+	TraceServices::FFrame Frame;
+	ReadFrameProvider(*Session).GetFrameFromTime(TraceFrameType_Game, TraceTime, Frame);
+	
+
+	TraceProvider->EnumerateMotionMatchingStateTimelines(InAnimInstanceId, [this, &Frame](const FTraceProvider::FMotionMatchingStateTimeline& InTimeline)
+	{
+		const FTraceMotionMatchingStateMessage* Message = nullptr;
+		
+		InTimeline.EnumerateEvents(Frame.StartTime, Frame.EndTime, [this, &Message, &Frame](double InStartTime, double InEndTime, const FTraceMotionMatchingStateMessage& InMessage)
+		{
+			Message = &InMessage;
+			return TraceServices::EEventEnumerate::Stop;
+		});
+
+		if (Message)
+		{
+			NodeIds.Add(Message->NodeId);
+			MotionMatchingStates.Add(Message);
+		}
+	});
+}
+
+
+const FTraceMotionMatchingStateMessage* FDebugger::GetMotionMatchingState() const
 {
-	return Instance()->MotionMatchingState;
+	return ActiveMotionMatchingState;
 }
 
-const UPoseSearchDatabase* FDebugger::GetPoseSearchDatabase() 
+const UPoseSearchDatabase* FDebugger::GetPoseSearchDatabase() const
 {
-	if (InternalInstance->MotionMatchingState == nullptr)
+	if (!InternalInstance->ActiveMotionMatchingState)
 	{
 		return nullptr;
 	}
 
-	const uint64 DatabaseId = InternalInstance->MotionMatchingState->DatabaseId;
+	const uint64 DatabaseId = InternalInstance->ActiveMotionMatchingState->DatabaseId;
 	if (DatabaseId == 0)
 	{
 	    return nullptr;
 	}
 	
-	// @TODO: Load the database if not currently loaded
 	UObject* DatabaseObject = FObjectTrace::GetObjectFromId(DatabaseId);
+	// @TODO: Load the object if unloaded
 	if (DatabaseObject == nullptr)
 	{
 		return nullptr;
 	}
-	
 	check(DatabaseObject->IsA<UPoseSearchDatabase>());
 
-	const UPoseSearchDatabase* Database = Cast<UPoseSearchDatabase>(DatabaseObject);
+	UPoseSearchDatabase* Database = Cast<UPoseSearchDatabase>(DatabaseObject);
 	const UPoseSearchSchema* Schema = Database->Schema;
 	if (Schema == nullptr || !Schema->IsValid())
 	{
@@ -1229,73 +1347,80 @@ const UPoseSearchDatabase* FDebugger::GetPoseSearchDatabase()
 	return Database;
 }
 
-UPoseSearchDebuggerReflection* FDebugger::GetReflection()
+UPoseSearchDebuggerReflection* FDebugger::GetReflection() const
 {
-	return InternalInstance->Reflection;
+	return Reflection;
 }
 
-bool FDebugger::GetIsPIESimulating()
+bool FDebugger::IsPIESimulating() const
 {
-	return InternalInstance->RewindDebugger->IsPIESimulating();
+	return RewindDebugger->IsPIESimulating();
 }
 
-TSet<int32> FDebugger::GetNodeIds(uint64 AnimInstanceId)
+const TArray<int32>* FDebugger::GetNodeIds() const
 {
-	const TraceServices::IAnalysisSession* Session = Instance()->RewindDebugger->GetAnalysisSession();
-	TraceServices::FAnalysisSessionReadScope SessionReadScope(*Session);
+	return &NodeIds;
+}
 
-	// Get provider and validate its existence in the session
-	const FTraceProvider* TraceProvider = Session->ReadProvider<FTraceProvider>(FTraceProvider::ProviderName);
-	if (!TraceProvider)
-	{
-		return TSet<int32>();
-	}
-	
-	return TraceProvider->GetMotionMatchingNodeIds(AnimInstanceId);
+bool FDebugger::IsRecording() const
+{
+	return RewindDebugger->IsRecording();
 }
 
-void FDebugger::OnSelectionChanged(uint64 AnimInstanceId, int32 NodeId)
+double FDebugger::GetRecordingDuration() const
 {
-	const TraceServices::IAnalysisSession* Session = RewindDebugger->GetAnalysisSession();
-	TraceServices::FAnalysisSessionReadScope SessionReadScope(*Session);
+	return RewindDebugger->GetRecordingDuration();
+}
 
-	// Get provider and validate
-	const FTraceProvider* TraceProvider = Session->ReadProvider<FTraceProvider>(FTraceProvider::ProviderName);
-	if (!TraceProvider)
-	{
-		return;
-	}
+int32 FDebugger::GetActiveNodesNum() const
+{
+	return MotionMatchingStates.Num();
+}
 
-	const double TraceTime = RewindDebugger->CurrentTraceTime();
-	TraceProvider->ReadMotionMatchingStateTimeline(AnimInstanceId, NodeId, [this, Session, TraceTime](const FTraceProvider::FMotionMatchingStateTimeline& TimelineData)
+void FDebugger::OnUpdate(uint64 InAnimInstanceId)
+{
+	UpdateMotionMatchingStates(InAnimInstanceId);
+}
+
+void FDebugger::OnSelectionChanged(int32 InNodeId)
+{
+	bool bFound = false;
+	// Find node in all motion matching states this frame
+	const int32 NodesNum = NodeIds.Num();
+	for(int i = 0; i < NodesNum; ++i)
 	{
-		const TraceServices::IFrameProvider& FrameProvider = TraceServices::ReadFrameProvider(*Session);
-		TraceServices::FFrame Frame;
-		if(FrameProvider.GetFrameFromTime(ETraceFrameType::TraceFrameType_Game, TraceTime, Frame))
+		if (NodeIds[i] == InNodeId)
 		{
-			TimelineData.EnumerateEvents(Frame.StartTime, Frame.EndTime,
-				[this](double InStartTime, double InEndTime, uint32 InDepth, const FTraceMotionMatchingStateMessage& Message)
-				{
-					MotionMatchingState = &Message;
-
-					return TraceServices::EEventEnumerate::Stop;
-				});
+			ActiveMotionMatchingState = MotionMatchingStates[i];
+			bFound = true;
+			break;
 		}
-	});
+	}
 
-	UpdateReflection();
+	if (bFound)
+	{
+		UpdateReflection();
+	}
+	else
+	{
+		ActiveMotionMatchingState = nullptr;
+	}
 }
 
 TSharedPtr<SDebuggerView> FDebugger::GenerateView(uint64 InAnimInstanceId)
 {
 	return
 		SNew(SDebuggerView, InAnimInstanceId)
-		.MotionMatchingNodeIds_Static(&FDebugger::GetNodeIds, InAnimInstanceId)
-		.MotionMatchingState_Static(&FDebugger::GetMotionMatchingState)
-		.Reflection_Static(&FDebugger::GetReflection)
-		.PoseSearchDatabase_Static(&FDebugger::GetPoseSearchDatabase)
-		.IsPIESimulating_Static(&FDebugger::GetIsPIESimulating)
-		.OnSelectionChanged_Raw(Instance(), &FDebugger::OnSelectionChanged);
+		.MotionMatchingNodeIds_Raw(Instance(), &FDebugger::GetNodeIds)
+		.MotionMatchingState_Raw(Instance(), &FDebugger::GetMotionMatchingState)
+		.Reflection_Raw(Instance(), &FDebugger::GetReflection)
+		.PoseSearchDatabase_Raw(Instance(), &FDebugger::GetPoseSearchDatabase)
+		.IsPIESimulating_Raw(Instance(), &FDebugger::IsPIESimulating)
+		.IsRecording_Raw(Instance(), &FDebugger::IsRecording)
+		.RecordingDuration_Raw(Instance(), &FDebugger::GetRecordingDuration)
+		.ActiveNodesNum_Raw(Instance(), &FDebugger::GetActiveNodesNum)
+		.OnSelectionChanged_Raw(Instance(), &FDebugger::OnSelectionChanged)
+		.OnUpdate_Raw(Instance(), &FDebugger::OnUpdate);
 }
 
 
diff --git a/Engine/Plugins/Experimental/Animation/PoseSearch/Source/Editor/Private/PoseSearchDebugger.h b/Engine/Plugins/Experimental/Animation/PoseSearch/Source/Editor/Private/PoseSearchDebugger.h
--- a/Engine/Plugins/Experimental/Animation/PoseSearch/Source/Editor/Private/PoseSearchDebugger.h
+++ b/Engine/Plugins/Experimental/Animation/PoseSearch/Source/Editor/Private/PoseSearchDebugger.h
@@ -94,7 +94,7 @@ private:
 	void AddColumn(TSharedRef<DebuggerDatabaseColumns::IColumn>&& Column);
 	void CreateRows(const UPoseSearchDatabase* Database);
 	void SortDatabaseRows();
-	void PopulateRows(const FTraceMotionMatchingStateMessage* State, const UPoseSearchDatabase* Database);
+	void UpdateRows(const FTraceMotionMatchingStateMessage* State, const UPoseSearchDatabase* Database);
 	EColumnSortMode::Type GetColumnSortMode(const FName ColumnId) const;
 	float GetColumnWidth(const FName ColumnId) const;
 	void OnColumnSortModeChanged(const EColumnSortPriority::Type SortPriority, const FName & ColumnId, const EColumnSortMode::Type InSortMode);
@@ -156,7 +156,13 @@ private:
 /**
  * Callback to update the debugger when node selection is changed
  */
-DECLARE_DELEGATE_TwoParams(FOnSelectionChanged, uint64 AnimInstanceId, int32 NodeId);
+DECLARE_DELEGATE_OneParam(FOnSelectionChanged, int32 NodeId);
+
+/**
+ * Callback to update the debugger when view update occurs, updating the motion matching state
+ * relative to the selected anim instance.
+ */
+DECLARE_DELEGATE_OneParam(FOnUpdate, uint64 AnimInstanceId);
 
 /**
  * Entire view of the PoseSearch debugger, containing all sub-widgets
@@ -170,9 +176,13 @@ public:
 		SLATE_ATTRIBUTE(const FTraceMotionMatchingStateMessage*, MotionMatchingState)
 		SLATE_ATTRIBUTE(UPoseSearchDebuggerReflection*, Reflection)
 		SLATE_ATTRIBUTE(const UPoseSearchDatabase*, PoseSearchDatabase)
-		SLATE_ATTRIBUTE(TSet<int32>, MotionMatchingNodeIds)
+		SLATE_ATTRIBUTE(const TArray<int32>*, MotionMatchingNodeIds)
 		SLATE_ATTRIBUTE(bool, IsPIESimulating)
+		SLATE_ATTRIBUTE(bool, IsRecording)
+		SLATE_ATTRIBUTE(double, RecordingDuration)
+		SLATE_ATTRIBUTE(int32, ActiveNodesNum)
 		SLATE_EVENT(FOnSelectionChanged, OnSelectionChanged)
+		SLATE_EVENT(FOnUpdate, OnUpdate)
 	SLATE_END_ARGS()
 
 	void Construct(const FArguments& InArgs, uint64 InAnimInstanceId);
@@ -182,27 +192,42 @@ public:
 
 private:
 	void UpdateViews();
+	
+	/** Returns an int32 appropriate to the index of our widget selector */
+	int32 SelectView() const;
+
+	FReply UpdateNodeSelection(int32 InSelectedNode);
+	
+	TSharedRef<SWidget> GenerateNoDataMessageView();
 	TSharedRef<SWidget> GenerateReturnButtonView();
 	TSharedRef<SWidget> GenerateNodeDebuggerView();
 
-	/** Gets all MM nodes being traced in this session */
-	TAttribute<TSet<int32>> MotionMatchingNodeIds;
+	/** Gets all MM nodes being traced in this frame */
+	TAttribute<const TArray<int32>*> MotionMatchingNodeIds;
 	/** Retrieves the reflection UObject from the debugger */
 	TAttribute<UPoseSearchDebuggerReflection*> Reflection;
 	/** Retrieves the MM state from the debugger */
 	TAttribute<const FTraceMotionMatchingStateMessage*> MotionMatchingState;
 	/** Retrieves the PoseSearch Database from the debugger */
 	TAttribute<const UPoseSearchDatabase*> PoseSearchDatabase;
-	/** Whether the game is unpaused and currently simulating */
+	/** Whether the game is un-paused and currently simulating */
 	TAttribute<bool> IsPIESimulating;
+	/** Whether the Rewind Debugger is currently recording gameplay */
+	TAttribute<bool> IsRecording;
+ 	/** Length of the recorded sequence (if any) */
+	TAttribute<double> RecordingDuration;
+ 	/** Number of active nodes at the current frame */
+	TAttribute<int32> ActiveNodesNum;
 	/** Update current debugger data when node selection is changed */
 	FOnSelectionChanged OnSelectionChanged;
+	/** Update current debugger data when update occurs */
+	FOnUpdate OnUpdate;
 
 	/** Active node being debugged */
 	int32 SelectedNode = -1;
 
 	/** List of all nodes being traced (stored for selection) */
-	TSet<int32> ActiveNodes;
+	TSet<int32> StoredNodes;
 
 	/** Database view of the motion matching node */
 	TSharedPtr<SDebuggerDatabaseView> DatabaseView;
@@ -211,24 +236,23 @@ private:
 	/** Node debugger view hosts the above two views */
 	TSharedPtr<SSplitter> NodeDebuggerView;
 
-	/** Contains the return to node selection (picked by database name) */
-	TSharedPtr<SHorizontalBox> ReturnButtonView;
 	/** Selection view before node is selected */
 	TSharedPtr<SVerticalBox> SelectionView;
-
 	
 	/** Gray box occluding the debugger view when simulating */
 	TSharedPtr<SVerticalBox> SimulatingView;
 	
-	/** Used to switch between views in the switcher */
-	enum ESwitcherViewType : uint16
+	/** Used to switch between views in the switcher, int32 maps to index in the SWidgetSwitcher */
+	enum ESwitcherViewType : int32
 	{
-		Waiting = 0,
-		Selection = 1,
-		Debugger = 2
-	} SwitcherViewType = Waiting;
+		Selection = 0,
+		Debugger = 1,
+		StoppedMsg = 2,
+		RecordingMsg = 3,
+		NoDataMsg = 4
+	} SwitcherViewType = StoppedMsg;
+	
 	/** Contains all the above, switches between them depending on context */
-	// @TODO: Make this resolve to Waiting when PIE is stopped
 	TSharedPtr<SWidgetSwitcher> Switcher;
 
 	/** Contains the switcher, the entire debugger view */
@@ -238,7 +262,16 @@ private:
 	uint64 AnimInstanceId = 0;
 
 	/** Current position of the time marker */
-	double TimeMarker = 0.0;
+	double TimeMarker = -1.0;
+
+	/** Tracks if the current time has been updated yet (delayed) */
+	bool bUpdated = false;
+
+	/** Tracks number of consecutive frames, once it reaches threshold it will update the view */
+	int32 CurrentConsecutiveFrames = 0;
+
+	/** Once the frame count has reached this value, an update will trigger for the view */
+	static constexpr int32 ConsecutiveFramesUpdateThreshold = 10;
 };
 
 
@@ -261,20 +294,29 @@ public:
 
 private:
 	// Used for view callbacks
-	static const FTraceMotionMatchingStateMessage* GetMotionMatchingState();
-	static const UPoseSearchDatabase* GetPoseSearchDatabase();
-	static UPoseSearchDebuggerReflection* GetReflection();
-	static bool GetIsPIESimulating();
-	static TSet<int32> GetNodeIds(uint64 AnimInstanceId);
-	void OnSelectionChanged(uint64 AnimInstanceId, int32 NodeId);
+	const FTraceMotionMatchingStateMessage* GetMotionMatchingState() const;
+	const UPoseSearchDatabase* GetPoseSearchDatabase() const;
+	UPoseSearchDebuggerReflection* GetReflection() const;
+	bool IsPIESimulating() const;
+	const TArray<int32>* GetNodeIds() const;
+	bool IsRecording() const;
+	double GetRecordingDuration() const;
+	int32 GetActiveNodesNum() const;
+	void OnUpdate(uint64 InAnimInstanceId);
+	void OnSelectionChanged(int32 InNodeId);
 
 	/** Updates the current reflection data relative to the MM state */
 	void UpdateReflection() const;
+	void UpdateMotionMatchingStates(uint64 InAnimInstanceId);
 	
 	/** Last stored Rewind Debugger */
 	const IRewindDebugger* RewindDebugger = nullptr;
+	/** List of all Node IDs associated with motion matching states */
+	TArray<int32> NodeIds;
+	/** List of all updated motion matching states per node */
+	TArray<const FTraceMotionMatchingStateMessage*> MotionMatchingStates;
 	/** Last stored MM state (updated from OnSelectionChanged) */
-    const FTraceMotionMatchingStateMessage* MotionMatchingState = nullptr;
+	const FTraceMotionMatchingStateMessage* ActiveMotionMatchingState = nullptr;
 	/** Last updated reflection data relative to MM state */
 	TObjectPtr<UPoseSearchDebuggerReflection> Reflection = nullptr;
 
diff --git a/Engine/Plugins/Experimental/Animation/PoseSearch/Source/Editor/Private/Trace/PoseSearchTraceProvider.cpp b/Engine/Plugins/Experimental/Animation/PoseSearch/Source/Editor/Private/Trace/PoseSearchTraceProvider.cpp
--- a/Engine/Plugins/Experimental/Animation/PoseSearch/Source/Editor/Private/Trace/PoseSearchTraceProvider.cpp
+++ b/Engine/Plugins/Experimental/Animation/PoseSearch/Source/Editor/Private/Trace/PoseSearchTraceProvider.cpp
@@ -21,6 +21,12 @@ bool FTraceProvider::ReadMotionMatchingStateTimeline(uint64 InAnimInstanceId, in
 	return MotionMatchingStateTimelineStorage.ReadTimeline(InAnimInstanceId, InNodeId, Callback);
 }
 
+bool FTraceProvider::EnumerateMotionMatchingStateTimelines(uint64 InAnimInstanceId, TFunctionRef<void(const FMotionMatchingStateTimeline&)> Callback) const
+{
+	Session.ReadAccessCheck();
+	return MotionMatchingStateTimelineStorage.EnumerateNodeTimelines(InAnimInstanceId, Callback);
+}
+
 
 void FTraceProvider::AppendMotionMatchingState(const FTraceMotionMatchingStateMessage& InMessage, double InTime)
 {
diff --git a/Engine/Plugins/Experimental/Animation/PoseSearch/Source/Editor/Private/Trace/PoseSearchTraceProvider.h b/Engine/Plugins/Experimental/Animation/PoseSearch/Source/Editor/Private/Trace/PoseSearchTraceProvider.h
--- a/Engine/Plugins/Experimental/Animation/PoseSearch/Source/Editor/Private/Trace/PoseSearchTraceProvider.h
+++ b/Engine/Plugins/Experimental/Animation/PoseSearch/Source/Editor/Private/Trace/PoseSearchTraceProvider.h
@@ -43,13 +43,18 @@ public:
 	/** Get all node ids of motion matching nodes appended to the object's timeline thus far */
 	TSet<int32> GetMotionMatchingNodeIds(uint64 InAnimInstanceId) const;
 
-	/** Read the object-relevant timeline info and provide it via callback */
 	using FMotionMatchingStateTimeline = TraceServices::ITimeline<FTraceMotionMatchingStateMessage>;
+	
+	/** Read the node-relative timeline info of the AnimInstance and provide it via callback */
 	bool ReadMotionMatchingStateTimeline(uint64 InAnimInstanceId, int32 InNodeId, TFunctionRef<void(const FMotionMatchingStateTimeline&)> Callback) const;
+	
+	/** Enumerate all node timelines on the AnimInstance and provide them via callback */
+	bool EnumerateMotionMatchingStateTimelines(uint64 InAnimInstanceId, TFunctionRef<void(const FMotionMatchingStateTimeline&)> Callback) const;
 
 	/** Append to the timeline on our object */
 	void AppendMotionMatchingState(const FTraceMotionMatchingStateMessage& InMessage, double InTime);
 
+	
 	static const FName ProviderName;
 
 private:
@@ -93,6 +98,18 @@ private:
 			return Timelines[*TimelineIndex];
 		}
 
+		bool ReadNodeTimeline(const FNodeToTimelineMap* NodeToTimelineMap, int32 InNodeId, TFunctionRef<void(const TraceServices::ITimeline<MessageType>&)> Callback) const
+		{
+			const uint32* TimelineIndex = NodeToTimelineMap->Find(InNodeId);
+			if (TimelineIndex == nullptr || !Timelines.IsValidIndex(*TimelineIndex))
+			{
+				return false;
+			}
+
+			Callback(*Timelines[*TimelineIndex]);
+			return true;
+		}
+
 		/** Retrieve a timeline from an anim instance + node and execute the callback */
 		bool ReadTimeline(uint64 InAnimInstanceId, int32 InNodeId, TFunctionRef<void(const TraceServices::ITimeline<MessageType>&)> Callback) const
 		{
@@ -102,14 +119,29 @@ private:
 				return false;
 			}
 
-			const uint32* TimelineIndex = NodeToTimelineMap->Find(InNodeId);
-			if (TimelineIndex == nullptr || !Timelines.IsValidIndex(*TimelineIndex))
+			return ReadNodeTimeline(NodeToTimelineMap, InNodeId, Callback);
+		}
+
+		/** Enumerates timelines of all nodes given an AnimInstance */
+		bool EnumerateNodeTimelines(uint64 InAnimInstanceId, TFunctionRef<void(const TraceServices::ITimeline<MessageType>&)> Callback) const
+		{
+			const FNodeToTimelineMap* NodeToTimelineMap = AnimInstanceIdToTimelines.Find(InAnimInstanceId);
+			if (NodeToTimelineMap == nullptr)
 			{
 				return false;
 			}
 
-			Callback(*Timelines[*TimelineIndex]);
-			return true;
+			bool bSuccess = true;
+			for (const TPair<int32, uint32>& Pair : *NodeToTimelineMap)
+			{
+				bSuccess &= ReadNodeTimeline(NodeToTimelineMap, Pair.Key, Callback);
+				if (!bSuccess)
+				{
+					break;
+				}
+			}
+
+			return bSuccess;
 		}
 
 		TSet<int32> GetNodeIds(uint64 InAnimInstanceId) const
