PoseSearch Debugger: Debug Visualization & Multiple Tabs

- Debug Visualization integration inside the debugger, currently supports drawing the query vector in blue, the active pose in green, and a database selection in red
- Multiple debugger tabs are now supported inside the Rewind Debugger, using the "Keep View Open" option in the context menu of the Pose Search tab.

#rnx

diff --git a/Engine/Plugins/Experimental/Animation/PoseSearch/PoseSearch.uplugin b/Engine/Plugins/Experimental/Animation/PoseSearch/PoseSearch.uplugin
--- a/Engine/Plugins/Experimental/Animation/PoseSearch/PoseSearch.uplugin
+++ b/Engine/Plugins/Experimental/Animation/PoseSearch/PoseSearch.uplugin
@@ -26,5 +26,14 @@
 			"Type": "UncookedOnly",
 			"LoadingPhase": "PreDefault"
 		}
+	],
+	"Plugins": [
+		{
+			"Name": "GameplayInsights",
+			"Enabled": true,
+			"WhitelistTargets": [
+				"Editor"
+			]
+		}
 	]
 }
\ No newline at end of file
diff --git a/Engine/Plugins/Experimental/Animation/PoseSearch/Source/Editor/PoseSearchEditor.Build.cs b/Engine/Plugins/Experimental/Animation/PoseSearch/Source/Editor/PoseSearchEditor.Build.cs
--- a/Engine/Plugins/Experimental/Animation/PoseSearch/Source/Editor/PoseSearchEditor.Build.cs
+++ b/Engine/Plugins/Experimental/Animation/PoseSearch/Source/Editor/PoseSearchEditor.Build.cs
@@ -23,6 +23,7 @@ public class PoseSearchEditor : ModuleRules
 				"TraceAnalysis",
 				"TraceServices",
 				"TraceInsights",
+				"GameplayInsights",
 				
 				// UI 
 				"PropertyEditor",
diff --git a/Engine/Plugins/Experimental/Animation/PoseSearch/Source/Editor/Private/PoseSearchDebugger.cpp b/Engine/Plugins/Experimental/Animation/PoseSearch/Source/Editor/Private/PoseSearchDebugger.cpp
--- a/Engine/Plugins/Experimental/Animation/PoseSearch/Source/Editor/Private/PoseSearchDebugger.cpp
+++ b/Engine/Plugins/Experimental/Animation/PoseSearch/Source/Editor/Private/PoseSearchDebugger.cpp
@@ -1,6 +1,8 @@
 ﻿// Copyright Epic Games, Inc. All Rights Reserved.
 
 #include "PoseSearchDebugger.h"
+#include "IAnimationProvider.h"
+#include "Editor/EditorEngine.h"
 #include "PoseSearch/PoseSearch.h"
 #include "Animation/AnimSequence.h"
 #include "RewindDebuggerInterface/Public/IRewindDebugger.h"
@@ -31,6 +33,8 @@ void FPoseSearchDebuggerFeatureReflection::EmptyAll()
 
 namespace UE { namespace PoseSearch {
 
+constexpr uint8 DrawDebugDepthPriority = SDPG_Foreground + 2;
+
 class FDebuggerDatabaseRowData : public TSharedFromThis<FDebuggerDatabaseRowData>
 {
 public:
@@ -99,7 +103,7 @@ namespace DebuggerDatabaseColumns
 			
 			return SNew(STextBlock)
             	.Font(RowFont)
-				.Text_Lambda([this, RowData] { return GetRowText(RowData); })
+				.Text_Lambda([this, RowData]() -> FText { return GetRowText(RowData); })
             	.Justification(ETextJustify::Center);
 		}
 
@@ -286,14 +290,13 @@ namespace DebuggerDatabaseColumns
 }
 
 
-DECLARE_DELEGATE_RetVal_OneParam(TSharedRef<DebuggerDatabaseColumns::IColumn>&, FGetColumnDelegate, FName);
-
 /**
  * Widget representing a single row of the database view
  */
 class SDebuggerDatabaseRow : public SMultiColumnTableRow<TSharedRef<FDebuggerDatabaseRowData>>
 {
 	SLATE_BEGIN_ARGS(SDebuggerDatabaseRow) {}
+		SLATE_ATTRIBUTE(const SDebuggerDatabaseView::FColumnMap*, ColumnMap)
 	SLATE_END_ARGS()
 
 	void Construct(
@@ -305,8 +308,13 @@ class SDebuggerDatabaseRow : public SMultiColumnTableRow<TSharedRef<FDebuggerDat
 		FMargin InPaddingMargin
 	)
 	{
+		ColumnMap = InArgs._ColumnMap;
+		check(ColumnMap.IsBound());
+		
 		Row = InRow;
+		
 		RowBrush = InRowBrush;
+		check(RowBrush);
 
 		SMultiColumnTableRow<TSharedRef<FDebuggerDatabaseRowData>>::Construct(
 			FSuperRowType::FArguments()
@@ -318,11 +326,8 @@ class SDebuggerDatabaseRow : public SMultiColumnTableRow<TSharedRef<FDebuggerDat
 
 	virtual TSharedRef<SWidget> GenerateWidgetForColumn(const FName& InColumnName) override
 	{
-		check(GetColumn.IsBound());
-		check(RowBrush);
-		
 		// Get column
-		TSharedRef<DebuggerDatabaseColumns::IColumn>& Column = GetColumn.Execute(InColumnName);
+		const TSharedRef<DebuggerDatabaseColumns::IColumn>& Column = (*ColumnMap.Get())[InColumnName];
 		
 		static FSlateFontInfo NormalFont = FEditorStyle::Get().GetFontStyle("DetailsView.CategoryTextStyle");
 		const TSharedRef<SWidget> Widget = Column->GenerateTextWidget(Row.ToSharedRef());
@@ -351,9 +356,8 @@ class SDebuggerDatabaseRow : public SMultiColumnTableRow<TSharedRef<FDebuggerDat
 	const FSlateBrush* RowBrush = nullptr;
 
 	/** Used to grab the column struct given a column name */
-	static FGetColumnDelegate GetColumn;
+	TAttribute<const SDebuggerDatabaseView::FColumnMap*> ColumnMap;
 };
-FGetColumnDelegate SDebuggerDatabaseRow::GetColumn;
 
 class SDebuggerMessageBox : public SCompoundWidget
 {
@@ -377,7 +381,7 @@ class SDebuggerMessageBox : public SCompoundWidget
 	}
 };
 
-void SDebuggerDatabaseView::Update(const FTraceMotionMatchingStateMessage* State, const UPoseSearchDatabase* Database)
+void SDebuggerDatabaseView::Update(const FTraceMotionMatchingStateMessage& State, const UPoseSearchDatabase& Database)
 {
 	UpdateRows(State, Database);
 }
@@ -472,16 +476,15 @@ void SDebuggerDatabaseView::SortDatabaseRows()
 	}
 
 	DatabaseView.ListView->RequestListRefresh();
-	DatabaseView.HeaderRow->RefreshColumns();
 }
 
-void SDebuggerDatabaseView::CreateRows(const UPoseSearchDatabase* Database)
+void SDebuggerDatabaseView::CreateRows(const UPoseSearchDatabase& Database)
 {
-	const int32 NumPoses = Database->SearchIndex.NumPoses;
+	const int32 NumPoses = Database.SearchIndex.NumPoses;
 	DatabaseView.Rows.Reserve(NumPoses);
 
 	// Build database rows
-	for(const FPoseSearchDatabaseSequence& DbSequence : Database->Sequences)
+	for(const FPoseSearchDatabaseSequence& DbSequence : Database.Sequences)
 	{
 		const int32 LastPoseIdx = DbSequence.FirstPoseIdx + DbSequence.NumPoses;
 		for (int32 PoseIdx = DbSequence.FirstPoseIdx; PoseIdx != LastPoseIdx; ++PoseIdx)
@@ -502,7 +505,7 @@ void SDebuggerDatabaseView::CreateRows(const UPoseSearchDatabase* Database)
 
 			Row->AnimSequenceName = SequenceName;
 			// Cap time in sequence to end of range
-			Row->Time = FGenericPlatformMath::Min(Range.Min + (PoseIdx - DbSequence.FirstPoseIdx) * Database->Schema->SamplingInterval, Range.Max);
+			Row->Time = FGenericPlatformMath::Min(Range.Min + (PoseIdx - DbSequence.FirstPoseIdx) * Database.Schema->SamplingInterval, Range.Max);
 			Row->Length = SequenceLength;
 		}
 	}
@@ -510,7 +513,7 @@ void SDebuggerDatabaseView::CreateRows(const UPoseSearchDatabase* Database)
 	ActiveView.Rows.Add(MakeShared<FDebuggerDatabaseRowData>());
 }
 
-void SDebuggerDatabaseView::UpdateRows(const FTraceMotionMatchingStateMessage* State, const UPoseSearchDatabase* Database)
+void SDebuggerDatabaseView::UpdateRows(const FTraceMotionMatchingStateMessage& State, const UPoseSearchDatabase& Database)
 {
 	if (DatabaseView.Rows.IsEmpty())
 	{
@@ -521,10 +524,10 @@ void SDebuggerDatabaseView::UpdateRows(const FTraceMotionMatchingStateMessage* S
 
 	// Active bias weights pulled from the MM node
 	FPoseSearchBiasWeights BiasWeights;
-    BiasWeights.Weights = State->BiasWeights;
+    BiasWeights.Weights = State.BiasWeights;
 
-	const FPoseSearchFeatureVectorLayout& Layout = Database->Schema->Layout;
-	const FPoseSearchIndex& SearchIndex = Database->SearchIndex;
+	const FPoseSearchFeatureVectorLayout& Layout = Database.Schema->Layout;
+	const FPoseSearchIndex& SearchIndex = Database.SearchIndex;
 
 	// Reverse engineer weight params from weights array
 	auto ExtractWeight = [](
@@ -565,19 +568,20 @@ void SDebuggerDatabaseView::UpdateRows(const FTraceMotionMatchingStateMessage* S
 	FPoseSearchBiasWeights TrajectoryBiasWeights;
 	TrajectoryBiasWeights.Init(Params, Layout);
 	
-	const FPoseSearchBiasWeightsContext BiasWeightsContext { &BiasWeights, Database };
-	const FPoseSearchBiasWeightsContext PoseBiasWeightsContext { &PoseBiasWeights, Database };
-	const FPoseSearchBiasWeightsContext TrajectoryBiasWeightsContext { &TrajectoryBiasWeights, Database };
+	const FPoseSearchBiasWeightsContext BiasWeightsContext { &BiasWeights, &Database };
+	const FPoseSearchBiasWeightsContext PoseBiasWeightsContext { &PoseBiasWeights, &Database };
+	const FPoseSearchBiasWeightsContext TrajectoryBiasWeightsContext { &TrajectoryBiasWeights, &Database };
 	
 	for(const TSharedRef<FDebuggerDatabaseRowData>& Row : DatabaseView.Rows)
 	{
 		const int32 PoseIdx = Row->PoseIdx;
-		Row->Score = ComparePoses(SearchIndex, PoseIdx, State->QueryVectorNormalized, &BiasWeightsContext);
-		Row->PoseScore = ComparePoses(SearchIndex, PoseIdx, State->QueryVectorNormalized, &PoseBiasWeightsContext);
-		Row->TrajectoryScore = ComparePoses(SearchIndex, PoseIdx, State->QueryVectorNormalized, &TrajectoryBiasWeightsContext);
+		
+		Row->Score = ComparePoses(SearchIndex, PoseIdx, State.QueryVectorNormalized, &BiasWeightsContext);
+		Row->PoseScore = ComparePoses(SearchIndex, PoseIdx, State.QueryVectorNormalized, &PoseBiasWeightsContext);
+		Row->TrajectoryScore = ComparePoses(SearchIndex, PoseIdx, State.QueryVectorNormalized, &TrajectoryBiasWeightsContext);
 
 		// If we are on the active pose for the frame
-		if (PoseIdx == State->DbPoseIdx)
+		if (PoseIdx == State.DbPoseIdx)
 		{
 			*ActiveView.Rows[0] = Row.Get();
 		}
@@ -588,18 +592,23 @@ void SDebuggerDatabaseView::UpdateRows(const FTraceMotionMatchingStateMessage* S
 
 TSharedRef<ITableRow> SDebuggerDatabaseView::HandleGenerateDatabaseRow(TSharedRef<FDebuggerDatabaseRowData> Item, const TSharedRef<STableViewBase>& OwnerTable) const
 {
-	return SNew(SDebuggerDatabaseRow, OwnerTable, Item, DatabaseView.RowStyle, &DatabaseView.RowBrush, FMargin(0.0f, 2.0f, 6.0f, 2.0f));
+	return
+		SNew(SDebuggerDatabaseRow, OwnerTable, Item, DatabaseView.RowStyle, &DatabaseView.RowBrush, FMargin(0.0f, 2.0f, 6.0f, 2.0f))
+		.ColumnMap(this, &SDebuggerDatabaseView::GetColumnMap);
+	
 }
 
 TSharedRef<ITableRow> SDebuggerDatabaseView::HandleGenerateActiveRow(TSharedRef<FDebuggerDatabaseRowData> Item, const TSharedRef<STableViewBase>& OwnerTable) const
 {
-	return SNew(SDebuggerDatabaseRow, OwnerTable, Item, ActiveView.RowStyle, &ActiveView.RowBrush, FMargin(0.0f, 2.0f, 6.0f, 4.0f));
+	return
+		SNew(SDebuggerDatabaseRow, OwnerTable, Item, ActiveView.RowStyle, &ActiveView.RowBrush, FMargin(0.0f, 2.0f, 6.0f, 4.0f))
+		.ColumnMap(this, &SDebuggerDatabaseView::GetColumnMap);
 }
 
 void SDebuggerDatabaseView::Construct(const FArguments& InArgs)
 {
 	using namespace DebuggerDatabaseColumns;
-	
+
 	// @TODO: Support runtime reordering of these indices
 	// Construct all column types
 	AddColumn(MakeShared<FAnimSequenceName>(0));
@@ -646,7 +655,7 @@ void SDebuggerDatabaseView::Construct(const FArguments& InArgs)
 			.HeaderRow(DatabaseView.HeaderRow.ToSharedRef())
 			.OnGenerateRow(this, &SDebuggerDatabaseView::HandleGenerateDatabaseRow)
 			.ExternalScrollbar(DatabaseView.ScrollBar)
-			.SelectionMode(ESelectionMode::Multi)
+			.SelectionMode(ESelectionMode::SingleToggle)
 			.ConsumeMouseWheel(EConsumeMouseWheel::WhenScrollingPossible);
 
 	DatabaseView.RowStyle = FEditorStyle::GetWidgetStyle<FTableRowStyle>("TableView.Row");
@@ -656,136 +665,132 @@ void SDebuggerDatabaseView::Construct(const FArguments& InArgs)
 	
 	ChildSlot
 	[
-		SNew(SWidgetSwitcher)
-		+ SWidgetSwitcher::Slot()
+		SNew(SVerticalBox)
+		+ SVerticalBox::Slot()
+		// Side and top margins, ignore bottom handled by the color border below
+		.Padding(0.0f, 5.0f, 0.0f, 0.0f)
+		.AutoHeight()
 		[
+			// Active Row text tab
 			SNew(SVerticalBox)
 			+ SVerticalBox::Slot()
-			// Side and top margins, ignore bottom handled by the color border below
-			.Padding(0.0f, 5.0f, 0.0f, 0.0f)
 			.AutoHeight()
+			.Padding(0.0f)
 			[
-				// Active Row text tab
-				SNew(SVerticalBox)
-				+ SVerticalBox::Slot()
-				.AutoHeight()
+				SNew(SHorizontalBox)
+				+ SHorizontalBox::Slot()
+				.HAlign(HAlign_Center)
+				.VAlign(VAlign_Fill)
 				.Padding(0.0f)
+				.AutoWidth()
 				[
-					SNew(SHorizontalBox)
-					+ SHorizontalBox::Slot()
+					SNew(SBorder)
+					.BorderImage(FEditorStyle::GetBrush("DetailsView.CategoryTop"))
+					.Padding(FMargin(30.0f, 3.0f, 30.0f, 0.0f))
 					.HAlign(HAlign_Center)
 					.VAlign(VAlign_Fill)
-					.Padding(0.0f)
-					.AutoWidth()
 					[
-						SNew(SBorder)
-						.BorderImage(FEditorStyle::GetBrush("DetailsView.CategoryTop"))
-						.Padding(FMargin(30.0f, 3.0f, 30.0f, 0.0f))
-						.HAlign(HAlign_Center)
-						.VAlign(VAlign_Fill)
-						[
-							SNew(STextBlock)
-							.Text(FText::FromString("Active Pose"))	
-						]
+						SNew(STextBlock)
+						.Text(FText::FromString("Active Pose"))	
 					]
 				]
+			]
 
-				// Active row list view with scroll bar
-				+ SVerticalBox::Slot()
+			// Active row list view with scroll bar
+			+ SVerticalBox::Slot()
+			
+			.AutoHeight()
+			[
+				SNew(SHorizontalBox)
 				
-				.AutoHeight()
+				+ SHorizontalBox::Slot()
+				.HAlign(HAlign_Fill)
+				.VAlign(VAlign_Fill)
+				.Padding(0.0f)
 				[
-					SNew(SHorizontalBox)
+					SNew(SBorder)
 					
-					+ SHorizontalBox::Slot()
-					.HAlign(HAlign_Fill)
-					.VAlign(VAlign_Fill)
+					.BorderImage(FEditorStyle::GetBrush("NoBorder"))
 					.Padding(0.0f)
 					[
-						SNew(SBorder)
-						
-						.BorderImage(FEditorStyle::GetBrush("NoBorder"))
-						.Padding(0.0f)
-						[
-							ActiveView.ListView.ToSharedRef()
-						]
+						ActiveView.ListView.ToSharedRef()
 					]
+				]
 
-					+ SHorizontalBox::Slot()
-					.AutoWidth()
-					[
-						ActiveView.ScrollBar.ToSharedRef()
-					]
-				]	
-			]
+				+ SHorizontalBox::Slot()
+				.AutoWidth()
+				[
+					ActiveView.ScrollBar.ToSharedRef()
+				]
+			]	
+		]
 
+		+ SVerticalBox::Slot()
+		.Padding(0.0f, 0.0f, 0.0f, 5.0f)
+		[
+			// Database view text tab
+			SNew(SVerticalBox)
 			+ SVerticalBox::Slot()
-			.Padding(0.0f, 0.0f, 0.0f, 5.0f)
 			[
-				// Database view text tab
-				SNew(SVerticalBox)
-				+ SVerticalBox::Slot()
+				SNew(SHorizontalBox)
+				+ SHorizontalBox::Slot()
+				.HAlign(HAlign_Center)
+				.VAlign(VAlign_Fill)
+				.Padding(0.0f)
 				[
-					SNew(SHorizontalBox)
-					+ SHorizontalBox::Slot()
+					SNew(SBorder)
+					.BorderImage(FEditorStyle::GetBrush("DetailsView.CategoryTop"))
+					.Padding(FMargin(30.0f, 3.0f, 30.0f, 0.0f))
 					.HAlign(HAlign_Center)
 					.VAlign(VAlign_Fill)
-					.Padding(0.0f)
-					[
-						SNew(SBorder)
-						.BorderImage(FEditorStyle::GetBrush("DetailsView.CategoryTop"))
-						.Padding(FMargin(30.0f, 3.0f, 30.0f, 0.0f))
-						.HAlign(HAlign_Center)
-						.VAlign(VAlign_Fill)
-						[
-							SNew(STextBlock)
-							.Text(FText::FromString("Selected Poses"))	
-						]
-					]
-					.AutoWidth()
-					
-					+ SHorizontalBox::Slot()
-					.HAlign(HAlign_Fill)
 					[
-						SNew(SBorder)
-						.BorderImage(&DatabaseView.RowStyle.EvenRowBackgroundBrush)
+						SNew(STextBlock)
+						.Text(FText::FromString("Selected Poses"))	
 					]
 				]
-				.AutoHeight()
-
-				// Gray line below the tab 
-				+ SVerticalBox::Slot()
-				.AutoHeight()
-				.Padding(0.0f)
+				.AutoWidth()
+				
+				+ SHorizontalBox::Slot()
+				.HAlign(HAlign_Fill)
 				[
 					SNew(SBorder)
-					.BorderImage(FEditorStyle::GetBrush("DetailsView.CategoryTop"))
-					.Padding(FMargin(0.0f, 3.0f, 0.0f, 3.0f))
-					.HAlign(HAlign_Fill)
-					.VAlign(VAlign_Fill)
+					.BorderImage(&DatabaseView.RowStyle.EvenRowBackgroundBrush)
 				]
-				
+			]
+			.AutoHeight()
+
+			// Gray line below the tab 
+			+ SVerticalBox::Slot()
+			.AutoHeight()
+			.Padding(0.0f)
+			[
+				SNew(SBorder)
+				.BorderImage(FEditorStyle::GetBrush("DetailsView.CategoryTop"))
+				.Padding(FMargin(0.0f, 3.0f, 0.0f, 3.0f))
+				.HAlign(HAlign_Fill)
+				.VAlign(VAlign_Fill)
+			]
 			
-				+ SVerticalBox::Slot()
+		
+			+ SVerticalBox::Slot()
+			[
+				SNew(SHorizontalBox)
+				+ SHorizontalBox::Slot()
+				.Padding(0.0f)
 				[
-					SNew(SHorizontalBox)
-					+ SHorizontalBox::Slot()
+					SNew(SBorder)
+					.BorderImage(FEditorStyle::GetBrush("NoBorder"))
 					.Padding(0.0f)
 					[
-						SNew(SBorder)
-						.BorderImage(FEditorStyle::GetBrush("NoBorder"))
-						.Padding(0.0f)
-						[
-							DatabaseView.ListView.ToSharedRef()
-						]
-					]
-					
-					+ SHorizontalBox::Slot()
-					.AutoWidth()
-					[
-						DatabaseView.ScrollBar.ToSharedRef()
+						DatabaseView.ListView.ToSharedRef()
 					]
 				]
+				
+				+ SHorizontalBox::Slot()
+				.AutoWidth()
+				[
+					DatabaseView.ScrollBar.ToSharedRef()
+				]
 			]
 		]
 	];
@@ -793,20 +798,22 @@ void SDebuggerDatabaseView::Construct(const FArguments& InArgs)
 	SortColumn = FPoseIdx::Name;
 	SortMode = EColumnSortMode::Ascending;
 
+	// Active view scroll bar only for indenting the columns to align w/ database
 	ActiveView.ScrollBar->SetVisibility(EVisibility::Hidden);
-
-	// Assign the get column function on the database rows
-	SDebuggerDatabaseRow::GetColumn = FGetColumnDelegate::CreateLambda([this](const FName& Name) -> TSharedRef<IColumn>&
-	{
-		return Columns[Name];
-	});
 	
 	RefreshColumns();
 }
 
-void SDebuggerDetailsView::Construct(const FArguments& InArgs, UPoseSearchDebuggerReflection* Reflection)
+void SDebuggerDetailsView::Construct(const FArguments& InArgs)
 {
-	/** Construct this panel using the property editor */
+	// Add property editor (detail view) UObject to world root so that it persists when PIE is stopped
+	Reflection = NewObject<UPoseSearchDebuggerReflection>();
+	Reflection->AddToRoot();
+	
+	check(IsValid(Reflection));
+
+	// @TODO: Convert this to a custom builder instead of of a standard details view
+	// Load property module and create details view with our reflection UObject
 	FPropertyEditorModule& PropPlugin = FModuleManager::LoadModuleChecked<FPropertyEditorModule>("PropertyEditor");
 	FDetailsViewArgs DetailsViewArgs;
     DetailsViewArgs.NameAreaSettings = FDetailsViewArgs::HideNameArea;
@@ -822,34 +829,129 @@ void SDebuggerDetailsView::Construct(const FArguments& InArgs, UPoseSearchDebugg
 	];
 }
 
+void SDebuggerDetailsView::Update(const FTraceMotionMatchingStateMessage& State, const UPoseSearchDatabase& Database) const
+{
+	UpdateReflection(State, Database);
+}
+
+SDebuggerDetailsView::~SDebuggerDetailsView()
+{
+	// Our previously instantiated object attached to root may be cleaned up at this point
+	if (UObjectInitialized())
+	{
+		Reflection->RemoveFromRoot();
+	}
+}
+
+void SDebuggerDetailsView::UpdateReflection(const FTraceMotionMatchingStateMessage& State, const UPoseSearchDatabase& Database) const
+{
+	check(Reflection);
+	const UPoseSearchSchema* Schema = Database.Schema;
+
+	Reflection->CurrentDatabaseName = Database.GetName();
+	Reflection->ElapsedPoseJumpTime = State.ElapsedPoseJumpTime;
+	Reflection->bFollowUpAnimation = EnumHasAnyFlags(State.Flags, FTraceMotionMatchingState::EFlags::FollowupAnimation);
+	
+	Reflection->PoseFeatures.EmptyAll();
+	Reflection->DistanceTrajectoryFeatures.EmptyAll();
+	Reflection->TimeTrajectoryFeatures.EmptyAll();
+
+	FFeatureVectorReader Reader;
+	// Ensure parity between Layout and QueryVector
+	Reader.Init(&Schema->Layout);
+	Reader.SetValues(State.QueryVector);
+
+	int32 NumSubsamples = Schema->PoseSampleTimes.Num();
+	const int32 NumBones = Schema->Bones.Num();
+
+	FPoseSearchFeatureDesc Feature;
+	
+	// Aggregate all features and place into the reflection struct
+	auto Extract = [](
+		const FFeatureVectorReader& InReader,
+		const FPoseSearchFeatureDesc& InFeature,
+		FPoseSearchDebuggerFeatureReflection& ReflectionRef
+	)
+	{
+		FVector OutputVec;
+		if (InReader.GetPosition(InFeature, &OutputVec))
+		{
+			ReflectionRef.Positions.Add(OutputVec);
+		}
+		if (InReader.GetLinearVelocity(InFeature, &OutputVec))
+		{
+			ReflectionRef.LinearVelocities.Add(OutputVec);
+		}
+		if (InReader.GetAngularVelocity(InFeature, &OutputVec))
+        {
+        	ReflectionRef.AngularVelocities.Add(OutputVec);
+        }
+	};
+
+	// Pose samples
+	Feature.Domain = EPoseSearchFeatureDomain::Time;
+	for (int32 SchemaSubsampleIndex = 0; SchemaSubsampleIndex < NumSubsamples; ++SchemaSubsampleIndex)
+	{
+		Feature.SubsampleIdx = SchemaSubsampleIndex;
+
+		for(int32 SchemaBoneIdx = 0; SchemaBoneIdx < NumBones; ++SchemaBoneIdx)
+		{
+			Feature.SchemaBoneIdx = SchemaBoneIdx;
+			Extract(Reader, Feature, Reflection->PoseFeatures);
+		}
+	}
+
+	// Used for classifying trajectories instead of bones, special index
+	Feature.SchemaBoneIdx = FPoseSearchFeatureDesc::TrajectoryBoneIndex;
+	
+	// Trajectory time samples
+	NumSubsamples = Schema->TrajectorySampleTimes.Num();
+	for (int32 SchemaSubsampleIndex = 0; SchemaSubsampleIndex < NumSubsamples; ++SchemaSubsampleIndex)
+	{
+		Feature.SubsampleIdx = SchemaSubsampleIndex;
+		Extract(Reader, Feature, Reflection->TimeTrajectoryFeatures);
+	}
+	// Trajectory distance samples
+	NumSubsamples = Schema->TrajectorySampleDistances.Num();
+	Feature.Domain = EPoseSearchFeatureDomain::Distance;
+	for (int32 SchemaSubsampleIndex = 0; SchemaSubsampleIndex < NumSubsamples; ++SchemaSubsampleIndex)
+	{
+		Feature.SubsampleIdx = SchemaSubsampleIndex;
+		Extract(Reader, Feature, Reflection->DistanceTrajectoryFeatures);
+	}
+}
 
 void SDebuggerView::Construct(const FArguments& InArgs, uint64 InAnimInstanceId)
 {
 	MotionMatchingNodeIds = InArgs._MotionMatchingNodeIds;
 	MotionMatchingState = InArgs._MotionMatchingState;
-	Reflection = InArgs._Reflection;
 	PoseSearchDatabase = InArgs._PoseSearchDatabase;
-	OnSelectionChanged = InArgs._OnSelectionChanged;
+	OnUpdateSelection = InArgs._OnUpdateSelection;
 	IsPIESimulating = InArgs._IsPIESimulating;
 	IsRecording = InArgs._IsRecording;
 	RecordingDuration = InArgs._RecordingDuration;
-	ActiveNodesNum = InArgs._ActiveNodesNum;
+	NodesNum = InArgs._NodesNum;
+	World = InArgs._World;
+	RootTransform = InArgs._RootTransform;
 	OnUpdate = InArgs._OnUpdate;
+	OnViewClosed = InArgs._OnViewClosed;
 
 	// Validate the existence of the passed getters
 	check(MotionMatchingNodeIds.IsBound());
     check(MotionMatchingState.IsBound());
-	check(Reflection.IsBound());
 	check(PoseSearchDatabase.IsBound());
-	check(OnSelectionChanged.IsBound());
+	check(OnUpdateSelection.IsBound());
 	check(IsPIESimulating.IsBound())
     check(IsRecording.IsBound());
     check(RecordingDuration.IsBound());
-	check(ActiveNodesNum.IsBound());
+	check(NodesNum.IsBound());
+	check(World.IsBound());
+	check(RootTransform.IsBound());
 	check(OnUpdate.IsBound());
+	check(OnViewClosed.IsBound());
 	
 	AnimInstanceId = InAnimInstanceId;
-	SelectedNode = -1;
+	SelectedNodeId = -1;
 
 	ChildSlot
 	[
@@ -860,7 +962,7 @@ void SDebuggerView::Construct(const FArguments& InArgs, uint64 InAnimInstanceId)
 		.HAlign(HAlign_Fill)
 		[
 			SAssignNew(Switcher, SWidgetSwitcher)
-			.WidgetIndex_Raw(this, &SDebuggerView::SelectView)
+			.WidgetIndex(this, &SDebuggerView::SelectView)
 
 			// [0] Selection view before node selection is made
 			+ SWidgetSwitcher::Slot()
@@ -905,98 +1007,204 @@ void SDebuggerView::SetTimeMarker(double InTimeMarker)
 		return;
 	}
 
-	const bool SameTime = FMath::Abs(InTimeMarker - TimeMarker) < DOUBLE_SMALL_NUMBER;
 	TimeMarker = InTimeMarker;
+}
+
+void SDebuggerView::Tick(const FGeometry& AllottedGeometry, const double InCurrentTime, const float InDeltaTime)
+{
+	if (IsPIESimulating.Get())
+	{
+		return;
+	}
+	
+	const UWorld* DebuggerWorld = World.Get();
+    // World should be always valid at this point
+    check(DebuggerWorld);
+	
+	// @TODO: Handle editor world when those features are enabled for the Rewind Debugger
+	// Currently prevents debug draw remnants from stopped world
+	if (DebuggerWorld->WorldType != EWorldType::PIE)
+	{
+		return;
+	}
+	
+	const bool bSameTime = FMath::Abs(TimeMarker - PreviousTimeMarker) < DOUBLE_SMALL_NUMBER;
+	PreviousTimeMarker = TimeMarker;
+
+	OnUpdate.Execute(AnimInstanceId);
+	OnUpdateSelection.Execute(SelectedNodeId);
+	const FTraceMotionMatchingStateMessage* State = MotionMatchingState.Get();
+	const UPoseSearchDatabase* Database = PoseSearchDatabase.Get();
 
 	// We haven't reached the update point yet
 	if (CurrentConsecutiveFrames < ConsecutiveFramesUpdateThreshold)
 	{
 		// If we're on the same time marker, it is consecutive
-		if (SameTime)
+		if (bSameTime)
 		{
 			++CurrentConsecutiveFrames;
 		}
-		return;
 	}
-	// New frame after having updated, reset consecutive frames count and start counting again
-	if (!SameTime)
+	else
 	{
-		CurrentConsecutiveFrames = 0;
-		bUpdated = false;
-		return;
+		// New frame after having updated, reset consecutive frames count and start counting again
+		if (!bSameTime)
+		{
+			CurrentConsecutiveFrames = 0;
+			bUpdated = false;
+		}
+		// Haven't updated since passing through frame gate, update once
+		else if (!bUpdated)
+		{
+			const bool bNodeSelected = UpdateSelection();
+			if (bNodeSelected)
+			{
+				OnUpdateSelection.Execute(SelectedNodeId);
+				// Refresh state and database info after update
+				State = MotionMatchingState.Get();
+				Database = PoseSearchDatabase.Get();
+				if (!State || !Database)
+				{
+					return;
+				}
+				
+				UpdateViews(*State, *Database);
+			}
+			bUpdated = true;	
+		}
 	}
-
-	// Haven't updated since passing through frame gate, update once
-	if (!bUpdated)
+	
+	const FTransform* Transform = RootTransform.Get();
+	if (!State || !Database || !Transform)
 	{
-		UpdateViews();
-    	bUpdated = true;	
+		return;
 	}
+
+	// Drawing of the debug vector should occur regardless of database update threshold
+	DrawFeatures(*DebuggerWorld, *State, *Database, *Transform);
 }
 
-void SDebuggerView::UpdateViews()
+bool SDebuggerView::UpdateSelection()
 {
-	OnUpdate.Execute(AnimInstanceId);
-	
 	// Update selection view if no node selected
-	if (SelectedNode == -1)
+	bool bNodeSelected = SelectedNodeId != INDEX_NONE;
+	if (!bNodeSelected)
 	{
-		TSet<int32> NodeIds = TSet<int32>(*MotionMatchingNodeIds.Get());
-
-		// If we have a new set of nodes
-		if (!NodeIds.Difference(StoredNodes).IsEmpty())
+		const TArray<int32>& NodeIds = *MotionMatchingNodeIds.Get();
+		// Only one node active, bypass selection view
+		if (NodeIds.Num() == 1)
 		{
-			// Only one node active, bypass selection view
-			if (NodeIds.Num() == 1)
-			{
-				SelectedNode = *NodeIds.begin();
-				UpdateViews();
-			}
-			// Create selection view with buttons for each node, displaying the database name
-			else
+			SelectedNodeId = *NodeIds.begin();
+			bNodeSelected = true;
+		}
+		// Create selection view with buttons for each node, displaying the database name
+		else
+		{
+			SelectionView->ClearChildren();
+			for (int32 NodeId : NodeIds)
 			{
-				SelectionView->ClearChildren();
-				for (int32 NodeId : NodeIds)
-				{
-					OnSelectionChanged.Execute(NodeId);
-					SelectionView->AddSlot()
-					.HAlign(HAlign_Fill)
+				OnUpdateSelection.Execute(NodeId);
+				SelectionView->AddSlot()
+				.HAlign(HAlign_Fill)
+				.VAlign(VAlign_Center)
+				.Padding(10.0f)
+				[
+					SNew(SButton)
+					.Text(FText::FromString(PoseSearchDatabase.Get()->GetName()))
+					.HAlign(HAlign_Center)
 					.VAlign(VAlign_Center)
-					.Padding(10.0f)
-					[
-						SNew(SButton)
-						.Text(FText::FromString(PoseSearchDatabase.Get()->GetName()))
-						.HAlign(HAlign_Center)
-						.VAlign(VAlign_Center)
-						.ContentPadding(10.0f)
-						.OnClicked_Raw(this, &SDebuggerView::UpdateNodeSelection, NodeId)
-					];
-				}
+					.ContentPadding(10.0f)
+					.OnClicked(this, &SDebuggerView::OnUpdateNodeSelection, NodeId)
+				];
 			}
 		}
-		
-		// Update active node list
-		StoredNodes = MoveTemp(NodeIds);
 	}
-	else
+
+	return bNodeSelected;
+}
+
+void SDebuggerView::UpdateViews(const FTraceMotionMatchingStateMessage& State, const UPoseSearchDatabase& Database) const
+{
+	DatabaseView->Update(State, Database);
+	DetailsView->Update(State, Database);
+}
+
+void SDebuggerView::DrawFeatures(
+	const UWorld& DebuggerWorld,
+	const FTraceMotionMatchingStateMessage& State,
+	const UPoseSearchDatabase& Database,
+	const FTransform& Transform
+) const
+{
+	// Set shared state
+	FDebugDrawParams DrawParams;
+	DrawParams.Database = &Database;
+	DrawParams.PoseVector = State.QueryVector;
+	DrawParams.World = &DebuggerWorld;
+	DrawParams.RootTransform = Transform;
+	// Single frame render
+	DrawParams.DefaultLifeTime = 0.0f;
+
+	const TObjectPtr<UPoseSearchDebuggerReflection>& Reflection = DetailsView->GetReflection();
+
+	auto SetDrawFlags = [](FDebugDrawParams& InDrawParams, const FPoseSearchDebuggerFeatureDrawOptions& Options)
+	{
+		InDrawParams.Flags = EDebugDrawFlags::None;
+		if (Options.bDisable)
+		{
+			return;
+		}
+
+		if (Options.bDrawPoseFeatures)
+		{
+			InDrawParams.Flags |= EDebugDrawFlags::IncludePose;
+		}
+
+		if (Options.bDrawTrajectoryFeatures)
+		{
+			InDrawParams.Flags |= EDebugDrawFlags::IncludeTrajectory;
+		}
+	};
+
+	// Draw query vector
+	DrawParams.Color = &FLinearColor::Blue;
+	SetDrawFlags(DrawParams, Reflection->QueryDrawOptions);
+	Draw(DrawParams);
+	DrawParams.PoseVector = {};
+
+	const TSharedPtr<SListView<TSharedRef<FDebuggerDatabaseRowData>>>& DatabaseRows = DatabaseView->GetDatabaseRows();
+	TArray<TSharedRef<FDebuggerDatabaseRowData>> Selected = DatabaseRows->GetSelectedItems();
+
+	// Red for non-active database view
+	DrawParams.Color = &FLinearColor::Red;
+	SetDrawFlags(DrawParams, Reflection->SelectedPoseDrawOptions);
+
+	// Draw any selected database vectors
+	for (const TSharedRef<FDebuggerDatabaseRowData>& Row : Selected)
+	{
+		DrawParams.PoseIdx = Row->PoseIdx;
+		Draw(DrawParams);
+	}
+
+	Selected = DatabaseView->GetActiveRow()->GetSelectedItems();
+	
+	// Active row should only have 0 or 1
+	check(Selected.Num() < 2);
+
+	if (!Selected.IsEmpty())
 	{
-		check(Reflection.Get());
-    	OnSelectionChanged.Execute(SelectedNode);
+		// Green for the active view
+		DrawParams.Color = &FLinearColor::Green;
 		
-    	const FTraceMotionMatchingStateMessage* State = MotionMatchingState.Get();
-    	const UPoseSearchDatabase* Database = PoseSearchDatabase.Get();
-    	if (!State || !Database)
-    	{
-    		return;
-    	}
-    
-    	SwitcherViewType = Debugger;
-    	DatabaseView->Update(State, Database);	
+		// Use the motion-matching state's pose idx, as the active row may be update-throttled at this point
+		DrawParams.PoseIdx = State.DbPoseIdx;
+		Draw(DrawParams);
 	}
 }
 
 int32 SDebuggerView::SelectView() const
 {
+	// Currently recording
 	if (IsPIESimulating.Get())
 	{
 		if (IsRecording.Get())
@@ -1005,32 +1213,36 @@ int32 SDebuggerView::SelectView() const
 		}
 	}
 
+	// Data has not been recorded yet
 	if (RecordingDuration.Get() < DOUBLE_SMALL_NUMBER)
 	{
 		return StoppedMsg;
 	}
 
-	const bool bNoActiveNodes = ActiveNodesNum.Get() == 0;
-	const bool bNodeSelectedWithoutData = SelectedNode != INDEX_NONE && MotionMatchingState.Get() == nullptr;
-	
+	const bool bNoActiveNodes = NodesNum.Get() == 0;
+	const bool bNodeSelectedWithoutData = SelectedNodeId != INDEX_NONE && MotionMatchingState.Get() == nullptr;
+
+	// No active nodes, or node selected has no data
 	if (bNoActiveNodes || bNodeSelectedWithoutData)
     {
     	return NoDataMsg;
     }
-	
-	if (SelectedNode == INDEX_NONE)
+
+	// Node not selected yet, showcase selection view
+	if (SelectedNodeId == INDEX_NONE)
 	{
 		return Selection;
 	}
 
+	// Standard debugger view
 	return Debugger;
 }
 
-FReply SDebuggerView::UpdateNodeSelection(int32 InSelectedNode)
+FReply SDebuggerView::OnUpdateNodeSelection(int32 InSelectedNodeId)
 {
 	// -1 will backtrack to selection view
-	SelectedNode = InSelectedNode;
-	UpdateViews();
+	SelectedNodeId = InSelectedNodeId;
+	bUpdated = false;
 	return FReply::Handled();
 }
 
@@ -1040,7 +1252,7 @@ TSharedRef<SWidget> SDebuggerView::GenerateNoDataMessageView()
 	ReturnButtonView->SetVisibility(TAttribute<EVisibility>::CreateLambda([this]() -> EVisibility
 	{
 		// Hide the return button for the no data message if we have no nodes at all
-		return ActiveNodesNum.Get() > 0 ? EVisibility::Visible : EVisibility::Hidden;
+		return NodesNum.Get() > 0 ? EVisibility::Visible : EVisibility::Hidden;
 	}));
 	
 	return 
@@ -1074,7 +1286,7 @@ TSharedRef<SWidget> SDebuggerView::GenerateReturnButtonView()
 			.VAlign(VAlign_Center)
 			.ButtonStyle(FEditorStyle::Get(), "SimpleButton")
 			.ContentPadding( FMargin(1, 0) )
-			.OnClicked_Raw(this, &SDebuggerView::UpdateNodeSelection, static_cast<int32>(INDEX_NONE))
+			.OnClicked(this, &SDebuggerView::OnUpdateNodeSelection, static_cast<int32>(INDEX_NONE))
 			// Contents of button, icon then text
 			[
 				SNew(SHorizontalBox)
@@ -1106,7 +1318,7 @@ TSharedRef<SWidget> SDebuggerView::GenerateNodeDebuggerView()
 	ReturnButtonView->SetVisibility(TAttribute<EVisibility>::CreateLambda([this]() -> EVisibility
 		{
 			// Collapse this view if we have don't have more than 1 node
-			return ActiveNodesNum.Get() > 1 ? EVisibility::Visible : EVisibility::Collapsed;
+			return NodesNum.Get() > 1 ? EVisibility::Visible : EVisibility::Collapsed;
 		}
 	));
 	
@@ -1137,7 +1349,7 @@ TSharedRef<SWidget> SDebuggerView::GenerateNodeDebuggerView()
 		+ SSplitter::Slot()
 		.Value(0.35f)
 		[
-			SAssignNew(DetailsView, SDebuggerDetailsView, Reflection.Get())
+			SAssignNew(DetailsView, SDebuggerDetailsView)
 		];
 }
 
@@ -1152,132 +1364,20 @@ uint64 SDebuggerView::GetObjectId() const
 	return AnimInstanceId;
 }
 
-
-FDebugger* FDebugger::InternalInstance;
-void FDebugger::Initialize()
+SDebuggerView::~SDebuggerView()
 {
-	InternalInstance = new FDebugger;
-	IModularFeatures::Get().RegisterModularFeature(IRewindDebuggerExtension::ModularFeatureName, InternalInstance);
-
-	// Add property editor (detail view) UObject to world root so that it persists when PIE is stopped
-	InternalInstance->Reflection = NewObject<UPoseSearchDebuggerReflection>();
-	InternalInstance->Reflection->AddToRoot();
-	
-	check(IsValid(InternalInstance->Reflection));
+	OnViewClosed.Execute(AnimInstanceId);
 }
 
-void FDebugger::Shutdown()
+void FDebuggerInstance::UpdateMotionMatchingStates(uint64 InAnimInstanceId)
 {
-	FPropertyEditorModule& PropPlugin = FModuleManager::LoadModuleChecked<FPropertyEditorModule>("PropertyEditor");
-	PropPlugin.UnregisterCustomClassLayout(UPoseSearchDebuggerReflection::StaticClass()->GetFName());
-
-	// Our previously instantiated object attached to root may be cleaned up at this point
-	if (UObjectInitialized())
-	{
-		InternalInstance->Reflection->RemoveFromRoot();
-	}
-	
-	IModularFeatures::Get().UnregisterModularFeature(IRewindDebuggerExtension::ModularFeatureName, InternalInstance);
-	delete InternalInstance;
-}
-
-void FDebugger::Update(float DeltaTime, IRewindDebugger* InRewindDebugger)
-{
-	// Update active rewind debugger in use
-	RewindDebugger = InRewindDebugger;
-}
-
-void FDebugger::UpdateReflection() const
-{
-	check(Reflection);
-	if (!ActiveMotionMatchingState)
-	{
-		return;
-	}
-
-	Reflection->ElapsedPoseJumpTime = ActiveMotionMatchingState->ElapsedPoseJumpTime;
-
-	const UPoseSearchDatabase* Database = GetPoseSearchDatabase();
-	if (Database == nullptr)
+	NodeIds.Empty();
+	MotionMatchingStates.Empty();
+	if (!RewindDebugger)
 	{
 		return;
 	}
-	// Validated by the get above
-	const UPoseSearchSchema* Schema = Database->Schema;
-	
-	Reflection->PoseFeatures.EmptyAll();
-	Reflection->DistanceTrajectoryFeatures.EmptyAll();
-	Reflection->TimeTrajectoryFeatures.EmptyAll();
-
-	FFeatureVectorReader Reader;
-	// Ensure parity between Layout and QueryVector
-	Reader.Init(&Schema->Layout);
-	Reader.SetValues(ActiveMotionMatchingState->QueryVector);
-
-	int32 NumSubsamples = Schema->PoseSampleTimes.Num();
-	const int32 NumBones = Schema->Bones.Num();
-
-	FPoseSearchFeatureDesc Feature;
-	
-	// Aggregate all features and place into the reflection struct
-	auto Extract = [](
-		const FFeatureVectorReader& InReader,
-		const FPoseSearchFeatureDesc& InFeature,
-		FPoseSearchDebuggerFeatureReflection& ReflectionRef
-	)
-	{
-		FVector OutputVec;
-		if (InReader.GetPosition(InFeature, &OutputVec))
-		{
-			ReflectionRef.Positions.Add(OutputVec);
-		}
-		if (InReader.GetLinearVelocity(InFeature, &OutputVec))
-		{
-			ReflectionRef.LinearVelocities.Add(OutputVec);
-		}
-		if (InReader.GetAngularVelocity(InFeature, &OutputVec))
-        {
-        	ReflectionRef.AngularVelocities.Add(OutputVec);
-        }
-	};
-
-	// Pose samples
-	Feature.Domain = EPoseSearchFeatureDomain::Time;
-	for (int32 SchemaSubsampleIndex = 0; SchemaSubsampleIndex < NumSubsamples; ++SchemaSubsampleIndex)
-	{
-		Feature.SubsampleIdx = SchemaSubsampleIndex;
-
-		for(int32 SchemaBoneIdx = 0; SchemaBoneIdx < NumBones; ++SchemaBoneIdx)
-		{
-			Feature.SchemaBoneIdx = SchemaBoneIdx;
-			Extract(Reader, Feature, Reflection->PoseFeatures);
-		}
-	}
-
-	// Used for classifying trajectories instead of bones, special index
-	Feature.SchemaBoneIdx = FPoseSearchFeatureDesc::TrajectoryBoneIndex;
 	
-	// Trajectory time samples
-	NumSubsamples = Schema->TrajectorySampleTimes.Num();
-	for (int32 SchemaSubsampleIndex = 0; SchemaSubsampleIndex < NumSubsamples; ++SchemaSubsampleIndex)
-	{
-		Feature.SubsampleIdx = SchemaSubsampleIndex;
-		Extract(Reader, Feature, Reflection->TimeTrajectoryFeatures);
-	}
-	// Trajectory distance samples
-	NumSubsamples = Schema->TrajectorySampleDistances.Num();
-	Feature.Domain = EPoseSearchFeatureDomain::Distance;
-	for (int32 SchemaSubsampleIndex = 0; SchemaSubsampleIndex < NumSubsamples; ++SchemaSubsampleIndex)
-	{
-		Feature.SubsampleIdx = SchemaSubsampleIndex;
-		Extract(Reader, Feature, Reflection->DistanceTrajectoryFeatures);
-	}
-}
-
-void FDebugger::UpdateMotionMatchingStates(uint64 InAnimInstanceId)
-{
-	NodeIds.Empty();
-	MotionMatchingStates.Empty();
 	// Get provider and validate
 	const TraceServices::IAnalysisSession* Session = RewindDebugger->GetAnalysisSession();
 	TraceServices::FAnalysisSessionReadScope SessionReadScope(*Session);
@@ -1312,19 +1412,24 @@ void FDebugger::UpdateMotionMatchingStates(uint64 InAnimInstanceId)
 }
 
 
-const FTraceMotionMatchingStateMessage* FDebugger::GetMotionMatchingState() const
+FDebuggerInstance::FDebuggerInstance(uint64 InAnimInstanceId)
+	: AnimInstanceId(InAnimInstanceId)
+{
+}
+
+const FTraceMotionMatchingStateMessage* FDebuggerInstance::GetMotionMatchingState() const
 {
 	return ActiveMotionMatchingState;
 }
 
-const UPoseSearchDatabase* FDebugger::GetPoseSearchDatabase() const
+const UPoseSearchDatabase* FDebuggerInstance::GetPoseSearchDatabase() const
 {
-	if (!InternalInstance->ActiveMotionMatchingState)
+	if (!ActiveMotionMatchingState)
 	{
 		return nullptr;
 	}
 
-	const uint64 DatabaseId = InternalInstance->ActiveMotionMatchingState->DatabaseId;
+	const uint64 DatabaseId = ActiveMotionMatchingState->DatabaseId;
 	if (DatabaseId == 0)
 	{
 	    return nullptr;
@@ -1347,44 +1452,65 @@ const UPoseSearchDatabase* FDebugger::GetPoseSearchDatabase() const
 	return Database;
 }
 
-UPoseSearchDebuggerReflection* FDebugger::GetReflection() const
+const TArray<int32>* FDebuggerInstance::GetNodeIds() const
 {
-	return Reflection;
+	return &NodeIds;
 }
 
-bool FDebugger::IsPIESimulating() const
-{
-	return RewindDebugger->IsPIESimulating();
-}
 
-const TArray<int32>* FDebugger::GetNodeIds() const
+int32 FDebuggerInstance::GetNodesNum() const
 {
-	return &NodeIds;
+	return MotionMatchingStates.Num();
 }
 
-bool FDebugger::IsRecording() const
+const FTransform* FDebuggerInstance::GetRootTransform() const
 {
-	return RewindDebugger->IsRecording();
-}
+	if (!RewindDebugger || !ActiveMotionMatchingState)
+	{
+		return nullptr;
+	}
+	
+	// Get provider and validate
+	const TraceServices::IAnalysisSession* Session = RewindDebugger->GetAnalysisSession();
+	TraceServices::FAnalysisSessionReadScope SessionReadScope(*Session);
+	
+	const IAnimationProvider* AnimationProvider = Session->ReadProvider<IAnimationProvider>("AnimationProvider");
+	if (!AnimationProvider)
+	{
+		return nullptr;
+	}
 
-double FDebugger::GetRecordingDuration() const
-{
-	return RewindDebugger->GetRecordingDuration();
-}
+	const double TraceTime = RewindDebugger->CurrentTraceTime();
+	TraceServices::FFrame Frame;
+	ReadFrameProvider(*Session).GetFrameFromTime(TraceFrameType_Game, TraceTime, Frame);
 
-int32 FDebugger::GetActiveNodesNum() const
-{
-	return MotionMatchingStates.Num();
+	const FTransform* RootTransform = nullptr;
+
+	AnimationProvider->ReadSkeletalMeshPoseTimeline(ActiveMotionMatchingState->SkeletalMeshComponentId, [&RootTransform, &Frame](const IAnimationProvider::SkeletalMeshPoseTimeline& TimelineData, bool bHasCurves)
+	{
+		TimelineData.EnumerateEvents(Frame.StartTime, Frame.EndTime,
+			[&RootTransform, &Frame](double InStartTime, double InEndTime, uint32 InDepth, const FSkeletalMeshPoseMessage& PoseMessage)
+			{
+				RootTransform = &PoseMessage.ComponentToWorld;
+				return TraceServices::EEventEnumerate::Stop;
+			});
+	});
+
+	return RootTransform;
 }
 
-void FDebugger::OnUpdate(uint64 InAnimInstanceId)
+void FDebuggerInstance::OnUpdate(uint64 InAnimInstanceId)
 {
 	UpdateMotionMatchingStates(InAnimInstanceId);
 }
 
-void FDebugger::OnSelectionChanged(int32 InNodeId)
+void FDebuggerInstance::OnUpdateSelection(int32 InNodeId)
 {
-	bool bFound = false;
+	if (InNodeId == INDEX_NONE)
+	{
+		return;
+	}
+	
 	// Find node in all motion matching states this frame
 	const int32 NodesNum = NodeIds.Num();
 	for(int i = 0; i < NodesNum; ++i)
@@ -1392,35 +1518,95 @@ void FDebugger::OnSelectionChanged(int32 InNodeId)
 		if (NodeIds[i] == InNodeId)
 		{
 			ActiveMotionMatchingState = MotionMatchingStates[i];
-			bFound = true;
-			break;
+			return;
 		}
 	}
 
-	if (bFound)
+	ActiveMotionMatchingState = nullptr;
+}
+
+FDebugger* FDebugger::Debugger;
+void FDebugger::Initialize()
+{
+	Debugger = new FDebugger;
+	IModularFeatures::Get().RegisterModularFeature(IRewindDebuggerExtension::ModularFeatureName, Debugger);
+
+}
+
+void FDebugger::Shutdown()
+{
+	IModularFeatures::Get().UnregisterModularFeature(IRewindDebuggerExtension::ModularFeatureName, Debugger);
+	delete Debugger;
+}
+
+void FDebugger::Update(float DeltaTime, IRewindDebugger* InRewindDebugger)
+{
+	// Update active rewind debugger in use
+	RewindDebugger = InRewindDebugger;
+
+	for (TSharedRef<FDebuggerInstance>& DebuggerInstance : DebuggerInstances)
 	{
-		UpdateReflection();
+		DebuggerInstance->RewindDebugger = RewindDebugger;
 	}
-	else
+}
+
+bool FDebugger::IsPIESimulating() const
+{
+	return RewindDebugger->IsPIESimulating();
+}
+
+bool FDebugger::IsRecording() const
+{
+	return RewindDebugger->IsRecording();
+}
+
+double FDebugger::GetRecordingDuration() const
+{
+	return RewindDebugger->GetRecordingDuration();
+}
+
+const UWorld* FDebugger::GetWorld() const
+{
+	return RewindDebugger->GetWorldToVisualize();
+}
+
+void FDebugger::OnViewClosed(uint64 InAnimInstanceId)
+{
+	for (int i = 0; i < DebuggerInstances.Num(); ++i)
 	{
-		ActiveMotionMatchingState = nullptr;
+		if (DebuggerInstances[i]->AnimInstanceId == InAnimInstanceId)
+		{
+			DebuggerInstances.RemoveAtSwap(i);
+			return;
+		}
 	}
+
+	// Should always be a valid remove
+	checkNoEntry();
 }
 
-TSharedPtr<SDebuggerView> FDebugger::GenerateView(uint64 InAnimInstanceId)
+TSharedPtr<SDebuggerView> FDebugger::GenerateInstance(uint64 InAnimInstanceId)
 {
-	return
-		SNew(SDebuggerView, InAnimInstanceId)
-		.MotionMatchingNodeIds_Raw(Instance(), &FDebugger::GetNodeIds)
-		.MotionMatchingState_Raw(Instance(), &FDebugger::GetMotionMatchingState)
-		.Reflection_Raw(Instance(), &FDebugger::GetReflection)
-		.PoseSearchDatabase_Raw(Instance(), &FDebugger::GetPoseSearchDatabase)
-		.IsPIESimulating_Raw(Instance(), &FDebugger::IsPIESimulating)
-		.IsRecording_Raw(Instance(), &FDebugger::IsRecording)
-		.RecordingDuration_Raw(Instance(), &FDebugger::GetRecordingDuration)
-		.ActiveNodesNum_Raw(Instance(), &FDebugger::GetActiveNodesNum)
-		.OnSelectionChanged_Raw(Instance(), &FDebugger::OnSelectionChanged)
-		.OnUpdate_Raw(Instance(), &FDebugger::OnUpdate);
+	TSharedRef<FDebuggerInstance>& DebuggerInstance = DebuggerInstances.Add_GetRef(MakeShared<FDebuggerInstance>(InAnimInstanceId));
+	
+	TSharedPtr<SDebuggerView> DebuggerView;
+
+	SAssignNew(DebuggerView, SDebuggerView, InAnimInstanceId)
+	.MotionMatchingNodeIds(DebuggerInstance, &FDebuggerInstance::GetNodeIds)
+	.MotionMatchingState(DebuggerInstance, &FDebuggerInstance::GetMotionMatchingState)
+	.PoseSearchDatabase(DebuggerInstance, &FDebuggerInstance::GetPoseSearchDatabase)
+	.NodesNum(DebuggerInstance, &FDebuggerInstance::GetNodesNum)
+	.RootTransform(DebuggerInstance, &FDebuggerInstance::GetRootTransform)
+	.OnUpdateSelection(DebuggerInstance, &FDebuggerInstance::OnUpdateSelection)
+	.OnUpdate(DebuggerInstance, &FDebuggerInstance::OnUpdate)
+	// Rewind debugger callbacks gathered from shared object
+	.World_Raw(this, &FDebugger::GetWorld)
+	.IsRecording_Raw(this, &FDebugger::IsRecording)
+	.IsPIESimulating_Raw(this, &FDebugger::IsPIESimulating)
+	.RecordingDuration_Raw(this, &FDebugger::GetRecordingDuration)
+	.OnViewClosed_Raw(this, &FDebugger::OnViewClosed);
+
+	return DebuggerView;
 }
 
 
@@ -1446,7 +1632,7 @@ FName FDebuggerViewCreator::GetTargetTypeName() const
 
 TSharedPtr<IRewindDebuggerView> FDebuggerViewCreator::CreateDebugView(uint64 ObjectId, double CurrentTime, const TraceServices::IAnalysisSession& InAnalysisSession) const
 {
-	return FDebugger::GenerateView(ObjectId);
+	return FDebugger::Get()->GenerateInstance(ObjectId);
 }
 
 FName FDebuggerViewCreator::GetName() const
diff --git a/Engine/Plugins/Experimental/Animation/PoseSearch/Source/Editor/Private/PoseSearchDebugger.h b/Engine/Plugins/Experimental/Animation/PoseSearch/Source/Editor/Private/PoseSearchDebugger.h
--- a/Engine/Plugins/Experimental/Animation/PoseSearch/Source/Editor/Private/PoseSearchDebugger.h
+++ b/Engine/Plugins/Experimental/Animation/PoseSearch/Source/Editor/Private/PoseSearchDebugger.h
@@ -10,11 +10,12 @@
 #include "PoseSearchDebugger.generated.h"
 
 namespace TraceServices { class IAnalysisSession; }
-class IUnrealInsightsModule;
+enum class EPoseSearchFeatureDomain;
 class UPoseSearchDatabase;
+class IUnrealInsightsModule;
 class ITableRow;
-class STableViewBase;
 class IDetailsView;
+class STableViewBase;
 class SVerticalBox;
 class SHorizontalBox;
 class SScrollBox;
@@ -22,7 +23,7 @@ class SWidgetSwitcher;
 template <typename ItemType> class SListView;
 
 /**
- * Used by the reflection UObject to encompass a set of features
+ * Used by the reflection UObject to encompass a set of feature vectors
  */
 USTRUCT()
 struct FPoseSearchDebuggerFeatureReflection
@@ -30,10 +31,13 @@ struct FPoseSearchDebuggerFeatureReflection
 	GENERATED_USTRUCT_BODY()
 
 	// @TODO: Should be ideally enumerated based on all possible schema features
+	
     UPROPERTY(VisibleAnywhere, Category="Query Data")
 	TArray<FVector> Positions;
+	
     UPROPERTY(VisibleAnywhere, Category="Query Data")
 	TArray<FVector> LinearVelocities;
+	
     UPROPERTY(VisibleAnywhere, Category="Query Data")
 	TArray<FVector> AngularVelocities;
 
@@ -41,6 +45,24 @@ struct FPoseSearchDebuggerFeatureReflection
 	void EmptyAll();
 };
 
+/**
+ * Used by the reflection UObject to encompass draw options for the query and database selections
+ */
+USTRUCT()
+struct FPoseSearchDebuggerFeatureDrawOptions
+{
+	GENERATED_USTRUCT_BODY()
+
+	UPROPERTY(EditAnywhere, Category="Draw Options")
+    bool bDisable = false;
+
+	UPROPERTY(EditAnywhere, Category="Draw Options", Meta=(EditCondition="!bDisable"))
+	bool bDrawPoseFeatures = true;
+
+	UPROPERTY(EditAnywhere, Category="Draw Options", Meta=(EditCondition="!bDisable"))
+    bool bDrawTrajectoryFeatures = true;
+};
+
 /**
  * Reflection UObject being observed in the details view panel of the debugger
  */
@@ -49,6 +71,9 @@ class POSESEARCHEDITOR_API UPoseSearchDebuggerReflection : public UObject
 {
 	GENERATED_UCLASS_BODY()
 
+	UPROPERTY(VisibleAnywhere, Category="Motion Matching State", Meta=(DisplayName="Current Database"))
+	FString CurrentDatabaseName = "";
+
 	/** Time since last PoseSearch jump */
 	UPROPERTY(VisibleAnywhere, Category="Motion Matching State")
 	float ElapsedPoseJumpTime = 0.0f;
@@ -57,6 +82,12 @@ class POSESEARCHEDITOR_API UPoseSearchDebuggerReflection : public UObject
 	UPROPERTY(VisibleAnywhere, Category="Motion Matching State")
 	bool bFollowUpAnimation = false;
 
+	UPROPERTY(EditAnywhere, Category="Draw Options", Meta=(DisplayName="Query"))
+	FPoseSearchDebuggerFeatureDrawOptions QueryDrawOptions;
+
+	UPROPERTY(EditAnywhere, Category="Draw Options", Meta=(DisplayName="Selected Pose"))
+	FPoseSearchDebuggerFeatureDrawOptions SelectedPoseDrawOptions;
+
 	/** Pose features of the current query vector */
     UPROPERTY(EditAnywhere, Category="Query Data")
 	FPoseSearchDebuggerFeatureReflection PoseFeatures;
@@ -74,10 +105,11 @@ class POSESEARCHEDITOR_API UPoseSearchDebuggerReflection : public UObject
 namespace UE { namespace PoseSearch {
 
 namespace DebuggerDatabaseColumns { struct IColumn; }
+struct FDebugDrawParams;
 struct FTraceMotionMatchingStateMessage;
+class FFeatureVectorReader;
 class FDebuggerDatabaseRowData;
 
-
 /**
  * Database panel view widget of the PoseSearch debugger
  */
@@ -87,19 +119,50 @@ class SDebuggerDatabaseView : public SCompoundWidget
 	SLATE_END_ARGS()
 	
 	void Construct(const FArguments& InArgs);
-	void Update(const FTraceMotionMatchingStateMessage* State, const UPoseSearchDatabase* Database);
+	void Update(const FTraceMotionMatchingStateMessage& State, const UPoseSearchDatabase& Database);
+
+	const TSharedPtr<SListView<TSharedRef<FDebuggerDatabaseRowData>>>& GetActiveRow() const { return ActiveView.ListView; }
+	const TSharedPtr<SListView<TSharedRef<FDebuggerDatabaseRowData>>>& GetDatabaseRows() const { return DatabaseView.ListView; }
+
+	/** Used by database rows to acquire column-specific information */
+	using FColumnMap = TMap<FName, TSharedRef<DebuggerDatabaseColumns::IColumn>>;
 
 private:
+	/** Deletes existing columns and initializes a new set */
 	void RefreshColumns();
+
+	/** Adds a column to the existing list */
 	void AddColumn(TSharedRef<DebuggerDatabaseColumns::IColumn>&& Column);
-	void CreateRows(const UPoseSearchDatabase* Database);
+
+	/** Retrieves current column map, used as an attribute by rows */
+	const FColumnMap* GetColumnMap() const { return &Columns; }
+
+	/** Creates widgets for every pose in the database, initializing the static data in the process */
+	void CreateRows(const UPoseSearchDatabase& Database);
+
+	/** Sorts the database by the current sort predicate, updating the view order */
 	void SortDatabaseRows();
-	void UpdateRows(const FTraceMotionMatchingStateMessage* State, const UPoseSearchDatabase* Database);
+
+	/** Sets dynamic data for each row, such as score at the current time */
+	void UpdateRows(const FTraceMotionMatchingStateMessage& State, const UPoseSearchDatabase& Database);
+
+	/** Acquires sort predicate for the given column */
 	EColumnSortMode::Type GetColumnSortMode(const FName ColumnId) const;
+
+	/** Gets active column width, used to align active and database view */
 	float GetColumnWidth(const FName ColumnId) const;
+	
+	/** Updates the active sort predicate, setting the sorting order of all other columns to none
+	 * (to be dependent on active column */
 	void OnColumnSortModeChanged(const EColumnSortPriority::Type SortPriority, const FName & ColumnId, const EColumnSortMode::Type InSortMode);
+
+	/** Aligns the active and database views */
 	void OnColumnWidthChanged(const float NewWidth, FName ColumnId) const;
+
+	/** Generates a database row widget for the given data */
 	TSharedRef<ITableRow> HandleGenerateDatabaseRow(TSharedRef<FDebuggerDatabaseRowData> Item, const TSharedRef<STableViewBase>& OwnerTable) const;
+	
+	/** Generates the active row widget for the given data */
 	TSharedRef<ITableRow> HandleGenerateActiveRow(TSharedRef<FDebuggerDatabaseRowData> Item, const TSharedRef<STableViewBase>& OwnerTable) const;
 
 	/** Current column to sort by */
@@ -108,7 +171,7 @@ private:
 	EColumnSortMode::Type SortMode = EColumnSortMode::Ascending;
 	
 	/** Column data container, used to emplace defined column structures of various types */
-    TMap<FName, TSharedRef<DebuggerDatabaseColumns::IColumn>> Columns;
+    FColumnMap Columns;
 
 	struct FTable
 	{
@@ -146,17 +209,27 @@ class SDebuggerDetailsView : public SCompoundWidget
 	SLATE_BEGIN_ARGS(SDebuggerDetailsView) {}
 	SLATE_END_ARGS()
 
-	void Construct(const FArguments& InArgs, UPoseSearchDebuggerReflection* Reflection);
+	void Construct(const FArguments& InArgs);
+	void Update(const FTraceMotionMatchingStateMessage& State, const UPoseSearchDatabase& Database) const;
+
+	/** Get a const version of our reflection object */
+	const TObjectPtr<UPoseSearchDebuggerReflection>& GetReflection() const { return Reflection; }
+	
+	virtual ~SDebuggerDetailsView() override;
+
 private:
+	/** Update our details view object with new state information */
+	void UpdateReflection(const FTraceMotionMatchingStateMessage& State, const UPoseSearchDatabase& Database) const;
+	
 	/** Details widget constructed for the MM node */
 	TSharedPtr<IDetailsView> Details;
+	/** Last updated reflection data relative to MM state */
+	TObjectPtr<UPoseSearchDebuggerReflection> Reflection = nullptr;
 };
 
 
-/**
- * Callback to update the debugger when node selection is changed
- */
-DECLARE_DELEGATE_OneParam(FOnSelectionChanged, int32 NodeId);
+/** Callback to update the debugger when node is actively selected */
+DECLARE_DELEGATE_OneParam(FOnUpdateSelection, int32 NodeId);
 
 /**
  * Callback to update the debugger when view update occurs, updating the motion matching state
@@ -164,6 +237,9 @@ DECLARE_DELEGATE_OneParam(FOnSelectionChanged, int32 NodeId);
  */
 DECLARE_DELEGATE_OneParam(FOnUpdate, uint64 AnimInstanceId);
 
+/** Callback to relay closing of the view to destroy the debugger instance */
+DECLARE_DELEGATE_OneParam(FOnViewClosed, uint64 AnimInstanceId);
+
 /**
  * Entire view of the PoseSearch debugger, containing all sub-widgets
  */
@@ -174,65 +250,97 @@ public:
 
 	SLATE_BEGIN_ARGS(SDebuggerView){}
 		SLATE_ATTRIBUTE(const FTraceMotionMatchingStateMessage*, MotionMatchingState)
-		SLATE_ATTRIBUTE(UPoseSearchDebuggerReflection*, Reflection)
 		SLATE_ATTRIBUTE(const UPoseSearchDatabase*, PoseSearchDatabase)
 		SLATE_ATTRIBUTE(const TArray<int32>*, MotionMatchingNodeIds)
 		SLATE_ATTRIBUTE(bool, IsPIESimulating)
 		SLATE_ATTRIBUTE(bool, IsRecording)
 		SLATE_ATTRIBUTE(double, RecordingDuration)
-		SLATE_ATTRIBUTE(int32, ActiveNodesNum)
-		SLATE_EVENT(FOnSelectionChanged, OnSelectionChanged)
+		SLATE_ATTRIBUTE(int32, NodesNum)
+		SLATE_ATTRIBUTE(const UWorld*, World)
+		SLATE_ATTRIBUTE(const FTransform*, RootTransform)
+		SLATE_EVENT(FOnUpdateSelection, OnUpdateSelection)
 		SLATE_EVENT(FOnUpdate, OnUpdate)
+		SLATE_EVENT(FOnViewClosed, OnViewClosed)
 	SLATE_END_ARGS()
 
 	void Construct(const FArguments& InArgs, uint64 InAnimInstanceId);
 	virtual void SetTimeMarker(double InTimeMarker) override;
+	virtual void Tick(const FGeometry& AllottedGeometry, const double InCurrentTime, const float InDeltaTime) override;
 	virtual FName GetName() const override;
 	virtual uint64 GetObjectId() const override;
 
+	virtual ~SDebuggerView() override;
+
 private:
-	void UpdateViews();
+	/** Called each frame to draw features of the query vector & database selections */
+	void DrawFeatures(const UWorld& DebuggerWorld, const FTraceMotionMatchingStateMessage& State, const UPoseSearchDatabase& Database, const FTransform& Transform) const;
+	
+	/** Check if a node selection was made, true if a node is selected */
+	bool UpdateSelection();
+
+	/** Update the database and details views */
+	void UpdateViews(const FTraceMotionMatchingStateMessage& State, const UPoseSearchDatabase& Database) const;
 	
 	/** Returns an int32 appropriate to the index of our widget selector */
 	int32 SelectView() const;
 
-	FReply UpdateNodeSelection(int32 InSelectedNode);
-	
+	/** Callback when a button in the selection view is clicked */
+	FReply OnUpdateNodeSelection(int32 InSelectedNodeId);
+
+	/** Generates the message view relaying that there is no data */
 	TSharedRef<SWidget> GenerateNoDataMessageView();
+
+	/** Generates the return button to go back to the selection mode */
 	TSharedRef<SWidget> GenerateReturnButtonView();
+
+	/** Generates the entire node debugger widget, including database and details view */
 	TSharedRef<SWidget> GenerateNodeDebuggerView();
 
 	/** Gets all MM nodes being traced in this frame */
 	TAttribute<const TArray<int32>*> MotionMatchingNodeIds;
-	/** Retrieves the reflection UObject from the debugger */
-	TAttribute<UPoseSearchDebuggerReflection*> Reflection;
+	
 	/** Retrieves the MM state from the debugger */
 	TAttribute<const FTraceMotionMatchingStateMessage*> MotionMatchingState;
+	
 	/** Retrieves the PoseSearch Database from the debugger */
 	TAttribute<const UPoseSearchDatabase*> PoseSearchDatabase;
+	
 	/** Whether the game is un-paused and currently simulating */
 	TAttribute<bool> IsPIESimulating;
+	
 	/** Whether the Rewind Debugger is currently recording gameplay */
 	TAttribute<bool> IsRecording;
+	
  	/** Length of the recorded sequence (if any) */
 	TAttribute<double> RecordingDuration;
+	
  	/** Number of active nodes at the current frame */
-	TAttribute<int32> ActiveNodesNum;
+	TAttribute<int32> NodesNum;
+	
+	/** Current world the Rewind Debugger is functioning in */
+	TAttribute<const UWorld*> World;
+	
+	/** Current Component to World Space transform of the Skeletal Mesh Component */
+	TAttribute<const FTransform*> RootTransform;
+	
 	/** Update current debugger data when node selection is changed */
-	FOnSelectionChanged OnSelectionChanged;
+	FOnUpdateSelection OnUpdateSelection;
+	
 	/** Update current debugger data when update occurs */
 	FOnUpdate OnUpdate;
+	
+	/** Destroy the debugger instanced when closed */
+	FOnViewClosed OnViewClosed;
 
 	/** Active node being debugged */
-	int32 SelectedNode = -1;
-
-	/** List of all nodes being traced (stored for selection) */
-	TSet<int32> StoredNodes;
+	int32 SelectedNodeId = -1;
 
 	/** Database view of the motion matching node */
 	TSharedPtr<SDebuggerDatabaseView> DatabaseView;
+	
 	/** Details panel for introspecting the motion matching node */
 	TSharedPtr<SDebuggerDetailsView> DetailsView;
+	
 	/** Node debugger view hosts the above two views */
 	TSharedPtr<SSplitter> NodeDebuggerView;
 
@@ -241,7 +349,7 @@ private:
 	
 	/** Gray box occluding the debugger view when simulating */
 	TSharedPtr<SVerticalBox> SimulatingView;
-	
+
 	/** Used to switch between views in the switcher, int32 maps to index in the SWidgetSwitcher */
 	enum ESwitcherViewType : int32
 	{
@@ -264,6 +372,9 @@ private:
 	/** Current position of the time marker */
 	double TimeMarker = -1.0;
 
+	/** Previous position of the time marker */
+	double PreviousTimeMarker = -1.0;
+
 	/** Tracks if the current time has been updated yet (delayed) */
 	bool bUpdated = false;
 
@@ -274,53 +385,78 @@ private:
 	static constexpr int32 ConsecutiveFramesUpdateThreshold = 10;
 };
 
+class FDebuggerInstance : public TSharedFromThis<FDebuggerInstance>
+{
+public:
+	explicit FDebuggerInstance(uint64 InAnimInstanceId);
+
+private:
+	// Used for view callbacks
+    const FTraceMotionMatchingStateMessage* GetMotionMatchingState() const;
+    const UPoseSearchDatabase* GetPoseSearchDatabase() const;
+	const TArray<int32>* GetNodeIds() const;
+	int32 GetNodesNum() const;
+	const FTransform* GetRootTransform() const;
+
+	/** Update motion matching states for frame */
+	void OnUpdate(uint64 InAnimInstanceId);
+	
+	/** Updates active motion matching state based on node selection */
+	void OnUpdateSelection(int32 InNodeId);
+
+	/** Update the list of states for this frame */	
+	void UpdateMotionMatchingStates(uint64 InAnimInstanceId);
+	
+	/** List of all Node IDs associated with motion matching states */
+	TArray<int32> NodeIds;
+	
+	/** List of all updated motion matching states per node */
+	TArray<const FTraceMotionMatchingStateMessage*> MotionMatchingStates;
+	
+	/** Currently active motion matching state based on node selection in the view */
+	const FTraceMotionMatchingStateMessage* ActiveMotionMatchingState = nullptr;
+
+	/** Stored Rewind Debugger instance provided by the extension */
+	const IRewindDebugger* RewindDebugger = nullptr;
+
+	/** Anim Instance associated with this debugger instance */
+	uint64 AnimInstanceId = 0;
+
+	/** Limits some public API */
+	friend class FDebugger;
+};
+
 
 /**
  * PoseSearch debugger, containing the data to be acquired and relayed to the view
  */
-class FDebugger : public IRewindDebuggerExtension, public TSharedFromThis<FDebugger>
+class FDebugger : public TSharedFromThis<FDebugger>, public IRewindDebuggerExtension
 {
 public:
 	virtual void Update(float DeltaTime, IRewindDebugger* InRewindDebugger) override;
 	virtual ~FDebugger() = default;
 
-	static FDebugger* Instance() { return InternalInstance; }
+	static FDebugger* Get() { return Debugger; }
 	static void Initialize();
 	static void Shutdown();
 	static const FName ModularFeatureName;
 
 	/** Generates the slate debugger view widget */
-	static TSharedPtr<SDebuggerView> GenerateView(uint64 InAnimInstanceId);
-
+	TSharedPtr<SDebuggerView> GenerateInstance(uint64 InAnimInstanceId);
 private:
-	// Used for view callbacks
-	const FTraceMotionMatchingStateMessage* GetMotionMatchingState() const;
-	const UPoseSearchDatabase* GetPoseSearchDatabase() const;
-	UPoseSearchDebuggerReflection* GetReflection() const;
 	bool IsPIESimulating() const;
-	const TArray<int32>* GetNodeIds() const;
 	bool IsRecording() const;
 	double GetRecordingDuration() const;
-	int32 GetActiveNodesNum() const;
-	void OnUpdate(uint64 InAnimInstanceId);
-	void OnSelectionChanged(int32 InNodeId);
+	const UWorld* GetWorld() const;
+	void OnViewClosed(uint64 InAnimInstanceId);
 
-	/** Updates the current reflection data relative to the MM state */
-	void UpdateReflection() const;
-	void UpdateMotionMatchingStates(uint64 InAnimInstanceId);
-	
 	/** Last stored Rewind Debugger */
 	const IRewindDebugger* RewindDebugger = nullptr;
-	/** List of all Node IDs associated with motion matching states */
-	TArray<int32> NodeIds;
-	/** List of all updated motion matching states per node */
-	TArray<const FTraceMotionMatchingStateMessage*> MotionMatchingStates;
-	/** Last stored MM state (updated from OnSelectionChanged) */
-	const FTraceMotionMatchingStateMessage* ActiveMotionMatchingState = nullptr;
-	/** Last updated reflection data relative to MM state */
-	TObjectPtr<UPoseSearchDebuggerReflection> Reflection = nullptr;
 
-	static FDebugger* InternalInstance;
+	/** List of all active debugger instances */
+	TArray<TSharedRef<FDebuggerInstance>> DebuggerInstances;
+	
+	static FDebugger* Debugger;
 };
 
 /**
@@ -335,6 +471,8 @@ public:
 	virtual FText GetTitle() const override;
 	virtual FSlateIcon GetIcon() const override;
 	virtual FName GetTargetTypeName() const override;
+	
+	/** Creates the PoseSearch Slate view for the provided AnimInstance */
 	virtual TSharedPtr<IRewindDebuggerView> CreateDebugView(uint64 ObjectId, double CurrentTime, const TraceServices::IAnalysisSession& InAnalysisSession) const override;
 };
 
diff --git a/Engine/Plugins/Experimental/Animation/PoseSearch/Source/Editor/Private/Trace/PoseSearchTraceAnalyzer.cpp b/Engine/Plugins/Experimental/Animation/PoseSearch/Source/Editor/Private/Trace/PoseSearchTraceAnalyzer.cpp
--- a/Engine/Plugins/Experimental/Animation/PoseSearch/Source/Editor/Private/Trace/PoseSearchTraceAnalyzer.cpp
+++ b/Engine/Plugins/Experimental/Animation/PoseSearch/Source/Editor/Private/Trace/PoseSearchTraceAnalyzer.cpp
@@ -33,6 +33,7 @@ bool FTraceAnalyzer::OnEvent(uint16 RouteId, EStyle Style, const FOnEventContext
 	const double Time = Context.EventTime.AsSeconds(EventData.GetValue<uint64>("Cycle"));
 	const uint16 FrameCounter = EventData.GetValue<uint16>("FrameCounter");
 	const uint64 AnimInstanceId = EventData.GetValue<uint64>("AnimInstanceId");
+	const uint64 SkeletalMeshComponentId = EventData.GetValue<uint64>("SkeletalMeshComponentId");
 	const int32 NodeId = EventData.GetValue<int32>("NodeId");
 
 	switch (RouteId)
@@ -60,6 +61,7 @@ bool FTraceAnalyzer::OnEvent(uint16 RouteId, EStyle Style, const FOnEventContext
 				// Common data
 				Message.NodeId = NodeId;
 				Message.AnimInstanceId = AnimInstanceId;
+				Message.SkeletalMeshComponentId = SkeletalMeshComponentId;
 				Message.FrameCounter = FrameCounter;
 
 				TraceProvider.AppendMotionMatchingState(Message, Time);
diff --git a/Engine/Plugins/Experimental/Animation/PoseSearch/Source/Editor/Private/Trace/PoseSearchTraceProvider.h b/Engine/Plugins/Experimental/Animation/PoseSearch/Source/Editor/Private/Trace/PoseSearchTraceProvider.h
--- a/Engine/Plugins/Experimental/Animation/PoseSearch/Source/Editor/Private/Trace/PoseSearchTraceProvider.h
+++ b/Engine/Plugins/Experimental/Animation/PoseSearch/Source/Editor/Private/Trace/PoseSearchTraceProvider.h
@@ -15,6 +15,7 @@ namespace UE { namespace PoseSearch {
 struct FTraceMessage
 {
 	uint64 AnimInstanceId;
+	uint64 SkeletalMeshComponentId;
 	int32 NodeId;
 	uint16 FrameCounter;
 };
diff --git a/Engine/Plugins/Experimental/Animation/PoseSearch/Source/Runtime/Private/AnimNode_MotionMatching.cpp b/Engine/Plugins/Experimental/Animation/PoseSearch/Source/Runtime/Private/AnimNode_MotionMatching.cpp
--- a/Engine/Plugins/Experimental/Animation/PoseSearch/Source/Runtime/Private/AnimNode_MotionMatching.cpp
+++ b/Engine/Plugins/Experimental/Animation/PoseSearch/Source/Runtime/Private/AnimNode_MotionMatching.cpp
@@ -212,22 +212,19 @@ void FAnimNode_MotionMatching::PreUpdate(const UAnimInstance* InAnimInstance)
 		check(SkeletalMeshComponent);
 
 		UE::PoseSearch::FDebugDrawParams DrawParams;
-		DrawParams.Flags = UE::PoseSearch::EDebugDrawFlags::None;
 		DrawParams.RootTransform = SkeletalMeshComponent->GetComponentTransform();
 		DrawParams.Database = Database;
-		DrawParams.Query = ComposedQuery.GetValues();
 		DrawParams.World = SkeletalMeshComponent->GetWorld();
 		DrawParams.DefaultLifeTime = 0.0f;
-		DrawParams.HighlightPoseIdx = DbPoseIdx;
 
 		if (bDebugDrawMatch)
 		{
-			DrawParams.Flags |= UE::PoseSearch::EDebugDrawFlags::DrawBest;
+			DrawParams.PoseIdx = DbPoseIdx;
 		}
 
 		if (bDebugDrawGoal)
 		{
-			DrawParams.Flags |= UE::PoseSearch::EDebugDrawFlags::DrawQuery;
+			DrawParams.PoseVector = ComposedQuery.GetValues();
 		}
 
 		UE::PoseSearch::Draw(DrawParams);
diff --git a/Engine/Plugins/Experimental/Animation/PoseSearch/Source/Runtime/Private/PoseSearch.cpp b/Engine/Plugins/Experimental/Animation/PoseSearch/Source/Runtime/Private/PoseSearch.cpp
--- a/Engine/Plugins/Experimental/Animation/PoseSearch/Source/Runtime/Private/PoseSearch.cpp
+++ b/Engine/Plugins/Experimental/Animation/PoseSearch/Source/Runtime/Private/PoseSearch.cpp
@@ -1339,7 +1339,7 @@ bool FFeatureVectorReader::GetVector(FPoseSearchFeatureDesc Element, FVector* Ou
 
 bool FDebugDrawParams::CanDraw() const
 {
-	if (!World || !EnumHasAnyFlags(Flags, EDebugDrawFlags::DrawAll))
+	if (!World)
 	{
 		return false;
 	}
@@ -2094,7 +2094,8 @@ static void DrawTrajectoryFeatures(const FDebugDrawParams& DrawParams, const FFe
 		if (Reader.GetPosition(Feature, &TrajectoryPos))
 		{	
 			Feature.Type = EPoseSearchFeatureType::Position;
-			FLinearColor LinearColor = GetColorForFeature(Feature, Reader.GetLayout());
+
+			FLinearColor LinearColor = DrawParams.Color ? *DrawParams.Color : GetColorForFeature(Feature, Reader.GetLayout());
 			FColor Color = LinearColor.ToFColor(true);
 
 			TrajectoryPos = DrawParams.RootTransform.TransformPosition(TrajectoryPos);
@@ -2109,7 +2110,8 @@ static void DrawTrajectoryFeatures(const FDebugDrawParams& DrawParams, const FFe
 		if (Reader.GetLinearVelocity(Feature, &TrajectoryVel))
 		{
 			Feature.Type = EPoseSearchFeatureType::LinearVelocity;
-			FLinearColor LinearColor = GetColorForFeature(Feature, Reader.GetLayout());
+			
+			FLinearColor LinearColor = DrawParams.Color ? *DrawParams.Color : GetColorForFeature(Feature, Reader.GetLayout());
 			FColor Color = LinearColor.ToFColor(true);
 
 			TrajectoryVel *= DrawDebugVelocityScale;
@@ -2148,26 +2150,30 @@ static void DrawPoseFeatures(const FDebugDrawParams& DrawParams, const FFeatureV
 			Feature.SchemaBoneIdx = SchemaBoneIdx;
 
 			FVector BonePos;
-			bool bHaveBonePos = Reader.GetPosition(Feature, &BonePos);
+			const bool bHaveBonePos = Reader.GetPosition(Feature, &BonePos);
 			if (bHaveBonePos)
 			{
 				Feature.Type = EPoseSearchFeatureType::Position;
-				FLinearColor Color = GetColorForFeature(Feature, Reader.GetLayout());
+				
+				FLinearColor LinearColor = DrawParams.Color ? *DrawParams.Color : GetColorForFeature(Feature, Reader.GetLayout());
+				FColor Color = LinearColor.ToFColor(true);
 
 				BonePos = DrawParams.RootTransform.TransformPosition(BonePos);
-				DrawDebugSphere(DrawParams.World, BonePos, DrawDebugSphereSize, DrawDebugSphereSegments, Color.ToFColor(true), false, LifeTime, DepthPriority, DrawDebugSphereLineThickness);
+				DrawDebugSphere(DrawParams.World, BonePos, DrawDebugSphereSize, DrawDebugSphereSegments, Color, false, LifeTime, DepthPriority, DrawDebugSphereLineThickness);
 			}
 
 			FVector BoneVel;
 			if (bHaveBonePos && Reader.GetLinearVelocity(Feature, &BoneVel))
 			{
 				Feature.Type = EPoseSearchFeatureType::LinearVelocity;
-				FLinearColor Color = GetColorForFeature(Feature, Reader.GetLayout());
+				
+				FLinearColor LinearColor = DrawParams.Color ? *DrawParams.Color : GetColorForFeature(Feature, Reader.GetLayout());
+				FColor Color = LinearColor.ToFColor(true);
 
 				BoneVel *= DrawDebugVelocityScale;
 				BoneVel = DrawParams.RootTransform.TransformVector(BoneVel);
 				FVector BoneVelDirection = BoneVel.GetSafeNormal();
-				DrawDebugDirectionalArrow(DrawParams.World, BonePos + BoneVelDirection * DrawDebugSphereSize, BonePos + BoneVel, DrawDebugArrowSize, Color.ToFColor(true), false, LifeTime, DepthPriority, DrawDebugLineThickness);
+				DrawDebugDirectionalArrow(DrawParams.World, BonePos + BoneVelDirection * DrawDebugSphereSize, BonePos + BoneVel, DrawDebugArrowSize, Color, false, LifeTime, DepthPriority, DrawDebugLineThickness);
 			}
 		}
 	}
@@ -2175,18 +2181,36 @@ static void DrawPoseFeatures(const FDebugDrawParams& DrawParams, const FFeatureV
 
 static void DrawFeatureVector(const FDebugDrawParams& DrawParams, const FFeatureVectorReader& Reader)
 {
-	DrawPoseFeatures(DrawParams, Reader);
-	DrawTrajectoryFeatures(DrawParams, Reader, EPoseSearchFeatureDomain::Time);
-	DrawTrajectoryFeatures(DrawParams, Reader, EPoseSearchFeatureDomain::Distance);
+	if (EnumHasAnyFlags(DrawParams.Flags, EDebugDrawFlags::IncludePose))
+	{
+		DrawPoseFeatures(DrawParams, Reader);
+	}
+
+	if (EnumHasAnyFlags(DrawParams.Flags, EDebugDrawFlags::IncludeTrajectory))
+	{
+		DrawTrajectoryFeatures(DrawParams, Reader, EPoseSearchFeatureDomain::Time);
+        DrawTrajectoryFeatures(DrawParams, Reader, EPoseSearchFeatureDomain::Distance);
+	}
 }
 
-static void DrawSearchIndex(const FDebugDrawParams& DrawParams)
+static void DrawFeatureVector(const FDebugDrawParams& DrawParams, TArrayView<const float> PoseVector)
 {
-	if (!DrawParams.CanDraw())
+	const UPoseSearchSchema* Schema = DrawParams.GetSchema();
+	check(Schema)
+
+	if (PoseVector.Num() != Schema->Layout.NumFloats)
 	{
 		return;
 	}
 
+	FFeatureVectorReader Reader;
+	Reader.Init(&Schema->Layout);
+	Reader.SetValues(PoseVector);
+	DrawFeatureVector(DrawParams, Reader);
+}
+
+static void DrawSearchIndex(const FDebugDrawParams& DrawParams)
+{
 	const UPoseSearchSchema* Schema = DrawParams.GetSchema();
 	const FPoseSearchIndex* SearchIndex = DrawParams.GetSearchIndex();
 	check(Schema);
@@ -2195,61 +2219,37 @@ static void DrawSearchIndex(const FDebugDrawParams& DrawParams)
 	FFeatureVectorReader Reader;
 	Reader.Init(&Schema->Layout);
 
-	int32 LastPoseIdx = SearchIndex->NumPoses;
-	int32 StartPoseIdx = 0;
-	if (!(DrawParams.Flags & EDebugDrawFlags::DrawSearchIndex))
-	{
-		StartPoseIdx = DrawParams.HighlightPoseIdx;
-		LastPoseIdx = StartPoseIdx + 1;
-	}
-
-	if (StartPoseIdx < 0)
-	{
-		return;
-	}
+	const int32 LastPoseIdx = SearchIndex->NumPoses;
 
 	TArray<float> PoseVector;
-	for (int32 PoseIdx = StartPoseIdx; PoseIdx != LastPoseIdx; ++PoseIdx)
+	for (int32 PoseIdx = 0; PoseIdx != LastPoseIdx; ++PoseIdx)
 	{
 		PoseVector = SearchIndex->GetPoseValues(PoseIdx);
 		SearchIndex->InverseNormalize(PoseVector);
 		Reader.SetValues(PoseVector);
-
 		DrawFeatureVector(DrawParams, Reader);
 	}
 }
 
-static void DrawQuery(const FDebugDrawParams& DrawParams)
-{
-	if (!DrawParams.CanDraw())
-	{
-		return;
-	}
-
-	const UPoseSearchSchema* Schema = DrawParams.GetSchema();
-	check(Schema);
-
-	if (DrawParams.Query.Num() != Schema->Layout.NumFloats)
-	{
-		return;
-	}
-
-	FFeatureVectorReader Reader;
-	Reader.Init(&Schema->Layout);
-	Reader.SetValues(DrawParams.Query);
-	DrawFeatureVector(DrawParams, Reader);
-}
-
 void Draw(const FDebugDrawParams& DebugDrawParams)
 {
 	if (DebugDrawParams.CanDraw())
 	{
-		if (EnumHasAnyFlags(DebugDrawParams.Flags, EDebugDrawFlags::DrawQuery))
+		if (DebugDrawParams.PoseIdx != INDEX_NONE)
 		{
-			DrawQuery(DebugDrawParams);
-		}
+			const FPoseSearchIndex* SearchIndex = DebugDrawParams.GetSearchIndex();
+			check(SearchIndex);
 
-		if (EnumHasAnyFlags(DebugDrawParams.Flags, EDebugDrawFlags::DrawSearchIndex | EDebugDrawFlags::DrawBest))
+			TArray<float> PoseVector;
+			PoseVector = SearchIndex->GetPoseValues(DebugDrawParams.PoseIdx);
+			SearchIndex->InverseNormalize(PoseVector);
+			DrawFeatureVector(DebugDrawParams, PoseVector);
+		}
+		if (!DebugDrawParams.PoseVector.IsEmpty())
+		{
+			DrawFeatureVector(DebugDrawParams, DebugDrawParams.PoseVector);
+		}
+		if (EnumHasAnyFlags(DebugDrawParams.Flags, EDebugDrawFlags::DrawSearchIndex))
 		{
 			DrawSearchIndex(DebugDrawParams);
 		}
@@ -2967,8 +2967,8 @@ FSearchResult Search(const UAnimSequenceBase* Sequence, TArrayView<const float>
 
 	// Do debug visualization
 	DebugDrawParams.SequenceMetaData = MetaData;
-	DebugDrawParams.Query = Query;
-	DebugDrawParams.HighlightPoseIdx = Result.PoseIdx;
+	DebugDrawParams.PoseVector = Query;
+	DebugDrawParams.PoseIdx = Result.PoseIdx;
 	Draw(DebugDrawParams);
 
 	return Result;
@@ -3003,8 +3003,8 @@ FDbSearchResult Search(const UPoseSearchDatabase* Database, TArrayView<const flo
 
 	// Do debug visualization
 	DebugDrawParams.Database = Database;
-	DebugDrawParams.Query = Query;
-	DebugDrawParams.HighlightPoseIdx = Result.PoseIdx;
+	DebugDrawParams.PoseVector = Query;
+	DebugDrawParams.PoseIdx = Result.PoseIdx;
 	Draw(DebugDrawParams);
 
 	return Result;
diff --git a/Engine/Plugins/Experimental/Animation/PoseSearch/Source/Runtime/Private/Trace/PoseSearchTraceLogger.cpp b/Engine/Plugins/Experimental/Animation/PoseSearch/Source/Runtime/Private/Trace/PoseSearchTraceLogger.cpp
--- a/Engine/Plugins/Experimental/Animation/PoseSearch/Source/Runtime/Private/Trace/PoseSearchTraceLogger.cpp
+++ b/Engine/Plugins/Experimental/Animation/PoseSearch/Source/Runtime/Private/Trace/PoseSearchTraceLogger.cpp
@@ -12,6 +12,7 @@ UE_TRACE_EVENT_BEGIN(PoseSearch, MotionMatchingState)
 	UE_TRACE_EVENT_FIELD(uint64, Cycle)
 	UE_TRACE_EVENT_FIELD(uint16, FrameCounter)
 	UE_TRACE_EVENT_FIELD(uint64, AnimInstanceId)
+	UE_TRACE_EVENT_FIELD(uint64, SkeletalMeshComponentId)
 	UE_TRACE_EVENT_FIELD(int32, NodeId)
 	UE_TRACE_EVENT_FIELD(float, ElapsedPoseJumpTime)
 	UE_TRACE_EVENT_FIELD(uint32, Flags)
@@ -53,11 +54,13 @@ void FTraceMotionMatchingState::Output(const FAnimationBaseContext& InContext, c
 
 	UObject* AnimInstance = InContext.AnimInstanceProxy->GetAnimInstanceObject();
 	TRACE_OBJECT(AnimInstance);
+	UObject* SkeletalMeshComponent = AnimInstance->GetOuter();
 
 	UE_TRACE_LOG(PoseSearch, MotionMatchingState, PoseSearchChannel)
 		<< MotionMatchingState.Cycle(FPlatformTime::Cycles64())
 		<< MotionMatchingState.FrameCounter(FObjectTrace::GetObjectWorldTickCounter(AnimInstance))
 		<< MotionMatchingState.AnimInstanceId(FObjectTrace::GetObjectId(AnimInstance))
+		<< MotionMatchingState.SkeletalMeshComponentId(FObjectTrace::GetObjectId(SkeletalMeshComponent))
 		<< MotionMatchingState.NodeId(InContext.GetCurrentNodeId())
 		<< MotionMatchingState.ElapsedPoseJumpTime(State.ElapsedPoseJumpTime)
 		// Cast to uint32 for event message
diff --git a/Engine/Plugins/Experimental/Animation/PoseSearch/Source/Runtime/Private/Trace/PoseSearchTraceLogger.h b/Engine/Plugins/Experimental/Animation/PoseSearch/Source/Runtime/Private/Trace/PoseSearchTraceLogger.h
--- a/Engine/Plugins/Experimental/Animation/PoseSearch/Source/Runtime/Private/Trace/PoseSearchTraceLogger.h
+++ b/Engine/Plugins/Experimental/Animation/PoseSearch/Source/Runtime/Private/Trace/PoseSearchTraceLogger.h
@@ -55,6 +55,11 @@ struct POSESEARCH_API FTraceMotionMatchingState
 	/** Node Id of the motion matching node associated with this message */
 	int32 NodeId = 0;
 
+	/** Skeletal Mesh Component Id, outer class of the AnimInstance.
+	 *  Used for retrieval of traced root transforms from the animation provider.
+	 */
+	uint64 SkeletalMeshComponentId = 0;
+
 	/** Output the current state info to the logger */
 	static void Output(const FAnimationBaseContext& InContext, const FTraceMotionMatchingState& State);
 	
diff --git a/Engine/Plugins/Experimental/Animation/PoseSearch/Source/Runtime/Public/PoseSearch/PoseSearch.h b/Engine/Plugins/Experimental/Animation/PoseSearch/Source/Runtime/Public/PoseSearch/PoseSearch.h
--- a/Engine/Plugins/Experimental/Animation/PoseSearch/Source/Runtime/Public/PoseSearch/PoseSearch.h
+++ b/Engine/Plugins/Experimental/Animation/PoseSearch/Source/Runtime/Public/PoseSearch/PoseSearch.h
@@ -603,26 +603,41 @@ private:
 
 enum class EDebugDrawFlags : uint32
 {
-	None			= 0,
-	DrawQuery		= 1 << 0,
-	DrawSearchIndex = 1 << 1,
-	DrawBest        = 1 << 2,
-	DrawAll		    = MAX_uint32
+	None			    = 0,
+
+	/** Draw the entire search index */
+	DrawSearchIndex     = 1 << 0,
+
+	/** Draw pose features for each pose vector */
+	IncludePose         = 1 << 1,
+
+	/** Draw trajectory features for each pose vector */
+	IncludeTrajectory   = 1 << 2,
+
+	/** Draw all pose vector features */
+	IncludeAllFeatures  = IncludePose | IncludeTrajectory,
 };
 ENUM_CLASS_FLAGS(EDebugDrawFlags);
 
-struct FDebugDrawParams
+struct POSESEARCH_API FDebugDrawParams
 {
 	const UWorld* World = nullptr;
 	const UPoseSearchDatabase* Database = nullptr;
 	const UPoseSearchSequenceMetaData* SequenceMetaData = nullptr;
-	EDebugDrawFlags Flags = EDebugDrawFlags::None;
+	EDebugDrawFlags Flags = EDebugDrawFlags::IncludeAllFeatures;
 	float DefaultLifeTime = 5.0f;
 	FTransform RootTransform = FTransform::Identity;
-	int32 HighlightPoseIdx = -1;
-	TArrayView<const float> Query;
 
-	bool CanDraw () const;
+	/** If set, draw the corresponding pose from the search index */
+	int32 PoseIdx = INDEX_NONE;
+
+	/** If set, draw using this uniform color instead of feature-based coloring */
+	const FLinearColor* Color = nullptr;
+
+	/** If set, interpret the buffer as a pose vector and draw it */
+	TArrayView<const float> PoseVector;
+
+	bool CanDraw() const;
 	const FPoseSearchIndex* GetSearchIndex() const;
 	const UPoseSearchSchema* GetSchema() const;
 };
