Trajectory sample extrapolation
Notes:
- It's always enabled. If the main clip or the lead/follow crips loop, it doesn't come into play. If they don't, it will extrapolate even the lead/follow as necessary
- Extrapolation parameters (min speed and sampling interval) were added to PoseSearchDatabase
- This does not handle playback issues related to no longer having lead/follow clips. That will come after.

diff --git a/Engine/Plugins/Experimental/Animation/PoseSearch/Source/Runtime/Private/PoseSearch.cpp b/Engine/Plugins/Experimental/Animation/PoseSearch/Source/Runtime/Private/PoseSearch.cpp
--- a/Engine/Plugins/Experimental/Animation/PoseSearch/Source/Runtime/Private/PoseSearch.cpp
+++ b/Engine/Plugins/Experimental/Animation/PoseSearch/Source/Runtime/Private/PoseSearch.cpp
@@ -1396,6 +1396,7 @@ public:
 		const UAnimSequence* Sequence = nullptr;
 		bool bLoopable = false;
 		int32 DistanceSamplingRate = 60;
+		FPoseSearchExtrapolationParameters ExtrapolationParameters;
 	} Input;
 
 	struct FOutput
@@ -1412,12 +1413,25 @@ public:
 	void Init(const FInput& Input);
 	void Process();
 
+
+	// Extracts root transform at the given time, using the extremities of the sequence to extrapolate beyond the 
+	// sequence limits when Time is less than zero or greater than the sequence length.
 	FTransform ExtractRootTransform(float Time) const;
+
+	// Extracts the accumulated root distance at the given time, using the extremities of the sequence to extrapolate 
+	// beyond the sequence limits when Time is less than zero or greater than the sequence length
 	float ExtractRootDistance(float Time) const;
 
 private:
 	void Reserve();
 	void ProcessRootMotion();
+
+	// Samples sequence and adjusts obtained root motion to ExtrapolationTime
+	FTransform ExtrapolateRootMotion(float SampleStart, float SampleEnd, float ExtrapolationTime) const;
+
+	// Uses distance delta between NextRootDistanceIndex and NextRootDistanceIndex - 1 and extrapolates it to 
+	// ExtrapolationTime
+	float ExtrapolateRootDistance(int32 NextRootDistanceIndex, float ExtrapolationTime) const;
 };
 
 void FSequenceSampler::Init(const FInput& InInput)
@@ -1458,16 +1472,54 @@ void FSequenceSampler::Process()
 
 FTransform FSequenceSampler::ExtractRootTransform(float Time) const
 {
-	return Input.Sequence->ExtractRootMotionFromRange(0.0f, Time);
+	if (Input.bLoopable)
+	{
+		FTransform LoopableRootTransform = Input.Sequence->ExtractRootMotion(0.0f, Time, true);
+		return LoopableRootTransform;
+	}
+
+	const float ExtrapolationSampleTime = Input.ExtrapolationParameters.SampleTime;
+
+	const float PlayLength = Input.Sequence->GetPlayLength();
+	const float ClampedTime = FMath::Clamp(Time, 0.0f, PlayLength);
+	const float ExtrapolationTime = Time - ClampedTime;
+
+	FTransform RootTransform = FTransform::Identity;
+
+	// If Time is less than zero, ExtrapolationTime will be negative. In this case, we extrapolate the beginning of the 
+	// animation to estimate where the root would be at Time
+	if (ExtrapolationTime < -SMALL_NUMBER)
+	{
+		const FTransform ExtrapolatedRootMotion = ExtrapolateRootMotion(
+			0.0f, ExtrapolationSampleTime, 
+			ExtrapolationTime);
+		RootTransform = ExtrapolatedRootMotion;
+	}
+	else
+	{
+		RootTransform = Input.Sequence->ExtractRootMotionFromRange(0.0f, ClampedTime);
+
+		// If Time is greater than PlayLength, ExtrapolationTIme will be a positive number. In this case, we extrapolate
+		// the end of the animation to estimate where the root would be at Time
+		if (ExtrapolationTime > SMALL_NUMBER)
+		{
+			const FTransform ExtrapolatedRootMotion = ExtrapolateRootMotion(
+				PlayLength - ExtrapolationSampleTime, PlayLength,
+				ExtrapolationTime);
+			RootTransform = ExtrapolatedRootMotion * RootTransform;
+		}
+	}
+
+	return RootTransform;
 }
 
 float FSequenceSampler::ExtractRootDistance(float Time) const
 {
-	check(Time <= Output.PlayLength);
+	const float ClampedTime = FMath::Clamp(Time, 0.0f, Output.PlayLength);
 
 	// Find the distance sample that corresponds with the time and split into whole and partial parts
 	float IntegralDistanceSample;
-	float DistanceAlpha = FMath::Modf(Time * Input.DistanceSamplingRate, &IntegralDistanceSample);
+	float DistanceAlpha = FMath::Modf(ClampedTime * Input.DistanceSamplingRate, &IntegralDistanceSample);
 	float DistanceIdx = (int32)IntegralDistanceSample;
 
 	// Verify the distance offset and any residual portion would be in bounds
@@ -1481,9 +1533,98 @@ float FSequenceSampler::ExtractRootDistance(float Time) const
 		Distance = FMath::Lerp(Distance, NextDistance, DistanceAlpha);
 	}
 
+	const float ExtrapolationTime = Time - ClampedTime;
+
+	if (ExtrapolationTime != 0.0f)
+	{
+		// If extrapolationTime is not zero, we extrapolate the beginning or the end of the animation to estimate
+		// the root distance.
+		const int32 DistIdx = (ExtrapolationTime > 0.0f) ? Output.AccumulatedRootDistance.Num() - 1 : 1;
+		const float ExtrapolatedDistance = ExtrapolateRootDistance(DistIdx, ExtrapolationTime);
+		Distance += ExtrapolatedDistance;
+	}
+
 	return Distance;
 }
 
+FTransform FSequenceSampler::ExtrapolateRootMotion(float SampleStart, float SampleEnd, float ExtrapolationTime) const
+{
+	const float SampleDelta = SampleEnd - SampleStart;
+	check(!FMath::IsNearlyZero(SampleDelta));
+
+	FTransform SampleToExtrapolate = Input.Sequence->ExtractRootMotionFromRange(SampleStart, SampleEnd);
+
+	const FVector LinearVelocityToExtrapolate = SampleToExtrapolate.GetTranslation() / SampleDelta;
+	const float LinearSpeedToExtrapolate = LinearVelocityToExtrapolate.Size();
+	const bool bCanExtrapolateTranslation = 
+		LinearSpeedToExtrapolate >= Input.ExtrapolationParameters.LinearSpeedThreshold;
+
+	const float AngularSpeedToExtrapolateRad = SampleToExtrapolate.GetRotation().GetAngle() / SampleDelta;
+	const bool bCanExtrapolateRotation = 
+		FMath::RadiansToDegrees(AngularSpeedToExtrapolateRad) >= Input.ExtrapolationParameters.AngularSpeedThreshold;
+
+	if (!bCanExtrapolateTranslation && !bCanExtrapolateRotation)
+	{
+		return FTransform::Identity;
+	}
+
+	if (!bCanExtrapolateTranslation)
+	{
+		SampleToExtrapolate.SetTranslation(FVector::ZeroVector);
+	}
+
+	if (!bCanExtrapolateRotation)
+	{
+		SampleToExtrapolate.SetRotation(FQuat::Identity);
+	}
+
+	// converting ExtrapolationTime to a positive number to avoid dealing with the negative extrapolation and inverting
+	// transforms later on.
+	const float AbsExtrapolationTime = FMath::Abs(ExtrapolationTime);
+	const float AbsSampleDelta = FMath::Abs(SampleDelta);
+	const FTransform AbsTimeSampleToExtrapolate = 
+		ExtrapolationTime >= 0.0f ? SampleToExtrapolate : SampleToExtrapolate.Inverse();
+
+	// because we're extrapolating rotation, the extrapolation must be integrated over time
+	const float SampleMultiplier = AbsExtrapolationTime / AbsSampleDelta;
+	float IntegralNumSamples;
+	float RemainingSampleFraction = FMath::Modf(SampleMultiplier, &IntegralNumSamples);
+	int32 NumSamples = (int32)IntegralNumSamples;
+
+	// adding full samples to the extrapolated root motion
+	FTransform ExtrapolatedRootMotion = FTransform::Identity;
+	for (int i = 0; i < NumSamples; ++i)
+	{
+		ExtrapolatedRootMotion = AbsTimeSampleToExtrapolate * ExtrapolatedRootMotion;
+	}
+	
+	// and a blend with identify for whatever is left
+	FTransform RemainingExtrapolatedRootMotion;
+	RemainingExtrapolatedRootMotion.Blend(
+		FTransform::Identity, 
+		AbsTimeSampleToExtrapolate, 
+		RemainingSampleFraction);
+
+	ExtrapolatedRootMotion = RemainingExtrapolatedRootMotion * ExtrapolatedRootMotion;
+	return ExtrapolatedRootMotion;
+}
+
+float FSequenceSampler::ExtrapolateRootDistance(int32 NextRootDistanceIndex, float ExtrapolationTime) const
+{
+	check(NextRootDistanceIndex > 0 && NextRootDistanceIndex < Output.AccumulatedRootDistance.Num());
+
+	const float DistanceDelta = 
+		Output.AccumulatedRootDistance[NextRootDistanceIndex] - 
+		Output.AccumulatedRootDistance[NextRootDistanceIndex - 1];
+	const float Speed = DistanceDelta * Input.DistanceSamplingRate;
+	const float ExtrapolationSpeed = Speed >= Input.ExtrapolationParameters.LinearSpeedThreshold ?
+		Speed : 0.0f;
+	const float ExtrapolatedDistance = ExtrapolationSpeed * ExtrapolationTime;
+
+	return ExtrapolatedDistance;
+}
+
+
 void FSequenceSampler::ProcessRootMotion()
 {
 	// Note the distance sampling interval is independent of the schema's sampling interval
@@ -1524,6 +1665,9 @@ struct FSamplingParam
 	float WrappedParam = 0.0f;
 	int32 NumCycles = 0;
 	bool bClamped = false;
+	
+	// If the animation can't loop, WrappedParam contains the clamped value and whatever is left is stored here
+	float Extrapolation = 0.0f;
 };
 
 static FSamplingParam WrapOrClampSamplingParam(bool bCanWrap, float SamplingParamExtent, float SamplingParam)
@@ -1543,9 +1687,6 @@ static FSamplingParam WrapOrClampSamplingParam(bool bCanWrap, float SamplingPara
 	{
 		if (SamplingParam < 0.0f)
 		{
-			// Note a full reverse cycle isn't completed merely when SamplingParam < 0.0f,
-			// rather when SamplingParam < -SamplingParamExtent
-			Result.WrappedParam += SamplingParamExtent;
 			while (Result.WrappedParam < 0.0f)
 			{
 				Result.WrappedParam += SamplingParamExtent;
@@ -1555,7 +1696,6 @@ static FSamplingParam WrapOrClampSamplingParam(bool bCanWrap, float SamplingPara
 
 		else
 		{
-			// On the other hand a full forward cycle is completed each time we exceed SamplingParamExtent
 			while (Result.WrappedParam > SamplingParamExtent)
 			{
 				Result.WrappedParam -= SamplingParamExtent;
@@ -1568,6 +1708,7 @@ static FSamplingParam WrapOrClampSamplingParam(bool bCanWrap, float SamplingPara
 	if (ParamClamped != Result.WrappedParam)
 	{
 		check(!bCanWrap);
+		Result.Extrapolation = Result.WrappedParam - ParamClamped;
 		Result.WrappedParam = ParamClamped;
 		Result.bClamped = true;
 	}
@@ -1705,21 +1846,28 @@ const float FSequenceIndexer::GetSampleTimeFromDistance(float SampleDistance) co
 
 	auto ClipTimeFromDistance = [](const FSequenceSampler* Sampler, float ClipDistance) -> float
 	{
-		// Search for the distance value
-		int32 NextSampleIdx = Algo::LowerBound(Sampler->Output.AccumulatedRootDistance, ClipDistance);
-		check(NextSampleIdx < Sampler->Output.AccumulatedRootDistance.Num());
-	
-		// Compute distance interpolation amount
-		int32 PrevSampleIdx = FMath::Max(0, NextSampleIdx - 1);
+		int32 NextSampleIdx = 1;
+		int32 PrevSampleIdx = 0;
+		if (ClipDistance > 0.0f)
+		{
+			// Search for the distance value. Because the values will be extrapolated if necessary
+			// LowerBound might go past the end of the array, in which case the last valid index is used
+			int32 ClipDistanceLowerBoundIndex = Algo::LowerBound(Sampler->Output.AccumulatedRootDistance, ClipDistance);
+			NextSampleIdx = FMath::Min(
+				ClipDistanceLowerBoundIndex,
+				Sampler->Output.AccumulatedRootDistance.Num() - 1);
+
+			// Compute distance interpolation amount
+			PrevSampleIdx = FMath::Max(0, NextSampleIdx - 1);
+		}
+
 		float NextDistance = Sampler->Output.AccumulatedRootDistance[NextSampleIdx];
-		float PrevDistance = PrevSampleIdx < NextSampleIdx ? Sampler->Output.AccumulatedRootDistance[PrevSampleIdx] : 0.0f;
+		float PrevDistance = Sampler->Output.AccumulatedRootDistance[PrevSampleIdx];
 		float DistanceSampleAlpha = FMath::GetRangePct(PrevDistance, NextDistance, ClipDistance);
 
 		// Convert to time
 		float ClipTime = (float(NextSampleIdx) - (1.0f - DistanceSampleAlpha)) / Sampler->Input.DistanceSamplingRate;
-		float ClipTimeClamped = FMath::Clamp(ClipTime, 0.0f, Sampler->Output.PlayLength);
-
-		return ClipTimeClamped;
+		return ClipTime;
 	};
 
 	float MainTotalDistance = Input.MainSequence->Output.TotalRootDistance;
@@ -1735,9 +1883,12 @@ const float FSequenceIndexer::GetSampleTimeFromDistance(float SampleDistance) co
 			const FSequenceSampler::FOutput& ClipData = Input.LeadInSequence->Output;
 
 			bool bLeadInCanWrap = CanWrapDistanceSamples(Input.LeadInSequence);
-			FSamplingParam SamplingParam = WrapOrClampSamplingParam(bLeadInCanWrap, ClipData.TotalRootDistance, SampleDistance);
+			float LeadRelativeDistance = SampleDistance + ClipData.TotalRootDistance;
+			FSamplingParam SamplingParam = WrapOrClampSamplingParam(bLeadInCanWrap, ClipData.TotalRootDistance, LeadRelativeDistance);
 
-			float ClipTime = ClipTimeFromDistance(Input.LeadInSequence, SamplingParam.WrappedParam);
+			float ClipTime = ClipTimeFromDistance(
+				Input.LeadInSequence, 
+				SamplingParam.WrappedParam + SamplingParam.Extrapolation);
 
 			// Make the lead in clip time relative to the main sequence again and unwrap
 			SampleTime = -((SamplingParam.NumCycles * ClipData.PlayLength) + (ClipData.PlayLength - ClipTime));
@@ -1752,7 +1903,9 @@ const float FSequenceIndexer::GetSampleTimeFromDistance(float SampleDistance) co
 			float FollowRelativeDistance = SampleDistance - MainTotalDistance;
 			FSamplingParam SamplingParam = WrapOrClampSamplingParam(bFollowUpCanWrap, ClipData.TotalRootDistance, FollowRelativeDistance);
 
-			float ClipTime = ClipTimeFromDistance(Input.FollowUpSequence, SamplingParam.WrappedParam);
+			float ClipTime = ClipTimeFromDistance(
+				Input.FollowUpSequence, 
+				SamplingParam.WrappedParam + SamplingParam.Extrapolation);
 
 			// Make the follow up clip time relative to the main sequence again and unwrap
 			SampleTime = Input.MainSequence->Output.PlayLength + SamplingParam.NumCycles * ClipData.PlayLength + ClipTime;
@@ -1765,17 +1918,34 @@ const float FSequenceIndexer::GetSampleTimeFromDistance(float SampleDistance) co
 	{
 		const FSequenceSampler::FOutput& ClipData = Input.MainSequence->Output;
 
-		FSamplingParam SamplingParam = WrapOrClampSamplingParam(bMainCanWrap, MainTotalDistance, SampleDistance);
-		float ClipTime = ClipTimeFromDistance(Input.MainSequence, SamplingParam.WrappedParam);
+		float MainRelativeDistance = SampleDistance;
+		if (SampleDistance < 0.0f && bMainCanWrap)
+		{
+			// In this case we're sampling a loop backwards, so MainRelativeDistance must adjust so the number of cycles 
+			// is counted correctly.
+			MainRelativeDistance += ClipData.TotalRootDistance;
+		}
+
+		FSamplingParam SamplingParam = WrapOrClampSamplingParam(bMainCanWrap, MainTotalDistance, MainRelativeDistance);
+		float ClipTime = ClipTimeFromDistance(
+			Input.MainSequence, 
+			SamplingParam.WrappedParam + SamplingParam.Extrapolation);
 
 		// Unwrap the main clip time
-		if (SampleDistance < 0.0f)
+		if (bMainCanWrap)
 		{
-			SampleTime = -((SamplingParam.NumCycles * ClipData.PlayLength) + (ClipData.PlayLength - ClipTime));
+			if (SampleDistance < 0.0f)
+			{
+				SampleTime = -((SamplingParam.NumCycles * ClipData.PlayLength) + (ClipData.PlayLength - ClipTime));
+			}
+			else
+			{
+				SampleTime = SamplingParam.NumCycles * ClipData.PlayLength + ClipTime;
+			}
 		}
 		else
 		{
-			SampleTime = SamplingParam.NumCycles * ClipData.PlayLength + ClipTime;
+			SampleTime = ClipTime;
 		}
 	}
 
@@ -1809,12 +1979,22 @@ FSequenceIndexer::FSampleInfo FSequenceIndexer::GetSampleInfo(float SampleTime)
 			const FSequenceSampler::FOutput& ClipData = Input.LeadInSequence->Output;
 
 			bool bLeadInCanWrap = CanWrapTimeSamples(Input.LeadInSequence);
-			SamplingParam = WrapOrClampSamplingParam(bLeadInCanWrap, ClipData.PlayLength, SampleTime);
+			float LeadRelativeTime = SampleTime + ClipData.PlayLength;
+			SamplingParam = WrapOrClampSamplingParam(bLeadInCanWrap, ClipData.PlayLength, LeadRelativeTime);
 
 			Sample.Clip = Input.LeadInSequence;
 
-			RootMotionInitial = FTransform::Identity;
-			RootDistanceInitial = 0.0f;
+			check(SamplingParam.Extrapolation <= 0.0f);
+			if (SamplingParam.Extrapolation < 0.0f)
+			{
+				RootMotionInitial = Input.LeadInSequence->Output.TotalRootMotion.Inverse();
+				RootDistanceInitial = -Input.LeadInSequence->Output.TotalRootDistance;
+			}
+			else
+			{
+				RootMotionInitial = FTransform::Identity;
+				RootDistanceInitial = 0.0f;
+			}
 
 			RootMotionLast = Input.LeadInSequence->Output.TotalRootMotion;
 			RootDistanceLast = Input.LeadInSequence->Output.TotalRootDistance;
@@ -1843,7 +2023,15 @@ FSequenceIndexer::FSampleInfo FSequenceIndexer::GetSampleInfo(float SampleTime)
 	// The main anim sample may have been wrapped or clamped
 	if (!Sample.IsValid())
 	{
-		SamplingParam = WrapOrClampSamplingParam(bMainCanWrap, MainPlayLength, SampleTime);
+		float MainRelativeTime = SampleTime;
+		if (SampleTime < 0.0f && bMainCanWrap)
+		{
+			// In this case we're sampling a loop backwards, so MainRelativeTime must adjust so the number of cycles is
+			// counted correctly.
+			MainRelativeTime += MainPlayLength;
+		}
+
+		SamplingParam = WrapOrClampSamplingParam(bMainCanWrap, MainPlayLength, MainRelativeTime);
 
 		Sample.Clip = Input.MainSequence;
 
@@ -1854,43 +2042,56 @@ FSequenceIndexer::FSampleInfo FSequenceIndexer::GetSampleInfo(float SampleTime)
 		RootDistanceLast = Input.MainSequence->Output.TotalRootDistance;
 	}
 
-	Sample.ClipTime = SamplingParam.WrappedParam;
 
-	// Determine how to accumulate motion for every cycle of the anim. If the sample
-	// had to be clamped, this motion will end up not getting applied below.
-	// Also invert the accumulation direction if the requested sample was wrapped backwards.
-	FTransform RootMotionPerCycle = RootMotionLast;
-	float RootDistancePerCycle = RootDistanceLast;
-	if (SampleTime < 0.0f)
+	if (FMath::Abs(SamplingParam.Extrapolation) > SMALL_NUMBER)
 	{
-		RootMotionPerCycle = RootMotionPerCycle.Inverse();
-		RootDistancePerCycle *= -1;
+		Sample.ClipTime = SamplingParam.WrappedParam + SamplingParam.Extrapolation;
+		const FTransform ClipRootMotion = Sample.Clip->ExtractRootTransform(Sample.ClipTime);
+		const float ClipDistance = Sample.Clip->ExtractRootDistance(Sample.ClipTime);
+
+		Sample.RootTransform = ClipRootMotion * RootMotionInitial;
+		Sample.RootDistance = RootDistanceInitial + ClipDistance;
 	}
+	else
+	{
+		Sample.ClipTime = SamplingParam.WrappedParam;
 
-	// Find the remaining motion deltas after wrapping
-	FTransform RootMotionRemainder = Sample.Clip->ExtractRootTransform(Sample.ClipTime);
-	float RootDistanceRemainder = Sample.Clip->ExtractRootDistance(Sample.ClipTime);
+		// Determine how to accumulate motion for every cycle of the anim. If the sample
+		// had to be clamped, this motion will end up not getting applied below.
+		// Also invert the accumulation direction if the requested sample was wrapped backwards.
+		FTransform RootMotionPerCycle = RootMotionLast;
+		float RootDistancePerCycle = RootDistanceLast;
+		if (SampleTime < 0.0f)
+		{
+			RootMotionPerCycle = RootMotionPerCycle.Inverse();
+			RootDistancePerCycle *= -1;
+		}
 
-	// Invert motion deltas if we wrapped backwards
-	if (SampleTime < 0.0f)
-	{
-		RootMotionRemainder.SetToRelativeTransform(RootMotionLast);
-		RootDistanceRemainder = -(RootDistanceLast - RootDistanceRemainder);
-	}
+		// Find the remaining motion deltas after wrapping
+		FTransform RootMotionRemainder = Sample.Clip->ExtractRootTransform(Sample.ClipTime);
+		float RootDistanceRemainder = Sample.Clip->ExtractRootDistance(Sample.ClipTime);
 
-	Sample.RootTransform = RootMotionInitial;
-	Sample.RootDistance = RootDistanceInitial;
+		// Invert motion deltas if we wrapped backwards
+		if (SampleTime < 0.0f)
+		{
+			RootMotionRemainder.SetToRelativeTransform(RootMotionLast);
+			RootDistanceRemainder = -(RootDistanceLast - RootDistanceRemainder);
+		}
 
-	// Note if the sample was clamped, no motion will be applied here because NumCycles will be zero
-	int32 CyclesRemaining = SamplingParam.NumCycles;
-	while (CyclesRemaining--)
-	{
-		Sample.RootTransform *= RootMotionPerCycle;
-		Sample.RootDistance += RootDistancePerCycle;
-	}
+		Sample.RootTransform = RootMotionInitial;
+		Sample.RootDistance = RootDistanceInitial;
 
-	Sample.RootTransform *= RootMotionRemainder;
-	Sample.RootDistance += RootDistanceRemainder;
+		// Note if the sample was clamped, no motion will be applied here because NumCycles will be zero
+		int32 CyclesRemaining = SamplingParam.NumCycles;
+		while (CyclesRemaining--)
+		{
+			Sample.RootTransform = RootMotionPerCycle * Sample.RootTransform;
+         			Sample.RootDistance += RootDistancePerCycle;
+		}
+
+		Sample.RootTransform = RootMotionRemainder * Sample.RootTransform;
+		Sample.RootDistance += RootDistanceRemainder;
+	}
 
 	return Sample;
 }
@@ -2072,6 +2273,7 @@ static void DrawTrajectoryFeatures(const FDebugDrawParams& DrawParams, const FFe
 {
 	const float LifeTime = DrawParams.DefaultLifeTime;
 	const uint8 DepthPriority = ESceneDepthPriorityGroup::SDPG_Foreground + 2;
+	const bool bPersistent = EnumHasAnyFlags(DrawParams.Flags, EDebugDrawFlags::Persistent);
 
 	FPoseSearchFeatureDesc Feature;
 	Feature.Domain = Domain;
@@ -2099,7 +2301,14 @@ static void DrawTrajectoryFeatures(const FDebugDrawParams& DrawParams, const FFe
 			FColor Color = LinearColor.ToFColor(true);
 
 			TrajectoryPos = DrawParams.RootTransform.TransformPosition(TrajectoryPos);
-			DrawDebugSphere(DrawParams.World, TrajectoryPos, DrawDebugSphereSize, DrawDebugSphereSegments, Color, false, LifeTime, DepthPriority,  DrawDebugSphereLineThickness);
+			if (EnumHasAnyFlags(DrawParams.Flags, EDebugDrawFlags::DrawSearchIndex))
+			{
+				DrawDebugPoint(DrawParams.World, TrajectoryPos, DrawParams.PointSize, Color, bPersistent, DrawParams.DefaultLifeTime, DepthPriority);
+			}
+			else
+			{
+				DrawDebugSphere(DrawParams.World, TrajectoryPos, DrawDebugSphereSize, DrawDebugSphereSegments, Color, bPersistent, LifeTime, DepthPriority, DrawDebugSphereLineThickness);
+			}
 		}
 		else
 		{
@@ -2117,7 +2326,14 @@ static void DrawTrajectoryFeatures(const FDebugDrawParams& DrawParams, const FFe
 			TrajectoryVel *= DrawDebugVelocityScale;
 			TrajectoryVel = DrawParams.RootTransform.TransformVector(TrajectoryVel);
 			FVector TrajectoryVelDirection = TrajectoryVel.GetSafeNormal();
-			DrawDebugDirectionalArrow(DrawParams.World, TrajectoryPos + TrajectoryVelDirection * DrawDebugSphereSize, TrajectoryPos + TrajectoryVel, DrawDebugArrowSize, Color, false, LifeTime, DepthPriority, DrawDebugLineThickness);
+			if (EnumHasAnyFlags(DrawParams.Flags, EDebugDrawFlags::DrawSearchIndex))
+			{
+				DrawDebugPoint(DrawParams.World, TrajectoryVel, DrawParams.PointSize, Color, bPersistent, DrawParams.DefaultLifeTime, DepthPriority);
+			}
+			else
+			{
+				DrawDebugDirectionalArrow(DrawParams.World, TrajectoryPos + TrajectoryVelDirection * DrawDebugSphereSize, TrajectoryPos + TrajectoryVel, DrawDebugArrowSize, Color, bPersistent, LifeTime, DepthPriority, DrawDebugLineThickness);
+			}
 		}
 	}
 }
@@ -2129,6 +2345,7 @@ static void DrawPoseFeatures(const FDebugDrawParams& DrawParams, const FFeatureV
 
 	const float LifeTime = DrawParams.DefaultLifeTime;
 	const uint8 DepthPriority = ESceneDepthPriorityGroup::SDPG_Foreground + 2;
+	const bool bPersistent = EnumHasAnyFlags(DrawParams.Flags, EDebugDrawFlags::Persistent);
 
 	FPoseSearchFeatureDesc Feature;
 	Feature.Domain = EPoseSearchFeatureDomain::Time;
@@ -2159,7 +2376,14 @@ static void DrawPoseFeatures(const FDebugDrawParams& DrawParams, const FFeatureV
 				FColor Color = LinearColor.ToFColor(true);
 
 				BonePos = DrawParams.RootTransform.TransformPosition(BonePos);
-				DrawDebugSphere(DrawParams.World, BonePos, DrawDebugSphereSize, DrawDebugSphereSegments, Color, false, LifeTime, DepthPriority, DrawDebugSphereLineThickness);
+				if (EnumHasAnyFlags(DrawParams.Flags, EDebugDrawFlags::DrawSearchIndex))
+				{
+					DrawDebugPoint(DrawParams.World, BonePos, DrawParams.PointSize, Color, bPersistent, DrawParams.DefaultLifeTime, DepthPriority);
+				}
+				else
+				{
+					DrawDebugSphere(DrawParams.World, BonePos, DrawDebugSphereSize, DrawDebugSphereSegments, Color, bPersistent, LifeTime, DepthPriority, DrawDebugSphereLineThickness);
+				}
 			}
 
 			FVector BoneVel;
@@ -2173,7 +2397,14 @@ static void DrawPoseFeatures(const FDebugDrawParams& DrawParams, const FFeatureV
 				BoneVel *= DrawDebugVelocityScale;
 				BoneVel = DrawParams.RootTransform.TransformVector(BoneVel);
 				FVector BoneVelDirection = BoneVel.GetSafeNormal();
-				DrawDebugDirectionalArrow(DrawParams.World, BonePos + BoneVelDirection * DrawDebugSphereSize, BonePos + BoneVel, DrawDebugArrowSize, Color, false, LifeTime, DepthPriority, DrawDebugLineThickness);
+				if (EnumHasAnyFlags(DrawParams.Flags, EDebugDrawFlags::DrawSearchIndex))
+				{
+					DrawDebugPoint(DrawParams.World, BoneVel, DrawParams.PointSize, Color, bPersistent, DrawParams.DefaultLifeTime, DepthPriority);
+				}
+				else
+				{
+					DrawDebugDirectionalArrow(DrawParams.World, BonePos + BoneVelDirection * DrawDebugSphereSize, BonePos + BoneVel, DrawDebugArrowSize, Color, bPersistent, LifeTime, DepthPriority, DrawDebugLineThickness);
+				}
 			}
 		}
 	}
@@ -2703,6 +2934,7 @@ bool BuildIndex(const UAnimSequence* Sequence, UPoseSearchSequenceMetaData* Sequ
 	FSequenceSampler Sampler;
 	FSequenceSampler::FInput SamplerInput;
 	SamplerInput.Schema = SequenceMetaData->Schema;
+	SamplerInput.ExtrapolationParameters = SequenceMetaData->ExtrapolationParameters;
 	SamplerInput.Sequence = Sequence;
 	SamplerInput.bLoopable = false;
 	Sampler.Init(SamplerInput);
@@ -2751,6 +2983,7 @@ bool BuildIndex(UPoseSearchDatabase* Database)
 
 			FSequenceSampler::FInput Input;
 			Input.Schema = Database->Schema;
+			Input.ExtrapolationParameters = Database->ExtrapolationParameters;
 			Input.Sequence = Sequence;
 			Input.bLoopable = bLoopable;
 			SequenceSamplers[SequenceSamplerIdx].Init(Input);
diff --git a/Engine/Plugins/Experimental/Animation/PoseSearch/Source/Runtime/Public/PoseSearch/PoseSearch.h b/Engine/Plugins/Experimental/Animation/PoseSearch/Source/Runtime/Public/PoseSearch/PoseSearch.h
--- a/Engine/Plugins/Experimental/Animation/PoseSearch/Source/Runtime/Public/PoseSearch/PoseSearch.h
+++ b/Engine/Plugins/Experimental/Animation/PoseSearch/Source/Runtime/Public/PoseSearch/PoseSearch.h
@@ -312,6 +312,25 @@ struct POSESEARCH_API FPoseSearchIndex
 	void InverseNormalize (TArrayView<float> PoseVector) const;
 };
 
+USTRUCT()
+struct FPoseSearchExtrapolationParameters
+{
+	GENERATED_BODY()
+
+public:
+	// If the angular root motion speed in degrees is below this value, it will be treated as zero.
+	UPROPERTY(EditAnywhere, Category = "Settings")
+	float AngularSpeedThreshold = 1.0f;
+	
+	// If the root motion linear speed is below this value, it will be treated as zero.
+	UPROPERTY(EditAnywhere, Category = "Settings")
+	float LinearSpeedThreshold = 1.0f;
+
+	// Time from sequence start/end used to extrapolate the trajectory.
+	UPROPERTY(EditAnywhere, Category = "Settings")
+	float SampleTime = 0.05f;
+};
+
 /** Animation metadata object for indexing a single animation. */
 UCLASS(BlueprintType, Category = "Animation|PoseSearch")
 class POSESEARCH_API UPoseSearchSequenceMetaData : public UAnimMetaData
@@ -325,6 +344,9 @@ public:
 	UPROPERTY(EditAnywhere, Category="Settings")
 	FFloatInterval SamplingRange = FFloatInterval(0.0f, 0.0f);
 
+	UPROPERTY(EditAnywhere, Category = "Settings")
+	FPoseSearchExtrapolationParameters ExtrapolationParameters;
+
 	UPROPERTY();
 	FPoseSearchIndex SearchIndex;
 
@@ -452,6 +474,9 @@ public:
 	UPROPERTY(EditAnywhere, Category="Database")
 	TArray<FPoseSearchDatabaseSequence> Sequences;
 
+	UPROPERTY(EditAnywhere, Category="Database")
+	FPoseSearchExtrapolationParameters ExtrapolationParameters;
+
 	UPROPERTY()
 	FPoseSearchIndex SearchIndex;
 
@@ -616,6 +641,8 @@ enum class EDebugDrawFlags : uint32
 
 	/** Draw all pose vector features */
 	IncludeAllFeatures  = IncludePose | IncludeTrajectory,
+
+	Persistent = 1 << 3,
 };
 ENUM_CLASS_FLAGS(EDebugDrawFlags);
 
@@ -625,7 +652,10 @@ struct POSESEARCH_API FDebugDrawParams
 	const UPoseSearchDatabase* Database = nullptr;
 	const UPoseSearchSequenceMetaData* SequenceMetaData = nullptr;
 	EDebugDrawFlags Flags = EDebugDrawFlags::IncludeAllFeatures;
+
 	float DefaultLifeTime = 5.0f;
+	float PointSize = 1.0f;
+
 	FTransform RootTransform = FTransform::Identity;
 
 	/** If set, draw the corresponding pose from the search index */
