Motion Trajectory Component for Motion Matching:

Motion Trajectory Component notes:
1) Abstract component/interface implemented with prediction and history API.
2) Implemented uniform, frame-rate independent history sampling algorithm for retaining trajectory sample coherence.
3) Implemented concrete Character Movement Trajectory Component for encapsulating ground locomotion prediction algorithm and API.
4) Motion Trajectory blueprint library containing:
5) FlattenTrajectory2D algorithm for isolating and removing Z axis direction contribution from tracjectory.
6) ClampTrajectoryDirection for projecting trajectory samples into a discrete, allowed set of directions (such as cardinal).

Engine changes:
1) Implemented shared TrajectorySample, Range, and Domain structures for common usage among Motion Trajectory, Pose Search, and Blueprint.
2) Expanded Sequence Player blueprint library functionality for Scripted Motion Matching content example.

Pose Search changes:
1) Match Pose renamed to Pose History with minor configuration parameters.
2) Complete removal of approximated past trajectory sampling. This has now been pushed to the Motion Trajectory Component.
3) Motion Matching interface changes: Goal has been replaced with Trajectory and Dynamic Play Rate Adjustment has been encapsulated to an experimental code module and blueprint library. New defaults, and encapsulating settings.
4) Motion Matching changes: New Pose Search blueprint library which encapsulates core Motion Matching algorithm.
5) Unified Trajectory to Pose Feature Vector transformation algorithms.
6) Removal of temporary Pose Search Prediction blueprint library (now Dynamic Play Rate).
7) End-to-end implementation support for Past/Future Time and Distance domain schemas.

AnimInput changes:
1) Removal of AnimInput_CapsulePrediction (replaced with Motion Trajectory Component).

Gameplay Anim Gym changes:
1) Pose search database is now fully distance-based (history and future).
2) Character blueprint includes Character Movement Trajectory Component with distance-based domain setup.
3) Scripted Motion Matching anim node function prototype with Engine content example.

diff --git a/Engine/Plugins/Experimental/Animation/MotionTrajectory/MotionTrajectory.uplugin b/Engine/Plugins/Experimental/Animation/MotionTrajectory/MotionTrajectory.uplugin
new file mode 100644
--- /dev/null
+++ b/Engine/Plugins/Experimental/Animation/MotionTrajectory/MotionTrajectory.uplugin
@@ -0,0 +1,25 @@
+{
+	"FileVersion": 3,
+	"Version": 1,
+	"VersionName": "0.1",
+	"FriendlyName": "Motion Trajectory",
+	"Description": "Generate predictions and track history of character motion.",
+	"Category": "Animation",
+	"CreatedBy": "Epic Games, Inc.",
+	"CreatedByURL": "http://epicgames.com",
+	"DocsURL": "",
+	"MarketplaceURL": "",
+	"SupportURL": "",
+	"EnabledByDefault": false,
+	"CanContainContent": true,
+	"IsBetaVersion": false,
+	"IsExperimentalVersion": true,
+	"Installed": false,
+	"Modules": [
+		{
+			"Name": "MotionTrajectory",
+			"Type": "Runtime",
+			"LoadingPhase": "PreDefault"
+		}
+	]
+}
\ No newline at end of file
diff --git a/Engine/Plugins/Experimental/Animation/MotionTrajectory/Source/MotionTrajectory/MotionTrajectory.Build.cs b/Engine/Plugins/Experimental/Animation/MotionTrajectory/Source/MotionTrajectory/MotionTrajectory.Build.cs
new file mode 100644
--- /dev/null
+++ b/Engine/Plugins/Experimental/Animation/MotionTrajectory/Source/MotionTrajectory/MotionTrajectory.Build.cs
@@ -0,0 +1,21 @@
+// Copyright Epic Games, Inc. All Rights Reserved.
+
+using UnrealBuildTool;
+
+public class MotionTrajectory : ModuleRules
+{
+	public MotionTrajectory(ReadOnlyTargetRules Target) : base(Target)
+	{
+		PCHUsage = ModuleRules.PCHUsageMode.UseExplicitOrSharedPCHs;
+		
+		PublicDependencyModuleNames.AddRange(
+			new string[]
+			{
+				"Core",
+				"CoreUObject",
+                "Engine",
+				"AnimGraphRuntime"
+			}
+		);
+	}
+}
diff --git a/Engine/Plugins/Experimental/Animation/MotionTrajectory/Source/MotionTrajectory/Private/MotionTrajectory.cpp b/Engine/Plugins/Experimental/Animation/MotionTrajectory/Source/MotionTrajectory/Private/MotionTrajectory.cpp
new file mode 100644
--- /dev/null
+++ b/Engine/Plugins/Experimental/Animation/MotionTrajectory/Source/MotionTrajectory/Private/MotionTrajectory.cpp
@@ -0,0 +1,273 @@
+// Copyright Epic Games, Inc. All Rights Reserved.
+
+#include "MotionTrajectory.h"
+
+#include "GameFramework/Pawn.h"
+#include "Kismet/KismetMathLibrary.h"
+#include "Kismet/KismetSystemLibrary.h"
+
+#define LOCTEXT_NAMESPACE "MotionTrajectory"
+
+void FMotionTrajectoryModule::StartupModule()
+{
+}
+
+void FMotionTrajectoryModule::ShutdownModule()
+{
+}
+
+UMotionTrajectoryComponent::UMotionTrajectoryComponent(const FObjectInitializer& ObjectInitializer) : Super(ObjectInitializer)
+{
+	PrimaryComponentTick.bCanEverTick = true;
+	PrimaryComponentTick.bStartWithTickEnabled = true;
+	PrimaryComponentTick.TickGroup = TG_PrePhysics;
+}
+
+FTrajectorySample UMotionTrajectoryComponent::GetPresentTrajectory() const
+{
+	const APawn* Pawn = TryGetOwnerPawn();
+	checkf(false, TEXT("UMotionTrajectoryComponent::GetPresentTrajectory for Pawn: %s requires implementation."), Pawn ? *(Pawn->GetHumanReadableName()) : TEXT("NULL"));
+	return {};
+}
+
+FTransform UMotionTrajectoryComponent::GetPresentWorldTransform() const
+{
+	const APawn* Pawn = TryGetOwnerPawn();
+	checkf(false, TEXT("UMotionTrajectoryComponent::GetPresentWorldTransform for Pawn: %s requires implementation."), Pawn ? *(Pawn->GetHumanReadableName()) : TEXT("NULL"));
+	return {};
+}
+
+const APawn* UMotionTrajectoryComponent::TryGetOwnerPawn() const
+{
+	const AActor* Actor = GetOwner();
+	return Actor ? Cast<APawn>(Actor) : nullptr;
+}
+
+FTrajectorySampleRange UMotionTrajectoryComponent::CombineHistoryPresentPrediction(bool bIncludeHistory, const FTrajectorySampleRange& Prediction) const
+{
+	FTrajectorySampleRange Trajectory;
+	int32 TotalSampleCount = 1 + Prediction.Samples.Num(); // present + future
+
+	FTrajectorySampleRange History;
+	if (bIncludeHistory)
+	{
+		History = GetHistory();
+		TotalSampleCount += History.Samples.Num();
+	}
+
+	// Linearly combines all uniformly sampled trajectory ranges from [0...n (samples)] in domain ordering:
+	// Ex: [Past + Present + Future], with negative values representing the past, a zeroed present, and positive values representing the future
+	Trajectory.Samples.Reserve(TotalSampleCount);
+	Trajectory.Samples.Append(History.Samples);
+	Trajectory.Samples.Add(PresentTrajectory);
+	Trajectory.Samples.Append(Prediction.Samples);
+	Trajectory.SampleRate = SampleRate;
+
+	return Trajectory;
+}
+
+void UMotionTrajectoryComponent::TickHistoryEvictionPolicy()
+{
+	// World space transform and time are used for determining local-relative historical samples
+	const FTransform WorldTransform = GetPresentWorldTransform();
+	const float WorldGameTime = UKismetSystemLibrary::GetGameTimeInSeconds(GetWorld());
+
+	// Skip on the very first sample
+	if (PreviousWorldGameTime != 0.f)
+	{
+		// Compute world space/time deltas for historical sample decay
+		const float DeltaSeconds = WorldGameTime - PreviousWorldGameTime;
+		const float DeltaDistance = FVector::Distance(WorldTransform.GetLocation(), PreviousWorldTransform.GetLocation());
+		const FTransform DeltaTransform = PreviousWorldTransform.GetRelativeTransform(WorldTransform);
+
+		for (auto& Sample : SampleHistory)
+		{
+			Sample.AccumulatedSeconds -= DeltaSeconds;
+			Sample.AccumulatedDistance -= DeltaDistance;
+			Sample.Position = DeltaTransform.TransformPosition(Sample.Position);
+		}
+
+		const FTrajectorySample FirstSample = SampleHistory.First();
+		const float FirstSampleTime = FirstSample.AccumulatedSeconds;
+
+		constexpr int32 DistanceDomainMask = static_cast<int32>(ETrajectorySampleDomain::Distance);
+		constexpr int32 TimeDomainMask = static_cast<int32>(ETrajectorySampleDomain::Time);
+		const bool bDistanceDomainEnabled = (HistorySettings.Domain & DistanceDomainMask) == DistanceDomainMask;
+		const bool bTimeDomainEnabled = (HistorySettings.Domain & TimeDomainMask) == TimeDomainMask;
+
+		// For both time and distance history domains, we compute the concept of an "effective" time horizon:
+		// This value is present to guarantee uniform history decay against a fixed anchor point in the past when there is lack of trajectory motion
+		// The anchor point is defined by the AccumulatedTime of the furthest history sample at the exact moment of zero motion
+		if (FMath::IsNearlyZero(DeltaDistance) && !FirstSample.IsZeroSample())
+		{
+			if (EffectiveTimeDomain == 0.f)
+			{
+				// One time initialization of an effective time domain at the exact moment of zero motion
+				EffectiveTimeDomain = FirstSampleTime;
+			}
+		}
+		else
+		{
+			// When we are in motion, an effective time domain is not required
+			EffectiveTimeDomain = 0.f;
+		}
+
+		// Remove all trajectory samples which are outside the boundaries of the enabled domain horizons
+		SampleHistory.RemoveAll([&](const FTrajectorySample& Sample)
+			{
+				if (Sample.IsZeroSample() && PresentTrajectory.IsZeroSample())
+				{
+					return true;
+				}
+
+				// Distance horizon
+				if (bDistanceDomainEnabled && (Sample.AccumulatedDistance < -HistorySettings.Distance))
+				{
+					return true;
+				}
+
+				// Time horizon
+				if (bTimeDomainEnabled && (Sample.AccumulatedSeconds < -HistorySettings.Seconds))
+				{
+					return true;
+				}
+
+				// Effective time horizon
+				if (EffectiveTimeDomain != 0.f&& (Sample.AccumulatedSeconds < EffectiveTimeDomain))
+				{
+					return true;
+				}
+
+				return false;
+			});
+	}
+
+	// Disallow trajectory samples beyond the maximum count
+	// An opportunity to insert more samples will arise as history decays
+	if (SampleHistory.Num() < MaxSamples)
+	{
+		SampleHistory.Emplace(PresentTrajectory);
+	}
+	
+	// Cache the present world space sample as the next frame's reference point for delta computation
+	PreviousWorldTransform = WorldTransform;
+	PreviousWorldGameTime = WorldGameTime;
+}
+
+void UMotionTrajectoryComponent::FlushHistory()
+{
+	SampleHistory.Reset();
+	PreviousWorldTransform = {};
+	PreviousWorldGameTime = 0.f;
+	EffectiveTimeDomain = 0.f;
+}
+
+void UMotionTrajectoryComponent::OnComponentCreated()
+{
+	FlushHistory();
+	SampleHistory.Reserve(MaxSamples);
+	Super::OnComponentCreated();
+}
+
+void UMotionTrajectoryComponent::TickComponent(float DeltaTime, enum ELevelTick TickType, FActorComponentTickFunction* ThisTickFunction)
+{
+	// Compute the instantaneous/present trajectory sample and guarantee that is zeroed on all accumulated domains
+	PresentTrajectory = GetPresentTrajectory();
+	PresentTrajectory.AccumulatedDistance = 0.f;
+	PresentTrajectory.AccumulatedSeconds = 0.f;
+
+	// Tick the historical sample retention/decay algorithm by one iteration
+	TickHistoryEvictionPolicy();
+
+	Super::TickComponent(DeltaTime, TickType, ThisTickFunction);
+}
+
+FTrajectorySampleRange UMotionTrajectoryComponent::GetTrajectory() const
+{
+	const APawn* Pawn = TryGetOwnerPawn();
+	checkf(false, TEXT("UMotionTrajectoryComponent::GetTrajectory for Pawn: %s requires implementation."), Pawn ? *(Pawn->GetHumanReadableName()) : TEXT("NULL"));
+	return FTrajectorySampleRange(SampleRate);
+}
+
+FTrajectorySampleRange UMotionTrajectoryComponent::GetTrajectoryWithSettings(const FMotionTrajectorySettings& Settings, bool bIncludeHistory) const
+{
+	const APawn* Pawn = TryGetOwnerPawn();
+	checkf(false, TEXT("UMotionTrajectoryComponent::GetTrajectoryWithSettings for Pawn: %s requires implementation."), Pawn ? *(Pawn->GetHumanReadableName()) : TEXT("NULL"));
+	return FTrajectorySampleRange(SampleRate);
+}
+
+void UMotionTrajectoryComponent::SetSampleRate(int32 Rate)
+{
+	// To maintain trajectory sample coherency, changing the sample rate will cause all historical samples to be flushed
+	// In the future, we may consider leveraging the frame-rate independent uniform resampling algorithm implemented within GetHistory()
+	FlushHistory();
+	SampleRate = Rate;
+}
+
+FTrajectorySampleRange UMotionTrajectoryComponent::GetHistory() const
+{
+	if (SampleHistory.IsEmpty())
+	{
+		return FTrajectorySampleRange(SampleRate);
+	}
+	else if (bUniformSampledHistory)
+	{
+		// Establish a uniform history range that will be frame-rate independent, and resampled against the SampleRate property specified by this component
+		FTrajectorySampleRange UniformHistory;
+		UniformHistory.Samples.Reserve(SampleHistory.Num());
+		UniformHistory.SampleRate = SampleRate;
+
+		int32 InterpStartIdx = 0;
+		const float SampleFrequency = 1.f / static_cast<float>(SampleRate);
+		const float FirstNonUniformSampleTime = SampleHistory.First().AccumulatedSeconds;
+		const float FirstUniformSampleTime = FMath::RoundToFloat(FirstNonUniformSampleTime / SampleFrequency) * SampleFrequency;
+
+		check(SampleFrequency > 0.f);
+		check(FirstUniformSampleTime <= 0.f);
+
+		// Iterate the source history buffer with a uniform sample frequency
+		// The iteration exit condition is meant to guarantee that we do not include (erronenously) a instantaneous/present trajectory sample at zero time
+		for (int32 Idx = 0;; ++Idx)
+		{
+			const float SampleTime = FirstUniformSampleTime + static_cast<float>(Idx) * SampleFrequency;
+			check(SampleTime <= 0.f);
+
+			if (!FMath::IsNearlyZero(SampleTime))
+			{
+				const FTrajectorySample Sample = FTrajectorySampleRange::IterSampleTrajectory(SampleHistory, ETrajectorySampleDomain::Time, SampleTime, InterpStartIdx, bSmoothInterpolation);
+				UniformHistory.Samples.Emplace(Sample);
+			}
+			else
+			{
+				break;
+			}
+		}
+
+		return UniformHistory;
+	}
+	else
+	{
+		const int32 HistorySampleCount = SampleHistory.Num();
+		const int32 PresentSampleIdx = HistorySampleCount - 1;
+
+		// Convert the source history ring buffer into a trajectory range
+		// We consider these values to be non-uniform with a variable, frame-rate dependent sample rate governed by the component tick rate
+		FTrajectorySampleRange SourceHistory;
+		SourceHistory.Samples.Reserve(HistorySampleCount);
+
+		for (const auto& Sample : SampleHistory)
+		{
+			SourceHistory.Samples.Emplace(Sample);
+		}
+
+		// Remove the present trajectory sample, this will not be included in the history
+		check(FMath::IsNearlyZero(SourceHistory.Samples[PresentSampleIdx].AccumulatedSeconds));
+		SourceHistory.Samples.RemoveAt(PresentSampleIdx);
+
+		return SourceHistory;
+	}
+}
+
+#undef LOCTEXT_NAMESPACE
+
+IMPLEMENT_MODULE(FMotionTrajectoryModule, MotionTrajectory)
\ No newline at end of file
diff --git a/Engine/Plugins/Experimental/Animation/MotionTrajectory/Source/MotionTrajectory/Private/MotionTrajectoryCharacterMovement.cpp b/Engine/Plugins/Experimental/Animation/MotionTrajectory/Source/MotionTrajectory/Private/MotionTrajectoryCharacterMovement.cpp
new file mode 100644
--- /dev/null
+++ b/Engine/Plugins/Experimental/Animation/MotionTrajectory/Source/MotionTrajectory/Private/MotionTrajectoryCharacterMovement.cpp
@@ -0,0 +1,225 @@
+// Copyright Epic Games, Inc. All Rights Reserved.
+
+#include "MotionTrajectoryCharacterMovement.h"
+
+#include "Animation/AnimInstance.h"
+#include "GameFramework/CharacterMovementComponent.h"
+#include "GameFramework/Pawn.h"
+#include "KismetAnimationLibrary.h"
+
+// ----------- BEGIN Derived from FCharacterMovementComponentAsyncInput::CalcVelocity() ----------- //
+
+static void StepPrediction(float IntegrationDelta
+	, const UCharacterMovementComponent* MovementComponent
+	, FTrajectorySample& Sample)
+{
+	if (MovementComponent->GetCurrentAcceleration().IsZero())
+	{
+		float ActualBrakingFriction = (MovementComponent->bUseSeparateBrakingFriction ? MovementComponent->BrakingFriction : MovementComponent->GroundFriction);
+
+		if (!Sample.LocalLinearVelocity.IsZero())
+		{
+			const float FrictionFactor = FMath::Max(0.f, MovementComponent->BrakingFrictionFactor);
+			ActualBrakingFriction = FMath::Max(0.f, ActualBrakingFriction * FrictionFactor);
+				
+			const float BrakingDeceleration = FMath::Max(0.f, MovementComponent->GetMaxBrakingDeceleration());
+			const bool bZeroFriction = (ActualBrakingFriction == 0.f);
+			const bool bZeroBraking = (BrakingDeceleration == 0.f);
+
+			if (bZeroFriction && bZeroBraking)
+			{
+				return;
+			}
+
+			float RemainingTime = IntegrationDelta;
+			const float MaxTimeStep = FMath::Clamp(MovementComponent->BrakingSubStepTime, 1.f / 75.f, 1.f / 20.f);
+
+			const FVector PrevLinearVelocity = Sample.LocalLinearVelocity;
+			const FVector RevAccel = (bZeroBraking ? FVector::ZeroVector : (-BrakingDeceleration * Sample.LocalLinearVelocity.GetSafeNormal()));
+
+			while (RemainingTime >= UCharacterMovementComponent::MIN_TICK_TIME)
+			{
+				const float dt = ((RemainingTime > MaxTimeStep && !bZeroFriction) ? FMath::Min(MaxTimeStep, RemainingTime * 0.5f) : RemainingTime);
+				RemainingTime -= dt;
+
+				Sample.LocalLinearVelocity = Sample.LocalLinearVelocity + ((-ActualBrakingFriction) * Sample.LocalLinearVelocity + RevAccel) * dt;
+
+				if ((Sample.LocalLinearVelocity | PrevLinearVelocity) <= 0.f)
+				{
+					Sample.LocalLinearVelocity = FVector::ZeroVector;
+					return;
+				}
+			}
+
+			// Clamp to zero if nearly zero, or if below min threshold and braking
+			const float VSizeSq = Sample.LocalLinearVelocity.SizeSquared();
+			if (VSizeSq <= KINDA_SMALL_NUMBER || (!bZeroBraking && VSizeSq <= FMath::Square(UCharacterMovementComponent::BRAKE_TO_STOP_VELOCITY)))
+			{
+				Sample.LocalLinearVelocity = FVector::ZeroVector;
+			}
+		}
+	}
+	else
+	{
+		const FVector AccelDir = Sample.LocalLinearAcceleration.GetSafeNormal();
+		const float VelSize = Sample.LocalLinearVelocity.Size();
+		Sample.LocalLinearVelocity = Sample.LocalLinearVelocity - (Sample.LocalLinearVelocity - AccelDir * VelSize) * FMath::Min(IntegrationDelta * MovementComponent->GroundFriction, 1.f);
+
+		const float MaxInputSpeed = FMath::Max(MovementComponent->GetMaxSpeed() * MovementComponent->GetAnalogInputModifier(), MovementComponent->GetMinAnalogSpeed());
+		Sample.LocalLinearVelocity += Sample.LocalLinearAcceleration * IntegrationDelta;
+		Sample.LocalLinearVelocity = Sample.LocalLinearVelocity.GetClampedToMaxSize(MaxInputSpeed);
+	}
+
+	Sample.Position += Sample.LocalLinearVelocity * IntegrationDelta;
+}
+
+static void PredictTrajectory(const UCharacterMovementComponent* MovementComponent
+	, int32 SampleRate
+	, int32 MaxSamples
+	, const FMotionTrajectorySettings& Settings
+	, const FTrajectorySample& PresentTrajectory
+	, FTrajectorySampleRange& ReturnValue)
+{
+	ReturnValue.SampleRate = SampleRate;
+	const float IntegrationDelta = 1.f / static_cast<float>(SampleRate);
+
+	if (!!Settings.Domain)
+	{
+		FTrajectorySample Sample = PresentTrajectory;
+		FTrajectorySample PreviousSample = PresentTrajectory;
+		float AccumulatedDistance = 0.f;
+		float AccumulatedSeconds = 0.f;
+
+		constexpr int32 DistanceDomainMask = static_cast<int32>(ETrajectorySampleDomain::Distance);
+		constexpr int32 TimeDomainMask = static_cast<int32>(ETrajectorySampleDomain::Time);
+
+		for (int32 Step = 0; Step < MaxSamples; ++Step)
+		{
+			PreviousSample = Sample;
+			StepPrediction(IntegrationDelta, MovementComponent, Sample);
+
+			AccumulatedDistance += FVector::Distance(PreviousSample.Position, Sample.Position);
+			Sample.AccumulatedDistance = AccumulatedDistance;
+			AccumulatedSeconds += IntegrationDelta;
+			Sample.AccumulatedSeconds = AccumulatedSeconds;
+
+			if (FMath::IsNearlyEqual(FMath::Abs(Sample.AccumulatedDistance - PreviousSample.AccumulatedDistance), SMALL_NUMBER))
+			{
+				break;
+			}
+
+			ReturnValue.Samples.Add(Sample);
+
+			if (((Settings.Domain & DistanceDomainMask) == DistanceDomainMask)
+				&& (Settings.Distance > 0.f)
+				&& (Sample.AccumulatedDistance < Settings.Distance))
+			{
+				continue;
+			}
+
+			if (((Settings.Domain & TimeDomainMask) == TimeDomainMask)
+				&& (Settings.Seconds > 0.f)
+				&& (Step * IntegrationDelta < Settings.Seconds)
+				&& (Sample.AccumulatedDistance > PreviousSample.AccumulatedDistance))
+			{
+				continue;
+			}
+
+			break;
+		}
+	}
+}
+
+// ------------ END Derived from FCharacterMovementComponentAsyncInput::CalcVelocity() ------------ //
+
+FTrajectorySample UCharacterMovementTrajectoryComponent::GetPresentTrajectory() const
+{
+	FTrajectorySample ReturnValue;
+	
+	const APawn* Pawn = TryGetOwnerPawn();
+	if (!Pawn)
+	{
+		return ReturnValue;
+	}
+
+	const UCharacterMovementComponent* MovementComponent = Cast<UCharacterMovementComponent>(Pawn->GetMovementComponent());
+	if (ensure(MovementComponent) && MovementComponent->MovementMode == EMovementMode::MOVE_Walking)
+	{
+		const USkeletalMeshComponent* SkeletalMeshComponent = Pawn->FindComponentByClass<USkeletalMeshComponent>();
+		FTransform ComponentWorldTransform = SkeletalMeshComponent ? SkeletalMeshComponent->GetComponentTransform() : FTransform();
+
+		ReturnValue.LocalLinearVelocity = ComponentWorldTransform.InverseTransformVectorNoScale(MovementComponent->Velocity);
+		ReturnValue.LocalLinearAcceleration = ComponentWorldTransform.InverseTransformVectorNoScale(MovementComponent->GetCurrentAcceleration());
+	}
+	return ReturnValue;
+}
+
+FTransform UCharacterMovementTrajectoryComponent::GetPresentWorldTransform() const
+{
+	FTransform ReturnValue = FTransform::Identity;
+
+	const APawn* Pawn = TryGetOwnerPawn();
+	if (!Pawn)
+	{
+		return ReturnValue;
+	}
+
+	const USkeletalMeshComponent* SkeletalMeshComponent = Pawn->FindComponentByClass<USkeletalMeshComponent>();
+	ReturnValue = SkeletalMeshComponent ? SkeletalMeshComponent->GetComponentTransform() : Pawn->GetActorTransform();
+
+	return ReturnValue;
+}
+
+UCharacterMovementTrajectoryComponent::UCharacterMovementTrajectoryComponent(const FObjectInitializer& ObjectInitializer) : Super(ObjectInitializer)
+{
+	bWantsInitializeComponent = true;
+}
+
+void UCharacterMovementTrajectoryComponent::InitializeComponent()
+{
+	// Add a tick dependency on the character movement component since this component implementation is reliant on its internal state
+	UCharacterMovementComponent* CharacterMovementComponent = GetOwner()->FindComponentByClass<UCharacterMovementComponent>();
+	AddTickPrerequisiteComponent(CharacterMovementComponent);
+
+	Super::InitializeComponent();
+}
+
+void UCharacterMovementTrajectoryComponent::UninitializeComponent()
+{
+	UCharacterMovementComponent* CharacterMovementComponent = GetOwner()->FindComponentByClass<UCharacterMovementComponent>();
+	RemoveTickPrerequisiteComponent(CharacterMovementComponent);
+
+	Super::UninitializeComponent();
+}
+
+FTrajectorySampleRange UCharacterMovementTrajectoryComponent::GetTrajectory() const
+{
+	return GetTrajectoryWithSettings(PredictionSettings, bPredictionIncludesHistory);
+}
+
+FTrajectorySampleRange UCharacterMovementTrajectoryComponent::GetTrajectoryWithSettings(const FMotionTrajectorySettings& Settings
+	, bool bIncludeHistory) const
+{
+	const APawn* Pawn = TryGetOwnerPawn();
+	if (!Pawn)
+	{
+		return FTrajectorySampleRange(SampleRate);
+	}
+
+	const UCharacterMovementComponent* MovementComponent = Cast<UCharacterMovementComponent>(Pawn->GetMovementComponent());
+
+	// Currently the trajectory prediction only supports the walking movement mode of the character movement component
+	if (ensure(MovementComponent) && MovementComponent->MovementMode == EMovementMode::MOVE_Walking)
+	{
+		// Step the prediction iteratively towards the specified domain horizon(s)
+		FTrajectorySampleRange Prediction(SampleRate);
+		PredictTrajectory(MovementComponent, SampleRate, MaxSamples, Settings, PresentTrajectory, Prediction);
+
+		// Combine past, present, and future into a uniformly sampled complete trajectory
+		return CombineHistoryPresentPrediction(bIncludeHistory, Prediction);
+	}
+	else
+	{
+		return FTrajectorySampleRange(SampleRate);
+	}
+}
\ No newline at end of file
diff --git a/Engine/Plugins/Experimental/Animation/MotionTrajectory/Source/MotionTrajectory/Private/MotionTrajectoryLibrary.cpp b/Engine/Plugins/Experimental/Animation/MotionTrajectory/Source/MotionTrajectory/Private/MotionTrajectoryLibrary.cpp
new file mode 100644
--- /dev/null
+++ b/Engine/Plugins/Experimental/Animation/MotionTrajectory/Source/MotionTrajectory/Private/MotionTrajectoryLibrary.cpp
@@ -0,0 +1,173 @@
+// Copyright Epic Games, Inc. All Rights Reserved.
+
+#include "MotionTrajectoryLibrary.h"
+#include "Algo/Find.h"
+
+static void FlattenTrajectoryPosition(FTrajectorySample& Sample, float& AccumulatedDeltaDistance, bool PreserveSpeed)
+{
+	if (!Sample.Position.IsZero())
+	{
+		if (PreserveSpeed)
+		{
+			// Take the full displacement, effectively meaning that the Z axis never existed
+			Sample.Position = Sample.Position.Size() * Sample.Position.GetSafeNormal2D();
+		}
+		else
+		{
+			// Accumulate the delta displacement difference as a result of Z axis being removed
+			const float DeltaDistance = Sample.Position.Size() - Sample.Position.Size2D();
+			AccumulatedDeltaDistance += DeltaDistance;
+
+			// Remove the accumulated delta displacement from each sample on the timeline, and recompute the new projected position
+			Sample.AccumulatedDistance -= AccumulatedDeltaDistance;
+			Sample.Position = FMath::Abs(Sample.AccumulatedDistance) * Sample.Position.GetSafeNormal2D();
+		}
+	}
+}
+
+FTrajectorySampleRange UMotionTrajectoryBlueprintLibrary::FlattenTrajectory2D(FTrajectorySampleRange Trajectory, bool PreserveSpeed)
+{
+	if (!Trajectory.HasSamples())
+	{
+		return Trajectory;
+	}
+
+	if (Trajectory.HasOnlyZeroSamples())
+	{
+		return Trajectory;
+	}
+
+	// Each iteration will preserve the linear magnitudes of velocity and acceleration while removing the direction Z-axis component
+	for (auto& Sample : Trajectory.Samples)
+	{
+		// Note: As as a consequence of magnitude preservation, AccumulatedDistance alongside AccumulatedTime should not need modification
+
+		// Linear velocity Z-axis component removal
+		if (!Sample.LocalLinearVelocity.IsZero())
+		{
+			const float VelMagnitude = PreserveSpeed ? Sample.LocalLinearVelocity.Size() : Sample.LocalLinearVelocity.Size2D();
+			Sample.LocalLinearVelocity = VelMagnitude * Sample.LocalLinearVelocity.GetSafeNormal2D();
+		}
+
+		// Align linear acceleration via projection onto velocity
+		if (!Sample.LocalLinearVelocity.IsZero() && !Sample.LocalLinearAcceleration.IsZero())
+		{
+			const float AccelMagnitude = PreserveSpeed ? Sample.LocalLinearAcceleration.Size() : Sample.LocalLinearAcceleration.Size2D();
+			Sample.LocalLinearAcceleration = AccelMagnitude * Sample.LocalLinearAcceleration.ProjectOnTo(Sample.LocalLinearVelocity).GetSafeNormal();
+		}
+	}
+
+	// The present position sample is used as the basis for recomputing the future and history Accumulated Distance
+	const int32 PresentSampleIdx = Trajectory.Samples.IndexOfByPredicate([](const FTrajectorySample& Sample){
+		return Sample.AccumulatedSeconds == 0.f;
+	});
+
+	check(PresentSampleIdx != INDEX_NONE);
+
+	// Walk all samples into the future, conditionally removing contribution of Z axis motion
+	float AccumulatedDeltaDistance = 0.f;
+	for (int32 Idx = PresentSampleIdx, Num = Trajectory.Samples.Num(); Idx < Num; ++Idx)
+	{
+		FlattenTrajectoryPosition(Trajectory.Samples[Idx], AccumulatedDeltaDistance, PreserveSpeed);
+	}
+
+	// There is a possibility history has not been computed yet
+	if (PresentSampleIdx == 0)
+	{
+		return Trajectory;
+	}
+
+	// Walk all samples in the past, conditionally removing the contribution of Z axis motion
+	AccumulatedDeltaDistance = 0.f;
+	for (int32 Idx = PresentSampleIdx - 1, Begin = 0; Idx >= Begin; --Idx)
+	{
+		FlattenTrajectoryPosition(Trajectory.Samples[Idx], AccumulatedDeltaDistance, PreserveSpeed);
+	}
+
+	return Trajectory;
+}
+
+static FVector ClampDirection(const FVector InputVector, const TArray<FTrajectoryDirectionClamp>& Directions)
+{
+	if (Directions.IsEmpty())
+	{
+		return InputVector;
+	}
+
+	FVector InputDirection;
+	float InputLength;
+	InputVector.ToDirectionAndLength(InputDirection, InputLength);
+	if (InputLength < SMALL_NUMBER)
+	{
+		return InputVector;
+	}
+
+	// Assume first direction is best then check if the input direction is within the remaining sectors
+	FVector NearestDirection = Directions[0].Direction;
+	for (int32 DirIdx = 1; DirIdx != Directions.Num(); ++DirIdx)
+	{
+		const auto& Clamp = Directions[DirIdx];
+		if (FMath::Acos(FVector::DotProduct(InputDirection, Clamp.Direction)) < FMath::DegreesToRadians(Clamp.AngleTresholdDegrees))
+		{
+			NearestDirection = Clamp.Direction;
+			break;
+		}
+	}
+
+	const FVector Output = InputLength * NearestDirection;
+	return Output;
+}
+
+FTrajectorySampleRange UMotionTrajectoryBlueprintLibrary::ClampTrajectoryDirection(FTrajectorySampleRange Trajectory, const TArray<FTrajectoryDirectionClamp>& Directions)
+{
+	if (Directions.IsEmpty())
+	{
+		return Trajectory;
+	}
+
+	if (!Trajectory.HasSamples())
+	{
+		return Trajectory;
+	}
+
+	if (Trajectory.HasOnlyZeroSamples())
+	{
+		return Trajectory;
+	}
+
+	// The clamped present (zero domain) sample is used as the basis for projecting samples along its trajectory
+	FTrajectorySample* PresentSample = Algo::FindByPredicate(Trajectory.Samples, [](const FTrajectorySample& Sample) {
+		return FMath::IsNearlyZero(Sample.AccumulatedSeconds);
+	}
+	);
+
+	check(PresentSample)
+
+	if (!PresentSample->LocalLinearVelocity.IsZero())
+	{
+		const FVector VelocityBasis = ClampDirection(PresentSample->LocalLinearVelocity, Directions).GetSafeNormal();
+
+		for (auto& Sample : Trajectory.Samples)
+		{
+			// Align linear velocity onto the velocity basis to maintain the present intended direction, while retaining per-sample magnitude
+			if (!Sample.LocalLinearVelocity.IsZero())
+			{
+				Sample.LocalLinearVelocity = Sample.LocalLinearVelocity.Size() * Sample.LocalLinearVelocity.ProjectOnTo(VelocityBasis).GetSafeNormal();
+			}
+
+			// Align linear acceleration through projection onto the modified velocity
+			if (!Sample.LocalLinearVelocity.IsZero() && !Sample.LocalLinearAcceleration.IsZero())
+			{
+				Sample.LocalLinearAcceleration = Sample.LocalLinearAcceleration.Size() * Sample.LocalLinearAcceleration.ProjectOnTo(Sample.LocalLinearVelocity).GetSafeNormal();
+			}
+
+			// Align the position path through projection onto the modified velocity
+			if (!Sample.LocalLinearVelocity.IsZero() && !Sample.Position.IsZero())
+			{
+				Sample.Position = FMath::Abs(Sample.AccumulatedDistance) * Sample.Position.ProjectOnTo(Sample.LocalLinearVelocity).GetSafeNormal();
+			}
+		}
+	}
+
+	return Trajectory;
+}
\ No newline at end of file
diff --git a/Engine/Plugins/Experimental/Animation/MotionTrajectory/Source/MotionTrajectory/Public/MotionTrajectory.h b/Engine/Plugins/Experimental/Animation/MotionTrajectory/Source/MotionTrajectory/Public/MotionTrajectory.h
new file mode 100644
--- /dev/null
+++ b/Engine/Plugins/Experimental/Animation/MotionTrajectory/Source/MotionTrajectory/Public/MotionTrajectory.h
@@ -0,0 +1,144 @@
+// Copyright Epic Games, Inc. All Rights Reserved.
+
+#pragma once
+
+#include "Animation/MotionTrajectoryTypes.h"
+#include "Components/ActorComponent.h"
+#include "Containers/RingBuffer.h"
+#include "Modules/ModuleManager.h"
+#include "UObject/ObjectMacros.h"
+
+#include "MotionTrajectory.generated.h"
+
+class APawn;
+
+class FMotionTrajectoryModule : public IModuleInterface
+{
+public: 
+
+	virtual void StartupModule() override;
+	virtual void ShutdownModule() override;
+};
+
+// Specifies the chosen domain parameters for trajectory sample retention and creation
+USTRUCT(BlueprintType, Category="Motion Trajectory")
+struct MOTIONTRAJECTORY_API FMotionTrajectorySettings
+{
+	GENERATED_BODY()
+
+	// Trajectory mask for specifying a time and/or distance domain
+	// Default: Time Domain of 2.0 seconds
+	UPROPERTY(EditAnywhere, BlueprintReadWrite, Category="Settings", meta=(Bitmask, BitmaskEnum=ETrajectorySampleDomain))
+	int32 Domain = 1; // ETrajectorySampleDomain::TIME
+
+	// Sample time horizon
+	UPROPERTY(EditAnywhere, BlueprintReadWrite, Category="Settings", meta=(ClampMin="0.0"))
+	float Seconds = 2.f;
+
+	// Sample distance horizon
+	UPROPERTY(EditAnywhere, BlueprintReadWrite, Category="Settings", meta=(ClampMin="0.0"))
+	float Distance = 0.f;
+};
+
+// Abstract component interface for the minimum Motion Trajectory prediction and history API
+UCLASS(ClassGroup=Movement, abstract, BlueprintType, Category="Motion Trajectory")
+class MOTIONTRAJECTORY_API UMotionTrajectoryComponent : public UActorComponent
+{
+	GENERATED_BODY()
+
+private:
+
+	// Previous frame world-space transform
+	FTransform PreviousWorldTransform = {};
+
+	// Previous frame world game time
+	float PreviousWorldGameTime = 0.f;
+
+	// A superceding "effective" time domain applied to distance and time domain history samples, which guarantees uniform decay
+	float EffectiveTimeDomain = 0.f;
+
+	// Component-driven tick which decays and evicts trajectory history outside of the defined domain horizons
+	void TickHistoryEvictionPolicy();
+
+protected:
+
+	// Retained trajectory history samples
+	TRingBuffer<FTrajectorySample> SampleHistory = {};
+
+	// Retained instantaneous/present trajectory
+	FTrajectorySample PresentTrajectory = {};
+
+	// Gets the instantaneous/present trajectory sample of the current frame
+	virtual FTrajectorySample GetPresentTrajectory() const;
+	
+	// Gets the instantaneous/present world-space transform of the current frame
+	virtual FTransform GetPresentWorldTransform() const;
+
+	// Gets the Pawn from the owning Actor of this component
+	const APawn* TryGetOwnerPawn() const;
+
+	// Combines all trajectory samples in the past, present, and future into a unifed trajectory range
+	FTrajectorySampleRange CombineHistoryPresentPrediction(bool bIncludeHistory, const FTrajectorySampleRange& Prediction) const;
+
+	// Forcefully evicts all trajectory history and resets internal history tracking state
+	void FlushHistory();
+
+public:
+
+	UMotionTrajectoryComponent(const FObjectInitializer& ObjectInitializer);
+
+	// Begin UActorComponent Interface
+	virtual void OnComponentCreated() override;
+	virtual void TickComponent(float DeltaTime
+		, enum ELevelTick TickType
+		, FActorComponentTickFunction* ThisTickFunction) override;
+	// End UActorComponent Interface
+
+	// Interface for computing a trajectory prediction
+	UFUNCTION(BlueprintPure, Category="Motion Trajectory")
+	virtual FTrajectorySampleRange GetTrajectory() const;
+
+	// Interface for computing a trajectory prediction with overridden component settings
+	UFUNCTION(BlueprintPure, Category="Motion Trajectory", meta=(AutoCreateRefTerm="Settings", DisplayName="Get Trajectory (With Settings)"))
+	virtual FTrajectorySampleRange GetTrajectoryWithSettings(const FMotionTrajectorySettings& Settings, bool bIncludeHistory) const;
+	
+	// Changes the history and prediction sample rate for the component
+	UFUNCTION(BlueprintCallable, Category="Motion Trajectory")
+	void SetSampleRate(int32 Rate);
+
+	// Retrieves the trajectory history as a range
+	UFUNCTION(BlueprintCallable, Category="Motion Trajectory")
+	FTrajectorySampleRange GetHistory() const;
+
+	// Prediction simulation settings
+	UPROPERTY(EditDefaultsOnly, BlueprintReadOnly, Category="Settings")
+	FMotionTrajectorySettings PredictionSettings;
+
+	// History retention settings
+	UPROPERTY(EditDefaultsOnly, BlueprintReadOnly, Category="Settings")
+	FMotionTrajectorySettings HistorySettings;
+
+	// The trajectory sampling rate for both prediction and history
+	UPROPERTY(EditDefaultsOnly, BlueprintReadOnly, Category="Settings", meta=(ClampMin="5", ClampMax="120"))
+	int32 SampleRate = 30;
+
+	// The maximum number of trajectory samples retained by the internal history buffer
+	UPROPERTY(EditDefaultsOnly, BlueprintReadOnly, Category="Settings", meta=(ClampMin="0", ClampMax="3600"))
+	int32 MaxSamples = 1000;
+
+	// Determines if GetTrajectory() will automatically combine history into the returned trajectory prediction (true)
+	// Consider disabling this option when historical trajectory samples are never used
+	// Example: No Motion Matching historical sample times or distances are defined in the Pose Search Schema
+	UPROPERTY(EditDefaultsOnly, BlueprintReadOnly, Category="Settings")
+	bool bPredictionIncludesHistory = true;
+
+	// Determines if history will be sampled by the specified Sample Rate (true) or the Component Tick Rate (false)
+	// Consider disabling this option when the engine frame rate and component sample rate are the same
+	UPROPERTY(EditDefaultsOnly, BlueprintReadOnly, Category="Settings")
+	bool bUniformSampledHistory = true;
+
+	// History samples will use Centripetal Catmull–Rom spline interpolation (true) (slow, accurate) instead of linear interpolation (false) (fast, approximate)
+	// Consider enabling this option when accurate trajectory sample speed at low tick rates (high interpolation usage) is important
+	UPROPERTY(EditAnywhere, BlueprintReadWrite, Category="Settings")
+	bool bSmoothInterpolation = false;
+};
\ No newline at end of file
diff --git a/Engine/Plugins/Experimental/Animation/MotionTrajectory/Source/MotionTrajectory/Public/MotionTrajectoryCharacterMovement.h b/Engine/Plugins/Experimental/Animation/MotionTrajectory/Source/MotionTrajectory/Public/MotionTrajectoryCharacterMovement.h
new file mode 100644
--- /dev/null
+++ b/Engine/Plugins/Experimental/Animation/MotionTrajectory/Source/MotionTrajectory/Public/MotionTrajectoryCharacterMovement.h
@@ -0,0 +1,35 @@
+// Copyright Epic Games, Inc. All Rights Reserved.
+
+#pragma once
+
+#include "MotionTrajectory.h"
+
+#include "MotionTrajectoryCharacterMovement.generated.h"
+
+// Example motion trajectory component implementation for encapsulating: Character Movement ground locomotion
+UCLASS(meta=(BlueprintSpawnableComponent), Category="Motion Trajectory")
+class MOTIONTRAJECTORY_API UCharacterMovementTrajectoryComponent : public UMotionTrajectoryComponent
+{
+	GENERATED_BODY()
+
+protected:
+
+	// Begin UMotionTrajectoryComponent Interface
+	virtual FTrajectorySample GetPresentTrajectory() const override;
+	virtual FTransform GetPresentWorldTransform() const override;
+	// End UMotionTrajectoryComponent Interface
+
+public:
+
+	UCharacterMovementTrajectoryComponent(const FObjectInitializer& ObjectInitializer);
+
+	// Begin UActorComponent Interface
+	virtual void InitializeComponent() override;
+	virtual void UninitializeComponent() override;
+	// End UActorComponent Interface
+
+	// Begin UMotionTrajectoryComponent Interface
+	virtual FTrajectorySampleRange GetTrajectory() const override;
+	virtual FTrajectorySampleRange GetTrajectoryWithSettings(const FMotionTrajectorySettings& Settings, bool bIncludeHistory) const override;
+	// End UMotionTrajectoryComponent Interface
+};
\ No newline at end of file
diff --git a/Engine/Plugins/Experimental/Animation/MotionTrajectory/Source/MotionTrajectory/Public/MotionTrajectoryLibrary.h b/Engine/Plugins/Experimental/Animation/MotionTrajectory/Source/MotionTrajectory/Public/MotionTrajectoryLibrary.h
new file mode 100644
--- /dev/null
+++ b/Engine/Plugins/Experimental/Animation/MotionTrajectory/Source/MotionTrajectory/Public/MotionTrajectoryLibrary.h
@@ -0,0 +1,53 @@
+// Copyright Epic Games, Inc. All Rights Reserved.
+
+#pragma once
+
+#include "Animation/MotionTrajectoryTypes.h"
+#include "Kismet/BlueprintFunctionLibrary.h"
+#include "UObject/ObjectMacros.h"
+
+#include "MotionTrajectoryLibrary.generated.h"
+
+USTRUCT(BlueprintType, Category="Motion Trajectory")
+struct MOTIONTRAJECTORY_API FTrajectoryDirectionClamp
+{
+	GENERATED_BODY()
+
+	// Target clamped direction for an incoming source direction
+	UPROPERTY(EditAnywhere, BlueprintReadWrite, Category="Settings")
+	FVector Direction;
+
+	// Target clamped directions will be applied to any source direction within this angle boundary
+	UPROPERTY(EditAnywhere, BlueprintReadWrite, Category="Settings", meta=(ClampMin="1.0", ClampMax="180.0"))
+	float AngleTresholdDegrees;
+};
+
+UCLASS(Category="Motion Trajectory")
+class MOTIONTRAJECTORY_API UMotionTrajectoryBlueprintLibrary : public UBlueprintFunctionLibrary
+{
+	GENERATED_BODY()
+
+public:
+
+	/**
+	* Removes the Z axis motion contribution from a trajectory range
+	*
+	* @param Trajectory			Input trajectory range
+	* @param PreserveSpeed		Input bool determining if velocity magnitude is preserved (true) or projected (false)
+	*
+	* @return					Z axis flattened, modified trajectory range
+	*/
+	UFUNCTION(BlueprintPure, Category="Motion Trajectory", meta=(BlueprintThreadSafe))
+	static FTrajectorySampleRange FlattenTrajectory2D(FTrajectorySampleRange Trajectory, bool PreserveSpeed = true);
+
+	/**
+	* Projects trajectory samples onto a defined set of allowed directions
+	*
+	* @param Trajectory			Input trajectory range
+	* @param Directions			Input direction clamping, containing angle thresholds for determining source to target direction
+	*
+	* @return					Direction clamped, modified trajectory range
+	*/
+	UFUNCTION(BlueprintPure, Category = "Motion Trajectory", meta=(BlueprintThreadSafe, AutoCreateRefTerm="Directions"))
+	static FTrajectorySampleRange ClampTrajectoryDirection(FTrajectorySampleRange Trajectory, const TArray<FTrajectoryDirectionClamp>& Directions);
+};
\ No newline at end of file
diff --git a/Engine/Plugins/Experimental/Animation/PoseSearch/Source/Editor/Private/AnimGraphNode_PoseSearchHistoryCollector.cpp b/Engine/Plugins/Experimental/Animation/PoseSearch/Source/Editor/Private/AnimGraphNode_PoseSearchHistoryCollector.cpp
--- a/Engine/Plugins/Experimental/Animation/PoseSearch/Source/Editor/Private/AnimGraphNode_PoseSearchHistoryCollector.cpp
+++ b/Engine/Plugins/Experimental/Animation/PoseSearch/Source/Editor/Private/AnimGraphNode_PoseSearchHistoryCollector.cpp
@@ -4,7 +4,6 @@
 
 #define LOCTEXT_NAMESPACE "AnimGraphNode_PoseSearchHistoryCollector"
 
-
 FLinearColor UAnimGraphNode_PoseSearchHistoryCollector::GetNodeTitleColor() const
 {
 	return FColor(86, 182, 194);
@@ -17,7 +16,7 @@ FText UAnimGraphNode_PoseSearchHistoryCollector::GetTooltipText() const
 
 FText UAnimGraphNode_PoseSearchHistoryCollector::GetNodeTitle(ENodeTitleType::Type TitleType) const
 {
-	return LOCTEXT("NodeTitle", "Match Pose");
+	return LOCTEXT("NodeTitle", "Pose History");
 }
 
 FText UAnimGraphNode_PoseSearchHistoryCollector::GetMenuCategory() const
@@ -25,5 +24,4 @@ FText UAnimGraphNode_PoseSearchHistoryCollector::GetMenuCategory() const
 	return LOCTEXT("NodeCategory", "Pose Search");
 }
 
-
 #undef LOCTEXT_NAMESPACE
diff --git a/Engine/Plugins/Experimental/Animation/PoseSearch/Source/Runtime/PoseSearch.Build.cs b/Engine/Plugins/Experimental/Animation/PoseSearch/Source/Runtime/PoseSearch.Build.cs
--- a/Engine/Plugins/Experimental/Animation/PoseSearch/Source/Runtime/PoseSearch.Build.cs
+++ b/Engine/Plugins/Experimental/Animation/PoseSearch/Source/Runtime/PoseSearch.Build.cs
@@ -33,7 +33,9 @@ public class PoseSearch : ModuleRules
 				"CoreUObject",
 				"Engine",
 				"AnimationCore",
-				"TraceLog"
+				"TraceLog",
+				"AnimGraphRuntime",
+				"DynamicPlayRate"
 			}
 		);
 
@@ -42,8 +44,8 @@ public class PoseSearch : ModuleRules
             {
 				"Core",
 				"CoreUObject",
-				"Engine",
-            }
+				"Engine"
+			}
 		);
 	}
 }
diff --git a/Engine/Plugins/Experimental/Animation/PoseSearch/Source/Runtime/Private/AnimNode_MotionMatching.cpp b/Engine/Plugins/Experimental/Animation/PoseSearch/Source/Runtime/Private/AnimNode_MotionMatching.cpp
--- a/Engine/Plugins/Experimental/Animation/PoseSearch/Source/Runtime/Private/AnimNode_MotionMatching.cpp
+++ b/Engine/Plugins/Experimental/Animation/PoseSearch/Source/Runtime/Private/AnimNode_MotionMatching.cpp
@@ -1,17 +1,18 @@
 // Copyright Epic Games, Inc. All Rights Reserved.
 
 #include "PoseSearch/AnimNode_MotionMatching.h"
-#include "PoseSearch/PoseSearch.h"
-#include "PoseSearch/PoseSearchPredictionLibrary.h"
-#include "PoseSearch/AnimNode_PoseSearchHistoryCollector.h"
+
 #include "Animation/AnimInstanceProxy.h"
 #include "Animation/AnimNode_Inertialization.h"
 #include "Animation/AnimSequence.h"
+#include "Animation/MotionTrajectoryTypes.h"
+#include "DynamicPlayRate/DynamicPlayRateLibrary.h"
+#include "PoseSearch/AnimNode_PoseSearchHistoryCollector.h"
+#include "PoseSearch/PoseSearch.h"
 #include "Trace/PoseSearchTraceLogger.h"
 
 #define LOCTEXT_NAMESPACE "AnimNode_MotionMatching"
 
-
 /////////////////////////////////////////////////////
 // FAnimNode_MotionMatching
 
@@ -21,7 +22,7 @@ void FAnimNode_MotionMatching::Initialize_AnyThread(const FAnimationInitializeCo
 
 	GetEvaluateGraphExposedInputs().Execute(Context);
 
-	InitNewDatabaseSearch();
+	MotionMatchingState.InitNewDatabaseSearch(Database, Settings.SearchThrottleTime);
 
 	Source.SetLinkNode(&SequencePlayerNode);
 	Source.Initialize(Context);
@@ -38,160 +39,41 @@ void FAnimNode_MotionMatching::UpdateAssetPlayer(const FAnimationUpdateContext&
 {
 	DECLARE_SCOPE_HIERARCHICAL_COUNTER_ANIMNODE(UpdateAssetPlayer);
 
-
 	GetEvaluateGraphExposedInputs().Execute(Context);
-	// Note: What if the input database changes? That's not being handled at all!
 
-	bool bJumpedToPose = false;
-	FPoseSearchBiasWeights QueryBiasWeights;
-	
-#if UE_POSE_SEARCH_TRACE_ENABLED
-	UE::PoseSearch::FTraceMotionMatchingState State;
-#endif
+	// Update with the sequence player's current time.
+	MotionMatchingState.AssetPlayerTime = SequencePlayerNode.GetAccumulatedTime();
 
-	const float DeltaTime = Context.GetDeltaTime();
+	// Execute core motion matching algorithm and retain across frame state
+	UpdateMotionMatchingState(Context
+		, Database
+		, Trajectory
+		, Settings
+		, MotionMatchingState
+	);
 
-	if (IsValidForSearch())
+	// If a new pose is requested, jump to the pose by updating the embedded sequence player node
+	if ((MotionMatchingState.Flags & EMotionMatchingFlags::JumpedToPose) == EMotionMatchingFlags::JumpedToPose)
 	{
-		if (PreviousDatabase != Database)
-		{
-			InitNewDatabaseSearch();
-		}
-
-		if (!ComposedQuery.IsInitializedForSchema(Database->Schema))
-		{
-			ComposedQuery.Init(Database->Schema);
-		}
-
-		const float AssetTime = SequencePlayerNode.GetAccumulatedTime();
-		const float AssetLength = SequencePlayerNode.GetCurrentAssetLength();
-
-		// Step the pose forward
-		if (DbPoseIdx != INDEX_NONE)
-		{
-			check(DbSequenceIdx != INDEX_NONE);
-
-			// Determine roughly which pose we're playing from the database based on the sequence player time
-			DbPoseIdx = Database->GetPoseIndexFromAssetTime(DbSequenceIdx, AssetTime);
-
-			// Overwrite query feature vector with the stepped pose from the database
-			if (DbPoseIdx != INDEX_NONE)
-			{
-				ComposedQuery.CopyFromSearchIndex(Database->SearchIndex, DbPoseIdx);
-			}
-		}
-
-		if (DbPoseIdx == INDEX_NONE)
-		{
-			UE::PoseSearch::IPoseHistoryProvider* PoseHistoryProvider = Context.GetMessage<UE::PoseSearch::IPoseHistoryProvider>();
-			if (PoseHistoryProvider)
-			{
-				UE::PoseSearch::FPoseHistory& History = PoseHistoryProvider->GetPoseHistory();
-				ComposedQuery.TrySetPoseFeatures(&History);
-				ComposedQuery.TrySetPastTrajectoryFeatures(&History);
-			}
-		}
-
-		// Update features in the query with the latest inputs
-		ComposeQuery(Context);
-
-		// Initialize the pose search query bias weights context
-		QueryBiasWeights.Init(BiasWeights, Database->Schema->Layout);
-
-		const FPoseSearchBiasWeightsContext BiasWeightsContext = { &QueryBiasWeights, Database };
-
-		// Determine how much the updated query vector deviates from the current pose vector
-		float CurrentDissimilarity = MAX_flt;
-		if (DbPoseIdx != INDEX_NONE)
-		{
-			CurrentDissimilarity = UE::PoseSearch::ComparePoses(Database->SearchIndex, DbPoseIdx, ComposedQuery.GetNormalizedValues(), &BiasWeightsContext);
-		}
-
-		// Search the database for the nearest match to the updated query vector
-		UE::PoseSearch::FDbSearchResult Result = UE::PoseSearch::Search(Database, ComposedQuery.GetNormalizedValues(), &BiasWeightsContext);
-		if (Result.IsValid() && (ElapsedPoseJumpTime >= SearchThrottleTime))
-		{
-			const FPoseSearchDatabaseSequence& ResultDbSequence = Database->Sequences[Result.DbSequenceIdx];
-
-			// Consider the search result better if it is more similar to the query than the current pose we're playing back from the database
-			bool bBetterPose = Result.Dissimilarity * (1.0f + (MinPercentImprovement / 100.0f)) < CurrentDissimilarity;
-
-			// We'll ignore the candidate pose if it is too near to our current pose
-			bool bNearbyPose = false;
-			if (DbSequenceIdx == Result.DbSequenceIdx)
-			{
-				bNearbyPose = FMath::Abs(AssetTime - Result.TimeOffsetSeconds) < PoseJumpThreshold;
-				if (!bNearbyPose && ResultDbSequence.bLoopAnimation)
-				{
-					bNearbyPose = FMath::Abs(AssetLength - AssetTime - Result.TimeOffsetSeconds) < PoseJumpThreshold;
-				}
-			}
-
-			// Start playback from the candidate pose if we determined it was a better option
-			if (bBetterPose && !bNearbyPose)
-			{
-				JumpToPose(Context, Result);
-				bJumpedToPose = true;
-			}
-		}
-
-		// Continue with the follow up sequence if we're finishing a one shot anim
-		if (!bJumpedToPose && SequencePlayerNode.GetSequence() && !SequencePlayerNode.GetLoopAnimation())
-		{
-			float AssetTimeAfterUpdate = AssetTime + DeltaTime;
-			if (AssetTimeAfterUpdate > AssetLength)
-			{
-				const FPoseSearchDatabaseSequence& DbSequence = Database->Sequences[DbSequenceIdx];
-				int32 FollowUpDbSequenceIdx = Database->Sequences.IndexOfByPredicate([&](const FPoseSearchDatabaseSequence& Entry){ return Entry.Sequence == DbSequence.FollowUpSequence; });
-				if (FollowUpDbSequenceIdx != INDEX_NONE)
-				{
-					const FPoseSearchDatabaseSequence& FollowUpDbSequence = Database->Sequences[FollowUpDbSequenceIdx];
-					float FollowUpAssetTime = AssetTimeAfterUpdate - AssetLength;
-					int32 FollowUpPoseIdx = Database->GetPoseIndexFromAssetTime(FollowUpDbSequenceIdx, FollowUpAssetTime);
-
-					UE::PoseSearch::FDbSearchResult FollowUpResult;
-					FollowUpResult.DbSequenceIdx = FollowUpDbSequenceIdx;
-					FollowUpResult.PoseIdx = FollowUpPoseIdx;
-					FollowUpResult.TimeOffsetSeconds = FollowUpAssetTime;
-					JumpToPose(Context, FollowUpResult);
-					bJumpedToPose = true;
-#if UE_POSE_SEARCH_TRACE_ENABLED
-					State.Flags |= UE::PoseSearch::FTraceMotionMatchingState::EFlags::FollowupAnimation;
-#endif
-				}
-			}
-		}
+		const FPoseSearchDatabaseSequence& ResultDbSequence = Database->Sequences[MotionMatchingState.DbSequenceIdx];
+		SequencePlayerNode.SetSequence(ResultDbSequence.Sequence);
+		SequencePlayerNode.SetAccumulatedTime(MotionMatchingState.AssetPlayerTime);
+		SequencePlayerNode.SetLoopAnimation(ResultDbSequence.bLoopAnimation);
+		SequencePlayerNode.SetPlayRate(1.f);
 	}
 
-	if (bJumpedToPose)
-	{
-		ElapsedPoseJumpTime = 0.0f;
-	}
-	else
-	{
-		ElapsedPoseJumpTime += DeltaTime;
-	}
-	 
-	const FPredictionSequenceState SequenceState = { SequencePlayerNode.GetSequence() , SequencePlayerNode.GetAccumulatedTime(), SequencePlayerNode.GetPlayRate(), SequencePlayerNode.GetLoopAnimation() };
-	const float PlayRate = UPoseSearchPredictionDistanceMatching::ComputePlayRate(Context, Prediction, PredictionSettings, SequenceState);
-	SequencePlayerNode.SetPlayRate(PlayRate);
+	// Optionally applying dynamic play rate adjustment to chosen sequences based on predictive motion analysis
+	const float PlayRate = DynamicPlayRateAdjustment(Context
+		, Trajectory
+		, DynamicPlayRateSettings
+		, SequencePlayerNode.GetSequence()
+		, SequencePlayerNode.GetAccumulatedTime()
+		, SequencePlayerNode.GetPlayRate()
+		, SequencePlayerNode.GetLoopAnimation()
+	);
 
+	SequencePlayerNode.SetPlayRate(PlayRate);
 	Source.Update(Context);
-	
-#if UE_POSE_SEARCH_TRACE_ENABLED
-	if (DbPoseIdx == INDEX_NONE)
-	{
-		return;
-	}
-	State.ElapsedPoseJumpTime = ElapsedPoseJumpTime;
-	// @TODO: Change this to only be the previous query, not persistently updated (i.e. if throttled)?
-	State.QueryVector = ComposedQuery.GetValues();
-	State.QueryVectorNormalized = ComposedQuery.GetNormalizedValues();
-	State.BiasWeights = QueryBiasWeights.Weights;
-	State.DbPoseIdx = DbPoseIdx;
-	State.DatabaseId = FObjectTrace::GetObjectId(Database.Get());
-	UE_TRACE_POSE_SEARCH_MOTION_MATCHING_STATE(Context, State)
-#endif
 }
 
 bool FAnimNode_MotionMatching::HasPreUpdate() const
@@ -219,12 +101,12 @@ void FAnimNode_MotionMatching::PreUpdate(const UAnimInstance* InAnimInstance)
 
 		if (bDebugDrawMatch)
 		{
-			DrawParams.PoseIdx = DbPoseIdx;
+			DrawParams.PoseIdx = MotionMatchingState.DbPoseIdx;
 		}
 
-		if (bDebugDrawGoal)
+		if (bDebugDrawQuery)
 		{
-			DrawParams.PoseVector = ComposedQuery.GetValues();
+			DrawParams.PoseVector = MotionMatchingState.ComposedQuery.GetValues();
 		}
 
 		UE::PoseSearch::Draw(DrawParams);
@@ -237,55 +119,6 @@ void FAnimNode_MotionMatching::GatherDebugData(FNodeDebugData& DebugData)
 	Source.GatherDebugData(DebugData);
 }
 
-bool FAnimNode_MotionMatching::IsValidForSearch() const
-{
-	return Database && Database->IsValidForSearch();
-}
-
-void FAnimNode_MotionMatching::ComposeQuery(const FAnimationBaseContext& Context)
-{
-	// Merge goal features into the query vector
-	if (ComposedQuery.IsCompatible(Goal))
-	{
-		ComposedQuery.MergeReplace(Goal);
-	}
-
-	ComposedQuery.Normalize(Database->SearchIndex);
-}
-
-void FAnimNode_MotionMatching::JumpToPose(const FAnimationUpdateContext& Context, UE::PoseSearch::FDbSearchResult Result)
-{
-	// Remember which pose and sequence we're playing from the database
-	DbPoseIdx = Result.PoseIdx;
-	DbSequenceIdx = Result.DbSequenceIdx;
-
-	// Immediately jump to the pose by updating the embedded sequence player node
-	const FPoseSearchDatabaseSequence& ResultDbSequence = Database->Sequences[Result.DbSequenceIdx];
-	SequencePlayerNode.SetSequence(ResultDbSequence.Sequence);
-	SequencePlayerNode.SetAccumulatedTime(Result.TimeOffsetSeconds);
-	SequencePlayerNode.SetLoopAnimation(ResultDbSequence.bLoopAnimation);
-	SequencePlayerNode.SetPlayRate(1.f);
-
-	// Use inertial blending to smooth over the transition
-	// It would be cool in the future to adjust the blend time by amount of dissimilarity, but we'll need a standardized distance metric first.
-	if (BlendTime > 0.0f)
-	{
-		UE::Anim::IInertializationRequester* InertializationRequester = Context.GetMessage<UE::Anim::IInertializationRequester>();
-		if (InertializationRequester)
-		{
-			InertializationRequester->RequestInertialization(BlendTime);
-		}
-	}
-}
-
-void FAnimNode_MotionMatching::InitNewDatabaseSearch()
-{
-	DbPoseIdx = INDEX_NONE;
-	DbSequenceIdx = INDEX_NONE;
-	ElapsedPoseJumpTime = SearchThrottleTime;
-	PreviousDatabase = Database;
-}
-
 // FAnimNode_AssetPlayerBase interface
 float FAnimNode_MotionMatching::GetAccumulatedTime() const
 {
diff --git a/Engine/Plugins/Experimental/Animation/PoseSearch/Source/Runtime/Private/AnimNode_PoseSearchHistoryCollector.cpp b/Engine/Plugins/Experimental/Animation/PoseSearch/Source/Runtime/Private/AnimNode_PoseSearchHistoryCollector.cpp
--- a/Engine/Plugins/Experimental/Animation/PoseSearch/Source/Runtime/Private/AnimNode_PoseSearchHistoryCollector.cpp
+++ b/Engine/Plugins/Experimental/Animation/PoseSearch/Source/Runtime/Private/AnimNode_PoseSearchHistoryCollector.cpp
@@ -42,7 +42,7 @@ void FAnimNode_PoseSearchHistoryCollector::Initialize_AnyThread(const FAnimation
 	Super::Initialize_AnyThread(Context);
 
 	// @@ need to do this once based on descendant node's (or input param?) search schema, not every node init
-	PoseHistory.Init(128, 4.0f);
+	PoseHistory.Init(PoseCount, PoseDuration);
 
 	UE::Anim::TScopedGraphMessage<UE::PoseSearch::Private::FPoseHistoryProvider> ScopedMessage(Context, this);
 
diff --git a/Engine/Plugins/Experimental/Animation/PoseSearch/Source/Runtime/Private/PoseSearch.cpp b/Engine/Plugins/Experimental/Animation/PoseSearch/Source/Runtime/Private/PoseSearch.cpp
--- a/Engine/Plugins/Experimental/Animation/PoseSearch/Source/Runtime/Private/PoseSearch.cpp
+++ b/Engine/Plugins/Experimental/Animation/PoseSearch/Source/Runtime/Private/PoseSearch.cpp
@@ -26,6 +26,8 @@
 
 IMPLEMENT_ANIMGRAPH_MESSAGE(UE::PoseSearch::IPoseHistoryProvider);
 
+PRAGMA_DISABLE_OPTIMIZATION
+
 namespace UE { namespace PoseSearch {
 
 DEFINE_LOG_CATEGORY(LogPoseSearch);
@@ -640,6 +642,16 @@ int32 UPoseSearchDatabase::GetPoseIndexFromAssetTime(int32 DbSequenceIdx, float
 	return INDEX_NONE;
 }
 
+float UPoseSearchDatabase::GetSequenceLength(int32 DbSequenceIdx) const
+{
+	return Sequences[DbSequenceIdx].Sequence->GetPlayLength();
+}
+
+bool UPoseSearchDatabase::DoesSequenceLoop(int32 DbSequenceIdx) const
+{
+	return Sequences[DbSequenceIdx].bLoopAnimation;
+}
+
 bool UPoseSearchDatabase::IsValidForIndexing() const
 {
 	bool bValid = Schema && Schema->IsValid() && !Sequences.IsEmpty();
@@ -869,90 +881,50 @@ bool FPoseSearchFeatureVectorBuilder::TrySetPoseFeatures(UE::PoseSearch::FPoseHi
 	return true;
 }
 
-bool FPoseSearchFeatureVectorBuilder::TrySetPastTrajectoryFeatures(UE::PoseSearch::FPoseHistory* History)
+void FPoseSearchFeatureVectorBuilder::BuildFromTrajectoryTimeBased(const FTrajectorySampleRange& Trajectory)
 {
 	check(Schema && Schema->IsValid());
-	check(History);
-
-	bool bSuccess = TrySetPastTrajectoryTimeFeatures(History);
-	bSuccess &= TrySetPastTrajectoryDistanceFeatures(History);
 
-	return bSuccess;
-}
-
-bool FPoseSearchFeatureVectorBuilder::TrySetPastTrajectoryTimeFeatures(UE::PoseSearch::FPoseHistory* History)
-{
 	FPoseSearchFeatureDesc Feature;
 	Feature.Domain = EPoseSearchFeatureDomain::Time;
 	Feature.SchemaBoneIdx = FPoseSearchFeatureDesc::TrajectoryBoneIndex;
 
-	for (int32 SchemaSubsampleIdx = 0; SchemaSubsampleIdx != Schema->TrajectorySampleTimes.Num(); ++SchemaSubsampleIdx)
+	for (int32 Idx = 0, NextIterStartIdx = 0, Num = Schema->TrajectorySampleTimes.Num(); Idx < Num; ++Idx)
 	{
-		Feature.SubsampleIdx = SchemaSubsampleIdx;
-		
-		// Stop when we've reached future samples
-		float SampleTime = Schema->TrajectorySampleTimes[SchemaSubsampleIdx];
-		if (SampleTime > 0.0f)
-		{
-			break;
-		}
+		const float SampleTime = Schema->TrajectorySampleTimes[Idx];
+		const FTrajectorySample Sample = FTrajectorySampleRange::IterSampleTrajectory(Trajectory.Samples, ETrajectorySampleDomain::Time, SampleTime, NextIterStartIdx);
 
-		float SecondsAgo = -SampleTime;
-		FTransform WorldComponentTransform;
-		if (!History->TrySampleRootTimeBased(SecondsAgo, &WorldComponentTransform))
-		{
-			return false;
-		}
+		Feature.SubsampleIdx = Idx;
 
-		FTransform WorldPrevComponentTransform;
-		if (!History->TrySampleRootTimeBased(SecondsAgo + History->GetSampleTimeInterval(), &WorldPrevComponentTransform))
-		{
-			return false;
-		}
+		Feature.Type = EPoseSearchFeatureType::LinearVelocity;
+		SetVector(Feature, Sample.LocalLinearVelocity);
 
-		SetTransform(Feature, WorldComponentTransform);
-		SetTransformVelocity(Feature, WorldComponentTransform, WorldPrevComponentTransform, History->GetSampleTimeInterval());
+		Feature.Type = EPoseSearchFeatureType::Position;
+		SetVector(Feature, Sample.Position);
 	}
-
-	return true;
 }
 
-bool FPoseSearchFeatureVectorBuilder::TrySetPastTrajectoryDistanceFeatures(UE::PoseSearch::FPoseHistory* History)
+void FPoseSearchFeatureVectorBuilder::BuildFromTrajectoryDistanceBased(const FTrajectorySampleRange& Trajectory)
 {
+	check(Schema && Schema->IsValid());
+
 	FPoseSearchFeatureDesc Feature;
 	Feature.Domain = EPoseSearchFeatureDomain::Distance;
 	Feature.SchemaBoneIdx = FPoseSearchFeatureDesc::TrajectoryBoneIndex;
 
-	for (int32 SchemaSubsampleIdx = 0; SchemaSubsampleIdx != Schema->TrajectorySampleDistances.Num(); ++SchemaSubsampleIdx)
+	for (int32 Idx = 0, NextIterStartIdx = 0, Num = Schema->TrajectorySampleDistances.Num(); Idx < Num; ++Idx)
 	{
-		Feature.SubsampleIdx = SchemaSubsampleIdx;
-		
-		// Stop when we've reached future samples
-		float SampleDistance = Schema->TrajectorySampleDistances[SchemaSubsampleIdx];
-		if (SampleDistance > 0.0f)
-		{
-			break;
-		}
+		const float SampleDistance = Schema->TrajectorySampleDistances[Idx];
+		const FTrajectorySample Sample = FTrajectorySampleRange::IterSampleTrajectory(Trajectory.Samples, ETrajectorySampleDomain::Distance, SampleDistance, NextIterStartIdx);
 
-		float SampleTime = 0.0f;
-		float PastDistance = -SampleDistance;
-		FTransform WorldComponentTransform;
-		if (!History->TrySampleRootDistanceBased(PastDistance, &WorldComponentTransform, &SampleTime))
-		{
-			return false;
-		}
-
-		FTransform WorldPrevComponentTransform;
-		if (!History->TrySampleRootTimeBased(SampleTime + History->GetSampleTimeInterval(), &WorldPrevComponentTransform))
-		{
-			return false;
-		}
+		Feature.SubsampleIdx = Idx;
 
-		SetTransform(Feature, WorldComponentTransform);
-		SetTransformVelocity(Feature, WorldComponentTransform, WorldPrevComponentTransform, History->GetSampleTimeInterval());
-	}
+		Feature.Type = EPoseSearchFeatureType::LinearVelocity;
+		SetVector(Feature, Sample.LocalLinearVelocity);
 
-	return true;
+		Feature.Type = EPoseSearchFeatureType::Position;
+		SetVector(Feature, Sample.Position);
+	} 
 }
 
 void FPoseSearchFeatureVectorBuilder::CopyFromSearchIndex(const FPoseSearchIndex& SearchIndex, int32 PoseIdx)
@@ -1026,6 +998,12 @@ void FPoseSearchFeatureVectorBuilder::Normalize(const FPoseSearchIndex& ForSearc
 	ForSearchIndex.Normalize(ValuesNormalized);
 }
 
+void FPoseSearchFeatureVectorBuilder::BuildFromTrajectory(const FTrajectorySampleRange& Trajectory)
+{
+	BuildFromTrajectoryTimeBased(Trajectory);
+	BuildFromTrajectoryDistanceBased(Trajectory);
+}
+
 namespace UE { namespace PoseSearch {
 
 //////////////////////////////////////////////////////////////////////////
@@ -1124,61 +1102,6 @@ bool FPoseHistory::TrySamplePose(float SecondsAgo, const FReferenceSkeleton& Ref
 	return bSampled;
 }
 
-bool FPoseHistory::TrySampleRootTimeBased(float SecondsAgo, FTransform* OutTransform) const
-{
-	int32 NextIdx = LowerBound(Knots.begin(), Knots.end(), SecondsAgo, TGreater<>());
-	if (NextIdx <= 0 || NextIdx >= Knots.Num())
-	{
-		return false;
-	}
-
-	int32 PrevIdx = NextIdx - 1;
-
-	const FPose& PrevPose = Poses[PrevIdx];
-	const FPose& NextPose = Poses[NextIdx];
-
-	// Compute alpha between previous and next knots
-	float Alpha = FMath::GetMappedRangeValueUnclamped(
-		FVector2D(Knots[PrevIdx], Knots[NextIdx]),
-		FVector2D(0.0f, 1.0f),
-		SecondsAgo);
-
-	FTransform RootTransform;
-	RootTransform.Blend(PrevPose.WorldComponentTransform, NextPose.WorldComponentTransform, Alpha);
-	RootTransform.SetToRelativeTransform(Poses.Last().WorldComponentTransform);
-
-	*OutTransform = RootTransform;
-	return true;
-}
-
-bool FPoseHistory::TrySampleRootDistanceBased(float SampleDistance, FTransform* OutTransform, float* OutSampleTime) const
-{
-	float SampleTime = 0.0f;
-
-	// Find the relevant distance samples and convert to time
-	double NextDistance = 0.0;
-	double PrevDistance = 0.0;
-	for (int32 NextIdx = Poses.Num() - 2; NextIdx >= 0; --NextIdx)
-	{
-		FTransform PrevTransform = Poses[NextIdx+1].WorldComponentTransform;
-		FTransform NextTransform = Poses[NextIdx].WorldComponentTransform;
-
-		NextDistance += FVector::Distance(PrevTransform.GetTranslation(), NextTransform.GetTranslation());
-		if (NextDistance >= SampleDistance)
-		{
-			float DistanceSampleAlpha = (float)FMath::GetRangePct(PrevDistance, NextDistance, (double)SampleDistance);
-			SampleTime = FMath::Lerp(Knots[NextIdx+1], Knots[NextIdx], DistanceSampleAlpha);
-			break;
-		}
-
-		PrevDistance = NextDistance;
-	}
-
-	bool bResult = TrySampleRootTimeBased(SampleTime, OutTransform);
-	*OutSampleTime = SampleTime;
-	return bResult;
-}
-
 void FPoseHistory::Update(float SecondsElapsed, const FPoseContext& PoseContext)
 {
 	// Age our elapsed times
@@ -1220,7 +1143,6 @@ void FPoseHistory::Update(float SecondsElapsed, const FPoseContext& PoseContext)
 	Knots.Last() = 0.0f;
 	FPose& CurrentPose = Poses.Last();
 	CopyCompactToSkeletonPose(PoseContext.Pose, CurrentPose.LocalTransforms);
-	CurrentPose.WorldComponentTransform = PoseContext.AnimInstanceProxy->GetComponentTransform();
 }
 
 float FPoseHistory::GetSampleTimeInterval() const
@@ -3335,4 +3257,6 @@ UE::Anim::IPoseSearchProvider::FSearchResult FModule::Search(const FAnimationBas
 
 }} // namespace UE::PoseSearch
 
+PRAGMA_ENABLE_OPTIMIZATION
+
 IMPLEMENT_MODULE(UE::PoseSearch::FModule, PoseSearch)
\ No newline at end of file
diff --git a/Engine/Plugins/Experimental/Animation/PoseSearch/Source/Runtime/Private/PoseSearchLibrary.cpp b/Engine/Plugins/Experimental/Animation/PoseSearch/Source/Runtime/Private/PoseSearchLibrary.cpp
new file mode 100644
--- /dev/null
+++ b/Engine/Plugins/Experimental/Animation/PoseSearch/Source/Runtime/Private/PoseSearchLibrary.cpp
@@ -0,0 +1,244 @@
+// Copyright Epic Games, Inc. All Rights Reserved.
+
+#include "PoseSearch/PoseSearchLibrary.h"
+
+#include "Animation/AnimNode_Inertialization.h"
+#include "Animation/AnimNode_SequencePlayer.h"
+#include "PoseSearch/AnimNode_MotionMatching.h"
+#include "Trace/PoseSearchTraceLogger.h"
+
+DEFINE_LOG_CATEGORY_STATIC(LogPoseSearchLibrary, Verbose, All);
+
+void FMotionMatchingState::InitNewDatabaseSearch(const UPoseSearchDatabase* Database, float SearchThrottleTime)
+{
+	DbPoseIdx = INDEX_NONE;
+	DbSequenceIdx = INDEX_NONE;
+	ElapsedPoseJumpTime = SearchThrottleTime;
+	AssetPlayerTime = 0.0f;
+	CurrentDatabase = Database;
+}
+
+void FMotionMatchingState::ComposeQuery(const UPoseSearchDatabase* Database, const FTrajectorySampleRange& Trajectory)
+{
+	FPoseSearchFeatureVectorBuilder Goal;
+	Goal.Init(Database->Schema);
+	Goal.BuildFromTrajectory(Trajectory);
+
+	// Merge goal features into the query vector
+	if (ComposedQuery.IsCompatible(Goal))
+	{
+		ComposedQuery.MergeReplace(Goal);
+	}
+
+	ComposedQuery.Normalize(Database->SearchIndex);
+}
+
+void FMotionMatchingState::JumpToPose(const UE::PoseSearch::FDbSearchResult& Result)
+{
+	// Remember which pose and sequence we're playing from the database
+	DbPoseIdx = Result.PoseIdx;
+	DbSequenceIdx = Result.DbSequenceIdx;
+
+	ElapsedPoseJumpTime = 0.0f;
+	AssetPlayerTime = Result.TimeOffsetSeconds;
+}
+
+static void RequestInertialBlend(const FAnimationUpdateContext& Context, float BlendTime)
+{
+	// Use inertial blending to smooth over the transition
+	// It would be cool in the future to adjust the blend time by amount of dissimilarity, but we'll need a standardized distance metric first.
+	if (BlendTime > 0.0f)
+	{
+		UE::Anim::IInertializationRequester* InertializationRequester = Context.GetMessage<UE::Anim::IInertializationRequester>();
+		if (InertializationRequester)
+		{
+			InertializationRequester->RequestInertialization(BlendTime);
+		}
+	}
+}
+
+void UpdateMotionMatchingState(const FAnimationUpdateContext& Context
+	, const UPoseSearchDatabase* Database
+	, const FTrajectorySampleRange& Trajectory
+	, const FMotionMatchingSettings& Settings
+	, FMotionMatchingState& InOutMotionMatchingState
+)
+{
+	FPoseSearchBiasWeights QueryBiasWeights;
+	InOutMotionMatchingState.Flags = EMotionMatchingFlags::None;
+	const float DeltaTime = Context.GetDeltaTime();
+
+	if (Database && Database->IsValidForSearch())
+	{
+		if (InOutMotionMatchingState.CurrentDatabase != Database)
+		{
+			InOutMotionMatchingState.InitNewDatabaseSearch(Database, Settings.SearchThrottleTime);
+		}
+
+		if (!InOutMotionMatchingState.ComposedQuery.IsInitializedForSchema(Database->Schema))
+		{
+			InOutMotionMatchingState.ComposedQuery.Init(Database->Schema);
+		}
+
+		// Step the pose forward
+		if (InOutMotionMatchingState.DbPoseIdx != INDEX_NONE)
+		{
+			check(InOutMotionMatchingState.DbSequenceIdx != INDEX_NONE);
+
+			// Determine roughly which pose we're playing from the database based on the sequence player time
+			InOutMotionMatchingState.DbPoseIdx = Database->GetPoseIndexFromAssetTime(InOutMotionMatchingState.DbSequenceIdx, InOutMotionMatchingState.AssetPlayerTime);
+
+			// Overwrite query feature vector with the stepped pose from the database
+			if (InOutMotionMatchingState.DbPoseIdx != INDEX_NONE)
+			{
+				InOutMotionMatchingState.ComposedQuery.CopyFromSearchIndex(Database->SearchIndex, InOutMotionMatchingState.DbPoseIdx);
+			}
+		}
+
+		if (InOutMotionMatchingState.DbPoseIdx == INDEX_NONE)
+		{
+			UE::PoseSearch::IPoseHistoryProvider* PoseHistoryProvider = Context.GetMessage<UE::PoseSearch::IPoseHistoryProvider>();
+			if (PoseHistoryProvider)
+			{
+				UE::PoseSearch::FPoseHistory& History = PoseHistoryProvider->GetPoseHistory();
+				InOutMotionMatchingState.ComposedQuery.TrySetPoseFeatures(&History);
+			}
+		}
+
+		// Update features in the query with the latest inputs
+		InOutMotionMatchingState.ComposeQuery(Database, Trajectory);
+
+		// Initialize the pose search query bias weights context
+		QueryBiasWeights.Init(Settings.BiasWeights, Database->Schema->Layout);
+
+		const FPoseSearchBiasWeightsContext BiasWeightsContext = { &QueryBiasWeights, Database };
+
+		// Determine how much the updated query vector deviates from the current pose vector
+		float CurrentDissimilarity = MAX_flt;
+		if (InOutMotionMatchingState.DbPoseIdx != INDEX_NONE)
+		{
+			CurrentDissimilarity = UE::PoseSearch::ComparePoses(Database->SearchIndex, InOutMotionMatchingState.DbPoseIdx, InOutMotionMatchingState.ComposedQuery.GetNormalizedValues(), &BiasWeightsContext);
+		}
+
+		// Search the database for the nearest match to the updated query vector
+		UE::PoseSearch::FDbSearchResult Result = UE::PoseSearch::Search(Database, InOutMotionMatchingState.ComposedQuery.GetNormalizedValues(), &BiasWeightsContext);
+		if (Result.IsValid() && (InOutMotionMatchingState.ElapsedPoseJumpTime >= Settings.SearchThrottleTime))
+		{
+			const FPoseSearchDatabaseSequence& ResultDbSequence = Database->Sequences[Result.DbSequenceIdx];
+
+			// Consider the search result better if it is more similar to the query than the current pose we're playing back from the database
+			bool bBetterPose = Result.Dissimilarity * (1.0f + (Settings.MinPercentImprovement / 100.0f)) < CurrentDissimilarity;
+
+			// We'll ignore the candidate pose if it is too near to our current pose
+			bool bNearbyPose = false;
+			if (InOutMotionMatchingState.DbSequenceIdx == Result.DbSequenceIdx)
+			{
+				bNearbyPose = FMath::Abs(InOutMotionMatchingState.AssetPlayerTime - Result.TimeOffsetSeconds) < Settings.PoseJumpThresholdTime;
+				if (!bNearbyPose && ResultDbSequence.bLoopAnimation)
+				{
+					bNearbyPose = FMath::Abs(Database->GetSequenceLength(InOutMotionMatchingState.DbSequenceIdx) - InOutMotionMatchingState.AssetPlayerTime - Result.TimeOffsetSeconds) < Settings.PoseJumpThresholdTime;
+				}
+			}
+
+			// Start playback from the candidate pose if we determined it was a better option
+			if (bBetterPose && !bNearbyPose)
+			{
+				InOutMotionMatchingState.JumpToPose(Result);
+				RequestInertialBlend(Context, Settings.BlendTime);
+				InOutMotionMatchingState.Flags |= EMotionMatchingFlags::JumpedToPose;
+			}
+		}
+
+		// Continue with the follow up sequence if we're finishing a one shot anim
+		if (!(InOutMotionMatchingState.Flags & EMotionMatchingFlags::JumpedToPose)
+			&& (InOutMotionMatchingState.DbSequenceIdx != INDEX_NONE) 
+			&& !Database->DoesSequenceLoop(InOutMotionMatchingState.DbSequenceIdx))
+		{
+			float AssetTimeAfterUpdate = InOutMotionMatchingState.AssetPlayerTime + DeltaTime;
+			const float AssetLength = Database->GetSequenceLength(InOutMotionMatchingState.DbSequenceIdx);
+			if (AssetTimeAfterUpdate > AssetLength)
+			{
+				const FPoseSearchDatabaseSequence& DbSequence = Database->Sequences[InOutMotionMatchingState.DbSequenceIdx];
+				int32 FollowUpDbSequenceIdx = Database->Sequences.IndexOfByPredicate([&](const FPoseSearchDatabaseSequence& Entry) { return Entry.Sequence == DbSequence.FollowUpSequence; });
+				if (FollowUpDbSequenceIdx != INDEX_NONE)
+				{
+					const FPoseSearchDatabaseSequence& FollowUpDbSequence = Database->Sequences[FollowUpDbSequenceIdx];
+					float FollowUpAssetTime = AssetTimeAfterUpdate - AssetLength;
+					int32 FollowUpPoseIdx = Database->GetPoseIndexFromAssetTime(FollowUpDbSequenceIdx, FollowUpAssetTime);
+
+					UE::PoseSearch::FDbSearchResult FollowUpResult;
+					FollowUpResult.DbSequenceIdx = FollowUpDbSequenceIdx;
+					FollowUpResult.PoseIdx = FollowUpPoseIdx;
+					FollowUpResult.TimeOffsetSeconds = FollowUpAssetTime;
+					InOutMotionMatchingState.JumpToPose(FollowUpResult);
+					RequestInertialBlend(Context, Settings.BlendTime);
+					InOutMotionMatchingState.Flags |= EMotionMatchingFlags::JumpedToPose;
+				}
+			}
+		}
+	}
+
+	if (!(InOutMotionMatchingState.Flags & EMotionMatchingFlags::JumpedToPose))
+	{
+		InOutMotionMatchingState.ElapsedPoseJumpTime += DeltaTime;
+	}
+
+#if UE_POSE_SEARCH_TRACE_ENABLED
+	UE::PoseSearch::FTraceMotionMatchingState TraceState;
+	if ((InOutMotionMatchingState.Flags & EMotionMatchingFlags::JumpedToPose) == EMotionMatchingFlags::JumpedToPose)
+	{
+		TraceState.Flags |= UE::PoseSearch::FTraceMotionMatchingState::EFlags::FollowupAnimation;
+	}
+	if (InOutMotionMatchingState.DbPoseIdx == INDEX_NONE)
+	{
+		return;
+	}
+	TraceState.ElapsedPoseJumpTime = InOutMotionMatchingState.ElapsedPoseJumpTime;
+	// @TODO: Change this to only be the previous query, not persistently updated (i.e. if throttled)?
+	TraceState.QueryVector = InOutMotionMatchingState.ComposedQuery.GetValues();
+	TraceState.QueryVectorNormalized = InOutMotionMatchingState.ComposedQuery.GetNormalizedValues();
+	TraceState.BiasWeights = QueryBiasWeights.Weights;
+	TraceState.DbPoseIdx = InOutMotionMatchingState.DbPoseIdx;
+	TraceState.DatabaseId = FObjectTrace::GetObjectId(Database);
+	UE_TRACE_POSE_SEARCH_MOTION_MATCHING_STATE(Context, TraceState)
+#endif
+}
+
+void UPoseSearchLibrary::UpdateMotionMatchingForSequencePlayer(const FAnimUpdateContext& AnimUpdateContext
+	, const FSequencePlayerReference& SequencePlayer
+	, const UPoseSearchDatabase* Database
+	, const FTrajectorySampleRange& Trajectory
+	, const FMotionMatchingSettings& Settings
+	, UPARAM(ref) FMotionMatchingState& InOutMotionMatchingState
+)
+{
+	if (const FAnimationUpdateContext* AnimationUpdateContext = AnimUpdateContext.GetContext())
+	{
+		if (FAnimNode_SequencePlayer* SequencePlayerNode = SequencePlayer.GetAnimNodePtr<FAnimNode_SequencePlayer>())
+		{
+			// Update with the sequence player's current time.
+			InOutMotionMatchingState.AssetPlayerTime = SequencePlayerNode->GetAccumulatedTime();
+
+			// Execute core motion matching algorithm and retain across frame state
+			UpdateMotionMatchingState(*AnimationUpdateContext, Database, Trajectory, Settings, InOutMotionMatchingState);
+
+			// If a new pose is requested, jump to the pose by updating the embedded sequence player node
+			if ((InOutMotionMatchingState.Flags & EMotionMatchingFlags::JumpedToPose) == EMotionMatchingFlags::JumpedToPose)
+			{
+				const FPoseSearchDatabaseSequence& ResultDbSequence = Database->Sequences[InOutMotionMatchingState.DbSequenceIdx];
+				SequencePlayerNode->SetSequence(ResultDbSequence.Sequence);
+				SequencePlayerNode->SetAccumulatedTime(InOutMotionMatchingState.AssetPlayerTime);
+				SequencePlayerNode->SetLoopAnimation(ResultDbSequence.bLoopAnimation);
+				SequencePlayerNode->SetPlayRate(1.f);
+			}
+		}
+		else
+		{
+			UE_LOG(LogPoseSearchLibrary, Warning, TEXT("UpdateMotionMatchingForSequencePlayer must be called on a Sequence Player node"));
+		}
+	}
+	else
+	{
+		UE_LOG(LogPoseSearchLibrary, Warning, TEXT("UpdateMotionMatchingForSequencePlayer called with invalid context"));
+	}
+}
\ No newline at end of file
diff --git a/Engine/Plugins/Experimental/Animation/PoseSearch/Source/Runtime/Private/Trace/PoseSearchTraceLogger.h b/Engine/Plugins/Experimental/Animation/PoseSearch/Source/Runtime/Private/Trace/PoseSearchTraceLogger.h
--- a/Engine/Plugins/Experimental/Animation/PoseSearch/Source/Runtime/Private/Trace/PoseSearchTraceLogger.h
+++ b/Engine/Plugins/Experimental/Animation/PoseSearch/Source/Runtime/Private/Trace/PoseSearchTraceLogger.h
@@ -71,12 +71,12 @@ ENUM_CLASS_FLAGS(FTraceMotionMatchingState::EFlags)
 }}
 
 #define UE_TRACE_POSE_SEARCH_MOTION_MATCHING_STATE(Context, MotionMatchingState) \
-	UE::PoseSearch::FTraceMotionMatchingState::Output(Context, State);
+	UE::PoseSearch::FTraceMotionMatchingState::Output(Context, MotionMatchingState);
 
 #else
 
 // Empty declarations if not enabled
-#define UE_TRACE_POSE_SEARCH_MOTION_MATCHING_STATE(Context, State)
+#define UE_TRACE_POSE_SEARCH_MOTION_MATCHING_STATE(Context, MotionMatchingState)
 
 #endif // UE_POSE_SEARCH_TRACE_ENABLED
 
diff --git a/Engine/Plugins/Experimental/Animation/PoseSearch/Source/Runtime/Public/PoseSearch/AnimNode_MotionMatching.h b/Engine/Plugins/Experimental/Animation/PoseSearch/Source/Runtime/Public/PoseSearch/AnimNode_MotionMatching.h
--- a/Engine/Plugins/Experimental/Animation/PoseSearch/Source/Runtime/Public/PoseSearch/AnimNode_MotionMatching.h
+++ b/Engine/Plugins/Experimental/Animation/PoseSearch/Source/Runtime/Public/PoseSearch/AnimNode_MotionMatching.h
@@ -2,10 +2,13 @@
 
 #pragma once
 
-#include "PoseSearch/PoseSearch.h"
-#include "PoseSearch/PoseSearchPredictionTypes.h"
 #include "Animation/AnimNode_AssetPlayerBase.h"
 #include "Animation/AnimNode_SequencePlayer.h"
+#include "Animation/MotionTrajectoryTypes.h"
+#include "DynamicPlayRate/DynamicPlayRateLibrary.h"
+#include "PoseSearch/PoseSearch.h"
+#include "PoseSearch/PoseSearchLibrary.h"
+
 #include "AnimNode_MotionMatching.generated.h"
 
 USTRUCT(BlueprintInternalUseOnly)
@@ -18,48 +21,27 @@ public:
 	FPoseLink Source;
 
 	// Collection of animations for motion matching
-	UPROPERTY(EditAnywhere, BlueprintReadWrite, Category=Settings, meta=(PinHiddenByDefault))
+	UPROPERTY(EditAnywhere, BlueprintReadWrite, Category=Settings, meta=(PinShownByDefault))
 	TObjectPtr<const UPoseSearchDatabase> Database = nullptr;
 
-	// Dynamic biasing weights during pose search queries
-	UPROPERTY(EditAnywhere, BlueprintReadWrite, Category = Settings, meta = (PinHiddenByDefault))
-	FPoseSearchBiasWeightParams BiasWeights;
-
-	// Motion matching goal
+	// Motion trajectory samples for pose search queries
 	UPROPERTY(EditAnywhere, BlueprintReadWrite, Category=Settings, meta=(PinShownByDefault))
-	FPoseSearchFeatureVectorBuilder Goal;
-
-	// Trajectory prediction samples
-	UPROPERTY(EditAnywhere, BlueprintReadWrite, Category = Settings, meta = (PinHiddenByDefault))
-	FPredictionTrajectoryRange Prediction;
-
-	// Configuration for prediction play rate adjustment, error tolerances, and sampling resolution
-	UPROPERTY(EditAnywhere, BlueprintReadWrite, Category = Settings, meta = (PinHiddenByDefault))
-	FPredictionTrajectorySettings PredictionSettings;
-
-	// Time in seconds to blend out to the new pose. Uses inertial blending and requires an Inertialization node after this node.
-	UPROPERTY(EditAnywhere, BlueprintReadWrite, Category=Settings, meta=(ClampMin="0"))
-	float BlendTime = 0.2f;
+	FTrajectorySampleRange Trajectory;
 
-	// Don't jump to poses that are less than this many seconds away
-	UPROPERTY(EditAnywhere, BlueprintReadWrite, Category=Settings, meta=(ClampMin="0"))
-	float PoseJumpThreshold = 4.0f;
-
-	// Minimum amount of time to wait between pose searches
-	UPROPERTY(EditAnywhere, BlueprintReadWrite, Category=Settings, meta=(ClampMin="0"))
-	float SearchThrottleTime = 0.1f;
+	// Settings for dynamic play rate adjustment on sequences chosen by motion matching
+	UPROPERTY(EditAnywhere, BlueprintReadWrite, Category=Settings, meta=(PinHiddenByDefault))
+	FDynamicPlayRateSettings DynamicPlayRateSettings;
 
-	// How much better the search result must be compared to the current pose in order to jump to it
-	// Note: This feature won't work quite as advertised until search data rescaling is implemented
-	UPROPERTY(EditAnywhere, BlueprintReadWrite, Category=Settings, meta=(ClampMin="0"))
-	float MinPercentImprovement = 0.0f;
+	// Settings for the core motion matching algorithm evaluation
+	UPROPERTY(EditAnywhere, BlueprintReadWrite, Category=Settings, meta=(PinHiddenByDefault))
+	FMotionMatchingSettings Settings;
 
 #if WITH_EDITORONLY_DATA
 	UPROPERTY(EditAnywhere, Category=Debug)
 	bool bDebugDraw = false;
 
 	UPROPERTY(EditAnywhere, Category=Debug)
-	bool bDebugDrawGoal = true;
+	bool bDebugDrawQuery = true;
 
 	UPROPERTY(EditAnywhere, Category=Debug)
 	bool bDebugDrawMatch = true;
@@ -78,26 +60,8 @@ private:
 	// Embedded sequence player node for playing animations from the motion matching database
 	FAnimNode_SequencePlayer_Standalone SequencePlayerNode;
 
-	// The current pose we're playing from the database
-	int32 DbPoseIdx = INDEX_NONE;
-
-	// The current animation we're playing from the database
-	int32 DbSequenceIdx = INDEX_NONE;
-
-	// The current query feature vector used to search the database for pose candidates
-	FPoseSearchFeatureVectorBuilder ComposedQuery;
-
-	// When the current database is different from the previous one, the search parameters are reset
-	UPROPERTY(Transient)
-	TObjectPtr<const UPoseSearchDatabase> PreviousDatabase = nullptr;
-
-	// Time since the last pose jump
-	float ElapsedPoseJumpTime = 0.0f;
-
-	bool IsValidForSearch() const;
-	void ComposeQuery(const FAnimationBaseContext& Context);
-	void JumpToPose(const FAnimationUpdateContext& Context, UE::PoseSearch::FDbSearchResult Result);
-	void InitNewDatabaseSearch();
+	// Encapsulated motion matching algorithm and internal state
+	FMotionMatchingState MotionMatchingState;
 	
 	// FAnimNode_AssetPlayerBase
 protected:
diff --git a/Engine/Plugins/Experimental/Animation/PoseSearch/Source/Runtime/Public/PoseSearch/AnimNode_PoseSearchHistoryCollector.h b/Engine/Plugins/Experimental/Animation/PoseSearch/Source/Runtime/Public/PoseSearch/AnimNode_PoseSearchHistoryCollector.h
--- a/Engine/Plugins/Experimental/Animation/PoseSearch/Source/Runtime/Public/PoseSearch/AnimNode_PoseSearchHistoryCollector.h
+++ b/Engine/Plugins/Experimental/Animation/PoseSearch/Source/Runtime/Public/PoseSearch/AnimNode_PoseSearchHistoryCollector.h
@@ -12,12 +12,19 @@ struct POSESEARCH_API FAnimNode_PoseSearchHistoryCollector : public FAnimNode_Ba
 	GENERATED_BODY()
 
 public:
-	UPROPERTY(EditAnywhere, BlueprintReadWrite, Category = Links)
+	UPROPERTY(EditAnywhere, BlueprintReadWrite, Category=Links)
 	FPoseLink Source;
 
+	// The maximum amount of poses that can be stored
+	UPROPERTY(EditAnywhere, BlueprintReadWrite, Category=Settings, meta=(ClampMin="0"))
+	int32 PoseCount = 128;
+	
+	// The time horizon for how long a pose will be stored in seconds
+	UPROPERTY(EditAnywhere, BlueprintReadWrite, Category=Settings, meta=(ClampMin="0"))
+	float PoseDuration = 4.f;
+
 public:
-	FAnimNode_PoseSearchHistoryCollector()
-	{}
+	FAnimNode_PoseSearchHistoryCollector() { }
 
 	// FAnimNode_Base interface
 	virtual void Initialize_AnyThread(const FAnimationInitializeContext& Context) override;
@@ -32,5 +39,4 @@ public:
 protected:
 
 	UE::PoseSearch::FPoseHistory PoseHistory;
-	TArray<float> Query;
 };
\ No newline at end of file
diff --git a/Engine/Plugins/Experimental/Animation/PoseSearch/Source/Runtime/Public/PoseSearch/PoseSearch.h b/Engine/Plugins/Experimental/Animation/PoseSearch/Source/Runtime/Public/PoseSearch/PoseSearch.h
--- a/Engine/Plugins/Experimental/Animation/PoseSearch/Source/Runtime/Public/PoseSearch/PoseSearch.h
+++ b/Engine/Plugins/Experimental/Animation/PoseSearch/Source/Runtime/Public/PoseSearch/PoseSearch.h
@@ -12,12 +12,12 @@
 #include "Animation/AnimMetaData.h"
 #include "Animation/AnimNodeMessages.h"
 #include "Animation/AnimationAsset.h"
+#include "Animation/MotionTrajectoryTypes.h"
 #include "BoneIndices.h"
 #include "Interfaces/Interface_BoneReferenceSkeletonProvider.h"
 
 #include "PoseSearch.generated.h"
 
-
 class UAnimSequence;
 struct FCompactPose;
 struct FPoseContext;
@@ -160,7 +160,7 @@ enum class EPoseSearchDataPreprocessor : int32
 /**
 * Specifies the format of a pose search index. At runtime, queries are built according to the schema for searching.
 */
-UCLASS(BlueprintType, Category = "Animation|PoseSearch")
+UCLASS(BlueprintType, Category = "Animation|Pose Search")
 class POSESEARCH_API UPoseSearchSchema : public UDataAsset, public IBoneReferenceSkeletonProvider
 {
 	GENERATED_BODY()
@@ -332,7 +332,7 @@ public:
 };
 
 /** Animation metadata object for indexing a single animation. */
-UCLASS(BlueprintType, Category = "Animation|PoseSearch")
+UCLASS(BlueprintType, Category = "Animation|Pose Search")
 class POSESEARCH_API UPoseSearchSequenceMetaData : public UAnimMetaData
 {
 	GENERATED_BODY()
@@ -358,7 +358,7 @@ public: // UObject
 };
 
 /** Bias weights for influencing the results of a pose search query */
-USTRUCT(BlueprintType, Category = "Animation|PoseSearch")
+USTRUCT(BlueprintType, Category = "Animation|Pose Search")
 struct POSESEARCH_API FPoseSearchBiasWeightParams
 {
 	GENERATED_BODY()
@@ -406,7 +406,7 @@ struct FPoseSearchBiasWeightsContext
 };
 
 /** Per-sequence bias weights for influencing the results of a pose search query */
-UCLASS(BlueprintType, Category = "Animation|PoseSearch")
+UCLASS(BlueprintType, Category = "Animation|Pose Search")
 class POSESEARCH_API UPoseSearchSequenceBiasWeightMetaData : public UAnimMetaData
 {
 	GENERATED_BODY()
@@ -417,7 +417,7 @@ public:
 };
 
 /** An entry in a UPoseSearchDatabase. */
-USTRUCT(BlueprintType, Category = "Animation|PoseSearch")
+USTRUCT(BlueprintType, Category = "Animation|Pose Search")
 struct POSESEARCH_API FPoseSearchDatabaseSequence
 {
 	GENERATED_BODY()
@@ -462,7 +462,7 @@ struct POSESEARCH_API FPoseSearchDatabaseSequence
 };
 
 /** A data asset for indexing a collection of animation sequences. */
-UCLASS(BlueprintType, Category = "Animation|PoseSearch")
+UCLASS(BlueprintType, Category = "Animation|Pose Search")
 class POSESEARCH_API UPoseSearchDatabase : public UDataAsset
 {
 	GENERATED_BODY()
@@ -482,6 +482,8 @@ public:
 
 	int32 FindSequenceForPose(int32 PoseIdx) const;
 	int32 GetPoseIndexFromAssetTime(int32 DbSequenceIdx, float AssetTime) const;
+	float GetSequenceLength(int32 DbSequenceIdx) const;
+	bool DoesSequenceLoop(int32 DbSequenceIdx) const;
 
 	bool IsValidForIndexing() const;
 	bool IsValidForSearch() const;
@@ -497,7 +499,7 @@ public: // UObject
 * FFeatureVectorBuilder is used to build search queries at runtime and for adding samples during search index construction.
 */
 
-USTRUCT(BlueprintType, Category = "Animation|PoseSearch")
+USTRUCT(BlueprintType, Category = "Animation|Pose Search")
 struct POSESEARCH_API FPoseSearchFeatureVectorBuilder
 {
 	GENERATED_BODY()
@@ -518,8 +520,8 @@ public:
 	void SetLinearVelocity(FPoseSearchFeatureDesc Feature, const FTransform& Transform, const FTransform& PrevTransform, float DeltaTime);
 	void SetAngularVelocity(FPoseSearchFeatureDesc Feature, const FTransform& Transform, const FTransform& PrevTransform, float DeltaTime);
 	void SetVector(FPoseSearchFeatureDesc Feature, const FVector& Vector);
+	void BuildFromTrajectory(const FTrajectorySampleRange& Trajectory);
 	bool TrySetPoseFeatures(UE::PoseSearch::FPoseHistory* History);
-	bool TrySetPastTrajectoryFeatures(UE::PoseSearch::FPoseHistory* History);
 
 	void CopyFromSearchIndex(const FPoseSearchIndex& SearchIndex, int32 PoseIdx);
 	void CopyFeature(const FPoseSearchFeatureVectorBuilder& OtherBuilder, int32 FeatureIdx);
@@ -534,8 +536,8 @@ public:
 	void Normalize(const FPoseSearchIndex& ForSearchIndex);
 
 private:
-	bool TrySetPastTrajectoryTimeFeatures(UE::PoseSearch::FPoseHistory* History);
-	bool TrySetPastTrajectoryDistanceFeatures(UE::PoseSearch::FPoseHistory* History);
+	void BuildFromTrajectoryTimeBased(const FTrajectorySampleRange& Trajectory);
+	void BuildFromTrajectoryDistanceBased(const FTrajectorySampleRange& Trajectory);
 
 	UPROPERTY(Transient)
 	const UPoseSearchSchema* Schema = nullptr;
@@ -557,8 +559,6 @@ public:
 	void Init(int32 InNumPoses, float InTimeHorizon);
 	void Init(const FPoseHistory& History);
 	bool TrySamplePose(float SecondsAgo, const FReferenceSkeleton& RefSkeleton, const TArray<FBoneIndexType>& RequiredBones);
-	bool TrySampleRootTimeBased(float SecondsAgo, FTransform* OutTransform) const;
-	bool TrySampleRootDistanceBased(float PastDistance, FTransform* OutTransform, float* OutSampleTime) const;
 
 	void Update(float SecondsElapsed, const FPoseContext& PoseContext);
 	float GetSampleTimeInterval() const;
@@ -575,7 +575,6 @@ private:
 	struct FPose
 	{
 		TArray<FTransform> LocalTransforms;
-		FTransform WorldComponentTransform;
 	};
 
 	TRingBuffer<FPose> Poses;
diff --git a/Engine/Plugins/Experimental/Animation/PoseSearch/Source/Runtime/Public/PoseSearch/PoseSearchLibrary.h b/Engine/Plugins/Experimental/Animation/PoseSearch/Source/Runtime/Public/PoseSearch/PoseSearchLibrary.h
new file mode 100644
--- /dev/null
+++ b/Engine/Plugins/Experimental/Animation/PoseSearch/Source/Runtime/Public/PoseSearch/PoseSearchLibrary.h
@@ -0,0 +1,128 @@
+// Copyright Epic Games, Inc. All Rights Reserved.
+
+#pragma once
+
+#include "AnimNodes/AnimNode_SequenceEvaluator.h"
+#include "Animation/AnimExecutionContext.h"
+#include "Animation/AnimNodeReference.h"
+#include "Kismet/BlueprintFunctionLibrary.h"
+#include "Kismet/BlueprintFunctionLibrary.h"
+#include "PoseSearch/PoseSearch.h"
+#include "SequenceEvaluatorLibrary.h"
+#include "SequencePlayerLibrary.h"
+
+#include "PoseSearchLibrary.generated.h"
+
+UENUM(BlueprintType, Category="Motion Trajectory", meta=(Bitflags, UseEnumValuesAsMaskValuesInEditor="true"))
+enum class EMotionMatchingFlags : uint8
+{
+	None = 0 UMETA(Hidden),
+	JumpedToPose = 1 << 0 // Signals that motion matching has made a significant deviation in the selected sequence/pose index
+};
+ENUM_CLASS_FLAGS(EMotionMatchingFlags);
+
+USTRUCT(BlueprintType, Category = "Animation|Pose Search")
+struct POSESEARCH_API FMotionMatchingSettings
+{
+	GENERATED_BODY()
+
+	// Dynamic biasing weights for influencing pose search queries
+	UPROPERTY(EditAnywhere, BlueprintReadWrite, Category=Settings, meta=(PinHiddenByDefault))
+	FPoseSearchBiasWeightParams BiasWeights;
+
+	// Time in seconds to blend out to the new pose. Uses inertial blending and requires an Inertialization node after this node.
+	UPROPERTY(EditAnywhere, BlueprintReadWrite, Category=Settings, meta=(ClampMin="0"))
+	float BlendTime = 0.2f;
+	
+	// Don't jump to poses that are less than this many seconds away
+	UPROPERTY(EditAnywhere, BlueprintReadWrite, Category=Settings, meta=(ClampMin="0"))
+	float PoseJumpThresholdTime = 1.f;
+
+	// Minimum amount of time to wait between pose search queries
+	UPROPERTY(EditAnywhere, BlueprintReadWrite, Category=Settings, meta=(ClampMin="0"))
+	float SearchThrottleTime = 0.1f;
+
+	// How much better the search result must be compared to the current pose in order to jump to it
+	// Note: This feature won't work quite as advertised until search data rescaling is implemented
+	UPROPERTY(EditAnywhere, BlueprintReadWrite, Category=Settings, meta=(ClampMin="0", ClampMax="100"))
+	float MinPercentImprovement = 100.f;
+};
+
+USTRUCT(BlueprintType, Category="Animation|Pose Search")
+struct POSESEARCH_API FMotionMatchingState
+{
+	GENERATED_BODY()
+
+	// Initializes the minimum required motion matching state
+	void InitNewDatabaseSearch(const UPoseSearchDatabase* Database, float SearchThrottleTime);
+
+	// Adds trajectory prediction and history information to ComposedQuery
+	void ComposeQuery(const UPoseSearchDatabase* Database, const FTrajectorySampleRange& Trajectory);
+
+	// Internally stores the 'jump' to a new pose/sequence index and asset time for evaluation
+	void JumpToPose(const UE::PoseSearch::FDbSearchResult& Result);
+
+	// The current pose we're playing from the database
+	UPROPERTY(Transient)
+	int32 DbPoseIdx = INDEX_NONE;
+
+	// The current animation we're playing from the database
+	UPROPERTY(Transient)
+	int32 DbSequenceIdx = INDEX_NONE;
+
+	// The current query feature vector used to search the database for pose candidates
+	UPROPERTY(Transient)
+	FPoseSearchFeatureVectorBuilder ComposedQuery;
+
+	// When the database changes, the search parameters are reset
+	UPROPERTY(Transient)
+	TWeakObjectPtr<const UPoseSearchDatabase> CurrentDatabase = nullptr;
+
+	// Time since the last pose jump
+	UPROPERTY(Transient)
+	float ElapsedPoseJumpTime = 0.f;
+
+	// Current time within the asset player node
+	UPROPERTY(Transient)
+	float AssetPlayerTime = 0.f;
+
+	// Evaluation flags relevant to the state of motion matching
+	UPROPERTY(VisibleAnywhere, BlueprintReadOnly, Category=State)
+	EMotionMatchingFlags Flags = EMotionMatchingFlags::None;
+};
+
+/**
+* Implementation of the core motion matching algorithm
+*
+* @param UpdateContext				Input animation update context providing access to the proxy and delta time
+* @param Database					Input collection of animations for motion matching
+* @param Trajectory					Input motion trajectory samples for pose search queries
+* @param Settings					Input motion matching algorithm configuration settings
+* @param InOutMotionMatchingState	Input/Output encapsulated motion matching algorithm and state
+* @param bOutJumpedToPose			Output motion matching decision to 'jump' to a new pose/sequence index in the database
+*/
+POSESEARCH_API void UpdateMotionMatchingState(const FAnimationUpdateContext& Context
+	, const UPoseSearchDatabase* Database
+	, const FTrajectorySampleRange& Trajectory
+	, const FMotionMatchingSettings& Settings
+	, FMotionMatchingState& InOutMotionMatchingState
+);
+
+// Exposes common pose search operations to Blueprint and native code
+UCLASS(Category="Animation|Pose Search")
+class POSESEARCH_API UPoseSearchLibrary : public UBlueprintFunctionLibrary
+{
+	GENERATED_BODY()
+
+public:
+
+	// Encapsulates and computes the motion matching algorithm for a sequence player
+	UFUNCTION(BlueprintPure, Category="Pose Search", meta=(BlueprintThreadSafe, AutoCreateRefTerm="Trajectory, Settings"))
+	static void UpdateMotionMatchingForSequencePlayer(const FAnimUpdateContext& AnimUpdateContext
+		, const FSequencePlayerReference& SequencePlayer
+		, const UPoseSearchDatabase* Database
+		, const FTrajectorySampleRange& Trajectory
+		, const FMotionMatchingSettings& Settings
+		, UPARAM(ref) FMotionMatchingState& InOutMotionMatchingState
+	);
+};
\ No newline at end of file
diff --git a/Engine/Plugins/Experimental/Animation/PoseSearch/Source/Runtime/Public/PoseSearch/PoseSearchPredictionLibrary.cpp b/Engine/Plugins/Experimental/Animation/PoseSearch/Source/Runtime/Public/PoseSearch/PoseSearchPredictionLibrary.cpp
deleted file mode 100644
diff --git a/Engine/Plugins/Experimental/Animation/PoseSearch/Source/Runtime/Public/PoseSearch/PoseSearchPredictionLibrary.h b/Engine/Plugins/Experimental/Animation/PoseSearch/Source/Runtime/Public/PoseSearch/PoseSearchPredictionLibrary.h
deleted file mode 100644
diff --git a/Engine/Plugins/Experimental/Animation/PoseSearch/Source/Runtime/Public/PoseSearch/PoseSearchPredictionTypes.cpp b/Engine/Plugins/Experimental/Animation/PoseSearch/Source/Runtime/Public/PoseSearch/PoseSearchPredictionTypes.cpp
deleted file mode 100644
diff --git a/Engine/Plugins/Experimental/Animation/PoseSearch/Source/Runtime/Public/PoseSearch/PoseSearchPredictionTypes.h b/Engine/Plugins/Experimental/Animation/PoseSearch/Source/Runtime/Public/PoseSearch/PoseSearchPredictionTypes.h
deleted file mode 100644
diff --git a/Engine/Source/Runtime/AnimGraphRuntime/Private/SequencePlayerLibrary.cpp b/Engine/Source/Runtime/AnimGraphRuntime/Private/SequencePlayerLibrary.cpp
--- a/Engine/Source/Runtime/AnimGraphRuntime/Private/SequencePlayerLibrary.cpp
+++ b/Engine/Source/Runtime/AnimGraphRuntime/Private/SequencePlayerLibrary.cpp
@@ -69,7 +69,6 @@ FSequencePlayerReference USequencePlayerLibrary::SetSequence(const FSequencePlay
 	return SequencePlayer;
 }
 
-
 FSequencePlayerReference USequencePlayerLibrary::SetSequenceWithInertialBlending(const FAnimUpdateContext& UpdateContext, const FSequencePlayerReference& SequencePlayer, UAnimSequenceBase* Sequence, float BlendTime)
 {
 	SequencePlayer.CallAnimNodeFunction<FAnimNode_SequencePlayer>(
@@ -101,4 +100,69 @@ FSequencePlayerReference USequencePlayerLibrary::SetSequenceWithInertialBlending
 		});
 
 	return SequencePlayer;
+}
+
+float USequencePlayerLibrary::GetAccumulatedTime(const FSequencePlayerReference& SequencePlayer)
+{
+	SequencePlayer.CallAnimNodeFunction<FAnimNode_SequencePlayer>(
+		TEXT("GetAccumulatedTime"),
+		[](FAnimNode_SequencePlayer& InSequencePlayer)
+		{
+			return InSequencePlayer.GetAccumulatedTime();
+		});
+
+	UE_LOG(LogSequencePlayerLibrary, Warning, TEXT("Could not get accumulated time on sequence player."));
+	return 0.f;
+}
+
+float USequencePlayerLibrary::GetStartPosition(const FSequencePlayerReference& SequencePlayer)
+{
+	SequencePlayer.CallAnimNodeFunction<FAnimNode_SequencePlayer>(
+		TEXT("GetStartPosition"),
+		[](FAnimNode_SequencePlayer& InSequencePlayer)
+		{
+			return InSequencePlayer.GetStartPosition();
+		});
+
+	UE_LOG(LogSequencePlayerLibrary, Warning, TEXT("Could not get start position on sequence player."));
+	return 0.f;
+}
+
+float USequencePlayerLibrary::GetPlayRate(const FSequencePlayerReference& SequencePlayer)
+{
+	SequencePlayer.CallAnimNodeFunction<FAnimNode_SequencePlayer>(
+		TEXT("GetPlayRate"),
+		[](FAnimNode_SequencePlayer& InSequencePlayer)
+		{
+			return InSequencePlayer.GetPlayRate();
+		});
+
+	UE_LOG(LogSequencePlayerLibrary, Warning, TEXT("Could not get play rate on sequence player."));
+	return 1.f;
+}
+
+bool USequencePlayerLibrary::GetLoopAnimation(const FSequencePlayerReference& SequencePlayer)
+{
+	SequencePlayer.CallAnimNodeFunction<FAnimNode_SequencePlayer>(
+		TEXT("GetLoopAnimation"),
+		[](FAnimNode_SequencePlayer& InSequencePlayer)
+		{
+			return InSequencePlayer.GetLoopAnimation();
+		});
+
+	UE_LOG(LogSequencePlayerLibrary, Warning, TEXT("Could not get looping state on sequence player."));
+	return false;
+}
+
+UAnimSequenceBase* USequencePlayerLibrary::GetSequence(const FSequencePlayerReference& SequencePlayer)
+{
+	SequencePlayer.CallAnimNodeFunction<FAnimNode_SequencePlayer>(
+		TEXT("GetSequence"),
+		[](FAnimNode_SequencePlayer& InSequencePlayer)
+		{
+			return InSequencePlayer.GetSequence();
+		});
+
+	UE_LOG(LogSequencePlayerLibrary, Warning, TEXT("Could not get sequence on sequence player."));
+	return nullptr;
 }
\ No newline at end of file
diff --git a/Engine/Source/Runtime/AnimGraphRuntime/Public/SequencePlayerLibrary.h b/Engine/Source/Runtime/AnimGraphRuntime/Public/SequencePlayerLibrary.h
--- a/Engine/Source/Runtime/AnimGraphRuntime/Public/SequencePlayerLibrary.h
+++ b/Engine/Source/Runtime/AnimGraphRuntime/Public/SequencePlayerLibrary.h
@@ -28,11 +28,11 @@ class ANIMGRAPHRUNTIME_API USequencePlayerLibrary : public UBlueprintFunctionLib
 
 public:
 	/** Get a sequence player context from an anim node context */
-	UFUNCTION(BlueprintCallable, Category = "Sequence Player", meta=(BlueprintThreadSafe, ExpandEnumAsExecs = "Result"))
+	UFUNCTION(BlueprintCallable, Category = "Sequence Player", meta = (BlueprintThreadSafe, ExpandEnumAsExecs = "Result"))
 	static FSequencePlayerReference ConvertToSequencePlayer(const FAnimNodeReference& Node, EAnimNodeReferenceConversionResult& Result);
 
 	/** Get a sequence player context from an anim node context (pure) */
-	UFUNCTION(BlueprintPure, Category = "Sequence Player", meta=(BlueprintThreadSafe, DisplayName = "Convert to Sequence Player"))
+	UFUNCTION(BlueprintPure, Category = "Sequence Player", meta = (BlueprintThreadSafe, DisplayName = "Convert to Sequence Player"))
 	static void ConvertToSequencePlayerPure(const FAnimNodeReference& Node, FSequencePlayerReference& SequencePlayer, bool& Result)
 	{
 		EAnimNodeReferenceConversionResult ConversionResult;
@@ -41,23 +41,42 @@ public:
 	}
 
 	/** Set the current accumulated time of the sequence player */
-	UFUNCTION(BlueprintCallable, Category = "Sequence Player", meta=(BlueprintThreadSafe))
+	UFUNCTION(BlueprintCallable, Category = "Sequence Player", meta = (BlueprintThreadSafe))
 	static FSequencePlayerReference SetAccumulatedTime(const FSequencePlayerReference& SequencePlayer, float Time);
 
 	/** Set the start position of the sequence player */
-	UFUNCTION(BlueprintCallable, Category = "Sequence Player", meta=(BlueprintThreadSafe))
+	UFUNCTION(BlueprintCallable, Category = "Sequence Player", meta = (BlueprintThreadSafe))
 	static FSequencePlayerReference SetStartPosition(const FSequencePlayerReference& SequencePlayer, float StartPosition);
 
 	/** Set the play rate of the sequence player */
-	UFUNCTION(BlueprintCallable, Category = "Sequence Player", meta=(BlueprintThreadSafe))
+	UFUNCTION(BlueprintCallable, Category = "Sequence Player", meta = (BlueprintThreadSafe))
 	static FSequencePlayerReference SetPlayRate(const FSequencePlayerReference& SequencePlayer, float PlayRate);
 
 	/** Set the current sequence of the sequence player */
-	UFUNCTION(BlueprintCallable, Category = "Sequence Player", meta=(BlueprintThreadSafe))
+	UFUNCTION(BlueprintCallable, Category = "Sequence Player", meta = (BlueprintThreadSafe))
 	static FSequencePlayerReference SetSequence(const FSequencePlayerReference& SequencePlayer, UAnimSequenceBase* Sequence);
 
 	/** Set the current sequence of the sequence player with an inertial blend time */
-	UFUNCTION(BlueprintCallable, Category = "Sequence Player", meta=(BlueprintThreadSafe))
+	UFUNCTION(BlueprintCallable, Category = "Sequence Player", meta = (BlueprintThreadSafe))
 	static FSequencePlayerReference SetSequenceWithInertialBlending(const FAnimUpdateContext& UpdateContext, const FSequencePlayerReference& SequencePlayer, UAnimSequenceBase* Sequence, float BlendTime = 0.2f);
 	
+	/** Gets the current accumulated time of the sequence player */
+	UFUNCTION(BlueprintCallable, Category = "Sequence Player", meta = (BlueprintThreadSafe))
+	static float GetAccumulatedTime(const FSequencePlayerReference& SequencePlayer);
+
+	/** Get the start position of the sequence player */
+	UFUNCTION(BlueprintCallable, Category = "Sequence Player", meta = (BlueprintThreadSafe))
+	static float GetStartPosition(const FSequencePlayerReference& SequencePlayer);
+
+	/** Get the play rate of the sequence player */
+	UFUNCTION(BlueprintCallable, Category = "Sequence Player", meta = (BlueprintThreadSafe))
+	static float GetPlayRate(const FSequencePlayerReference& SequencePlayer);
+
+	/** Get the looping state of the sequence player */
+	UFUNCTION(BlueprintCallable, Category = "Sequence Player", meta = (BlueprintThreadSafe))
+	static bool GetLoopAnimation(const FSequencePlayerReference& SequencePlayer);
+
+	/** Get the current sequence of the sequence player */
+	UFUNCTION(BlueprintCallable, Category = "Sequence Player", meta = (BlueprintThreadSafe))
+	static UAnimSequenceBase* GetSequence(const FSequencePlayerReference& SequencePlayer);
 };
\ No newline at end of file
diff --git a/Engine/Source/Runtime/Engine/Private/Animation/MotionTrajectoryTypes.cpp b/Engine/Source/Runtime/Engine/Private/Animation/MotionTrajectoryTypes.cpp
new file mode 100644
--- /dev/null
+++ b/Engine/Source/Runtime/Engine/Private/Animation/MotionTrajectoryTypes.cpp
@@ -0,0 +1,89 @@
+// Copyright Epic Games, Inc. All Rights Reserved.
+
+#include "Animation/MotionTrajectoryTypes.h"
+#include "Algo/AllOf.h"
+
+namespace
+{
+	template<class U> static inline U CubicCRSplineInterpSafe(const U& P0, const U& P1, const U& P2, const U& P3, const float I, const float A = 0.5f)
+	{
+		float D1;
+		float D2;
+		float D3;
+
+		if constexpr (TIsFloatingPoint<U>::Value)
+		{
+			D1 = FMath::Abs(P1 - P0);
+			D2 = FMath::Abs(P2 - P1);
+			D3 = FMath::Abs(P3 - P2);
+		}
+		else
+		{
+			D1 = static_cast<float>(FVector::Distance(P0, P1));
+			D2 = static_cast<float>(FVector::Distance(P2, P1));
+			D3 = static_cast<float>(FVector::Distance(P3, P2));
+		}
+
+		const float T0 = 0.f;
+		const float T1 = T0 + FMath::Pow(D1, A);
+		const float T2 = T1 + FMath::Pow(D2, A);
+		const float T3 = T2 + FMath::Pow(D3, A);
+
+		return FMath::CubicCRSplineInterpSafe(P0, P1, P2, P3, T0, T1, T2, T3, FMath::Lerp(T1, T2, I));
+	}
+}
+
+bool FTrajectorySample::IsZeroSample() const
+{
+	// AccumulatedTime is specifically omitted here to allow for the zero sample semantic across an entire trajectory range
+	return LocalLinearVelocity.IsNearlyZero()
+		&& LocalLinearAcceleration.IsNearlyZero()
+		&& Position.IsNearlyZero()
+		&& FMath::IsNearlyZero(AccumulatedDistance);
+}
+
+FTrajectorySample FTrajectorySample::Lerp(const FTrajectorySample& Sample, float Alpha) const
+{
+	FTrajectorySample Interp;
+	Interp.AccumulatedSeconds = FMath::Lerp(AccumulatedSeconds, Sample.AccumulatedSeconds, Alpha);
+	Interp.AccumulatedDistance = FMath::Lerp(AccumulatedDistance, Sample.AccumulatedDistance, Alpha);
+	Interp.Position = FMath::Lerp(Position, Sample.Position, Alpha);
+	Interp.LocalLinearVelocity = FMath::Lerp(LocalLinearVelocity, Sample.LocalLinearVelocity, Alpha);
+	Interp.LocalLinearAcceleration = FMath::Lerp(LocalLinearAcceleration, Sample.LocalLinearAcceleration, Alpha);
+	return Interp;
+}
+
+FTrajectorySample FTrajectorySample::CubicCRSplineInterp(const FTrajectorySample& PrevSample
+	, const FTrajectorySample& Sample
+	, const FTrajectorySample& NextSample
+	, float Alpha) const
+{
+	FTrajectorySample Interp;
+	Interp.AccumulatedDistance = CubicCRSplineInterpSafe(PrevSample.AccumulatedDistance, AccumulatedDistance, Sample.AccumulatedDistance, NextSample.AccumulatedDistance, Alpha);
+	Interp.AccumulatedSeconds = CubicCRSplineInterpSafe(PrevSample.AccumulatedSeconds, AccumulatedSeconds, Sample.AccumulatedSeconds, NextSample.AccumulatedSeconds, Alpha);
+	Interp.Position = CubicCRSplineInterpSafe(PrevSample.Position, Position, Sample.Position, NextSample.Position, Alpha);
+	Interp.LocalLinearVelocity = CubicCRSplineInterpSafe(PrevSample.LocalLinearVelocity, LocalLinearVelocity, Sample.LocalLinearVelocity, NextSample.LocalLinearVelocity, Alpha);
+	Interp.LocalLinearAcceleration = CubicCRSplineInterpSafe(PrevSample.LocalLinearAcceleration, LocalLinearAcceleration, Sample.LocalLinearAcceleration, NextSample.LocalLinearAcceleration, Alpha);
+	return Interp;
+}
+
+bool FTrajectorySampleRange::HasSamples() const
+{
+	return !Samples.IsEmpty();
+}
+
+bool FTrajectorySampleRange::HasOnlyZeroSamples() const
+{
+	return Algo::AllOf(Samples, [](const FTrajectorySample& Sample)
+		{
+			return Sample.IsZeroSample();
+		});
+}
+
+void FTrajectorySampleRange::RemoveHistory()
+{
+	Samples.RemoveAll([](const FTrajectorySample& Sample)
+		{
+			return Sample.AccumulatedSeconds < 0.f;
+		});
+}
\ No newline at end of file
diff --git a/Engine/Source/Runtime/Engine/Public/Animation/MotionTrajectoryTypes.h b/Engine/Source/Runtime/Engine/Public/Animation/MotionTrajectoryTypes.h
new file mode 100644
--- /dev/null
+++ b/Engine/Source/Runtime/Engine/Public/Animation/MotionTrajectoryTypes.h
@@ -0,0 +1,146 @@
+// Copyright Epic Games, Inc. All Rights Reserved.
+
+#pragma once
+
+#include "UObject/ObjectMacros.h"
+
+#include "MotionTrajectoryTypes.generated.h"
+
+// Enumeration for signaling which "Accumulated" domain values to respect when determining past and future sampling horizons
+UENUM(BlueprintType, Category="Motion Trajectory", meta=(Bitflags, UseEnumValuesAsMaskValuesInEditor="true"))
+enum class ETrajectorySampleDomain : uint8
+{
+	None = 0 UMETA(Hidden),
+	Time = 1 << 0, // Seconds
+	Distance = 1 << 1 // Centimeters (Unreal units)
+};
+ENUM_CLASS_FLAGS(ETrajectorySampleDomain);
+
+// A motion trajectory sample associated with a specific time or distance domain value
+USTRUCT(BlueprintType, Category="Motion Trajectory")
+struct ENGINE_API FTrajectorySample
+{
+	GENERATED_BODY()
+
+	// The relative accumulated time that this sample is associated with
+	// Zero value for instantaneous, negative values for the past, and positive values for the future
+	UPROPERTY(EditAnywhere, BlueprintReadWrite, Category="Motion Trajectory")
+	float AccumulatedSeconds = 0.f;
+
+	// The relative accumulated distance that this trajectory sample is associated with
+	// Zero value for instantaneous, negative values for the past, and positive values for the future
+	UPROPERTY(EditAnywhere, BlueprintReadWrite, Category="Motion Trajectory")
+	float AccumulatedDistance = 0.f;
+
+	// Position relative to the sampled in-motion object
+	UPROPERTY(EditAnywhere, BlueprintReadWrite, Category="Motion Trajectory")
+	FVector Position = FVector::ZeroVector;
+
+	// Linear velocity relative to the sampled in-motion object
+	UPROPERTY(EditAnywhere, BlueprintReadWrite, Category="Motion Trajectory")
+	FVector LocalLinearVelocity = FVector::ZeroVector;
+
+	// Linear acceleration relative to the sampled in-motion object
+	UPROPERTY(EditAnywhere, BlueprintReadWrite, Category="Motion Trajectory")
+	FVector LocalLinearAcceleration = FVector::ZeroVector;
+	
+	// Linear interpolation of all parameters of two trajectory samples
+	FTrajectorySample Lerp(const FTrajectorySample& Sample, float Alpha) const;
+
+	// Centripetal Catmull–Rom spline interpolation of all parameters of two trajectory samples
+	FTrajectorySample CubicCRSplineInterp(const FTrajectorySample& PrevSample
+		, const FTrajectorySample& Sample
+		, const FTrajectorySample& NextSample
+		, float Alpha) const;
+
+	// Determines if all sample properties are zeroed
+	bool IsZeroSample() const;
+};
+
+// A container of ordered trajectory samples and associated sampling rate
+USTRUCT(BlueprintType, Category="Motion Trajectory")
+struct ENGINE_API FTrajectorySampleRange
+{
+	GENERATED_BODY()
+
+	// Linearly ordered (Time or Distance domain) container for past, present, and future trajectory samples
+	UPROPERTY(EditAnywhere, BlueprintReadWrite, Category="Motion Trajectory")
+	TArray<FTrajectorySample> Samples;
+
+	// Per-second sample rate of the trajectory sample container
+	UPROPERTY(EditAnywhere, BlueprintReadWrite, Category="Motion Trajectory")
+	int32 SampleRate = 30;
+
+	// Constructors
+	FTrajectorySampleRange() = default;
+
+	FTrajectorySampleRange(int32 Rate) : Samples(), SampleRate(Rate) { }
+
+	// Removes history samples from trajectory (retains present and future)
+	void RemoveHistory();
+
+	// Determine if any trajectory samples are present
+	bool HasSamples() const;
+
+	// Determine if all trajectory samples are default values
+	bool HasOnlyZeroSamples() const;
+
+	// Iterator for precise subsampling of the trajectory
+	template <typename Container> static FTrajectorySample IterSampleTrajectory(const Container& Samples, ETrajectorySampleDomain DomainType, float DomainValue, int32& InitialIdx, bool bSmoothInterp = false)
+	{
+		auto SafeNextIndex = [](float Idx, float Num)
+		{
+			return FMath::Min(Idx + 1, Num - 1);
+		};
+
+		auto SafePrevIndex = [](float Idx)
+		{
+			return Idx > 0 ? Idx - 1 : 0;
+		};
+
+		// Iterative invocations of this function should use the NextIdx output parameter as a method of short circuiting a full linear traversal of the range
+		for (int32 Idx = InitialIdx, Num = Samples.Num(); Idx < Num; ++Idx)
+		{
+			const FTrajectorySample& NextSample = Samples[Idx];
+			const float NextDomainValue = DomainType == ETrajectorySampleDomain::Time ? NextSample.AccumulatedSeconds : NextSample.AccumulatedDistance;
+
+			// Continue traversing the samples until an appropriate "right-hand" side domain value is found for interpolation 
+			if (NextDomainValue < DomainValue)
+			{
+				continue;
+			}
+			else
+			{
+				// Range check to disallow a right-hand 0 index from using -1 as a left index
+				InitialIdx = SafePrevIndex(Idx);
+
+				// Find the interpolation factor of P between: [P1 .. P .. P2] by mapping the domain values to [0 ... 1]
+				// Note: NextIdx is biased to equal P1 rather than P2 to account for cases where subsequent trajectory iterations may require subsampling between [P1 ... P2] again
+				const FTrajectorySample& InitialSample = Samples[InitialIdx];
+				const float PrevDomainValue = DomainType == ETrajectorySampleDomain::Time ? InitialSample.AccumulatedSeconds : InitialSample.AccumulatedDistance;
+				const float Alpha = FMath::GetMappedRangeValueUnclamped(FVector2D(PrevDomainValue, NextDomainValue), FVector2D(0.f, 1.f), DomainValue);
+
+				FTrajectorySample InterpSample;
+				if (bSmoothInterp)
+				{
+					// Apply Centripetal Catmull–Rom spline interpolation when four valid control points can be established during sub-sampling
+					const int32 PrevIdx = SafePrevIndex(InitialIdx);
+					const int32 NextIdx = SafeNextIndex(Idx, Num);
+
+					if (PrevIdx != InitialIdx && InitialIdx != Idx && Idx != NextIdx)
+					{
+						InterpSample = InitialSample.CubicCRSplineInterp(Samples[PrevIdx], NextSample, Samples[NextIdx], Alpha);
+						return InterpSample;
+					}
+				}
+
+				// For now, we will by default defer to linear interpolation
+				InterpSample = InitialSample.Lerp(NextSample, Alpha);
+				return InterpSample;
+			}
+		}
+
+		// Sampling beyond the available range will extrapolate as the last sample
+		return !Samples.IsEmpty() ? Samples.Last() : FTrajectorySample();
+	}
+};
\ No newline at end of file
diff --git a/Engine/Source/Runtime/Experimental/Animation/DynamicPlayRate/DynamicPlayRate.Build.cs b/Engine/Source/Runtime/Experimental/Animation/DynamicPlayRate/DynamicPlayRate.Build.cs
new file mode 100644
--- /dev/null
+++ b/Engine/Source/Runtime/Experimental/Animation/DynamicPlayRate/DynamicPlayRate.Build.cs
@@ -0,0 +1,23 @@
+// Copyright Epic Games, Inc. All Rights Reserved.
+
+using UnrealBuildTool;
+
+public class DynamicPlayRate : ModuleRules
+{
+	public DynamicPlayRate(ReadOnlyTargetRules Target) : base(Target)
+	{
+		PCHUsage = ModuleRules.PCHUsageMode.UseExplicitOrSharedPCHs;
+		PrivateIncludePaths.Add("Runtime/Experimental/Animation/DynamicPlayRate/Private");
+		PublicIncludePaths.Add(ModuleDirectory + "/Public");
+
+		PublicDependencyModuleNames.AddRange(
+			new string[]
+			{
+				"Core",
+				"CoreUObject",
+                "Engine",
+				"AnimGraphRuntime"
+			}
+		);
+	}
+}
diff --git a/Engine/Source/Runtime/Experimental/Animation/DynamicPlayRate/Private/DynamicPlayRate/DynamicPlayRateLibrary.cpp b/Engine/Source/Runtime/Experimental/Animation/DynamicPlayRate/Private/DynamicPlayRate/DynamicPlayRateLibrary.cpp
new file mode 100644
--- /dev/null
+++ b/Engine/Source/Runtime/Experimental/Animation/DynamicPlayRate/Private/DynamicPlayRate/DynamicPlayRateLibrary.cpp
@@ -0,0 +1,252 @@
+// Copyright Epic Games, Inc. All Rights Reserved.
+
+#include "DynamicPlayRate/DynamicPlayRateLibrary.h"
+
+#include "Algo/MinElement.h"
+#include "Animation/AnimExecutionContext.h"
+#include "Animation/AnimInstanceProxy.h"
+#include "Animation/AnimNodeBase.h"
+#include "Animation/AnimSequence.h"
+
+DEFINE_LOG_CATEGORY_STATIC(LogDynamicPlayRateLibrary, Verbose, All);
+
+#if ENABLE_ANIM_DEBUG
+TAutoConsoleVariable<int32> CVarAnimNodeDynamicPlayRateDebug(TEXT("a.DynamicPlayRate.Debug"), 0, TEXT("Turn on debug for dynamic play rate adjustment"));
+TAutoConsoleVariable<int32> CVarAnimNodeDynamicPlayRateEnable(TEXT("a.DynamicPlayRate.Enable"), 1, TEXT("Toggle dynamic play rate adjustment"));
+#endif
+
+float FDynamicPlayRateSettings::ComputePlayRate(float PlayRate, float DeltaTime) const
+{
+	return RemappingCurve ? RemappingCurve->GetFloatValue(PlayRate) : ScaleBiasClamp.ApplyTo(PlayRate, DeltaTime);
+}
+
+float DynamicPlayRateAdjustment(const FAnimationUpdateContext& Context
+	, FTrajectorySampleRange Trajectory
+	, const FDynamicPlayRateSettings& Settings
+	, const UAnimSequenceBase* Sequence
+	, float AccumulatedTime
+	, float PlayRate
+	, bool bLooping)
+{
+	if (!Settings.bEnabled)
+	{
+		return PlayRate;
+	}
+
+	const float DeltaTime = Context.GetDeltaTime();
+
+#if ENABLE_ANIM_DEBUG
+	// Debug enable/disable toggle for play rate scaling
+	if (!CVarAnimNodeDynamicPlayRateEnable.GetValueOnAnyThread())
+	{
+		return PlayRate;
+	}
+#if WITH_EDITORONLY_DATA
+	bool bDebugDraw = Settings.bDebugDraw || CVarAnimNodeDynamicPlayRateDebug.GetValueOnAnyThread();
+#endif
+#else
+#if WITH_EDITORONLY_DATA
+	bool bDebugDraw = Settings.bDebugDraw;
+#endif
+#endif
+#if WITH_EDITORONLY_DATA
+	if (bDebugDraw)
+	{
+		Context.AnimInstanceProxy->AnimDrawDebugSphere(Context.AnimInstanceProxy->GetComponentTransform().GetLocation(), 8.f, 16, FColor::Green);
+	}
+#endif
+	// Delta time is not progressing
+	if (FMath::IsNearlyZero(DeltaTime, SMALL_NUMBER))
+	{
+		return PlayRate;
+	}
+
+	// Trajectory isn't being updated
+	if (!Trajectory.HasSamples())
+	{
+		return PlayRate;
+	}
+
+	// Currently dynamic play rate adjustmenmt only takes into consideration present and future motion samples
+	Trajectory.RemoveHistory();
+
+	// Trajectory contains only zeroed samples
+	if (Trajectory.HasOnlyZeroSamples())
+	{
+		return PlayRate;
+	}
+
+	// No sequence available to play rate scale
+	if (!Sequence)
+	{
+		return PlayRate;
+	}
+
+	// Sequence Base pointer is not an actual sequence asset
+	const UAnimSequence* InternalSequence = Cast<const UAnimSequence>(Sequence);
+	if (!ensure(InternalSequence))
+	{
+		return PlayRate;
+	}
+
+	// Find the minima trajectory trajectory velocity
+	// Approximate zeroed values may indicate the synchronization point for a stop or pivot
+	const FTrajectorySample* MinimaSample = Algo::MinElementBy(Trajectory.Samples, [](const FTrajectorySample& Sample)
+		{
+			return Sample.LocalLinearVelocity.SizeSquared();
+		});
+
+	// Given a high-resolution time step, walk the current animation sequence to find a corresponding minima root motion delta
+	// Changes in direction are considered extreme minima events, meaning no other subsequent minima is significant enough to take precedence
+	float MinimaSampleTime = FLT_MAX;
+	bool bPivotDetected = false;
+
+	check(MinimaSample);
+	check(Settings.RootMotionSampleRate > 0.f);
+
+	for (auto [Idx, RootMotionSampleStep, PreviousDirection, MinimaDisplacement, CosOfPivotAngleThreshold] = std::tuple{ 0, 1.f / Settings.RootMotionSampleRate, FVector::ZeroVector, FLT_MAX, FMath::Cos(Settings.ZeroRootMotionAngleThreshold) };; ++Idx)
+	{
+		const float SampleTime = AccumulatedTime + static_cast<float>(Idx) * RootMotionSampleStep;
+
+		if (SampleTime > InternalSequence->GetPlayLength())
+		{
+			break;
+		}
+
+		const FVector RootMotion = InternalSequence->ExtractRootMotion(SampleTime, RootMotionSampleStep, bLooping).GetTranslation();
+
+		FVector RootMotionDirection;
+		float RootMotionDisplacement;
+		RootMotion.ToDirectionAndLength(RootMotionDirection, RootMotionDisplacement);
+
+		// Found a smaller displacement in the root motion track
+		if (RootMotionDisplacement <= MinimaDisplacement)
+		{
+			MinimaDisplacement = RootMotionDisplacement;
+			MinimaSampleTime = SampleTime;
+		}
+
+		const float CosOfPotentialPivotAngle = RootMotionDirection.Dot(PreviousDirection);
+		PreviousDirection = RootMotionDirection;
+
+		// Significant changes in direction will be defined as a pivot
+		// Hack: Unfortunately in practice, we may erroneously identify animations as a pivot when in fact they have malformed root motion tracks
+		if (CosOfPotentialPivotAngle < CosOfPivotAngleThreshold && RootMotionDisplacement > Settings.ZeroRootMotionDisplacementError)
+		{
+			// Bias the minima sample time for pivots to favor the pre-pivot phase (moment prior to direction change)
+			MinimaSampleTime = SampleTime - RootMotionSampleStep;
+			bPivotDetected = true;
+			break;
+		}
+	}
+
+	// We should always be able to find a minima, however a few situations could lead to this firing
+	// 1) If the root motion sampling time step is not high enough resolution, we may not be able to sample the
+	//    track near the end of a non-looping sequence
+	// 2) We may be sampling at time 'sequence length' of a non-looping sequence
+	if (MinimaSampleTime == FLT_MAX)
+	{
+		return PlayRate;
+	}
+
+	// Extrapolate the minima forward in time to detect a potential complete loss in velocity
+	const FVector MinimaRootMotionDelta = InternalSequence->ExtractRootMotion(MinimaSampleTime, DeltaTime, bLooping).GetTranslation();
+
+	// Play rate scaling is root motion driven when a near zero root motion delta or pivot has been detected
+	// Otherwise we are locomotion driven, which is reflected in the numerator of the divisor
+	bool bRootMotionDrivenPlayRate = MinimaRootMotionDelta.IsNearlyZero(KINDA_SMALL_NUMBER) || bPivotDetected;
+#if WITH_EDITORONLY_DATA
+	FColor SynchronizationColor = FColor::Red;
+#endif
+	FVector RootMotionDelta = FVector::ZeroVector;
+
+	// This loop allows for play rate scaling to apply correction in the event that we discover the animation and trajectory trajectory minima mismatch
+	// Example: 
+	//	If both the trajectory trajectory and chosen animation are decelerating to zero, minima driven play rate scaling can be correctly applied
+	//  However if the chosen animation is not decelerating to zero -- such as Jog_Right vs Jog_Right_Stop -- then a mismatch has been detected
+	//	and locomotion/instantaneous driven play rate scaling is attempted instead
+	for (int32 MaxIterations = 1; MaxIterations >= 0; MaxIterations--)
+	{
+		// Minima driven play rate scaling synchronizes using the remaining displacement to near zero with: animation / locomotion
+		// Locomotion driven play rate scaling synchronizes using the per-frame instantaneous displacement with: locomotion / animation
+		const float SequenceDelta = bRootMotionDrivenPlayRate ? MinimaSampleTime - AccumulatedTime : DeltaTime;
+		const float TrajectoryDisplacement = bRootMotionDrivenPlayRate ? MinimaSample->AccumulatedDistance : Trajectory.Samples[0].LocalLinearVelocity.Size() * DeltaTime;
+
+		RootMotionDelta = InternalSequence->ExtractRootMotion(AccumulatedTime, SequenceDelta, bLooping).GetTranslation();
+		const float RootMotionDeltaDisplacement = RootMotionDelta.Size();
+
+		// Zero displacement is left in the animation, which may result in sliding if the trajectory minima has non-zero displacement
+		const bool bZeroRootMotion = FMath::IsNearlyZero(RootMotionDeltaDisplacement, KINDA_SMALL_NUMBER);
+
+		// Zero displacement is left in the trajectory minima, which may result in a pop or a pose break if the animation has non-zero displacement
+		const bool bZeroTrajectory = FMath::IsNearlyZero(TrajectoryDisplacement, KINDA_SMALL_NUMBER);
+
+		// Play rate scaling isn't require since no trajectory trajectory motion or root motion is present
+		if (bZeroRootMotion && bZeroTrajectory)
+		{
+			PlayRate = FMath::Clamp(PlayRate, 0.f, 1.f);
+			break;
+		}
+		// The computed minima in the root motion and trajectory mismatch, so flip the synchronization behavior as an attempt to correctly play rate scale
+		// If this fails, the algorithm will automatically fall back to a default play rate of 1.f, which may introduce sliding.
+		else if (bZeroTrajectory || bZeroRootMotion)
+		{
+			// There are cases where root motion may be available in the absence of predicted motion. We need to guarantee that the play rate 
+			// will not superfluously reach extreme values due to bRootMotionDrivenPlayRate being enabled.
+			if (bRootMotionDrivenPlayRate && FMath::IsNearlyZero(MinimaSample->AccumulatedDistance, KINDA_SMALL_NUMBER))
+			{
+				PlayRate = FMath::Clamp(PlayRate, 0.f, 1.f);
+				break;
+			}
+
+			bRootMotionDrivenPlayRate = !bRootMotionDrivenPlayRate;
+			continue;
+		}
+		// Play rate scaling succeeded
+		else
+		{
+#if WITH_EDITORONLY_DATA
+			SynchronizationColor = bRootMotionDrivenPlayRate ? FColor::Purple : FColor::Blue;
+#endif
+			PlayRate = bRootMotionDrivenPlayRate
+				? RootMotionDeltaDisplacement / TrajectoryDisplacement
+				: TrajectoryDisplacement / RootMotionDeltaDisplacement;
+
+			break;
+		}
+	}
+
+#if WITH_EDITORONLY_DATA
+	// Render the starting and ending trajectory trajectory positions for the distance matching play rate synchronization
+	if (bDebugDraw)
+	{
+		Context.AnimInstanceProxy->AnimDrawDebugSphere(Context.AnimInstanceProxy->GetComponentTransform().TransformPosition(RootMotionDelta), 8.f, 16, SynchronizationColor);
+
+		if (bRootMotionDrivenPlayRate)
+		{
+			Context.AnimInstanceProxy->AnimDrawDebugSphere(Context.AnimInstanceProxy->GetComponentTransform().TransformPosition(MinimaSample->Position), 8.f, 16, FColor::Yellow);
+		}
+	}
+#endif
+
+	// Optionally remap the computed play rate against a curve
+	return Settings.ComputePlayRate(PlayRate, DeltaTime);
+}
+
+float UDynamicPlayRateLibrary::DynamicPlayRateAdjustment(const FAnimUpdateContext& UpdateContext
+	, FTrajectorySampleRange Trajectory
+	, const FDynamicPlayRateSettings& Settings
+	, const UAnimSequenceBase* Sequence
+	, float AccumulatedTime
+	, float PlayRate
+	, bool bLooping)
+{
+	return ::DynamicPlayRateAdjustment(*(UpdateContext.GetContext())
+		, Trajectory
+		, Settings
+		, Sequence
+		, AccumulatedTime
+		, PlayRate
+		, bLooping
+	);
+}
\ No newline at end of file
diff --git a/Engine/Source/Runtime/Experimental/Animation/DynamicPlayRate/Public/DynamicPlayRate/DynamicPlayRateLibrary.h b/Engine/Source/Runtime/Experimental/Animation/DynamicPlayRate/Public/DynamicPlayRate/DynamicPlayRateLibrary.h
new file mode 100644
--- /dev/null
+++ b/Engine/Source/Runtime/Experimental/Animation/DynamicPlayRate/Public/DynamicPlayRate/DynamicPlayRateLibrary.h
@@ -0,0 +1,94 @@
+// Copyright Epic Games, Inc. All Rights Reserved.
+
+#pragma once
+
+#include "Animation/AnimExecutionContext.h"
+#include "Animation/AnimSequenceBase.h"
+#include "Animation/InputScaleBias.h"
+#include "Animation/MotionTrajectoryTypes.h"
+#include "Curves/CurveFloat.h"
+#include "UObject/ObjectMacros.h"
+
+#include "DynamicPlayRateLibrary.generated.h"
+
+USTRUCT(BlueprintType, Category="Dynamic Play Rate")
+struct DYNAMICPLAYRATE_API FDynamicPlayRateSettings
+{
+	GENERATED_BODY()
+
+	// Enable/Disable dynamic play rate adjusment
+	UPROPERTY(EditAnywhere, BlueprintReadWrite, Category="Settings", meta=(PinShownByDefault))
+	bool bEnabled = false;
+
+	// Optional scaling, biasing, and clamping controls for play rate adjustment
+	UPROPERTY(EditAnywhere, BlueprintReadWrite, Category="Settings")
+	FInputScaleBiasClamp ScaleBiasClamp;
+
+	// Optional play rate remapping curve (X axis: Source play rate, Y axis: Target play rate)
+	UPROPERTY(EditAnywhere, BlueprintReadWrite, Category="Settings")
+	TObjectPtr<const UCurveFloat> RemappingCurve = nullptr;
+
+	// Root motion time step (per second) used for analyzing and determining future zero velocities
+	UPROPERTY(EditAnywhere, BlueprintReadWrite, Category="Settings", meta=(ClampMin="15.0", ClampMax="120.0"))
+	float RootMotionSampleRate = 60.f;
+
+	// Root motion angle threshold for determining significant changes in direction (ex: Identifying pivots)
+	UPROPERTY(EditAnywhere, BlueprintReadWrite, Category="Settings", meta=(ClampMin="0.0", ClampMax="180.0"))
+	float ZeroRootMotionAngleThreshold = 90.f;
+
+	// Root motion displacement error tolerance, for identifying false positive associated with ZeroRootMotionAngleThreshold
+	UPROPERTY(EditAnywhere, BlueprintReadWrite, Category="Settings")
+	float ZeroRootMotionDisplacementError = 0.0004f;
+
+	// Computes a new play rate by first attempting the remapping curve and then defaulting to the scale/bias/clamp control adjustment
+	float ComputePlayRate(float PlayRate, float DeltaTime) const;
+
+#if WITH_EDITORONLY_DATA
+	// Display in-game world markers visualizing the most significant motion values in consideration for play rate adjustment
+	UPROPERTY(EditAnywhere, Category=Debug)
+	bool bDebugDraw = false;
+#endif
+};
+
+/**
+* Computes a dynamically adjusted play rate value for a given playing sequence asset dependent on:
+*	1. Analysis of the root motion vs prediction motion contribution ratio
+*	2. Determination of significant events such as stopping, pivoting, and motion discontinuity
+*
+* @param UpdateContext	Input animation update context providing access to the proxy and delta time
+* @param Trajectory		Input trajectory samples used for predictive motion analysis against the chosen sequence root motion track
+* @param Settings		Input dynamic play rate adjustment settings for specifying error tolerances, sample rate, and additional play rate behavior
+* @param Sequence		Input playing sequence asset
+* @param PlayRate		Input sequence play rate, prior to adjustment
+* @param bLooping		Input sequence looping behavior
+*
+* @return				Adjusted play rate value
+*/
+DYNAMICPLAYRATE_API float DynamicPlayRateAdjustment(const FAnimationUpdateContext& Context
+	, FTrajectorySampleRange Trajectory
+	, const FDynamicPlayRateSettings& Settings
+	, const UAnimSequenceBase* Sequence
+	, float AccumulatedTime
+	, float PlayRate
+	, bool bLooping
+);
+
+UCLASS(Category="Dynamic Play Rate")
+class DYNAMICPLAYRATE_API UDynamicPlayRateLibrary : public UBlueprintFunctionLibrary
+{
+	GENERATED_BODY()
+
+public:
+
+	// Computes a dynamically adjusted play rate value for a given playing sequence asset
+	UFUNCTION(BlueprintPure, Category="Motion Trajectory", meta=(BlueprintThreadSafe))
+	static float DynamicPlayRateAdjustment(const FAnimUpdateContext& Context
+	    , FTrajectorySampleRange Trajectory
+		, const FDynamicPlayRateSettings& Settings
+		, const UAnimSequenceBase* Sequence
+		, float AccumulatedTime
+		, float PlayRate
+		, bool bLooping
+	);
+};
+
