PoseSearch: New weighted metric for pose similarity

Scoring weights can now be specified at several levels of abstraction.

Channel: Pose weight vs trajectory weight.
Horizon: History weight vs prediction weight.
Type: Weights per feature types, such as positions and velocities.

Weights are normalized at each of these levels. So ultimately the weights that
make up all channels of a weights group will sum to 1 (or zero if entirely disabled).

Several new concepts were introduced for this system that aren't fully
formalized. Channels aren't a concrete data type, but I'm imagining we'll head
in that direction where a Schema is made up of Channels which themselves contain
sampling paramters like time/distance offsets and what to sample. The weights
system sets itself up in such a way to effectively pretend channels exist.

Similarly, horizons aren't an explicit data type either and the weights system
pretends they sort of are for presenting this high level concept of past vs future
weight to the user.

I have removed the previous weighting system including per sequence weights.
This feature will be replaced by a grouping system for databases in the future.
Groups will be collections of sequences and a set of weights for that group.
This is why the WeightsContext contains an array of 1 group of weights-- it's
setup that way to vector the implementation toward groups.

Speaking of the WeightsContext, the motion matching node holds onto one and
updates it every frame. This is in order to support changing weights at
runtime. There are a set of dynamic weight parameters that can be used to scale
the weights of all weight groups within a database. The WeightsContext tracks
the prior dynamic weight params and knows if the weights buffer needs to be
recomputed or not. So most of the time, updating the WeightsContext is a no-op.
In the future we'll need to add support for interpolating changes to dynamic
weights to avoid abrupt scoring changes.

diff --git a/Engine/Plugins/Experimental/Animation/PoseSearch/Source/Editor/Private/PoseSearchDebugger.cpp b/Engine/Plugins/Experimental/Animation/PoseSearch/Source/Editor/Private/PoseSearchDebugger.cpp
--- a/Engine/Plugins/Experimental/Animation/PoseSearch/Source/Editor/Private/PoseSearchDebugger.cpp
+++ b/Engine/Plugins/Experimental/Animation/PoseSearch/Source/Editor/Private/PoseSearchDebugger.cpp
@@ -33,8 +33,6 @@ void FPoseSearchDebuggerFeatureReflection::EmptyAll()
 
 namespace UE { namespace PoseSearch {
 
-constexpr uint8 DrawDebugDepthPriority = SDPG_Foreground + 2;
-
 class FDebuggerDatabaseRowData : public TSharedFromThis<FDebuggerDatabaseRowData>
 {
 public:
@@ -521,64 +519,30 @@ void SDebuggerDatabaseView::UpdateRows(const FTraceMotionMatchingStateMessage& S
 		CreateRows(Database);
 	}
 	check(ActiveView.Rows.Num() == 1);
-
-	// Active bias weights pulled from the MM node
-	FPoseSearchBiasWeights BiasWeights;
-    BiasWeights.Weights = State.BiasWeights;
-
-	const FPoseSearchFeatureVectorLayout& Layout = Database.Schema->Layout;
-	const FPoseSearchIndex& SearchIndex = Database.SearchIndex;
-
-	// Reverse engineer weight params from weights array
-	auto ExtractWeight = [](
-		const FPoseSearchFeatureVectorLayout& InLayout,
-		const FPoseSearchBiasWeights& InBiasWeights,
-		const EPoseSearchFeatureType InFeatureType,
-		const bool bTrajectory
-		) -> float
-		{
-			int32 FeatureIdx = INDEX_NONE;
-			if (InLayout.EnumerateFeature(InFeatureType, bTrajectory, FeatureIdx))
-			{
-				const FPoseSearchFeatureDesc& Feature = InLayout.Features[FeatureIdx];
-            	// Return first weight found associated with the feature
-				// as the same weight is applied to all features in the buffer (for now?)
-            	return InBiasWeights.Weights[Feature.ValueOffset];	
-			}
-			return 0.0f;
-		};
-
-	// @TODO: Compute alternate scores based on column visibility in view options
-	// Zeroed trajectory for pose score exclusively
-	FPoseSearchBiasWeightParams Params;
-	Params.PosePositionWeight = ExtractWeight(Layout, BiasWeights, EPoseSearchFeatureType::Position, false);
-	Params.PoseLinearVelocityWeight = ExtractWeight(Layout, BiasWeights, EPoseSearchFeatureType::LinearVelocity, false);
-	Params.TrajectoryPositionWeight = 0.0f;
-	Params.TrajectoryLinearVelocityWeight = 0.0f;
-	
-	FPoseSearchBiasWeights PoseBiasWeights;
-	PoseBiasWeights.Init(Params, Layout);
-
-	// Zeroed pose for trajectory score
-	Params.PosePositionWeight = 0.0f;
-	Params.PoseLinearVelocityWeight = 0.0f;
-	Params.TrajectoryPositionWeight = ExtractWeight(Layout, BiasWeights, EPoseSearchFeatureType::Position, true);
-	Params.TrajectoryLinearVelocityWeight = ExtractWeight(Layout, BiasWeights, EPoseSearchFeatureType::LinearVelocity, true);
 	
-	FPoseSearchBiasWeights TrajectoryBiasWeights;
-	TrajectoryBiasWeights.Init(Params, Layout);
-	
-	const FPoseSearchBiasWeightsContext BiasWeightsContext { &BiasWeights, &Database };
-	const FPoseSearchBiasWeightsContext PoseBiasWeightsContext { &PoseBiasWeights, &Database };
-	const FPoseSearchBiasWeightsContext TrajectoryBiasWeightsContext { &TrajectoryBiasWeights, &Database };
+	FPoseSearchDynamicWeightParams StateDynamicWeightParams;
+	FPoseSearchWeightsContext StateWeights;
+	StateWeights.Update(StateDynamicWeightParams, &Database);
+
+	FPoseSearchWeightsContext PoseWeights;
+	FPoseSearchDynamicWeightParams WeightParams = StateDynamicWeightParams;
+	WeightParams.TrajectoryDynamicWeights.ChannelWeightScale = 0.0f;
+	PoseWeights.Update(WeightParams, &Database);
+
+	FPoseSearchWeightsContext TrajectoryWeights;
+	WeightParams = StateDynamicWeightParams;
+	WeightParams.PoseDynamicWeights.ChannelWeightScale = 0.0f;
+	TrajectoryWeights.Update(WeightParams, &Database);
 	
+	const FPoseSearchIndex& SearchIndex = Database.SearchIndex;
+
 	for(const TSharedRef<FDebuggerDatabaseRowData>& Row : DatabaseView.Rows)
 	{
 		const int32 PoseIdx = Row->PoseIdx;
 		
-		Row->Score = ComparePoses(SearchIndex, PoseIdx, State.QueryVectorNormalized, &BiasWeightsContext);
-		Row->PoseScore = ComparePoses(SearchIndex, PoseIdx, State.QueryVectorNormalized, &PoseBiasWeightsContext);
-		Row->TrajectoryScore = ComparePoses(SearchIndex, PoseIdx, State.QueryVectorNormalized, &TrajectoryBiasWeightsContext);
+		Row->Score = ComparePoses(SearchIndex, PoseIdx, State.QueryVectorNormalized, &StateWeights);
+		Row->PoseScore = ComparePoses(SearchIndex, PoseIdx, State.QueryVectorNormalized, &PoseWeights);
+		Row->TrajectoryScore = ComparePoses(SearchIndex, PoseIdx, State.QueryVectorNormalized, &TrajectoryWeights);
 
 		// If we are on the active pose for the frame
 		if (PoseIdx == State.DbPoseIdx)
@@ -1641,4 +1605,4 @@ FName FDebuggerViewCreator::GetName() const
 	return Name;
 }
 
-}}
+}} // namespace UE::PoseSearch
diff --git a/Engine/Plugins/Experimental/Animation/PoseSearch/Source/Runtime/Private/AnimNode_PoseSearchHistoryCollector.cpp b/Engine/Plugins/Experimental/Animation/PoseSearch/Source/Runtime/Private/AnimNode_PoseSearchHistoryCollector.cpp
--- a/Engine/Plugins/Experimental/Animation/PoseSearch/Source/Runtime/Private/AnimNode_PoseSearchHistoryCollector.cpp
+++ b/Engine/Plugins/Experimental/Animation/PoseSearch/Source/Runtime/Private/AnimNode_PoseSearchHistoryCollector.cpp
@@ -49,6 +49,13 @@ void FAnimNode_PoseSearchHistoryCollector::Initialize_AnyThread(const FAnimation
 	Source.Initialize(Context);
 }
 
+void FAnimNode_PoseSearchHistoryCollector::CacheBones_AnyThread(const FAnimationCacheBonesContext& Context)
+{
+	DECLARE_SCOPE_HIERARCHICAL_COUNTER_ANIMNODE(CacheBones_AnyThread)
+
+	Source.CacheBones(Context);
+}
+
 void FAnimNode_PoseSearchHistoryCollector::Evaluate_AnyThread(FPoseContext& Output)
 {
 	DECLARE_SCOPE_HIERARCHICAL_COUNTER_ANIMNODE(Evaluate_AnyThread);
diff --git a/Engine/Plugins/Experimental/Animation/PoseSearch/Source/Runtime/Private/PoseSearch.cpp b/Engine/Plugins/Experimental/Animation/PoseSearch/Source/Runtime/Private/PoseSearch.cpp
--- a/Engine/Plugins/Experimental/Animation/PoseSearch/Source/Runtime/Private/PoseSearch.cpp
+++ b/Engine/Plugins/Experimental/Animation/PoseSearch/Source/Runtime/Private/PoseSearch.cpp
@@ -3,6 +3,7 @@
 #include "PoseSearch/PoseSearch.h"
 #include "PoseSearchEigenHelper.h"
 
+#include "Algo/BinarySearch.h"
 #include "Async/ParallelFor.h"
 #include "Templates/IdentityFunctor.h"
 #include "Templates/Invoke.h"
@@ -23,10 +24,10 @@
 #include "BonePose.h"
 #include "Trace/PoseSearchTraceLogger.h"
 #include "UObject/ObjectSaveContext.h"
+#include "Misc/MemStack.h"
 
 IMPLEMENT_ANIMGRAPH_MESSAGE(UE::PoseSearch::IPoseHistoryProvider);
 
-PRAGMA_DISABLE_OPTIMIZATION
 
 namespace UE { namespace PoseSearch {
 
@@ -61,15 +62,28 @@ static FFloatInterval GetEffectiveSamplingRange(const UAnimSequenceBase* Sequenc
 
 static inline float CompareFeatureVectors(int32 NumValues, const float* A, const float* B, const float* Weights)
 {
-	float Dissimilarity = 0.f;
+	double Dissimilarity = 0.f;
 
 	for (int32 ValueIdx = 0; ValueIdx != NumValues; ++ValueIdx)
 	{
-		const float Diff = Weights[ValueIdx] * (A[ValueIdx] - B[ValueIdx]);
+		const float Diff = A[ValueIdx] - B[ValueIdx];
+		Dissimilarity += Weights[ValueIdx] * (Diff * Diff);
+	}
+
+	return (float)Dissimilarity;
+}
+
+static inline float CompareFeatureVectors(int32 NumValues, const float* A, const float* B)
+{
+	double Dissimilarity = 0.f;
+
+	for (int32 ValueIdx = 0; ValueIdx != NumValues; ++ValueIdx)
+	{
+		const float Diff = A[ValueIdx] - B[ValueIdx];
 		Dissimilarity += Diff * Diff;
 	}
 
-	return Dissimilarity;
+	return (float)Dissimilarity;
 }
 
 FLinearColor GetColorForFeature(FPoseSearchFeatureDesc Feature, const FPoseSearchFeatureVectorLayout* Layout)
@@ -140,6 +154,12 @@ FORCEINLINE auto LowerBound(IteratorType First, IteratorType Last, const ValueTy
 }
 
 
+// Stopgap channel indices since schemas don't yet support explicit channels of data
+constexpr int32 ChannelIdxPose = 0;
+constexpr int32 ChannelIdxTrajectoryTime = 1;
+constexpr int32 ChannelIdxTrajectoryDistance = 2;
+
+
 //////////////////////////////////////////////////////////////////////////
 // FFeatureTypeTraits
 
@@ -188,14 +208,6 @@ bool FPoseSearchFeatureDesc::operator==(const FPoseSearchFeatureDesc& Other) con
 		(Domain == Other.Domain);
 }
 
-bool FPoseSearchFeatureDesc::IsSubsampleOfSameFeature(const FPoseSearchFeatureDesc& Other) const
-{
-	return
-		(SchemaBoneIdx == Other.SchemaBoneIdx) &&
-		(Type == Other.Type) &&
-		(Domain == Other.Domain);
-}
-
 
 //////////////////////////////////////////////////////////////////////////
 // FPoseSearchFeatureVectorLayout
@@ -210,25 +222,6 @@ void FPoseSearchFeatureVectorLayout::Init()
 
 		uint32 FeatureNumFloats = UE::PoseSearch::GetFeatureTypeTraits(Feature.Type).NumFloats;
 		FloatCount += FeatureNumFloats;
-
-		if (Feature.SchemaBoneIdx == FPoseSearchFeatureDesc::TrajectoryBoneIndex)
-		{
-			if (FirstTrajectoryValueOffset == -1)
-			{
-				FirstTrajectoryValueOffset = Feature.ValueOffset;
-			}
-
-			NumTrajectoryValues += FeatureNumFloats;
-		}
-		else
-		{
-			if (FirstPoseValueOffset == -1)
-			{
-				FirstPoseValueOffset = Feature.ValueOffset;
-			}
-
-			NumPoseValues += FeatureNumFloats;
-		}
 	}
 
 	NumFloats = FloatCount;
@@ -238,10 +231,6 @@ void FPoseSearchFeatureVectorLayout::Reset()
 {
 	Features.Reset();
 	NumFloats = 0;
-	NumTrajectoryValues = 0;
-	NumPoseValues = 0;
-	FirstTrajectoryValueOffset = -1;
-	FirstPoseValueOffset = -1;
 }
 
 bool FPoseSearchFeatureVectorLayout::IsValid(int32 MaxNumBones) const
@@ -262,18 +251,45 @@ bool FPoseSearchFeatureVectorLayout::IsValid(int32 MaxNumBones) const
 	return true;
 }
 
-bool FPoseSearchFeatureVectorLayout::EnumerateFeature(EPoseSearchFeatureType FeatureType, bool bTrajectory, int32& OutFeatureIdx) const
+bool FPoseSearchFeatureVectorLayout::EnumerateBy(int32 ChannelIdx, EPoseSearchFeatureType Type, int32& InOutFeatureIdx) const
 {
-	// This function behaves similar to a generator
-	// OutFeatureIdx represents the 'next' Pose Feature Description that matches the inner criteria
-	// OutFeatureIdx can then be used again as a starting index to begin a subsequent search
-	for (int32 Idx = OutFeatureIdx + 1, Size = Features.Num(); Idx < Size; ++Idx)
+	auto IsChannelMatch = [](int32 ChannelIdx, const FPoseSearchFeatureDesc& Feature) -> bool
 	{
-		// A trajectory feature match will result when bTrajectory = True and SchemaBoneIdx = -1
-		// A pose feature match will result when bTrajectory = False and SchemaBoneIdx != -1
-		if (Features[Idx].Type == FeatureType && (bTrajectory == (Features[Idx].SchemaBoneIdx == FPoseSearchFeatureDesc::TrajectoryBoneIndex)))
+		bool bChannelMatch = true;
+		if (ChannelIdx == UE::PoseSearch::ChannelIdxPose)
+		{
+			bChannelMatch = Feature.SchemaBoneIdx != FPoseSearchFeatureDesc::TrajectoryBoneIndex;
+		}
+		else if (ChannelIdx == UE::PoseSearch::ChannelIdxTrajectoryTime)
+		{
+			bChannelMatch = (Feature.SchemaBoneIdx == FPoseSearchFeatureDesc::TrajectoryBoneIndex) && (Feature.Domain == EPoseSearchFeatureDomain::Time);
+		}
+		else if (ChannelIdx == UE::PoseSearch::ChannelIdxTrajectoryDistance)
+		{
+			bChannelMatch = (Feature.SchemaBoneIdx == FPoseSearchFeatureDesc::TrajectoryBoneIndex) && (Feature.Domain == EPoseSearchFeatureDomain::Distance);
+		}
+		return bChannelMatch;
+	};
+
+	auto IsTypeMatch = [](EPoseSearchFeatureType Type, const FPoseSearchFeatureDesc& Feature) -> bool
+	{
+		bool bTypeMatch = true;
+		if (Type != EPoseSearchFeatureType::Invalid)
+		{
+			bTypeMatch = Feature.Type == Type;
+		}
+		return bTypeMatch;
+	};
+
+	for (int32 Size = Features.Num(); ++InOutFeatureIdx < Size; )
+	{
+		const FPoseSearchFeatureDesc& Feature = Features[InOutFeatureIdx];
+
+		bool bChannelMatch = IsChannelMatch(ChannelIdx, Feature);
+		bool bTypeMatch = IsTypeMatch(Type, Feature);
+
+		if (bChannelMatch && bTypeMatch)
 		{
-			OutFeatureIdx = Idx;
 			return true;
 		}
 	}
@@ -348,6 +364,24 @@ float UPoseSearchSchema::GetTrajectoryPastDistanceHorizon () const
 	return TrajectorySampleDistances.Num() ? TrajectorySampleDistances[0] : 1.0f;
 }
 
+TArrayView<const float> UPoseSearchSchema::GetChannelSampleOffsets (int32 ChannelIdx) const
+{
+	if (ChannelIdx == UE::PoseSearch::ChannelIdxPose)
+	{
+		return PoseSampleTimes;
+	}
+	else if (ChannelIdx == UE::PoseSearch::ChannelIdxTrajectoryTime)
+	{
+		return TrajectorySampleTimes;
+	}
+	else if (ChannelIdx == UE::PoseSearch::ChannelIdxTrajectoryDistance)
+	{
+		return TrajectorySampleDistances;
+	}
+
+	return {};
+}
+
 void UPoseSearchSchema::GenerateLayout()
 {
 	Layout.Reset();
@@ -463,36 +497,333 @@ void UPoseSearchSchema::ResolveBoneReferences()
 
 
 //////////////////////////////////////////////////////////////////////////
-// FPoseSearchBiasWeights
+// FPoseSearchChannelWeightParams
 
-void FPoseSearchBiasWeights::Init(const FPoseSearchBiasWeightParams& WeightParams, const FPoseSearchFeatureVectorLayout& Layout)
+FPoseSearchChannelWeightParams::FPoseSearchChannelWeightParams()
 {
-	// Initialize all weights to a default value of 1, and subsequently override all bound weights to their assigned value
-	Weights.Init(1.f, Layout.NumFloats);
-	BindSemanticWeight(WeightParams.TrajectoryPositionWeight, Layout, EPoseSearchFeatureType::Position, true);
-	BindSemanticWeight(WeightParams.TrajectoryLinearVelocityWeight, Layout, EPoseSearchFeatureType::LinearVelocity, true);
-	BindSemanticWeight(WeightParams.PosePositionWeight, Layout, EPoseSearchFeatureType::Position, false);
-	BindSemanticWeight(WeightParams.PoseLinearVelocityWeight, Layout, EPoseSearchFeatureType::LinearVelocity, false);
+	for (int32 Type = 0; Type != (int32)EPoseSearchFeatureType::Num; ++Type)
+	{
+		TypeWeights.Add((EPoseSearchFeatureType)Type, 1.0f);
+	}
+}
+
+
+//////////////////////////////////////////////////////////////////////////
+// FPoseSearchWeights
+
+void FPoseSearchWeights::Init(const FPoseSearchWeightParams& WeightParams, const UPoseSearchSchema* Schema, const FPoseSearchDynamicWeightParams& DynamicWeightParams)
+{
+	using namespace UE::PoseSearch;
+	using namespace Eigen;
+
+	// Convenience enum for indexing by horizon
+	enum EHorizon : int
+	{
+		History,
+		Prediction,
+		Num
+	};
+
+	// Initialize weights
+	Weights.Reset();
+	Weights.SetNumZeroed(Schema->Layout.NumFloats);
+
+	// Completely disable weights if requested
+	if (DynamicWeightParams.bDebugDisableWeights)
+	{
+		for (float& Weight: Weights)
+		{
+			Weight = 1.0f;
+		}
+		return;
+	}
+
+	
+	FMemMark MemMark(FMemStack::Get());
+
+
+	// Setup channel indexable weight params
+	constexpr int ChannelNum = 3;
+	const FPoseSearchChannelWeightParams* ChannelWeightParams[ChannelNum];
+	const FPoseSearchChannelDynamicWeightParams* ChannelDynamicWeightParams[ChannelNum];
+
+	ChannelWeightParams[ChannelIdxPose] = &WeightParams.PoseWeight;
+	ChannelWeightParams[ChannelIdxTrajectoryTime] = &WeightParams.TrajectoryWeight;
+	ChannelWeightParams[ChannelIdxTrajectoryDistance] = &WeightParams.TrajectoryWeight;
+
+	ChannelDynamicWeightParams[ChannelIdxPose] = &DynamicWeightParams.PoseDynamicWeights;
+	ChannelDynamicWeightParams[ChannelIdxTrajectoryTime] = &DynamicWeightParams.TrajectoryDynamicWeights;
+	ChannelDynamicWeightParams[ChannelIdxTrajectoryDistance] = &DynamicWeightParams.TrajectoryDynamicWeights;
+
+
+	// Normalize channel weights
+	Array<float, ChannelNum, 1> NormalizedChannelWeights;
+	for (int ChannelIdx = 0; ChannelIdx != ChannelNum; ++ChannelIdx)
+	{
+		NormalizedChannelWeights[ChannelIdx] = ChannelWeightParams[ChannelIdx]->ChannelWeight * ChannelDynamicWeightParams[ChannelIdx]->ChannelWeightScale;
+
+		// Zero the channel weight if there are no features in this channel
+		int32 FeatureIdx = INDEX_NONE;
+		if (!Schema->Layout.EnumerateBy(ChannelIdx, EPoseSearchFeatureType::Invalid, FeatureIdx))
+		{
+			NormalizedChannelWeights[ChannelIdx] = 0.0f;
+		}
+	}
+
+	const float ChannelWeightSum = NormalizedChannelWeights.sum();
+	if (!FMath::IsNearlyZero(ChannelWeightSum))
+	{
+		NormalizedChannelWeights *= (1.0f / ChannelWeightSum);
+	}	
+
+
+	// Determine maximum number of channel sample offsets for allocation
+	int32 MaxChannelSampleOffsets = 0;
+	for (int ChannelIdx = 0; ChannelIdx != ChannelNum; ++ChannelIdx)
+	{
+		TArrayView<const float> ChannelSampleOffsets = Schema->GetChannelSampleOffsets(ChannelIdx);
+		MaxChannelSampleOffsets = FMath::Max(MaxChannelSampleOffsets, ChannelSampleOffsets.Num());
+	}
+
+	// WeightsByFeature is indexed by FeatureIdx in a Layout
+	TArray<float, TMemStackAllocator<>> WeightsByFeatureStorage;
+	WeightsByFeatureStorage.SetNum(Schema->Layout.Features.Num());
+	Map<ArrayXf> WeightsByFeature(WeightsByFeatureStorage.GetData(), WeightsByFeatureStorage.Num());
+
+	// HorizonWeightsBySample is indexed by the channel's sample offsets in a Schema
+	TArray<float, TMemStackAllocator<>> HorizonWeightsBySampleStorage;
+	HorizonWeightsBySampleStorage.SetNum(MaxChannelSampleOffsets);
+	Map<ArrayXf> HorizonWeightsBySample(HorizonWeightsBySampleStorage.GetData(), HorizonWeightsBySampleStorage.Num());
+
+	// WeightsByType is indexed by feature type
+	Array<float, (int)EPoseSearchFeatureType::Num, 1> WeightsByType;
+
+
+	// Determine each channel's feature weights
+	for (int ChannelIdx = 0; ChannelIdx != ChannelNum; ++ChannelIdx)
+	{
+		// Ignore this channel entirely if it has no weight
+		if (FMath::IsNearlyZero(NormalizedChannelWeights[ChannelIdx]))
+		{
+			continue;
+		}
+
+		// Get channel info
+		const FPoseSearchChannelWeightParams& ChannelWeights = *ChannelWeightParams[ChannelIdx];
+		const FPoseSearchChannelDynamicWeightParams& ChannelDynamicWeights = *ChannelDynamicWeightParams[ChannelIdx];
+		TArrayView<const float> ChannelSampleOffsets = Schema->GetChannelSampleOffsets(ChannelIdx);
+
+		// Reset scratch weights
+		WeightsByFeature.setConstant(0.0f);
+		WeightsByType.setConstant(0.0f);
+		HorizonWeightsBySample.setConstant(0.0f);
+
+
+		// Initialize weights by type lookup
+		for (int Type = 0; Type != (int)EPoseSearchFeatureType::Num; ++Type)
+		{
+			WeightsByType[Type] = ChannelWeights.TypeWeights.FindRef((EPoseSearchFeatureType)Type);
+
+			// Zero the weight if this channel doesn't have any features using this type
+			int32 FeatureIdx = INDEX_NONE;
+			if (!Schema->Layout.EnumerateBy(ChannelIdx, (EPoseSearchFeatureType)Type, FeatureIdx))
+			{
+				WeightsByType[Type] = 0.0f;
+			}
+		}
+
+		// Normalize type weights
+		float TypeWeightsSum = WeightsByType.sum();
+		if (!FMath::IsNearlyZero(TypeWeightsSum))
+		{
+			WeightsByType *= (1.0f / TypeWeightsSum);
+		}
+		else
+		{
+			// Ignore this channel entirely if there are no types contributing weight
+			continue;
+		}
+
+
+		// Determine the range of sample offsets that make up the history and prediction horizons
+		FInt32Range HorizonSampleIdxRanges[EHorizon::Num];
+		{
+			int32 IdxUpper = Algo::UpperBound(ChannelSampleOffsets, 0.0f);
+			int32 IdxLower = ChannelSampleOffsets[0] <= 0.0f ? 0 : IdxUpper;
+			HorizonSampleIdxRanges[EHorizon::History] = FInt32Range(IdxLower, IdxUpper);
+
+			IdxLower = IdxUpper;
+			IdxUpper = ChannelSampleOffsets.Num();
+			HorizonSampleIdxRanges[EHorizon::Prediction] = FInt32Range(IdxLower, IdxUpper);
+		}
+
+
+		// Initialize horizon weights
+		Array<float, 1, EHorizon::Num> NormalizedHorizonWeights;
+		NormalizedHorizonWeights.setConstant(0.0f);
+
+		if (!HorizonSampleIdxRanges[EHorizon::History].IsEmpty())
+		{
+			NormalizedHorizonWeights[EHorizon::History] = ChannelWeights.HistoryParams.Weight * ChannelDynamicWeights.HistoryWeightScale;
+		}
+		if (!HorizonSampleIdxRanges[EHorizon::Prediction].IsEmpty())
+		{
+			NormalizedHorizonWeights[EHorizon::Prediction] = ChannelWeights.PredictionParams.Weight * ChannelDynamicWeights.PredictionWeightScale;
+		}
+		
+		// Normalize horizon weights
+		float HorizonWeightSum = NormalizedHorizonWeights.sum();
+		if (!FMath::IsNearlyZero(HorizonWeightSum))
+		{
+			NormalizedHorizonWeights *= (1.0f / HorizonWeightSum);
+		}
+		else
+		{
+			// Ignore this channel entirely if the horizons don't contribute any weight
+			continue;
+		}
+
+
+		auto SetHorizonSampleWeights = [&HorizonWeightsBySample, &ChannelSampleOffsets](FInt32Range SampleIdxRange, const FPoseSearchChannelHorizonParams& HorizonParams)
+		{
+			// Segment length represents the number of sample offsets in the span that make up this horizon
+			int32 SegmentLength = SampleIdxRange.Size<int32>();
+
+			if (SegmentLength > 0)
+			{
+				int32 SegmentBegin = SampleIdxRange.GetLowerBoundValue();
+				if (HorizonParams.bInterpolate && SegmentLength > 1)
+				{
+					// We'll map the range spanned by the horizon's sample offsets to the interpolation range
+					// The interpolation range is 0 to 1 unless an initial value was set
+					// The initial value allows the user to set a minimum weight or reverse the lerp direction
+					// We'll normalize these weights in the next step
+					FVector2D InputRange(ChannelSampleOffsets[SegmentBegin], ChannelSampleOffsets[SegmentBegin + SegmentLength - 1]);
+					FVector2D OutputRange(HorizonParams.InitialValue, 1.0f - HorizonParams.InitialValue);
+
+					for (int32 OffsetIdx = SegmentBegin; OffsetIdx != SegmentBegin + SegmentLength; ++OffsetIdx)
+					{
+						float SampleOffset = ChannelSampleOffsets[OffsetIdx];
+						float Alpha = FMath::GetMappedRangeValueUnclamped(InputRange, OutputRange, SampleOffset);
+						float Weight = FAlphaBlend::AlphaToBlendOption(Alpha, HorizonParams.InterpolationMethod);
+						HorizonWeightsBySample[OffsetIdx] = Weight;
+					}
+				}
+				else
+				{
+					// If we're not interpolating weights across this horizon, just give them all equal weight
+					HorizonWeightsBySample.segment(SegmentBegin, SegmentLength).setConstant(1.0f);
+				}
+
+				// Normalize weights within the horizon's segment of sample offsets
+				float HorizonSum = HorizonWeightsBySample.segment(SegmentBegin, SegmentLength).sum();
+				if (!FMath::IsNearlyZero(HorizonSum))
+				{
+					HorizonWeightsBySample.segment(SegmentBegin, SegmentLength) *= 1.0f / HorizonSum;
+				}
+			}
+		};
+
+		SetHorizonSampleWeights(HorizonSampleIdxRanges[EHorizon::History], ChannelWeights.HistoryParams);
+		SetHorizonSampleWeights(HorizonSampleIdxRanges[EHorizon::Prediction], ChannelWeights.PredictionParams);
+
+
+		// Now set this channel's weights for every feature in each horizon
+		Array<float, 1, EHorizon::Num> HorizonSums;
+		HorizonSums = 0.0f;
+		for (int FeatureIdx = INDEX_NONE; Schema->Layout.EnumerateBy(ChannelIdx, EPoseSearchFeatureType::Invalid, FeatureIdx); /*empty*/)
+		{
+			const FPoseSearchFeatureDesc& Feature = Schema->Layout.Features[FeatureIdx];
+
+			for (int HorizonIdx = 0; HorizonIdx != EHorizon::Num; ++HorizonIdx)
+			{
+				if (HorizonSampleIdxRanges[HorizonIdx].Contains(Feature.SubsampleIdx))
+				{
+					int HorizonSize = HorizonSampleIdxRanges[HorizonIdx].Size<int>();
+					WeightsByFeature[FeatureIdx] = HorizonWeightsBySample[Feature.SubsampleIdx] * (HorizonSize * WeightsByType[(int)Feature.Type]);
+					HorizonSums[HorizonIdx] += WeightsByFeature[FeatureIdx];
+					break;
+				}
+			}
+		}
+
+		// Scale feature weights within horizons so that they have the desired total horizon weight
+		for (int FeatureIdx = INDEX_NONE; Schema->Layout.EnumerateBy(ChannelIdx, EPoseSearchFeatureType::Invalid, FeatureIdx); /*empty*/)
+		{
+			const FPoseSearchFeatureDesc& Feature = Schema->Layout.Features[FeatureIdx];
+
+			for (int HorizonIdx = 0; HorizonIdx != EHorizon::Num; ++HorizonIdx)
+			{
+				if (HorizonSampleIdxRanges[HorizonIdx].Contains(Feature.SubsampleIdx))
+				{
+					float HorizonWeight = NormalizedHorizonWeights[HorizonIdx] / HorizonSums[HorizonIdx];
+					WeightsByFeature[FeatureIdx] *= HorizonWeight;
+					break;
+				}
+			}
+		}
+
+		// Scale all features in all horizons so they have the desired channel weight
+		WeightsByFeature *= NormalizedChannelWeights[ChannelIdx];
+
+		// Weights should sum to channel weight at this point
+		ensure(FMath::IsNearlyEqual(WeightsByFeature.sum(), NormalizedChannelWeights[ChannelIdx], KINDA_SMALL_NUMBER));
+
+		// Merge feature weights for channel into per-value weights buffer
+		// Weights are replicated per feature dimension so the cost function can directly index weights by value index
+		for (int FeatureIdx = INDEX_NONE; Schema->Layout.EnumerateBy(ChannelIdx, EPoseSearchFeatureType::Invalid, FeatureIdx); /*empty*/)
+		{
+			const FPoseSearchFeatureDesc& Feature = Schema->Layout.Features[FeatureIdx];
+			int32 ValueSize = GetFeatureTypeTraits(Feature.Type).NumFloats;
+			int32 ValueTerm = Feature.ValueOffset + ValueSize;
+			for (int32 ValueIdx = Feature.ValueOffset; ValueIdx != ValueTerm; ++ValueIdx)
+			{
+				Weights[ValueIdx] = WeightsByFeature[FeatureIdx];
+			}
+		}
+	}
 }
 
-void FPoseSearchBiasWeights::BindSemanticWeight(float Weight, const FPoseSearchFeatureVectorLayout& Layout, EPoseSearchFeatureType FeatureType, bool bTrajectory)
+
+//////////////////////////////////////////////////////////////////////////
+// FPoseSearchWeightsContext
+
+void FPoseSearchWeightsContext::Update(const FPoseSearchDynamicWeightParams& ActiveWeights, const UPoseSearchDatabase* ActiveDatabase)
 {
-	// The Weight parameter will be bound to a specific feature described by the FPoseSearchFeatureVectorLayout
-	int32 FeatureIdx = INDEX_NONE;
-	while (Layout.EnumerateFeature(FeatureType, bTrajectory, FeatureIdx))
+	bool bRecomputeWeights = false;
+	if (Database != ActiveDatabase)
 	{
-		const FPoseSearchFeatureDesc& Feature = Layout.Features[FeatureIdx];
-		const int32 FirstValueIdx = Feature.ValueOffset;
-		const int32 NumValues = UE::PoseSearch::GetFeatureTypeTraits(FeatureType).NumFloats;
+		Database = ActiveDatabase;
+		bRecomputeWeights = true;
+	}
 
-		// NOTE: This enumeration method is currently assumed by the PoseSearch Debugger
-		for (int32 Idx = 0; Idx < NumValues; ++Idx)
+	if (DynamicWeights != ActiveWeights)
+	{
+		DynamicWeights = ActiveWeights;
+		bRecomputeWeights = true;
+	}
+
+	if (bRecomputeWeights)
+	{
+		int32 NumGroups = ActiveDatabase ? 1 : 0;
+		ComputedGroupWeights.SetNum(NumGroups);
+		for (FPoseSearchWeights& GroupWeights: ComputedGroupWeights)
 		{
-			Weights[FirstValueIdx + Idx] = Weight;
+			GroupWeights.Init(Database->Weights, Database->Schema, ActiveWeights);
 		}
 	}
 }
 
+const FPoseSearchWeights* FPoseSearchWeightsContext::GetGroupWeights(int32 WeightsGroupIdx) const
+{
+	if (ComputedGroupWeights.IsValidIndex(WeightsGroupIdx))
+	{
+		return &ComputedGroupWeights[WeightsGroupIdx];
+	}
+
+	return nullptr;
+}
+
 
 //////////////////////////////////////////////////////////////////////////
 // FPoseSearchIndex
@@ -847,7 +1178,7 @@ void FPoseSearchFeatureVectorBuilder::SetVector(FPoseSearchFeatureDesc Element,
 
 bool FPoseSearchFeatureVectorBuilder::TrySetPoseFeatures(UE::PoseSearch::FPoseHistory* History)
 {
-	check(Schema && Schema->IsValid());
+	check(Schema.IsValid() && Schema->IsValid());
 	check(History);
 
 	FPoseSearchFeatureDesc Feature;
@@ -889,7 +1220,7 @@ bool FPoseSearchFeatureVectorBuilder::TrySetPoseFeatures(UE::PoseSearch::FPoseHi
 
 void FPoseSearchFeatureVectorBuilder::BuildFromTrajectoryTimeBased(const FTrajectorySampleRange& Trajectory)
 {
-	check(Schema && Schema->IsValid());
+	check(Schema.IsValid() && Schema->IsValid());
 
 	FPoseSearchFeatureDesc Feature;
 	Feature.Domain = EPoseSearchFeatureDomain::Time;
@@ -912,7 +1243,7 @@ void FPoseSearchFeatureVectorBuilder::BuildFromTrajectoryTimeBased(const FTrajec
 
 void FPoseSearchFeatureVectorBuilder::BuildFromTrajectoryDistanceBased(const FTrajectorySampleRange& Trajectory)
 {
-	check(Schema && Schema->IsValid());
+	check(Schema.IsValid() && Schema->IsValid());
 
 	FPoseSearchFeatureDesc Feature;
 	Feature.Domain = EPoseSearchFeatureDomain::Distance;
@@ -2348,7 +2679,7 @@ static void DrawFeatureVector(const FDebugDrawParams& DrawParams, const FFeature
 	if (EnumHasAnyFlags(DrawParams.Flags, EDebugDrawFlags::IncludeTrajectory))
 	{
 		DrawTrajectoryFeatures(DrawParams, Reader, EPoseSearchFeatureDomain::Time);
-        DrawTrajectoryFeatures(DrawParams, Reader, EPoseSearchFeatureDomain::Distance);
+		DrawTrajectoryFeatures(DrawParams, Reader, EPoseSearchFeatureDomain::Distance);
 	}
 }
 
@@ -2497,7 +2828,7 @@ static void PreprocessSearchIndexNormalize(FPoseSearchIndex* SearchIndex)
 	// exponentially rather than additively and square rooting the sum of squares does not 
 	// remove that bias. [1]
 	//
-	// The pose matrix is transformed in place and the tranformation matrix, its inverse,
+	// The pose matrix is transformed in place and the transformation matrix, its inverse,
 	// and data mean vector are computed and stored along with it.
 	//
 	// N:	number of dimensions for input column vectors
@@ -2614,7 +2945,7 @@ static void PreprocessSearchIndexNormalize(FPoseSearchIndex* SearchIndex)
 static void PreprocessSearchIndexSphere(FPoseSearchIndex* SearchIndex)
 {
 	// This function performs correlation based zero-phase component analysis sphering (ZCA-cor sphering)
-	// The pose matrix is transformed in place and the tranformation matrix, its inverse,
+	// The pose matrix is transformed in place and the transformation matrix, its inverse,
 	// and data mean vector are computed and stored along with it.
 	//
 	// N:	number of dimensions for input column vectors
@@ -2657,13 +2988,13 @@ static void PreprocessSearchIndexSphere(FPoseSearchIndex* SearchIndex)
 	// Z^(-1) = V * D^(1/2) * V^T
 	// x := (Z^(-1) * x) + u
 	//
-	// The sphering processs allows nearest neighbor queries to use the Mahalonobis metric
+	// The sphering process allows nearest neighbor queries to use the Mahalanobis metric
 	// which is unitless, scale-invariant, and accounts for feature correlation.
 	// The Mahalanobis distance between two random vectors x and y in data matrix X is:
 	// d(x,y) = ((x-y)^T * cov(X)^(-1) * (x-y))^(1/2)
 	//
-	// Because sphering transforms X into a new matrix with identity covariance, the Mahalonobis
-	// distance equation above reduces to Eucledean distance since cov(X)^(-1) = I:
+	// Because sphering transforms X into a new matrix with identity covariance, the Mahalanobis
+	// distance equation above reduces to Euclidean distance since cov(X)^(-1) = I:
 	// d(x,y) = ((x-y)^T * (x-y))^(1/2)
 	// 
 	// References:
@@ -2679,6 +3010,7 @@ static void PreprocessSearchIndexSphere(FPoseSearchIndex* SearchIndex)
 	//
 	// Note this sphering preprocessor needs more work and isn't yet exposed in the editor as an option.
 	// Todo:
+	// - Figure out apparent flipping behavior
 	// - Try singular value decomposition in place of eigendecomposition
 	// - Remove zero variance feature axes from data and search queries
 	// - Support weighted Mahalanobis metric. User supplied weights need to be transformed to data's new basis.
@@ -2980,20 +3312,6 @@ bool BuildIndex(UPoseSearchDatabase* Database)
 		TotalFloats += Output.FeatureVectorTable.Num();
 	}
 
-	// Establish per-sequence pose search bias weights if metadata is present
-	for (int32 SequenceIdx = 0; SequenceIdx != Database->Sequences.Num(); ++SequenceIdx)
-	{
-		FPoseSearchDatabaseSequence& DbSequence = Database->Sequences[SequenceIdx];
-		if (DbSequence.Sequence)
-		{
-			const UPoseSearchSequenceBiasWeightMetaData* BiasWeightMetadata = DbSequence.Sequence->FindMetaDataByClass<UPoseSearchSequenceBiasWeightMetaData>();
-			if (BiasWeightMetadata)
-			{
-				DbSequence.BiasWeights.Init(BiasWeightMetadata->BiasWeights, Database->Schema->Layout);
-			}
-		}
-	}
-
 	// Join animation data into a single search index
 	Database->SearchIndex.Values.Reset(TotalFloats);
 	for (const FSequenceIndexer& Indexer : Indexers)
@@ -3010,21 +3328,7 @@ bool BuildIndex(UPoseSearchDatabase* Database)
 	return true;
 }
 
-static inline void DefaultInitializeWeights(const FPoseSearchBiasWeightsContext* BiasWeightsContext, int32 Size, TArray<float>& Weights, bool& bBiasWeightContextAvailable)
-{
-	bBiasWeightContextAvailable = BiasWeightsContext && BiasWeightsContext->HasBiasWeights();
-
-	if (bBiasWeightContextAvailable)
-	{
-		Weights = BiasWeightsContext->BiasWeights->Weights;
-	}
-	else
-	{
-		Weights.Init(1.f, Size);
-	}
-}
-
-static FSearchResult Search(const FPoseSearchIndex& SearchIndex, TArrayView<const float> Query, const FPoseSearchBiasWeightsContext* BiasWeightsContext = nullptr)
+static FSearchResult Search(const FPoseSearchIndex& SearchIndex, TArrayView<const float> Query, const FPoseSearchWeightsContext* WeightsContext = nullptr)
 {
 	FSearchResult Result;
 	if (!ensure(SearchIndex.IsValid()))
@@ -3032,64 +3336,17 @@ static FSearchResult Search(const FPoseSearchIndex& SearchIndex, TArrayView<cons
 		return Result;
 	}
 
-	if(!ensure(Query.Num() == SearchIndex.Schema->Layout.NumFloats))
+	if (!ensure(Query.Num() == SearchIndex.Schema->Layout.NumFloats))
 	{
 		return Result;
 	}
 
-	bool bBiasWeightContextAvailable = false;
-
-	// Initial weights by default are set to 1, but may be independently set by an external system such as motion matching
-	TArray<float> InitialWeights;
-	DefaultInitializeWeights(BiasWeightsContext, Query.Num(), InitialWeights, bBiasWeightContextAvailable);
-	const auto InitialWeightsMap = Eigen::Map<const Eigen::ArrayXf>(InitialWeights.GetData(), InitialWeights.Num());
-
-	// Accumulated weights will contain the per-pose final weights, optionally including per-sequence and/or other external values
-	TArray<float> AccumulatedWeights(InitialWeights);
-	auto AccumulatedWeightsMap = Eigen::Map<Eigen::ArrayXf>(AccumulatedWeights.GetData(), AccumulatedWeights.Num());
-
 	float BestPoseDissimilarity = MAX_flt;
 	int32 BestPoseIdx = INDEX_NONE;
 
-	for (int32 PoseIdx = 0, PrevSequenceIdx = INDEX_NONE; PoseIdx != SearchIndex.NumPoses; ++PoseIdx)
+	for (int32 PoseIdx = 0; PoseIdx != SearchIndex.NumPoses; ++PoseIdx)
 	{
-		// Sequence index and metadata tracking are done within this loop in order to optimize
-		// and elide unnecessary recomputing of the weight buffer.
-		bool bSequenceWeightsAvailable = false;
-		int32_t SequenceIdx = INDEX_NONE;
-
-		if (bBiasWeightContextAvailable)
-		{
-			// Apply the per-sequence bias weights if they are present within the sequence metadata
-			SequenceIdx = BiasWeightsContext->Database->FindSequenceForPose(PoseIdx);
-			if (SequenceIdx != PrevSequenceIdx)
-			{
-				const FPoseSearchDatabaseSequence& SequenceEntry = BiasWeightsContext->Database->Sequences[SequenceIdx];
-				bSequenceWeightsAvailable = SequenceEntry.BiasWeights.IsInitialized();
-
-				if (bSequenceWeightsAvailable)
-				{
-					const auto& SequenceBiasWeights = SequenceEntry.BiasWeights.Weights;
-					AccumulatedWeightsMap *= Eigen::Map<const Eigen::ArrayXf>(SequenceBiasWeights.GetData(), SequenceBiasWeights.Num());
-				}
-			}
-		}
-
-		const int32 FeatureValueOffset = PoseIdx * SearchIndex.Schema->Layout.NumFloats;
-
-		const float PoseDissimilarity = CompareFeatureVectors(
-			SearchIndex.Schema->Layout.NumFloats,
-			Query.GetData(),
-			&SearchIndex.Values[FeatureValueOffset],
-			AccumulatedWeights.GetData()
-		);
-
-		if (bSequenceWeightsAvailable && SequenceIdx != PrevSequenceIdx)
-		{
-			// Reset pose weights to remove any extraneous sequence contributions for the next iteration
-			AccumulatedWeightsMap = InitialWeightsMap;
-			PrevSequenceIdx = SequenceIdx;
-		}
+		const float PoseDissimilarity = ComparePoses(SearchIndex, PoseIdx, Query, WeightsContext);
 		
 		if (PoseDissimilarity < BestPoseDissimilarity)
 		{
@@ -3135,7 +3392,7 @@ FSearchResult Search(const UAnimSequenceBase* Sequence, TArrayView<const float>
 	return Result;
 }
 
-FDbSearchResult Search(const UPoseSearchDatabase* Database, TArrayView<const float> Query, const FPoseSearchBiasWeightsContext* BiasWeightsContext, FDebugDrawParams DebugDrawParams)
+FDbSearchResult Search(const UPoseSearchDatabase* Database, TArrayView<const float> Query, const FPoseSearchWeightsContext* WeightsContext, FDebugDrawParams DebugDrawParams)
 {
 	if (!ensure(Database && Database->IsValidForSearch()))
 	{
@@ -3144,7 +3401,7 @@ FDbSearchResult Search(const UPoseSearchDatabase* Database, TArrayView<const flo
 
 	const FPoseSearchIndex& SearchIndex = Database->SearchIndex;
 
-	FDbSearchResult Result = Search(SearchIndex, Query, BiasWeightsContext);
+	FDbSearchResult Result = Search(SearchIndex, Query, WeightsContext);
 	if (!Result.IsValid())
 	{
 		return FDbSearchResult();
@@ -3171,30 +3428,23 @@ FDbSearchResult Search(const UPoseSearchDatabase* Database, TArrayView<const flo
 	return Result;
 }
 
-float ComparePoses(const FPoseSearchIndex& SearchIndex, int32 PoseIdx, TArrayView<const float> Query, const FPoseSearchBiasWeightsContext* BiasWeightsContext)
+float ComparePoses(const FPoseSearchIndex& SearchIndex, int32 PoseIdx, TArrayView<const float> Query, const FPoseSearchWeightsContext* WeightsContext)
 {
 	TArrayView<const float> PoseValues = SearchIndex.GetPoseValues(PoseIdx);
 	check(PoseValues.Num() == Query.Num());
 
-	bool bBiasWeightContextAvailable = false;
-
-	TArray<float> Weights;
-	DefaultInitializeWeights(BiasWeightsContext, Query.Num(), Weights, bBiasWeightContextAvailable);
-
-	if (bBiasWeightContextAvailable)
+	float Dissimilarity;
+	if (WeightsContext)
 	{
-		// Apply the per-sequence bias weights if present on the sequence metadata
-		const int32 SequenceIdx = BiasWeightsContext->Database->FindSequenceForPose(PoseIdx);
-		const FPoseSearchDatabaseSequence& SequenceEntry = BiasWeightsContext->Database->Sequences[SequenceIdx];
-
-		if (SequenceEntry.BiasWeights.IsInitialized())
-		{
-			const auto& SequenceWeights = SequenceEntry.BiasWeights.Weights;
-			Eigen::Map<Eigen::ArrayXf>(Weights.GetData(), Weights.Num()) *= Eigen::Map<const Eigen::ArrayXf>(SequenceWeights.GetData(), SequenceWeights.Num());
-		}
+		const FPoseSearchWeights* WeightsSet = WeightsContext->GetGroupWeights(0);
+		Dissimilarity = CompareFeatureVectors(PoseValues.Num(), PoseValues.GetData(), Query.GetData(), WeightsSet->Weights.GetData());
+	}
+	else
+	{
+		Dissimilarity = CompareFeatureVectors(PoseValues.Num(), PoseValues.GetData(), Query.GetData());
 	}
 
-	return CompareFeatureVectors(PoseValues.Num(), PoseValues.GetData(), Query.GetData(), Weights.GetData());
+	return Dissimilarity;
 }
 
 
@@ -3263,6 +3513,4 @@ UE::Anim::IPoseSearchProvider::FSearchResult FModule::Search(const FAnimationBas
 
 }} // namespace UE::PoseSearch
 
-PRAGMA_ENABLE_OPTIMIZATION
-
 IMPLEMENT_MODULE(UE::PoseSearch::FModule, PoseSearch)
\ No newline at end of file
diff --git a/Engine/Plugins/Experimental/Animation/PoseSearch/Source/Runtime/Private/PoseSearchLibrary.cpp b/Engine/Plugins/Experimental/Animation/PoseSearch/Source/Runtime/Private/PoseSearchLibrary.cpp
--- a/Engine/Plugins/Experimental/Animation/PoseSearch/Source/Runtime/Private/PoseSearchLibrary.cpp
+++ b/Engine/Plugins/Experimental/Animation/PoseSearch/Source/Runtime/Private/PoseSearchLibrary.cpp
@@ -2,6 +2,7 @@
 
 #include "PoseSearch/PoseSearchLibrary.h"
 
+#include "Animation/AnimSequence.h"
 #include "Animation/AnimNode_Inertialization.h"
 #include "Animation/AnimNode_SequencePlayer.h"
 #include "PoseSearch/AnimNode_MotionMatching.h"
@@ -64,7 +65,6 @@ void UpdateMotionMatchingState(const FAnimationUpdateContext& Context
 	, FMotionMatchingState& InOutMotionMatchingState
 )
 {
-	FPoseSearchBiasWeights QueryBiasWeights;
 	InOutMotionMatchingState.Flags = EMotionMatchingFlags::None;
 	const float DeltaTime = Context.GetDeltaTime();
 
@@ -108,20 +108,27 @@ void UpdateMotionMatchingState(const FAnimationUpdateContext& Context
 		// Update features in the query with the latest inputs
 		InOutMotionMatchingState.ComposeQuery(Database, Trajectory);
 
-		// Initialize the pose search query bias weights context
-		QueryBiasWeights.Init(Settings.BiasWeights, Database->Schema->Layout);
-
-		const FPoseSearchBiasWeightsContext BiasWeightsContext = { &QueryBiasWeights, Database };
+		// Update weight groups
+		InOutMotionMatchingState.WeightsContext.Update(Settings.Weights, Database);
 
 		// Determine how much the updated query vector deviates from the current pose vector
 		float CurrentDissimilarity = MAX_flt;
 		if (InOutMotionMatchingState.DbPoseIdx != INDEX_NONE)
 		{
-			CurrentDissimilarity = UE::PoseSearch::ComparePoses(Database->SearchIndex, InOutMotionMatchingState.DbPoseIdx, InOutMotionMatchingState.ComposedQuery.GetNormalizedValues(), &BiasWeightsContext);
+			CurrentDissimilarity = UE::PoseSearch::ComparePoses(
+				Database->SearchIndex,
+				InOutMotionMatchingState.DbPoseIdx,
+				InOutMotionMatchingState.ComposedQuery.GetNormalizedValues(),
+				&InOutMotionMatchingState.WeightsContext
+			);
 		}
 
 		// Search the database for the nearest match to the updated query vector
-		UE::PoseSearch::FDbSearchResult Result = UE::PoseSearch::Search(Database, InOutMotionMatchingState.ComposedQuery.GetNormalizedValues(), &BiasWeightsContext);
+		UE::PoseSearch::FDbSearchResult Result = UE::PoseSearch::Search(
+			Database,
+			InOutMotionMatchingState.ComposedQuery.GetNormalizedValues(),
+			&InOutMotionMatchingState.WeightsContext
+		);
 		if (Result.IsValid() && (InOutMotionMatchingState.ElapsedPoseJumpTime >= Settings.SearchThrottleTime))
 		{
 			const FPoseSearchDatabaseSequence& ResultDbSequence = Database->Sequences[Result.DbSequenceIdx];
@@ -136,7 +143,8 @@ void UpdateMotionMatchingState(const FAnimationUpdateContext& Context
 				bNearbyPose = FMath::Abs(InOutMotionMatchingState.AssetPlayerTime - Result.TimeOffsetSeconds) < Settings.PoseJumpThresholdTime;
 				if (!bNearbyPose && ResultDbSequence.bLoopAnimation)
 				{
-					bNearbyPose = FMath::Abs(Database->GetSequenceLength(InOutMotionMatchingState.DbSequenceIdx) - InOutMotionMatchingState.AssetPlayerTime - Result.TimeOffsetSeconds) < Settings.PoseJumpThresholdTime;
+					const float AssetLength = Database->GetSequenceLength(InOutMotionMatchingState.DbSequenceIdx);
+					bNearbyPose = FMath::Abs(AssetLength - InOutMotionMatchingState.AssetPlayerTime - Result.TimeOffsetSeconds) < Settings.PoseJumpThresholdTime;
 				}
 			}
 
@@ -184,23 +192,23 @@ void UpdateMotionMatchingState(const FAnimationUpdateContext& Context
 	}
 
 #if UE_POSE_SEARCH_TRACE_ENABLED
-	UE::PoseSearch::FTraceMotionMatchingState TraceState;
-	if ((InOutMotionMatchingState.Flags & EMotionMatchingFlags::JumpedToPose) == EMotionMatchingFlags::JumpedToPose)
-	{
-		TraceState.Flags |= UE::PoseSearch::FTraceMotionMatchingState::EFlags::FollowupAnimation;
-	}
-	if (InOutMotionMatchingState.DbPoseIdx == INDEX_NONE)
+	if (InOutMotionMatchingState.DbPoseIdx != INDEX_NONE)
 	{
-		return;
+		UE::PoseSearch::FTraceMotionMatchingState TraceState;
+		if ((InOutMotionMatchingState.Flags & EMotionMatchingFlags::JumpedToPose) == EMotionMatchingFlags::JumpedToPose)
+		{
+			TraceState.Flags |= UE::PoseSearch::FTraceMotionMatchingState::EFlags::FollowupAnimation;
+		}
+
+		TraceState.ElapsedPoseJumpTime = InOutMotionMatchingState.ElapsedPoseJumpTime;
+		// @TODO: Change this to only be the previous query, not persistently updated (i.e. if throttled)?
+		TraceState.QueryVector = InOutMotionMatchingState.ComposedQuery.GetValues();
+		TraceState.QueryVectorNormalized = InOutMotionMatchingState.ComposedQuery.GetNormalizedValues();
+		TraceState.BiasWeights = InOutMotionMatchingState.WeightsContext.GetGroupWeights(0)->Weights;
+		TraceState.DbPoseIdx = InOutMotionMatchingState.DbPoseIdx;
+		TraceState.DatabaseId = FObjectTrace::GetObjectId(Database);
+		UE_TRACE_POSE_SEARCH_MOTION_MATCHING_STATE(Context, TraceState)
 	}
-	TraceState.ElapsedPoseJumpTime = InOutMotionMatchingState.ElapsedPoseJumpTime;
-	// @TODO: Change this to only be the previous query, not persistently updated (i.e. if throttled)?
-	TraceState.QueryVector = InOutMotionMatchingState.ComposedQuery.GetValues();
-	TraceState.QueryVectorNormalized = InOutMotionMatchingState.ComposedQuery.GetNormalizedValues();
-	TraceState.BiasWeights = QueryBiasWeights.Weights;
-	TraceState.DbPoseIdx = InOutMotionMatchingState.DbPoseIdx;
-	TraceState.DatabaseId = FObjectTrace::GetObjectId(Database);
-	UE_TRACE_POSE_SEARCH_MOTION_MATCHING_STATE(Context, TraceState)
 #endif
 }
 
diff --git a/Engine/Plugins/Experimental/Animation/PoseSearch/Source/Runtime/Private/Trace/PoseSearchTraceLogger.h b/Engine/Plugins/Experimental/Animation/PoseSearch/Source/Runtime/Private/Trace/PoseSearchTraceLogger.h
--- a/Engine/Plugins/Experimental/Animation/PoseSearch/Source/Runtime/Private/Trace/PoseSearchTraceLogger.h
+++ b/Engine/Plugins/Experimental/Animation/PoseSearch/Source/Runtime/Private/Trace/PoseSearchTraceLogger.h
@@ -73,7 +73,7 @@ ENUM_CLASS_FLAGS(FTraceMotionMatchingState::EFlags)
 #define UE_TRACE_POSE_SEARCH_MOTION_MATCHING_STATE(Context, MotionMatchingState) \
 	UE::PoseSearch::FTraceMotionMatchingState::Output(Context, MotionMatchingState);
 
-#else
+#else // UE_POSE_SEARCH_TRACE_ENABLED
 
 // Empty declarations if not enabled
 #define UE_TRACE_POSE_SEARCH_MOTION_MATCHING_STATE(Context, MotionMatchingState)
diff --git a/Engine/Plugins/Experimental/Animation/PoseSearch/Source/Runtime/Public/PoseSearch/AnimNode_PoseSearchHistoryCollector.h b/Engine/Plugins/Experimental/Animation/PoseSearch/Source/Runtime/Public/PoseSearch/AnimNode_PoseSearchHistoryCollector.h
--- a/Engine/Plugins/Experimental/Animation/PoseSearch/Source/Runtime/Public/PoseSearch/AnimNode_PoseSearchHistoryCollector.h
+++ b/Engine/Plugins/Experimental/Animation/PoseSearch/Source/Runtime/Public/PoseSearch/AnimNode_PoseSearchHistoryCollector.h
@@ -28,6 +28,7 @@ public:
 
 	// FAnimNode_Base interface
 	virtual void Initialize_AnyThread(const FAnimationInitializeContext& Context) override;
+	virtual void CacheBones_AnyThread(const FAnimationCacheBonesContext& Context) override;
 	virtual void Evaluate_AnyThread(FPoseContext& Output) override;
 	virtual void Update_AnyThread(const FAnimationUpdateContext& Context);
 	virtual void GatherDebugData(FNodeDebugData& DebugData);
diff --git a/Engine/Plugins/Experimental/Animation/PoseSearch/Source/Runtime/Public/PoseSearch/PoseSearch.h b/Engine/Plugins/Experimental/Animation/PoseSearch/Source/Runtime/Public/PoseSearch/PoseSearch.h
--- a/Engine/Plugins/Experimental/Animation/PoseSearch/Source/Runtime/Public/PoseSearch/PoseSearch.h
+++ b/Engine/Plugins/Experimental/Animation/PoseSearch/Source/Runtime/Public/PoseSearch/PoseSearch.h
@@ -13,6 +13,7 @@
 #include "Animation/AnimNodeMessages.h"
 #include "Animation/AnimationAsset.h"
 #include "Animation/MotionTrajectoryTypes.h"
+#include "AlphaBlend.h"
 #include "BoneIndices.h"
 #include "Interfaces/Interface_BoneReferenceSkeletonProvider.h"
 
@@ -61,7 +62,7 @@ struct POSESEARCH_API FPoseSearchFeatureDesc
 {
 	GENERATED_BODY()
 
-	static constexpr int32 TrajectoryBoneIndex = -1;
+	static constexpr int32 TrajectoryBoneIndex = -1; 
 
 	UPROPERTY()
 	int32 SchemaBoneIdx;
@@ -81,8 +82,6 @@ struct POSESEARCH_API FPoseSearchFeatureDesc
 
 	bool operator==(const FPoseSearchFeatureDesc& Other) const;
 
-	bool IsSubsampleOfSameFeature(const FPoseSearchFeatureDesc& Other) const;
-
 	FPoseSearchFeatureDesc()
 		: SchemaBoneIdx(MAX_int32)
 		, SubsampleIdx(MAX_int32)
@@ -92,6 +91,7 @@ struct POSESEARCH_API FPoseSearchFeatureDesc
 	{}
 };
 
+
 /**
 * Explicit description of a pose feature vector.
 * Determined by options set in a UPoseSearchSchema and owned by the schema.
@@ -109,49 +109,19 @@ struct POSESEARCH_API FPoseSearchFeatureVectorLayout
 	TArray<FPoseSearchFeatureDesc> Features;
 
 	UPROPERTY()
-	uint32 NumFloats = 0;
-
-	UPROPERTY()
-	uint32 NumTrajectoryValues = 0;
-
-	UPROPERTY()
-	int32 FirstTrajectoryValueOffset = -1;
-
-	UPROPERTY()
-	uint32 NumPoseValues = 0;
-
-	UPROPERTY()
-	int32 FirstPoseValueOffset = -1;
+	int32 NumFloats = 0;
 
 	bool IsValid(int32 MaxNumBones) const;
 
-	bool EnumerateFeature(EPoseSearchFeatureType FeatureType, bool bTrajectory, int32& OutFeatureIdx) const;
+	bool EnumerateBy(int32 ChannelIdx, EPoseSearchFeatureType Type, int32& InOutFeatureIdx) const;
 };
 
-
-//USTRUCT()
-//struct FPoseSearchSchemaBoneRef
-//{
-//	GENERATED_BODY()
-//
-//	/** Index into UPoseSearchSchema::Bones. */
-//	UPROPERTY()
-//	int32 BoneIdx = -1;
-//
-//	/**
-//	* Optional index into UPoseSearchSchema::Bones.
-//	* Causes BoneIdx to be sampled in the reference frame of this bone.
-//	*/
-//	UPROPERTY()
-//	int32 RelativeBoneIdx = -1;
-//};
-
 UENUM()
 enum class EPoseSearchDataPreprocessor : int32
 {
 	None			= 0,
 	Automatic		= 1,
-	Normalize    	= 2,
+	Normalize		= 2,
 	Sphere			= 3 UMETA(Hidden),
 
 	Invalid = -1 UMETA(Hidden)
@@ -170,7 +140,7 @@ public:
 	static constexpr int32 DefaultSampleRate = 10;
 
 	UPROPERTY(EditAnywhere, Category = "Schema")
-	USkeleton* Skeleton = nullptr;
+	TObjectPtr<USkeleton> Skeleton = nullptr;
 
 	UPROPERTY(EditAnywhere, meta = (ClampMin = "1", ClampMax = "60"), Category = "Schema")
 	int32 SampleRate = DefaultSampleRate;
@@ -190,11 +160,6 @@ public:
 	UPROPERTY(EditAnywhere, Category = "Schema")
 	TArray<FBoneReference> Bones;
 
-	/*
-	UPROPERTY(EditAnywhere, Category = "Schema")
-	TArray<FPoseSearchSchemaBoneRef> PoseSampleBones;
-	*/
-
 	UPROPERTY(EditAnywhere, Category = "Schema")
 	TArray<float> PoseSampleTimes;
 
@@ -243,6 +208,8 @@ public:
 	// Returns a positive value when there are no past sample distances.
 	float GetTrajectoryPastDistanceHorizon () const;
 
+	TArrayView<const float> GetChannelSampleOffsets (int32 ChannelIdx) const;
+
 public: // UObject
 	virtual void PreSave(FObjectPreSaveContext ObjectSaveContext) override;
 	virtual void PostLoad() override;
@@ -261,7 +228,7 @@ struct POSESEARCH_API FPoseSearchIndexPreprocessInfo
 	GENERATED_BODY()
 
 	UPROPERTY()
-	uint32 NumDimensions = 0;
+	int32 NumDimensions = 0;
 
 	UPROPERTY()
 	TArray<float> TransformationMatrix;
@@ -297,7 +264,7 @@ struct POSESEARCH_API FPoseSearchIndex
 	TArray<float> Values;
 
 	UPROPERTY()
-	const UPoseSearchSchema* Schema = nullptr;
+	TObjectPtr<const UPoseSearchSchema> Schema = nullptr;
 
 	UPROPERTY()
 	FPoseSearchIndexPreprocessInfo PreprocessInfo;
@@ -339,7 +306,7 @@ class POSESEARCH_API UPoseSearchSequenceMetaData : public UAnimMetaData
 public:
 
 	UPROPERTY(EditAnywhere, Category="Settings")
-	const UPoseSearchSchema* Schema = nullptr;
+	TObjectPtr<const UPoseSearchSchema> Schema = nullptr;
 
 	UPROPERTY(EditAnywhere, Category="Settings")
 	FFloatInterval SamplingRange = FFloatInterval(0.0f, 0.0f);
@@ -357,65 +324,148 @@ public: // UObject
 	virtual void PreSave(FObjectPreSaveContext ObjectSaveContext) override;
 };
 
-/** Bias weights for influencing the results of a pose search query */
 USTRUCT(BlueprintType, Category = "Animation|Pose Search")
-struct POSESEARCH_API FPoseSearchBiasWeightParams
+struct POSESEARCH_API FPoseSearchChannelHorizonParams
 {
 	GENERATED_BODY()
 
-	UPROPERTY(BlueprintReadWrite, EditAnywhere, Category = "Settings", meta = (ClampMin = "0.0"));
-	float TrajectoryLinearVelocityWeight = 1.f;
+	// Total score contribution of all samples within this horizon, normalized with other horizons
+	UPROPERTY(BlueprintReadWrite, EditAnywhere, Category = "Settings", meta = (ClampMin = "0.0"))
+	float Weight = 1.0f;
 
-	UPROPERTY(BlueprintReadWrite, EditAnywhere, Category = "Settings", meta = (ClampMin = "0.0"));
-	float TrajectoryPositionWeight = 1.f;
+	// Whether to interpolate samples within this horizon
+	UPROPERTY(BlueprintReadWrite, EditAnywhere, Category = "Advanced")
+	bool bInterpolate = false;
 
-	UPROPERTY(BlueprintReadWrite, EditAnywhere, Category = "Settings", meta = (ClampMin = "0.0"));
-	float PoseLinearVelocityWeight = 1.f;
+	// Horizon sample weights will be interpolated from InitialValue to 1.0 - InitialValue and then normalized
+	UPROPERTY(BlueprintReadWrite, EditAnywhere, Category = "Advanced", meta = (EditCondition = "bInterpolate", ClampMin="0.0", ClampMax="1.0"))
+	float InitialValue = 0.1f;
 
-	UPROPERTY(BlueprintReadWrite, EditAnywhere, Category = "Settings", meta = (ClampMin = "0.0"));
-	float PosePositionWeight = 1.f;
+	// Curve type for horizon interpolation 
+	UPROPERTY(BlueprintReadWrite, EditAnywhere, Category = "Advanced", meta = (EditCondition = "bInterpolate"))
+	EAlphaBlendOption InterpolationMethod = EAlphaBlendOption::Linear;
 };
 
-/** Expanded bias weights from FPoseSearchBiasWeightParams which match the binary layout described by FPoseSearchFeatureVectorLayout */
-USTRUCT()
-struct POSESEARCH_API FPoseSearchBiasWeights
+USTRUCT(BlueprintType, Category = "Animation|Pose Search")
+struct POSESEARCH_API FPoseSearchChannelWeightParams
 {
 	GENERATED_BODY()
 
-	UPROPERTY()
-	TArray<float> Weights;
+	// Contribution of this score component. Normalized with other channels.
+	UPROPERTY(BlueprintReadWrite, EditAnywhere, Category = "Settings", meta = (ClampMin = "0.0"))
+	float ChannelWeight = 1.0f;
 
-	bool IsInitialized() const { return !Weights.IsEmpty(); }
-	void Init(const FPoseSearchBiasWeightParams& WeightParams, const FPoseSearchFeatureVectorLayout& Layout);
+	// History horizon params (for sample offsets <= 0)
+	UPROPERTY(BlueprintReadWrite, EditAnywhere, Category = "Settings")
+	FPoseSearchChannelHorizonParams HistoryParams;
 
-private:
-	void BindSemanticWeight(float Weight, const FPoseSearchFeatureVectorLayout& Layout, EPoseSearchFeatureType FeatureType, bool bTrajectory);
+	// Prediction horizon params (for sample offsets > 0)
+	UPROPERTY(BlueprintReadWrite, EditAnywhere, Category = "Settings")
+	FPoseSearchChannelHorizonParams PredictionParams;
+
+	// Contribution of each type within this channel
+	UPROPERTY(BlueprintReadWrite, EditAnywhere, Category = "Settings")
+	TMap<EPoseSearchFeatureType, float> TypeWeights;
+
+	FPoseSearchChannelWeightParams();
 };
 
-/** 
-* Bias weight context contains a reference to global/external weights and a relevant pose search database
-* The database parameter is primarily used for locating per-sequence bias weights
-*/
-struct FPoseSearchBiasWeightsContext
+USTRUCT(BlueprintType, Category = "Animation|Pose Search")
+struct POSESEARCH_API FPoseSearchChannelDynamicWeightParams
 {
-	const FPoseSearchBiasWeights* BiasWeights = nullptr;
-	const UPoseSearchDatabase* Database = nullptr;
+	GENERATED_BODY()
+
+	// Multiplier for the contribution of this score component. Final weight will be normalized with other channels after scaling.
+	UPROPERTY(BlueprintReadWrite, EditAnywhere, Category = "Settings", meta = (ClampMin = "0.0"))
+	float ChannelWeightScale = 1.0f;
+
+	// Multiplier for history score contribution. Normalized with prediction weight after scaling.
+	UPROPERTY(BlueprintReadWrite, EditAnywhere, Category = "Settings", meta = (ClampMin = "0.0"))
+	float HistoryWeightScale = 1.0f;
+
+	// Multiplier for prediction score contribution. Normalized with history weight after scaling.
+	UPROPERTY(BlueprintReadWrite, EditAnywhere, Category = "Settings", meta = (ClampMin = "0.0"))
+	float PredictionWeightScale = 1.0f;
 
-	bool IsInitialized() const { return !!Database; }
-	bool HasBiasWeights() const { return !!BiasWeights && BiasWeights->IsInitialized(); }
+	bool operator==(const FPoseSearchChannelDynamicWeightParams& Rhs) const
+	{
+		return (ChannelWeightScale == Rhs.ChannelWeightScale)
+			&& (HistoryWeightScale == Rhs.HistoryWeightScale)
+			&& (PredictionWeightScale == Rhs.PredictionWeightScale);
+	}
+	bool operator!=(const FPoseSearchChannelDynamicWeightParams& Rhs) const { return !(*this == Rhs); }
 };
 
-/** Per-sequence bias weights for influencing the results of a pose search query */
-UCLASS(BlueprintType, Category = "Animation|Pose Search")
-class POSESEARCH_API UPoseSearchSequenceBiasWeightMetaData : public UAnimMetaData
+USTRUCT(BlueprintType, Category = "Animation|Pose Search")
+struct POSESEARCH_API FPoseSearchWeightParams
+{
+	GENERATED_BODY()
+
+	UPROPERTY(BlueprintReadWrite, EditAnywhere, Category = "Settings")
+	FPoseSearchChannelWeightParams PoseWeight;
+
+	UPROPERTY(BlueprintReadWrite, EditAnywhere, Category = "Settings")
+	FPoseSearchChannelWeightParams TrajectoryWeight;
+};
+
+USTRUCT(BlueprintType, Category="Animation|Pose Search")
+struct POSESEARCH_API FPoseSearchDynamicWeightParams
 {
 	GENERATED_BODY()
+
+	UPROPERTY(BlueprintReadWrite, EditAnywhere, Category = "Settings")
+	FPoseSearchChannelDynamicWeightParams PoseDynamicWeights;
+
+	UPROPERTY(BlueprintReadWrite, EditAnywhere, Category = "Settings")
+	FPoseSearchChannelDynamicWeightParams TrajectoryDynamicWeights;
+
+	UPROPERTY(BlueprintReadWrite, EditAnywhere, Category = "Settings")
+	bool bDebugDisableWeights = false;
+
+	bool operator==(const FPoseSearchDynamicWeightParams& Rhs) const
+	{
+		return (PoseDynamicWeights == Rhs.PoseDynamicWeights)
+			&& (TrajectoryDynamicWeights == Rhs.TrajectoryDynamicWeights)
+			&& (bDebugDisableWeights == Rhs.bDebugDisableWeights);
+	}
+	bool operator!=(const FPoseSearchDynamicWeightParams& Rhs) const { return !(*this == Rhs); }
+};
+
+USTRUCT()
+struct POSESEARCH_API FPoseSearchWeights
+{
+	GENERATED_BODY()
+
+	UPROPERTY(Transient)
+	TArray<float> Weights;
+
+	bool IsInitialized() const { return !Weights.IsEmpty(); }
+	void Init(const FPoseSearchWeightParams& WeightParams, const UPoseSearchSchema* Schema, const FPoseSearchDynamicWeightParams& RuntimeParams);
+};
+
+USTRUCT()
+struct POSESEARCH_API FPoseSearchWeightsContext
+{
+	GENERATED_BODY()
+
 public:
+	// Check if the database or runtime weight parameters have changed and then computes and caches new group weights
+	void Update(const FPoseSearchDynamicWeightParams& DynamicWeights, const UPoseSearchDatabase * Database);
 
-	UPROPERTY(EditAnywhere, Category = "Settings")
-	FPoseSearchBiasWeightParams BiasWeights;
+	const FPoseSearchWeights* GetGroupWeights (int32 WeightsGroupIdx) const;
+	
+private:
+	UPROPERTY(Transient)
+	TWeakObjectPtr<const UPoseSearchDatabase> Database = nullptr;
+
+	UPROPERTY(Transient)
+	FPoseSearchDynamicWeightParams DynamicWeights;
+
+	UPROPERTY(Transient)
+	TArray<FPoseSearchWeights> ComputedGroupWeights;
 };
 
+
 /** An entry in a UPoseSearchDatabase. */
 USTRUCT(BlueprintType, Category = "Animation|Pose Search")
 struct POSESEARCH_API FPoseSearchDatabaseSequence
@@ -423,7 +473,7 @@ struct POSESEARCH_API FPoseSearchDatabaseSequence
 	GENERATED_BODY()
 
 	UPROPERTY(EditAnywhere, Category="Sequence")
-	UAnimSequence* Sequence = nullptr;
+	TObjectPtr<UAnimSequence> Sequence = nullptr;
 
 	UPROPERTY(EditAnywhere, Category="Sequence")
 	FFloatInterval SamplingRange = FFloatInterval(0.0f, 0.0f);
@@ -436,7 +486,7 @@ struct POSESEARCH_API FPoseSearchDatabaseSequence
 	// for one shot anims with past sampling. When past sampling is used without a lead in sequence,
 	// the sampling range of the main sequence will be clamped if necessary.
 	UPROPERTY(EditAnywhere, Category="Sequence")
-	UAnimSequence* LeadInSequence = nullptr;
+	TObjectPtr<UAnimSequence> LeadInSequence = nullptr;
 
 	UPROPERTY(EditAnywhere, Category="Sequence")
 	bool bLoopLeadInAnimation = false;
@@ -446,7 +496,7 @@ struct POSESEARCH_API FPoseSearchDatabaseSequence
 	// for one shot anims with future sampling. When future sampling is used without a follow up sequence,
 	// the sampling range of the main sequence will be clamped if necessary.
 	UPROPERTY(EditAnywhere, Category="Sequence")
-	UAnimSequence* FollowUpSequence = nullptr;
+	TObjectPtr<UAnimSequence> FollowUpSequence = nullptr;
 
 	UPROPERTY(EditAnywhere, Category="Sequence")
 	bool bLoopFollowUpAnimation = false;
@@ -456,9 +506,6 @@ struct POSESEARCH_API FPoseSearchDatabaseSequence
 
 	UPROPERTY()
 	int32 NumPoses = 0;
-
-	UPROPERTY()
-	FPoseSearchBiasWeights BiasWeights;
 };
 
 /** A data asset for indexing a collection of animation sequences. */
@@ -471,12 +518,15 @@ public:
 	UPROPERTY(EditAnywhere, BlueprintReadOnly, Category="Database")
 	const UPoseSearchSchema* Schema;
 
-	UPROPERTY(EditAnywhere, Category="Database")
-	TArray<FPoseSearchDatabaseSequence> Sequences;
+	UPROPERTY(EditAnywhere, Category = "Database")
+	FPoseSearchWeightParams Weights;
 
-	UPROPERTY(EditAnywhere, Category="Database")
+	UPROPERTY(EditAnywhere, Category = "Database")
 	FPoseSearchExtrapolationParameters ExtrapolationParameters;
 
+	UPROPERTY(EditAnywhere, Category="Database")
+	TArray<FPoseSearchDatabaseSequence> Sequences;
+
 	UPROPERTY()
 	FPoseSearchIndex SearchIndex;
 
@@ -508,7 +558,7 @@ public:
 	void Reset();
 	void ResetFeatures();
 
-	const UPoseSearchSchema* GetSchema() const { return Schema; }
+	const UPoseSearchSchema* GetSchema() const { return Schema.Get(); }
 
 	TArrayView<const float> GetValues() const { return Values; }
 	TArrayView<const float> GetNormalizedValues() const { return ValuesNormalized; }
@@ -540,7 +590,7 @@ private:
 	void BuildFromTrajectoryDistanceBased(const FTrajectorySampleRange& Trajectory);
 
 	UPROPERTY(Transient)
-	const UPoseSearchSchema* Schema = nullptr;
+	TWeakObjectPtr<const UPoseSearchSchema> Schema = nullptr;
 
 	TArray<float> Values;
 	TArray<float> ValuesNormalized;
@@ -706,7 +756,7 @@ POSESEARCH_API void Draw(const FDebugDrawParams& DrawParams);
 * @param Sequence			The input sequence create a search index for
 * @param SequenceMetaData	The input sequence indexing info and output search index
 * 
-* @return					Whether the index was built successfully
+* @return Whether the index was built successfully
 */
 POSESEARCH_API bool BuildIndex(const UAnimSequence* Sequence, UPoseSearchSequenceMetaData* SequenceMetaData);
 
@@ -716,7 +766,7 @@ POSESEARCH_API bool BuildIndex(const UAnimSequence* Sequence, UPoseSearchSequenc
 * 
 * @param Database	The input collection of animations and output search index
 * 
-* @return			Whether the index was built successfully
+* @return Whether the index was built successfully
 */
 POSESEARCH_API bool BuildIndex(UPoseSearchDatabase* Database);
 
@@ -724,12 +774,12 @@ POSESEARCH_API bool BuildIndex(UPoseSearchDatabase* Database);
 /**
 * Performs a pose search on a sequence's UPoseSearchSequenceMetaData.
 *
-* @param Sequence			The sequence to search within. Must contain a UPoseSearchSequenceMetaData object.
-* @param Query				The pose query to search for. To build a query, see FFeatureVectorBuilder. Must have been built using the same schema as the UPoseSearchSequenceMetaData.
-* @param BiasWeightsContext	Optional bias weight context used to influence pose search query results
-* @param DrawParams			Visualization options
+* @param Sequence		The sequence to search within. Must contain a UPoseSearchSequenceMetaData object.
+* @param Query			The pose query to search for. To build a query, see FFeatureVectorBuilder. Must have been built using the same schema as the UPoseSearchSequenceMetaData.
+* @param WeightsContext	Optional weights context used to influence pose search query results
+* @param DrawParams		Visualization options
 * 
-* @return					The pose in the sequence that most closely matches the Query.
+* @return The pose in the sequence that most closely matches the Query.
 */
 POSESEARCH_API FSearchResult Search(const UAnimSequenceBase* Sequence, TArrayView<const float> Query, FDebugDrawParams DrawParams = FDebugDrawParams());
 
@@ -737,26 +787,26 @@ POSESEARCH_API FSearchResult Search(const UAnimSequenceBase* Sequence, TArrayVie
 /**
 * Performs a pose search on a UPoseSearchDatabase.
 *
-* @param Database			The database to search within
-* @param Query				The pose query to search for. To build a query, see FFeatureVectorBuilder. Must have been built using the same schema as the Database.
-* @param BiasWeightsContext	Optional bias weight context used to influence pose search query results
-* @param DrawParams			Visualization options
+* @param Database		The database to search within
+* @param Query			The pose query to search for. To build a query, see FFeatureVectorBuilder. Must have been built using the same schema as the Database.
+* @param WeightsContext	Optional weights context used to influence pose search query results
+* @param DrawParams		Visualization options
 * 
-* @return					The pose in the database that most closely matches the Query.
+* @return The pose in the database that most closely matches the Query.
 */
-POSESEARCH_API FDbSearchResult Search(const UPoseSearchDatabase* Database, TArrayView<const float> Query, const FPoseSearchBiasWeightsContext* BiasWeightsContext = nullptr, FDebugDrawParams DrawParams = FDebugDrawParams());
+POSESEARCH_API FDbSearchResult Search(const UPoseSearchDatabase* Database, TArrayView<const float> Query, const FPoseSearchWeightsContext* WeightsContext = nullptr, FDebugDrawParams DrawParams = FDebugDrawParams());
 
 
 /**
 * Evaluate pose comparison metric between a pose in the search index and an input query
 * 
-* @param SearchIndex		The search index containing the pose to compare to the query
-* @param PoseIdx			The index of the pose in the search index to compare to the query
-* @param Query				The query to compare against. Must have been built using the same schema as the SearchIndex.
-* @param BiasWeightsContext Optional bias weight context used to influence pose comparison
+* @param SearchIndex	The search index containing the pose to compare to the query
+* @param PoseIdx		The index of the pose in the search index to compare to the query
+* @param Query			The query to compare against. Must have been built using the same schema as the SearchIndex.
+* @param WeightsContext	Optional weights context used to influence pose comparison
 * 
-* @return					Dissimilarity between the two poses
+* @return Dissimilarity between the two poses
 */
-POSESEARCH_API float ComparePoses(const FPoseSearchIndex& SearchIndex, int32 PoseIdx, TArrayView<const float> Query, const FPoseSearchBiasWeightsContext* BiasWeightsContext = nullptr);
+POSESEARCH_API float ComparePoses(const FPoseSearchIndex& SearchIndex, int32 PoseIdx, TArrayView<const float> Query, const FPoseSearchWeightsContext* WeightsContext = nullptr);
 
 }} // namespace UE::PoseSearch
\ No newline at end of file
diff --git a/Engine/Plugins/Experimental/Animation/PoseSearch/Source/Runtime/Public/PoseSearch/PoseSearchLibrary.h b/Engine/Plugins/Experimental/Animation/PoseSearch/Source/Runtime/Public/PoseSearch/PoseSearchLibrary.h
--- a/Engine/Plugins/Experimental/Animation/PoseSearch/Source/Runtime/Public/PoseSearch/PoseSearchLibrary.h
+++ b/Engine/Plugins/Experimental/Animation/PoseSearch/Source/Runtime/Public/PoseSearch/PoseSearchLibrary.h
@@ -26,9 +26,9 @@ struct POSESEARCH_API FMotionMatchingSettings
 {
 	GENERATED_BODY()
 
-	// Dynamic biasing weights for influencing pose search queries
+	// Dynamic weights for influencing pose selection
 	UPROPERTY(EditAnywhere, BlueprintReadWrite, Category=Settings, meta=(PinHiddenByDefault))
-	FPoseSearchBiasWeightParams BiasWeights;
+	FPoseSearchDynamicWeightParams Weights;
 
 	// Time in seconds to blend out to the new pose. Uses inertial blending and requires an Inertialization node after this node.
 	UPROPERTY(EditAnywhere, BlueprintReadWrite, Category=Settings, meta=(ClampMin="0"))
@@ -74,6 +74,10 @@ struct POSESEARCH_API FMotionMatchingState
 	UPROPERTY(Transient)
 	FPoseSearchFeatureVectorBuilder ComposedQuery;
 
+	// Precomputed runtime weights
+	UPROPERTY(Transient)
+	FPoseSearchWeightsContext WeightsContext;
+
 	// When the database changes, the search parameters are reset
 	UPROPERTY(Transient)
 	TWeakObjectPtr<const UPoseSearchDatabase> CurrentDatabase = nullptr;
diff --git a/Engine/Source/Runtime/AnimGraphRuntime/Public/AnimNodes/AnimNode_SequenceEvaluator.h b/Engine/Source/Runtime/AnimGraphRuntime/Public/AnimNodes/AnimNode_SequenceEvaluator.h
--- a/Engine/Source/Runtime/AnimGraphRuntime/Public/AnimNodes/AnimNode_SequenceEvaluator.h
+++ b/Engine/Source/Runtime/AnimGraphRuntime/Public/AnimNodes/AnimNode_SequenceEvaluator.h
@@ -48,9 +48,9 @@ public:
 	// End of FAnimNode_Base interface
 
 	// FAnimNode_AssetPlayerBase Interface
-	virtual float GetAccumulatedTime() const {return GetExplicitTime();}
-	virtual void SetAccumulatedTime(float NewTime) { SetExplicitTime(NewTime); }
-	virtual UAnimationAsset* GetAnimAsset() const { return GetSequence(); }
+	virtual float GetAccumulatedTime() const override {return GetExplicitTime();}
+	virtual void SetAccumulatedTime(float NewTime) override { SetExplicitTime(NewTime); }
+	virtual UAnimationAsset* GetAnimAsset() const override { return GetSequence(); }
 	// End of FAnimNode_AssetPlayerBase Interface
 
 	void SetExplicitPreviousTime(float PreviousTime) { InternalTimeAccumulator = PreviousTime; }
