PoseSearch: Tracing for dynamic weight params

diff --git a/Engine/Plugins/Experimental/Animation/PoseSearch/Source/Editor/Private/PoseSearchDebugger.cpp b/Engine/Plugins/Experimental/Animation/PoseSearch/Source/Editor/Private/PoseSearchDebugger.cpp
--- a/Engine/Plugins/Experimental/Animation/PoseSearch/Source/Editor/Private/PoseSearchDebugger.cpp
+++ b/Engine/Plugins/Experimental/Animation/PoseSearch/Source/Editor/Private/PoseSearchDebugger.cpp
@@ -520,19 +520,22 @@ void SDebuggerDatabaseView::UpdateRows(const FTraceMotionMatchingStateMessage& S
 	}
 	check(ActiveView.Rows.Num() == 1);
 	
-	FPoseSearchDynamicWeightParams StateDynamicWeightParams;
 	FPoseSearchWeightsContext StateWeights;
-	StateWeights.Update(StateDynamicWeightParams, &Database);
+	StateWeights.Update(State.Weights, &Database);
 
 	FPoseSearchWeightsContext PoseWeights;
-	FPoseSearchDynamicWeightParams WeightParams = StateDynamicWeightParams;
-	WeightParams.TrajectoryDynamicWeights.ChannelWeightScale = 0.0f;
-	PoseWeights.Update(WeightParams, &Database);
+	{
+		FPoseSearchDynamicWeightParams WeightParams = State.Weights;
+		WeightParams.TrajectoryDynamicWeights.ChannelWeightScale = 0.0f;
+		PoseWeights.Update(WeightParams, &Database);
+	}
 
 	FPoseSearchWeightsContext TrajectoryWeights;
-	WeightParams = StateDynamicWeightParams;
-	WeightParams.PoseDynamicWeights.ChannelWeightScale = 0.0f;
-	TrajectoryWeights.Update(WeightParams, &Database);
+	{
+		FPoseSearchDynamicWeightParams WeightParams = State.Weights;
+		WeightParams.PoseDynamicWeights.ChannelWeightScale = 0.0f;
+		TrajectoryWeights.Update(WeightParams, &Database);
+	}
 	
 	const FPoseSearchIndex& SearchIndex = Database.SearchIndex;
 
diff --git a/Engine/Plugins/Experimental/Animation/PoseSearch/Source/Editor/Private/Trace/PoseSearchTraceAnalyzer.cpp b/Engine/Plugins/Experimental/Animation/PoseSearch/Source/Editor/Private/Trace/PoseSearchTraceAnalyzer.cpp
--- a/Engine/Plugins/Experimental/Animation/PoseSearch/Source/Editor/Private/Trace/PoseSearchTraceAnalyzer.cpp
+++ b/Engine/Plugins/Experimental/Animation/PoseSearch/Source/Editor/Private/Trace/PoseSearchTraceAnalyzer.cpp
@@ -41,22 +41,27 @@ bool FTraceAnalyzer::OnEvent(uint16 RouteId, EStyle Style, const FOnEventContext
 		case RouteId_MotionMatchingState:
 			{
 				FTraceMotionMatchingStateMessage Message;
-				Message.ElapsedPoseJumpTime = EventData.GetValue<float>("ElapsedPoseJumpTime");
-				// Cast back to our flag type
+				Message.DatabaseId = EventData.GetValue<uint64>("DatabaseId");
 				Message.Flags = static_cast<FTraceMotionMatchingState::EFlags>(EventData.GetValue<uint32>("Flags"));
 				Message.DbPoseIdx = EventData.GetValue<int32>("DbPoseIdx");
+				Message.ElapsedPoseJumpTime = EventData.GetValue<float>("ElapsedPoseJumpTime");
+				Message.QueryVector = EventData.GetArrayView<float>("QueryVector");
+                Message.QueryVectorNormalized = EventData.GetArrayView<float>("QueryVectorNormalized");
 
-				// Gather array types
-				TArrayView<const float> View = EventData.GetArrayView<float>("QueryVector");
-				Message.QueryVector = View;
-				
-				View = EventData.GetArrayView<float>("QueryVectorNormalized");
-                Message.QueryVectorNormalized = View;
-				
-				View = EventData.GetArrayView<float>("BiasWeights");
-				Message.BiasWeights = View;
-				
-				Message.DatabaseId = EventData.GetValue<uint64>("DatabaseId");
+				TArrayView<const float> ChannelWeightScales = EventData.GetArrayView<float>("ChannelWeightScales");
+				TArrayView<const float> HistoryWeightScales = EventData.GetArrayView<float>("HistoryWeightScales");
+				TArrayView<const float> PredictionWeightScales = EventData.GetArrayView<float>("PredictionWeightScales");
+
+				if ((ChannelWeightScales.Num() == 2) && (HistoryWeightScales.Num() == 2) && (PredictionWeightScales.Num() == 2))
+				{
+					Message.Weights.PoseDynamicWeights.ChannelWeightScale = ChannelWeightScales[0];
+					Message.Weights.TrajectoryDynamicWeights.ChannelWeightScale = ChannelWeightScales[1];
+					Message.Weights.PoseDynamicWeights.HistoryWeightScale = HistoryWeightScales[0];
+					Message.Weights.TrajectoryDynamicWeights.HistoryWeightScale = HistoryWeightScales[1];
+					Message.Weights.PoseDynamicWeights.PredictionWeightScale = PredictionWeightScales[0];
+					Message.Weights.TrajectoryDynamicWeights.PredictionWeightScale = PredictionWeightScales[1];
+					Message.Weights.bDebugDisableWeights = EventData.GetValue<bool>("DebugDisableWeights");
+				}
 
 				// Common data
 				Message.NodeId = NodeId;
diff --git a/Engine/Plugins/Experimental/Animation/PoseSearch/Source/Editor/Private/Trace/PoseSearchTraceProvider.h b/Engine/Plugins/Experimental/Animation/PoseSearch/Source/Editor/Private/Trace/PoseSearchTraceProvider.h
--- a/Engine/Plugins/Experimental/Animation/PoseSearch/Source/Editor/Private/Trace/PoseSearchTraceProvider.h
+++ b/Engine/Plugins/Experimental/Animation/PoseSearch/Source/Editor/Private/Trace/PoseSearchTraceProvider.h
@@ -27,7 +27,7 @@ struct FTraceMotionMatchingStateMessage : FTraceMessage
 	float ElapsedPoseJumpTime;
 	TArray<float> QueryVector;
 	TArray<float> QueryVectorNormalized;
-	TArray<float> BiasWeights;
+	FPoseSearchDynamicWeightParams Weights;
 	int32 DbPoseIdx;
 	uint64 DatabaseId;
 };
diff --git a/Engine/Plugins/Experimental/Animation/PoseSearch/Source/Runtime/Private/PoseSearchLibrary.cpp b/Engine/Plugins/Experimental/Animation/PoseSearch/Source/Runtime/Private/PoseSearchLibrary.cpp
--- a/Engine/Plugins/Experimental/Animation/PoseSearch/Source/Runtime/Private/PoseSearchLibrary.cpp
+++ b/Engine/Plugins/Experimental/Animation/PoseSearch/Source/Runtime/Private/PoseSearchLibrary.cpp
@@ -204,7 +204,7 @@ void UpdateMotionMatchingState(const FAnimationUpdateContext& Context
 		// @TODO: Change this to only be the previous query, not persistently updated (i.e. if throttled)?
 		TraceState.QueryVector = InOutMotionMatchingState.ComposedQuery.GetValues();
 		TraceState.QueryVectorNormalized = InOutMotionMatchingState.ComposedQuery.GetNormalizedValues();
-		TraceState.BiasWeights = InOutMotionMatchingState.WeightsContext.GetGroupWeights(0)->Weights;
+		TraceState.Weights = Settings.Weights;
 		TraceState.DbPoseIdx = InOutMotionMatchingState.DbPoseIdx;
 		TraceState.DatabaseId = FObjectTrace::GetObjectId(Database);
 		UE_TRACE_POSE_SEARCH_MOTION_MATCHING_STATE(Context, TraceState)
diff --git a/Engine/Plugins/Experimental/Animation/PoseSearch/Source/Runtime/Private/Trace/PoseSearchTraceLogger.cpp b/Engine/Plugins/Experimental/Animation/PoseSearch/Source/Runtime/Private/Trace/PoseSearchTraceLogger.cpp
--- a/Engine/Plugins/Experimental/Animation/PoseSearch/Source/Runtime/Private/Trace/PoseSearchTraceLogger.cpp
+++ b/Engine/Plugins/Experimental/Animation/PoseSearch/Source/Runtime/Private/Trace/PoseSearchTraceLogger.cpp
@@ -20,7 +20,10 @@ UE_TRACE_EVENT_BEGIN(PoseSearch, MotionMatchingState)
 	UE_TRACE_EVENT_FIELD(int32, DbPoseIdx)
 	UE_TRACE_EVENT_FIELD(float[], QueryVector)
 	UE_TRACE_EVENT_FIELD(float[], QueryVectorNormalized)
-	UE_TRACE_EVENT_FIELD(float[], BiasWeights)
+	UE_TRACE_EVENT_FIELD(float[], ChannelWeightScales)
+	UE_TRACE_EVENT_FIELD(float[], HistoryWeightScales)
+	UE_TRACE_EVENT_FIELD(float[], PredictionWeightScales)
+	UE_TRACE_EVENT_FIELD(bool, DebugDisableWeights)
 UE_TRACE_EVENT_END()
 
 namespace UE { namespace PoseSearch {
@@ -56,6 +59,23 @@ void FTraceMotionMatchingState::Output(const FAnimationBaseContext& InContext, c
 	TRACE_OBJECT(AnimInstance);
 	UObject* SkeletalMeshComponent = AnimInstance->GetOuter();
 
+	float ChannelWeightScales[] =
+	{
+		State.Weights.PoseDynamicWeights.ChannelWeightScale,
+		State.Weights.TrajectoryDynamicWeights.ChannelWeightScale
+	};
+
+	float HistoryWeightScales[] =
+	{
+		State.Weights.PoseDynamicWeights.HistoryWeightScale,
+		State.Weights.TrajectoryDynamicWeights.HistoryWeightScale
+	};
+	float PredictionWeightScales[] =
+	{
+		State.Weights.PoseDynamicWeights.PredictionWeightScale,
+		State.Weights.TrajectoryDynamicWeights.PredictionWeightScale
+	};
+
 	UE_TRACE_LOG(PoseSearch, MotionMatchingState, PoseSearchChannel)
 		<< MotionMatchingState.Cycle(FPlatformTime::Cycles64())
 		<< MotionMatchingState.FrameCounter(FObjectTrace::GetObjectWorldTickCounter(AnimInstance))
@@ -63,13 +83,15 @@ void FTraceMotionMatchingState::Output(const FAnimationBaseContext& InContext, c
 		<< MotionMatchingState.SkeletalMeshComponentId(FObjectTrace::GetObjectId(SkeletalMeshComponent))
 		<< MotionMatchingState.NodeId(InContext.GetCurrentNodeId())
 		<< MotionMatchingState.ElapsedPoseJumpTime(State.ElapsedPoseJumpTime)
-		// Cast to uint32 for event message
 		<< MotionMatchingState.Flags(static_cast<uint32>(State.Flags))
+		<< MotionMatchingState.DatabaseId(State.DatabaseId)
+		<< MotionMatchingState.DbPoseIdx(State.DbPoseIdx)
 		<< MotionMatchingState.QueryVector(State.QueryVector.GetData(), State.QueryVector.Num())
 		<< MotionMatchingState.QueryVectorNormalized(State.QueryVectorNormalized.GetData(), State.QueryVectorNormalized.Num())
-		<< MotionMatchingState.BiasWeights(State.BiasWeights.GetData(), State.BiasWeights.Num())
-		<< MotionMatchingState.DbPoseIdx(State.DbPoseIdx)
-		<< MotionMatchingState.DatabaseId(State.DatabaseId);
+		<< MotionMatchingState.ChannelWeightScales(ChannelWeightScales, UE_ARRAY_COUNT(ChannelWeightScales))
+		<< MotionMatchingState.HistoryWeightScales(HistoryWeightScales, UE_ARRAY_COUNT(HistoryWeightScales))
+		<< MotionMatchingState.PredictionWeightScales(PredictionWeightScales, UE_ARRAY_COUNT(PredictionWeightScales))
+		<< MotionMatchingState.DebugDisableWeights(State.Weights.bDebugDisableWeights);
 }
 
 }}
diff --git a/Engine/Plugins/Experimental/Animation/PoseSearch/Source/Runtime/Private/Trace/PoseSearchTraceLogger.h b/Engine/Plugins/Experimental/Animation/PoseSearch/Source/Runtime/Private/Trace/PoseSearchTraceLogger.h
--- a/Engine/Plugins/Experimental/Animation/PoseSearch/Source/Runtime/Private/Trace/PoseSearchTraceLogger.h
+++ b/Engine/Plugins/Experimental/Animation/PoseSearch/Source/Runtime/Private/Trace/PoseSearchTraceLogger.h
@@ -2,6 +2,7 @@
 
 #pragma once
 
+#include "PoseSearch/PoseSearch.h"
 #include "ObjectTrace.h"
 
 // Enable this if object tracing is enabled, mimics animation tracing
@@ -39,12 +40,12 @@ struct POSESEARCH_API FTraceMotionMatchingState
 	/** Storage container for state booleans */
 	EFlags Flags = EFlags::None;
 
-	/** Feature value array of the query data */
+	/** Search vectors in normalized and unnormalized forms */
 	TArrayView<const float> QueryVector;
 	TArrayView<const float> QueryVectorNormalized;
 
-	/** Weight bias params for motion matching queries */
-	TArrayView<const float> BiasWeights;
+	/** Runtime weights */
+	FPoseSearchDynamicWeightParams Weights;
 
 	/** Index of the pose in our database */
 	int32 DbPoseIdx = 0;
