AnimNode_MotionMatching no longer gets stuck when sequences don't have follow-ups in the database
- Moved logic to select the follow-up sequence before the search results are processed
- If there's no way to continue on the current sequence or the follow-up, the new search result is always used.

diff --git a/Engine/Plugins/Experimental/Animation/PoseSearch/Source/Runtime/Private/PoseSearch.cpp b/Engine/Plugins/Experimental/Animation/PoseSearch/Source/Runtime/Private/PoseSearch.cpp
--- a/Engine/Plugins/Experimental/Animation/PoseSearch/Source/Runtime/Private/PoseSearch.cpp
+++ b/Engine/Plugins/Experimental/Animation/PoseSearch/Source/Runtime/Private/PoseSearch.cpp
@@ -959,7 +959,7 @@ int32 UPoseSearchDatabase::GetPoseIndexFromAssetTime(int32 DbSequenceIdx, float
 	FFloatInterval Range = UE::PoseSearch::GetEffectiveSamplingRange(DbSequence.Sequence, DbSequence.SamplingRange);
 	if (Range.Contains(AssetTime))
 	{
-		int32 PoseOffset = FMath::RoundToInt(Schema->SampleRate * (AssetTime - Range.Min));		
+		int32 PoseOffset = FMath::FloorToInt(Schema->SampleRate * (AssetTime - Range.Min));
 		if (PoseOffset >= DbSequence.NumPoses)
 		{
 			if (DbSequence.bLoopAnimation)
@@ -979,6 +979,13 @@ int32 UPoseSearchDatabase::GetPoseIndexFromAssetTime(int32 DbSequenceIdx, float
 	return INDEX_NONE;
 }
 
+FFloatInterval UPoseSearchDatabase::GetEffectiveSamplingRange(int32 DbSequenceIdx) const
+{
+	const FPoseSearchDatabaseSequence& DbSequence = Sequences[DbSequenceIdx];
+	FFloatInterval Range = UE::PoseSearch::GetEffectiveSamplingRange(DbSequence.Sequence, DbSequence.SamplingRange);
+	return Range;
+}
+
 float UPoseSearchDatabase::GetSequenceLength(int32 DbSequenceIdx) const
 {
 	return Sequences[DbSequenceIdx].Sequence->GetPlayLength();
@@ -3328,7 +3335,7 @@ bool BuildIndex(UPoseSearchDatabase* Database)
 	return true;
 }
 
-static FSearchResult Search(const FPoseSearchIndex& SearchIndex, TArrayView<const float> Query, const FPoseSearchWeightsContext* WeightsContext = nullptr)
+static FSearchResult Search(const FPoseSearchIndex& SearchIndex, TArrayView<const float> Query, const FPoseSearchWeightsContext* WeightsContext = nullptr, const TSet<int32>& ExcludedIndices = TSet<int32>())
 {
 	FSearchResult Result;
 	if (!ensure(SearchIndex.IsValid()))
@@ -3346,6 +3353,11 @@ static FSearchResult Search(const FPoseSearchIndex& SearchIndex, TArrayView<cons
 
 	for (int32 PoseIdx = 0; PoseIdx != SearchIndex.NumPoses; ++PoseIdx)
 	{
+		if (ExcludedIndices.Contains(PoseIdx))
+		{
+			continue;
+		}
+
 		const float PoseDissimilarity = ComparePoses(SearchIndex, PoseIdx, Query, WeightsContext);
 		
 		if (PoseDissimilarity < BestPoseDissimilarity)
@@ -3392,7 +3404,11 @@ FSearchResult Search(const UAnimSequenceBase* Sequence, TArrayView<const float>
 	return Result;
 }
 
-FDbSearchResult Search(const UPoseSearchDatabase* Database, TArrayView<const float> Query, const FPoseSearchWeightsContext* WeightsContext, FDebugDrawParams DebugDrawParams)
+FDbSearchResult Search(
+	const UPoseSearchDatabase* Database, 
+	TArrayView<const float> Query, 
+	const FPoseSearchWeightsContext* WeightsContext, 
+	const float EndTimeToExclude, FDebugDrawParams DebugDrawParams)
 {
 	if (!ensure(Database && Database->IsValidForSearch()))
 	{
@@ -3401,7 +3417,23 @@ FDbSearchResult Search(const UPoseSearchDatabase* Database, TArrayView<const flo
 
 	const FPoseSearchIndex& SearchIndex = Database->SearchIndex;
 
-	FDbSearchResult Result = Search(SearchIndex, Query, WeightsContext);
+	const int32 NumSequenceIndicesToExclude = Database->Schema->SampleRate * EndTimeToExclude;
+	TSet<int32> ExcludedIndices;
+	if (NumSequenceIndicesToExclude > 0)
+	{
+		for (const FPoseSearchDatabaseSequence& Sequence : Database->Sequences)
+		{
+			// leaving at least one sample per sequence
+			const int32 ExclusionStartIndex = Sequence.FirstPoseIdx + FMath::Max(1, Sequence.NumPoses - NumSequenceIndicesToExclude);
+			const int32 ExclusionEndIndex = Sequence.FirstPoseIdx + Sequence.NumPoses;
+			for (int32 ExcludedIndex = ExclusionStartIndex; ExcludedIndex < ExclusionEndIndex; ++ExcludedIndex)
+			{
+				ExcludedIndices.Add(ExcludedIndex);
+			}
+		}
+	}
+	
+	FDbSearchResult Result = Search(SearchIndex, Query, WeightsContext, ExcludedIndices);
 	if (!Result.IsValid())
 	{
 		return FDbSearchResult();
diff --git a/Engine/Plugins/Experimental/Animation/PoseSearch/Source/Runtime/Private/PoseSearchLibrary.cpp b/Engine/Plugins/Experimental/Animation/PoseSearch/Source/Runtime/Private/PoseSearchLibrary.cpp
--- a/Engine/Plugins/Experimental/Animation/PoseSearch/Source/Runtime/Private/PoseSearchLibrary.cpp
+++ b/Engine/Plugins/Experimental/Animation/PoseSearch/Source/Runtime/Private/PoseSearchLibrary.cpp
@@ -81,21 +81,20 @@ void UpdateMotionMatchingState(const FAnimationUpdateContext& Context
 		}
 
 		// Step the pose forward
-		if (InOutMotionMatchingState.DbPoseIdx != INDEX_NONE)
+		UE::PoseSearch::FMotionMatchingContinuityParams ContinuityParameters = InOutMotionMatchingState.ComputeContinuityParameters(Context);
+		const bool bCanContinue = ContinuityParameters.IsValid();
+		if (bCanContinue)
 		{
-			check(InOutMotionMatchingState.DbSequenceIdx != INDEX_NONE);
-
-			// Determine roughly which pose we're playing from the database based on the sequence player time
-			InOutMotionMatchingState.DbPoseIdx = Database->GetPoseIndexFromAssetTime(InOutMotionMatchingState.DbSequenceIdx, InOutMotionMatchingState.AssetPlayerTime);
-
-			// Overwrite query feature vector with the stepped pose from the database
-			if (InOutMotionMatchingState.DbPoseIdx != INDEX_NONE)
-			{
-				InOutMotionMatchingState.ComposedQuery.CopyFromSearchIndex(Database->SearchIndex, InOutMotionMatchingState.DbPoseIdx);
-			}
+			InOutMotionMatchingState.DbPoseIdx = ContinuityParameters.Result.PoseIdx;
+			InOutMotionMatchingState.DbSequenceIdx = ContinuityParameters.Result.DbSequenceIdx;
 		}
 
-		if (InOutMotionMatchingState.DbPoseIdx == INDEX_NONE)
+		// build query
+		if (InOutMotionMatchingState.DbPoseIdx != INDEX_NONE)
+		{
+			InOutMotionMatchingState.ComposedQuery.CopyFromSearchIndex(Database->SearchIndex, InOutMotionMatchingState.DbPoseIdx);
+		}
+		else
 		{
 			UE::PoseSearch::IPoseHistoryProvider* PoseHistoryProvider = Context.GetMessage<UE::PoseSearch::IPoseHistoryProvider>();
 			if (PoseHistoryProvider)
@@ -127,63 +126,55 @@ void UpdateMotionMatchingState(const FAnimationUpdateContext& Context
 		UE::PoseSearch::FDbSearchResult Result = UE::PoseSearch::Search(
 			Database,
 			InOutMotionMatchingState.ComposedQuery.GetNormalizedValues(),
-			&InOutMotionMatchingState.WeightsContext
+			&InOutMotionMatchingState.WeightsContext,
+			Settings.SequenceEndExlusionTime
 		);
-		if (Result.IsValid() && (InOutMotionMatchingState.ElapsedPoseJumpTime >= Settings.SearchThrottleTime))
+		if (Result.IsValid() && ((InOutMotionMatchingState.ElapsedPoseJumpTime >= Settings.SearchThrottleTime) || !bCanContinue))
 		{
-			const FPoseSearchDatabaseSequence& ResultDbSequence = Database->Sequences[Result.DbSequenceIdx];
-
-			// Consider the search result better if it is more similar to the query than the current pose we're playing back from the database
-			bool bBetterPose = Result.Dissimilarity * (1.0f + (Settings.MinPercentImprovement / 100.0f)) < CurrentDissimilarity;
-
-			// We'll ignore the candidate pose if it is too near to our current pose
-			bool bNearbyPose = false;
-			if (InOutMotionMatchingState.DbSequenceIdx == Result.DbSequenceIdx)
-			{
-				bNearbyPose = FMath::Abs(InOutMotionMatchingState.AssetPlayerTime - Result.TimeOffsetSeconds) < Settings.PoseJumpThresholdTime;
-				if (!bNearbyPose && ResultDbSequence.bLoopAnimation)
-				{
-					const float AssetLength = Database->GetSequenceLength(InOutMotionMatchingState.DbSequenceIdx);
-					bNearbyPose = FMath::Abs(AssetLength - InOutMotionMatchingState.AssetPlayerTime - Result.TimeOffsetSeconds) < Settings.PoseJumpThresholdTime;
-				}
-			}
-
-			// Start playback from the candidate pose if we determined it was a better option
-			if (bBetterPose && !bNearbyPose)
+			if (!bCanContinue)
 			{
 				InOutMotionMatchingState.JumpToPose(Result);
 				RequestInertialBlend(Context, Settings.BlendTime);
 				InOutMotionMatchingState.Flags |= EMotionMatchingFlags::JumpedToPose;
 			}
-		}
-
-		// Continue with the follow up sequence if we're finishing a one shot anim
-		if (!(InOutMotionMatchingState.Flags & EMotionMatchingFlags::JumpedToPose)
-			&& (InOutMotionMatchingState.DbSequenceIdx != INDEX_NONE) 
-			&& !Database->DoesSequenceLoop(InOutMotionMatchingState.DbSequenceIdx))
-		{
-			float AssetTimeAfterUpdate = InOutMotionMatchingState.AssetPlayerTime + DeltaTime;
-			const float AssetLength = Database->GetSequenceLength(InOutMotionMatchingState.DbSequenceIdx);
-			if (AssetTimeAfterUpdate > AssetLength)
+			else
 			{
-				const FPoseSearchDatabaseSequence& DbSequence = Database->Sequences[InOutMotionMatchingState.DbSequenceIdx];
-				int32 FollowUpDbSequenceIdx = Database->Sequences.IndexOfByPredicate([&](const FPoseSearchDatabaseSequence& Entry) { return Entry.Sequence == DbSequence.FollowUpSequence; });
-				if (FollowUpDbSequenceIdx != INDEX_NONE)
+				const FPoseSearchDatabaseSequence& ResultDbSequence = Database->Sequences[Result.DbSequenceIdx];
+
+				// Consider the search result better if it is more similar to the query than the current pose we're playing back from the database
+				bool bBetterPose = Result.Dissimilarity * (1.0f + (Settings.MinPercentImprovement / 100.0f)) < CurrentDissimilarity;
+
+				// We'll ignore the candidate pose if it is too near to our current pose
+				bool bNearbyPose = false;
+				if (InOutMotionMatchingState.DbSequenceIdx == Result.DbSequenceIdx)
 				{
-					const FPoseSearchDatabaseSequence& FollowUpDbSequence = Database->Sequences[FollowUpDbSequenceIdx];
-					float FollowUpAssetTime = AssetTimeAfterUpdate - AssetLength;
-					int32 FollowUpPoseIdx = Database->GetPoseIndexFromAssetTime(FollowUpDbSequenceIdx, FollowUpAssetTime);
-
-					UE::PoseSearch::FDbSearchResult FollowUpResult;
-					FollowUpResult.DbSequenceIdx = FollowUpDbSequenceIdx;
-					FollowUpResult.PoseIdx = FollowUpPoseIdx;
-					FollowUpResult.TimeOffsetSeconds = FollowUpAssetTime;
-					InOutMotionMatchingState.JumpToPose(FollowUpResult);
+					bNearbyPose = FMath::Abs(InOutMotionMatchingState.AssetPlayerTime - Result.TimeOffsetSeconds) < Settings.PoseJumpThresholdTime;
+					if (!bNearbyPose && ResultDbSequence.bLoopAnimation)
+					{
+						const float AssetLength = Database->GetSequenceLength(InOutMotionMatchingState.DbSequenceIdx);
+						bNearbyPose = FMath::Abs(AssetLength - InOutMotionMatchingState.AssetPlayerTime - Result.TimeOffsetSeconds) < Settings.PoseJumpThresholdTime;
+					}
+				}
+
+				// Start playback from the candidate pose if we determined it was a better option
+				if (bBetterPose && !bNearbyPose)
+				{
+					InOutMotionMatchingState.JumpToPose(Result);
 					RequestInertialBlend(Context, Settings.BlendTime);
 					InOutMotionMatchingState.Flags |= EMotionMatchingFlags::JumpedToPose;
 				}
 			}
 		}
+
+		// Continue with the follow up sequence if we're finishing a one shot anim
+		if (!(InOutMotionMatchingState.Flags & EMotionMatchingFlags::JumpedToPose)
+			&& bCanContinue
+			&& ContinuityParameters.bJumpRequired)
+		{
+			InOutMotionMatchingState.JumpToPose(ContinuityParameters.Result);
+			RequestInertialBlend(Context, Settings.BlendTime);
+			InOutMotionMatchingState.Flags |= EMotionMatchingFlags::JumpedToPose;
+		}
 	}
 
 	if (!(InOutMotionMatchingState.Flags & EMotionMatchingFlags::JumpedToPose))
@@ -249,4 +240,59 @@ void UPoseSearchLibrary::UpdateMotionMatchingForSequencePlayer(const FAnimUpdate
 	{
 		UE_LOG(LogPoseSearchLibrary, Warning, TEXT("UpdateMotionMatchingForSequencePlayer called with invalid context"));
 	}
-}
\ No newline at end of file
+}
+
+UE::PoseSearch::FMotionMatchingContinuityParams FMotionMatchingState::ComputeContinuityParameters(const FAnimationUpdateContext& Context) const
+{
+	UE::PoseSearch::FMotionMatchingContinuityParams ContinuityParameters;
+
+	if (DbPoseIdx != INDEX_NONE)
+	{
+		check(DbSequenceIdx != INDEX_NONE);
+
+		const FPoseSearchDatabaseSequence& DbSequence = CurrentDatabase->Sequences[DbSequenceIdx];
+		const float AssetLength = DbSequence.Sequence->GetPlayLength();
+		const float DeltaTime = Context.GetDeltaTime();
+
+		float ContinuityAssetTime = AssetPlayerTime;
+		ETypeAdvanceAnim ContinuityAdvanceType = FAnimationRuntime::AdvanceTime(
+			DbSequence.bLoopAnimation,
+			DeltaTime,
+			ContinuityAssetTime,
+			AssetLength);
+
+		if (ContinuityAdvanceType != ETAA_Finished)
+		{
+			// we can continue ticking the same sequence forward
+			ContinuityParameters.Result.PoseIdx = CurrentDatabase->GetPoseIndexFromAssetTime(DbSequenceIdx,
+																					  ContinuityAssetTime);
+			ContinuityParameters.Result.DbSequenceIdx = DbSequenceIdx;
+			ContinuityParameters.Result.TimeOffsetSeconds = AssetPlayerTime;
+		}
+		else
+		{
+			// check if there's a follow-up that can be used
+			int32 FollowUpDbSequenceIdx = CurrentDatabase->Sequences.IndexOfByPredicate(
+				[&](const FPoseSearchDatabaseSequence& Entry)
+			{
+				return Entry.Sequence == DbSequence.FollowUpSequence;
+			});
+			if (FollowUpDbSequenceIdx != INDEX_NONE)
+			{
+				const FPoseSearchDatabaseSequence& FollowUpDbSequence = CurrentDatabase->Sequences[FollowUpDbSequenceIdx];
+				const float FollowUpAssetTime = AssetPlayerTime + DeltaTime - AssetLength;
+				const int32 FollowUpPoseIdx = CurrentDatabase->GetPoseIndexFromAssetTime(FollowUpDbSequenceIdx,
+																						 FollowUpAssetTime);
+				const FFloatInterval SamplingRange = CurrentDatabase->GetEffectiveSamplingRange(FollowUpDbSequenceIdx);
+
+				ContinuityParameters.Result.PoseIdx = FollowUpPoseIdx;
+				ContinuityParameters.Result.DbSequenceIdx = FollowUpDbSequenceIdx;
+				ContinuityParameters.Result.TimeOffsetSeconds =
+					SamplingRange.Min + (CurrentDatabase->Schema->SamplingInterval *
+										 (ContinuityParameters.Result.PoseIdx - FollowUpDbSequence.FirstPoseIdx));
+				ContinuityParameters.bJumpRequired = true;
+			}
+		}
+	}
+	return ContinuityParameters;
+}
diff --git a/Engine/Plugins/Experimental/Animation/PoseSearch/Source/Runtime/Public/PoseSearch/PoseSearch.h b/Engine/Plugins/Experimental/Animation/PoseSearch/Source/Runtime/Public/PoseSearch/PoseSearch.h
--- a/Engine/Plugins/Experimental/Animation/PoseSearch/Source/Runtime/Public/PoseSearch/PoseSearch.h
+++ b/Engine/Plugins/Experimental/Animation/PoseSearch/Source/Runtime/Public/PoseSearch/PoseSearch.h
@@ -534,6 +534,7 @@ public:
 	int32 GetPoseIndexFromAssetTime(int32 DbSequenceIdx, float AssetTime) const;
 	float GetSequenceLength(int32 DbSequenceIdx) const;
 	bool DoesSequenceLoop(int32 DbSequenceIdx) const;
+	FFloatInterval GetEffectiveSamplingRange(int32 DbSequenceIdx) const;
 
 	bool IsValidForIndexing() const;
 	bool IsValidForSearch() const;
@@ -787,14 +788,15 @@ POSESEARCH_API FSearchResult Search(const UAnimSequenceBase* Sequence, TArrayVie
 /**
 * Performs a pose search on a UPoseSearchDatabase.
 *
-* @param Database		The database to search within
-* @param Query			The pose query to search for. To build a query, see FFeatureVectorBuilder. Must have been built using the same schema as the Database.
-* @param WeightsContext	Optional weights context used to influence pose search query results
-* @param DrawParams		Visualization options
+* @param Database			The database to search within
+* @param Query				The pose query to search for. To build a query, see FFeatureVectorBuilder. Must have been built using the same schema as the Database.
+* @param WeightsContext		Optional weights context used to influence pose search query results
+* @param EndTimeToExclude	Samples this close to the end of an animation sequence will be ignored by the search. 
+* @param DrawParams			Visualization options
 * 
 * @return The pose in the database that most closely matches the Query.
 */
-POSESEARCH_API FDbSearchResult Search(const UPoseSearchDatabase* Database, TArrayView<const float> Query, const FPoseSearchWeightsContext* WeightsContext = nullptr, FDebugDrawParams DrawParams = FDebugDrawParams());
+POSESEARCH_API FDbSearchResult Search(const UPoseSearchDatabase* Database, TArrayView<const float> Query, const FPoseSearchWeightsContext* WeightsContext = nullptr, const float EndTimeToExclude = 0.0f, FDebugDrawParams DrawParams = FDebugDrawParams());
 
 
 /**
diff --git a/Engine/Plugins/Experimental/Animation/PoseSearch/Source/Runtime/Public/PoseSearch/PoseSearchLibrary.h b/Engine/Plugins/Experimental/Animation/PoseSearch/Source/Runtime/Public/PoseSearch/PoseSearchLibrary.h
--- a/Engine/Plugins/Experimental/Animation/PoseSearch/Source/Runtime/Public/PoseSearch/PoseSearchLibrary.h
+++ b/Engine/Plugins/Experimental/Animation/PoseSearch/Source/Runtime/Public/PoseSearch/PoseSearchLibrary.h
@@ -13,6 +13,26 @@
 
 #include "PoseSearchLibrary.generated.h"
 
+namespace UE::PoseSearch
+{
+	struct FMotionMatchingContinuityParams
+	{
+		FDbSearchResult Result;
+		bool bJumpRequired = false;
+
+		bool IsValid()
+		{
+			return Result.IsValid();
+		}
+
+		void Reset()
+		{
+			Result = UE::PoseSearch::FDbSearchResult();
+			bJumpRequired = false;
+		}
+	};
+}
+
 UENUM(BlueprintType, Category="Motion Trajectory", meta=(Bitflags, UseEnumValuesAsMaskValuesInEditor="true"))
 enum class EMotionMatchingFlags : uint8
 {
@@ -46,6 +66,11 @@ struct POSESEARCH_API FMotionMatchingSettings
 	// Note: This feature won't work quite as advertised until search data rescaling is implemented
 	UPROPERTY(EditAnywhere, BlueprintReadWrite, Category=Settings, meta=(ClampMin="0", ClampMax="100"))
 	float MinPercentImprovement = 100.f;
+
+	// Pose indices SequenceEndExclusionTime or less away from the end of the database sequence will be ignored by the 
+	// PoseSearch. This ensures the search doesn't get stuck at the end of a sequence.
+	UPROPERTY(EditAnywhere, BlueprintReadWrite, Category = Settings, meta = (ClampMin = "0"))
+	float SequenceEndExlusionTime = 0.25f;
 };
 
 USTRUCT(BlueprintType, Category="Animation|Pose Search")
@@ -62,6 +87,9 @@ struct POSESEARCH_API FMotionMatchingState
 	// Internally stores the 'jump' to a new pose/sequence index and asset time for evaluation
 	void JumpToPose(const UE::PoseSearch::FDbSearchResult& Result);
 
+	// Updates DbPoseIdx to track DeltaTime and jumps to a follow-up sequence when available
+	UE::PoseSearch::FMotionMatchingContinuityParams ComputeContinuityParameters(const FAnimationUpdateContext& Context) const;
+
 	// The current pose we're playing from the database
 	UPROPERTY(Transient)
 	int32 DbPoseIdx = INDEX_NONE;
