PoseSearch: Added a simple mechanism to support bulk adding animation clips to a database

diff --git a/Engine/Plugins/Experimental/Animation/PoseSearch/Source/Runtime/Private/PoseSearch.cpp b/Engine/Plugins/Experimental/Animation/PoseSearch/Source/Runtime/Private/PoseSearch.cpp
--- a/Engine/Plugins/Experimental/Animation/PoseSearch/Source/Runtime/Private/PoseSearch.cpp
+++ b/Engine/Plugins/Experimental/Animation/PoseSearch/Source/Runtime/Private/PoseSearch.cpp
@@ -1030,6 +1030,25 @@ bool UPoseSearchDatabase::IsValidForSearch() const
 	return IsValidForIndexing() && SearchIndex.IsValid();
 }
 
+void UPoseSearchDatabase::CollectSimpleSequences()
+{
+	for (auto& SimpleSequence: SimpleSequences)
+	{
+		auto Predicate = [&SimpleSequence](FPoseSearchDatabaseSequence& DbSequence) -> bool
+		{
+			return DbSequence.Sequence == SimpleSequence;
+		};
+
+		if (!Sequences.ContainsByPredicate(Predicate))
+		{
+			FPoseSearchDatabaseSequence& DbSequence = Sequences.AddDefaulted_GetRef();
+			DbSequence.Sequence = SimpleSequence;
+		}
+	}
+
+	SimpleSequences.Reset();
+}
+
 void UPoseSearchDatabase::PreSave(FObjectPreSaveContext ObjectSaveContext)
 {
 	SearchIndex.Reset();
@@ -1047,6 +1066,21 @@ void UPoseSearchDatabase::PreSave(FObjectPreSaveContext ObjectSaveContext)
 	Super::PreSave(ObjectSaveContext);
 }
 
+#if WITH_EDITOR
+void UPoseSearchDatabase::PostEditChangeProperty(struct FPropertyChangedEvent& PropertyChangedEvent)
+{
+	Super::PostEditChangeProperty(PropertyChangedEvent);
+
+	if (PropertyChangedEvent.Property->GetFName() == GET_MEMBER_NAME_CHECKED(UPoseSearchDatabase, SimpleSequences))
+	{
+		if (!SimpleSequences.IsEmpty())
+		{
+			CollectSimpleSequences();
+		}
+	}
+}
+#endif // WITH_EDITOR
+
 
 //////////////////////////////////////////////////////////////////////////
 // FPoseSearchFeatureVectorBuilder
diff --git a/Engine/Plugins/Experimental/Animation/PoseSearch/Source/Runtime/Public/PoseSearch/PoseSearch.h b/Engine/Plugins/Experimental/Animation/PoseSearch/Source/Runtime/Public/PoseSearch/PoseSearch.h
--- a/Engine/Plugins/Experimental/Animation/PoseSearch/Source/Runtime/Public/PoseSearch/PoseSearch.h
+++ b/Engine/Plugins/Experimental/Animation/PoseSearch/Source/Runtime/Public/PoseSearch/PoseSearch.h
@@ -524,6 +524,10 @@ public:
 	UPROPERTY(EditAnywhere, Category = "Database")
 	FPoseSearchExtrapolationParameters ExtrapolationParameters;
 
+	// Drag and drop animations here to add them in bulk to Sequences
+	UPROPERTY(EditAnywhere, Category = "Database", DisplayName="Drag And Drop Anims Here")
+	TArray<TObjectPtr<UAnimSequence>> SimpleSequences;
+
 	UPROPERTY(EditAnywhere, Category="Database")
 	TArray<FPoseSearchDatabaseSequence> Sequences;
 
@@ -541,6 +545,13 @@ public:
 
 public: // UObject
 	virtual void PreSave(FObjectPreSaveContext ObjectSaveContext) override;
+
+#if WITH_EDITOR
+	virtual void PostEditChangeProperty(struct FPropertyChangedEvent& PropertyChangedEvent) override;
+#endif
+
+private:
+	void CollectSimpleSequences();
 };
 
 
