PoseSearch: Added ForwardVector trajectory feature

diff --git a/Engine/Plugins/Experimental/Animation/PoseSearch/Source/Runtime/Private/PoseSearch.cpp b/Engine/Plugins/Experimental/Animation/PoseSearch/Source/Runtime/Private/PoseSearch.cpp
--- a/Engine/Plugins/Experimental/Animation/PoseSearch/Source/Runtime/Private/PoseSearch.cpp
+++ b/Engine/Plugins/Experimental/Animation/PoseSearch/Source/Runtime/Private/PoseSearch.cpp
@@ -159,6 +159,17 @@ constexpr int32 ChannelIdxPose = 0;
 constexpr int32 ChannelIdxTrajectoryTime = 1;
 constexpr int32 ChannelIdxTrajectoryDistance = 2;
 
+static EPoseSearchFeatureDomain FeatureDomainFromChannel (int32 ChannelIdx)
+{
+	switch (ChannelIdx)
+	{
+		case ChannelIdxPose: return EPoseSearchFeatureDomain::Time;
+		case ChannelIdxTrajectoryTime: return EPoseSearchFeatureDomain::Time;
+		case ChannelIdxTrajectoryDistance: return EPoseSearchFeatureDomain::Distance;
+		default: return EPoseSearchFeatureDomain::Invalid;
+	}
+}
+
 
 //////////////////////////////////////////////////////////////////////////
 // FFeatureTypeTraits
@@ -176,6 +187,7 @@ static constexpr FFeatureTypeTraits FeatureTypeTraits[] =
 	{ EPoseSearchFeatureType::Rotation, 6 },
 	{ EPoseSearchFeatureType::LinearVelocity, 3 },
 	{ EPoseSearchFeatureType::AngularVelocity, 3 },
+	{ EPoseSearchFeatureType::ForwardVector, 3 },
 };
 
 FFeatureTypeTraits GetFeatureTypeTraits(EPoseSearchFeatureType Type)
@@ -309,8 +321,8 @@ void UPoseSearchSchema::PreSave(FObjectPreSaveContext ObjectSaveContext)
 	TrajectorySampleTimes.Sort(TLess<>());
 	TrajectorySampleDistances.Sort(TLess<>());
 
-	GenerateLayout();
 	ResolveBoneReferences();
+	GenerateLayout();
 
 	EffectiveDataPreprocessor = DataPreprocessor;
 	if (EffectiveDataPreprocessor == EPoseSearchDataPreprocessor::Automatic)
@@ -386,86 +398,82 @@ void UPoseSearchSchema::GenerateLayout()
 {
 	Layout.Reset();
 
-	// Time domain trajectory positions
-	if (bUseTrajectoryPositions && TrajectorySampleTimes.Num())
+	auto AddChannelFeatures = [this](int32 ChannelIdx, EPoseSearchFeatureType Type)
 	{
+		const int32 NumSamples = GetChannelSampleOffsets(ChannelIdx).Num();
+		const int32 NumBones = GetNumBones();
+		EPoseSearchFeatureDomain Domain = UE::PoseSearch::FeatureDomainFromChannel(ChannelIdx);
+
 		FPoseSearchFeatureDesc Feature;
-		Feature.SchemaBoneIdx = FPoseSearchFeatureDesc::TrajectoryBoneIndex;
-		Feature.Domain = EPoseSearchFeatureDomain::Time;
-		Feature.Type = EPoseSearchFeatureType::Position;
-		for (Feature.SubsampleIdx = 0; Feature.SubsampleIdx != TrajectorySampleTimes.Num(); ++Feature.SubsampleIdx)
+		Feature.Domain = Domain;
+		Feature.Type = Type;
+
+		if (ChannelIdx == UE::PoseSearch::ChannelIdxPose)
+		{
+			for (Feature.SubsampleIdx = 0; Feature.SubsampleIdx != NumSamples; ++Feature.SubsampleIdx)
+			{
+				for (Feature.SchemaBoneIdx = 0; Feature.SchemaBoneIdx != NumBones; ++Feature.SchemaBoneIdx)
+				{
+					Layout.Features.Add(Feature);
+				}
+			}
+		}
+		else
 		{
-			Layout.Features.Add(Feature);
+			Feature.SchemaBoneIdx = FPoseSearchFeatureDesc::TrajectoryBoneIndex;
+			for (Feature.SubsampleIdx = 0; Feature.SubsampleIdx != NumSamples; ++Feature.SubsampleIdx)
+			{
+				Layout.Features.Add(Feature);
+			}
 		}
+	};
+
+	// Time domain trajectory positions
+	if (bUseTrajectoryPositions)
+	{
+		AddChannelFeatures(UE::PoseSearch::ChannelIdxTrajectoryTime, EPoseSearchFeatureType::Position);
 	}
 
 	// Time domain trajectory linear velocities
-	if (bUseTrajectoryVelocities && TrajectorySampleTimes.Num())
+	if (bUseTrajectoryVelocities)
 	{
-		FPoseSearchFeatureDesc Feature;
-		Feature.SchemaBoneIdx = FPoseSearchFeatureDesc::TrajectoryBoneIndex;
-		Feature.Domain = EPoseSearchFeatureDomain::Time;
-		Feature.Type = EPoseSearchFeatureType::LinearVelocity;
-		for (Feature.SubsampleIdx = 0; Feature.SubsampleIdx != TrajectorySampleTimes.Num(); ++Feature.SubsampleIdx)
-		{
-			Layout.Features.Add(Feature);
-		}
+		AddChannelFeatures(UE::PoseSearch::ChannelIdxTrajectoryTime, EPoseSearchFeatureType::LinearVelocity);
+	}
+
+	// Time domain trajectory forward vectors
+	if (bUseTrajectoryForwardVectors)
+	{
+		AddChannelFeatures(UE::PoseSearch::ChannelIdxTrajectoryTime, EPoseSearchFeatureType::ForwardVector);
 	}
 
 	// Distance domain trajectory positions
-	if (bUseTrajectoryPositions && TrajectorySampleDistances.Num())
+	if (bUseTrajectoryPositions)
 	{
-		FPoseSearchFeatureDesc Feature;
-		Feature.SchemaBoneIdx = FPoseSearchFeatureDesc::TrajectoryBoneIndex;
-		Feature.Domain = EPoseSearchFeatureDomain::Distance;
-		Feature.Type = EPoseSearchFeatureType::Position;
-		for (Feature.SubsampleIdx = 0; Feature.SubsampleIdx != TrajectorySampleDistances.Num(); ++Feature.SubsampleIdx)
-		{
-			Layout.Features.Add(Feature);
-		}
+		AddChannelFeatures(UE::PoseSearch::ChannelIdxTrajectoryDistance, EPoseSearchFeatureType::Position);
 	}
 
 	// Distance domain trajectory linear velocities
-	if (bUseTrajectoryVelocities && TrajectorySampleDistances.Num())
+	if (bUseTrajectoryVelocities)
 	{
-		FPoseSearchFeatureDesc Feature;
-		Feature.SchemaBoneIdx = FPoseSearchFeatureDesc::TrajectoryBoneIndex;
-		Feature.Domain = EPoseSearchFeatureDomain::Distance;
-		Feature.Type = EPoseSearchFeatureType::LinearVelocity;
-		for (Feature.SubsampleIdx = 0; Feature.SubsampleIdx != TrajectorySampleDistances.Num(); ++Feature.SubsampleIdx)
-		{
-			Layout.Features.Add(Feature);
-		}
+		AddChannelFeatures(UE::PoseSearch::ChannelIdxTrajectoryDistance, EPoseSearchFeatureType::LinearVelocity);
+	}
+
+	// Distance domain trajectory forward vectors
+	if (bUseTrajectoryForwardVectors)
+	{
+		AddChannelFeatures(UE::PoseSearch::ChannelIdxTrajectoryDistance, EPoseSearchFeatureType::ForwardVector);
 	}
 
 	// Time domain bone positions
-	if (bUseBonePositions && PoseSampleTimes.Num())
+	if (bUseBonePositions)
 	{
-		FPoseSearchFeatureDesc Feature;
-		Feature.Domain = EPoseSearchFeatureDomain::Time;
-		Feature.Type = EPoseSearchFeatureType::Position;
-		for (Feature.SubsampleIdx = 0; Feature.SubsampleIdx != PoseSampleTimes.Num(); ++Feature.SubsampleIdx)
-		{
-			for (Feature.SchemaBoneIdx = 0; Feature.SchemaBoneIdx != Bones.Num(); ++Feature.SchemaBoneIdx)
-			{
-				Layout.Features.Add(Feature);
-			}
-		}
+		AddChannelFeatures(UE::PoseSearch::ChannelIdxPose, EPoseSearchFeatureType::Position);
 	}
 
 	// Time domain bone linear velocities
 	if (bUseBoneVelocities && PoseSampleTimes.Num())
 	{
-		FPoseSearchFeatureDesc Feature;
-		Feature.Domain = EPoseSearchFeatureDomain::Time;
-		Feature.Type = EPoseSearchFeatureType::LinearVelocity;
-		for (Feature.SubsampleIdx = 0; Feature.SubsampleIdx != PoseSampleTimes.Num(); ++Feature.SubsampleIdx)
-		{
-			for (Feature.SchemaBoneIdx = 0; Feature.SchemaBoneIdx != Bones.Num(); ++Feature.SchemaBoneIdx)
-			{
-				Layout.Features.Add(Feature);
-			}
-		}
+		AddChannelFeatures(UE::PoseSearch::ChannelIdxPose, EPoseSearchFeatureType::LinearVelocity);
 	}
 
 	Layout.Init();
@@ -1153,6 +1161,9 @@ void FPoseSearchFeatureVectorBuilder::SetRotation(FPoseSearchFeatureDesc Element
 			++NumFeaturesAdded;
 		}
 	}
+
+	Element.Type = EPoseSearchFeatureType::ForwardVector;
+	SetVector(Element, Rotation.GetAxisY());
 }
 
 void FPoseSearchFeatureVectorBuilder::SetLinearVelocity(FPoseSearchFeatureDesc Element, const FTransform& Transform, const FTransform& PrevTransform, float DeltaTime)
@@ -1277,8 +1288,7 @@ void FPoseSearchFeatureVectorBuilder::BuildFromTrajectoryTimeBased(const FTrajec
 		Feature.Type = EPoseSearchFeatureType::LinearVelocity;
 		SetVector(Feature, Sample.LinearVelocity);
 
-		Feature.Type = EPoseSearchFeatureType::Position;
-		SetVector(Feature, Sample.Transform.GetLocation());
+		SetTransform(Feature, Sample.Transform);
 	}
 }
 
@@ -1300,8 +1310,7 @@ void FPoseSearchFeatureVectorBuilder::BuildFromTrajectoryDistanceBased(const FTr
 		Feature.Type = EPoseSearchFeatureType::LinearVelocity;
 		SetVector(Feature, Sample.LinearVelocity);
 
-		Feature.Type = EPoseSearchFeatureType::Position;
-		SetVector(Feature, Sample.Transform.GetLocation());
+		SetTransform(Feature, Sample.Transform);
 	} 
 }
 
@@ -1602,6 +1611,12 @@ bool FFeatureVectorReader::GetRotation(FPoseSearchFeatureDesc Element, FQuat* Ou
 	return false;
 }
 
+bool FFeatureVectorReader::GetForwardVector(FPoseSearchFeatureDesc Element, FVector* OutForwardVector) const
+{
+	Element.Type = EPoseSearchFeatureType::ForwardVector;
+	return GetVector(Element, OutForwardVector);
+}
+
 bool FFeatureVectorReader::GetLinearVelocity(FPoseSearchFeatureDesc Element, FVector* OutLinearVelocity) const
 {
 	Element.Type = EPoseSearchFeatureType::LinearVelocity;
@@ -2470,7 +2485,7 @@ void FSequenceIndexer::AddPoseFeatures(int32 SampleIdx)
 		// Get each bone's component transform, velocity, and acceleration and add accumulated root motion at this time offset
 		// Think of this process as freezing the character in place (at SampleTime) and then tracing the paths of their joints
 		// as they move through space from past to present to future (at times indicated by PoseSampleTimes).
-		for (int32 SchemaBoneIndex = 0; SchemaBoneIndex != Input.Schema->NumBones(); ++SchemaBoneIndex)
+		for (int32 SchemaBoneIndex = 0; SchemaBoneIndex != Input.Schema->GetNumBones(); ++SchemaBoneIndex)
 		{
 			Feature.SchemaBoneIdx = SchemaBoneIndex;
 
@@ -2632,7 +2647,46 @@ static void DrawTrajectoryFeatures(const FDebugDrawParams& DrawParams, const FFe
 			}
 			else
 			{
-				DrawDebugDirectionalArrow(DrawParams.World, TrajectoryPos + TrajectoryVelDirection * DrawDebugSphereSize, TrajectoryPos + TrajectoryVel, DrawDebugArrowSize, Color, bPersistent, LifeTime, DepthPriority, DrawDebugLineThickness);
+				DrawDebugDirectionalArrow(
+					DrawParams.World,
+					TrajectoryPos + TrajectoryVelDirection * DrawDebugSphereSize,
+					TrajectoryPos + TrajectoryVel,
+					DrawDebugArrowSize,
+					Color,
+					bPersistent,
+					LifeTime,
+					DepthPriority,
+					DrawDebugLineThickness
+				);
+			}
+		}
+
+		FVector TrajectoryForward;
+		if (Reader.GetForwardVector(Feature, &TrajectoryForward))
+		{
+			Feature.Type = EPoseSearchFeatureType::ForwardVector;
+
+			FLinearColor LinearColor = DrawParams.Color ? *DrawParams.Color : GetColorForFeature(Feature, Reader.GetLayout());
+			FColor Color = LinearColor.ToFColor(true);
+
+			TrajectoryForward = DrawParams.RootTransform.TransformVector(TrajectoryForward);
+			if (EnumHasAnyFlags(DrawParams.Flags, EDebugDrawFlags::DrawSearchIndex))
+			{
+				DrawDebugPoint(DrawParams.World, TrajectoryForward, DrawParams.PointSize, Color, bPersistent, DrawParams.DefaultLifeTime, DepthPriority);
+			}
+			else
+			{
+				DrawDebugDirectionalArrow(
+					DrawParams.World,
+					TrajectoryPos + TrajectoryForward * DrawDebugSphereSize,
+					TrajectoryPos + TrajectoryForward * DrawDebugSphereSize * 2.0f,
+					DrawDebugArrowSize,
+					Color,
+					bPersistent,
+					LifeTime,
+					DepthPriority,
+					DrawDebugLineThickness
+				);
 			}
 		}
 	}
diff --git a/Engine/Plugins/Experimental/Animation/PoseSearch/Source/Runtime/Public/PoseSearch/PoseSearch.h b/Engine/Plugins/Experimental/Animation/PoseSearch/Source/Runtime/Public/PoseSearch/PoseSearch.h
--- a/Engine/Plugins/Experimental/Animation/PoseSearch/Source/Runtime/Public/PoseSearch/PoseSearch.h
+++ b/Engine/Plugins/Experimental/Animation/PoseSearch/Source/Runtime/Public/PoseSearch/PoseSearch.h
@@ -35,25 +35,24 @@ DECLARE_LOG_CATEGORY_EXTERN(LogPoseSearch, Log, All);
 UENUM()
 enum class EPoseSearchFeatureType : int32
 {
-	Invalid			= -1,
-
-	Position		= 0,	
-	Rotation		= 1,
-	LinearVelocity	= 2,
-	AngularVelocity	= 3,
-
-	Num
+	Position,	
+	Rotation,
+	LinearVelocity,
+	AngularVelocity,
+	ForwardVector,
+
+	Num UMETA(Hidden),
+	Invalid = Num UMETA(Hidden)
 };
 
 UENUM()
 enum class EPoseSearchFeatureDomain : int32
 {
-	Invalid		= -1,
-
-	Time		= 0,
-	Distance	= 1,
+	Time,
+	Distance,
 
-	Num
+	Num UMETA(Hidden),
+	Invalid = Num UMETA(Hidden)
 };
 
 /** Describes each feature of a vector, including data type, sampling options, and buffer offset. */
@@ -119,12 +118,13 @@ struct POSESEARCH_API FPoseSearchFeatureVectorLayout
 UENUM()
 enum class EPoseSearchDataPreprocessor : int32
 {
-	None			= 0,
-	Automatic		= 1,
-	Normalize		= 2,
-	Sphere			= 3 UMETA(Hidden),
+	None,
+	Automatic,
+	Normalize,
+	Sphere UMETA(Hidden),
 
-	Invalid = -1 UMETA(Hidden)
+	Num UMETA(Hidden),
+	Invalid = Num UMETA(Hidden)
 };
 
 /**
@@ -157,6 +157,9 @@ public:
 	UPROPERTY(EditAnywhere, Category = "Schema")
 	bool bUseTrajectoryPositions = true;
 
+	UPROPERTY(EditAnywhere, Category = "Schema")
+	bool bUseTrajectoryForwardVectors = false;
+
 	UPROPERTY(EditAnywhere, Category = "Schema")
 	TArray<FBoneReference> Bones;
 
@@ -190,7 +193,7 @@ public:
 
 	bool IsValid () const;
 
-	int32 NumBones () const { return BoneIndices.Num(); }
+	int32 GetNumBones () const { return BoneIndices.Num(); }
 
 	// Returns farthest future sample time >= 0.0f.
 	// Returns a negative value when there are no future sample times.
@@ -672,6 +675,7 @@ public:
 	bool GetTransform(FPoseSearchFeatureDesc Feature, FTransform* OutTransform) const;
 	bool GetPosition(FPoseSearchFeatureDesc Feature, FVector* OutPosition) const;
 	bool GetRotation(FPoseSearchFeatureDesc Feature, FQuat* OutRotation) const;
+	bool GetForwardVector(FPoseSearchFeatureDesc Feature, FVector* OutForwardVector) const;
 	bool GetLinearVelocity(FPoseSearchFeatureDesc Feature, FVector* OutLinearVelocity) const;
 	bool GetAngularVelocity(FPoseSearchFeatureDesc Feature, FVector* OutAngularVelocity) const;
 	bool GetVector(FPoseSearchFeatureDesc Feature, FVector* OutVector) const;
