Fixed bone index mismatch in FPoseSearchFeatureVectorBuilder::TrySetPoseFeatures. Now the function takes a bone container as a parameter so we can convert history compact bone index into skeleton index.

diff --git a/Engine/Plugins/Experimental/Animation/PoseSearch/Source/Runtime/Private/PoseSearch.cpp b/Engine/Plugins/Experimental/Animation/PoseSearch/Source/Runtime/Private/PoseSearch.cpp
--- a/Engine/Plugins/Experimental/Animation/PoseSearch/Source/Runtime/Private/PoseSearch.cpp
+++ b/Engine/Plugins/Experimental/Animation/PoseSearch/Source/Runtime/Private/PoseSearch.cpp
@@ -1228,7 +1228,9 @@ void FPoseSearchFeatureVectorBuilder::SetVector(FPoseSearchFeatureDesc Element,
 	}
 }
 
-bool FPoseSearchFeatureVectorBuilder::TrySetPoseFeatures(UE::PoseSearch::FPoseHistory* History)
+bool FPoseSearchFeatureVectorBuilder::TrySetPoseFeatures(
+	UE::PoseSearch::FPoseHistory* History, 
+	const FBoneContainer& BoneContainer)
 {
 	check(Schema.IsValid() && Schema->IsValid());
 	check(History);
@@ -1260,8 +1262,9 @@ bool FPoseSearchFeatureVectorBuilder::TrySetPoseFeatures(UE::PoseSearch::FPoseHi
 			Feature.SchemaBoneIdx = SchemaBoneIdx;
 
 			int32 SkeletonBoneIndex = Schema->BoneIndices[SchemaBoneIdx];
-			const FTransform& Transform = ComponentPose[SkeletonBoneIndex];
-			const FTransform& PrevTransform = ComponentPrevPose[SkeletonBoneIndex];
+			FCompactPoseBoneIndex CompactBoneIndex = BoneContainer.GetCompactPoseIndexFromSkeletonIndex(SkeletonBoneIndex);
+			const FTransform& Transform = ComponentPose[CompactBoneIndex.GetInt()];
+			const FTransform& PrevTransform = ComponentPrevPose[CompactBoneIndex.GetInt()];
 			SetTransform(Feature, Transform);
 			SetTransformVelocity(Feature, Transform, PrevTransform, History->GetSampleTimeInterval());
 		}
@@ -3616,7 +3619,7 @@ UE::Anim::IPoseSearchProvider::FSearchResult FModule::Search(const FAnimationBas
 	FPoseSearchFeatureVectorBuilder& QueryBuilder = PoseHistory.GetQueryBuilder();
 
 	QueryBuilder.Init(MetaData->Schema);
-	if (!QueryBuilder.TrySetPoseFeatures(&PoseHistory))
+	if (!QueryBuilder.TrySetPoseFeatures(&PoseHistory, GraphContext.AnimInstanceProxy->GetRequiredBones()))
 	{
 		return ProviderResult;
 	}
diff --git a/Engine/Plugins/Experimental/Animation/PoseSearch/Source/Runtime/Private/PoseSearchLibrary.cpp b/Engine/Plugins/Experimental/Animation/PoseSearch/Source/Runtime/Private/PoseSearchLibrary.cpp
--- a/Engine/Plugins/Experimental/Animation/PoseSearch/Source/Runtime/Private/PoseSearchLibrary.cpp
+++ b/Engine/Plugins/Experimental/Animation/PoseSearch/Source/Runtime/Private/PoseSearchLibrary.cpp
@@ -100,7 +100,9 @@ void UpdateMotionMatchingState(const FAnimationUpdateContext& Context
 			if (PoseHistoryProvider)
 			{
 				UE::PoseSearch::FPoseHistory& History = PoseHistoryProvider->GetPoseHistory();
-				InOutMotionMatchingState.ComposedQuery.TrySetPoseFeatures(&History);
+				InOutMotionMatchingState.ComposedQuery.TrySetPoseFeatures(
+					&History, 
+					Context.AnimInstanceProxy->GetRequiredBones());
 			}
 		}
 
diff --git a/Engine/Plugins/Experimental/Animation/PoseSearch/Source/Runtime/Public/PoseSearch/PoseSearch.h b/Engine/Plugins/Experimental/Animation/PoseSearch/Source/Runtime/Public/PoseSearch/PoseSearch.h
--- a/Engine/Plugins/Experimental/Animation/PoseSearch/Source/Runtime/Public/PoseSearch/PoseSearch.h
+++ b/Engine/Plugins/Experimental/Animation/PoseSearch/Source/Runtime/Public/PoseSearch/PoseSearch.h
@@ -586,7 +586,7 @@ public:
 	void SetAngularVelocity(FPoseSearchFeatureDesc Feature, const FTransform& Transform, const FTransform& PrevTransform, float DeltaTime);
 	void SetVector(FPoseSearchFeatureDesc Feature, const FVector& Vector);
 	void BuildFromTrajectory(const FTrajectorySampleRange& Trajectory);
-	bool TrySetPoseFeatures(UE::PoseSearch::FPoseHistory* History);
+	bool TrySetPoseFeatures(UE::PoseSearch::FPoseHistory* History, const FBoneContainer& BoneContainer);
 
 	void CopyFromSearchIndex(const FPoseSearchIndex& SearchIndex, int32 PoseIdx);
 	void CopyFeature(const FPoseSearchFeatureVectorBuilder& OtherBuilder, int32 FeatureIdx);
