PoseSearch: Fixed debug drawing perf caused by sphere line thickness

diff --git a/Engine/Plugins/Experimental/Animation/PoseSearch/Source/Runtime/Private/PoseSearch.cpp b/Engine/Plugins/Experimental/Animation/PoseSearch/Source/Runtime/Private/PoseSearch.cpp
--- a/Engine/Plugins/Experimental/Animation/PoseSearch/Source/Runtime/Private/PoseSearch.cpp
+++ b/Engine/Plugins/Experimental/Animation/PoseSearch/Source/Runtime/Private/PoseSearch.cpp
@@ -41,8 +41,7 @@ constexpr float DrawDebugPointSize = 3.0f;
 constexpr float DrawDebugVelocityScale = 0.08f;
 constexpr float DrawDebugArrowSize = 30.0f;
 constexpr float DrawDebugSphereSize = 3.0f;
-constexpr int32 DrawDebugSphereSegments = 8;
-constexpr float DrawDebugSphereLineThickness = 0.5f;
+constexpr int32 DrawDebugSphereSegments = 10;
 
 static bool IsSamplingRangeValid(FFloatInterval Range)
 {
@@ -1228,9 +1227,7 @@ void FPoseSearchFeatureVectorBuilder::SetVector(FPoseSearchFeatureDesc Element,
 	}
 }
 
-bool FPoseSearchFeatureVectorBuilder::TrySetPoseFeatures(
-	UE::PoseSearch::FPoseHistory* History, 
-	const FBoneContainer& BoneContainer)
+bool FPoseSearchFeatureVectorBuilder::TrySetPoseFeatures(UE::PoseSearch::FPoseHistory* History, const FBoneContainer& BoneContainer)
 {
 	check(Schema.IsValid() && Schema->IsValid());
 	check(History);
@@ -2625,7 +2622,7 @@ static void DrawTrajectoryFeatures(const FDebugDrawParams& DrawParams, const FFe
 			}
 			else
 			{
-				DrawDebugSphere(DrawParams.World, TrajectoryPos, DrawDebugSphereSize, DrawDebugSphereSegments, Color, bPersistent, LifeTime, DepthPriority, DrawDebugSphereLineThickness);
+				DrawDebugSphere(DrawParams.World, TrajectoryPos, DrawDebugSphereSize, DrawDebugSphereSegments, Color, bPersistent, LifeTime, DepthPriority);
 			}
 		}
 		else
@@ -2739,7 +2736,7 @@ static void DrawPoseFeatures(const FDebugDrawParams& DrawParams, const FFeatureV
 				}
 				else
 				{
-					DrawDebugSphere(DrawParams.World, BonePos, DrawDebugSphereSize, DrawDebugSphereSegments, Color, bPersistent, LifeTime, DepthPriority, DrawDebugSphereLineThickness);
+					DrawDebugSphere(DrawParams.World, BonePos, DrawDebugSphereSize, DrawDebugSphereSegments, Color, bPersistent, LifeTime, DepthPriority);
 				}
 			}
 
