Motion Matching debugger improvements:
- Drawing
  - Added trajectory samples color fade to help distinguish different samples
  - Added trajectory sample labels. Unfortunately they don't show when the controller is ejected in PIE
- Rewind Debugger
  - Added filter box to allow filtering poses by anim sequence name. It tokenizes whatever is in the search box using white spaces as separators and filters out anim sequences whose name don't containt all tokens

diff --git a/Engine/Plugins/Experimental/Animation/PoseSearch/Source/Editor/Private/PoseSearchDebugger.cpp b/Engine/Plugins/Experimental/Animation/PoseSearch/Source/Editor/Private/PoseSearchDebugger.cpp
--- a/Engine/Plugins/Experimental/Animation/PoseSearch/Source/Editor/Private/PoseSearchDebugger.cpp
+++ b/Engine/Plugins/Experimental/Animation/PoseSearch/Source/Editor/Private/PoseSearchDebugger.cpp
@@ -13,9 +13,11 @@
 #include "Widgets/Layout/SWidgetSwitcher.h"
 #include "Widgets/SBoxPanel.h"
 #include "Widgets/Input/SButton.h"
+#include "Widgets/Input/SSearchBox.h"
 #include "Widgets/Images/SImage.h"
 #include "Styling/SlateIconFinder.h"
 #include "TraceServices/Model/Frames.h"
+#include "Algo/AllOf.h"
 
 #define LOCTEXT_NAMESPACE "PoseSearchDebugger"
 
@@ -389,7 +391,7 @@ void SDebuggerDatabaseView::RefreshColumns()
 	using namespace DebuggerDatabaseColumns;
 	
 	ActiveView.HeaderRow->ClearColumns();
-	DatabaseView.HeaderRow->ClearColumns();
+	FilteredDatabaseView.HeaderRow->ClearColumns();
 	
 	// Sort columns by index
 	Columns.ValueSort([](const TSharedRef<IColumn> Column0, const TSharedRef<IColumn> Column1)
@@ -415,7 +417,7 @@ void SDebuggerDatabaseView::RefreshColumns()
 				.HAlignHeader(HAlign_Center)
 				.HAlignCell(HAlign_Fill);
 
-			DatabaseView.HeaderRow->AddColumn(ColumnArgs);
+			FilteredDatabaseView.HeaderRow->AddColumn(ColumnArgs);
 			
 			// Every time the active column is changed, update the database column
 			ActiveView.HeaderRow->AddColumn(ColumnArgs.OnWidthChanged(this, &SDebuggerDatabaseView::OnColumnWidthChanged, ColumnName));
@@ -452,6 +454,7 @@ void SDebuggerDatabaseView::OnColumnSortModeChanged(const EColumnSortPriority::T
 	SortColumn = ColumnId;
 	SortMode = InSortMode;
 	SortDatabaseRows();
+	FilterDatabaseRows();
 }
 
 void SDebuggerDatabaseView::OnColumnWidthChanged(const float NewWidth, FName ColumnId) const
@@ -461,25 +464,65 @@ void SDebuggerDatabaseView::OnColumnWidthChanged(const float NewWidth, FName Col
 	Columns[ColumnId]->Width = NewWidth;
 }
 
+void SDebuggerDatabaseView::OnFilterTextChanged(const FText& SearchText)
+{
+	FilterText = SearchText;
+	FilterDatabaseRows();
+}
+
 
 void SDebuggerDatabaseView::SortDatabaseRows()
 {
 	if (SortMode == EColumnSortMode::Ascending)
 	{
-		DatabaseView.Rows.Sort(Columns[SortColumn]->GetSortPredicateAscending());
+		UnfilteredDatabaseRows.Sort(Columns[SortColumn]->GetSortPredicateAscending());
 	}
 	else if (SortMode == EColumnSortMode::Descending)
 	{
-		DatabaseView.Rows.Sort(Columns[SortColumn]->GetSortPredicateDescending());
+		UnfilteredDatabaseRows.Sort(Columns[SortColumn]->GetSortPredicateDescending());
 	}
+}
+
+void SDebuggerDatabaseView::FilterDatabaseRows()
+{
+	FilteredDatabaseView.Rows.Empty();
 
-	DatabaseView.ListView->RequestListRefresh();
+	FString FilterString = FilterText.ToString();
+	TArray<FString> Tokens;
+	FilterString.ParseIntoArrayWS(Tokens);
+
+	if (Tokens.Num() == 0)
+	{
+		for (const auto& UnfilteredRow : UnfilteredDatabaseRows)
+		{
+			FilteredDatabaseView.Rows.Add(UnfilteredRow);
+		}
+	}
+	else
+	{
+		for (const auto& UnfilteredRow : UnfilteredDatabaseRows)
+		{
+			bool bMatchesAllTokens = Algo::AllOf(
+				Tokens,
+				[&](FString Token)
+			{
+				return UnfilteredRow->AnimSequenceName.Contains(Token);
+			});
+
+			if (bMatchesAllTokens)
+			{
+				FilteredDatabaseView.Rows.Add(UnfilteredRow);
+			}
+		}
+	}
+
+	FilteredDatabaseView.ListView->RequestListRefresh();
 }
 
 void SDebuggerDatabaseView::CreateRows(const UPoseSearchDatabase& Database)
 {
 	const int32 NumPoses = Database.SearchIndex.NumPoses;
-	DatabaseView.Rows.Reserve(NumPoses);
+	UnfilteredDatabaseRows.Reserve(NumPoses);
 
 	// Build database rows
 	for(const FPoseSearchDatabaseSequence& DbSequence : Database.Sequences)
@@ -487,7 +530,7 @@ void SDebuggerDatabaseView::CreateRows(const UPoseSearchDatabase& Database)
 		const int32 LastPoseIdx = DbSequence.FirstPoseIdx + DbSequence.NumPoses;
 		for (int32 PoseIdx = DbSequence.FirstPoseIdx; PoseIdx != LastPoseIdx; ++PoseIdx)
 		{
-			TSharedRef<FDebuggerDatabaseRowData> Row = DatabaseView.Rows.Add_GetRef(MakeShared<FDebuggerDatabaseRowData>());
+			TSharedRef<FDebuggerDatabaseRowData> Row = UnfilteredDatabaseRows.Add_GetRef(MakeShared<FDebuggerDatabaseRowData>());
 			Row->PoseIdx = PoseIdx;
 
 			const FString SequenceName = DbSequence.Sequence->GetName();
@@ -513,7 +556,7 @@ void SDebuggerDatabaseView::CreateRows(const UPoseSearchDatabase& Database)
 
 void SDebuggerDatabaseView::UpdateRows(const FTraceMotionMatchingStateMessage& State, const UPoseSearchDatabase& Database)
 {
-	if (DatabaseView.Rows.IsEmpty())
+	if (UnfilteredDatabaseRows.IsEmpty())
 	{
 		check(ActiveView.Rows.IsEmpty());
 		CreateRows(Database);
@@ -539,7 +582,7 @@ void SDebuggerDatabaseView::UpdateRows(const FTraceMotionMatchingStateMessage& S
 	
 	const FPoseSearchIndex& SearchIndex = Database.SearchIndex;
 
-	for(const TSharedRef<FDebuggerDatabaseRowData>& Row : DatabaseView.Rows)
+	for(const TSharedRef<FDebuggerDatabaseRowData>& Row : UnfilteredDatabaseRows)
 	{
 		const int32 PoseIdx = Row->PoseIdx;
 		
@@ -555,12 +598,13 @@ void SDebuggerDatabaseView::UpdateRows(const FTraceMotionMatchingStateMessage& S
 	}
 
 	SortDatabaseRows();
+	FilterDatabaseRows();
 }
 
 TSharedRef<ITableRow> SDebuggerDatabaseView::HandleGenerateDatabaseRow(TSharedRef<FDebuggerDatabaseRowData> Item, const TSharedRef<STableViewBase>& OwnerTable) const
 {
 	return
-		SNew(SDebuggerDatabaseRow, OwnerTable, Item, DatabaseView.RowStyle, &DatabaseView.RowBrush, FMargin(0.0f, 2.0f, 6.0f, 2.0f))
+		SNew(SDebuggerDatabaseRow, OwnerTable, Item, FilteredDatabaseView.RowStyle, &FilteredDatabaseView.RowBrush, FMargin(0.0f, 2.0f, 6.0f, 2.0f))
 		.ColumnMap(this, &SDebuggerDatabaseView::GetColumnMap);
 	
 }
@@ -608,28 +652,29 @@ void SDebuggerDatabaseView::Construct(const FArguments& InArgs)
 	ActiveView.RowStyle = FEditorStyle::GetWidgetStyle<FTableRowStyle>("TableView.Row");
 	ActiveView.RowBrush = *FEditorStyle::GetBrush("DetailsView.CategoryTop");
 
-	// Database 
-	DatabaseView.ScrollBar =
+
+	// Filtered Database
+	FilteredDatabaseView.ScrollBar =
 		SNew(SScrollBar)
 		.Orientation(Orient_Vertical)
 		.HideWhenNotInUse(false)
 		.AlwaysShowScrollbar(true)
 		.AlwaysShowScrollbarTrack(true);
-	DatabaseView.HeaderRow = SNew(SHeaderRow).Visibility(EVisibility::Collapsed);
-	
-	DatabaseView.ListView = SNew(SListView<TSharedRef<FDebuggerDatabaseRowData>>)
-			.ListItemsSource(&DatabaseView.Rows)
-			.HeaderRow(DatabaseView.HeaderRow.ToSharedRef())
-			.OnGenerateRow(this, &SDebuggerDatabaseView::HandleGenerateDatabaseRow)
-			.ExternalScrollbar(DatabaseView.ScrollBar)
-			.SelectionMode(ESelectionMode::SingleToggle)
-			.ConsumeMouseWheel(EConsumeMouseWheel::WhenScrollingPossible);
-
-	DatabaseView.RowStyle = FEditorStyle::GetWidgetStyle<FTableRowStyle>("TableView.Row");
+	FilteredDatabaseView.HeaderRow = SNew(SHeaderRow).Visibility(EVisibility::Collapsed);
+
+	FilteredDatabaseView.ListView = SNew(SListView<TSharedRef<FDebuggerDatabaseRowData>>)
+		.ListItemsSource(&FilteredDatabaseView.Rows)
+		.HeaderRow(FilteredDatabaseView.HeaderRow.ToSharedRef())
+		.OnGenerateRow(this, &SDebuggerDatabaseView::HandleGenerateDatabaseRow)
+		.ExternalScrollbar(FilteredDatabaseView.ScrollBar)
+		.SelectionMode(ESelectionMode::SingleToggle)
+		.ConsumeMouseWheel(EConsumeMouseWheel::WhenScrollingPossible);
+
+	FilteredDatabaseView.RowStyle = FEditorStyle::GetWidgetStyle<FTableRowStyle>("TableView.Row");
 	// Set selected color to white to retain visibility when multi-selecting
-	DatabaseView.RowStyle.SetSelectedTextColor(FLinearColor(FVector3f(0.8f)));
-	DatabaseView.RowBrush = *FEditorStyle::GetBrush("ToolPanel.GroupBorder");
-	
+	FilteredDatabaseView.RowStyle.SetSelectedTextColor(FLinearColor(FVector3f(0.8f)));
+	FilteredDatabaseView.RowBrush = *FEditorStyle::GetBrush("ToolPanel.GroupBorder");
+
 	ChildSlot
 	[
 		SNew(SVerticalBox)
@@ -721,7 +766,7 @@ void SDebuggerDatabaseView::Construct(const FArguments& InArgs)
 				.HAlign(HAlign_Fill)
 				[
 					SNew(SBorder)
-					.BorderImage(&DatabaseView.RowStyle.EvenRowBackgroundBrush)
+					.BorderImage(&FilteredDatabaseView.RowStyle.EvenRowBackgroundBrush)
 				]
 			]
 			.AutoHeight()
@@ -737,7 +782,14 @@ void SDebuggerDatabaseView::Construct(const FArguments& InArgs)
 				.HAlign(HAlign_Fill)
 				.VAlign(VAlign_Fill)
 			]
-			
+
+			+ SVerticalBox::Slot()
+			.Padding(0.0f, 0.0f, 0.0f, 5.0f)
+			.AutoHeight()
+			[
+				SAssignNew(FilterBox, SSearchBox)
+				.OnTextChanged(this, &SDebuggerDatabaseView::OnFilterTextChanged)
+			]
 		
 			+ SVerticalBox::Slot()
 			[
@@ -749,14 +801,14 @@ void SDebuggerDatabaseView::Construct(const FArguments& InArgs)
 					.BorderImage(FEditorStyle::GetBrush("NoBorder"))
 					.Padding(0.0f)
 					[
-						DatabaseView.ListView.ToSharedRef()
+						FilteredDatabaseView.ListView.ToSharedRef()
 					]
 				]
 				
 				+ SHorizontalBox::Slot()
 				.AutoWidth()
 				[
-					DatabaseView.ScrollBar.ToSharedRef()
+					FilteredDatabaseView.ScrollBar.ToSharedRef()
 				]
 			]
 		]
@@ -1131,6 +1183,16 @@ void SDebuggerView::DrawFeatures(
 		{
 			InDrawParams.Flags |= EDebugDrawFlags::IncludeTrajectory;
 		}
+
+		if (Options.bDrawSampleLabels)
+		{
+			InDrawParams.Flags |= EDebugDrawFlags::DrawSampleLabels;
+		}
+
+		if (Options.bDrawSamplesWithColorGradient)
+		{
+			InDrawParams.Flags |= EDebugDrawFlags::DrawSamplesWithColorGradient;
+		}
 	};
 
 	// Draw query vector
diff --git a/Engine/Plugins/Experimental/Animation/PoseSearch/Source/Editor/Private/PoseSearchDebugger.h b/Engine/Plugins/Experimental/Animation/PoseSearch/Source/Editor/Private/PoseSearchDebugger.h
--- a/Engine/Plugins/Experimental/Animation/PoseSearch/Source/Editor/Private/PoseSearchDebugger.h
+++ b/Engine/Plugins/Experimental/Animation/PoseSearch/Source/Editor/Private/PoseSearchDebugger.h
@@ -20,6 +20,7 @@ class SVerticalBox;
 class SHorizontalBox;
 class SScrollBox;
 class SWidgetSwitcher;
+class SSearchBox;
 template <typename ItemType> class SListView;
 
 /**
@@ -61,6 +62,12 @@ struct FPoseSearchDebuggerFeatureDrawOptions
 
 	UPROPERTY(EditAnywhere, Category="Draw Options", Meta=(EditCondition="!bDisable"))
     bool bDrawTrajectoryFeatures = true;
+
+	UPROPERTY(EditAnywhere, Category = "Draw Options", Meta = (EditCondition = "!bDisable"))
+	bool bDrawSampleLabels = true;
+
+	UPROPERTY(EditAnywhere, Category = "Draw Options", Meta = (EditCondition = "!bDisable"))
+	bool bDrawSamplesWithColorGradient = true;
 };
 
 /**
@@ -122,7 +129,7 @@ class SDebuggerDatabaseView : public SCompoundWidget
 	void Update(const FTraceMotionMatchingStateMessage& State, const UPoseSearchDatabase& Database);
 
 	const TSharedPtr<SListView<TSharedRef<FDebuggerDatabaseRowData>>>& GetActiveRow() const { return ActiveView.ListView; }
-	const TSharedPtr<SListView<TSharedRef<FDebuggerDatabaseRowData>>>& GetDatabaseRows() const { return DatabaseView.ListView; }
+	const TSharedPtr<SListView<TSharedRef<FDebuggerDatabaseRowData>>>& GetDatabaseRows() const { return FilteredDatabaseView.ListView; }
 
 	/** Used by database rows to acquire column-specific information */
 	using FColumnMap = TMap<FName, TSharedRef<DebuggerDatabaseColumns::IColumn>>;
@@ -143,6 +150,8 @@ private:
 	/** Sorts the database by the current sort predicate, updating the view order */
 	void SortDatabaseRows();
 
+	void FilterDatabaseRows();
+
 	/** Sets dynamic data for each row, such as score at the current time */
 	void UpdateRows(const FTraceMotionMatchingStateMessage& State, const UPoseSearchDatabase& Database);
 
@@ -159,6 +168,9 @@ private:
 	/** Aligns the active and database views */
 	void OnColumnWidthChanged(const float NewWidth, FName ColumnId) const;
 
+	/** Called when the text in the filter box is modified to update the filtering */
+	void OnFilterTextChanged(const FText& SearchText);
+
 	/** Generates a database row widget for the given data */
 	TSharedRef<ITableRow> HandleGenerateDatabaseRow(TSharedRef<FDebuggerDatabaseRowData> Item, const TSharedRef<STableViewBase>& OwnerTable) const;
 	
@@ -197,8 +209,18 @@ private:
 
 	/** Active row at the top of the view */
 	FTable ActiveView;
-	/** Database listings for all poses */
-	FTable DatabaseView;
+
+	/** All database poses */
+	TArray<TSharedRef<FDebuggerDatabaseRowData>> UnfilteredDatabaseRows;
+
+	/** Database listing for filtered poses */
+	FTable FilteredDatabaseView;
+
+	/** Search box widget */
+	TSharedPtr<SSearchBox> FilterBox;
+
+	/** Text used to filter DatabaseView */
+	FText FilterText;
 };
 
 /**
diff --git a/Engine/Plugins/Experimental/Animation/PoseSearch/Source/Runtime/Private/PoseSearch.cpp b/Engine/Plugins/Experimental/Animation/PoseSearch/Source/Runtime/Private/PoseSearch.cpp
--- a/Engine/Plugins/Experimental/Animation/PoseSearch/Source/Runtime/Private/PoseSearch.cpp
+++ b/Engine/Plugins/Experimental/Animation/PoseSearch/Source/Runtime/Private/PoseSearch.cpp
@@ -42,6 +42,9 @@ constexpr float DrawDebugVelocityScale = 0.08f;
 constexpr float DrawDebugArrowSize = 30.0f;
 constexpr float DrawDebugSphereSize = 3.0f;
 constexpr int32 DrawDebugSphereSegments = 10;
+constexpr float DrawDebugGradientStrength = 0.8f;
+constexpr float DrawDebugSampleLabelFontScale = 1.0f;
+static const FVector DrawDebugSampleLabelOffset = FVector(0.0f, 0.0f, -10.0f);
 
 static bool IsSamplingRangeValid(FFloatInterval Range)
 {
@@ -2603,6 +2606,17 @@ static void DrawTrajectoryFeatures(const FDebugDrawParams& DrawParams, const FFe
 		return;
 	}
 
+	auto GetGradientColor = [](const FLinearColor& OriginalColor, int SampleIdx, int NumSamples, EDebugDrawFlags Flags)
+	{
+		int Denominator = NumSamples - 1;
+		if (Denominator <= 0 || !EnumHasAnyFlags(Flags, EDebugDrawFlags::DrawSamplesWithColorGradient))
+		{
+			return OriginalColor;
+		}
+
+		return OriginalColor * (1.0f - DrawDebugGradientStrength * (SampleIdx / (float)Denominator));
+	};
+
 	for (int32 SchemaSubsampleIdx = 0; SchemaSubsampleIdx != NumSubsamples; ++SchemaSubsampleIdx)
 	{
 		Feature.SubsampleIdx = SchemaSubsampleIdx;
@@ -2613,7 +2627,8 @@ static void DrawTrajectoryFeatures(const FDebugDrawParams& DrawParams, const FFe
 			Feature.Type = EPoseSearchFeatureType::Position;
 
 			FLinearColor LinearColor = DrawParams.Color ? *DrawParams.Color : GetColorForFeature(Feature, Reader.GetLayout());
-			FColor Color = LinearColor.ToFColor(true);
+			FLinearColor GradientColor = GetGradientColor(LinearColor, SchemaSubsampleIdx, NumSubsamples, DrawParams.Flags);
+			FColor Color = GradientColor.ToFColor(true);
 
 			TrajectoryPos = DrawParams.RootTransform.TransformPosition(TrajectoryPos);
 			if (EnumHasAnyFlags(DrawParams.Flags, EDebugDrawFlags::DrawSearchIndex))
@@ -2636,7 +2651,8 @@ static void DrawTrajectoryFeatures(const FDebugDrawParams& DrawParams, const FFe
 			Feature.Type = EPoseSearchFeatureType::LinearVelocity;
 			
 			FLinearColor LinearColor = DrawParams.Color ? *DrawParams.Color : GetColorForFeature(Feature, Reader.GetLayout());
-			FColor Color = LinearColor.ToFColor(true);
+			FLinearColor GradientColor = GetGradientColor(LinearColor, SchemaSubsampleIdx, NumSubsamples, DrawParams.Flags);
+			FColor Color = GradientColor.ToFColor(true);
 
 			TrajectoryVel *= DrawDebugVelocityScale;
 			TrajectoryVel = DrawParams.RootTransform.TransformVector(TrajectoryVel);
@@ -2667,7 +2683,8 @@ static void DrawTrajectoryFeatures(const FDebugDrawParams& DrawParams, const FFe
 			Feature.Type = EPoseSearchFeatureType::ForwardVector;
 
 			FLinearColor LinearColor = DrawParams.Color ? *DrawParams.Color : GetColorForFeature(Feature, Reader.GetLayout());
-			FColor Color = LinearColor.ToFColor(true);
+			FLinearColor GradientColor = GetGradientColor(LinearColor, SchemaSubsampleIdx, NumSubsamples, DrawParams.Flags);
+			FColor Color = GradientColor.ToFColor(true);
 
 			TrajectoryForward = DrawParams.RootTransform.TransformVector(TrajectoryForward);
 			if (EnumHasAnyFlags(DrawParams.Flags, EDebugDrawFlags::DrawSearchIndex))
@@ -2689,6 +2706,24 @@ static void DrawTrajectoryFeatures(const FDebugDrawParams& DrawParams, const FFe
 				);
 			}
 		}
+
+		if (EnumHasAnyFlags(DrawParams.Flags, EDebugDrawFlags::DrawSampleLabels))
+		{
+			FLinearColor LinearColor = DrawParams.Color ? *DrawParams.Color : GetColorForFeature(Feature, Reader.GetLayout());
+			FLinearColor GradientColor = GetGradientColor(LinearColor, SchemaSubsampleIdx, NumSubsamples, DrawParams.Flags);
+			FColor Color = GradientColor.ToFColor(true);
+
+			FString SampleLabel = FString::Format(TEXT("{0}"), { SchemaSubsampleIdx });
+			DrawDebugString(
+				DrawParams.World,
+				TrajectoryPos + DrawDebugSampleLabelOffset,
+				SampleLabel,
+				nullptr,
+				Color,
+				LifeTime,
+				false,
+				DrawDebugSampleLabelFontScale);
+		}
 	}
 }
 
diff --git a/Engine/Plugins/Experimental/Animation/PoseSearch/Source/Runtime/Public/PoseSearch/PoseSearch.h b/Engine/Plugins/Experimental/Animation/PoseSearch/Source/Runtime/Public/PoseSearch/PoseSearch.h
--- a/Engine/Plugins/Experimental/Animation/PoseSearch/Source/Runtime/Public/PoseSearch/PoseSearch.h
+++ b/Engine/Plugins/Experimental/Animation/PoseSearch/Source/Runtime/Public/PoseSearch/PoseSearch.h
@@ -708,6 +708,10 @@ enum class EDebugDrawFlags : uint32
 	IncludeAllFeatures  = IncludePose | IncludeTrajectory,
 
 	Persistent = 1 << 3,
+	
+	DrawSampleLabels = 1 << 4,
+
+	DrawSamplesWithColorGradient = 1 << 5,
 };
 ENUM_CLASS_FLAGS(EDebugDrawFlags);
 
