MotionMatchingFFDemo - Added pose search anim notify states
- Block Transition: used to stop animation sequence segments from being returned from the pose search
- Modify Cost: used to change the cost of an animtion segment to make it more or less likely to be selected by pose search

diff --git a/Engine/Plugins/Experimental/Animation/PoseSearch/Source/Runtime/Private/PoseSearch.cpp b/Engine/Plugins/Experimental/Animation/PoseSearch/Source/Runtime/Private/PoseSearch.cpp
--- a/Engine/Plugins/Experimental/Animation/PoseSearch/Source/Runtime/Private/PoseSearch.cpp
+++ b/Engine/Plugins/Experimental/Animation/PoseSearch/Source/Runtime/Private/PoseSearch.cpp
@@ -1,6 +1,7 @@
 // Copyright Epic Games, Inc. All Rights Reserved.
 
 #include "PoseSearch/PoseSearch.h"
+#include "PoseSearch/PoseSearchAnimNotifies.h"
 #include "PoseSearchEigenHelper.h"
 
 #include "Algo/BinarySearch.h"
@@ -1740,6 +1741,10 @@ public:
 	// beyond the sequence limits when Time is less than zero or greater than the sequence length
 	float ExtractRootDistance(float Time) const;
 
+	// Extracts notify states inheriting from UAnimNotifyState_PoseSearchBase present in the sequence at Time.
+	// The function does not empty NotifyStates before adding new notifies!
+	void ExtractPoseSearchNotifyStates(float Time, TArray<class UAnimNotifyState_PoseSearchBase*>& NotifyStates) const;
+
 private:
 	void Reserve();
 	void ProcessRootMotion();
@@ -1750,6 +1755,7 @@ private:
 	// Uses distance delta between NextRootDistanceIndex and NextRootDistanceIndex - 1 and extrapolates it to 
 	// ExtrapolationTime
 	float ExtrapolateRootDistance(int32 NextRootDistanceIndex, float ExtrapolationTime) const;
+
 };
 
 void FSequenceSampler::Init(const FInput& InInput)
@@ -1865,6 +1871,40 @@ float FSequenceSampler::ExtractRootDistance(float Time) const
 	return Distance;
 }
 
+void FSequenceSampler::ExtractPoseSearchNotifyStates(
+	float Time, 
+	TArray<UAnimNotifyState_PoseSearchBase*>& NotifyStates) const
+{
+	// getting pose search notifies in an interval of size ExtractionInterval, centered on Time
+	constexpr float ExtractionInterval = 1.0f / 120.0f;
+	FAnimNotifyContext NotifyContext;
+	Input.Sequence->GetAnimNotifies(Time - (ExtractionInterval * 0.5f), ExtractionInterval, NotifyContext);
+
+	// check which notifies actually overlap Time and are of the right base type
+	for (const FAnimNotifyEventReference& EventReference : NotifyContext.ActiveNotifies)
+	{
+		const FAnimNotifyEvent* NotifyEvent = EventReference.GetNotify();
+		if (!NotifyEvent)
+		{
+			continue;
+		}
+
+		if (NotifyEvent->GetTriggerTime() > Time ||
+			NotifyEvent->GetEndTriggerTime() < Time)
+		{
+			continue;
+		}
+
+		UAnimNotifyState_PoseSearchBase* PoseSearchAnimNotify = 
+			Cast<UAnimNotifyState_PoseSearchBase>(NotifyEvent->NotifyStateClass);
+		if (PoseSearchAnimNotify)
+		{
+			NotifyStates.Add(PoseSearchAnimNotify);
+		}
+	}
+}
+
+
 FTransform FSequenceSampler::ExtrapolateRootMotion(float SampleStart, float SampleEnd, float ExtrapolationTime) const
 {
 	const float SampleDelta = SampleEnd - SampleStart;
@@ -2050,6 +2090,7 @@ public:
 		const FSequenceSampler* LeadInSequence = nullptr;
 		const FSequenceSampler* FollowUpSequence = nullptr;
 		FFloatInterval RequestedSamplingRange = FFloatInterval(0.0f, 0.0f);
+		FPoseSearchBlockTransitionParameters BlockTransitionParameters;
 	} Input;
 
 	struct FOutput
@@ -2058,6 +2099,7 @@ public:
 		int32 LastIndexedSample = 0;
 		int32 NumIndexedPoses = 0;
 		TArray<float> FeatureVectorTable;
+		TArray<FPoseSearchPoseMetadata> PoseMetadata;
 	} Output;
 
 	void Reset();
@@ -2066,6 +2108,7 @@ public:
 
 private:
 	FPoseSearchFeatureVectorBuilder FeatureVector;
+	FPoseSearchPoseMetadata Metadata;
 
 	struct FSampleInfo
 	{
@@ -2087,6 +2130,7 @@ private:
 	void AddPoseFeatures(int32 SampleIdx);
 	void AddTrajectoryTimeFeatures(int32 SampleIdx);
 	void AddTrajectoryDistanceFeatures(int32 SampleIdx);
+	void AddMetadata(int32 SampleIdx);
 };
 
 void FSequenceIndexer::Reset()
@@ -2096,11 +2140,13 @@ void FSequenceIndexer::Reset()
 	Output.NumIndexedPoses = 0;
 
 	Output.FeatureVectorTable.Reset(0);
+	Output.PoseMetadata.Reset(0);
 }
 
 void FSequenceIndexer::Reserve()
 {
 	Output.FeatureVectorTable.SetNumZeroed(Input.Schema->Layout.NumFloats * Output.NumIndexedPoses);
+	Output.PoseMetadata.SetNum(Output.NumIndexedPoses);
 }
 
 void FSequenceIndexer::Init(const FInput& InSettings)
@@ -2131,6 +2177,7 @@ void FSequenceIndexer::Process()
 		AddPoseFeatures(SampleIdx);
 		AddTrajectoryTimeFeatures(SampleIdx);
 		AddTrajectoryDistanceFeatures(SampleIdx);
+		AddMetadata(SampleIdx);
 
 		SampleEnd(SampleIdx);
 	}
@@ -2145,13 +2192,17 @@ void FSequenceIndexer::SampleEnd(int32 SampleIdx)
 {
 	check(FeatureVector.IsComplete());
 
-	int32 FirstValueIdx = (SampleIdx - Output.FirstIndexedSample) * Input.Schema->Layout.NumFloats;
+	const int32 PoseIdx = SampleIdx - Output.FirstIndexedSample;
+
+	const int32 FirstValueIdx = PoseIdx * Input.Schema->Layout.NumFloats;
 	TArrayView<float> WriteValues = MakeArrayView(&Output.FeatureVectorTable[FirstValueIdx], Input.Schema->Layout.NumFloats);
 
 	TArrayView<const float> ReadValues = FeatureVector.GetValues();
 	
 	check(WriteValues.Num() == ReadValues.Num());
 	FMemory::Memcpy(WriteValues.GetData(), ReadValues.GetData(), WriteValues.Num() * WriteValues.GetTypeSize());
+
+	Output.PoseMetadata[PoseIdx] = Metadata;
 }
 
 const float FSequenceIndexer::GetSampleTimeFromDistance(float SampleDistance) const
@@ -2583,6 +2634,39 @@ void FSequenceIndexer::AddTrajectoryDistanceFeatures(int32 SampleIdx)
 	}
 }
 
+void FSequenceIndexer::AddMetadata(int32 SampleIdx)
+{
+	const float SequenceLength = Input.MainSequence->Output.PlayLength;
+	const float SampleTime = FMath::Min(SampleIdx * Input.Schema->SamplingInterval, SequenceLength);
+
+	Metadata = FPoseSearchPoseMetadata();
+
+	const bool bBlockTransition =
+		SampleTime < Input.BlockTransitionParameters.SequenceStartInterval ||
+		SampleTime > SequenceLength - Input.BlockTransitionParameters.SequenceEndInterval;
+
+	if (bBlockTransition)
+	{
+		EnumAddFlags(Metadata.Flags, EPoseSearchPoseFlags::BlockTransition);
+	}
+
+	TArray<UAnimNotifyState_PoseSearchBase*> NotifyStates;
+	Input.MainSequence->ExtractPoseSearchNotifyStates(SampleTime, NotifyStates);
+	for (const UAnimNotifyState_PoseSearchBase* PoseSearchNotify : NotifyStates)
+	{
+		if (PoseSearchNotify->GetClass()->IsChildOf<UAnimNotifyState_PoseSearchBlockTransition>())
+		{
+			EnumAddFlags(Metadata.Flags, EPoseSearchPoseFlags::BlockTransition);
+		}
+		else if (PoseSearchNotify->GetClass()->IsChildOf<UAnimNotifyState_PoseSearchModifyCost>())
+		{
+			const UAnimNotifyState_PoseSearchModifyCost* ModifyCostNotify =
+				Cast<const UAnimNotifyState_PoseSearchModifyCost>(PoseSearchNotify);
+			Metadata.CostAddend = ModifyCostNotify->CostAddend;
+		}
+	}
+}
+
 
 //////////////////////////////////////////////////////////////////////////
 // PoseSearch API
@@ -3421,6 +3505,7 @@ bool BuildIndex(UPoseSearchDatabase* Database)
 		Input.LeadInSequence = GetSampler(DbSequence.LeadInSequence);
 		Input.FollowUpSequence = GetSampler(DbSequence.FollowUpSequence);
 		Input.Schema = Database->Schema;
+		Input.BlockTransitionParameters = Database->BlockTransitionParameters;
 		Input.RequestedSamplingRange = DbSequence.SamplingRange;
 		Indexer.Init(Input);
 	}
@@ -3444,10 +3529,12 @@ bool BuildIndex(UPoseSearchDatabase* Database)
 
 	// Join animation data into a single search index
 	Database->SearchIndex.Values.Reset(TotalFloats);
+	Database->SearchIndex.PoseMetadata.Reset(TotalPoses);
 	for (const FSequenceIndexer& Indexer : Indexers)
 	{
 		const FSequenceIndexer::FOutput& Output = Indexer.Output;
 		Database->SearchIndex.Values.Append(Output.FeatureVectorTable.GetData(), Output.FeatureVectorTable.Num());
+		Database->SearchIndex.PoseMetadata.Append(Output.PoseMetadata);
 	}
 
 	Database->SearchIndex.NumPoses = TotalPoses;
@@ -3458,7 +3545,7 @@ bool BuildIndex(UPoseSearchDatabase* Database)
 	return true;
 }
 
-static FSearchResult Search(const FPoseSearchIndex& SearchIndex, TArrayView<const float> Query, const FPoseSearchWeightsContext* WeightsContext = nullptr, const TSet<int32>& ExcludedIndices = TSet<int32>())
+static FSearchResult Search(const FPoseSearchIndex& SearchIndex, TArrayView<const float> Query, const FPoseSearchWeightsContext* WeightsContext = nullptr)
 {
 	FSearchResult Result;
 	if (!ensure(SearchIndex.IsValid()))
@@ -3476,16 +3563,19 @@ static FSearchResult Search(const FPoseSearchIndex& SearchIndex, TArrayView<cons
 
 	for (int32 PoseIdx = 0; PoseIdx != SearchIndex.NumPoses; ++PoseIdx)
 	{
-		if (ExcludedIndices.Contains(PoseIdx))
+		const FPoseSearchPoseMetadata& Metadata = SearchIndex.PoseMetadata[PoseIdx];
+
+		if (EnumHasAnyFlags(Metadata.Flags, EPoseSearchPoseFlags::BlockTransition))
 		{
 			continue;
 		}
 
 		const float PoseDissimilarity = ComparePoses(SearchIndex, PoseIdx, Query, WeightsContext);
+		const float ModifiedPoseDissimilarity = PoseDissimilarity + Metadata.CostAddend;
 		
-		if (PoseDissimilarity < BestPoseDissimilarity)
+		if (ModifiedPoseDissimilarity < BestPoseDissimilarity)
 		{
-			BestPoseDissimilarity = PoseDissimilarity;
+			BestPoseDissimilarity = ModifiedPoseDissimilarity;
 			BestPoseIdx = PoseIdx;
 		}
 	}
@@ -3531,7 +3621,7 @@ FDbSearchResult Search(
 	const UPoseSearchDatabase* Database, 
 	TArrayView<const float> Query, 
 	const FPoseSearchWeightsContext* WeightsContext, 
-	const float EndTimeToExclude, FDebugDrawParams DebugDrawParams)
+	FDebugDrawParams DebugDrawParams)
 {
 	if (!ensure(Database && Database->IsValidForSearch()))
 	{
@@ -3540,23 +3630,7 @@ FDbSearchResult Search(
 
 	const FPoseSearchIndex& SearchIndex = Database->SearchIndex;
 
-	const int32 NumSequenceIndicesToExclude = Database->Schema->SampleRate * EndTimeToExclude;
-	TSet<int32> ExcludedIndices;
-	if (NumSequenceIndicesToExclude > 0)
-	{
-		for (const FPoseSearchDatabaseSequence& Sequence : Database->Sequences)
-		{
-			// leaving at least one sample per sequence
-			const int32 ExclusionStartIndex = Sequence.FirstPoseIdx + FMath::Max(1, Sequence.NumPoses - NumSequenceIndicesToExclude);
-			const int32 ExclusionEndIndex = Sequence.FirstPoseIdx + Sequence.NumPoses;
-			for (int32 ExcludedIndex = ExclusionStartIndex; ExcludedIndex < ExclusionEndIndex; ++ExcludedIndex)
-			{
-				ExcludedIndices.Add(ExcludedIndex);
-			}
-		}
-	}
-	
-	FDbSearchResult Result = Search(SearchIndex, Query, WeightsContext, ExcludedIndices);
+	FDbSearchResult Result = Search(SearchIndex, Query, WeightsContext);
 	if (!Result.IsValid())
 	{
 		return FDbSearchResult();
diff --git a/Engine/Plugins/Experimental/Animation/PoseSearch/Source/Runtime/Private/PoseSearchAnimNotifies.cpp b/Engine/Plugins/Experimental/Animation/PoseSearch/Source/Runtime/Private/PoseSearchAnimNotifies.cpp
new file mode 100644
--- /dev/null
+++ b/Engine/Plugins/Experimental/Animation/PoseSearch/Source/Runtime/Private/PoseSearchAnimNotifies.cpp
@@ -0,0 +1,3 @@
+// Copyright Epic Games, Inc. All Rights Reserved.
+
+#include "PoseSearch/PoseSearchAnimNotifies.h"
\ No newline at end of file
diff --git a/Engine/Plugins/Experimental/Animation/PoseSearch/Source/Runtime/Private/PoseSearchLibrary.cpp b/Engine/Plugins/Experimental/Animation/PoseSearch/Source/Runtime/Private/PoseSearchLibrary.cpp
--- a/Engine/Plugins/Experimental/Animation/PoseSearch/Source/Runtime/Private/PoseSearchLibrary.cpp
+++ b/Engine/Plugins/Experimental/Animation/PoseSearch/Source/Runtime/Private/PoseSearchLibrary.cpp
@@ -128,8 +128,7 @@ void UpdateMotionMatchingState(const FAnimationUpdateContext& Context
 		UE::PoseSearch::FDbSearchResult Result = UE::PoseSearch::Search(
 			Database,
 			InOutMotionMatchingState.ComposedQuery.GetNormalizedValues(),
-			&InOutMotionMatchingState.WeightsContext,
-			Settings.SequenceEndExlusionTime
+			&InOutMotionMatchingState.WeightsContext
 		);
 		if (Result.IsValid() && ((InOutMotionMatchingState.ElapsedPoseJumpTime >= Settings.SearchThrottleTime) || !bCanContinue))
 		{
diff --git a/Engine/Plugins/Experimental/Animation/PoseSearch/Source/Runtime/Public/PoseSearch/PoseSearch.h b/Engine/Plugins/Experimental/Animation/PoseSearch/Source/Runtime/Public/PoseSearch/PoseSearch.h
--- a/Engine/Plugins/Experimental/Animation/PoseSearch/Source/Runtime/Public/PoseSearch/PoseSearch.h
+++ b/Engine/Plugins/Experimental/Animation/PoseSearch/Source/Runtime/Public/PoseSearch/PoseSearch.h
@@ -251,6 +251,30 @@ struct POSESEARCH_API FPoseSearchIndexPreprocessInfo
 	}
 };
 
+UENUM()
+enum class EPoseSearchPoseFlags : uint32
+{
+	None = 0,
+
+	/** Can be returned in a pose search */
+	BlockTransition = 1 << 0,
+};
+ENUM_CLASS_FLAGS(EPoseSearchPoseFlags);
+
+// This is kept for each pose in the search index along side the feature vector values and is used to influence the
+// search.
+USTRUCT()
+struct POSESEARCH_API FPoseSearchPoseMetadata
+{
+	GENERATED_BODY()
+
+	UPROPERTY()
+	EPoseSearchPoseFlags Flags = EPoseSearchPoseFlags::None;
+
+	UPROPERTY()
+	float CostAddend = 0.0f;
+};
+
 /**
 * A search index for animation poses. The structure of the search index is determined by its UPoseSearchSchema.
 * May represent a single animation (see UPoseSearchSequenceMetaData) or a collection (see UPoseSearchDatabase).
@@ -266,6 +290,9 @@ struct POSESEARCH_API FPoseSearchIndex
 	UPROPERTY()
 	TArray<float> Values;
 
+	UPROPERTY()
+	TArray<FPoseSearchPoseMetadata> PoseMetadata;
+
 	UPROPERTY()
 	TObjectPtr<const UPoseSearchSchema> Schema = nullptr;
 
@@ -301,6 +328,20 @@ public:
 	float SampleTime = 0.05f;
 };
 
+USTRUCT()
+struct FPoseSearchBlockTransitionParameters
+{
+	GENERATED_BODY()
+
+	// Excluding the beginning of sequences can help ensure an exact past trajectory is used when building the features
+	float SequenceStartInterval = 0.0f;
+
+	// Excluding the end of sequences help ensure an exact future trajectory, and also prevents the selection of
+	// a sequence which will end too soon to be worth selecting.
+	float SequenceEndInterval = 0.2f;
+
+};
+
 /** Animation metadata object for indexing a single animation. */
 UCLASS(BlueprintType, Category = "Animation|Pose Search")
 class POSESEARCH_API UPoseSearchSequenceMetaData : public UAnimMetaData
@@ -527,6 +568,9 @@ public:
 	UPROPERTY(EditAnywhere, Category = "Database")
 	FPoseSearchExtrapolationParameters ExtrapolationParameters;
 
+	UPROPERTY(EditAnywhere, Category = "Database")
+	FPoseSearchBlockTransitionParameters BlockTransitionParameters;
+
 	// Drag and drop animations here to add them in bulk to Sequences
 	UPROPERTY(EditAnywhere, Category = "Database", DisplayName="Drag And Drop Anims Here")
 	TArray<TObjectPtr<UAnimSequence>> SimpleSequences;
@@ -815,7 +859,7 @@ POSESEARCH_API FSearchResult Search(const UAnimSequenceBase* Sequence, TArrayVie
 * 
 * @return The pose in the database that most closely matches the Query.
 */
-POSESEARCH_API FDbSearchResult Search(const UPoseSearchDatabase* Database, TArrayView<const float> Query, const FPoseSearchWeightsContext* WeightsContext = nullptr, const float EndTimeToExclude = 0.0f, FDebugDrawParams DrawParams = FDebugDrawParams());
+POSESEARCH_API FDbSearchResult Search(const UPoseSearchDatabase* Database, TArrayView<const float> Query, const FPoseSearchWeightsContext* WeightsContext = nullptr, FDebugDrawParams DrawParams = FDebugDrawParams());
 
 
 /**
diff --git a/Engine/Plugins/Experimental/Animation/PoseSearch/Source/Runtime/Public/PoseSearch/PoseSearchAnimNotifies.h b/Engine/Plugins/Experimental/Animation/PoseSearch/Source/Runtime/Public/PoseSearch/PoseSearchAnimNotifies.h
new file mode 100644
--- /dev/null
+++ b/Engine/Plugins/Experimental/Animation/PoseSearch/Source/Runtime/Public/PoseSearch/PoseSearchAnimNotifies.h
@@ -0,0 +1,37 @@
+// Copyright Epic Games, Inc. All Rights Reserved.
+
+#pragma once
+
+#include "Animation/AnimNotifies/AnimNotifyState.h"
+
+#include "PoseSearchAnimNotifies.generated.h"
+
+// Base class for pose search anim notify states
+UCLASS(Abstract)
+class POSESEARCH_API UAnimNotifyState_PoseSearchBase : public UAnimNotifyState
+{
+	GENERATED_BODY()
+};
+
+// A pose matching search will not return results that overlap this notify, but the animation segment can still play
+// if a previous search result advances into it.
+UCLASS(Blueprintable, meta = (DisplayName = "Pose Matching: Block Transition"))
+class POSESEARCH_API UAnimNotifyState_PoseSearchBlockTransition : public UAnimNotifyState_PoseSearchBase
+{
+	GENERATED_BODY()
+};
+
+// Pose matching cost will be affected by this, making the animation segment more or less likely to be selected based
+// on the notify parameters
+UCLASS(Blueprintable, meta = (DisplayName = "Pose Matching: Modify Cost"))
+class POSESEARCH_API UAnimNotifyState_PoseSearchModifyCost : public UAnimNotifyState_PoseSearchBase
+{
+	GENERATED_BODY()
+
+public:
+
+	// A negative value reduces the cost and makes the segment more likely to be chosen. A positive value, conversely,
+	// makes the segment less likely to be chosen
+	UPROPERTY(EditAnywhere, Category = Config, meta = (DisplayName = "Modifier"))
+	float CostAddend = -1.0f;
+};
diff --git a/Engine/Plugins/Experimental/Animation/PoseSearch/Source/Runtime/Public/PoseSearch/PoseSearchLibrary.h b/Engine/Plugins/Experimental/Animation/PoseSearch/Source/Runtime/Public/PoseSearch/PoseSearchLibrary.h
--- a/Engine/Plugins/Experimental/Animation/PoseSearch/Source/Runtime/Public/PoseSearch/PoseSearchLibrary.h
+++ b/Engine/Plugins/Experimental/Animation/PoseSearch/Source/Runtime/Public/PoseSearch/PoseSearchLibrary.h
@@ -66,11 +66,6 @@ struct POSESEARCH_API FMotionMatchingSettings
 	// Note: This feature won't work quite as advertised until search data rescaling is implemented
 	UPROPERTY(EditAnywhere, BlueprintReadWrite, Category=Settings, meta=(ClampMin="0", ClampMax="100"))
 	float MinPercentImprovement = 100.f;
-
-	// Pose indices SequenceEndExclusionTime or less away from the end of the database sequence will be ignored by the 
-	// PoseSearch. This ensures the search doesn't get stuck at the end of a sequence.
-	UPROPERTY(EditAnywhere, BlueprintReadWrite, Category = Settings, meta = (ClampMin = "0"))
-	float SequenceEndExlusionTime = 0.25f;
 };
 
 USTRUCT(BlueprintType, Category="Animation|Pose Search")
