MotionMatchingFFDemo - ExcludeFromDatabase notify state
This notify state completely removes animation segments from the database, they'll never be returned from a search or play.
To implement this, FPoseSearchIndex now hold an array of FPoseSearchIndexAsset, which contains information about continuous animation segments to index. If a database sequence contains an ExcludeFromDatabase notify that splits the animation frames into more than 1 segment, multiple FPoseSearchIndexAsset entries will be added to the array, one for each segment.

diff --git a/Engine/Plugins/Experimental/Animation/PoseSearch/Source/Editor/Private/PoseSearchDebugger.cpp b/Engine/Plugins/Experimental/Animation/PoseSearch/Source/Editor/Private/PoseSearchDebugger.cpp
--- a/Engine/Plugins/Experimental/Animation/PoseSearch/Source/Editor/Private/PoseSearchDebugger.cpp
+++ b/Engine/Plugins/Experimental/Animation/PoseSearch/Source/Editor/Private/PoseSearchDebugger.cpp
@@ -525,28 +525,24 @@ void SDebuggerDatabaseView::CreateRows(const UPoseSearchDatabase& Database)
 	UnfilteredDatabaseRows.Reserve(NumPoses);
 
 	// Build database rows
-	for(const FPoseSearchDatabaseSequence& DbSequence : Database.Sequences)
+	for(const FPoseSearchIndexAsset& SearchIndexAsset : Database.SearchIndex.Assets)
 	{
-		const int32 LastPoseIdx = DbSequence.FirstPoseIdx + DbSequence.NumPoses;
-		for (int32 PoseIdx = DbSequence.FirstPoseIdx; PoseIdx != LastPoseIdx; ++PoseIdx)
+		const FPoseSearchDatabaseSequence& DbSequence = Database.GetSourceAsset(&SearchIndexAsset);
+		const int32 LastPoseIdx = SearchIndexAsset.FirstPoseIdx + SearchIndexAsset.NumPoses;
+		for (int32 PoseIdx = SearchIndexAsset.FirstPoseIdx; PoseIdx != LastPoseIdx; ++PoseIdx)
 		{
 			TSharedRef<FDebuggerDatabaseRowData> Row = UnfilteredDatabaseRows.Add_GetRef(MakeShared<FDebuggerDatabaseRowData>());
 			Row->PoseIdx = PoseIdx;
 
 			const FString SequenceName = DbSequence.Sequence->GetName();
 			const float SequenceLength = DbSequence.Sequence->GetPlayLength();
-			FFloatInterval Range = DbSequence.SamplingRange;
-
-			// @TODO: Update this when range is computed natively as part of the sequence class
-			const bool bSampleAll = (Range.Min == 0.0f) && (Range.Max == 0.0f);
-			const float SequencePlayLength = DbSequence.Sequence->GetPlayLength();
-			Range.Min = bSampleAll ? 0.0f : Range.Min;
-			Range.Max = bSampleAll ? SequencePlayLength : FMath::Min(SequencePlayLength, Range.Max);
-			// ---
+			FFloatInterval Range = SearchIndexAsset.SamplingInterval;
 
 			Row->AnimSequenceName = SequenceName;
 			// Cap time in sequence to end of range
-			Row->Time = FGenericPlatformMath::Min(Range.Min + (PoseIdx - DbSequence.FirstPoseIdx) * Database.Schema->SamplingInterval, Range.Max);
+			Row->Time = FGenericPlatformMath::Min(
+				Range.Min + (PoseIdx - SearchIndexAsset.FirstPoseIdx) * Database.Schema->SamplingInterval, 
+				Range.Max);
 			Row->Length = SequenceLength;
 		}
 	}
diff --git a/Engine/Plugins/Experimental/Animation/PoseSearch/Source/Runtime/Private/AnimNode_MotionMatching.cpp b/Engine/Plugins/Experimental/Animation/PoseSearch/Source/Runtime/Private/AnimNode_MotionMatching.cpp
--- a/Engine/Plugins/Experimental/Animation/PoseSearch/Source/Runtime/Private/AnimNode_MotionMatching.cpp
+++ b/Engine/Plugins/Experimental/Animation/PoseSearch/Source/Runtime/Private/AnimNode_MotionMatching.cpp
@@ -55,7 +55,9 @@ void FAnimNode_MotionMatching::UpdateAssetPlayer(const FAnimationUpdateContext&
 	// If a new pose is requested, jump to the pose by updating the embedded sequence player node
 	if ((MotionMatchingState.Flags & EMotionMatchingFlags::JumpedToPose) == EMotionMatchingFlags::JumpedToPose)
 	{
-		const FPoseSearchDatabaseSequence& ResultDbSequence = Database->Sequences[MotionMatchingState.DbSequenceIdx];
+		const FPoseSearchIndexAsset* SearchIndexAsset = 
+			&MotionMatchingState.CurrentDatabase->SearchIndex.Assets[MotionMatchingState.SearchIndexAssetIdx];
+		const FPoseSearchDatabaseSequence& ResultDbSequence = Database->GetSourceAsset(SearchIndexAsset);
 		SequencePlayerNode.SetSequence(ResultDbSequence.Sequence);
 		SequencePlayerNode.SetAccumulatedTime(MotionMatchingState.AssetPlayerTime);
 		SequencePlayerNode.SetLoopAnimation(ResultDbSequence.bLoopAnimation);
diff --git a/Engine/Plugins/Experimental/Animation/PoseSearch/Source/Runtime/Private/PoseSearch.cpp b/Engine/Plugins/Experimental/Animation/PoseSearch/Source/Runtime/Private/PoseSearch.cpp
--- a/Engine/Plugins/Experimental/Animation/PoseSearch/Source/Runtime/Private/PoseSearch.cpp
+++ b/Engine/Plugins/Experimental/Animation/PoseSearch/Source/Runtime/Private/PoseSearch.cpp
@@ -29,11 +29,10 @@
 
 IMPLEMENT_ANIMGRAPH_MESSAGE(UE::PoseSearch::IPoseHistoryProvider);
 
+DEFINE_LOG_CATEGORY(LogPoseSearch);
 
 namespace UE { namespace PoseSearch {
 
-DEFINE_LOG_CATEGORY(LogPoseSearch);
-
 //////////////////////////////////////////////////////////////////////////
 // Constants and utilities
 
@@ -835,6 +834,37 @@ const FPoseSearchWeights* FPoseSearchWeightsContext::GetGroupWeights(int32 Weigh
 	return nullptr;
 }
 
+//////////////////////////////////////////////////////////////////////////
+// FPoseSearchIndexAsset
+// 
+
+int32 FPoseSearchIndex::FindAssetIndex(const FPoseSearchIndexAsset* Asset) const
+{
+	if (Asset == nullptr || Assets.Num() == 0)
+	{
+		return INDEX_NONE;
+	}
+
+	const FPoseSearchIndexAsset* Start = &Assets[0];
+	int32 Result = Asset - Start;
+
+	if (!Assets.IsValidIndex(Result))
+	{
+		return INDEX_NONE;
+	}
+
+	return Result;
+}
+
+const FPoseSearchIndexAsset* FPoseSearchIndex::FindAssetForPose(int32 PoseIdx) const
+{
+	auto Predicate = [PoseIdx](const FPoseSearchIndexAsset& Asset)
+	{
+		return Asset.IsPoseInRange(PoseIdx);
+	};
+	return Assets.FindByPredicate(Predicate);
+}
+
 
 //////////////////////////////////////////////////////////////////////////
 // FPoseSearchIndex
@@ -842,7 +872,8 @@ const FPoseSearchWeights* FPoseSearchWeightsContext::GetGroupWeights(int32 Weigh
 bool FPoseSearchIndex::IsValid() const
 {
 	bool bSchemaValid = Schema && Schema->IsValid();
-	bool bSearchIndexValid = bSchemaValid && (NumPoses * Schema->Layout.NumFloats == Values.Num());
+	bool bAssetsValid = Assets.Num() > 0;
+	bool bSearchIndexValid = bSchemaValid && bAssetsValid && (NumPoses * Schema->Layout.NumFloats == Values.Num());
 
 	return bSearchIndexValid;
 }
@@ -950,51 +981,74 @@ bool UPoseSearchSequenceMetaData::IsValidForSearch() const
 	return IsValidForIndexing() && SearchIndex.IsValid();
 }
 
+//////////////////////////////////////////////////////////////////////////
+// FPoseSearchDatabaseSequence
+FFloatInterval FPoseSearchDatabaseSequence::GetEffectiveSamplingRange() const
+{
+	const bool bSampleAll = (SamplingRange.Min == 0.0f) && (SamplingRange.Max == 0.0f);
+	const float SequencePlayLength = Sequence->GetPlayLength();
+	FFloatInterval Range;
+	Range.Min = bSampleAll ? 0.0f : SamplingRange.Min;
+	Range.Max = bSampleAll ? SequencePlayLength : FMath::Min(SequencePlayLength, SamplingRange.Max);
+	return Range;
+}
 
 //////////////////////////////////////////////////////////////////////////
 // UPoseSearchDatabase
 
-int32 UPoseSearchDatabase::FindSequenceForPose(int32 PoseIdx) const
+int32 UPoseSearchDatabase::GetPoseIndexFromAssetTime(float AssetTime, const FPoseSearchIndexAsset* SearchIndexAsset) const
 {
-	auto Predicate = [PoseIdx](const FPoseSearchDatabaseSequence& DbSequence)
-	{
-		return (PoseIdx >= DbSequence.FirstPoseIdx) && (PoseIdx < DbSequence.FirstPoseIdx + DbSequence.NumPoses);
-	};
+	const FPoseSearchDatabaseSequence& DbSequence = Sequences[SearchIndexAsset->SourceAssetIdx];
 
-	return Sequences.IndexOfByPredicate(Predicate);
-}
+	FFloatInterval Range = SearchIndexAsset->SamplingInterval;
 
-int32 UPoseSearchDatabase::GetPoseIndexFromAssetTime(int32 DbSequenceIdx, float AssetTime) const
-{
-	const FPoseSearchDatabaseSequence& DbSequence = Sequences[DbSequenceIdx];
-	FFloatInterval Range = UE::PoseSearch::GetEffectiveSamplingRange(DbSequence.Sequence, DbSequence.SamplingRange);
 	if (Range.Contains(AssetTime))
 	{
-		int32 PoseOffset = FMath::FloorToInt(Schema->SampleRate * (AssetTime - Range.Min));
-		if (PoseOffset >= DbSequence.NumPoses)
+		int32 PoseOffset = FMath::RoundToInt(Schema->SampleRate * (AssetTime - Range.Min));
+		if (PoseOffset >= SearchIndexAsset->NumPoses)
 		{
 			if (DbSequence.bLoopAnimation)
 			{
-				PoseOffset -= DbSequence.NumPoses;
+				PoseOffset -= SearchIndexAsset->NumPoses;
 			}
 			else
 			{
-				PoseOffset = DbSequence.NumPoses - 1;
+				PoseOffset = SearchIndexAsset->NumPoses - 1;
 			}
 		}
-		
-		int32 PoseIdx = DbSequence.FirstPoseIdx + PoseOffset;
+		int32 PoseIdx = SearchIndexAsset->FirstPoseIdx + PoseOffset;
 		return PoseIdx;
 	}
-	
+
 	return INDEX_NONE;
 }
 
-FFloatInterval UPoseSearchDatabase::GetEffectiveSamplingRange(int32 DbSequenceIdx) const
+float UPoseSearchDatabase::GetTimeOffset(int32 PoseIdx, const FPoseSearchIndexAsset* Asset) const
 {
-	const FPoseSearchDatabaseSequence& DbSequence = Sequences[DbSequenceIdx];
-	FFloatInterval Range = UE::PoseSearch::GetEffectiveSamplingRange(DbSequence.Sequence, DbSequence.SamplingRange);
-	return Range;
+	const FPoseSearchIndexAsset* ValidAsset = Asset;
+	if (!ValidAsset)
+	{
+		ValidAsset = SearchIndex.FindAssetForPose(PoseIdx);
+		if (!ValidAsset)
+		{
+			UE_LOG(LogPoseSearch, Error, TEXT("Couldn't find asset for pose %i in database %s"), PoseIdx, *GetNameSafe(this));
+			return -1.0f;
+		}
+	}
+	if (!ValidAsset->IsPoseInRange(PoseIdx))
+	{
+		UE_LOG(LogPoseSearch, Error, TEXT("Pose %i out of range in database %s"), PoseIdx, *GetNameSafe(this));
+		return -1.0f;
+	}
+	const FPoseSearchDatabaseSequence& DbSequence = Sequences[ValidAsset->SourceAssetIdx];
+	const FFloatInterval SamplingRange = ValidAsset->SamplingInterval;
+	float TimeOffsetSeconds = SamplingRange.Min + Schema->SamplingInterval * (PoseIdx - Asset->FirstPoseIdx);
+	return TimeOffsetSeconds;
+}
+
+const FPoseSearchDatabaseSequence& UPoseSearchDatabase::GetSourceAsset(const FPoseSearchIndexAsset* SearchIndexAsset) const
+{
+	return Sequences[SearchIndexAsset->SourceAssetIdx];
 }
 
 float UPoseSearchDatabase::GetSequenceLength(int32 DbSequenceIdx) const
@@ -1060,6 +1114,71 @@ void UPoseSearchDatabase::CollectSimpleSequences()
 	SimpleSequences.Reset();
 }
 
+void FindValidSequenceIntervals(const FPoseSearchDatabaseSequence& DbSequence, TArray<FFloatRange>& ValidRanges)
+{
+	const UAnimSequence* Sequence = DbSequence.Sequence;
+	check(DbSequence.Sequence);
+
+	const float SequenceLength = DbSequence.Sequence->GetPlayLength();
+	const FFloatInterval EffectiveSamplingInterval = DbSequence.GetEffectiveSamplingRange();
+
+	// start from a single interval defined by the database sequence sampling range
+	ValidRanges.Empty();
+	ValidRanges.Add(FFloatRange::Inclusive(EffectiveSamplingInterval.Min, EffectiveSamplingInterval.Max));
+
+	FAnimNotifyContext NotifyContext;
+	Sequence->GetAnimNotifies(0.0f, SequenceLength, NotifyContext);
+
+	for (const FAnimNotifyEventReference& EventReference : NotifyContext.ActiveNotifies)
+	{
+		const FAnimNotifyEvent* NotifyEvent = EventReference.GetNotify();
+		if (!NotifyEvent)
+		{
+			continue;
+		}
+
+		const UAnimNotifyState_PoseSearchExcludeFromDatabase* ExclusionNotifyState =
+			Cast<const UAnimNotifyState_PoseSearchExcludeFromDatabase>(NotifyEvent->NotifyStateClass);
+		if (ExclusionNotifyState)
+		{
+			FFloatRange ExclusionRange = 
+				FFloatRange::Inclusive(NotifyEvent->GetTriggerTime(), NotifyEvent->GetEndTriggerTime());
+
+			// Split every valid range based on the exclusion range just found. Because this might increase the 
+			// number of ranges in ValidRanges, the algorithm iterates from end to start.
+			for (int RangeIdx = ValidRanges.Num() - 1; RangeIdx >= 0; --RangeIdx)
+			{
+				FFloatRange EvaluatedRange = ValidRanges[RangeIdx];
+				ValidRanges.RemoveAt(RangeIdx);
+
+				TArray<FFloatRange> Diff = FFloatRange::Difference(EvaluatedRange, ExclusionRange);
+				ValidRanges.Append(Diff);
+			}
+		}
+	}
+}
+
+bool UPoseSearchDatabase::TryInitSearchIndexAssets()
+{
+	SearchIndex.Assets.Empty();
+	bool bAnyMirrored = false;
+	TArray<FFloatRange> ValidRanges;
+	for (int32 SequenceIdx = 0; SequenceIdx < Sequences.Num(); ++SequenceIdx)
+	{
+		ValidRanges.Reset();
+		FindValidSequenceIntervals(Sequences[SequenceIdx], ValidRanges);
+		for (const FFloatRange& Range : ValidRanges)
+		{
+			SearchIndex.Assets.Add(
+				FPoseSearchIndexAsset(
+					SequenceIdx, 
+					FFloatInterval(Range.GetLowerBoundValue(), Range.GetUpperBoundValue())));
+		}
+	}
+
+	return true;
+}
+
 void UPoseSearchDatabase::PreSave(FObjectPreSaveContext ObjectSaveContext)
 {
 	SearchIndex.Reset();
@@ -2157,11 +2276,10 @@ void FSequenceIndexer::Init(const FInput& InSettings)
 
 	Input = InSettings;
 
-	const FFloatInterval SamplingRange = GetEffectiveSamplingRange(Input.MainSequence->Input.Sequence, Input.RequestedSamplingRange);
-
 	Reset();
-	Output.FirstIndexedSample = FMath::FloorToInt(SamplingRange.Min * Input.Schema->SampleRate);
-	Output.LastIndexedSample = FMath::Max(0, FMath::CeilToInt(SamplingRange.Max * Input.Schema->SampleRate));
+	Output.FirstIndexedSample = FMath::FloorToInt(Input.RequestedSamplingRange.Min * Input.Schema->SampleRate);
+	Output.LastIndexedSample = 
+		FMath::Max(0, FMath::CeilToInt(Input.RequestedSamplingRange.Max * Input.Schema->SampleRate));
 	Output.NumIndexedPoses = Output.LastIndexedSample - Output.FirstIndexedSample + 1;
 	Reserve();
 }
@@ -2642,8 +2760,9 @@ void FSequenceIndexer::AddMetadata(int32 SampleIdx)
 	Metadata = FPoseSearchPoseMetadata();
 
 	const bool bBlockTransition =
-		SampleTime < Input.BlockTransitionParameters.SequenceStartInterval ||
-		SampleTime > SequenceLength - Input.BlockTransitionParameters.SequenceEndInterval;
+		!Input.MainSequence->Input.bLoopable &&
+		(SampleTime < Input.BlockTransitionParameters.SequenceStartInterval || 
+		 SampleTime > SequenceLength - Input.BlockTransitionParameters.SequenceEndInterval);
 
 	if (bBlockTransition)
 	{
@@ -3427,6 +3546,14 @@ bool BuildIndex(const UAnimSequence* Sequence, UPoseSearchSequenceMetaData* Sequ
 	SequenceMetaData->SearchIndex.NumPoses = Indexer.Output.NumIndexedPoses;
 	SequenceMetaData->SearchIndex.Schema = SequenceMetaData->Schema;
 
+	SequenceMetaData->SearchIndex.Assets.Empty();
+	FPoseSearchIndexAsset SearchIndexAsset;
+	SearchIndexAsset.SourceAssetIdx = 0;
+	SearchIndexAsset.FirstPoseIdx = 0;
+	SearchIndexAsset.NumPoses = Indexer.Output.NumIndexedPoses;
+	SequenceMetaData->SearchIndex.Assets.Add(SearchIndexAsset);
+
+
 	PreprocessSearchIndex(&SequenceMetaData->SearchIndex);
 
 	return true;
@@ -3441,6 +3568,11 @@ bool BuildIndex(UPoseSearchDatabase* Database)
 		return false;
 	}
 
+	if (!Database->TryInitSearchIndexAssets())
+	{
+		return false;
+	}
+
 	FBoneContainer BoneContainer;
 	BoneContainer.InitializeTo(Database->Schema->BoneIndicesWithParents, FCurveEvaluationOption(false), *Database->Schema->Skeleton);
 
@@ -3493,20 +3625,25 @@ bool BuildIndex(UPoseSearchDatabase* Database)
 
 	// Prepare animation indexing tasks
 	TArray<FSequenceIndexer> Indexers;
-	Indexers.SetNum(Database->Sequences.Num());
-	for (int32 SequenceIdx = 0; SequenceIdx != Database->Sequences.Num(); ++SequenceIdx)
+	Indexers.Reserve(Database->SearchIndex.Assets.Num());
+	for (int32 AssetIdx = 0; AssetIdx != Database->SearchIndex.Assets.Num(); ++AssetIdx)
 	{
-		const FPoseSearchDatabaseSequence& DbSequence = Database->Sequences[SequenceIdx];
-		FSequenceIndexer& Indexer = Indexers[SequenceIdx];
+		const FPoseSearchIndexAsset& SearchIndexAsset = Database->SearchIndex.Assets[AssetIdx];
+		const FPoseSearchDatabaseSequence& DbSequence = Database->Sequences[SearchIndexAsset.SourceAssetIdx];
+		const float SequenceLength = DbSequence.Sequence->GetPlayLength();
 
 		FSequenceIndexer::FInput Input;
 		Input.BoneContainer = &BoneContainer;
 		Input.MainSequence = GetSampler(DbSequence.Sequence);
-		Input.LeadInSequence = GetSampler(DbSequence.LeadInSequence);
-		Input.FollowUpSequence = GetSampler(DbSequence.FollowUpSequence);
+		Input.LeadInSequence = 
+			SearchIndexAsset.SamplingInterval.Min == 0.0f ? GetSampler(DbSequence.LeadInSequence) : nullptr;
+		Input.FollowUpSequence = 
+			SearchIndexAsset.SamplingInterval.Max == SequenceLength ? GetSampler(DbSequence.FollowUpSequence) : nullptr;
 		Input.Schema = Database->Schema;
 		Input.BlockTransitionParameters = Database->BlockTransitionParameters;
-		Input.RequestedSamplingRange = DbSequence.SamplingRange;
+		Input.RequestedSamplingRange = SearchIndexAsset.SamplingInterval;
+
+		FSequenceIndexer& Indexer = Indexers.AddDefaulted_GetRef();
 		Indexer.Init(Input);
 	}
 
@@ -3514,15 +3651,18 @@ bool BuildIndex(UPoseSearchDatabase* Database)
 	ParallelFor(Indexers.Num(), [&Indexers](int32 SequenceIdx){ Indexers[SequenceIdx].Process(); });
 
 
+
 	// Write index info to sequence and count up total poses and storage required
 	int32 TotalPoses = 0;
 	int32 TotalFloats = 0;
-	for (int32 SequenceIdx = 0; SequenceIdx != Database->Sequences.Num(); ++SequenceIdx)
+	for (int32 AssetIdx = 0; AssetIdx != Database->SearchIndex.Assets.Num(); ++AssetIdx)
 	{
-		FPoseSearchDatabaseSequence& DbSequence = Database->Sequences[SequenceIdx];
-		const FSequenceIndexer::FOutput& Output = Indexers[SequenceIdx].Output;
-		DbSequence.NumPoses = Output.NumIndexedPoses;
-		DbSequence.FirstPoseIdx = TotalPoses;
+		const FSequenceIndexer::FOutput& Output = Indexers[AssetIdx].Output;
+
+		FPoseSearchIndexAsset& SearchIndexAsset = Database->SearchIndex.Assets[AssetIdx];
+		SearchIndexAsset.NumPoses = Output.NumIndexedPoses;
+		SearchIndexAsset.FirstPoseIdx = TotalPoses;
+
 		TotalPoses += Output.NumIndexedPoses;
 		TotalFloats += Output.FeatureVectorTable.Num();
 	}
@@ -3584,7 +3724,9 @@ static FSearchResult Search(const FPoseSearchIndex& SearchIndex, TArrayView<cons
 
 	Result.Dissimilarity = BestPoseDissimilarity;
 	Result.PoseIdx = BestPoseIdx;
-	// Result.TimeOffsetSeconds is set by caller
+	Result.SearchIndexAsset = SearchIndex.FindAssetForPose(BestPoseIdx);
+	Result.TimeOffsetSeconds = 
+		Result.SearchIndexAsset->SamplingInterval.Min + SearchIndex.Schema->SamplingInterval * (Result.PoseIdx - Result.SearchIndexAsset->FirstPoseIdx);
 
 	return Result;
 }
@@ -3605,9 +3747,6 @@ FSearchResult Search(const UAnimSequenceBase* Sequence, TArrayView<const float>
 		return Result;
 	}
 
-	const FFloatInterval SamplingRange = GetEffectiveSamplingRange(Sequence, MetaData->SamplingRange);
-	Result.TimeOffsetSeconds = SamplingRange.Min + (SearchIndex.Schema->SamplingInterval * Result.PoseIdx);
-
 	// Do debug visualization
 	DebugDrawParams.SequenceMetaData = MetaData;
 	DebugDrawParams.PoseVector = Query;
@@ -3617,7 +3756,7 @@ FSearchResult Search(const UAnimSequenceBase* Sequence, TArrayView<const float>
 	return Result;
 }
 
-FDbSearchResult Search(
+FSearchResult Search(
 	const UPoseSearchDatabase* Database, 
 	TArrayView<const float> Query, 
 	const FPoseSearchWeightsContext* WeightsContext, 
@@ -3625,29 +3764,22 @@ FDbSearchResult Search(
 {
 	if (!ensure(Database && Database->IsValidForSearch()))
 	{
-		return FDbSearchResult();
+		return FSearchResult();
 	}
 
 	const FPoseSearchIndex& SearchIndex = Database->SearchIndex;
 
-	FDbSearchResult Result = Search(SearchIndex, Query, WeightsContext);
+	FSearchResult Result = Search(SearchIndex, Query, WeightsContext);
 	if (!Result.IsValid())
 	{
-		return FDbSearchResult();
+		return FSearchResult();
 	}
 
-	int32 DbSequenceIdx = Database->FindSequenceForPose(Result.PoseIdx);
-	if (DbSequenceIdx == INDEX_NONE)
+	if (!ensure(Result.TimeOffsetSeconds >= 0.0f))
 	{
-		return FDbSearchResult();
+		return FSearchResult();
 	}
 	
-	const FPoseSearchDatabaseSequence& DbSequence = Database->Sequences[DbSequenceIdx];
-	const FFloatInterval SamplingRange = GetEffectiveSamplingRange(DbSequence.Sequence, DbSequence.SamplingRange);
-
-	Result.DbSequenceIdx = DbSequenceIdx;
-	Result.TimeOffsetSeconds = SamplingRange.Min + SearchIndex.Schema->SamplingInterval * (Result.PoseIdx - DbSequence.FirstPoseIdx);
-
 	// Do debug visualization
 	DebugDrawParams.Database = Database;
 	DebugDrawParams.PoseVector = Query;
@@ -3660,7 +3792,10 @@ FDbSearchResult Search(
 float ComparePoses(const FPoseSearchIndex& SearchIndex, int32 PoseIdx, TArrayView<const float> Query, const FPoseSearchWeightsContext* WeightsContext)
 {
 	TArrayView<const float> PoseValues = SearchIndex.GetPoseValues(PoseIdx);
-	check(PoseValues.Num() == Query.Num());
+	if (!ensure(PoseValues.Num() == Query.Num()))
+	{
+		return MAX_flt;
+	}
 
 	float Dissimilarity;
 	if (WeightsContext)
diff --git a/Engine/Plugins/Experimental/Animation/PoseSearch/Source/Runtime/Private/PoseSearchLibrary.cpp b/Engine/Plugins/Experimental/Animation/PoseSearch/Source/Runtime/Private/PoseSearchLibrary.cpp
--- a/Engine/Plugins/Experimental/Animation/PoseSearch/Source/Runtime/Private/PoseSearchLibrary.cpp
+++ b/Engine/Plugins/Experimental/Animation/PoseSearch/Source/Runtime/Private/PoseSearchLibrary.cpp
@@ -13,7 +13,7 @@ DEFINE_LOG_CATEGORY_STATIC(LogPoseSearchLibrary, Verbose, All);
 void FMotionMatchingState::InitNewDatabaseSearch(const UPoseSearchDatabase* Database, float SearchThrottleTime)
 {
 	DbPoseIdx = INDEX_NONE;
-	DbSequenceIdx = INDEX_NONE;
+	SearchIndexAssetIdx = INDEX_NONE;
 	ElapsedPoseJumpTime = SearchThrottleTime;
 	AssetPlayerTime = 0.0f;
 	CurrentDatabase = Database;
@@ -34,11 +34,11 @@ void FMotionMatchingState::ComposeQuery(const UPoseSearchDatabase* Database, con
 	ComposedQuery.Normalize(Database->SearchIndex);
 }
 
-void FMotionMatchingState::JumpToPose(const UE::PoseSearch::FDbSearchResult& Result)
+void FMotionMatchingState::JumpToPose(const UE::PoseSearch::FSearchResult& Result)
 {
 	// Remember which pose and sequence we're playing from the database
 	DbPoseIdx = Result.PoseIdx;
-	DbSequenceIdx = Result.DbSequenceIdx;
+	SearchIndexAssetIdx = CurrentDatabase->SearchIndex.FindAssetIndex(Result.SearchIndexAsset);
 
 	ElapsedPoseJumpTime = 0.0f;
 	AssetPlayerTime = Result.TimeOffsetSeconds;
@@ -81,12 +81,14 @@ void UpdateMotionMatchingState(const FAnimationUpdateContext& Context
 		}
 
 		// Step the pose forward
-		UE::PoseSearch::FMotionMatchingContinuityParams ContinuityParameters = InOutMotionMatchingState.ComputeContinuityParameters(Context);
+		UE::PoseSearch::FMotionMatchingContinuityParams ContinuityParameters = 
+			InOutMotionMatchingState.ComputeContinuityParameters(Context);
 		const bool bCanContinue = ContinuityParameters.IsValid();
 		if (bCanContinue)
 		{
 			InOutMotionMatchingState.DbPoseIdx = ContinuityParameters.Result.PoseIdx;
-			InOutMotionMatchingState.DbSequenceIdx = ContinuityParameters.Result.DbSequenceIdx;
+			InOutMotionMatchingState.SearchIndexAssetIdx = 
+				InOutMotionMatchingState.CurrentDatabase->SearchIndex.FindAssetIndex(ContinuityParameters.Result.SearchIndexAsset);
 		}
 
 		// build query
@@ -125,7 +127,7 @@ void UpdateMotionMatchingState(const FAnimationUpdateContext& Context
 		}
 
 		// Search the database for the nearest match to the updated query vector
-		UE::PoseSearch::FDbSearchResult Result = UE::PoseSearch::Search(
+		UE::PoseSearch::FSearchResult Result = UE::PoseSearch::Search(
 			Database,
 			InOutMotionMatchingState.ComposedQuery.GetNormalizedValues(),
 			&InOutMotionMatchingState.WeightsContext
@@ -140,19 +142,21 @@ void UpdateMotionMatchingState(const FAnimationUpdateContext& Context
 			}
 			else
 			{
-				const FPoseSearchDatabaseSequence& ResultDbSequence = Database->Sequences[Result.DbSequenceIdx];
-
 				// Consider the search result better if it is more similar to the query than the current pose we're playing back from the database
 				bool bBetterPose = Result.Dissimilarity * (1.0f + (Settings.MinPercentImprovement / 100.0f)) < CurrentDissimilarity;
 
 				// We'll ignore the candidate pose if it is too near to our current pose
 				bool bNearbyPose = false;
-				if (InOutMotionMatchingState.DbSequenceIdx == Result.DbSequenceIdx)
+				const FPoseSearchIndexAsset* StateSearchIndexAsset = 
+					&InOutMotionMatchingState.CurrentDatabase->SearchIndex.Assets[InOutMotionMatchingState.SearchIndexAssetIdx];
+				if (StateSearchIndexAsset == Result.SearchIndexAsset)
 				{
+					const FPoseSearchDatabaseSequence& ResultDbSequence = Database->GetSourceAsset(Result.SearchIndexAsset);
 					bNearbyPose = FMath::Abs(InOutMotionMatchingState.AssetPlayerTime - Result.TimeOffsetSeconds) < Settings.PoseJumpThresholdTime;
 					if (!bNearbyPose && ResultDbSequence.bLoopAnimation)
 					{
-						const float AssetLength = Database->GetSequenceLength(InOutMotionMatchingState.DbSequenceIdx);
+						const FPoseSearchDatabaseSequence& StateDbSequence = Database->GetSourceAsset(StateSearchIndexAsset);
+						const float AssetLength = StateDbSequence.Sequence->GetPlayLength();
 						bNearbyPose = FMath::Abs(AssetLength - InOutMotionMatchingState.AssetPlayerTime - Result.TimeOffsetSeconds) < Settings.PoseJumpThresholdTime;
 					}
 				}
@@ -225,7 +229,8 @@ void UPoseSearchLibrary::UpdateMotionMatchingForSequencePlayer(const FAnimUpdate
 			// If a new pose is requested, jump to the pose by updating the embedded sequence player node
 			if ((InOutMotionMatchingState.Flags & EMotionMatchingFlags::JumpedToPose) == EMotionMatchingFlags::JumpedToPose)
 			{
-				const FPoseSearchDatabaseSequence& ResultDbSequence = Database->Sequences[InOutMotionMatchingState.DbSequenceIdx];
+				const FPoseSearchIndexAsset* SearchIndexAsset = &Database->SearchIndex.Assets[InOutMotionMatchingState.SearchIndexAssetIdx];
+				const FPoseSearchDatabaseSequence& ResultDbSequence = Database->GetSourceAsset(SearchIndexAsset);
 				SequencePlayerNode->SetSequence(Cast<UAnimSequenceBase>(ResultDbSequence.Sequence));
 				SequencePlayerNode->SetAccumulatedTime(InOutMotionMatchingState.AssetPlayerTime);
 				SequencePlayerNode->SetLoopAnimation(ResultDbSequence.bLoopAnimation);
@@ -249,9 +254,10 @@ UE::PoseSearch::FMotionMatchingContinuityParams FMotionMatchingState::ComputeCon
 
 	if (DbPoseIdx != INDEX_NONE)
 	{
-		check(DbSequenceIdx != INDEX_NONE);
+		check(SearchIndexAssetIdx != INDEX_NONE);
 
-		const FPoseSearchDatabaseSequence& DbSequence = CurrentDatabase->Sequences[DbSequenceIdx];
+		const FPoseSearchIndexAsset* SearchIndexAsset = &CurrentDatabase->SearchIndex.Assets[SearchIndexAssetIdx];
+		const FPoseSearchDatabaseSequence& DbSequence = CurrentDatabase->GetSourceAsset(SearchIndexAsset);
 		const float AssetLength = DbSequence.Sequence->GetPlayLength();
 		const float DeltaTime = Context.GetDeltaTime();
 
@@ -264,11 +270,14 @@ UE::PoseSearch::FMotionMatchingContinuityParams FMotionMatchingState::ComputeCon
 
 		if (ContinuityAdvanceType != ETAA_Finished)
 		{
-			// we can continue ticking the same sequence forward
-			ContinuityParameters.Result.PoseIdx = CurrentDatabase->GetPoseIndexFromAssetTime(DbSequenceIdx,
-																					  ContinuityAssetTime);
-			ContinuityParameters.Result.DbSequenceIdx = DbSequenceIdx;
-			ContinuityParameters.Result.TimeOffsetSeconds = AssetPlayerTime;
+			// if the sampling interval doesn't contain the new time, there's no pose to return.
+			if (SearchIndexAsset->SamplingInterval.Contains(ContinuityAssetTime))
+			{
+				// we can continue ticking the same sequence forward
+				ContinuityParameters.Result.PoseIdx = CurrentDatabase->GetPoseIndexFromAssetTime(ContinuityAssetTime, SearchIndexAsset);
+				ContinuityParameters.Result.SearchIndexAsset = SearchIndexAsset;
+				ContinuityParameters.Result.TimeOffsetSeconds = AssetPlayerTime;
+			}
 		}
 		else
 		{
@@ -278,19 +287,29 @@ UE::PoseSearch::FMotionMatchingContinuityParams FMotionMatchingState::ComputeCon
 			{
 				return Entry.Sequence == DbSequence.FollowUpSequence;
 			});
-			if (FollowUpDbSequenceIdx != INDEX_NONE)
+
+			int32 FollowUpSearchIndexAssetIdx = CurrentDatabase->SearchIndex.Assets.IndexOfByPredicate(
+				[&](const FPoseSearchIndexAsset& Entry)
+			{
+				return Entry.SourceAssetIdx == FollowUpDbSequenceIdx && Entry.SamplingInterval.Contains(0.0f);
+			});
+
+			if (FollowUpSearchIndexAssetIdx != INDEX_NONE)
 			{
-				const FPoseSearchDatabaseSequence& FollowUpDbSequence = CurrentDatabase->Sequences[FollowUpDbSequenceIdx];
+				const FPoseSearchIndexAsset* FollowUpSearchIndexAsset = 
+					&CurrentDatabase->SearchIndex.Assets[FollowUpSearchIndexAssetIdx];
 				const float FollowUpAssetTime = AssetPlayerTime + DeltaTime - AssetLength;
-				const int32 FollowUpPoseIdx = CurrentDatabase->GetPoseIndexFromAssetTime(FollowUpDbSequenceIdx,
-																						 FollowUpAssetTime);
-				const FFloatInterval SamplingRange = CurrentDatabase->GetEffectiveSamplingRange(FollowUpDbSequenceIdx);
+				const int32 FollowUpPoseIdx = CurrentDatabase->GetPoseIndexFromAssetTime(
+					FollowUpAssetTime,
+					FollowUpSearchIndexAsset);
+				const FFloatInterval SamplingRange = FollowUpSearchIndexAsset->SamplingInterval;
 
 				ContinuityParameters.Result.PoseIdx = FollowUpPoseIdx;
-				ContinuityParameters.Result.DbSequenceIdx = FollowUpDbSequenceIdx;
+				ContinuityParameters.Result.SearchIndexAsset = FollowUpSearchIndexAsset;
 				ContinuityParameters.Result.TimeOffsetSeconds =
-					SamplingRange.Min + (CurrentDatabase->Schema->SamplingInterval *
-										 (ContinuityParameters.Result.PoseIdx - FollowUpDbSequence.FirstPoseIdx));
+					SamplingRange.Min + 
+					(CurrentDatabase->Schema->SamplingInterval *
+					 (ContinuityParameters.Result.PoseIdx - FollowUpSearchIndexAsset->FirstPoseIdx));
 				ContinuityParameters.bJumpRequired = true;
 			}
 		}
diff --git a/Engine/Plugins/Experimental/Animation/PoseSearch/Source/Runtime/Public/PoseSearch/PoseSearch.h b/Engine/Plugins/Experimental/Animation/PoseSearch/Source/Runtime/Public/PoseSearch/PoseSearch.h
--- a/Engine/Plugins/Experimental/Animation/PoseSearch/Source/Runtime/Public/PoseSearch/PoseSearch.h
+++ b/Engine/Plugins/Experimental/Animation/PoseSearch/Source/Runtime/Public/PoseSearch/PoseSearch.h
@@ -24,12 +24,12 @@ struct FCompactPose;
 struct FPoseContext;
 struct FReferenceSkeleton;
 
+DECLARE_LOG_CATEGORY_EXTERN(LogPoseSearch, Log, All);
+
 namespace UE { namespace PoseSearch {
 
 class FPoseHistory;
 
-DECLARE_LOG_CATEGORY_EXTERN(LogPoseSearch, Log, All);
-
 }}
 
 UENUM()
@@ -275,6 +275,42 @@ struct POSESEARCH_API FPoseSearchPoseMetadata
 	float CostAddend = 0.0f;
 };
 
+/**
+* Information about a source animation asset used by a search index.
+* Some source animation entries may generate multiple FPoseSearchIndexAsset entries.
+**/
+USTRUCT()
+struct POSESEARCH_API FPoseSearchIndexAsset
+{
+	GENERATED_BODY()
+public:
+	FPoseSearchIndexAsset()
+	{}
+
+	FPoseSearchIndexAsset(int32 InSourceAssetIdx, const FFloatInterval& InSamplingInterval)
+		: SourceAssetIdx(InSourceAssetIdx)
+		, SamplingInterval(InSamplingInterval)
+	{}
+
+	// Index of the source asset in search index's container (i.e. UPoseSearchDatabase)
+	UPROPERTY()
+	int32 SourceAssetIdx = INDEX_NONE;
+
+	UPROPERTY()
+	FFloatInterval SamplingInterval;
+
+	UPROPERTY()
+	int32 FirstPoseIdx = INDEX_NONE;
+
+	UPROPERTY()
+	int32 NumPoses = 0;
+
+	bool IsPoseInRange(int32 PoseIdx) const
+	{
+		return (PoseIdx >= FirstPoseIdx) && (PoseIdx < FirstPoseIdx + NumPoses);
+	}
+};
+
 /**
 * A search index for animation poses. The structure of the search index is determined by its UPoseSearchSchema.
 * May represent a single animation (see UPoseSearchSequenceMetaData) or a collection (see UPoseSearchDatabase).
@@ -299,10 +335,16 @@ struct POSESEARCH_API FPoseSearchIndex
 	UPROPERTY()
 	FPoseSearchIndexPreprocessInfo PreprocessInfo;
 
+	UPROPERTY()
+	TArray<FPoseSearchIndexAsset> Assets;
+
 	bool IsValid() const;
 
 	TArrayView<const float> GetPoseValues(int32 PoseIdx) const;
 
+	int32 FindAssetIndex(const FPoseSearchIndexAsset* Asset) const;
+	const FPoseSearchIndexAsset* FindAssetForPose(int32 PoseIdx) const;
+
 	void Reset();
 
 	void Normalize (TArrayView<float> PoseVector) const;
@@ -334,10 +376,12 @@ struct FPoseSearchBlockTransitionParameters
 	GENERATED_BODY()
 
 	// Excluding the beginning of sequences can help ensure an exact past trajectory is used when building the features
+	UPROPERTY(EditAnywhere, Category = "Settings")
 	float SequenceStartInterval = 0.0f;
 
 	// Excluding the end of sequences help ensure an exact future trajectory, and also prevents the selection of
 	// a sequence which will end too soon to be worth selecting.
+	UPROPERTY(EditAnywhere, Category = "Settings")
 	float SequenceEndInterval = 0.2f;
 
 };
@@ -545,11 +589,8 @@ struct POSESEARCH_API FPoseSearchDatabaseSequence
 	UPROPERTY(EditAnywhere, Category="Sequence")
 	bool bLoopFollowUpAnimation = false;
 
-	UPROPERTY()
-	int32 FirstPoseIdx = 0;
-
-	UPROPERTY()
-	int32 NumPoses = 0;
+public:
+	FFloatInterval GetEffectiveSamplingRange() const;
 };
 
 /** A data asset for indexing a collection of animation sequences. */
@@ -582,7 +623,6 @@ public:
 	FPoseSearchIndex SearchIndex;
 
 	int32 FindSequenceForPose(int32 PoseIdx) const;
-	int32 GetPoseIndexFromAssetTime(int32 DbSequenceIdx, float AssetTime) const;
 	float GetSequenceLength(int32 DbSequenceIdx) const;
 	bool DoesSequenceLoop(int32 DbSequenceIdx) const;
 	FFloatInterval GetEffectiveSamplingRange(int32 DbSequenceIdx) const;
@@ -590,7 +630,12 @@ public:
 	bool IsValidForIndexing() const;
 	bool IsValidForSearch() const;
 
+	int32 GetPoseIndexFromAssetTime(float AssetTime, const FPoseSearchIndexAsset* SearchIndexAsset) const;
+	float GetTimeOffset(int32 PoseIdx, const FPoseSearchIndexAsset* SearchIndexAsset = nullptr) const;
+	const FPoseSearchDatabaseSequence& GetSourceAsset(const FPoseSearchIndexAsset* SearchIndexAsset) const;
+
 public: // UObject
+
 	virtual void PreSave(FObjectPreSaveContext ObjectSaveContext) override;
 
 #if WITH_EDITOR
@@ -599,6 +644,10 @@ public: // UObject
 
 private:
 	void CollectSimpleSequences();
+
+public:
+	// Populates the FPoseSearchIndex::Assets array by evaluating the data in the Sequences array
+	bool TryInitSearchIndexAssets();
 };
 
 
@@ -787,24 +836,14 @@ struct POSESEARCH_API FDebugDrawParams
 
 struct FSearchResult
 {
-	int32 PoseIdx = -1;
+	int32 PoseIdx = INDEX_NONE;
+	const FPoseSearchIndexAsset* SearchIndexAsset = nullptr;
 	float TimeOffsetSeconds = 0.0f;
 	float Dissimilarity = MAX_flt;
 
 	bool IsValid() const { return PoseIdx >= 0; }
 };
 
-struct FDbSearchResult : public FSearchResult
-{
-	FDbSearchResult() = default;
-
-	FDbSearchResult(const FSearchResult& Result)
-		: FSearchResult(Result)
-	{}
-
-	int32 DbSequenceIdx = INDEX_NONE;
-};
-
 
 /**
 * Visualize pose search debug information
@@ -859,7 +898,7 @@ POSESEARCH_API FSearchResult Search(const UAnimSequenceBase* Sequence, TArrayVie
 * 
 * @return The pose in the database that most closely matches the Query.
 */
-POSESEARCH_API FDbSearchResult Search(const UPoseSearchDatabase* Database, TArrayView<const float> Query, const FPoseSearchWeightsContext* WeightsContext = nullptr, FDebugDrawParams DrawParams = FDebugDrawParams());
+POSESEARCH_API FSearchResult Search(const UPoseSearchDatabase* Database, TArrayView<const float> Query, const FPoseSearchWeightsContext* WeightsContext = nullptr, FDebugDrawParams DrawParams = FDebugDrawParams());
 
 
 /**
diff --git a/Engine/Plugins/Experimental/Animation/PoseSearch/Source/Runtime/Public/PoseSearch/PoseSearchAnimNotifies.h b/Engine/Plugins/Experimental/Animation/PoseSearch/Source/Runtime/Public/PoseSearch/PoseSearchAnimNotifies.h
--- a/Engine/Plugins/Experimental/Animation/PoseSearch/Source/Runtime/Public/PoseSearch/PoseSearchAnimNotifies.h
+++ b/Engine/Plugins/Experimental/Animation/PoseSearch/Source/Runtime/Public/PoseSearch/PoseSearchAnimNotifies.h
@@ -13,6 +13,14 @@ class POSESEARCH_API UAnimNotifyState_PoseSearchBase : public UAnimNotifyState
 	GENERATED_BODY()
 };
 
+// Use this notify state to remove animation segments from the database completely, they will never play or return from
+// a search result
+UCLASS(Blueprintable, meta = (DisplayName = "Pose Matching: Exclude From Database"))
+class POSESEARCH_API UAnimNotifyState_PoseSearchExcludeFromDatabase : public UAnimNotifyState_PoseSearchBase
+{
+	GENERATED_BODY()
+};
+
 // A pose matching search will not return results that overlap this notify, but the animation segment can still play
 // if a previous search result advances into it.
 UCLASS(Blueprintable, meta = (DisplayName = "Pose Matching: Block Transition"))
diff --git a/Engine/Plugins/Experimental/Animation/PoseSearch/Source/Runtime/Public/PoseSearch/PoseSearchLibrary.h b/Engine/Plugins/Experimental/Animation/PoseSearch/Source/Runtime/Public/PoseSearch/PoseSearchLibrary.h
--- a/Engine/Plugins/Experimental/Animation/PoseSearch/Source/Runtime/Public/PoseSearch/PoseSearchLibrary.h
+++ b/Engine/Plugins/Experimental/Animation/PoseSearch/Source/Runtime/Public/PoseSearch/PoseSearchLibrary.h
@@ -17,7 +17,7 @@ namespace UE::PoseSearch
 {
 	struct FMotionMatchingContinuityParams
 	{
-		FDbSearchResult Result;
+		FSearchResult Result;
 		bool bJumpRequired = false;
 
 		bool IsValid()
@@ -27,7 +27,7 @@ namespace UE::PoseSearch
 
 		void Reset()
 		{
-			Result = UE::PoseSearch::FDbSearchResult();
+			Result = UE::PoseSearch::FSearchResult();
 			bJumpRequired = false;
 		}
 	};
@@ -80,7 +80,7 @@ struct POSESEARCH_API FMotionMatchingState
 	void ComposeQuery(const UPoseSearchDatabase* Database, const FTrajectorySampleRange& Trajectory);
 
 	// Internally stores the 'jump' to a new pose/sequence index and asset time for evaluation
-	void JumpToPose(const UE::PoseSearch::FDbSearchResult& Result);
+	void JumpToPose(const UE::PoseSearch::FSearchResult& Result);
 
 	// Updates DbPoseIdx to track DeltaTime and jumps to a follow-up sequence when available
 	UE::PoseSearch::FMotionMatchingContinuityParams ComputeContinuityParameters(const FAnimationUpdateContext& Context) const;
@@ -91,7 +91,7 @@ struct POSESEARCH_API FMotionMatchingState
 
 	// The current animation we're playing from the database
 	UPROPERTY(Transient)
-	int32 DbSequenceIdx = INDEX_NONE;
+	int32 SearchIndexAssetIdx = INDEX_NONE;
 
 	// The current query feature vector used to search the database for pose candidates
 	UPROPERTY(Transient)
