PoseSearch: Debugger updates
- Rename "Score" columns to "Cost"
- Sort by cost by default
- Added Modifier cost column
- Score breakdown is now handled by new PoseSearch ComparePoses function that outputs FPoseCostInfo
- Pose and Trajectory score now reflect their actual contribution to the total score rather than being separately normalized
- Time column is now the Frame column and displays Frame # (Frame time)
- Removed length column
- Renamed "Selected Poses" header to "Pose Database"
- Renamed "PoseIdx" column to "Pose Index"
- Removed commas from pose indices
- Debug pose vector rendering now shows label prefixes to help associate them with the query, active, and selected poses
- Renamed "Use Trajectory Forward Vectors" to "Use Facing Directions"
- Added option to expose motion matching debug node drawing options as pins
- Misc cleanup

diff --git a/Engine/Plugins/Experimental/Animation/PoseSearch/Source/Editor/Private/PoseSearchDebugger.cpp b/Engine/Plugins/Experimental/Animation/PoseSearch/Source/Editor/Private/PoseSearchDebugger.cpp
--- a/Engine/Plugins/Experimental/Animation/PoseSearch/Source/Editor/Private/PoseSearchDebugger.cpp
+++ b/Engine/Plugins/Experimental/Animation/PoseSearch/Source/Editor/Private/PoseSearchDebugger.cpp
@@ -39,32 +39,16 @@ class FDebuggerDatabaseRowData : public TSharedFromThis<FDebuggerDatabaseRowData
 {
 public:
 	FDebuggerDatabaseRowData() = default;
-	FDebuggerDatabaseRowData(
-	    int32 InPoseIdx, 
-	    const FString& InAnimSequenceName,
-	    float InTime,
-	    float InLength,
-	    float InScore,
-	    float InPoseScore,
-	    float InTrajectoryScore
-	)
-		: PoseIdx(InPoseIdx)
-		, AnimSequenceName(InAnimSequenceName)
-		, Time(InTime)
-		, Length(InLength)
-		, Score(InScore)
-		, PoseScore(InPoseScore)
-		, TrajectoryScore(InTrajectoryScore)
-	{
-	}
+	
+	float GetPoseCost() const { return PoseCostInfo.ChannelCosts.Num() ? PoseCostInfo.ChannelCosts[0] : 0.0f; }
+	float GetTrajectoryCost() const { return PoseCostInfo.ChannelCosts.Num() >= 3 ? PoseCostInfo.ChannelCosts[1] + PoseCostInfo.ChannelCosts[2] : 0.0f; }
 
 	int32 PoseIdx = 0;
 	FString AnimSequenceName = "";
+	int32 DbSequenceIdx = 0;
+	int32 AnimFrame = 0;
 	float Time = 0.0f;
-	float Length = 0.0f;
-	float Score = 0.0f;
-	float PoseScore = 0.0f;
-	float TrajectoryScore = 0.0f;
+	FPoseCostInfo PoseCostInfo;
 };
 
 namespace DebuggerDatabaseColumns
@@ -131,10 +115,10 @@ namespace DebuggerDatabaseColumns
 
 		virtual FText GetRowText(const FRowDataRef& Row) const override
 		{
-			return FText::AsNumber(Row->PoseIdx);
+			return FText::AsNumber(Row->PoseIdx, &FNumberFormattingOptions::DefaultNoGrouping());
 		}
 	};
-	const FName FPoseIdx::Name = "PoseIdx";
+	const FName FPoseIdx::Name = "Pose Index";
 
 	struct FAnimSequenceName : IColumn
 	{
@@ -159,9 +143,9 @@ namespace DebuggerDatabaseColumns
 			return FText::FromString(Row->AnimSequenceName);
 		}
 	};
-	const FName FAnimSequenceName::Name = "AnimSequence";
+	const FName FAnimSequenceName::Name = "Sequence";
 
-	struct FTime : IColumn
+	struct FFrame : IColumn
 	{
 		using IColumn::IColumn;
 		static const FName Name;
@@ -181,12 +165,20 @@ namespace DebuggerDatabaseColumns
 
 		virtual FText GetRowText(const FRowDataRef& Row) const override
 		{
-			return FText::AsNumber(Row->Time);
+			FNumberFormattingOptions TimeFormattingOptions = FNumberFormattingOptions()
+				.SetUseGrouping(false)
+				.SetMaximumFractionalDigits(2);
+
+			return FText::Format(
+				FText::FromString("{0} ({1})"),
+				FText::AsNumber(Row->AnimFrame, &FNumberFormattingOptions::DefaultNoGrouping()),
+				FText::AsNumber(Row->Time, &TimeFormattingOptions)
+			);
 		}
 	};
-	const FName FTime::Name = "Time";
+	const FName FFrame::Name = "Frame";
 
-	struct FLength : IColumn
+	struct FCost : IColumn
 	{
 		using IColumn::IColumn;
 		static const FName Name;
@@ -194,24 +186,25 @@ namespace DebuggerDatabaseColumns
 		
 		virtual const FSortPredicate& GetSortPredicateAscending() const override
 		{
-			static FSortPredicate Predicate = [](const FRowDataRef& Row0, const FRowDataRef& Row1) -> bool { return Row0->Length < Row1->Length; };
+			static FSortPredicate Predicate = [](const FRowDataRef& Row0, const FRowDataRef& Row1) -> bool { return Row0->PoseCostInfo.TotalCost < Row1->PoseCostInfo.TotalCost; };
 			return Predicate;
 		}
 
 		virtual const FSortPredicate& GetSortPredicateDescending() const override
 		{
-			static FSortPredicate Predicate = [](const FRowDataRef& Row0, const FRowDataRef& Row1) -> bool { return Row0->Length >= Row1->Length; };
+			static FSortPredicate Predicate = [](const FRowDataRef& Row0, const FRowDataRef& Row1) -> bool { return Row0->PoseCostInfo.TotalCost >= Row1->PoseCostInfo.TotalCost; };
 			return Predicate;
 		}
-
+		
 		virtual FText GetRowText(const FRowDataRef& Row) const override
-		{
-			return FText::AsNumber(Row->Length);
-		}
+        {
+        	return FText::AsNumber(Row->PoseCostInfo.TotalCost);
+        }
 	};
-	const FName FLength::Name = "Length";
+	const FName FCost::Name = "Cost";
+
 
-	struct FScore : IColumn
+	struct FPoseCost : IColumn
 	{
 		using IColumn::IColumn;
 		static const FName Name;
@@ -219,25 +212,25 @@ namespace DebuggerDatabaseColumns
 		
 		virtual const FSortPredicate& GetSortPredicateAscending() const override
 		{
-			static FSortPredicate Predicate = [](const FRowDataRef& Row0, const FRowDataRef& Row1) -> bool { return Row0->Score < Row1->Score; };
+			static FSortPredicate Predicate = [](const FRowDataRef& Row0, const FRowDataRef& Row1) -> bool { return Row0->GetPoseCost() < Row1->GetPoseCost(); };
 			return Predicate;
 		}
 
 		virtual const FSortPredicate& GetSortPredicateDescending() const override
 		{
-			static FSortPredicate Predicate = [](const FRowDataRef& Row0, const FRowDataRef& Row1) -> bool { return Row0->Score >= Row1->Score; };
+			static FSortPredicate Predicate = [](const FRowDataRef& Row0, const FRowDataRef& Row1) -> bool { return Row0->GetPoseCost() >= Row1->GetPoseCost(); };
 			return Predicate;
 		}
 		
 		virtual FText GetRowText(const FRowDataRef& Row) const override
-        {
-        	return FText::AsNumber(Row->Score);
-        }
+		{
+			return FText::AsNumber(Row->GetPoseCost());
+		}
 	};
-	const FName FScore::Name = "Score";
+	const FName FPoseCost::Name = "Pose Cost";
 
 
-	struct FPoseScore : IColumn
+	struct FTrajectoryCost : IColumn
 	{
 		using IColumn::IColumn;
 		static const FName Name;
@@ -245,48 +238,47 @@ namespace DebuggerDatabaseColumns
 		
 		virtual const FSortPredicate& GetSortPredicateAscending() const override
 		{
-			static FSortPredicate Predicate = [](const FRowDataRef& Row0, const FRowDataRef& Row1) -> bool { return Row0->PoseScore < Row1->PoseScore; };
+			static FSortPredicate Predicate = [](const FRowDataRef& Row0, const FRowDataRef& Row1) -> bool { return Row0->GetTrajectoryCost() < Row1->GetTrajectoryCost(); };
 			return Predicate;
 		}
 
 		virtual const FSortPredicate& GetSortPredicateDescending() const override
 		{
-			static FSortPredicate Predicate = [](const FRowDataRef& Row0, const FRowDataRef& Row1) -> bool { return Row0->PoseScore >= Row1->PoseScore; };
+			static FSortPredicate Predicate = [](const FRowDataRef& Row0, const FRowDataRef& Row1) -> bool { return Row0->GetTrajectoryCost() >= Row1->GetTrajectoryCost(); };
 			return Predicate;
 		}
 		
 		virtual FText GetRowText(const FRowDataRef& Row) const override
 		{
-			return FText::AsNumber(Row->PoseScore);
-		}
+			return FText::AsNumber(Row->GetTrajectoryCost());
+        }
 	};
-	const FName FPoseScore::Name = "Pose Score";
+	const FName FTrajectoryCost::Name = "Trajectory Cost";
 
-
-	struct FTrajectoryScore : IColumn
+	struct FCostModifier : IColumn
 	{
 		using IColumn::IColumn;
 		static const FName Name;
 		virtual FName GetName() const override { return Name; }
-		
+
 		virtual const FSortPredicate& GetSortPredicateAscending() const override
 		{
-			static FSortPredicate Predicate = [](const FRowDataRef& Row0, const FRowDataRef& Row1) -> bool { return Row0->TrajectoryScore < Row1->TrajectoryScore; };
+			static FSortPredicate Predicate = [](const FRowDataRef& Row0, const FRowDataRef& Row1) -> bool { return Row0->PoseCostInfo.CostAddend < Row1->PoseCostInfo.CostAddend; };
 			return Predicate;
 		}
 
 		virtual const FSortPredicate& GetSortPredicateDescending() const override
 		{
-			static FSortPredicate Predicate = [](const FRowDataRef& Row0, const FRowDataRef& Row1) -> bool { return Row0->TrajectoryScore >= Row1->TrajectoryScore; };
+			static FSortPredicate Predicate = [](const FRowDataRef& Row0, const FRowDataRef& Row1) -> bool { return Row0->PoseCostInfo.CostAddend >= Row1->PoseCostInfo.CostAddend; };
 			return Predicate;
 		}
-		
+
 		virtual FText GetRowText(const FRowDataRef& Row) const override
 		{
-			return FText::AsNumber(Row->TrajectoryScore);
-        }
+			return FText::AsNumber(Row->PoseCostInfo.CostAddend);
+		}
 	};
-	const FName FTrajectoryScore::Name = "Trajectory Score";
+	const FName FCostModifier::Name = "Cost Modifier";
 }
 
 
@@ -531,19 +523,13 @@ void SDebuggerDatabaseView::CreateRows(const UPoseSearchDatabase& Database)
 		const int32 LastPoseIdx = SearchIndexAsset.FirstPoseIdx + SearchIndexAsset.NumPoses;
 		for (int32 PoseIdx = SearchIndexAsset.FirstPoseIdx; PoseIdx != LastPoseIdx; ++PoseIdx)
 		{
+			float Time = Database.GetTimeOffset(PoseIdx);
+
 			TSharedRef<FDebuggerDatabaseRowData> Row = UnfilteredDatabaseRows.Add_GetRef(MakeShared<FDebuggerDatabaseRowData>());
 			Row->PoseIdx = PoseIdx;
-
-			const FString SequenceName = DbSequence.Sequence->GetName();
-			const float SequenceLength = DbSequence.Sequence->GetPlayLength();
-			FFloatInterval Range = SearchIndexAsset.SamplingInterval;
-
-			Row->AnimSequenceName = SequenceName;
-			// Cap time in sequence to end of range
-			Row->Time = FGenericPlatformMath::Min(
-				Range.Min + (PoseIdx - SearchIndexAsset.FirstPoseIdx) * Database.Schema->SamplingInterval, 
-				Range.Max);
-			Row->Length = SequenceLength;
+			Row->AnimSequenceName = DbSequence.Sequence->GetName();
+			Row->Time = Time;
+			Row->AnimFrame = DbSequence.Sequence->GetFrameAtTime(Time);
 		}
 	}
 
@@ -561,20 +547,6 @@ void SDebuggerDatabaseView::UpdateRows(const FTraceMotionMatchingStateMessage& S
 	
 	FPoseSearchWeightsContext StateWeights;
 	StateWeights.Update(State.Weights, &Database);
-
-	FPoseSearchWeightsContext PoseWeights;
-	{
-		FPoseSearchDynamicWeightParams WeightParams = State.Weights;
-		WeightParams.TrajectoryDynamicWeights.ChannelWeightScale = 0.0f;
-		PoseWeights.Update(WeightParams, &Database);
-	}
-
-	FPoseSearchWeightsContext TrajectoryWeights;
-	{
-		FPoseSearchDynamicWeightParams WeightParams = State.Weights;
-		WeightParams.PoseDynamicWeights.ChannelWeightScale = 0.0f;
-		TrajectoryWeights.Update(WeightParams, &Database);
-	}
 	
 	const FPoseSearchIndex& SearchIndex = Database.SearchIndex;
 
@@ -582,9 +554,7 @@ void SDebuggerDatabaseView::UpdateRows(const FTraceMotionMatchingStateMessage& S
 	{
 		const int32 PoseIdx = Row->PoseIdx;
 		
-		Row->Score = ComparePoses(SearchIndex, PoseIdx, State.QueryVectorNormalized, &StateWeights);
-		Row->PoseScore = ComparePoses(SearchIndex, PoseIdx, State.QueryVectorNormalized, &PoseWeights);
-		Row->TrajectoryScore = ComparePoses(SearchIndex, PoseIdx, State.QueryVectorNormalized, &TrajectoryWeights);
+		ComparePoses(SearchIndex, PoseIdx, State.QueryVectorNormalized, &StateWeights, Row->PoseCostInfo);
 
 		// If we are on the active pose for the frame
 		if (PoseIdx == State.DbPoseIdx)
@@ -619,12 +589,12 @@ void SDebuggerDatabaseView::Construct(const FArguments& InArgs)
 	// @TODO: Support runtime reordering of these indices
 	// Construct all column types
 	AddColumn(MakeShared<FAnimSequenceName>(0));
-	AddColumn(MakeShared<FPoseIdx>(1));
-	AddColumn(MakeShared<FTime>(2));
-	AddColumn(MakeShared<FLength>(3));
-	AddColumn(MakeShared<FScore>(4));
-	AddColumn(MakeShared<FPoseScore>(5));
-	AddColumn(MakeShared<FTrajectoryScore>(6));
+	AddColumn(MakeShared<FCost>(1));
+	AddColumn(MakeShared<FPoseCost>(2));
+	AddColumn(MakeShared<FTrajectoryCost>(3));
+	AddColumn(MakeShared<FCostModifier>(4));
+	AddColumn(MakeShared<FFrame>(5));
+	AddColumn(MakeShared<FPoseIdx>(6));
 
 	// Active Row
 	ActiveView.HeaderRow = SNew(SHeaderRow);
@@ -753,7 +723,7 @@ void SDebuggerDatabaseView::Construct(const FArguments& InArgs)
 					.VAlign(VAlign_Fill)
 					[
 						SNew(STextBlock)
-						.Text(FText::FromString("Selected Poses"))	
+						.Text(FText::FromString("Pose Database"))	
 					]
 				]
 				.AutoWidth()
@@ -810,7 +780,7 @@ void SDebuggerDatabaseView::Construct(const FArguments& InArgs)
 		]
 	];
 	
-	SortColumn = FPoseIdx::Name;
+	SortColumn = FCost::Name;
 	SortMode = EColumnSortMode::Ascending;
 
 	// Active view scroll bar only for indenting the columns to align w/ database
@@ -1194,6 +1164,7 @@ void SDebuggerView::DrawFeatures(
 	// Draw query vector
 	DrawParams.Color = &FLinearColor::Blue;
 	SetDrawFlags(DrawParams, Reflection->QueryDrawOptions);
+	DrawParams.LabelPrefix = TEXT("Q");
 	Draw(DrawParams);
 	DrawParams.PoseVector = {};
 
@@ -1202,6 +1173,7 @@ void SDebuggerView::DrawFeatures(
 
 	// Red for non-active database view
 	DrawParams.Color = &FLinearColor::Red;
+	DrawParams.LabelPrefix = TEXT("S");
 	SetDrawFlags(DrawParams, Reflection->SelectedPoseDrawOptions);
 
 	// Draw any selected database vectors
@@ -1223,6 +1195,7 @@ void SDebuggerView::DrawFeatures(
 		
 		// Use the motion-matching state's pose idx, as the active row may be update-throttled at this point
 		DrawParams.PoseIdx = State.DbPoseIdx;
+		DrawParams.LabelPrefix = TEXT("A");
 		Draw(DrawParams);
 	}
 }
diff --git a/Engine/Plugins/Experimental/Animation/PoseSearch/Source/Runtime/Private/PoseSearch.cpp b/Engine/Plugins/Experimental/Animation/PoseSearch/Source/Runtime/Private/PoseSearch.cpp
--- a/Engine/Plugins/Experimental/Animation/PoseSearch/Source/Runtime/Private/PoseSearch.cpp
+++ b/Engine/Plugins/Experimental/Animation/PoseSearch/Source/Runtime/Private/PoseSearch.cpp
@@ -36,6 +36,7 @@ namespace UE { namespace PoseSearch {
 //////////////////////////////////////////////////////////////////////////
 // Constants and utilities
 
+// Debug drawing options
 constexpr float DrawDebugLineThickness = 2.0f;
 constexpr float DrawDebugPointSize = 3.0f;
 constexpr float DrawDebugVelocityScale = 0.08f;
@@ -46,22 +47,17 @@ constexpr float DrawDebugGradientStrength = 0.8f;
 constexpr float DrawDebugSampleLabelFontScale = 1.0f;
 static const FVector DrawDebugSampleLabelOffset = FVector(0.0f, 0.0f, -10.0f);
 
+// Stopgap channel indices since schemas don't yet support explicit channels of data
+constexpr int32 ChannelIdxPose = 0;
+constexpr int32 ChannelIdxTrajectoryTime = 1;
+constexpr int32 ChannelIdxTrajectoryDistance = 2;
+constexpr int32 MaxChannels = 3;
+
 static bool IsSamplingRangeValid(FFloatInterval Range)
 {
 	return Range.IsValid() && (Range.Min >= 0.0f);
 }
 
-static FFloatInterval GetEffectiveSamplingRange(const UAnimSequenceBase* Sequence, FFloatInterval SamplingRange)
-{
-	const bool bSampleAll = (SamplingRange.Min == 0.0f) && (SamplingRange.Max == 0.0f);
-	const float SequencePlayLength = Sequence->GetPlayLength();
-
-	FFloatInterval Range;
-	Range.Min = bSampleAll ? 0.0f : SamplingRange.Min;
-	Range.Max = bSampleAll ? SequencePlayLength : FMath::Min(SequencePlayLength, SamplingRange.Max);
-	return Range;
-}
-
 static inline float CompareFeatureVectors(int32 NumValues, const float* A, const float* B, const float* Weights)
 {
 	double Dissimilarity = 0.f;
@@ -155,12 +151,6 @@ FORCEINLINE auto LowerBound(IteratorType First, IteratorType Last, const ValueTy
 	return LowerBound(First, Last, Value, FIdentityFunctor(), TLess<>());
 }
 
-
-// Stopgap channel indices since schemas don't yet support explicit channels of data
-constexpr int32 ChannelIdxPose = 0;
-constexpr int32 ChannelIdxTrajectoryTime = 1;
-constexpr int32 ChannelIdxTrajectoryDistance = 2;
-
 static EPoseSearchFeatureDomain FeatureDomainFromChannel (int32 ChannelIdx)
 {
 	switch (ChannelIdx)
@@ -207,6 +197,25 @@ FFeatureTypeTraits GetFeatureTypeTraits(EPoseSearchFeatureType Type)
 	return FFeatureTypeTraits();
 }
 
+static void CalcChannelCosts(const UPoseSearchSchema* Schema, TArrayView<const float> CostVector, TArray<float>& OutChannelCosts)
+{
+	OutChannelCosts.Reset();
+	OutChannelCosts.AddZeroed(MaxChannels);
+	for (int ChannelIdx = 0; ChannelIdx != MaxChannels; ++ChannelIdx)
+	{
+		for (int FeatureIdx = INDEX_NONE; Schema->Layout.EnumerateBy(ChannelIdx, EPoseSearchFeatureType::Invalid, FeatureIdx); /*empty*/)
+		{
+			const FPoseSearchFeatureDesc& Feature = Schema->Layout.Features[FeatureIdx];
+			int32 ValueSize = GetFeatureTypeTraits(Feature.Type).NumFloats;
+			int32 ValueTerm = Feature.ValueOffset + ValueSize;
+			for (int32 ValueIdx = Feature.ValueOffset; ValueIdx != ValueTerm; ++ValueIdx)
+			{
+				OutChannelCosts[ChannelIdx] += CostVector[ValueIdx];
+			}
+		}
+	}
+}
+
 }} // namespace UE::PoseSearch
 
 
@@ -553,7 +562,7 @@ void FPoseSearchWeights::Init(const FPoseSearchWeightParams& WeightParams, const
 
 
 	// Setup channel indexable weight params
-	constexpr int ChannelNum = 3;
+	constexpr int ChannelNum = MaxChannels;
 	const FPoseSearchChannelWeightParams* ChannelWeightParams[ChannelNum];
 	const FPoseSearchChannelDynamicWeightParams* ChannelDynamicWeightParams[ChannelNum];
 
@@ -1040,9 +1049,9 @@ float UPoseSearchDatabase::GetTimeOffset(int32 PoseIdx, const FPoseSearchIndexAs
 		UE_LOG(LogPoseSearch, Error, TEXT("Pose %i out of range in database %s"), PoseIdx, *GetNameSafe(this));
 		return -1.0f;
 	}
-	const FPoseSearchDatabaseSequence& DbSequence = Sequences[ValidAsset->SourceAssetIdx];
+
 	const FFloatInterval SamplingRange = ValidAsset->SamplingInterval;
-	float TimeOffsetSeconds = SamplingRange.Min + Schema->SamplingInterval * (PoseIdx - Asset->FirstPoseIdx);
+	float TimeOffsetSeconds = FMath::Min(SamplingRange.Min + Schema->SamplingInterval * (PoseIdx - ValidAsset->FirstPoseIdx), SamplingRange.Max);
 	return TimeOffsetSeconds;
 }
 
@@ -2916,7 +2925,15 @@ static void DrawTrajectoryFeatures(const FDebugDrawParams& DrawParams, const FFe
 			FLinearColor GradientColor = GetGradientColor(LinearColor, SchemaSubsampleIdx, NumSubsamples, DrawParams.Flags);
 			FColor Color = GradientColor.ToFColor(true);
 
-			FString SampleLabel = FString::Format(TEXT("{0}"), { SchemaSubsampleIdx });
+			FString SampleLabel;
+			if (DrawParams.LabelPrefix.IsEmpty())
+			{
+				SampleLabel = FString::Format(TEXT("{0}"), { SchemaSubsampleIdx });
+			}
+			else
+			{
+				SampleLabel = FString::Format(TEXT("{1}[{0}]"), { SchemaSubsampleIdx, DrawParams.LabelPrefix.GetData() });
+			}
 			DrawDebugString(
 				DrawParams.World,
 				TrajectoryPos + DrawDebugSampleLabelOffset,
@@ -3811,6 +3828,60 @@ float ComparePoses(const FPoseSearchIndex& SearchIndex, int32 PoseIdx, TArrayVie
 	return Dissimilarity;
 }
 
+float ComparePoses(const FPoseSearchIndex& SearchIndex, int32 PoseIdx, TArrayView<const float> Query, const FPoseSearchWeightsContext* WeightsContext, FPoseCostInfo& OutPoseCostInfo)
+{
+	using namespace Eigen;
+
+	TArrayView<const float> PoseValues = SearchIndex.GetPoseValues(PoseIdx);
+	const int32 Dims = PoseValues.Num();
+	if (!ensure(PoseValues.Num() == Dims))
+	{
+		return MAX_flt;
+	}
+
+	OutPoseCostInfo.CostVector.SetNum(Dims);
+
+	// Setup Eigen views onto our vectors
+	auto OutCostVector = Map<ArrayXf>(OutPoseCostInfo.CostVector.GetData(), Dims);
+	auto PoseVector = Map<const ArrayXf>(PoseValues.GetData(), Dims);
+	auto QueryVector = Map<const ArrayXf>(Query.GetData(), Dims);
+	
+	// Compute weighted squared difference vector
+	float Dissimilarity;
+	if (WeightsContext)
+	{
+		const FPoseSearchWeights* WeightsSet = WeightsContext->GetGroupWeights(0);
+		check(WeightsSet);
+		check(WeightsSet->Weights.Num() == Dims);
+		auto WeightsVector = Map<const ArrayXf>(WeightsSet->Weights.GetData(), Dims);
+
+		OutCostVector = WeightsVector * (PoseVector - QueryVector).square();
+		Dissimilarity = OutCostVector.sum();
+	}
+	else
+	{
+		OutCostVector = (PoseVector - QueryVector).square();
+		Dissimilarity = OutCostVector.sum();
+	}
+
+	// Verify this math agrees with the runtime pose comparator
+	checkSlow(FMath::IsNearlyEqual(Dissimilarity, ComparePoses(SearchIndex, PoseIdx, Query, WeightsContext), KINDA_SMALL_NUMBER));
+
+	const FPoseSearchPoseMetadata& PoseMetadata = SearchIndex.PoseMetadata[PoseIdx];
+	float TotalDissimilarity = Dissimilarity + PoseMetadata.CostAddend;
+	
+	// Output cost details
+	OutPoseCostInfo.TotalCost = TotalDissimilarity;
+	OutPoseCostInfo.CostAddend = PoseMetadata.CostAddend;
+	CalcChannelCosts(SearchIndex.Schema, OutPoseCostInfo.CostVector, OutPoseCostInfo.ChannelCosts);
+
+	// Verify channel cost decomposition agrees with pose comparator
+	auto OutChannelCosts = Map<const ArrayXf>(OutPoseCostInfo.ChannelCosts.GetData(), OutPoseCostInfo.ChannelCosts.Num());
+	checkSlow(FMath::IsNearlyEqual(Dissimilarity, OutChannelCosts.sum(), KINDA_SMALL_NUMBER));
+
+	return TotalDissimilarity;
+}
+
 
 //////////////////////////////////////////////////////////////////////////
 // FModule
diff --git a/Engine/Plugins/Experimental/Animation/PoseSearch/Source/Runtime/Public/PoseSearch/AnimNode_MotionMatching.h b/Engine/Plugins/Experimental/Animation/PoseSearch/Source/Runtime/Public/PoseSearch/AnimNode_MotionMatching.h
--- a/Engine/Plugins/Experimental/Animation/PoseSearch/Source/Runtime/Public/PoseSearch/AnimNode_MotionMatching.h
+++ b/Engine/Plugins/Experimental/Animation/PoseSearch/Source/Runtime/Public/PoseSearch/AnimNode_MotionMatching.h
@@ -37,13 +37,13 @@ public:
 	FMotionMatchingSettings Settings;
 
 #if WITH_EDITORONLY_DATA
-	UPROPERTY(EditAnywhere, Category=Debug)
+	UPROPERTY(EditAnywhere, Category=Debug, meta = (PinShownByDefault))
 	bool bDebugDraw = false;
 
-	UPROPERTY(EditAnywhere, Category=Debug)
+	UPROPERTY(EditAnywhere, Category=Debug, meta = (PinShownByDefault))
 	bool bDebugDrawQuery = true;
 
-	UPROPERTY(EditAnywhere, Category=Debug)
+	UPROPERTY(EditAnywhere, Category=Debug, meta = (PinShownByDefault))
 	bool bDebugDrawMatch = true;
 #endif
 
diff --git a/Engine/Plugins/Experimental/Animation/PoseSearch/Source/Runtime/Public/PoseSearch/PoseSearch.h b/Engine/Plugins/Experimental/Animation/PoseSearch/Source/Runtime/Public/PoseSearch/PoseSearch.h
--- a/Engine/Plugins/Experimental/Animation/PoseSearch/Source/Runtime/Public/PoseSearch/PoseSearch.h
+++ b/Engine/Plugins/Experimental/Animation/PoseSearch/Source/Runtime/Public/PoseSearch/PoseSearch.h
@@ -157,7 +157,7 @@ public:
 	UPROPERTY(EditAnywhere, Category = "Schema")
 	bool bUseTrajectoryPositions = true;
 
-	UPROPERTY(EditAnywhere, Category = "Schema")
+	UPROPERTY(EditAnywhere, Category = "Schema", DisplayName="Use Facing Directions")
 	bool bUseTrajectoryForwardVectors = false;
 
 	UPROPERTY(EditAnywhere, Category = "Schema")
@@ -256,7 +256,7 @@ enum class EPoseSearchPoseFlags : uint32
 {
 	None = 0,
 
-	/** Can be returned in a pose search */
+	// Don't return this pose as a search result
 	BlockTransition = 1 << 0,
 };
 ENUM_CLASS_FLAGS(EPoseSearchPoseFlags);
@@ -402,7 +402,7 @@ public:
 	UPROPERTY(EditAnywhere, Category = "Settings")
 	FPoseSearchExtrapolationParameters ExtrapolationParameters;
 
-	UPROPERTY();
+	UPROPERTY()
 	FPoseSearchIndex SearchIndex;
 
 	bool IsValidForIndexing() const;
@@ -589,7 +589,6 @@ struct POSESEARCH_API FPoseSearchDatabaseSequence
 	UPROPERTY(EditAnywhere, Category="Sequence")
 	bool bLoopFollowUpAnimation = false;
 
-public:
 	FFloatInterval GetEffectiveSamplingRange() const;
 };
 
@@ -625,7 +624,6 @@ public:
 	int32 FindSequenceForPose(int32 PoseIdx) const;
 	float GetSequenceLength(int32 DbSequenceIdx) const;
 	bool DoesSequenceLoop(int32 DbSequenceIdx) const;
-	FFloatInterval GetEffectiveSamplingRange(int32 DbSequenceIdx) const;
 
 	bool IsValidForIndexing() const;
 	bool IsValidForSearch() const;
@@ -710,7 +708,6 @@ private:
 namespace UE { namespace PoseSearch {
 
 /** Records poses over time in a ring buffer. FFeatureVectorBuilder uses this to sample from the present or past poses according to the search schema. */
-// @@@Rename to FPoseQueryContext?
 class POSESEARCH_API FPoseHistory
 {
 public:
@@ -788,22 +785,28 @@ enum class EDebugDrawFlags : uint32
 {
 	None			    = 0,
 
-	/** Draw the entire search index */
+	// Draw the entire search index as a point cloud
 	DrawSearchIndex     = 1 << 0,
 
-	/** Draw pose features for each pose vector */
+	// Draw pose features for each pose vector
 	IncludePose         = 1 << 1,
 
-	/** Draw trajectory features for each pose vector */
+	// Draw trajectory features for each pose vector
 	IncludeTrajectory   = 1 << 2,
 
-	/** Draw all pose vector features */
+	// Draw all pose vector features
 	IncludeAllFeatures  = IncludePose | IncludeTrajectory,
 
+	/**
+	 * Keep rendered data until the next call to FlushPersistentDebugLines().
+	 * Combine with DrawSearchIndex to draw the search index only once.
+	 */
 	Persistent = 1 << 3,
 	
+	// Label samples with their indices
 	DrawSampleLabels = 1 << 4,
 
+	// Fade colors
 	DrawSamplesWithColorGradient = 1 << 5,
 };
 ENUM_CLASS_FLAGS(EDebugDrawFlags);
@@ -820,15 +823,18 @@ struct POSESEARCH_API FDebugDrawParams
 
 	FTransform RootTransform = FTransform::Identity;
 
-	/** If set, draw the corresponding pose from the search index */
+	// If set, draw the corresponding pose from the search index
 	int32 PoseIdx = INDEX_NONE;
 
-	/** If set, draw using this uniform color instead of feature-based coloring */
+	// If set, draw using this uniform color instead of feature-based coloring
 	const FLinearColor* Color = nullptr;
 
-	/** If set, interpret the buffer as a pose vector and draw it */
+	// If set, interpret the buffer as a pose vector and draw it
 	TArrayView<const float> PoseVector;
 
+	// Optional prefix for sample labels
+	FStringView LabelPrefix;
+
 	bool CanDraw() const;
 	const FPoseSearchIndex* GetSearchIndex() const;
 	const UPoseSearchSchema* GetSchema() const;
@@ -911,6 +917,40 @@ POSESEARCH_API FSearchResult Search(const UPoseSearchDatabase* Database, TArrayV
 * 
 * @return Dissimilarity between the two poses
 */
-POSESEARCH_API float ComparePoses(const FPoseSearchIndex& SearchIndex, int32 PoseIdx, TArrayView<const float> Query, const FPoseSearchWeightsContext* WeightsContext = nullptr);
+POSESEARCH_API float ComparePoses(const FPoseSearchIndex& SearchIndex, int32 PoseIdx, TArrayView<const float> Query, const FPoseSearchWeightsContext* WeightsContext);
+
+
+/**
+ * Cost details for pose analysis in the rewind debugger
+ */
+struct FPoseCostInfo
+{
+	// The dissimilarity of this pose compared to the query
+	float TotalCost = 0.0f;
+
+	// Contribution from ModifyCost anim notify
+	float CostAddend = 0.0f;
+
+	// Cost breakdown per channel (e.g. pose cost, time-based trajectory cost, distance-based trajectory cost, etc.)
+	TArray<float> ChannelCosts;
+
+	// Difference vector computed as W*((P-Q)^2) without the cost modifier applied
+	// Where P is the pose vector, Q is the query vector, W is the weights vector, and multiplication/exponentiation are element-wise operations
+	TArray<float> CostVector;
+};
+
+/**
+* Evaluate pose comparison metric between a pose in the search index and an input query with cost details
+*
+* @param SearchIndex		The search index containing the pose to compare to the query
+* @param PoseIdx			The index of the pose in the search index to compare to the query
+* @param Query				The query to compare against. Must have been built using the same schema as the SearchIndex.
+* @param WeightsContext		Optional weights context used to influence pose comparison
+* &param OutPoseCostInfo	Cost details for analysis in the debugger
+*
+* @return Dissimilarity between the two poses
+*/
+
+POSESEARCH_API float ComparePoses(const FPoseSearchIndex& SearchIndex, int32 PoseIdx, TArrayView<const float> Query, const FPoseSearchWeightsContext* WeightsContext, FPoseCostInfo& OutPoseCostInfo);
 
 }} // namespace UE::PoseSearch
\ No newline at end of file
