MotionMatchingFFDemo - Added option to index and play mirrorred sequences in UPoseSearchDatabase
- Mirrorring information kept only in FPoseSearchIndexAsset, no longer a feature as in the first implementation
- Mirror mismatch cost set in the database and added after feature vector cost computations, like it happens with the cost modifier notify

diff --git a/Engine/Plugins/Experimental/Animation/PoseSearch/Source/Editor/PoseSearchEditor.Build.cs b/Engine/Plugins/Experimental/Animation/PoseSearch/Source/Editor/PoseSearchEditor.Build.cs
--- a/Engine/Plugins/Experimental/Animation/PoseSearch/Source/Editor/PoseSearchEditor.Build.cs
+++ b/Engine/Plugins/Experimental/Animation/PoseSearch/Source/Editor/PoseSearchEditor.Build.cs
@@ -12,6 +12,7 @@ public class PoseSearchEditor : ModuleRules
 			new string[]
 			{
 				"AnimGraph",
+				"AnimGraphRuntime",
 				"AnimationCore",
 				"Core",
 				"CoreUObject",
diff --git a/Engine/Plugins/Experimental/Animation/PoseSearch/Source/Editor/Private/PoseSearchDebugger.cpp b/Engine/Plugins/Experimental/Animation/PoseSearch/Source/Editor/Private/PoseSearchDebugger.cpp
--- a/Engine/Plugins/Experimental/Animation/PoseSearch/Source/Editor/Private/PoseSearchDebugger.cpp
+++ b/Engine/Plugins/Experimental/Animation/PoseSearch/Source/Editor/Private/PoseSearchDebugger.cpp
@@ -42,12 +42,14 @@ public:
 	
 	float GetPoseCost() const { return PoseCostInfo.ChannelCosts.Num() ? PoseCostInfo.ChannelCosts[0] : 0.0f; }
 	float GetTrajectoryCost() const { return PoseCostInfo.ChannelCosts.Num() >= 3 ? PoseCostInfo.ChannelCosts[1] + PoseCostInfo.ChannelCosts[2] : 0.0f; }
+	float GetAddendsCost() const { return PoseCostInfo.NotifyCostAddend + PoseCostInfo.MirrorMismatchCostAddend; }
 
 	int32 PoseIdx = 0;
 	FString AnimSequenceName = "";
 	int32 DbSequenceIdx = 0;
 	int32 AnimFrame = 0;
 	float Time = 0.0f;
+	bool bMirrored = false;
 	FPoseCostInfo PoseCostInfo;
 };
 
@@ -263,22 +265,59 @@ namespace DebuggerDatabaseColumns
 
 		virtual const FSortPredicate& GetSortPredicateAscending() const override
 		{
-			static FSortPredicate Predicate = [](const FRowDataRef& Row0, const FRowDataRef& Row1) -> bool { return Row0->PoseCostInfo.CostAddend < Row1->PoseCostInfo.CostAddend; };
+			static FSortPredicate Predicate = [](const FRowDataRef& Row0, const FRowDataRef& Row1) -> bool 
+			{ 
+				return Row0->GetAddendsCost() < Row1->GetAddendsCost();
+			};
 			return Predicate;
 		}
 
 		virtual const FSortPredicate& GetSortPredicateDescending() const override
 		{
-			static FSortPredicate Predicate = [](const FRowDataRef& Row0, const FRowDataRef& Row1) -> bool { return Row0->PoseCostInfo.CostAddend >= Row1->PoseCostInfo.CostAddend; };
+			static FSortPredicate Predicate = [](const FRowDataRef& Row0, const FRowDataRef& Row1) -> bool 
+			{ 
+				return Row0->GetAddendsCost() >= Row1->GetAddendsCost();
+			};
 			return Predicate;
 		}
 
 		virtual FText GetRowText(const FRowDataRef& Row) const override
 		{
-			return FText::AsNumber(Row->PoseCostInfo.CostAddend);
+			return FText::AsNumber(Row->GetAddendsCost());
 		}
 	};
 	const FName FCostModifier::Name = "Cost Modifier";
+
+	struct FMirrored : IColumn
+	{
+		using IColumn::IColumn;
+		static const FName Name;
+		virtual FName GetName() const override { return Name; }
+
+		virtual const FSortPredicate& GetSortPredicateAscending() const override
+		{
+			static FSortPredicate Predicate = [](const FRowDataRef& Row0, const FRowDataRef& Row1) -> bool 
+			{ 
+				return Row0->bMirrored < Row1->bMirrored; 
+			};
+			return Predicate;
+		}
+
+		virtual const FSortPredicate& GetSortPredicateDescending() const override
+		{
+			static FSortPredicate Predicate = [](const FRowDataRef& Row0, const FRowDataRef& Row1) -> bool 
+			{ 
+				return Row0->bMirrored >= Row1->bMirrored; 
+			};
+			return Predicate;
+		}
+
+		virtual FText GetRowText(const FRowDataRef& Row) const override
+		{
+			return FText::Format(LOCTEXT("Mirrored", "{0}"), { Row->bMirrored } );
+		}
+	};
+	const FName FMirrored::Name = "Mirrored";
 }
 
 
@@ -530,6 +569,7 @@ void SDebuggerDatabaseView::CreateRows(const UPoseSearchDatabase& Database)
 			Row->AnimSequenceName = DbSequence.Sequence->GetName();
 			Row->Time = Time;
 			Row->AnimFrame = DbSequence.Sequence->GetFrameAtTime(Time);
+			Row->bMirrored = SearchIndexAsset.bMirrored;
 		}
 	}
 
@@ -544,17 +584,25 @@ void SDebuggerDatabaseView::UpdateRows(const FTraceMotionMatchingStateMessage& S
 		CreateRows(Database);
 	}
 	check(ActiveView.Rows.Num() == 1);
-	
+
 	FPoseSearchWeightsContext StateWeights;
 	StateWeights.Update(State.Weights, &Database);
-	
-	const FPoseSearchIndex& SearchIndex = Database.SearchIndex;
 
+	UE::PoseSearch::FSearchContext SearchContext;
+	SearchContext.SetSource(&Database);
+	SearchContext.QueryValues = State.QueryVectorNormalized;
+	SearchContext.WeightsContext = &StateWeights;
+	if (const FPoseSearchIndexAsset* CurrentIndexAsset = Database.SearchIndex.FindAssetForPose(State.DbPoseIdx))
+	{
+		SearchContext.QueryMirrorRequest = CurrentIndexAsset->bMirrored ? 
+			EPoseSearchBooleanRequest::TrueValue : EPoseSearchBooleanRequest::FalseValue;
+	}
+	
 	for(const TSharedRef<FDebuggerDatabaseRowData>& Row : UnfilteredDatabaseRows)
 	{
 		const int32 PoseIdx = Row->PoseIdx;
-		
-		ComparePoses(SearchIndex, PoseIdx, State.QueryVectorNormalized, &StateWeights, Row->PoseCostInfo);
+
+		ComparePoses(PoseIdx, SearchContext, Row->PoseCostInfo);
 
 		// If we are on the active pose for the frame
 		if (PoseIdx == State.DbPoseIdx)
@@ -594,7 +642,8 @@ void SDebuggerDatabaseView::Construct(const FArguments& InArgs)
 	AddColumn(MakeShared<FTrajectoryCost>(3));
 	AddColumn(MakeShared<FCostModifier>(4));
 	AddColumn(MakeShared<FFrame>(5));
-	AddColumn(MakeShared<FPoseIdx>(6));
+	AddColumn(MakeShared<FMirrored>(6));
+	AddColumn(MakeShared<FPoseIdx>(7));
 
 	// Active Row
 	ActiveView.HeaderRow = SNew(SHeaderRow);
@@ -847,7 +896,7 @@ void SDebuggerDetailsView::UpdateReflection(const FTraceMotionMatchingStateMessa
 	Reader.SetValues(State.QueryVector);
 
 	int32 NumSubsamples = Schema->PoseSampleTimes.Num();
-	const int32 NumBones = Schema->Bones.Num();
+	const int32 NumBones = Schema->SampledBones.Num();
 
 	FPoseSearchFeatureDesc Feature;
 	
diff --git a/Engine/Plugins/Experimental/Animation/PoseSearch/Source/Runtime/Private/AnimNode_MotionMatching.cpp b/Engine/Plugins/Experimental/Animation/PoseSearch/Source/Runtime/Private/AnimNode_MotionMatching.cpp
--- a/Engine/Plugins/Experimental/Animation/PoseSearch/Source/Runtime/Private/AnimNode_MotionMatching.cpp
+++ b/Engine/Plugins/Experimental/Animation/PoseSearch/Source/Runtime/Private/AnimNode_MotionMatching.cpp
@@ -24,7 +24,8 @@ void FAnimNode_MotionMatching::Initialize_AnyThread(const FAnimationInitializeCo
 
 	MotionMatchingState.InitNewDatabaseSearch(Database, Settings.SearchThrottleTime);
 
-	Source.SetLinkNode(&SequencePlayerNode);
+	Source.SetLinkNode(&MirrorNode);
+	MirrorNode.SetSourceLinkNode(&SequencePlayerNode);
 	Source.Initialize(Context);
 }
 
@@ -55,13 +56,15 @@ void FAnimNode_MotionMatching::UpdateAssetPlayer(const FAnimationUpdateContext&
 	// If a new pose is requested, jump to the pose by updating the embedded sequence player node
 	if ((MotionMatchingState.Flags & EMotionMatchingFlags::JumpedToPose) == EMotionMatchingFlags::JumpedToPose)
 	{
-		const FPoseSearchIndexAsset* SearchIndexAsset = 
-			&MotionMatchingState.CurrentDatabase->SearchIndex.Assets[MotionMatchingState.SearchIndexAssetIdx];
+		const FPoseSearchIndexAsset* SearchIndexAsset = MotionMatchingState.GetCurrentSearchIndexAsset();
 		const FPoseSearchDatabaseSequence& ResultDbSequence = Database->GetSourceAsset(SearchIndexAsset);
 		SequencePlayerNode.SetSequence(ResultDbSequence.Sequence);
 		SequencePlayerNode.SetAccumulatedTime(MotionMatchingState.AssetPlayerTime);
 		SequencePlayerNode.SetLoopAnimation(ResultDbSequence.bLoopAnimation);
 		SequencePlayerNode.SetPlayRate(1.f);
+
+		MirrorNode.SetMirrorDataTable(Database->Schema->MirrorDataTable.Get());
+		MirrorNode.SetMirror(SearchIndexAsset->bMirrored);
 	}
 
 	// Optionally applying dynamic play rate adjustment to chosen sequences based on predictive motion analysis
diff --git a/Engine/Plugins/Experimental/Animation/PoseSearch/Source/Runtime/Private/PoseSearch.cpp b/Engine/Plugins/Experimental/Animation/PoseSearch/Source/Runtime/Private/PoseSearch.cpp
--- a/Engine/Plugins/Experimental/Animation/PoseSearch/Source/Runtime/Private/PoseSearch.cpp
+++ b/Engine/Plugins/Experimental/Animation/PoseSearch/Source/Runtime/Private/PoseSearch.cpp
@@ -21,6 +21,7 @@
 #include "Animation/AnimSequence.h"
 #include "Animation/AnimSequenceBase.h"
 #include "Animation/AnimationPoseData.h"
+#include "Animation/MirrorDataTable.h"
 #include "AnimationRuntime.h"
 #include "BonePose.h"
 #include "Trace/PoseSearchTraceLogger.h"
@@ -320,6 +321,31 @@ bool FPoseSearchFeatureVectorLayout::EnumerateBy(int32 ChannelIdx, EPoseSearchFe
 	return false;
 }
 
+
+//////////////////////////////////////////////////////////////////////////
+// FPoseSearchBone
+
+uint32 FPoseSearchBone::GetTypeMask() const
+{
+	uint32 Mask = 0;
+	
+	if (bUsePosition)
+	{
+		Mask |= 1 << static_cast<int>(EPoseSearchFeatureType::Position);
+	}
+	if (bUseVelocity)
+	{
+		Mask |= 1 << static_cast<int>(EPoseSearchFeatureType::LinearVelocity);
+	}
+	if (bUseRotation)
+	{
+		Mask |= 1 << static_cast<int>(EPoseSearchFeatureType::Rotation);
+	}
+
+	return Mask;
+}
+
+
 //////////////////////////////////////////////////////////////////////////
 // UPoseSearchSchema
 
@@ -349,18 +375,34 @@ void UPoseSearchSchema::PostLoad()
 	Super::PostLoad();
 
 	ResolveBoneReferences();
+
+	// todo: remove this and deprecated variables
+	if (SampledBones.Num() == 0 && Bones_DEPRECATED.Num() > 0)
+	{
+		SampledBones.Empty();
+		for (const FBoneReference& Ref : Bones_DEPRECATED)
+		{
+			FPoseSearchBone Bone;
+			Bone.Reference = Ref;
+			Bone.bUsePosition = bUseBonePositions_DEPRECATED;
+			Bone.bUseVelocity = bUseBoneVelocities_DEPRECATED;
+			Bone.bUseRotation = false;
+			SampledBones.Add(Bone);
+		}
+		Bones_DEPRECATED.Empty();
+	}
 }
 
 bool UPoseSearchSchema::IsValid() const
 {
 	bool bValid = Skeleton != nullptr;
 
-	for (const FBoneReference& BoneRef : Bones)
+	for (const FPoseSearchBone& Bone : SampledBones)
 	{
-		bValid &= BoneRef.HasValidSetup();
+		bValid &= Bone.Reference.HasValidSetup();
 	}
 
-	bValid &= (Bones.Num() == BoneIndices.Num());
+	bValid &= (SampledBones.Num() == BoneIndices.Num());
 
 	bValid &= Layout.IsValid(BoneIndices.Num());
 	
@@ -412,7 +454,6 @@ void UPoseSearchSchema::GenerateLayout()
 	auto AddChannelFeatures = [this](int32 ChannelIdx, EPoseSearchFeatureType Type)
 	{
 		const int32 NumSamples = GetChannelSampleOffsets(ChannelIdx).Num();
-		const int32 NumBones = GetNumBones();
 		EPoseSearchFeatureDomain Domain = UE::PoseSearch::FeatureDomainFromChannel(ChannelIdx);
 
 		FPoseSearchFeatureDesc Feature;
@@ -421,11 +462,17 @@ void UPoseSearchSchema::GenerateLayout()
 
 		if (ChannelIdx == UE::PoseSearch::ChannelIdxPose)
 		{
+			const int32 NumBones = SampledBones.Num();
 			for (Feature.SubsampleIdx = 0; Feature.SubsampleIdx != NumSamples; ++Feature.SubsampleIdx)
 			{
 				for (Feature.SchemaBoneIdx = 0; Feature.SchemaBoneIdx != NumBones; ++Feature.SchemaBoneIdx)
 				{
-					Layout.Features.Add(Feature);
+					const FPoseSearchBone& SampledBone = SampledBones[Feature.SchemaBoneIdx];
+					bool bAddFeature = (SampledBone.GetTypeMask() & (1 << static_cast<int>(Type))) != 0;
+					if (bAddFeature)
+					{
+						Layout.Features.Add(Feature);
+					}
 				}
 			}
 		}
@@ -475,17 +522,11 @@ void UPoseSearchSchema::GenerateLayout()
 		AddChannelFeatures(UE::PoseSearch::ChannelIdxTrajectoryDistance, EPoseSearchFeatureType::ForwardVector);
 	}
 
-	// Time domain bone positions
-	if (bUseBonePositions)
-	{
-		AddChannelFeatures(UE::PoseSearch::ChannelIdxPose, EPoseSearchFeatureType::Position);
-	}
 
-	// Time domain bone linear velocities
-	if (bUseBoneVelocities && PoseSampleTimes.Num())
-	{
-		AddChannelFeatures(UE::PoseSearch::ChannelIdxPose, EPoseSearchFeatureType::LinearVelocity);
-	}
+	// Time domain bone features
+	AddChannelFeatures(UE::PoseSearch::ChannelIdxPose, EPoseSearchFeatureType::Position);
+	AddChannelFeatures(UE::PoseSearch::ChannelIdxPose, EPoseSearchFeatureType::Rotation);
+	AddChannelFeatures(UE::PoseSearch::ChannelIdxPose, EPoseSearchFeatureType::LinearVelocity);
 
 	Layout.Init();
 }
@@ -493,16 +534,16 @@ void UPoseSearchSchema::GenerateLayout()
 void UPoseSearchSchema::ResolveBoneReferences()
 {
 	// Initialize references to obtain bone indices
-	for (FBoneReference& BoneRef : Bones)
+	for (FPoseSearchBone& Bone : SampledBones)
 	{
-		BoneRef.Initialize(Skeleton);
+		Bone.Reference.Initialize(Skeleton);
 	}
 
 	// Fill out bone index array and sort by bone index
-	BoneIndices.SetNum(Bones.Num());
-	for (int32 Index = 0; Index != Bones.Num(); ++Index)
+	BoneIndices.SetNum(SampledBones.Num());
+	for (int32 Index = 0; Index != SampledBones.Num(); ++Index)
 	{
-		BoneIndices[Index] = Bones[Index].BoneIndex;
+		BoneIndices[Index] = SampledBones[Index].Reference.BoneIndex;
 	}
 	BoneIndices.Sort();
 
@@ -512,6 +553,12 @@ void UPoseSearchSchema::ResolveBoneReferences()
 	{
 		FAnimationRuntime::EnsureParentsPresent(BoneIndicesWithParents, Skeleton->GetReferenceSkeleton());
 	}
+
+	// BoneIndicesWithParents should at least contain the root to support mirroring root motion
+	if (BoneIndicesWithParents.Num() == 0)
+	{
+		BoneIndicesWithParents.Add(0);
+	}
 }
 
 
@@ -574,7 +621,6 @@ void FPoseSearchWeights::Init(const FPoseSearchWeightParams& WeightParams, const
 	ChannelDynamicWeightParams[ChannelIdxTrajectoryTime] = &DynamicWeightParams.TrajectoryDynamicWeights;
 	ChannelDynamicWeightParams[ChannelIdxTrajectoryDistance] = &DynamicWeightParams.TrajectoryDynamicWeights;
 
-
 	// Normalize channel weights
 	Array<float, ChannelNum, 1> NormalizedChannelWeights;
 	for (int ChannelIdx = 0; ChannelIdx != ChannelNum; ++ChannelIdx)
@@ -807,7 +853,9 @@ void FPoseSearchWeights::Init(const FPoseSearchWeightParams& WeightParams, const
 //////////////////////////////////////////////////////////////////////////
 // FPoseSearchWeightsContext
 
-void FPoseSearchWeightsContext::Update(const FPoseSearchDynamicWeightParams& ActiveWeights, const UPoseSearchDatabase* ActiveDatabase)
+void FPoseSearchWeightsContext::Update(
+	const FPoseSearchDynamicWeightParams& ActiveWeights, 
+	const UPoseSearchDatabase* ActiveDatabase)
 {
 	bool bRecomputeWeights = false;
 	if (Database != ActiveDatabase)
@@ -874,6 +922,31 @@ const FPoseSearchIndexAsset* FPoseSearchIndex::FindAssetForPose(int32 PoseIdx) c
 	return Assets.FindByPredicate(Predicate);
 }
 
+float FPoseSearchIndex::GetTimeOffset(int32 PoseIdx, const FPoseSearchIndexAsset* Asset) const
+{
+	const FPoseSearchIndexAsset* ValidAsset = Asset;
+	if (!ValidAsset)
+	{
+		ValidAsset = FindAssetForPose(PoseIdx);
+		if (!ValidAsset)
+		{
+			UE_LOG(LogPoseSearch, Error, TEXT("Couldn't find asset for pose %i in database"), PoseIdx);
+			return -1.0f;
+		}
+	}
+	if (!ValidAsset->IsPoseInRange(PoseIdx))
+	{
+		UE_LOG(LogPoseSearch, Error, TEXT("Pose %i out of range in database"), PoseIdx);
+		return -1.0f;
+	}
+
+	const FFloatInterval SamplingRange = ValidAsset->SamplingInterval;
+	float TimeOffsetSeconds = FMath::Min(
+		SamplingRange.Min + Schema->SamplingInterval * (PoseIdx - ValidAsset->FirstPoseIdx), 
+		SamplingRange.Max);
+	return TimeOffsetSeconds;
+}
+
 
 //////////////////////////////////////////////////////////////////////////
 // FPoseSearchIndex
@@ -1034,25 +1107,7 @@ int32 UPoseSearchDatabase::GetPoseIndexFromAssetTime(float AssetTime, const FPos
 
 float UPoseSearchDatabase::GetTimeOffset(int32 PoseIdx, const FPoseSearchIndexAsset* Asset) const
 {
-	const FPoseSearchIndexAsset* ValidAsset = Asset;
-	if (!ValidAsset)
-	{
-		ValidAsset = SearchIndex.FindAssetForPose(PoseIdx);
-		if (!ValidAsset)
-		{
-			UE_LOG(LogPoseSearch, Error, TEXT("Couldn't find asset for pose %i in database %s"), PoseIdx, *GetNameSafe(this));
-			return -1.0f;
-		}
-	}
-	if (!ValidAsset->IsPoseInRange(PoseIdx))
-	{
-		UE_LOG(LogPoseSearch, Error, TEXT("Pose %i out of range in database %s"), PoseIdx, *GetNameSafe(this));
-		return -1.0f;
-	}
-
-	const FFloatInterval SamplingRange = ValidAsset->SamplingInterval;
-	float TimeOffsetSeconds = FMath::Min(SamplingRange.Min + Schema->SamplingInterval * (PoseIdx - ValidAsset->FirstPoseIdx), SamplingRange.Max);
-	return TimeOffsetSeconds;
+	return SearchIndex.GetTimeOffset(PoseIdx, Asset);
 }
 
 const FPoseSearchDatabaseSequence& UPoseSearchDatabase::GetSourceAsset(const FPoseSearchIndexAsset* SearchIndexAsset) const
@@ -1060,6 +1115,12 @@ const FPoseSearchDatabaseSequence& UPoseSearchDatabase::GetSourceAsset(const FPo
 	return Sequences[SearchIndexAsset->SourceAssetIdx];
 }
 
+int32 UPoseSearchDatabase::FindSequenceForPose(int32 PoseIdx) const
+{
+	const FPoseSearchIndexAsset* Asset = SearchIndex.FindAssetForPose(PoseIdx);
+	return Asset->SourceAssetIdx;
+}
+
 float UPoseSearchDatabase::GetSequenceLength(int32 DbSequenceIdx) const
 {
 	return Sequences[DbSequenceIdx].Sequence->GetPlayLength();
@@ -1174,17 +1235,51 @@ bool UPoseSearchDatabase::TryInitSearchIndexAssets()
 	TArray<FFloatRange> ValidRanges;
 	for (int32 SequenceIdx = 0; SequenceIdx < Sequences.Num(); ++SequenceIdx)
 	{
+		const FPoseSearchDatabaseSequence& Sequence = Sequences[SequenceIdx];
+		bool bAddUnmirrored = 
+			Sequence.MirrorOption == EPoseSearchMirrorOption::UnmirroredOnly ||
+			Sequence.MirrorOption == EPoseSearchMirrorOption::UnmirroredAndMirrored;
+		bool bAddMirrored =
+			Sequence.MirrorOption == EPoseSearchMirrorOption::MirroredOnly ||
+			Sequence.MirrorOption == EPoseSearchMirrorOption::UnmirroredAndMirrored;
+
 		ValidRanges.Reset();
-		FindValidSequenceIntervals(Sequences[SequenceIdx], ValidRanges);
+		FindValidSequenceIntervals(Sequence, ValidRanges);
 		for (const FFloatRange& Range : ValidRanges)
 		{
-			SearchIndex.Assets.Add(
-				FPoseSearchIndexAsset(
-					SequenceIdx, 
-					FFloatInterval(Range.GetLowerBoundValue(), Range.GetUpperBoundValue())));
+			if (bAddUnmirrored)
+			{
+				SearchIndex.Assets.Add(
+					FPoseSearchIndexAsset(
+						SequenceIdx,
+						false,
+						FFloatInterval(Range.GetLowerBoundValue(), Range.GetUpperBoundValue())));
+			}
+
+			if (bAddMirrored)
+			{
+				SearchIndex.Assets.Add(
+					FPoseSearchIndexAsset(
+						SequenceIdx,
+						true,
+						FFloatInterval(Range.GetLowerBoundValue(), Range.GetUpperBoundValue())));
+				bAnyMirrored = true;
+			}
 		}
 	}
 
+	if (bAnyMirrored && !Schema->MirrorDataTable)
+	{
+		UE_LOG(
+			LogPoseSearch, 
+			Error, 
+			TEXT("Database %s is asking for mirrored sequences but MirrorDataBase is null in %s"),
+			*GetNameSafe(this), 
+			*GetNameSafe(Schema));
+		SearchIndex.Assets.Empty();
+		return false;
+	}
+
 	return true;
 }
 
@@ -1830,6 +1925,121 @@ const UPoseSearchSchema* FDebugDrawParams::GetSchema() const
 	return nullptr;
 }
 
+//////////////////////////////////////////////////////////////////////////
+// FSearchContext
+
+void FSearchContext::SetSource(const UPoseSearchDatabase* InSourceDatabase)
+{
+	SearchIndex = nullptr;
+	DebugDrawParams.Database = nullptr;
+	DebugDrawParams.SequenceMetaData = nullptr;
+
+	SourceDatabase = InSourceDatabase;
+	if (SourceDatabase)
+	{
+		if (ensure(SourceDatabase->IsValidForSearch()))
+		{
+			SearchIndex = &SourceDatabase->SearchIndex;
+			DebugDrawParams.Database = SourceDatabase;
+			MirrorMismatchCost = SourceDatabase->MirroringMismatchCost;
+		}
+	}
+}
+
+void FSearchContext::SetSource(const UAnimSequenceBase* InSourceSequence)
+{
+	SearchIndex = nullptr;
+	DebugDrawParams.Database = nullptr;
+	DebugDrawParams.SequenceMetaData = nullptr;
+
+	SourceSequence = InSourceSequence;
+	const UPoseSearchSequenceMetaData* MetaData =
+		SourceSequence->FindMetaDataByClass<UPoseSearchSequenceMetaData>();
+	if (MetaData && MetaData->IsValidForSearch())
+	{
+		SearchIndex = &MetaData->SearchIndex;
+		DebugDrawParams.SequenceMetaData = MetaData;
+	}
+}
+
+const FPoseSearchIndex* FSearchContext::GetSearchIndex() const
+{
+	return SearchIndex;
+}
+
+float FSearchContext::GetMirrorMismatchCost() const
+{
+	return MirrorMismatchCost;
+}
+
+//////////////////////////////////////////////////////////////////////////
+// FAnimSamplingContext
+
+struct FAnimSamplingContext
+{
+	// Time delta used for computing pose derivatives
+	static constexpr float FiniteDelta = 1 / 60.0f;
+
+	FBoneContainer BoneContainer;
+	
+	// Mirror data table pointer copied from Schema for convenience
+	TObjectPtr<UMirrorDataTable> MirrorDataTable = nullptr;
+	
+	// Compact pose format of Mirror Bone Map
+	TCustomBoneIndexArray<FCompactPoseBoneIndex, FCompactPoseBoneIndex> CompactPoseMirrorBones;
+	
+	// Pre-calculated component space rotations of reference pose, which allows mirror to work with any joint orientation
+	// Only initialized and used when a mirroring table is specified
+	TCustomBoneIndexArray<FQuat, FCompactPoseBoneIndex> ComponentSpaceRefRotations;
+	
+	void Init(const UPoseSearchSchema* Schema);
+	
+	FTransform MirrorTransform(const FTransform& Transform) const;
+
+private:
+	void FillCompactPoseAndComponentRefRotations();
+};
+
+void FAnimSamplingContext::Init(const UPoseSearchSchema* Schema)
+{
+	MirrorDataTable = Schema->MirrorDataTable;
+	// ***
+	// Note this is a change for the UPoseSearchSequenceMetaData flavor of BuildIndex
+	// Was previously doing Sequence->GetSkeleton() instead of Schema->GetSkeleton()
+	// Hmm, is that okay?
+	BoneContainer.InitializeTo(Schema->BoneIndicesWithParents, FCurveEvaluationOption(false), *Schema->Skeleton);
+	FillCompactPoseAndComponentRefRotations();
+}
+
+FTransform FAnimSamplingContext::MirrorTransform(const FTransform& InTransform) const
+{
+	const EAxis::Type MirrorAxis = MirrorDataTable->MirrorAxis;
+	FVector T = InTransform.GetTranslation();
+	T = FAnimationRuntime::MirrorVector(T, MirrorAxis);
+	const FQuat ReferenceRotation = ComponentSpaceRefRotations[FCompactPoseBoneIndex(0)];
+	FQuat Q = InTransform.GetRotation();
+	Q = FAnimationRuntime::MirrorQuat(Q, MirrorAxis);
+	Q *= FAnimationRuntime::MirrorQuat(ReferenceRotation, MirrorAxis).Inverse() * ReferenceRotation;
+	FTransform Result = FTransform(Q, T, InTransform.GetScale3D());
+	return Result;
+}
+
+void FAnimSamplingContext::FillCompactPoseAndComponentRefRotations()
+{
+	if (!MirrorDataTable.IsNull())
+	{
+		const UMirrorDataTable* MirrorDataTablePtr = MirrorDataTable.Get();
+		MirrorDataTablePtr->FillCompactPoseAndComponentRefRotations(
+			BoneContainer, 
+			CompactPoseMirrorBones, 
+			ComponentSpaceRefRotations);
+	}
+	else
+	{
+		CompactPoseMirrorBones.Reset();
+		ComponentSpaceRefRotations.Reset();
+	}
+}
 
 //////////////////////////////////////////////////////////////////////////
 // FSequenceSampler
@@ -1839,6 +2049,7 @@ struct FSequenceSampler
 public:
 	struct FInput
 	{
+		const FAnimSamplingContext* SamplingContext = nullptr;
 		const UPoseSearchSchema* Schema = nullptr;
 		const UAnimSequence* Sequence = nullptr;
 		bool bLoopable = false;
@@ -1849,7 +2060,6 @@ public:
 	struct FOutput
 	{
 		TArray<float> AccumulatedRootDistance;
-
 		int32 NumDistanceSamples = 0;
 		float PlayLength = 0.0f;
 		float TotalRootDistance = 0.0f;
@@ -1860,6 +2070,8 @@ public:
 	void Init(const FInput& Input);
 	void Process();
 
+	// Extracts pose from input sequence and mirrors it if necessary
+	void ExtractPose(const FAnimExtractContext& ExtractionCtx, bool bMirrored, FAnimationPoseData& OutAnimPoseData) const;
 
 	// Extracts root transform at the given time, using the extremities of the sequence to extrapolate beyond the 
 	// sequence limits when Time is less than zero or greater than the sequence length.
@@ -1922,6 +2134,22 @@ void FSequenceSampler::Process()
 	ProcessRootMotion();
 }
 
+void FSequenceSampler::ExtractPose(const FAnimExtractContext& ExtractionCtx, bool bMirrored, FAnimationPoseData& OutAnimPoseData) const
+{
+	check(!bMirrored || Input.Schema->MirrorDataTable);
+	Input.Sequence->GetAnimationPose(OutAnimPoseData, ExtractionCtx);
+	if (bMirrored)
+	{
+		FAnimationRuntime::MirrorPose(
+			OutAnimPoseData.GetPose(),
+			Input.Schema->MirrorDataTable->MirrorAxis,
+			Input.SamplingContext->CompactPoseMirrorBones,
+			Input.SamplingContext->ComponentSpaceRefRotations
+		);
+		// Note curves and attributes are not used during the indexing process and therefore don't need to be mirrored
+	}
+}
+
 FTransform FSequenceSampler::ExtractRootTransform(float Time) const
 {
 	if (Input.bLoopable)
@@ -2146,6 +2374,9 @@ void FSequenceSampler::ProcessRootMotion()
 	Output.TotalRootDistance = Output.AccumulatedRootDistance.Last();
 }
 
+//////////////////////////////////////////////////////////////////////////
+// FSequenceIndexer helpers
+
 struct FSamplingParam
 {
 	float WrappedParam = 0.0f;
@@ -2212,11 +2443,12 @@ public:
 
 	struct FInput
 	{
-		const FBoneContainer* BoneContainer = nullptr;
+		const FAnimSamplingContext* SamplingContext = nullptr;
 		const UPoseSearchSchema* Schema = nullptr;
 		const FSequenceSampler* MainSequence = nullptr;
 		const FSequenceSampler* LeadInSequence = nullptr;
 		const FSequenceSampler* FollowUpSequence = nullptr;
+		bool bMirrored = false;
 		FFloatInterval RequestedSamplingRange = FFloatInterval(0.0f, 0.0f);
 		FPoseSearchBlockTransitionParameters BlockTransitionParameters;
 	} Input;
@@ -2251,6 +2483,7 @@ private:
 	FSampleInfo GetSampleInfo(float SampleTime) const;
 	FSampleInfo GetSampleInfoRelative(float SampleTime, const FSampleInfo& Origin) const;
 	const float GetSampleTimeFromDistance(float Distance) const;
+	FTransform MirrorTransform(const FTransform& Transform) const;
 
 	void Reserve();
 	void SampleBegin(int32 SampleIdx);
@@ -2600,6 +2833,11 @@ FSequenceIndexer::FSampleInfo FSequenceIndexer::GetSampleInfoRelative(float Samp
 	return Sample;
 }
 
+FTransform FSequenceIndexer::MirrorTransform(const FTransform& Transform) const
+{
+	return Input.bMirrored ? Input.SamplingContext->MirrorTransform(Transform) : Transform;
+}
+
 void FSequenceIndexer::AddPoseFeatures(int32 SampleIdx)
 {
 	// This function samples the instantaneous pose at time t as well as the pose's velocity and acceleration at time t.
@@ -2608,29 +2846,28 @@ void FSequenceIndexer::AddPoseFeatures(int32 SampleIdx)
 	//	Second symmetric derivative: f''(t) ~ (f(t+h) - 2f(t) + f(t-h)) / h^2
 	// Where h is a constant time delta
 	// So this means three pose extractions are taken at time t-h, t, and t+h
-	constexpr float FiniteDelta = 1 / 60.0f;
 	constexpr int32 NumFiniteDiffTerms = 3;
 
-	if (Input.Schema->Bones.IsEmpty() || Input.Schema->PoseSampleTimes.IsEmpty())
+	if (Input.Schema->SampledBones.IsEmpty() || Input.Schema->PoseSampleTimes.IsEmpty())
 	{
 		return;
 	}
 
 	FCompactPose Poses[NumFiniteDiffTerms];
 	FCSPose<FCompactPose> ComponentSpacePoses[NumFiniteDiffTerms];
-	FBlendedCurve Curves[NumFiniteDiffTerms];
-	UE::Anim::FStackAttributeContainer Atrributes[NumFiniteDiffTerms];
+	FBlendedCurve UnusedCurves[NumFiniteDiffTerms];
+	UE::Anim::FStackAttributeContainer UnusedAtrributes[NumFiniteDiffTerms];
 
 	for(FCompactPose& Pose: Poses)
 	{
-		Pose.SetBoneContainer(Input.BoneContainer);
+		Pose.SetBoneContainer(&Input.SamplingContext->BoneContainer);
 	}
 
 	FAnimationPoseData AnimPoseData[NumFiniteDiffTerms] = 
 	{
-		{Poses[0], Curves[0], Atrributes[0]},
-		{Poses[1], Curves[1], Atrributes[1]},
-		{Poses[2], Curves[2], Atrributes[2]},
+		{Poses[0], UnusedCurves[0], UnusedAtrributes[0]},
+		{Poses[1], UnusedCurves[1], UnusedAtrributes[1]},
+		{Poses[2], UnusedCurves[2], UnusedAtrributes[2]},
 	};
 
 	FAnimExtractContext ExtractionCtx;
@@ -2651,15 +2888,15 @@ void FSequenceIndexer::AddPoseFeatures(int32 SampleIdx)
 		// For each pose subsample term, get the corresponding clip, accumulated root motion,
 		// and wrap the time parameter based on the clip's length.
 		FSampleInfo Samples[NumFiniteDiffTerms];
-		Samples[0] = GetSampleInfoRelative(SubsampleTime - FiniteDelta, Origin);
+		Samples[0] = GetSampleInfoRelative(SubsampleTime - Input.SamplingContext->FiniteDelta, Origin);
 		Samples[1] = GetSampleInfoRelative(SubsampleTime, Origin);
-		Samples[2] = GetSampleInfoRelative(SubsampleTime + FiniteDelta, Origin);
+		Samples[2] = GetSampleInfoRelative(SubsampleTime + Input.SamplingContext->FiniteDelta, Origin);
 
 		// Get pose samples
 		for (int32 Term = 0; Term != NumFiniteDiffTerms; ++Term)
 		{
 			ExtractionCtx.CurrentTime = Samples[Term].ClipTime;
-			Samples[Term].Clip->Input.Sequence->GetAnimationPose(AnimPoseData[Term], ExtractionCtx);
+			Samples[Term].Clip->ExtractPose(ExtractionCtx, Input.bMirrored, AnimPoseData[Term]);
 			ComponentSpacePoses[Term].InitPose(Poses[Term]);
 		}
 
@@ -2670,19 +2907,21 @@ void FSequenceIndexer::AddPoseFeatures(int32 SampleIdx)
 		{
 			Feature.SchemaBoneIdx = SchemaBoneIndex;
 
-			FCompactPoseBoneIndex CompactBoneIndex = Input.BoneContainer->MakeCompactPoseIndex(FMeshPoseBoneIndex(Input.Schema->BoneIndices[SchemaBoneIndex]));
+			FCompactPoseBoneIndex CompactBoneIndex = 
+				Input.SamplingContext->BoneContainer.MakeCompactPoseIndex(
+					FMeshPoseBoneIndex(Input.Schema->BoneIndices[SchemaBoneIndex]));
 
 			FTransform BoneTransforms[NumFiniteDiffTerms];
 			for (int32 Term = 0; Term != NumFiniteDiffTerms; ++Term)
 			{
 				BoneTransforms[Term] = ComponentSpacePoses[Term].GetComponentSpaceTransform(CompactBoneIndex);
-				BoneTransforms[Term] = BoneTransforms[Term] * Samples[Term].RootTransform;
+				BoneTransforms[Term] = BoneTransforms[Term] * MirrorTransform(Samples[Term].RootTransform);
 			}
 
 			// Add properties to the feature vector for the pose at SampleIdx
 			FeatureVector.SetTransform(Feature, BoneTransforms[1]);
-			FeatureVector.SetTransformVelocity(Feature, BoneTransforms[2], BoneTransforms[0], 2.0f * FiniteDelta);
-			//FeatureVector.SetTransformAccleration(Feature, BoneTransforms[2], BoneTransforms[1], BoneTransforms[0], FiniteDelta * FiniteDelta);
+			FeatureVector.SetTransformVelocity(Feature, BoneTransforms[2], BoneTransforms[0], 2.0f * Input.SamplingContext->FiniteDelta);
+			//FeatureVector.SetTransformAccleration(Feature, BoneTransforms[2], BoneTransforms[1], BoneTransforms[0], FMath::Square(Input.SamplingContext->FiniteDelta));
 		}
 	}
 }
@@ -2695,7 +2934,6 @@ void FSequenceIndexer::AddTrajectoryTimeFeatures(int32 SampleIdx)
 	//	Second symmetric derivative: f''(t) ~ (f(t+h) - 2f(t) + f(t-h)) / h^2
 	// Where h is a constant time delta
 	// So this means three root motion extractions are taken at time t-h, t, and t+h
-	constexpr float FiniteDelta = 1 / 60.0f;
 
 	FPoseSearchFeatureDesc Feature;
 	Feature.Domain = EPoseSearchFeatureDomain::Time;
@@ -2713,14 +2951,20 @@ void FSequenceIndexer::AddTrajectoryTimeFeatures(int32 SampleIdx)
 		// For each pose subsample term, get the corresponding clip, accumulated root motion,
 		// and wrap the time parameter based on the clip's length.
 		FSampleInfo Samples[3];
-		Samples[0] = GetSampleInfoRelative(SubsampleTime - FiniteDelta, Origin);
+		Samples[0] = GetSampleInfoRelative(SubsampleTime - Input.SamplingContext->FiniteDelta, Origin);
 		Samples[1] = GetSampleInfoRelative(SubsampleTime, Origin);
-		Samples[2] = GetSampleInfoRelative(SubsampleTime + FiniteDelta, Origin);
+		Samples[2] = GetSampleInfoRelative(SubsampleTime + Input.SamplingContext->FiniteDelta, Origin);
+
+		// Mirror transforms if requested
+		FTransform MirroredRoots[3];
+		MirroredRoots[0] = MirrorTransform(Samples[0].RootTransform);
+		MirroredRoots[1] = MirrorTransform(Samples[1].RootTransform);
+		MirroredRoots[2] = MirrorTransform(Samples[2].RootTransform);
 
 		// Add properties to the feature vector for the pose at SampleIdx
-		FeatureVector.SetTransform(Feature, Samples[1].RootTransform);
-		FeatureVector.SetTransformVelocity(Feature, Samples[2].RootTransform, Samples[0].RootTransform, 2.0f * FiniteDelta);
-		//FeatureVector.SetTransformAcceleration(Feature, Samples[2].RootTransform, Samples[1].RootTransform, Samples[0].RootTransform, FiniteDelta * FiniteDelta);
+		FeatureVector.SetTransform(Feature, MirroredRoots[1]);
+		FeatureVector.SetTransformVelocity(Feature, MirroredRoots[2], MirroredRoots[0], 2.0f * Input.SamplingContext->FiniteDelta);
+		//FeatureVector.SetTransformAcceleration(Feature, MirroredRoots[2], MirroredRoots[1], MirroredRoots[0], FMath::Square(Input.SamplingContext->FiniteDelta));
 	}
 }
 
@@ -2728,7 +2972,6 @@ void FSequenceIndexer::AddTrajectoryDistanceFeatures(int32 SampleIdx)
 {
 	// This function is very similar to AddTrajectoryTimeFeatures, but samples are taken in the distance domain
 	// instead of time domain.
-	constexpr float FiniteDelta = 1 / 60.0f;
 
 	FPoseSearchFeatureDesc Feature;
 	Feature.Domain = EPoseSearchFeatureDomain::Distance;
@@ -2737,6 +2980,7 @@ void FSequenceIndexer::AddTrajectoryDistanceFeatures(int32 SampleIdx)
 	float SampleTime = FMath::Min(SampleIdx * Input.Schema->SamplingInterval, Input.MainSequence->Output.PlayLength);
 	FSampleInfo Origin = GetSampleInfo(SampleTime);
 	
+	FQuat RootReferenceRot = Input.SamplingContext->BoneContainer.GetRefPoseTransform(FCompactPoseBoneIndex(0)).GetRotation();
 	for (int32 SchemaSubsampleIdx = 0; SchemaSubsampleIdx != Input.Schema->TrajectorySampleDistances.Num(); ++SchemaSubsampleIdx)
 	{
 		Feature.SubsampleIdx = SchemaSubsampleIdx;
@@ -2750,14 +2994,20 @@ void FSequenceIndexer::AddTrajectoryDistanceFeatures(int32 SampleIdx)
 		// For each pose subsample term, get the corresponding clip, accumulated root motion,
 		// and wrap the time parameter based on the clip's length.
 		FSampleInfo Samples[3];
-		Samples[0] = GetSampleInfoRelative(SubsampleTime - FiniteDelta, Origin);
+		Samples[0] = GetSampleInfoRelative(SubsampleTime - Input.SamplingContext->FiniteDelta, Origin);
 		Samples[1] = GetSampleInfoRelative(SubsampleTime, Origin);
-		Samples[2] = GetSampleInfoRelative(SubsampleTime + FiniteDelta, Origin);
+		Samples[2] = GetSampleInfoRelative(SubsampleTime + Input.SamplingContext->FiniteDelta, Origin);
+
+		// Mirror transforms if requested
+		FTransform MirroredRoots[3];
+		MirroredRoots[0] = MirrorTransform(Samples[0].RootTransform);
+		MirroredRoots[1] = MirrorTransform(Samples[1].RootTransform);
+		MirroredRoots[2] = MirrorTransform(Samples[2].RootTransform);
 
 		// Add properties to the feature vector for the pose at SampleIdx
-		FeatureVector.SetTransform(Feature, Samples[1].RootTransform);
-		FeatureVector.SetTransformVelocity(Feature, Samples[2].RootTransform, Samples[0].RootTransform, 2.0f * FiniteDelta);
-		//FeatureVector.SetTransformAcceleration(Feature, Samples[0].RootTransform, Samples[1].RootTransform, Samples[2].RootTransform, FiniteDelta * FiniteDelta);
+		FeatureVector.SetTransform(Feature, MirroredRoots[1]);
+		FeatureVector.SetTransformVelocity(Feature, MirroredRoots[2], MirroredRoots[0], 2.0f * Input.SamplingContext->FiniteDelta);
+		//FeatureVector.SetTransformAcceleration(Feature, MirroredRoots[2], MirroredRoots[1], MirroredRoots[0], FMath::Square(Input.SamplingContext->FiniteDelta));
 	}
 }
 
@@ -2960,7 +3210,7 @@ static void DrawPoseFeatures(const FDebugDrawParams& DrawParams, const FFeatureV
 	Feature.Domain = EPoseSearchFeatureDomain::Time;
 
 	const int32 NumSubsamples = Schema->PoseSampleTimes.Num();
-	const int32 NumBones = Schema->Bones.Num();
+	const int32 NumBones = Schema->SampledBones.Num();
 
 	if ((NumSubsamples * NumBones) == 0)
 	{
@@ -3538,11 +3788,12 @@ bool BuildIndex(const UAnimSequence* Sequence, UPoseSearchSequenceMetaData* Sequ
 		return false;
 	}
 
-	FBoneContainer BoneContainer;
-	BoneContainer.InitializeTo(SequenceMetaData->Schema->BoneIndicesWithParents, FCurveEvaluationOption(false), *SeqSkeleton);
+	FAnimSamplingContext SamplingContext;
+	SamplingContext.Init(SequenceMetaData->Schema);
 
 	FSequenceSampler Sampler;
 	FSequenceSampler::FInput SamplerInput;
+	SamplerInput.SamplingContext = &SamplingContext;
 	SamplerInput.Schema = SequenceMetaData->Schema;
 	SamplerInput.ExtrapolationParameters = SequenceMetaData->ExtrapolationParameters;
 	SamplerInput.Sequence = Sequence;
@@ -3552,10 +3803,10 @@ bool BuildIndex(const UAnimSequence* Sequence, UPoseSearchSequenceMetaData* Sequ
 
 	FSequenceIndexer Indexer;
 	FSequenceIndexer::FInput IndexerInput;
+	IndexerInput.SamplingContext = &SamplingContext;
 	IndexerInput.MainSequence = &Sampler;
 	IndexerInput.Schema = SequenceMetaData->Schema;
 	IndexerInput.RequestedSamplingRange = SequenceMetaData->SamplingRange;
-	IndexerInput.BoneContainer = &BoneContainer;
 	Indexer.Init(IndexerInput);
 	Indexer.Process();
 
@@ -3570,7 +3821,6 @@ bool BuildIndex(const UAnimSequence* Sequence, UPoseSearchSequenceMetaData* Sequ
 	SearchIndexAsset.NumPoses = Indexer.Output.NumIndexedPoses;
 	SequenceMetaData->SearchIndex.Assets.Add(SearchIndexAsset);
 
-
 	PreprocessSearchIndex(&SequenceMetaData->SearchIndex);
 
 	return true;
@@ -3590,8 +3840,8 @@ bool BuildIndex(UPoseSearchDatabase* Database)
 		return false;
 	}
 
-	FBoneContainer BoneContainer;
-	BoneContainer.InitializeTo(Database->Schema->BoneIndicesWithParents, FCurveEvaluationOption(false), *Database->Schema->Skeleton);
+	FAnimSamplingContext SamplingContext;
+	SamplingContext.Init(Database->Schema);
 
 	// Prepare animation preprocessing tasks
 	TArray<FSequenceSampler> SequenceSamplers;
@@ -3605,6 +3855,7 @@ bool BuildIndex(UPoseSearchDatabase* Database)
 			SequenceSamplerMap.Add(Sequence, SequenceSamplerIdx);
 
 			FSequenceSampler::FInput Input;
+			Input.SamplingContext = &SamplingContext;
 			Input.Schema = Database->Schema;
 			Input.ExtrapolationParameters = Database->ExtrapolationParameters;
 			Input.Sequence = Sequence;
@@ -3650,7 +3901,7 @@ bool BuildIndex(UPoseSearchDatabase* Database)
 		const float SequenceLength = DbSequence.Sequence->GetPlayLength();
 
 		FSequenceIndexer::FInput Input;
-		Input.BoneContainer = &BoneContainer;
+		Input.SamplingContext = &SamplingContext;
 		Input.MainSequence = GetSampler(DbSequence.Sequence);
 		Input.LeadInSequence = 
 			SearchIndexAsset.SamplingInterval.Min == 0.0f ? GetSampler(DbSequence.LeadInSequence) : nullptr;
@@ -3659,6 +3910,7 @@ bool BuildIndex(UPoseSearchDatabase* Database)
 		Input.Schema = Database->Schema;
 		Input.BlockTransitionParameters = Database->BlockTransitionParameters;
 		Input.RequestedSamplingRange = SearchIndexAsset.SamplingInterval;
+		Input.bMirrored = SearchIndexAsset.bMirrored;
 
 		FSequenceIndexer& Indexer = Indexers.AddDefaulted_GetRef();
 		Indexer.Init(Input);
@@ -3702,15 +3954,23 @@ bool BuildIndex(UPoseSearchDatabase* Database)
 	return true;
 }
 
-static FSearchResult Search(const FPoseSearchIndex& SearchIndex, TArrayView<const float> Query, const FPoseSearchWeightsContext* WeightsContext = nullptr)
+
+FSearchResult Search(FSearchContext& SearchContext)
 {
 	FSearchResult Result;
-	if (!ensure(SearchIndex.IsValid()))
+
+	const FPoseSearchIndex* SearchIndex = SearchContext.GetSearchIndex();
+	if (!SearchIndex)
 	{
 		return Result;
 	}
 
-	if (!ensure(Query.Num() == SearchIndex.Schema->Layout.NumFloats))
+	if (!ensure(SearchIndex->IsValid()))
+	{
+		return Result;
+	}
+
+	if (!ensure(SearchContext.QueryValues.Num() == SearchIndex->Schema->Layout.NumFloats))
 	{
 		return Result;
 	}
@@ -3718,21 +3978,20 @@ static FSearchResult Search(const FPoseSearchIndex& SearchIndex, TArrayView<cons
 	float BestPoseDissimilarity = MAX_flt;
 	int32 BestPoseIdx = INDEX_NONE;
 
-	for (int32 PoseIdx = 0; PoseIdx != SearchIndex.NumPoses; ++PoseIdx)
+	for (int32 PoseIdx = 0; PoseIdx != SearchIndex->NumPoses; ++PoseIdx)
 	{
-		const FPoseSearchPoseMetadata& Metadata = SearchIndex.PoseMetadata[PoseIdx];
+		const FPoseSearchPoseMetadata& Metadata = SearchIndex->PoseMetadata[PoseIdx];
 
 		if (EnumHasAnyFlags(Metadata.Flags, EPoseSearchPoseFlags::BlockTransition))
 		{
 			continue;
 		}
 
-		const float PoseDissimilarity = ComparePoses(SearchIndex, PoseIdx, Query, WeightsContext);
-		const float ModifiedPoseDissimilarity = PoseDissimilarity + Metadata.CostAddend;
-		
-		if (ModifiedPoseDissimilarity < BestPoseDissimilarity)
+		float PoseDissimilarity = ComparePoses(PoseIdx, SearchContext);
+
+		if (PoseDissimilarity < BestPoseDissimilarity)
 		{
-			BestPoseDissimilarity = ModifiedPoseDissimilarity;
+			BestPoseDissimilarity = PoseDissimilarity;
 			BestPoseIdx = PoseIdx;
 		}
 	}
@@ -3741,100 +4000,80 @@ static FSearchResult Search(const FPoseSearchIndex& SearchIndex, TArrayView<cons
 
 	Result.Dissimilarity = BestPoseDissimilarity;
 	Result.PoseIdx = BestPoseIdx;
-	Result.SearchIndexAsset = SearchIndex.FindAssetForPose(BestPoseIdx);
-	Result.TimeOffsetSeconds = 
-		Result.SearchIndexAsset->SamplingInterval.Min + SearchIndex.Schema->SamplingInterval * (Result.PoseIdx - Result.SearchIndexAsset->FirstPoseIdx);
-
-	return Result;
-}
-
-FSearchResult Search(const UAnimSequenceBase* Sequence, TArrayView<const float> Query, FDebugDrawParams DebugDrawParams)
-{
-	const UPoseSearchSequenceMetaData* MetaData = Sequence ? Sequence->FindMetaDataByClass<UPoseSearchSequenceMetaData>() : nullptr;
-	if (!MetaData || !MetaData->IsValidForSearch())
-	{
-		return FSearchResult();
-	}
+	Result.SearchIndexAsset = SearchIndex->FindAssetForPose(BestPoseIdx);
+	Result.TimeOffsetSeconds = SearchIndex->GetTimeOffset(BestPoseIdx, Result.SearchIndexAsset);
 
-	const FPoseSearchIndex& SearchIndex = MetaData->SearchIndex;
-
-	FSearchResult Result = Search(SearchIndex, Query);
-	if (!Result.IsValid())
-	{
-		return Result;
-	}
-
-	// Do debug visualization
-	DebugDrawParams.SequenceMetaData = MetaData;
-	DebugDrawParams.PoseVector = Query;
-	DebugDrawParams.PoseIdx = Result.PoseIdx;
-	Draw(DebugDrawParams);
+	SearchContext.DebugDrawParams.PoseVector = SearchContext.QueryValues;
+	SearchContext.DebugDrawParams.PoseIdx = Result.PoseIdx;
+	Draw(SearchContext.DebugDrawParams);
 
 	return Result;
 }
 
-FSearchResult Search(
-	const UPoseSearchDatabase* Database, 
-	TArrayView<const float> Query, 
-	const FPoseSearchWeightsContext* WeightsContext, 
-	FDebugDrawParams DebugDrawParams)
+void ComputePoseCostAddends(
+	int32 PoseIdx, 
+	FSearchContext& SearchContext, 
+	float& OutNotifyAddend, 
+	float& OutMirrorMismatchAddend)
 {
-	if (!ensure(Database && Database->IsValidForSearch()))
-	{
-		return FSearchResult();
-	}
+	OutNotifyAddend = 0.0f;
+	OutMirrorMismatchAddend = 0.0f;
 
-	const FPoseSearchIndex& SearchIndex = Database->SearchIndex;
-
-	FSearchResult Result = Search(SearchIndex, Query, WeightsContext);
-	if (!Result.IsValid())
+	if (SearchContext.QueryMirrorRequest != EPoseSearchBooleanRequest::Indifferent)
 	{
-		return FSearchResult();
+		const FPoseSearchIndexAsset* IndexAsset = SearchContext.GetSearchIndex()->FindAssetForPose(PoseIdx);
+		const bool bMirroringMismatch =
+			(IndexAsset->bMirrored && SearchContext.QueryMirrorRequest == EPoseSearchBooleanRequest::FalseValue) ||
+			(!IndexAsset->bMirrored && SearchContext.QueryMirrorRequest == EPoseSearchBooleanRequest::TrueValue);
+		if (bMirroringMismatch)
+		{
+			OutMirrorMismatchAddend = SearchContext.GetMirrorMismatchCost();
+		}
 	}
 
-	if (!ensure(Result.TimeOffsetSeconds >= 0.0f))
-	{
-		return FSearchResult();
-	}
-	
-	// Do debug visualization
-	DebugDrawParams.Database = Database;
-	DebugDrawParams.PoseVector = Query;
-	DebugDrawParams.PoseIdx = Result.PoseIdx;
-	Draw(DebugDrawParams);
-
-	return Result;
+	const FPoseSearchPoseMetadata& PoseMetadata = SearchContext.GetSearchIndex()->PoseMetadata[PoseIdx];
+	OutNotifyAddend = PoseMetadata.CostAddend;
 }
 
-float ComparePoses(const FPoseSearchIndex& SearchIndex, int32 PoseIdx, TArrayView<const float> Query, const FPoseSearchWeightsContext* WeightsContext)
+float ComparePoses(int32 PoseIdx, FSearchContext& SearchContext)
 {
-	TArrayView<const float> PoseValues = SearchIndex.GetPoseValues(PoseIdx);
-	if (!ensure(PoseValues.Num() == Query.Num()))
+	TArrayView<const float> PoseValues = SearchContext.GetSearchIndex()->GetPoseValues(PoseIdx);
+	if (!ensure(PoseValues.Num() == SearchContext.QueryValues.Num()))
 	{
 		return MAX_flt;
 	}
 
 	float Dissimilarity;
-	if (WeightsContext)
+	if (SearchContext.WeightsContext)
 	{
-		const FPoseSearchWeights* WeightsSet = WeightsContext->GetGroupWeights(0);
-		Dissimilarity = CompareFeatureVectors(PoseValues.Num(), PoseValues.GetData(), Query.GetData(), WeightsSet->Weights.GetData());
+		const FPoseSearchWeights* WeightsSet = SearchContext.WeightsContext->GetGroupWeights(0);
+		Dissimilarity = CompareFeatureVectors(
+			PoseValues.Num(), 
+			PoseValues.GetData(), 
+			SearchContext.QueryValues.GetData(),
+			WeightsSet->Weights.GetData());
 	}
 	else
 	{
-		Dissimilarity = CompareFeatureVectors(PoseValues.Num(), PoseValues.GetData(), Query.GetData());
+		Dissimilarity = 
+			CompareFeatureVectors(PoseValues.Num(), PoseValues.GetData(), SearchContext.QueryValues.GetData());
 	}
 
+	float NotifyAddend = 0.0f;
+	float MirrorMismatchAddend = 0.0f;
+	ComputePoseCostAddends(PoseIdx, SearchContext, NotifyAddend, MirrorMismatchAddend);
+	Dissimilarity += (NotifyAddend + MirrorMismatchAddend);
+	
 	return Dissimilarity;
 }
 
-float ComparePoses(const FPoseSearchIndex& SearchIndex, int32 PoseIdx, TArrayView<const float> Query, const FPoseSearchWeightsContext* WeightsContext, FPoseCostInfo& OutPoseCostInfo)
+float ComparePoses(int32 PoseIdx, FSearchContext& SearchContext, FPoseCostInfo& OutPoseCostInfo)
 {
 	using namespace Eigen;
 
-	TArrayView<const float> PoseValues = SearchIndex.GetPoseValues(PoseIdx);
+	TArrayView<const float> PoseValues = SearchContext.GetSearchIndex()->GetPoseValues(PoseIdx);
 	const int32 Dims = PoseValues.Num();
-	if (!ensure(PoseValues.Num() == Dims))
+	if (!ensure(Dims == SearchContext.QueryValues.Num()))
 	{
 		return MAX_flt;
 	}
@@ -3844,13 +4083,13 @@ float ComparePoses(const FPoseSearchIndex& SearchIndex, int32 PoseIdx, TArrayVie
 	// Setup Eigen views onto our vectors
 	auto OutCostVector = Map<ArrayXf>(OutPoseCostInfo.CostVector.GetData(), Dims);
 	auto PoseVector = Map<const ArrayXf>(PoseValues.GetData(), Dims);
-	auto QueryVector = Map<const ArrayXf>(Query.GetData(), Dims);
+	auto QueryVector = Map<const ArrayXf>(SearchContext.QueryValues.GetData(), Dims);
 	
 	// Compute weighted squared difference vector
 	float Dissimilarity;
-	if (WeightsContext)
+	if (SearchContext.WeightsContext)
 	{
-		const FPoseSearchWeights* WeightsSet = WeightsContext->GetGroupWeights(0);
+		const FPoseSearchWeights* WeightsSet = SearchContext.WeightsContext->GetGroupWeights(0);
 		check(WeightsSet);
 		check(WeightsSet->Weights.Num() == Dims);
 		auto WeightsVector = Map<const ArrayXf>(WeightsSet->Weights.GetData(), Dims);
@@ -3864,16 +4103,27 @@ float ComparePoses(const FPoseSearchIndex& SearchIndex, int32 PoseIdx, TArrayVie
 		Dissimilarity = OutCostVector.sum();
 	}
 
+	ComputePoseCostAddends(
+		PoseIdx, 
+		SearchContext, 
+		OutPoseCostInfo.NotifyCostAddend, 
+		OutPoseCostInfo.MirrorMismatchCostAddend);
+
+	const float TotalDissimilarity = 
+		Dissimilarity + 
+		OutPoseCostInfo.NotifyCostAddend + 
+		OutPoseCostInfo.MirrorMismatchCostAddend;
+
 	// Verify this math agrees with the runtime pose comparator
-	checkSlow(FMath::IsNearlyEqual(Dissimilarity, ComparePoses(SearchIndex, PoseIdx, Query, WeightsContext), KINDA_SMALL_NUMBER));
+	checkSlow(FMath::IsNearlyEqual(
+		TotalDissimilarity,
+		ComparePoses(PoseIdx, SearchContext), 
+		KINDA_SMALL_NUMBER));
 
-	const FPoseSearchPoseMetadata& PoseMetadata = SearchIndex.PoseMetadata[PoseIdx];
-	float TotalDissimilarity = Dissimilarity + PoseMetadata.CostAddend;
 	
 	// Output cost details
 	OutPoseCostInfo.TotalCost = TotalDissimilarity;
-	OutPoseCostInfo.CostAddend = PoseMetadata.CostAddend;
-	CalcChannelCosts(SearchIndex.Schema, OutPoseCostInfo.CostVector, OutPoseCostInfo.ChannelCosts);
+	CalcChannelCosts(SearchContext.GetSearchIndex()->Schema, OutPoseCostInfo.CostVector, OutPoseCostInfo.ChannelCosts);
 
 	// Verify channel cost decomposition agrees with pose comparator
 	auto OutChannelCosts = Map<const ArrayXf>(OutPoseCostInfo.ChannelCosts.GetData(), OutPoseCostInfo.ChannelCosts.Num());
@@ -3938,7 +4188,10 @@ UE::Anim::IPoseSearchProvider::FSearchResult FModule::Search(const FAnimationBas
 
 	QueryBuilder.Normalize(MetaData->SearchIndex);
 
-	::UE::PoseSearch::FSearchResult Result = ::UE::PoseSearch::Search(Sequence, QueryBuilder.GetNormalizedValues());
+	::UE::PoseSearch::FSearchContext SearchContext;
+	SearchContext.SetSource(Sequence);
+	SearchContext.QueryValues = QueryBuilder.GetNormalizedValues();
+	::UE::PoseSearch::FSearchResult Result = ::UE::PoseSearch::Search(SearchContext);
 
 	ProviderResult.Dissimilarity = Result.Dissimilarity;
 	ProviderResult.PoseIdx = Result.PoseIdx;
diff --git a/Engine/Plugins/Experimental/Animation/PoseSearch/Source/Runtime/Private/PoseSearchLibrary.cpp b/Engine/Plugins/Experimental/Animation/PoseSearch/Source/Runtime/Private/PoseSearchLibrary.cpp
--- a/Engine/Plugins/Experimental/Animation/PoseSearch/Source/Runtime/Private/PoseSearchLibrary.cpp
+++ b/Engine/Plugins/Experimental/Animation/PoseSearch/Source/Runtime/Private/PoseSearchLibrary.cpp
@@ -111,6 +111,18 @@ void UpdateMotionMatchingState(const FAnimationUpdateContext& Context
 		// Update features in the query with the latest inputs
 		InOutMotionMatchingState.ComposeQuery(Database, Trajectory);
 
+		UE::PoseSearch::FSearchContext SearchContext;
+		SearchContext.SetSource(InOutMotionMatchingState.CurrentDatabase.Get());
+		SearchContext.QueryValues = InOutMotionMatchingState.ComposedQuery.GetNormalizedValues();
+		SearchContext.WeightsContext = &InOutMotionMatchingState.WeightsContext;
+		if (const FPoseSearchIndexAsset* CurrentIndexAsset = InOutMotionMatchingState.GetCurrentSearchIndexAsset())
+		{
+			SearchContext.QueryMirrorRequest =
+				CurrentIndexAsset->bMirrored ?
+				EPoseSearchBooleanRequest::TrueValue :
+				EPoseSearchBooleanRequest::FalseValue;
+		}
+
 		// Update weight groups
 		InOutMotionMatchingState.WeightsContext.Update(Settings.Weights, Database);
 
@@ -118,26 +130,18 @@ void UpdateMotionMatchingState(const FAnimationUpdateContext& Context
 		float CurrentDissimilarity = MAX_flt;
 		if (InOutMotionMatchingState.DbPoseIdx != INDEX_NONE)
 		{
-			CurrentDissimilarity = UE::PoseSearch::ComparePoses(
-				Database->SearchIndex,
-				InOutMotionMatchingState.DbPoseIdx,
-				InOutMotionMatchingState.ComposedQuery.GetNormalizedValues(),
-				&InOutMotionMatchingState.WeightsContext
-			);
+			CurrentDissimilarity = UE::PoseSearch::ComparePoses(InOutMotionMatchingState.DbPoseIdx, SearchContext);
 		}
 
 		// Search the database for the nearest match to the updated query vector
-		UE::PoseSearch::FSearchResult Result = UE::PoseSearch::Search(
-			Database,
-			InOutMotionMatchingState.ComposedQuery.GetNormalizedValues(),
-			&InOutMotionMatchingState.WeightsContext
-		);
+		UE::PoseSearch::FSearchResult Result = UE::PoseSearch::Search(SearchContext);
 		if (Result.IsValid() && ((InOutMotionMatchingState.ElapsedPoseJumpTime >= Settings.SearchThrottleTime) || !bCanContinue))
 		{
 			if (!bCanContinue)
 			{
+				const float JumpBlendTime = InOutMotionMatchingState.ComputeJumpBlendTime(Result, Settings);
 				InOutMotionMatchingState.JumpToPose(Result);
-				RequestInertialBlend(Context, Settings.BlendTime);
+				RequestInertialBlend(Context, JumpBlendTime);
 				InOutMotionMatchingState.Flags |= EMotionMatchingFlags::JumpedToPose;
 			}
 			else
@@ -148,7 +152,7 @@ void UpdateMotionMatchingState(const FAnimationUpdateContext& Context
 				// We'll ignore the candidate pose if it is too near to our current pose
 				bool bNearbyPose = false;
 				const FPoseSearchIndexAsset* StateSearchIndexAsset = 
-					&InOutMotionMatchingState.CurrentDatabase->SearchIndex.Assets[InOutMotionMatchingState.SearchIndexAssetIdx];
+					InOutMotionMatchingState.GetCurrentSearchIndexAsset();
 				if (StateSearchIndexAsset == Result.SearchIndexAsset)
 				{
 					const FPoseSearchDatabaseSequence& ResultDbSequence = Database->GetSourceAsset(Result.SearchIndexAsset);
@@ -164,8 +168,9 @@ void UpdateMotionMatchingState(const FAnimationUpdateContext& Context
 				// Start playback from the candidate pose if we determined it was a better option
 				if (bBetterPose && !bNearbyPose)
 				{
+					const float JumpBlendTime = InOutMotionMatchingState.ComputeJumpBlendTime(Result, Settings);
 					InOutMotionMatchingState.JumpToPose(Result);
-					RequestInertialBlend(Context, Settings.BlendTime);
+					RequestInertialBlend(Context, JumpBlendTime);
 					InOutMotionMatchingState.Flags |= EMotionMatchingFlags::JumpedToPose;
 				}
 			}
@@ -256,7 +261,7 @@ UE::PoseSearch::FMotionMatchingContinuityParams FMotionMatchingState::ComputeCon
 	{
 		check(SearchIndexAssetIdx != INDEX_NONE);
 
-		const FPoseSearchIndexAsset* SearchIndexAsset = &CurrentDatabase->SearchIndex.Assets[SearchIndexAssetIdx];
+		const FPoseSearchIndexAsset* SearchIndexAsset = GetCurrentSearchIndexAsset();
 		const FPoseSearchDatabaseSequence& DbSequence = CurrentDatabase->GetSourceAsset(SearchIndexAsset);
 		const float AssetLength = DbSequence.Sequence->GetPlayLength();
 		const float DeltaTime = Context.GetDeltaTime();
@@ -291,7 +296,11 @@ UE::PoseSearch::FMotionMatchingContinuityParams FMotionMatchingState::ComputeCon
 			int32 FollowUpSearchIndexAssetIdx = CurrentDatabase->SearchIndex.Assets.IndexOfByPredicate(
 				[&](const FPoseSearchIndexAsset& Entry)
 			{
-				return Entry.SourceAssetIdx == FollowUpDbSequenceIdx && Entry.SamplingInterval.Contains(0.0f);
+				const bool bIsMatch = 
+					Entry.SourceAssetIdx == FollowUpDbSequenceIdx && 
+					Entry.bMirrored == SearchIndexAsset->bMirrored &&
+					Entry.SamplingInterval.Contains(0.0f);
+				return bIsMatch;
 			});
 
 			if (FollowUpSearchIndexAssetIdx != INDEX_NONE)
@@ -316,3 +325,31 @@ UE::PoseSearch::FMotionMatchingContinuityParams FMotionMatchingState::ComputeCon
 	}
 	return ContinuityParameters;
 }
+
+const FPoseSearchIndexAsset* FMotionMatchingState::GetCurrentSearchIndexAsset() const
+{
+	if (!CurrentDatabase->SearchIndex.Assets.IsValidIndex(SearchIndexAssetIdx))
+	{
+		return nullptr;
+	}
+
+	return &CurrentDatabase->SearchIndex.Assets[SearchIndexAssetIdx];
+}
+
+float FMotionMatchingState::ComputeJumpBlendTime(
+	const UE::PoseSearch::FSearchResult& Result, 
+	const FMotionMatchingSettings& Settings) const
+{
+	const FPoseSearchIndexAsset* SearchIndexAsset = GetCurrentSearchIndexAsset();
+	// Use alternate blend time when changing between mirrored and unmirrored
+	float JumpBlendTime = Settings.BlendTime;
+	if ((SearchIndexAsset != nullptr) && (Settings.MirrorChangeBlendTime > 0.0f))
+	{
+		if (Result.SearchIndexAsset->bMirrored != SearchIndexAsset->bMirrored)
+		{
+			JumpBlendTime = Settings.MirrorChangeBlendTime;
+		}
+	}
+
+	return JumpBlendTime;
+}
diff --git a/Engine/Plugins/Experimental/Animation/PoseSearch/Source/Runtime/Public/PoseSearch/AnimNode_MotionMatching.h b/Engine/Plugins/Experimental/Animation/PoseSearch/Source/Runtime/Public/PoseSearch/AnimNode_MotionMatching.h
--- a/Engine/Plugins/Experimental/Animation/PoseSearch/Source/Runtime/Public/PoseSearch/AnimNode_MotionMatching.h
+++ b/Engine/Plugins/Experimental/Animation/PoseSearch/Source/Runtime/Public/PoseSearch/AnimNode_MotionMatching.h
@@ -4,6 +4,7 @@
 
 #include "Animation/AnimNode_AssetPlayerBase.h"
 #include "Animation/AnimNode_SequencePlayer.h"
+#include "AnimNodes/AnimNode_Mirror.h"
 #include "Animation/MotionTrajectoryTypes.h"
 #include "DynamicPlayRate/DynamicPlayRateLibrary.h"
 #include "PoseSearch/PoseSearch.h"
@@ -60,6 +61,9 @@ private:
 	// Embedded sequence player node for playing animations from the motion matching database
 	FAnimNode_SequencePlayer_Standalone SequencePlayerNode;
 
+	// Embedded mirror node to handle mirroring if the pose search results in a mirrored sequence
+	FAnimNode_Mirror MirrorNode;
+
 	// Encapsulated motion matching algorithm and internal state
 	FMotionMatchingState MotionMatchingState;
 	
diff --git a/Engine/Plugins/Experimental/Animation/PoseSearch/Source/Runtime/Public/PoseSearch/PoseSearch.h b/Engine/Plugins/Experimental/Animation/PoseSearch/Source/Runtime/Public/PoseSearch/PoseSearch.h
--- a/Engine/Plugins/Experimental/Animation/PoseSearch/Source/Runtime/Public/PoseSearch/PoseSearch.h
+++ b/Engine/Plugins/Experimental/Animation/PoseSearch/Source/Runtime/Public/PoseSearch/PoseSearch.h
@@ -55,6 +55,18 @@ enum class EPoseSearchFeatureDomain : int32
 	Invalid = Num UMETA(Hidden)
 };
 
+UENUM()
+enum class EPoseSearchBooleanRequest : int32
+{
+	FalseValue,
+	TrueValue,
+	Indifferent, // if this is used, there will be no cost difference between true and false results
+
+	Num UMETA(Hidden),
+	Invalid = Num UMETA(Hidden)
+};
+
+
 /** Describes each feature of a vector, including data type, sampling options, and buffer offset. */
 USTRUCT()
 struct POSESEARCH_API FPoseSearchFeatureDesc
@@ -127,6 +139,28 @@ enum class EPoseSearchDataPreprocessor : int32
 	Invalid = Num UMETA(Hidden)
 };
 
+USTRUCT()
+struct POSESEARCH_API FPoseSearchBone
+{
+	GENERATED_BODY()
+
+	UPROPERTY(EditAnywhere, Category = Config)
+	FBoneReference Reference;
+	
+	UPROPERTY(EditAnywhere, Category = Config)
+	bool bUseVelocity = false;
+	
+	UPROPERTY(EditAnywhere, Category = Config)
+	bool bUsePosition = false;
+	
+	UPROPERTY(EditAnywhere, Category = Config)
+	bool bUseRotation = false;
+
+	// this function will return a mask out of EPoseSearchFeatureType based on which features were selected for
+	// the bone.
+	uint32 GetTypeMask() const;
+};
+
 /**
 * Specifies the format of a pose search index. At runtime, queries are built according to the schema for searching.
 */
@@ -145,11 +179,11 @@ public:
 	UPROPERTY(EditAnywhere, meta = (ClampMin = "1", ClampMax = "60"), Category = "Schema")
 	int32 SampleRate = DefaultSampleRate;
 
-	UPROPERTY(EditAnywhere, Category = "Schema")
-	bool bUseBoneVelocities = true;
+	UPROPERTY()
+	bool bUseBoneVelocities_DEPRECATED = true;
 
-	UPROPERTY(EditAnywhere, Category = "Schema")
-	bool bUseBonePositions = true;
+	UPROPERTY()
+	bool bUseBonePositions_DEPRECATED = true;
 
 	UPROPERTY(EditAnywhere, Category = "Schema")
 	bool bUseTrajectoryVelocities = true;
@@ -159,9 +193,12 @@ public:
 
 	UPROPERTY(EditAnywhere, Category = "Schema", DisplayName="Use Facing Directions")
 	bool bUseTrajectoryForwardVectors = false;
-
+	
 	UPROPERTY(EditAnywhere, Category = "Schema")
-	TArray<FBoneReference> Bones;
+	TArray<FPoseSearchBone> SampledBones;
+	
+	UPROPERTY()
+	TArray<FBoneReference> Bones_DEPRECATED;
 
 	UPROPERTY(EditAnywhere, Category = "Schema")
 	TArray<float> PoseSampleTimes;
@@ -172,6 +209,10 @@ public:
 	UPROPERTY(EditAnywhere, Category = "Schema")
 	TArray<float> TrajectorySampleDistances;
 
+	// If set, this schema will support mirroring pose search databases
+	UPROPERTY(EditAnywhere, Category = "Schema")
+	TObjectPtr<UMirrorDataTable> MirrorDataTable;
+
 	UPROPERTY(EditAnywhere, Category = "Schema")
 	EPoseSearchDataPreprocessor DataPreprocessor = EPoseSearchDataPreprocessor::Automatic;
 
@@ -287,8 +328,9 @@ public:
 	FPoseSearchIndexAsset()
 	{}
 
-	FPoseSearchIndexAsset(int32 InSourceAssetIdx, const FFloatInterval& InSamplingInterval)
+	FPoseSearchIndexAsset(int32 InSourceAssetIdx, bool bInMirrored, const FFloatInterval& InSamplingInterval)
 		: SourceAssetIdx(InSourceAssetIdx)
+		, bMirrored(bInMirrored)
 		, SamplingInterval(InSamplingInterval)
 	{}
 
@@ -296,6 +338,9 @@ public:
 	UPROPERTY()
 	int32 SourceAssetIdx = INDEX_NONE;
 
+	UPROPERTY()
+	bool bMirrored = false;
+
 	UPROPERTY()
 	FFloatInterval SamplingInterval;
 
@@ -344,6 +389,7 @@ struct POSESEARCH_API FPoseSearchIndex
 
 	int32 FindAssetIndex(const FPoseSearchIndexAsset* Asset) const;
 	const FPoseSearchIndexAsset* FindAssetForPose(int32 PoseIdx) const;
+	float GetTimeOffset(int32 PoseIdx, const FPoseSearchIndexAsset* Asset) const;
 
 	void Reset();
 
@@ -553,6 +599,16 @@ private:
 	TArray<FPoseSearchWeights> ComputedGroupWeights;
 };
 
+UENUM()
+enum class EPoseSearchMirrorOption : int32
+{
+	UnmirroredOnly,
+	MirroredOnly,
+	UnmirroredAndMirrored,
+	
+	Num UMETA(Hidden),
+	Invalid = Num UMETA(Hidden)
+};
 
 /** An entry in a UPoseSearchDatabase. */
 USTRUCT(BlueprintType, Category = "Animation|Pose Search")
@@ -569,6 +625,9 @@ struct POSESEARCH_API FPoseSearchDatabaseSequence
 	UPROPERTY(EditAnywhere, Category="Sequence")
 	bool bLoopAnimation = false;
 
+	UPROPERTY(EditAnywhere, Category = "Sequence")
+	EPoseSearchMirrorOption MirrorOption = EPoseSearchMirrorOption::UnmirroredOnly;
+
 	// Used for sampling past pose information at the beginning of the main sequence.
 	// This setting is intended for transitions between cycles. It is optional and only used
 	// for one shot anims with past sampling. When past sampling is used without a lead in sequence,
@@ -605,6 +664,11 @@ public:
 	UPROPERTY(EditAnywhere, Category = "Database")
 	FPoseSearchWeightParams Weights;
 
+	// If there's a mirroring mismatch between the currently playing sequence and a search candidate, this cost will be 
+	// added to the candidate, making it less likely to be selected
+	UPROPERTY(EditAnywhere, BlueprintReadWrite, Category="Database")
+	float MirroringMismatchCost = 0.0f;
+
 	UPROPERTY(EditAnywhere, Category = "Database")
 	FPoseSearchExtrapolationParameters ExtrapolationParameters;
 
@@ -850,6 +914,24 @@ struct FSearchResult
 	bool IsValid() const { return PoseIdx >= 0; }
 };
 
+struct POSESEARCH_API FSearchContext
+{
+	TArrayView<const float> QueryValues;
+	EPoseSearchBooleanRequest QueryMirrorRequest = EPoseSearchBooleanRequest::Indifferent;
+	const FPoseSearchWeightsContext* WeightsContext = nullptr;
+	FDebugDrawParams DebugDrawParams;
+
+	void SetSource(const UPoseSearchDatabase* InSourceDatabase);
+	void SetSource(const UAnimSequenceBase* InSourceSequence);
+	const FPoseSearchIndex* GetSearchIndex() const;
+	float GetMirrorMismatchCost() const;
+
+private:
+	const UPoseSearchDatabase* SourceDatabase = nullptr;
+	const UAnimSequenceBase* SourceSequence = nullptr;
+	const FPoseSearchIndex* SearchIndex = nullptr;
+	float MirrorMismatchCost = 0.0f;
+};
 
 /**
 * Visualize pose search debug information
@@ -880,44 +962,25 @@ POSESEARCH_API bool BuildIndex(const UAnimSequence* Sequence, UPoseSearchSequenc
 POSESEARCH_API bool BuildIndex(UPoseSearchDatabase* Database);
 
 
-/**
-* Performs a pose search on a sequence's UPoseSearchSequenceMetaData.
-*
-* @param Sequence		The sequence to search within. Must contain a UPoseSearchSequenceMetaData object.
-* @param Query			The pose query to search for. To build a query, see FFeatureVectorBuilder. Must have been built using the same schema as the UPoseSearchSequenceMetaData.
-* @param WeightsContext	Optional weights context used to influence pose search query results
-* @param DrawParams		Visualization options
-* 
-* @return The pose in the sequence that most closely matches the Query.
-*/
-POSESEARCH_API FSearchResult Search(const UAnimSequenceBase* Sequence, TArrayView<const float> Query, FDebugDrawParams DrawParams = FDebugDrawParams());
-
-
 /**
 * Performs a pose search on a UPoseSearchDatabase.
 *
-* @param Database			The database to search within
-* @param Query				The pose query to search for. To build a query, see FFeatureVectorBuilder. Must have been built using the same schema as the Database.
-* @param WeightsContext		Optional weights context used to influence pose search query results
-* @param EndTimeToExclude	Samples this close to the end of an animation sequence will be ignored by the search. 
-* @param DrawParams			Visualization options
+* @param SearchContext	Structure containing search parameters
 * 
 * @return The pose in the database that most closely matches the Query.
 */
-POSESEARCH_API FSearchResult Search(const UPoseSearchDatabase* Database, TArrayView<const float> Query, const FPoseSearchWeightsContext* WeightsContext = nullptr, FDebugDrawParams DrawParams = FDebugDrawParams());
-
+POSESEARCH_API FSearchResult Search(FSearchContext& SearchContext);
 
 /**
 * Evaluate pose comparison metric between a pose in the search index and an input query
 * 
 * @param SearchIndex	The search index containing the pose to compare to the query
 * @param PoseIdx		The index of the pose in the search index to compare to the query
-* @param Query			The query to compare against. Must have been built using the same schema as the SearchIndex.
-* @param WeightsContext	Optional weights context used to influence pose comparison
+* @param SearchContext	Structure containing search parameters
 * 
 * @return Dissimilarity between the two poses
 */
-POSESEARCH_API float ComparePoses(const FPoseSearchIndex& SearchIndex, int32 PoseIdx, TArrayView<const float> Query, const FPoseSearchWeightsContext* WeightsContext);
+POSESEARCH_API float ComparePoses(int32 PoseIdx, FSearchContext& SearchContext);
 
 
 /**
@@ -929,7 +992,9 @@ struct FPoseCostInfo
 	float TotalCost = 0.0f;
 
 	// Contribution from ModifyCost anim notify
-	float CostAddend = 0.0f;
+	float NotifyCostAddend = 0.0f;
+
+	float MirrorMismatchCostAddend = 0.0f;
 
 	// Cost breakdown per channel (e.g. pose cost, time-based trajectory cost, distance-based trajectory cost, etc.)
 	TArray<float> ChannelCosts;
@@ -944,13 +1009,10 @@ struct FPoseCostInfo
 *
 * @param SearchIndex		The search index containing the pose to compare to the query
 * @param PoseIdx			The index of the pose in the search index to compare to the query
-* @param Query				The query to compare against. Must have been built using the same schema as the SearchIndex.
-* @param WeightsContext		Optional weights context used to influence pose comparison
+* @param SearchContext		Structure containing search parameters
 * &param OutPoseCostInfo	Cost details for analysis in the debugger
 *
 * @return Dissimilarity between the two poses
 */
-
-POSESEARCH_API float ComparePoses(const FPoseSearchIndex& SearchIndex, int32 PoseIdx, TArrayView<const float> Query, const FPoseSearchWeightsContext* WeightsContext, FPoseCostInfo& OutPoseCostInfo);
-
+POSESEARCH_API float ComparePoses(int32 PoseIdx, FSearchContext& SearchContext, FPoseCostInfo& OutPoseCostInfo);
 }} // namespace UE::PoseSearch
\ No newline at end of file
diff --git a/Engine/Plugins/Experimental/Animation/PoseSearch/Source/Runtime/Public/PoseSearch/PoseSearchLibrary.h b/Engine/Plugins/Experimental/Animation/PoseSearch/Source/Runtime/Public/PoseSearch/PoseSearchLibrary.h
--- a/Engine/Plugins/Experimental/Animation/PoseSearch/Source/Runtime/Public/PoseSearch/PoseSearchLibrary.h
+++ b/Engine/Plugins/Experimental/Animation/PoseSearch/Source/Runtime/Public/PoseSearch/PoseSearchLibrary.h
@@ -53,6 +53,10 @@ struct POSESEARCH_API FMotionMatchingSettings
 	// Time in seconds to blend out to the new pose. Uses inertial blending and requires an Inertialization node after this node.
 	UPROPERTY(EditAnywhere, BlueprintReadWrite, Category=Settings, meta=(ClampMin="0"))
 	float BlendTime = 0.2f;
+
+	// If the pose jump requires a mirroring change and this value is greater than 0, it will be used instead of BlendTime
+	UPROPERTY(EditAnywhere, BlueprintReadWrite, Category = Settings, meta = (ClampMin = "0", DislayAfter = "BlendTime"))
+	float MirrorChangeBlendTime = 0.0f;
 	
 	// Don't jump to poses that are less than this many seconds away
 	UPROPERTY(EditAnywhere, BlueprintReadWrite, Category=Settings, meta=(ClampMin="0"))
@@ -85,6 +89,10 @@ struct POSESEARCH_API FMotionMatchingState
 	// Updates DbPoseIdx to track DeltaTime and jumps to a follow-up sequence when available
 	UE::PoseSearch::FMotionMatchingContinuityParams ComputeContinuityParameters(const FAnimationUpdateContext& Context) const;
 
+	const FPoseSearchIndexAsset* GetCurrentSearchIndexAsset() const;
+
+	float ComputeJumpBlendTime(const UE::PoseSearch::FSearchResult& Result, const FMotionMatchingSettings& Settings) const;
+
 	// The current pose we're playing from the database
 	UPROPERTY(Transient)
 	int32 DbPoseIdx = INDEX_NONE;
diff --git a/Engine/Source/Runtime/AnimGraphRuntime/Private/AnimNodes/AnimNode_Mirror.cpp b/Engine/Source/Runtime/AnimGraphRuntime/Private/AnimNodes/AnimNode_Mirror.cpp
--- a/Engine/Source/Runtime/AnimGraphRuntime/Private/AnimNodes/AnimNode_Mirror.cpp
+++ b/Engine/Source/Runtime/AnimGraphRuntime/Private/AnimNodes/AnimNode_Mirror.cpp
@@ -115,18 +115,10 @@ void FAnimNode_MirrorBase::FillCompactPoseAndComponentRefRotations(const FBoneCo
 	UMirrorDataTable* MirrorDataTable = GetMirrorDataTable();
 	if (MirrorDataTable)
 	{
-		TCustomBoneIndexArray<FSkeletonPoseBoneIndex, FSkeletonPoseBoneIndex> MirrorBoneIndexes;
-		MirrorDataTable->FillMirrorBoneIndexes(BoneContainer.GetSkeletonAsset(), MirrorBoneIndexes);
-		MirrorDataTable->FillCompactPoseMirrorBones(BoneContainer, MirrorBoneIndexes, CompactPoseMirrorBones);
-
-		const int32 NumBones = BoneContainer.GetCompactPoseNumBones();
-		ComponentSpaceRefRotations.SetNumUninitialized(NumBones);
-		ComponentSpaceRefRotations[FCompactPoseBoneIndex(0)] = BoneContainer.GetRefPoseTransform(FCompactPoseBoneIndex(0)).GetRotation();
-		for (FCompactPoseBoneIndex BoneIndex(1); BoneIndex < NumBones; ++BoneIndex)
-		{
-			const FCompactPoseBoneIndex ParentBoneIndex = BoneContainer.GetParentBoneIndex(BoneIndex);
-			ComponentSpaceRefRotations[BoneIndex] = ComponentSpaceRefRotations[ParentBoneIndex] * BoneContainer.GetRefPoseTransform(BoneIndex).GetRotation();
-		}
+		MirrorDataTable->FillCompactPoseAndComponentRefRotations(
+			BoneContainer, 
+			CompactPoseMirrorBones, 
+			ComponentSpaceRefRotations);
 	}
 	else
 	{
diff --git a/Engine/Source/Runtime/Engine/Classes/Animation/MirrorDataTable.h b/Engine/Source/Runtime/Engine/Classes/Animation/MirrorDataTable.h
--- a/Engine/Source/Runtime/Engine/Classes/Animation/MirrorDataTable.h
+++ b/Engine/Source/Runtime/Engine/Classes/Animation/MirrorDataTable.h
@@ -153,6 +153,18 @@ public:
 	 */
 	ENGINE_API void FillMirrorBoneIndexes(const USkeleton* Skeleton, TCustomBoneIndexArray<FSkeletonPoseBoneIndex, FSkeletonPoseBoneIndex>& OutMirrorBoneIndexes) const;
 
+	/**
+	 * Populates two arrays with a mapping of compact pose mirror bones and reference rotations
+	 *
+	 * @param	BoneContainer					Structure which holds the required bones
+	 * @param	OutCompactPoseMirrorBones		Output array mapping compact pose mirror bones
+	 * @param	OutComponentSpaceRefRotations	Output array mapping reference rotations
+	 */
+	ENGINE_API void FillCompactPoseAndComponentRefRotations(
+		const FBoneContainer& BoneContainer,
+		TCustomBoneIndexArray<FCompactPoseBoneIndex, FCompactPoseBoneIndex>& OutCompactPoseMirrorBones,
+		TCustomBoneIndexArray<FQuat, FCompactPoseBoneIndex>& OutComponentSpaceRefRotations) const;
+
 #if WITH_EDITOR  
 	/**
 	 * Populates the table by running the MirrorFindReplaceExpressions on bone names in the Skeleton.  If the mirrored name is also found 
diff --git a/Engine/Source/Runtime/Engine/Private/Animation/MirrorDataTable.cpp b/Engine/Source/Runtime/Engine/Private/Animation/MirrorDataTable.cpp
--- a/Engine/Source/Runtime/Engine/Private/Animation/MirrorDataTable.cpp
+++ b/Engine/Source/Runtime/Engine/Private/Animation/MirrorDataTable.cpp
@@ -352,6 +352,28 @@ void UMirrorDataTable::FillMirrorBoneIndexes(const USkeleton* InSkeleton, TCusto
 	}
 }
 
+void UMirrorDataTable::FillCompactPoseAndComponentRefRotations(
+	const FBoneContainer& BoneContainer,
+	TCustomBoneIndexArray<FCompactPoseBoneIndex, FCompactPoseBoneIndex>& OutCompactPoseMirrorBones,
+	TCustomBoneIndexArray<FQuat, FCompactPoseBoneIndex>& OutComponentSpaceRefRotations) const
+{
+	TCustomBoneIndexArray<FSkeletonPoseBoneIndex, FSkeletonPoseBoneIndex> MirrorBoneIndexes;
+	FillMirrorBoneIndexes(BoneContainer.GetSkeletonAsset(), MirrorBoneIndexes);
+	FillCompactPoseMirrorBones(BoneContainer, MirrorBoneIndexes, OutCompactPoseMirrorBones);
+
+	const int32 NumBones = BoneContainer.GetCompactPoseNumBones();
+	OutComponentSpaceRefRotations.SetNumUninitialized(NumBones);
+	OutComponentSpaceRefRotations[FCompactPoseBoneIndex(0)] = 
+		BoneContainer.GetRefPoseTransform(FCompactPoseBoneIndex(0)).GetRotation();
+	for (FCompactPoseBoneIndex BoneIndex(1); BoneIndex < NumBones; ++BoneIndex)
+	{
+		const FCompactPoseBoneIndex ParentBoneIndex = BoneContainer.GetParentBoneIndex(BoneIndex);
+		OutComponentSpaceRefRotations[BoneIndex] = 
+			OutComponentSpaceRefRotations[ParentBoneIndex] * BoneContainer.GetRefPoseTransform(BoneIndex).GetRotation();
+	}
+}
+
+
 void UMirrorDataTable::FillMirrorArrays()
 {
 	SyncToMirrorSyncMap.Empty();
