PoseSearch: Pose cost breakdown in debugger
- Refactored debugger display of feature vectors for reusability
- Added debugger display of cost vector
- Refactored debugger feature extraction via channels
- Refactored debugger tick to simplify update flow and reduce potentially duplicated work
- Continued realization of data channels concept
  - Added ChannelIndex to feature descriptions
  - Cleaned up FPoseSearchFeatureVectorLayout::EnumerateBy
  - Added FPoseSearchFeatureVectorLayout::NumLayouts
- Changed FeatureDesc indices to default to zero in anticipation of reducing field sizes in the future

diff --git a/Engine/Plugins/Experimental/Animation/PoseSearch/Source/Editor/Private/PoseSearchDebugger.cpp b/Engine/Plugins/Experimental/Animation/PoseSearch/Source/Editor/Private/PoseSearchDebugger.cpp
--- a/Engine/Plugins/Experimental/Animation/PoseSearch/Source/Editor/Private/PoseSearchDebugger.cpp
+++ b/Engine/Plugins/Experimental/Animation/PoseSearch/Source/Editor/Private/PoseSearchDebugger.cpp
@@ -22,18 +22,73 @@
 #define LOCTEXT_NAMESPACE "PoseSearchDebugger"
 
 
-UPoseSearchDebuggerReflection::UPoseSearchDebuggerReflection(const FObjectInitializer& ObjectInitializer) : Super(ObjectInitializer)
+void FPoseSearchDebuggerPoseVectorChannel::Reset()
 {
+	Positions.Reset();
+	LinearVelocities.Reset();
+	FacingDirections.Reset();
 }
 
-void FPoseSearchDebuggerFeatureReflection::EmptyAll()
+void FPoseSearchDebuggerPoseVector::Reset()
 {
-	Positions.Empty();
-	LinearVelocities.Empty();
-	AngularVelocities.Empty();
+	Pose.Reset();
+	TrajectoryTimeBased.Reset();
+	TrajectoryDistanceBased.Reset();
 }
 
-namespace UE { namespace PoseSearch {
+void FPoseSearchDebuggerPoseVector::ExtractFeatures(const UE::PoseSearch::FFeatureVectorReader& Reader)
+{
+	using namespace UE::PoseSearch;
+
+	Reset();
+
+	FPoseSearchDebuggerPoseVectorChannel* Channels[] = 
+	{
+		&Pose,
+		&TrajectoryTimeBased,
+		&TrajectoryDistanceBased
+	};
+
+	for (const FPoseSearchFeatureDesc& Feature: Reader.GetLayout()->Features)
+	{
+		if (Feature.ChannelIdx >= UE_ARRAY_COUNT(Channels))
+		{
+			continue;
+		}
+
+		FVector Vector;
+		if (Reader.GetVector(Feature, &Vector))
+		{
+			if (Feature.Type == EPoseSearchFeatureType::Position)
+			{
+				Channels[Feature.ChannelIdx]->Positions.Add(Vector);
+			}
+			else if (Feature.Type == EPoseSearchFeatureType::LinearVelocity)
+			{
+				Channels[Feature.ChannelIdx]->LinearVelocities.Add(Vector);
+			}
+			else if (Feature.Type == EPoseSearchFeatureType::ForwardVector)
+			{
+				Channels[Feature.ChannelIdx]->FacingDirections.Add(Vector);
+			}
+		}
+		
+	}
+
+	for (FPoseSearchDebuggerPoseVectorChannel* Channel : Channels)
+	{
+		Channel->bShowPositions = !Channel->Positions.IsEmpty();
+		Channel->bShowLinearVelocities = !Channel->LinearVelocities.IsEmpty();
+		Channel->bShowFacingDirections = !Channel->FacingDirections.IsEmpty();
+	}
+
+	bShowPose = !Pose.IsEmpty();
+	bShowTrajectoryTimeBased = !TrajectoryTimeBased.IsEmpty();
+	bShowTrajectoryDistanceBased = !TrajectoryDistanceBased.IsEmpty();
+}
+
+
+namespace UE::PoseSearch {
 
 class FDebuggerDatabaseRowData : public TSharedFromThis<FDebuggerDatabaseRowData>
 {
@@ -501,6 +556,10 @@ void SDebuggerDatabaseView::OnFilterTextChanged(const FText& SearchText)
 	FilterDatabaseRows();
 }
 
+void SDebuggerDatabaseView::OnDatabaseRowSelectionChanged(TSharedPtr<FDebuggerDatabaseRowData> Row, ESelectInfo::Type SelectInfo)
+{
+	OnPoseSelectionChanged.ExecuteIfBound();
+}
 
 void SDebuggerDatabaseView::SortDatabaseRows()
 {
@@ -632,6 +691,9 @@ TSharedRef<ITableRow> SDebuggerDatabaseView::HandleGenerateActiveRow(TSharedRef<
 
 void SDebuggerDatabaseView::Construct(const FArguments& InArgs)
 {
+	ParentDebuggerViewPtr = InArgs._Parent;
+	OnPoseSelectionChanged = InArgs._OnPoseSelectionChanged;
+
 	using namespace DebuggerDatabaseColumns;
 
 	// @TODO: Support runtime reordering of these indices
@@ -683,7 +745,8 @@ void SDebuggerDatabaseView::Construct(const FArguments& InArgs)
 		.OnGenerateRow(this, &SDebuggerDatabaseView::HandleGenerateDatabaseRow)
 		.ExternalScrollbar(FilteredDatabaseView.ScrollBar)
 		.SelectionMode(ESelectionMode::SingleToggle)
-		.ConsumeMouseWheel(EConsumeMouseWheel::WhenScrollingPossible);
+		.ConsumeMouseWheel(EConsumeMouseWheel::WhenScrollingPossible)
+		.OnSelectionChanged(this, &SDebuggerDatabaseView::OnDatabaseRowSelectionChanged);
 
 	FilteredDatabaseView.RowStyle = FEditorStyle::GetWidgetStyle<FTableRowStyle>("TableView.Row");
 	// Set selected color to white to retain visibility when multi-selecting
@@ -840,10 +903,11 @@ void SDebuggerDatabaseView::Construct(const FArguments& InArgs)
 
 void SDebuggerDetailsView::Construct(const FArguments& InArgs)
 {
+	ParentDebuggerViewPtr = InArgs._Parent;
+
 	// Add property editor (detail view) UObject to world root so that it persists when PIE is stopped
 	Reflection = NewObject<UPoseSearchDebuggerReflection>();
 	Reflection->AddToRoot();
-	
 	check(IsValid(Reflection));
 
 	// @TODO: Convert this to a custom builder instead of of a standard details view
@@ -851,10 +915,9 @@ void SDebuggerDetailsView::Construct(const FArguments& InArgs)
 	FPropertyEditorModule& PropPlugin = FModuleManager::LoadModuleChecked<FPropertyEditorModule>("PropertyEditor");
 	FDetailsViewArgs DetailsViewArgs;
     DetailsViewArgs.NameAreaSettings = FDetailsViewArgs::HideNameArea;
+	DetailsViewArgs.DefaultsOnlyVisibility = EEditDefaultsOnlyNodeVisibility::Hide;
 
-	// @TODO: Hide arrays with zero elements in the detail view, if possible
 	Details = PropPlugin.CreateDetailView(DetailsViewArgs);
-
 	Details->SetObject(Reflection);
 	
 	ChildSlot
@@ -886,73 +949,37 @@ void SDebuggerDetailsView::UpdateReflection(const FTraceMotionMatchingStateMessa
 	Reflection->ElapsedPoseJumpTime = State.ElapsedPoseJumpTime;
 	Reflection->bFollowUpAnimation = EnumHasAnyFlags(State.Flags, FTraceMotionMatchingState::EFlags::FollowupAnimation);
 	
-	Reflection->PoseFeatures.EmptyAll();
-	Reflection->DistanceTrajectoryFeatures.EmptyAll();
-	Reflection->TimeTrajectoryFeatures.EmptyAll();
-
 	FFeatureVectorReader Reader;
-	// Ensure parity between Layout and QueryVector
 	Reader.Init(&Schema->Layout);
+
+	// Query pose
 	Reader.SetValues(State.QueryVector);
+	Reflection->QueryPoseVector.ExtractFeatures(Reader);
 
-	int32 NumSubsamples = Schema->PoseSampleTimes.Num();
-	const int32 NumBones = Schema->SampledBones.Num();
+	// Active pose
+	TArray<float> Pose(Database.SearchIndex.GetPoseValues(State.DbPoseIdx));
+	Database.SearchIndex.InverseNormalize(Pose);
+	Reader.SetValues(Pose);
+	Reflection->ActivePoseVector.ExtractFeatures(Reader);
 
-	FPoseSearchFeatureDesc Feature;
-	
-	// Aggregate all features and place into the reflection struct
-	auto Extract = [](
-		const FFeatureVectorReader& InReader,
-		const FPoseSearchFeatureDesc& InFeature,
-		FPoseSearchDebuggerFeatureReflection& ReflectionRef
-	)
+	auto DebuggerView = ParentDebuggerViewPtr.Pin();
+	if (DebuggerView.IsValid()) 
 	{
-		FVector OutputVec;
-		if (InReader.GetPosition(InFeature, &OutputVec))
-		{
-			ReflectionRef.Positions.Add(OutputVec);
-		}
-		if (InReader.GetLinearVelocity(InFeature, &OutputVec))
-		{
-			ReflectionRef.LinearVelocities.Add(OutputVec);
-		}
-		if (InReader.GetAngularVelocity(InFeature, &OutputVec))
-        {
-        	ReflectionRef.AngularVelocities.Add(OutputVec);
-        }
-	};
-
-	// Pose samples
-	Feature.Domain = EPoseSearchFeatureDomain::Time;
-	for (int32 SchemaSubsampleIndex = 0; SchemaSubsampleIndex < NumSubsamples; ++SchemaSubsampleIndex)
-	{
-		Feature.SubsampleIdx = SchemaSubsampleIndex;
-
-		for(int32 SchemaBoneIdx = 0; SchemaBoneIdx < NumBones; ++SchemaBoneIdx)
+		TArray<TSharedRef<FDebuggerDatabaseRowData>> SelectedRows = DebuggerView->GetSelectedDatabaseRows();
+		if (!SelectedRows.IsEmpty())
 		{
-			Feature.SchemaBoneIdx = SchemaBoneIdx;
-			Extract(Reader, Feature, Reflection->PoseFeatures);
+			const TSharedRef<FDebuggerDatabaseRowData>& Selected = SelectedRows[0];
+			Pose = Database.SearchIndex.GetPoseValues(Selected->PoseIdx);
+			Database.SearchIndex.InverseNormalize(Pose);
+			Reader.SetValues(Pose);
+			Reflection->SelectedPoseVector.ExtractFeatures(Reader);
+
+			Pose = Selected->PoseCostInfo.CostVector;
+			//Database.SearchIndex.InverseNormalize(Pose);
+			Reader.SetValues(Pose);
+			Reflection->CostVector.ExtractFeatures(Reader);
 		}
 	}
-
-	// Used for classifying trajectories instead of bones, special index
-	Feature.SchemaBoneIdx = FPoseSearchFeatureDesc::TrajectoryBoneIndex;
-	
-	// Trajectory time samples
-	NumSubsamples = Schema->TrajectorySampleTimes.Num();
-	for (int32 SchemaSubsampleIndex = 0; SchemaSubsampleIndex < NumSubsamples; ++SchemaSubsampleIndex)
-	{
-		Feature.SubsampleIdx = SchemaSubsampleIndex;
-		Extract(Reader, Feature, Reflection->TimeTrajectoryFeatures);
-	}
-	// Trajectory distance samples
-	NumSubsamples = Schema->TrajectorySampleDistances.Num();
-	Feature.Domain = EPoseSearchFeatureDomain::Distance;
-	for (int32 SchemaSubsampleIndex = 0; SchemaSubsampleIndex < NumSubsamples; ++SchemaSubsampleIndex)
-	{
-		Feature.SubsampleIdx = SchemaSubsampleIndex;
-		Extract(Reader, Feature, Reflection->DistanceTrajectoryFeatures);
-	}
 }
 
 void SDebuggerView::Construct(const FArguments& InArgs, uint64 InAnimInstanceId)
@@ -970,6 +997,8 @@ void SDebuggerView::Construct(const FArguments& InArgs, uint64 InAnimInstanceId)
 	OnUpdate = InArgs._OnUpdate;
 	OnViewClosed = InArgs._OnViewClosed;
 
+	
+
 	// Validate the existence of the passed getters
 	check(MotionMatchingNodeIds.IsBound());
     check(MotionMatchingState.IsBound());
@@ -1050,9 +1079,8 @@ void SDebuggerView::Tick(const FGeometry& AllottedGeometry, const double InCurre
 	{
 		return;
 	}
-	
+
 	const UWorld* DebuggerWorld = World.Get();
-    // World should be always valid at this point
     check(DebuggerWorld);
 	
 	// @TODO: Handle editor world when those features are enabled for the Rewind Debugger
@@ -1065,10 +1093,7 @@ void SDebuggerView::Tick(const FGeometry& AllottedGeometry, const double InCurre
 	const bool bSameTime = FMath::Abs(TimeMarker - PreviousTimeMarker) < DOUBLE_SMALL_NUMBER;
 	PreviousTimeMarker = TimeMarker;
 
-	OnUpdate.Execute(AnimInstanceId);
-	OnUpdateSelection.Execute(SelectedNodeId);
-	const FTraceMotionMatchingStateMessage* State = MotionMatchingState.Get();
-	const UPoseSearchDatabase* Database = PoseSearchDatabase.Get();
+	bool bNeedUpdate = false;
 
 	// We haven't reached the update point yet
 	if (CurrentConsecutiveFrames < ConsecutiveFramesUpdateThreshold)
@@ -1090,32 +1115,23 @@ void SDebuggerView::Tick(const FGeometry& AllottedGeometry, const double InCurre
 		// Haven't updated since passing through frame gate, update once
 		else if (!bUpdated)
 		{
-			const bool bNodeSelected = UpdateSelection();
-			if (bNodeSelected)
-			{
-				OnUpdateSelection.Execute(SelectedNodeId);
-				// Refresh state and database info after update
-				State = MotionMatchingState.Get();
-				Database = PoseSearchDatabase.Get();
-				if (!State || !Database)
-				{
-					return;
-				}
-				
-				UpdateViews(*State, *Database);
-			}
-			bUpdated = true;	
+			bNeedUpdate = true;
 		}
 	}
-	
-	const FTransform* Transform = RootTransform.Get();
-	if (!State || !Database || !Transform)
+
+	if (bNeedUpdate)
 	{
-		return;
+		OnUpdate.Execute(AnimInstanceId);
+		if (UpdateSelection())
+		{
+			OnUpdateSelection.Execute(SelectedNodeId);
+			UpdateViews();
+		}
+		bUpdated = true;
 	}
-
-	// Drawing of the debug vector should occur regardless of database update threshold
-	DrawFeatures(*DebuggerWorld, *State, *Database, *Transform);
+	
+	// Draw visualization every tick
+	DrawVisualization();
 }
 
 bool SDebuggerView::UpdateSelection()
@@ -1157,10 +1173,34 @@ bool SDebuggerView::UpdateSelection()
 	return bNodeSelected;
 }
 
-void SDebuggerView::UpdateViews(const FTraceMotionMatchingStateMessage& State, const UPoseSearchDatabase& Database) const
+void SDebuggerView::UpdateViews() const
 {
-	DatabaseView->Update(State, Database);
-	DetailsView->Update(State, Database);
+	const FTraceMotionMatchingStateMessage* State = MotionMatchingState.Get();
+	const UPoseSearchDatabase* Database = PoseSearchDatabase.Get();
+	if (State && Database)
+	{
+		DatabaseView->Update(*State, *Database);
+		DetailsView->Update(*State, *Database);
+	}
+}
+
+void SDebuggerView::DrawVisualization() const
+{
+	const UWorld* DebuggerWorld = World.Get();
+	check(DebuggerWorld);
+
+	const FTraceMotionMatchingStateMessage* State = MotionMatchingState.Get();
+	const UPoseSearchDatabase* Database = PoseSearchDatabase.Get();
+	const FTransform* Transform = RootTransform.Get();
+	if (State && Database && Transform)
+	{
+		DrawFeatures(*DebuggerWorld, *State, *Database, *Transform);
+	}
+}
+
+TArray<TSharedRef<FDebuggerDatabaseRowData>> SDebuggerView::GetSelectedDatabaseRows() const
+{
+	return DatabaseView->GetDatabaseRows()->GetSelectedItems();
 }
 
 void SDebuggerView::DrawFeatures(
@@ -1285,6 +1325,16 @@ int32 SDebuggerView::SelectView() const
 	return Debugger;
 }
 
+void SDebuggerView::OnPoseSelectionChanged()
+{
+	const FTraceMotionMatchingStateMessage* State = MotionMatchingState.Get();
+	const UPoseSearchDatabase* Database = PoseSearchDatabase.Get();
+	if (State && Database)
+	{
+		DetailsView->Update(*State, *Database);
+	}
+}
+
 FReply SDebuggerView::OnUpdateNodeSelection(int32 InSelectedNodeId)
 {
 	// -1 will backtrack to selection view
@@ -1389,6 +1439,8 @@ TSharedRef<SWidget> SDebuggerView::GenerateNodeDebuggerView()
 			+ SVerticalBox::Slot()
 			[
 				SAssignNew(DatabaseView, SDebuggerDatabaseView)
+				.Parent(SharedThis(this))
+				.OnPoseSelectionChanged(this, &SDebuggerView::OnPoseSelectionChanged)
 			]
 		]
 
@@ -1397,6 +1449,7 @@ TSharedRef<SWidget> SDebuggerView::GenerateNodeDebuggerView()
 		.Value(0.35f)
 		[
 			SAssignNew(DetailsView, SDebuggerDetailsView)
+			.Parent(SharedThis(this))
 		];
 }
 
@@ -1688,4 +1741,4 @@ FName FDebuggerViewCreator::GetName() const
 	return Name;
 }
 
-}} // namespace UE::PoseSearch
+} // namespace UE::PoseSearch
\ No newline at end of file
diff --git a/Engine/Plugins/Experimental/Animation/PoseSearch/Source/Editor/Private/PoseSearchDebugger.h b/Engine/Plugins/Experimental/Animation/PoseSearch/Source/Editor/Private/PoseSearchDebugger.h
--- a/Engine/Plugins/Experimental/Animation/PoseSearch/Source/Editor/Private/PoseSearchDebugger.h
+++ b/Engine/Plugins/Experimental/Animation/PoseSearch/Source/Editor/Private/PoseSearchDebugger.h
@@ -7,6 +7,7 @@
 #include "RewindDebuggerInterface/Public/IRewindDebuggerView.h"
 #include "RewindDebuggerInterface/Public/IRewindDebuggerViewCreator.h"
 #include "Widgets/Views/SHeaderRow.h"
+#include "PoseSearch/PoseSearch.h"
 #include "PoseSearchDebugger.generated.h"
 
 namespace TraceServices { class IAnalysisSession; }
@@ -23,27 +24,75 @@ class SWidgetSwitcher;
 class SSearchBox;
 template <typename ItemType> class SListView;
 
+namespace UE::PoseSearch
+{
+	namespace DebuggerDatabaseColumns { struct IColumn; }
+	struct FDebugDrawParams;
+	struct FTraceMotionMatchingStateMessage;
+	class FFeatureVectorReader;
+	class FDebuggerDatabaseRowData;
+} // namespace UE::PoseSearch
+
+
 /**
  * Used by the reflection UObject to encompass a set of feature vectors
  */
 USTRUCT()
-struct FPoseSearchDebuggerFeatureReflection
+struct FPoseSearchDebuggerPoseVectorChannel
 {
-	GENERATED_USTRUCT_BODY()
+	GENERATED_BODY()
 
 	// @TODO: Should be ideally enumerated based on all possible schema features
+
+	UPROPERTY(VisibleDefaultsOnly, Category = "Channel")
+	bool bShowPositions = false;
 	
-    UPROPERTY(VisibleAnywhere, Category="Query Data")
+	UPROPERTY(VisibleDefaultsOnly, Category = "Channel")
+	bool bShowLinearVelocities = false;
+
+	UPROPERTY(VisibleDefaultsOnly, Category = "Channel")
+	bool bShowFacingDirections = false;
+
+    UPROPERTY(VisibleAnywhere, EditFixedSize, Category="Channel", meta=(EditCondition="bShowPositions", EditConditionHides))
 	TArray<FVector> Positions;
-	
-    UPROPERTY(VisibleAnywhere, Category="Query Data")
+
+    UPROPERTY(VisibleAnywhere, EditFixedSize, Category="Channel", meta=(EditCondition="bShowLinearVelocities", EditConditionHides))
 	TArray<FVector> LinearVelocities;
 	
-    UPROPERTY(VisibleAnywhere, Category="Query Data")
-	TArray<FVector> AngularVelocities;
+    UPROPERTY(VisibleAnywhere, EditFixedSize, Category="Channel", meta=(EditCondition="bShowFacingDirections", EditConditionHides))
+	TArray<FVector> FacingDirections;
+
+	void Reset();
+	bool IsEmpty() const { return Positions.IsEmpty() && LinearVelocities.IsEmpty() && FacingDirections.IsEmpty(); }
+};
+
+USTRUCT()
+struct FPoseSearchDebuggerPoseVector
+{
+	GENERATED_BODY()
+
+
+	UPROPERTY(VisibleDefaultsOnly, Category = "Pose Vector")
+	bool bShowPose = false;
+
+	UPROPERTY(VisibleDefaultsOnly, Category = "Pose Vector")
+	bool bShowTrajectoryTimeBased = false;
+
+	UPROPERTY(VisibleDefaultsOnly, Category = "Pose Vector")
+	bool bShowTrajectoryDistanceBased = false;
 
-	/** Empty contents of above arrays */
-	void EmptyAll();
+	UPROPERTY(VisibleAnywhere, Category="Pose Vector", meta=(EditCondition="bShowPose", EditConditionHides))
+	FPoseSearchDebuggerPoseVectorChannel Pose;
+
+	UPROPERTY(VisibleAnywhere, Category="Pose Vector", meta=(EditCondition="bShowTrajectoryTimeBased", EditConditionHides, DisplayName="Trajectory (Time)"))
+	FPoseSearchDebuggerPoseVectorChannel TrajectoryTimeBased;
+
+	UPROPERTY(VisibleAnywhere, Category="Pose Vector", meta=(EditCondition="bShowTrajectoryDistanceBased", EditConditionHides, DisplayName="Trajectory (Distance)"))
+	FPoseSearchDebuggerPoseVectorChannel TrajectoryDistanceBased;
+
+	void Reset();
+	void ExtractFeatures(const UE::PoseSearch::FFeatureVectorReader& Reader);
+	bool IsEmpty() const { return Pose.IsEmpty() && TrajectoryTimeBased.IsEmpty() && TrajectoryDistanceBased.IsEmpty(); }
 };
 
 /**
@@ -52,7 +101,7 @@ struct FPoseSearchDebuggerFeatureReflection
 USTRUCT()
 struct FPoseSearchDebuggerFeatureDrawOptions
 {
-	GENERATED_USTRUCT_BODY()
+	GENERATED_BODY()
 
 	UPROPERTY(EditAnywhere, Category="Draw Options")
     bool bDisable = false;
@@ -76,7 +125,9 @@ struct FPoseSearchDebuggerFeatureDrawOptions
 UCLASS()
 class POSESEARCHEDITOR_API UPoseSearchDebuggerReflection : public UObject
 {
-	GENERATED_UCLASS_BODY()
+	GENERATED_BODY()
+
+public:
 
 	UPROPERTY(VisibleAnywhere, Category="Motion Matching State", Meta=(DisplayName="Current Database"))
 	FString CurrentDatabaseName = "";
@@ -95,34 +146,33 @@ class POSESEARCHEDITOR_API UPoseSearchDebuggerReflection : public UObject
 	UPROPERTY(EditAnywhere, Category="Draw Options", Meta=(DisplayName="Selected Pose"))
 	FPoseSearchDebuggerFeatureDrawOptions SelectedPoseDrawOptions;
 
-	/** Pose features of the current query vector */
-    UPROPERTY(EditAnywhere, Category="Query Data")
-	FPoseSearchDebuggerFeatureReflection PoseFeatures;
+    UPROPERTY(VisibleAnywhere, Category="Pose Vectors")
+	FPoseSearchDebuggerPoseVector QueryPoseVector;
     	
-	/** Time-based trajectory features of the current query vector */
-    UPROPERTY(EditAnywhere, Category="Query Data")
-	FPoseSearchDebuggerFeatureReflection TimeTrajectoryFeatures;
+    UPROPERTY(VisibleAnywhere, Category="Pose Vectors")
+	FPoseSearchDebuggerPoseVector ActivePoseVector;
 
-	/** Distance-based trajectory features of the current query vector */
-	UPROPERTY(EditAnywhere, Category="Query Data")
-    FPoseSearchDebuggerFeatureReflection DistanceTrajectoryFeatures;
+	UPROPERTY(VisibleAnywhere, Category="Pose Vectors")
+	FPoseSearchDebuggerPoseVector SelectedPoseVector;
+
+	UPROPERTY(VisibleAnywhere, Category="Pose Vectors")
+	FPoseSearchDebuggerPoseVector CostVector;
 };
 
 
-namespace UE { namespace PoseSearch {
+namespace UE::PoseSearch {
 
-namespace DebuggerDatabaseColumns { struct IColumn; }
-struct FDebugDrawParams;
-struct FTraceMotionMatchingStateMessage;
-class FFeatureVectorReader;
-class FDebuggerDatabaseRowData;
+DECLARE_DELEGATE(FOnPoseSelectionChanged)
 
 /**
  * Database panel view widget of the PoseSearch debugger
  */
+
 class SDebuggerDatabaseView : public SCompoundWidget
 {
 	SLATE_BEGIN_ARGS(SDebuggerDatabaseView) {}
+		SLATE_ARGUMENT(TWeakPtr<class SDebuggerView>, Parent)
+		SLATE_EVENT(FOnPoseSelectionChanged, OnPoseSelectionChanged)
 	SLATE_END_ARGS()
 	
 	void Construct(const FArguments& InArgs);
@@ -171,14 +221,21 @@ private:
 	/** Called when the text in the filter box is modified to update the filtering */
 	void OnFilterTextChanged(const FText& SearchText);
 
+	void OnDatabaseRowSelectionChanged(TSharedPtr<FDebuggerDatabaseRowData> Row, ESelectInfo::Type SelectInfo);
+
 	/** Generates a database row widget for the given data */
 	TSharedRef<ITableRow> HandleGenerateDatabaseRow(TSharedRef<FDebuggerDatabaseRowData> Item, const TSharedRef<STableViewBase>& OwnerTable) const;
 	
 	/** Generates the active row widget for the given data */
 	TSharedRef<ITableRow> HandleGenerateActiveRow(TSharedRef<FDebuggerDatabaseRowData> Item, const TSharedRef<STableViewBase>& OwnerTable) const;
 
+	TWeakPtr<SDebuggerView> ParentDebuggerViewPtr;
+
+	FOnPoseSelectionChanged OnPoseSelectionChanged;
+
 	/** Current column to sort by */
 	FName SortColumn = "";
+
 	/** Current sorting mode */
 	EColumnSortMode::Type SortMode = EColumnSortMode::Ascending;
 	
@@ -229,6 +286,7 @@ private:
 class SDebuggerDetailsView : public SCompoundWidget
 {
 	SLATE_BEGIN_ARGS(SDebuggerDetailsView) {}
+		SLATE_ARGUMENT(TWeakPtr<class SDebuggerView>, Parent)
 	SLATE_END_ARGS()
 
 	void Construct(const FArguments& InArgs);
@@ -243,13 +301,15 @@ private:
 	/** Update our details view object with new state information */
 	void UpdateReflection(const FTraceMotionMatchingStateMessage& State, const UPoseSearchDatabase& Database) const;
 	
+	TWeakPtr<SDebuggerView> ParentDebuggerViewPtr;
+
 	/** Details widget constructed for the MM node */
 	TSharedPtr<IDetailsView> Details;
+
 	/** Last updated reflection data relative to MM state */
 	TObjectPtr<UPoseSearchDebuggerReflection> Reflection = nullptr;
 };
 
-
 /** Callback to update the debugger when node is actively selected */
 DECLARE_DELEGATE_OneParam(FOnUpdateSelection, int32 NodeId);
 
@@ -291,6 +351,8 @@ public:
 	virtual FName GetName() const override;
 	virtual uint64 GetObjectId() const override;
 
+	TArray<TSharedRef<FDebuggerDatabaseRowData>> GetSelectedDatabaseRows() const;
+
 	virtual ~SDebuggerView() override;
 
 private:
@@ -301,7 +363,9 @@ private:
 	bool UpdateSelection();
 
 	/** Update the database and details views */
-	void UpdateViews(const FTraceMotionMatchingStateMessage& State, const UPoseSearchDatabase& Database) const;
+	void UpdateViews() const;
+
+	void DrawVisualization() const;
 	
 	/** Returns an int32 appropriate to the index of our widget selector */
 	int32 SelectView() const;
@@ -309,6 +373,8 @@ private:
 	/** Callback when a button in the selection view is clicked */
 	FReply OnUpdateNodeSelection(int32 InSelectedNodeId);
 
+	void OnPoseSelectionChanged();
+
 	/** Generates the message view relaying that there is no data */
 	TSharedRef<SWidget> GenerateNoDataMessageView();
 
@@ -499,4 +565,4 @@ public:
 };
 
 
-}}
+} // namespace UE::PoseSearch
diff --git a/Engine/Plugins/Experimental/Animation/PoseSearch/Source/Runtime/Private/PoseSearch.cpp b/Engine/Plugins/Experimental/Animation/PoseSearch/Source/Runtime/Private/PoseSearch.cpp
--- a/Engine/Plugins/Experimental/Animation/PoseSearch/Source/Runtime/Private/PoseSearch.cpp
+++ b/Engine/Plugins/Experimental/Animation/PoseSearch/Source/Runtime/Private/PoseSearch.cpp
@@ -240,6 +240,7 @@ void FPoseSearchFeatureVectorLayout::Init()
 {
 	uint32 FloatCount = 0;
 
+	// Initialize value offsets
 	for (FPoseSearchFeatureDesc& Feature : Features)
 	{
 		Feature.ValueOffset = FloatCount;
@@ -248,7 +249,31 @@ void FPoseSearchFeatureVectorLayout::Init()
 		FloatCount += FeatureNumFloats;
 	}
 
+	// Determine channel count
+	TBitArray ChannelsFound;
+	ChannelsFound.Init(false, UE::PoseSearch::MaxChannels);
+	for (FPoseSearchFeatureDesc& Feature : Features)
+	{
+		if (Feature.SchemaBoneIdx != FPoseSearchFeatureDesc::TrajectoryBoneIndex)
+		{
+			ChannelsFound[UE::PoseSearch::ChannelIdxPose] = true;
+		}
+		else if (Feature.Domain == EPoseSearchFeatureDomain::Time)
+		{
+			ChannelsFound[UE::PoseSearch::ChannelIdxTrajectoryTime] = true;
+		}
+		else if (Feature.Domain == EPoseSearchFeatureDomain::Distance)
+		{
+			ChannelsFound[UE::PoseSearch::ChannelIdxTrajectoryDistance] = true;
+		}
+		else
+		{
+			check(false);
+		}
+	}
+
 	NumFloats = FloatCount;
+	NumChannels = ChannelsFound.CountSetBits();
 }
 
 void FPoseSearchFeatureVectorLayout::Reset()
@@ -280,17 +305,9 @@ bool FPoseSearchFeatureVectorLayout::EnumerateBy(int32 ChannelIdx, EPoseSearchFe
 	auto IsChannelMatch = [](int32 ChannelIdx, const FPoseSearchFeatureDesc& Feature) -> bool
 	{
 		bool bChannelMatch = true;
-		if (ChannelIdx == UE::PoseSearch::ChannelIdxPose)
-		{
-			bChannelMatch = Feature.SchemaBoneIdx != FPoseSearchFeatureDesc::TrajectoryBoneIndex;
-		}
-		else if (ChannelIdx == UE::PoseSearch::ChannelIdxTrajectoryTime)
-		{
-			bChannelMatch = (Feature.SchemaBoneIdx == FPoseSearchFeatureDesc::TrajectoryBoneIndex) && (Feature.Domain == EPoseSearchFeatureDomain::Time);
-		}
-		else if (ChannelIdx == UE::PoseSearch::ChannelIdxTrajectoryDistance)
+		if (ChannelIdx >= 0)
 		{
-			bChannelMatch = (Feature.SchemaBoneIdx == FPoseSearchFeatureDesc::TrajectoryBoneIndex) && (Feature.Domain == EPoseSearchFeatureDomain::Distance);
+			bChannelMatch = Feature.ChannelIdx == ChannelIdx;
 		}
 		return bChannelMatch;
 	};
@@ -459,6 +476,7 @@ void UPoseSearchSchema::GenerateLayout()
 		FPoseSearchFeatureDesc Feature;
 		Feature.Domain = Domain;
 		Feature.Type = Type;
+		Feature.ChannelIdx = ChannelIdx;
 
 		if (ChannelIdx == UE::PoseSearch::ChannelIdxPose)
 		{
diff --git a/Engine/Plugins/Experimental/Animation/PoseSearch/Source/Runtime/Public/PoseSearch/PoseSearch.h b/Engine/Plugins/Experimental/Animation/PoseSearch/Source/Runtime/Public/PoseSearch/PoseSearch.h
--- a/Engine/Plugins/Experimental/Animation/PoseSearch/Source/Runtime/Public/PoseSearch/PoseSearch.h
+++ b/Engine/Plugins/Experimental/Animation/PoseSearch/Source/Runtime/Public/PoseSearch/PoseSearch.h
@@ -76,30 +76,27 @@ struct POSESEARCH_API FPoseSearchFeatureDesc
 	static constexpr int32 TrajectoryBoneIndex = -1; 
 
 	UPROPERTY()
-	int32 SchemaBoneIdx;
+	int32 SchemaBoneIdx = 0;
 
 	UPROPERTY()
-	int32 SubsampleIdx;
+	int32 SubsampleIdx = 0;
 
 	UPROPERTY()
-	EPoseSearchFeatureType Type;
+	EPoseSearchFeatureType Type = EPoseSearchFeatureType::Invalid;
 
 	UPROPERTY()
-	EPoseSearchFeatureDomain Domain;
+	EPoseSearchFeatureDomain Domain = EPoseSearchFeatureDomain::Invalid;
+
+	UPROPERTY()
+	int8 ChannelIdx = 0;
 
 	// Set via FPoseSearchFeatureLayout::Init() and ignored by operator==
 	UPROPERTY()
-	int32 ValueOffset;
+	int32 ValueOffset = 0;
 
 	bool operator==(const FPoseSearchFeatureDesc& Other) const;
 
-	FPoseSearchFeatureDesc()
-		: SchemaBoneIdx(MAX_int32)
-		, SubsampleIdx(MAX_int32)
-		, Type(EPoseSearchFeatureType::Invalid)
-		, Domain(EPoseSearchFeatureDomain::Invalid)
-		, ValueOffset(MAX_int32)
-	{}
+	bool IsValid() const { return Type != EPoseSearchFeatureType::Invalid; }
 };
 
 
@@ -122,6 +119,9 @@ struct POSESEARCH_API FPoseSearchFeatureVectorLayout
 	UPROPERTY()
 	int32 NumFloats = 0;
 
+	UPROPERTY()
+	int32 NumChannels = 0;
+
 	bool IsValid(int32 MaxNumBones) const;
 
 	bool EnumerateBy(int32 ChannelIdx, EPoseSearchFeatureType Type, int32& InOutFeatureIdx) const;
