MotionMatchingFFDemo - added animation playback to motion matching debugger.

diff --git a/Engine/Plugins/Experimental/Animation/PoseSearch/Source/Editor/Private/PoseSearchDebugger.cpp b/Engine/Plugins/Experimental/Animation/PoseSearch/Source/Editor/Private/PoseSearchDebugger.cpp
--- a/Engine/Plugins/Experimental/Animation/PoseSearch/Source/Editor/Private/PoseSearchDebugger.cpp
+++ b/Engine/Plugins/Experimental/Animation/PoseSearch/Source/Editor/Private/PoseSearchDebugger.cpp
@@ -2,9 +2,11 @@
 
 #include "PoseSearchDebugger.h"
 #include "IAnimationProvider.h"
-#include "Editor/EditorEngine.h"
+#include "IGameplayProvider.h"
 #include "PoseSearch/PoseSearch.h"
+#include "Animation/AnimationPoseData.h"
 #include "Animation/AnimSequence.h"
+#include "Animation/MirrorDataTable.h"
 #include "RewindDebuggerInterface/Public/IRewindDebugger.h"
 #include "ObjectTrace.h"
 #include "TraceServices/Model/AnalysisSession.h"
@@ -13,6 +15,7 @@
 #include "Widgets/Layout/SWidgetSwitcher.h"
 #include "Widgets/SBoxPanel.h"
 #include "Widgets/Input/SButton.h"
+#include "Widgets/Input/SNumericEntryBox.h"
 #include "Widgets/Input/SSearchBox.h"
 #include "Widgets/Images/SImage.h"
 #include "Styling/SlateIconFinder.h"
@@ -87,6 +90,116 @@ void FPoseSearchDebuggerPoseVector::ExtractFeatures(const UE::PoseSearch::FFeatu
 	bShowTrajectoryDistanceBased = !TrajectoryDistanceBased.IsEmpty();
 }
 
+void UPoseSearchMeshComponent::Initialize(const FTransform& InComponentToWorld)
+{
+	SetComponentToWorld(InComponentToWorld);
+	const FReferenceSkeleton& SkeletalMeshRefSkeleton = SkeletalMesh->GetRefSkeleton();
+
+	// set up bone visibility states as this gets skipped since we allocate the component array before registration
+	for (int32 BaseIndex = 0; BaseIndex < 2; ++BaseIndex)
+	{
+		BoneVisibilityStates[BaseIndex].SetNum(SkeletalMeshRefSkeleton.GetNum());
+		for (int32 BoneIndex = 0; BoneIndex < SkeletalMeshRefSkeleton.GetNum(); BoneIndex++)
+		{
+			BoneVisibilityStates[BaseIndex][BoneIndex] = BVS_ExplicitlyHidden;
+		}
+	}
+
+	StartingTransform = InComponentToWorld;
+	Refresh();
+}
+
+void UPoseSearchMeshComponent::Refresh()
+{
+	// Flip buffers once to copy the directly-written component space transforms
+	bNeedToFlipSpaceBaseBuffers = true;
+	bHasValidBoneTransform = false;
+	FlipEditableSpaceBases();
+	bHasValidBoneTransform = true;
+
+	InvalidateCachedBounds();
+	UpdateBounds();
+	MarkRenderTransformDirty();
+	MarkRenderDynamicDataDirty();
+	MarkRenderStateDirty();
+}
+
+void UPoseSearchMeshComponent::ResetToStart()
+{
+	SetComponentToWorld(StartingTransform);
+	Refresh();
+}
+
+void UPoseSearchMeshComponent::UpdatePose(const FUpdateContext& UpdateContext)
+{
+	FMemMark Mark(FMemStack::Get());
+
+	FCompactPose CompactPose;
+	CompactPose.SetBoneContainer(&RequiredBones);
+	FBlendedCurve Curve;
+	Curve.InitFrom(RequiredBones);
+	UE::Anim::FStackAttributeContainer Attributes;
+	FAnimationPoseData PoseData(CompactPose, Curve, Attributes);
+
+	float AdvancedTime = UpdateContext.SequenceStartTime;
+	FAnimationRuntime::AdvanceTime(
+		UpdateContext.bLoop,
+		UpdateContext.SequenceTime - UpdateContext.SequenceStartTime,
+		AdvancedTime,
+		UpdateContext.Sequence->GetPlayLength());
+
+	FAnimExtractContext ExtractionCtx;
+	ExtractionCtx.CurrentTime = AdvancedTime;
+
+	UpdateContext.Sequence->GetAnimationPose(PoseData, ExtractionCtx);
+	if (UpdateContext.bMirrored)
+	{
+		FAnimationRuntime::MirrorPose(
+			CompactPose, 
+			UpdateContext.MirrorDataTable->MirrorAxis, 
+			*UpdateContext.CompactPoseMirrorBones, 
+			*UpdateContext.ComponentSpaceRefRotations);
+	}
+
+	FCSPose<FCompactPose> ComponentSpacePose;
+	ComponentSpacePose.InitPose(CompactPose);
+
+	for (const FBoneIndexType BoneIndex : RequiredBones.GetBoneIndicesArray())
+	{
+		const FTransform BoneTransform = 
+			ComponentSpacePose.GetComponentSpaceTransform(FCompactPoseBoneIndex(BoneIndex));
+
+		FSkeletonPoseBoneIndex SkeletonBoneIndex =
+			RequiredBones.GetSkeletonPoseIndexFromCompactPoseIndex(FCompactPoseBoneIndex(BoneIndex));
+		FName BoneName = 
+			RequiredBones.GetSkeletonAsset()->GetReferenceSkeleton().GetBoneName(SkeletonBoneIndex.GetInt());
+		SetBoneTransformByName(BoneName, BoneTransform, EBoneSpaces::ComponentSpace);
+	}
+
+	LastRootMotionDelta = UpdateContext.Sequence->ExtractRootMotion(
+		UpdateContext.SequenceStartTime, 
+		UpdateContext.SequenceTime - UpdateContext.SequenceStartTime,
+		UpdateContext.bLoop);
+
+	if (UpdateContext.bMirrored)
+	{
+		const EAxis::Type MirrorAxis = UpdateContext.MirrorDataTable->MirrorAxis;
+		FVector T = LastRootMotionDelta.GetTranslation();
+		T = FAnimationRuntime::MirrorVector(T, MirrorAxis);
+		const FQuat ReferenceRotation = (*UpdateContext.ComponentSpaceRefRotations)[FCompactPoseBoneIndex(0)];
+		FQuat Q = LastRootMotionDelta.GetRotation();
+		Q = FAnimationRuntime::MirrorQuat(Q, MirrorAxis);
+		Q *= FAnimationRuntime::MirrorQuat(ReferenceRotation, MirrorAxis).Inverse() * ReferenceRotation;
+		LastRootMotionDelta = FTransform(Q, T, LastRootMotionDelta.GetScale3D());
+	}
+
+	const FTransform ComponentTransform = LastRootMotionDelta * StartingTransform;
+
+	SetComponentToWorld(ComponentTransform);
+	FillComponentSpaceTransforms();
+	Refresh();
+}
+
 
 namespace UE::PoseSearch {
 
@@ -556,9 +669,15 @@ void SDebuggerDatabaseView::OnFilterTextChanged(const FText& SearchText)
 	FilterDatabaseRows();
 }
 
-void SDebuggerDatabaseView::OnDatabaseRowSelectionChanged(TSharedPtr<FDebuggerDatabaseRowData> Row, ESelectInfo::Type SelectInfo)
+
+void SDebuggerDatabaseView::OnDatabaseRowSelectionChanged(
+	TSharedPtr<FDebuggerDatabaseRowData> Row, 
+	ESelectInfo::Type SelectInfo)
 {
-	OnPoseSelectionChanged.ExecuteIfBound();
+	if (Row.IsValid())
+	{
+		OnPoseSelectionChanged.ExecuteIfBound(Row->PoseIdx, Row->Time);
+	}
 }
 
 void SDebuggerDatabaseView::SortDatabaseRows()
@@ -691,10 +810,11 @@ TSharedRef<ITableRow> SDebuggerDatabaseView::HandleGenerateActiveRow(TSharedRef<
 
 void SDebuggerDatabaseView::Construct(const FArguments& InArgs)
 {
+	using namespace DebuggerDatabaseColumns;
+
 	ParentDebuggerViewPtr = InArgs._Parent;
 	OnPoseSelectionChanged = InArgs._OnPoseSelectionChanged;
-
-	using namespace DebuggerDatabaseColumns;
+	check(OnPoseSelectionChanged.IsBound());
 
 	// @TODO: Support runtime reordering of these indices
 	// Construct all column types
@@ -984,33 +1104,13 @@ void SDebuggerDetailsView::UpdateReflection(const FTraceMotionMatchingStateMessa
 
 void SDebuggerView::Construct(const FArguments& InArgs, uint64 InAnimInstanceId)
 {
-	MotionMatchingNodeIds = InArgs._MotionMatchingNodeIds;
-	MotionMatchingState = InArgs._MotionMatchingState;
-	PoseSearchDatabase = InArgs._PoseSearchDatabase;
-	OnUpdateSelection = InArgs._OnUpdateSelection;
-	IsPIESimulating = InArgs._IsPIESimulating;
-	IsRecording = InArgs._IsRecording;
-	RecordingDuration = InArgs._RecordingDuration;
-	NodesNum = InArgs._NodesNum;
-	World = InArgs._World;
-	RootTransform = InArgs._RootTransform;
-	OnUpdate = InArgs._OnUpdate;
+	ViewModel = InArgs._ViewModel;
 	OnViewClosed = InArgs._OnViewClosed;
 
 	
 
 	// Validate the existence of the passed getters
-	check(MotionMatchingNodeIds.IsBound());
-    check(MotionMatchingState.IsBound());
-	check(PoseSearchDatabase.IsBound());
-	check(OnUpdateSelection.IsBound());
-	check(IsPIESimulating.IsBound())
-    check(IsRecording.IsBound());
-    check(RecordingDuration.IsBound());
-	check(NodesNum.IsBound());
-	check(World.IsBound());
-	check(RootTransform.IsBound());
-	check(OnUpdate.IsBound());
+	check(ViewModel.IsBound())
 	check(OnViewClosed.IsBound());
 	
 	AnimInstanceId = InAnimInstanceId;
@@ -1065,7 +1165,7 @@ void SDebuggerView::Construct(const FArguments& InArgs, uint64 InAnimInstanceId)
 
 void SDebuggerView::SetTimeMarker(double InTimeMarker)
 {
-	if (IsPIESimulating.Get())
+	if (FDebugger::IsPIESimulating())
 	{
 		return;
 	}
@@ -1075,12 +1175,12 @@ void SDebuggerView::SetTimeMarker(double InTimeMarker)
 
 void SDebuggerView::Tick(const FGeometry& AllottedGeometry, const double InCurrentTime, const float InDeltaTime)
 {
-	if (IsPIESimulating.Get())
+	if (FDebugger::IsPIESimulating())
 	{
 		return;
 	}
 
-	const UWorld* DebuggerWorld = World.Get();
+	const UWorld* DebuggerWorld = FDebugger::GetWorld();
     check(DebuggerWorld);
 	
 	// @TODO: Handle editor world when those features are enabled for the Rewind Debugger
@@ -1093,6 +1193,8 @@ void SDebuggerView::Tick(const FGeometry& AllottedGeometry, const double InCurre
 	const bool bSameTime = FMath::Abs(TimeMarker - PreviousTimeMarker) < DOUBLE_SMALL_NUMBER;
 	PreviousTimeMarker = TimeMarker;
 
+	TSharedPtr<FDebuggerViewModel> Model = ViewModel.Get();
+
 	bool bNeedUpdate = false;
 
 	// We haven't reached the update point yet
@@ -1121,14 +1223,16 @@ void SDebuggerView::Tick(const FGeometry& AllottedGeometry, const double InCurre
 
 	if (bNeedUpdate)
 	{
-		OnUpdate.Execute(AnimInstanceId);
+		Model->OnUpdate();
 		if (UpdateSelection())
 		{
-			OnUpdateSelection.Execute(SelectedNodeId);
+			Model->OnUpdateNodeSelection(SelectedNodeId);
 			UpdateViews();
 		}
 		bUpdated = true;
 	}
+
+	Model->UpdateAnimSequence();
 	
 	// Draw visualization every tick
 	DrawVisualization();
@@ -1136,11 +1240,13 @@ void SDebuggerView::Tick(const FGeometry& AllottedGeometry, const double InCurre
 
 bool SDebuggerView::UpdateSelection()
 {
+	TSharedPtr<FDebuggerViewModel> Model = ViewModel.Get();
+
 	// Update selection view if no node selected
 	bool bNodeSelected = SelectedNodeId != INDEX_NONE;
 	if (!bNodeSelected)
 	{
-		const TArray<int32>& NodeIds = *MotionMatchingNodeIds.Get();
+		const TArray<int32>& NodeIds = *Model->GetNodeIds();
 		// Only one node active, bypass selection view
 		if (NodeIds.Num() == 1)
 		{
@@ -1153,14 +1259,14 @@ bool SDebuggerView::UpdateSelection()
 			SelectionView->ClearChildren();
 			for (int32 NodeId : NodeIds)
 			{
-				OnUpdateSelection.Execute(NodeId);
+				Model->OnUpdateNodeSelection(NodeId);
 				SelectionView->AddSlot()
 				.HAlign(HAlign_Fill)
 				.VAlign(VAlign_Center)
 				.Padding(10.0f)
 				[
 					SNew(SButton)
-					.Text(FText::FromString(PoseSearchDatabase.Get()->GetName()))
+					.Text(FText::FromString(Model->GetPoseSearchDatabase()->GetName()))
 					.HAlign(HAlign_Center)
 					.VAlign(VAlign_Center)
 					.ContentPadding(10.0f)
@@ -1175,8 +1281,8 @@ bool SDebuggerView::UpdateSelection()
 
 void SDebuggerView::UpdateViews() const
 {
-	const FTraceMotionMatchingStateMessage* State = MotionMatchingState.Get();
-	const UPoseSearchDatabase* Database = PoseSearchDatabase.Get();
+	const FTraceMotionMatchingStateMessage* State = ViewModel.Get()->GetMotionMatchingState();
+	const UPoseSearchDatabase* Database = ViewModel.Get()->GetPoseSearchDatabase();
 	if (State && Database)
 	{
 		DatabaseView->Update(*State, *Database);
@@ -1186,12 +1292,12 @@ void SDebuggerView::UpdateViews() const
 
 void SDebuggerView::DrawVisualization() const
 {
-	const UWorld* DebuggerWorld = World.Get();
+	const UWorld* DebuggerWorld = FDebugger::GetWorld();
 	check(DebuggerWorld);
 
-	const FTraceMotionMatchingStateMessage* State = MotionMatchingState.Get();
-	const UPoseSearchDatabase* Database = PoseSearchDatabase.Get();
-	const FTransform* Transform = RootTransform.Get();
+	const FTraceMotionMatchingStateMessage* State = ViewModel.Get()->GetMotionMatchingState();
+	const UPoseSearchDatabase* Database = ViewModel.Get()->GetPoseSearchDatabase();
+	const FTransform* Transform = ViewModel.Get()->GetRootTransform();
 	if (State && Database && Transform)
 	{
 		DrawFeatures(*DebuggerWorld, *State, *Database, *Transform);
@@ -1265,6 +1371,13 @@ void SDebuggerView::DrawFeatures(
 	DrawParams.LabelPrefix = TEXT("S");
 	SetDrawFlags(DrawParams, Reflection->SelectedPoseDrawOptions);
 
+	FSkeletonDrawParams SkeletonDrawParams;
+	const bool bDrawSelectedPose = Reflection->bDrawSelectedSkeleton;
+	if (bDrawSelectedPose)
+	{
+		SkeletonDrawParams.Flags |= ESkeletonDrawFlags::SelectedPose;
+	}
+
 	// Draw any selected database vectors
 	for (const TSharedRef<FDebuggerDatabaseRowData>& Row : Selected)
 	{
@@ -1287,27 +1400,35 @@ void SDebuggerView::DrawFeatures(
 		DrawParams.LabelPrefix = TEXT("A");
 		Draw(DrawParams);
 	}
+
+	if (Reflection->bDrawActiveSkeleton)
+	{
+		SkeletonDrawParams.Flags |= ESkeletonDrawFlags::ActivePose;
+	}
+
+	SkeletonDrawParams.Flags |= ESkeletonDrawFlags::AnimSequence;
+
+	ViewModel.Get()->OnDraw(SkeletonDrawParams);
 }
 
 int32 SDebuggerView::SelectView() const
 {
 	// Currently recording
-	if (IsPIESimulating.Get())
+	if (FDebugger::IsPIESimulating() && FDebugger::IsRecording())
 	{
-		if (IsRecording.Get())
-		{
-			return RecordingMsg;
-		}
+		return RecordingMsg;
 	}
 
 	// Data has not been recorded yet
-	if (RecordingDuration.Get() < DOUBLE_SMALL_NUMBER)
+	if (FDebugger::GetRecordingDuration() < DOUBLE_SMALL_NUMBER)
 	{
 		return StoppedMsg;
 	}
 
-	const bool bNoActiveNodes = NodesNum.Get() == 0;
-	const bool bNodeSelectedWithoutData = SelectedNodeId != INDEX_NONE && MotionMatchingState.Get() == nullptr;
+	const TSharedPtr<FDebuggerViewModel> Model = ViewModel.Get();
+
+	const bool bNoActiveNodes = Model->GetNodesNum() == 0;
+	const bool bNodeSelectedWithoutData = SelectedNodeId != INDEX_NONE && Model->GetMotionMatchingState() == nullptr;
 
 	// No active nodes, or node selected has no data
 	if (bNoActiveNodes || bNodeSelectedWithoutData)
@@ -1325,14 +1446,27 @@ int32 SDebuggerView::SelectView() const
 	return Debugger;
 }
 
-void SDebuggerView::OnPoseSelectionChanged()
+void SDebuggerView::OnPoseSelectionChanged(int32 PoseIdx, float Time)
 {
-	const FTraceMotionMatchingStateMessage* State = MotionMatchingState.Get();
-	const UPoseSearchDatabase* Database = PoseSearchDatabase.Get();
+	const TSharedPtr<FDebuggerViewModel> Model = ViewModel.Get();
+	const FTraceMotionMatchingStateMessage* State = Model->GetMotionMatchingState();
+	const UPoseSearchDatabase* Database = Model->GetPoseSearchDatabase();
+
 	if (State && Database)
 	{
 		DetailsView->Update(*State, *Database);
 	}
+	
+	if (PoseIdx == INDEX_NONE)
+	{
+		Model->ClearSelectedSkeleton();
+	}
+	else
+	{
+		Model->ShowSelectedSkeleton(PoseIdx, Time);
+		// Stop sequence player when switching selections
+		Model->StopSelection();
+	}
 }
 
 FReply SDebuggerView::OnUpdateNodeSelection(int32 InSelectedNodeId)
@@ -1343,13 +1477,34 @@ FReply SDebuggerView::OnUpdateNodeSelection(int32 InSelectedNodeId)
 	return FReply::Handled();
 }
 
+FReply SDebuggerView::TogglePlaySelectedSequences() const
+{
+	const TSharedPtr<SListView<TSharedRef<FDebuggerDatabaseRowData>>>& DatabaseRows = DatabaseView->GetDatabaseRows();
+	TArray<TSharedRef<FDebuggerDatabaseRowData>> Selected = DatabaseRows->GetSelectedItems();
+	const bool bPlaying = ViewModel.Get()->IsPlayingSelections();
+	if (!bPlaying)
+	{
+		if (!Selected.IsEmpty())
+		{
+			// @TODO: Make functional with multiple poses being selected
+			ViewModel.Get()->PlaySelection(Selected[0]->PoseIdx, Selected[0]->Time);
+		}
+	}
+	else
+	{
+		ViewModel.Get()->StopSelection();
+	}
+
+	return FReply::Handled();
+}
+
 TSharedRef<SWidget> SDebuggerView::GenerateNoDataMessageView()
 {
 	TSharedRef<SWidget> ReturnButtonView = GenerateReturnButtonView();
 	ReturnButtonView->SetVisibility(TAttribute<EVisibility>::CreateLambda([this]() -> EVisibility
 	{
 		// Hide the return button for the no data message if we have no nodes at all
-		return NodesNum.Get() > 0 ? EVisibility::Visible : EVisibility::Hidden;
+		return ViewModel.Get()->GetNodesNum() > 0 ? EVisibility::Visible : EVisibility::Hidden;
 	}));
 	
 	return 
@@ -1370,7 +1525,7 @@ TSharedRef<SWidget> SDebuggerView::GenerateNoDataMessageView()
 		];
 }
 
-TSharedRef<SWidget> SDebuggerView::GenerateReturnButtonView()
+TSharedRef<SHorizontalBox> SDebuggerView::GenerateReturnButtonView()
 {
 	return
 		SNew(SHorizontalBox)
@@ -1381,6 +1536,7 @@ TSharedRef<SWidget> SDebuggerView::GenerateReturnButtonView()
 		[
 			SNew(SButton)
 			.VAlign(VAlign_Center)
+			.Visibility_Lambda([this] { return ViewModel.Get()->GetNodesNum() > 1 ? EVisibility::Visible : EVisibility::Hidden; })
 			.ButtonStyle(FEditorStyle::Get(), "SimpleButton")
 			.ContentPadding( FMargin(1, 0) )
 			.OnClicked(this, &SDebuggerView::OnUpdateNodeSelection, static_cast<int32>(INDEX_NONE))
@@ -1411,14 +1567,76 @@ TSharedRef<SWidget> SDebuggerView::GenerateReturnButtonView()
 
 TSharedRef<SWidget> SDebuggerView::GenerateNodeDebuggerView()
 {
-	TSharedRef<SWidget> ReturnButtonView = GenerateReturnButtonView();
-	ReturnButtonView->SetVisibility(TAttribute<EVisibility>::CreateLambda([this]() -> EVisibility
-		{
-			// Collapse this view if we have don't have more than 1 node
-			return NodesNum.Get() > 1 ? EVisibility::Visible : EVisibility::Collapsed;
-		}
-	));
+	TSharedRef<SHorizontalBox> ReturnButtonView = GenerateReturnButtonView();
+	ReturnButtonView->AddSlot()
+	.VAlign(VAlign_Center)
+	.HAlign(HAlign_Fill)
+	.Padding(32, 5, 0, 0)
+	.AutoWidth()
+	[
+		SNew(SButton)
+		.VAlign(VAlign_Center)
+		.HAlign(HAlign_Fill)
+		.ButtonStyle(FEditorStyle::Get(), "Button")
+		.ContentPadding( FMargin(5, 0) )
+		.OnClicked(this, &SDebuggerView::TogglePlaySelectedSequences)
+		[
+			SNew(SHorizontalBox)
+			// Icon
+			+ SHorizontalBox::Slot()
+			.AutoWidth()
+			[
+				SNew(SImage)
+				.Image_Lambda([this]
+				{
+					const bool bPlayingSelections = ViewModel.Get()->IsPlayingSelections();
+					return FSlateIcon("FEditorStyle", bPlayingSelections ? "PlayWorld.StopPlaySession.Small" : "PlayWorld.PlayInViewport.Small").GetSmallIcon();
+				})
+			]
+			// Text
+			+ SHorizontalBox::Slot()
+			.Padding(FMargin(8, 0, 0, 0))
+			.AutoWidth()
+			[
+				SNew(STextBlock)
+				.Text_Lambda([this] { return ViewModel.Get()->IsPlayingSelections() ? FText::FromString("Stop Selected Sequence") : FText::FromString("Play Selected Sequence"); })
+				.Justification(ETextJustify::Center)
+			]
+		]
+	];
 	
+	ReturnButtonView->AddSlot()
+	.VAlign(VAlign_Center)
+	.HAlign(HAlign_Left)
+	.Padding(64, 5, 0, 0)
+	.AutoWidth()
+	[
+		SNew(SHorizontalBox)
+		
+		+ SHorizontalBox::Slot()
+		.AutoWidth()
+		.Padding(0, 5, 0, 0)
+		[
+			SNew(STextBlock)
+			.Text(FText::FromString("Sequence Play Rate: "))
+		]
+		+ SHorizontalBox::Slot()
+		.AutoWidth()
+		.Padding(8, 0, 0, 0)
+		[
+			SNew(SNumericEntryBox<float>)
+			.MinValue(0.0f)
+			.MaxValue(5.0f)
+			.MinSliderValue(0.0f)
+			.MaxSliderValue(5.0f)
+			.Delta(0.01f)
+			.AllowSpin(true)
+			// Lambda to accomodate the TOptional this requires (for now)
+			.Value_Lambda([this] { return ViewModel.Get()->GetPlayRate(); })
+			.OnValueChanged(ViewModel.Get().ToSharedRef(), &FDebuggerViewModel::ChangePlayRate)	
+		]
+	];
+
 	return 
 		SNew(SSplitter)
 		.Orientation(Orient_Horizontal)
@@ -1469,60 +1687,31 @@ SDebuggerView::~SDebuggerView()
 	OnViewClosed.Execute(AnimInstanceId);
 }
 
-void FDebuggerInstance::UpdateMotionMatchingStates(uint64 InAnimInstanceId)
+FDebuggerViewModel::FDebuggerViewModel(uint64 InAnimInstanceId)
+	: AnimInstanceId(InAnimInstanceId)
 {
-	NodeIds.Empty();
-	MotionMatchingStates.Empty();
-	if (!RewindDebugger)
-	{
-		return;
-	}
-	
-	// Get provider and validate
-	const TraceServices::IAnalysisSession* Session = RewindDebugger->GetAnalysisSession();
-	TraceServices::FAnalysisSessionReadScope SessionReadScope(*Session);
-	
-	const FTraceProvider* TraceProvider = Session->ReadProvider<FTraceProvider>(FTraceProvider::ProviderName);
-	if (!TraceProvider)
-	{
-		return;
-	}
-
-	const double TraceTime = RewindDebugger->CurrentTraceTime();
-	TraceServices::FFrame Frame;
-	ReadFrameProvider(*Session).GetFrameFromTime(TraceFrameType_Game, TraceTime, Frame);
-	
+	Skeletons.AddDefaulted(ESkeletonIndex::Num);
+}
 
-	TraceProvider->EnumerateMotionMatchingStateTimelines(InAnimInstanceId, [this, &Frame](const FTraceProvider::FMotionMatchingStateTimeline& InTimeline)
+FDebuggerViewModel::~FDebuggerViewModel()
+{
+	for (FSkeleton& Skeleton : Skeletons)
 	{
-		const FTraceMotionMatchingStateMessage* Message = nullptr;
-		
-		InTimeline.EnumerateEvents(Frame.StartTime, Frame.EndTime, [this, &Message, &Frame](double InStartTime, double InEndTime, const FTraceMotionMatchingStateMessage& InMessage)
-		{
-			Message = &InMessage;
-			return TraceServices::EEventEnumerate::Stop;
-		});
-
-		if (Message)
+		if (Skeleton.Actor.IsValid())
 		{
-			NodeIds.Add(Message->NodeId);
-			MotionMatchingStates.Add(Message);
+			Skeleton.Actor->Destroy();
 		}
-	});
-}
-
+	}
 
-FDebuggerInstance::FDebuggerInstance(uint64 InAnimInstanceId)
-	: AnimInstanceId(InAnimInstanceId)
-{
+	Skeletons.Empty();
 }
 
-const FTraceMotionMatchingStateMessage* FDebuggerInstance::GetMotionMatchingState() const
+const FTraceMotionMatchingStateMessage* FDebuggerViewModel::GetMotionMatchingState() const
 {
 	return ActiveMotionMatchingState;
 }
 
-const UPoseSearchDatabase* FDebuggerInstance::GetPoseSearchDatabase() const
+const UPoseSearchDatabase* FDebuggerViewModel::GetPoseSearchDatabase() const
 {
 	if (!ActiveMotionMatchingState)
 	{
@@ -1552,77 +1741,387 @@ const UPoseSearchDatabase* FDebuggerInstance::GetPoseSearchDatabase() const
 	return Database;
 }
 
-const TArray<int32>* FDebuggerInstance::GetNodeIds() const
+const FPoseSearchDatabaseSequence* FDebuggerViewModel::GetAnimSequence(int32 SequenceIdx) const
 {
-	return &NodeIds;
+	const UPoseSearchDatabase* Database = GetPoseSearchDatabase();
+	if (Database && Database->Sequences.IsValidIndex(SequenceIdx))
+	{
+		const FPoseSearchDatabaseSequence* DatabaseSequence = &Database->Sequences[SequenceIdx];
+		if (DatabaseSequence)
+		{
+			return DatabaseSequence;
+		}
+	}
+	return nullptr;
+}
+
+void FDebuggerViewModel::ShowSelectedSkeleton(int32 PoseIdx, float Time)
+{
+	const UPoseSearchDatabase* Database = GetPoseSearchDatabase();
+	if (!Database)
+	{
+		return;
+	}
+	UPoseSearchMeshComponent* Component = Skeletons[SelectedPose].Component;
+	if (!Component)
+	{
+		return;
+	}
+
+	Component->ResetToStart();
+	bSelecting = true;
+	const FPoseSearchIndexAsset* Asset = Database->SearchIndex.FindAssetForPose(PoseIdx);
+	Skeletons[SelectedPose].SequenceIdx = Asset->SourceAssetIdx;
+	Skeletons[SelectedPose].Time = Time;
+	Skeletons[SelectedPose].bMirrored = Asset->bMirrored;
+}
+
+void FDebuggerViewModel::ClearSelectedSkeleton()
+{
+	bSelecting = false;
 }
 
+const TArray<int32>* FDebuggerViewModel::GetNodeIds() const
+{
+	return &NodeIds;
+}
 
-int32 FDebuggerInstance::GetNodesNum() const
+int32 FDebuggerViewModel::GetNodesNum() const
 {
 	return MotionMatchingStates.Num();
 }
 
-const FTransform* FDebuggerInstance::GetRootTransform() const
+const FTransform* FDebuggerViewModel::GetRootTransform() const
+{
+	return RootTransform;
+}
+
+void FDebuggerViewModel::OnUpdate()
 {
-	if (!RewindDebugger || !ActiveMotionMatchingState)
+	if (!bSkeletonsInitialized)
 	{
-		return nullptr;
+		UWorld* World = RewindDebugger.Get()->GetWorldToVisualize();
+		for (FSkeleton& Skeleton : Skeletons)
+		{
+			FActorSpawnParameters ActorSpawnParameters;
+			ActorSpawnParameters.bHideFromSceneOutliner = false;
+			ActorSpawnParameters.ObjectFlags |= RF_Transient;
+			Skeleton.Actor = World->SpawnActor<AActor>(ActorSpawnParameters);
+			Skeleton.Actor->SetActorLabel(TEXT("PoseSearch"));
+			Skeleton.Component = NewObject<UPoseSearchMeshComponent>(Skeleton.Actor.Get());
+			Skeleton.Actor->AddInstanceComponent(Skeleton.Component);
+			Skeleton.Component->RegisterComponentWithWorld(World);
+		}
+		FWorldDelegates::OnWorldCleanup.AddRaw(this, &FDebuggerViewModel::OnWorldCleanup);
+		bSkeletonsInitialized = true;
 	}
-	
+
+	UpdateFromTimeline();
+}
+
+void FDebuggerViewModel::OnUpdateNodeSelection(int32 InNodeId)
+{
+	if (InNodeId == INDEX_NONE)
+	{
+		return;
+	}
+
+	ActiveMotionMatchingState = nullptr;
+
+	// Find node in all motion matching states this frame
+	const int32 NodesNum = NodeIds.Num();
+	for (int i = 0; i < NodesNum; ++i)
+	{
+		if (NodeIds[i] == InNodeId)
+		{
+			ActiveMotionMatchingState = MotionMatchingStates[i];
+			break;
+		}
+	}
+	if (ActiveMotionMatchingState)
+	{
+		Skeletons[ActivePose].SequenceIdx = ActiveMotionMatchingState->DbPoseIdx;
+	}
+}
+
+void FDebuggerViewModel::OnDraw(FSkeletonDrawParams& DrawParams)
+{
+	// Returns if it is to be drawn this frame
+	auto SetDrawSkeleton = [this](UPoseSearchMeshComponent* InComponent, bool bDraw)
+	{
+		const bool bIsDrawingSkeleton = InComponent->ShouldDrawDebugSkeleton();
+		if (bIsDrawingSkeleton != bDraw)
+		{
+			InComponent->SetDrawDebugSkeleton(bDraw);
+		}
+		InComponent->MarkRenderStateDirty();
+	};
+	const bool bDrawActivePose = EnumHasAnyFlags(DrawParams.Flags, ESkeletonDrawFlags::ActivePose);
+	SetDrawSkeleton(Skeletons[ActivePose].Component, bDrawActivePose);
+	// If flag is set and we are currently in a valid drawing state
+	const bool bDrawSelectedPose = EnumHasAnyFlags(DrawParams.Flags, ESkeletonDrawFlags::SelectedPose) && bSelecting;
+	SetDrawSkeleton(Skeletons[SelectedPose].Component, bDrawSelectedPose);
+
+	FillCompactPoseAndComponentRefRotations();
+
+	UPoseSearchMeshComponent::FUpdateContext UpdateContext;
+	UpdateContext.MirrorDataTable = GetPoseSearchDatabase()->Schema->MirrorDataTable;
+	UpdateContext.CompactPoseMirrorBones = &CompactPoseMirrorBones;
+	UpdateContext.ComponentSpaceRefRotations = &ComponentSpaceRefRotations;
+
+	if (bDrawSelectedPose)
+	{
+		UPoseSearchMeshComponent* Component = Skeletons[SelectedPose].Component;
+		const FPoseSearchDatabaseSequence* DatabaseSequence = GetAnimSequence(Skeletons[SelectedPose].SequenceIdx);
+		UAnimSequence* Sequence = DatabaseSequence->Sequence;
+
+		UpdateContext.Sequence = Sequence;
+		UpdateContext.SequenceStartTime = Skeletons[SelectedPose].Time;
+		UpdateContext.SequenceTime = Skeletons[SelectedPose].Time;
+		UpdateContext.bMirrored = Skeletons[SelectedPose].bMirrored;
+		UpdateContext.bLoop = DatabaseSequence->bLoopAnimation;
+
+		Component->UpdatePose(UpdateContext);
+	}
+
+	const bool bDrawSequence = EnumHasAnyFlags(DrawParams.Flags, ESkeletonDrawFlags::AnimSequence);
+	if (bDrawSequence && SequenceData.bActive)
+	{
+		UPoseSearchMeshComponent* Component = Skeletons[AnimSequence].Component;
+		SetDrawSkeleton(Component, true);
+
+		const FPoseSearchDatabaseSequence* DatabaseSequence = GetAnimSequence(Skeletons[AnimSequence].SequenceIdx);
+		UAnimSequence* Sequence = DatabaseSequence->Sequence;
+
+		UpdateContext.Sequence = Sequence;
+		UpdateContext.SequenceStartTime = SequenceData.StartTime;
+		UpdateContext.SequenceTime = Skeletons[AnimSequence].Time;
+		UpdateContext.bMirrored = Skeletons[SelectedPose].bMirrored;
+		UpdateContext.bLoop = DatabaseSequence->bLoopAnimation;
+
+
+		Component->UpdatePose(UpdateContext);
+	}
+}
+
+void FDebuggerViewModel::UpdateFromTimeline()
+{
+	NodeIds.Empty();
+	MotionMatchingStates.Empty();
+	SkeletalMeshComponentId = 0;
+
 	// Get provider and validate
-	const TraceServices::IAnalysisSession* Session = RewindDebugger->GetAnalysisSession();
+	const TraceServices::IAnalysisSession* Session = RewindDebugger.Get()->GetAnalysisSession();
 	TraceServices::FAnalysisSessionReadScope SessionReadScope(*Session);
-	
+
+	const FTraceProvider* PoseSearchProvider = Session->ReadProvider<FTraceProvider>(FTraceProvider::ProviderName);
 	const IAnimationProvider* AnimationProvider = Session->ReadProvider<IAnimationProvider>("AnimationProvider");
-	if (!AnimationProvider)
+	const IGameplayProvider* GameplayProvider = Session->ReadProvider<IGameplayProvider>("GameplayProvider");
+	if (!(PoseSearchProvider && AnimationProvider && GameplayProvider))
 	{
-		return nullptr;
+		return;
 	}
-
-	const double TraceTime = RewindDebugger->CurrentTraceTime();
+	const double TraceTime = RewindDebugger.Get()->CurrentTraceTime();
 	TraceServices::FFrame Frame;
 	ReadFrameProvider(*Session).GetFrameFromTime(TraceFrameType_Game, TraceTime, Frame);
+	PoseSearchProvider->EnumerateMotionMatchingStateTimelines(AnimInstanceId, [&](const FTraceProvider::FMotionMatchingStateTimeline& InTimeline)
+	{
+		const FTraceMotionMatchingStateMessage* Message = nullptr;
 
-	const FTransform* RootTransform = nullptr;
-
-	AnimationProvider->ReadSkeletalMeshPoseTimeline(ActiveMotionMatchingState->SkeletalMeshComponentId, [&RootTransform, &Frame](const IAnimationProvider::SkeletalMeshPoseTimeline& TimelineData, bool bHasCurves)
+		InTimeline.EnumerateEvents(Frame.StartTime, Frame.EndTime, [&Message](double InStartTime, double InEndTime, const FTraceMotionMatchingStateMessage& InMessage)
+		{
+			Message = &InMessage;
+			return TraceServices::EEventEnumerate::Stop;
+		});
+		if (Message)
+		{
+			NodeIds.Add(Message->NodeId);
+			MotionMatchingStates.Add(Message);
+			SkeletalMeshComponentId = Message->SkeletalMeshComponentId;
+		}
+	});
+	/** No active motion matching state as no messages were read */
+	if (SkeletalMeshComponentId == 0)
+	{
+		return;
+	}
+	AnimationProvider->ReadSkeletalMeshPoseTimeline(SkeletalMeshComponentId, [&](const IAnimationProvider::SkeletalMeshPoseTimeline& TimelineData, bool bHasCurves)
 	{
-		TimelineData.EnumerateEvents(Frame.StartTime, Frame.EndTime,
-			[&RootTransform, &Frame](double InStartTime, double InEndTime, uint32 InDepth, const FSkeletalMeshPoseMessage& PoseMessage)
+		TimelineData.EnumerateEvents(Frame.StartTime, Frame.EndTime, [&](double InStartTime, double InEndTime, uint32 InDepth, const FSkeletalMeshPoseMessage& PoseMessage) -> TraceServices::EEventEnumerate
+		{
+			// Update root transform
+			RootTransform = &PoseMessage.ComponentToWorld;
+			const FSkeletalMeshInfo* SkeletalMeshInfo = AnimationProvider->FindSkeletalMeshInfo(PoseMessage.MeshId);
+			const FObjectInfo* SkeletalMeshObjectInfo = GameplayProvider->FindObjectInfo(PoseMessage.MeshId);
+			if (!SkeletalMeshInfo || !SkeletalMeshObjectInfo)
 			{
-				RootTransform = &PoseMessage.ComponentToWorld;
+
 				return TraceServices::EEventEnumerate::Stop;
-			});
+			}
+			UPoseSearchMeshComponent* ActiveComponent = Skeletons[ActivePose].Component;
+			UPoseSearchMeshComponent* SelectedComponent = Skeletons[SelectedPose].Component;
+			UPoseSearchMeshComponent* SequenceComponent = Skeletons[AnimSequence].Component;
+			USkeletalMesh* SkeletalMesh = TSoftObjectPtr<USkeletalMesh>(FSoftObjectPath(SkeletalMeshObjectInfo->PathName)).LoadSynchronous();
+			if (SkeletalMesh)
+			{
+				ActiveComponent->SetSkeletalMesh(SkeletalMesh);
+				SelectedComponent->SetSkeletalMesh(SkeletalMesh);
+				SequenceComponent->SetSkeletalMesh(SkeletalMesh);
+			}
+			FTransform ComponentWorldTransform;
+			// Active skeleton is simply the traced bone transforms
+			AnimationProvider->GetSkeletalMeshComponentSpacePose(PoseMessage, *SkeletalMeshInfo, ComponentWorldTransform, ActiveComponent->GetEditableComponentSpaceTransforms());
+			ActiveComponent->Initialize(ComponentWorldTransform);
+			ActiveComponent->SetDebugDrawColor(FLinearColor::Green);
+			SelectedComponent->SetDebugDrawColor(FLinearColor::Blue);
+			SelectedComponent->Initialize(ComponentWorldTransform);
+			SequenceComponent->SetDebugDrawColor(FLinearColor::Red);
+			SequenceComponent->Initialize(ComponentWorldTransform);
+
+			return TraceServices::EEventEnumerate::Stop;
+		});
 	});
+}
 
-	return RootTransform;
+void FDebuggerViewModel::UpdateAnimSequence()
+{
+	if (!IsPlayingSelections())
+	{
+		return;
+	}
+
+	const UPoseSearchDatabase* Database = GetPoseSearchDatabase();
+	const FPoseSearchDatabaseSequence* DatabaseSequence = GetAnimSequence(Skeletons[AnimSequence].SequenceIdx);
+	UAnimSequence* Sequence = DatabaseSequence->Sequence;
+	FSkeleton& SequenceSkeleton = Skeletons[AnimSequence];
+	UPoseSearchMeshComponent* Component = SequenceSkeleton.Component;
+	auto RestartSequence = [&]()
+	{
+		Component->ResetToStart();
+		SequenceData.AccumulatedTime = 0.0;
+		SequenceSkeleton.Time = SequenceData.StartTime;
+	};
+
+	if (Sequence)
+
+
+	{
+		const float DT = static_cast<float>(FApp::GetDeltaTime()) * SequencePlayRate;
+		const float PlayLength = Sequence->GetPlayLength();
+
+		const float DistanceHorizon = Database->Schema->GetTrajectoryFutureDistanceHorizon();
+		const bool bExceededDistanceHorizon = Component->LastRootMotionDelta.GetTranslation().Size() > DistanceHorizon;
+		const float TimeHorizon = Database->Schema->GetTrajectoryFutureTimeHorizon();
+		const bool bExceededTimeHorizon = (SequenceSkeleton.Time - SequenceData.StartTime) > TimeHorizon;
+		const bool bExceededHorizon = bExceededDistanceHorizon && bExceededTimeHorizon;
+		if (DatabaseSequence->bLoopAnimation)
+		{
+			if (bExceededHorizon)
+			{
+				// Delay before restarting the sequence to give the user some idea of where it would land
+				if (SequenceData.AccumulatedTime > SequenceData.StopDuration)
+				{
+					RestartSequence();
+				}
+				else
+				{
+					SequenceData.AccumulatedTime += DT;
+				}
+				return;
+			}
+
+			SequenceSkeleton.Time += DT;
+			SequenceData.AccumulatedTime += DT;
+		}
+		else
+		{
+			// Used to cap the sequence, but avoid modding when updating the pose
+			static constexpr float LengthOffset = 0.001f;
+			const bool bFinishedSequence = SequenceSkeleton.Time >= PlayLength - LengthOffset;
+
+			// Sequence player reached end of clip or reached distance horizon of trajectory vector
+			if (bFinishedSequence || bExceededHorizon)
+			{
+				// Delay before restarting the sequence to give the user some idea of where it would land
+				if (SequenceData.AccumulatedTime > SequenceData.StopDuration)
+				{
+					RestartSequence();
+				}
+				else
+				{
+					SequenceData.AccumulatedTime += DT;
+				}
+			}
+			else
+			{
+				// If we haven't finished, update the play time capped by the anim sequence (not looping)
+				SequenceSkeleton.Time += DT;
+			}
+		}
+	}
 }
 
-void FDebuggerInstance::OnUpdate(uint64 InAnimInstanceId)
+void FDebuggerViewModel::FillCompactPoseAndComponentRefRotations()
 {
-	UpdateMotionMatchingStates(InAnimInstanceId);
+	const UPoseSearchDatabase* Database = GetPoseSearchDatabase();
+	if (Database != nullptr)
+	{
+		UMirrorDataTable* MirrorDataTable = Database->Schema->MirrorDataTable;
+		if (MirrorDataTable != nullptr)
+		{
+			if (CompactPoseMirrorBones.Num() == 0 || ComponentSpaceRefRotations.Num() == 0)
+			{
+				MirrorDataTable->FillCompactPoseAndComponentRefRotations(
+					Skeletons[ActivePose].Component->RequiredBones,
+					CompactPoseMirrorBones,
+					ComponentSpaceRefRotations);
+			}
+
+			return;
+		}
+	}
+
+	CompactPoseMirrorBones.Reset();
+	ComponentSpaceRefRotations.Reset();
 }
 
-void FDebuggerInstance::OnUpdateSelection(int32 InNodeId)
+void FDebuggerViewModel::PlaySelection(int32 PoseIdx, float Time)
 {
-	if (InNodeId == INDEX_NONE)
+	const UPoseSearchDatabase* Database = GetPoseSearchDatabase();
+	if (!Database)
 	{
 		return;
 	}
-	
-	// Find node in all motion matching states this frame
-	const int32 NodesNum = NodeIds.Num();
-	for(int i = 0; i < NodesNum; ++i)
+	UPoseSearchMeshComponent* Component = Skeletons[AnimSequence].Component;
+	if (!Component)
 	{
-		if (NodeIds[i] == InNodeId)
-		{
-			ActiveMotionMatchingState = MotionMatchingStates[i];
-			return;
-		}
+		return;
 	}
 
-	ActiveMotionMatchingState = nullptr;
+	const FPoseSearchIndexAsset* Asset = Database->SearchIndex.FindAssetForPose(PoseIdx);
+
+	Component->ResetToStart();
+	Skeletons[AnimSequence].SequenceIdx = Asset->SourceAssetIdx;
+	Skeletons[AnimSequence].Time = SequenceData.StartTime = Time;
+	Skeletons[AnimSequence].bMirrored = Asset->bMirrored;
+	SequenceData.AccumulatedTime = 0.0f;
+	SequenceData.bActive = true;
+}
+void FDebuggerViewModel::StopSelection()
+{
+	UPoseSearchMeshComponent* Component = Skeletons[AnimSequence].Component;
+	if (!Component)
+	{
+		return;
+	}
+
+	SequenceData = {};
+	// @TODO: Make more functionality rely on checking if it should draw the sequence
+	Component->SetDrawDebugSkeleton(false);
+}
+void FDebuggerViewModel::OnWorldCleanup(UWorld* InWorld, bool bSessionEnded, bool bCleanupResources)
+{
+	bSkeletonsInitialized = false;
 }
 
 FDebugger* FDebugger::Debugger;
@@ -1639,72 +2138,75 @@ void FDebugger::Shutdown()
 	delete Debugger;
 }
 
-void FDebugger::Update(float DeltaTime, IRewindDebugger* InRewindDebugger)
+bool FDebugger::IsPIESimulating()
 {
-	// Update active rewind debugger in use
-	RewindDebugger = InRewindDebugger;
+	return Debugger->RewindDebugger->IsPIESimulating();
+}
+
+bool FDebugger::IsRecording()
+{
+	return Debugger->RewindDebugger->IsRecording();
 
-	for (TSharedRef<FDebuggerInstance>& DebuggerInstance : DebuggerInstances)
-	{
-		DebuggerInstance->RewindDebugger = RewindDebugger;
-	}
 }
 
-bool FDebugger::IsPIESimulating() const
+double FDebugger::GetRecordingDuration()
 {
-	return RewindDebugger->IsPIESimulating();
+	return Debugger->RewindDebugger->GetRecordingDuration();
 }
 
-bool FDebugger::IsRecording() const
+UWorld* FDebugger::GetWorld()
 {
-	return RewindDebugger->IsRecording();
+	return Debugger->RewindDebugger->GetWorldToVisualize();
 }
 
-double FDebugger::GetRecordingDuration() const
+const IRewindDebugger* FDebugger::GetRewindDebugger()
 {
-	return RewindDebugger->GetRecordingDuration();
+	return Debugger->RewindDebugger;
 }
 
-const UWorld* FDebugger::GetWorld() const
+void FDebugger::Update(float DeltaTime, IRewindDebugger* InRewindDebugger)
 {
-	return RewindDebugger->GetWorldToVisualize();
+	// Update active rewind debugger in use
+	RewindDebugger = InRewindDebugger;
 }
 
 void FDebugger::OnViewClosed(uint64 InAnimInstanceId)
 {
-	for (int i = 0; i < DebuggerInstances.Num(); ++i)
+	TArray<TSharedRef<FDebuggerViewModel>>& Models = Debugger->ViewModels;
+	for (int i = 0; i < Models.Num(); ++i)
 	{
-		if (DebuggerInstances[i]->AnimInstanceId == InAnimInstanceId)
+		if (Models[i]->AnimInstanceId == InAnimInstanceId)
 		{
-			DebuggerInstances.RemoveAtSwap(i);
+			Models.RemoveAtSwap(i);
 			return;
 		}
 	}
-
 	// Should always be a valid remove
 	checkNoEntry();
 }
 
+TSharedPtr<FDebuggerViewModel> FDebugger::GetViewModel(uint64 InAnimInstanceId)
+{
+	TArray<TSharedRef<FDebuggerViewModel>>& Models = Debugger->ViewModels;
+	for (int i = 0; i < Models.Num(); ++i)
+	{
+		if (Models[i]->AnimInstanceId == InAnimInstanceId)
+		{
+			return Models[i];
+		}
+	}
+	return nullptr;
+}
+
 TSharedPtr<SDebuggerView> FDebugger::GenerateInstance(uint64 InAnimInstanceId)
 {
-	TSharedRef<FDebuggerInstance>& DebuggerInstance = DebuggerInstances.Add_GetRef(MakeShared<FDebuggerInstance>(InAnimInstanceId));
-	
+	ViewModels.Add_GetRef(MakeShared<FDebuggerViewModel>(InAnimInstanceId))->RewindDebugger.BindStatic(&FDebugger::GetRewindDebugger);
+
 	TSharedPtr<SDebuggerView> DebuggerView;
 
 	SAssignNew(DebuggerView, SDebuggerView, InAnimInstanceId)
-	.MotionMatchingNodeIds(DebuggerInstance, &FDebuggerInstance::GetNodeIds)
-	.MotionMatchingState(DebuggerInstance, &FDebuggerInstance::GetMotionMatchingState)
-	.PoseSearchDatabase(DebuggerInstance, &FDebuggerInstance::GetPoseSearchDatabase)
-	.NodesNum(DebuggerInstance, &FDebuggerInstance::GetNodesNum)
-	.RootTransform(DebuggerInstance, &FDebuggerInstance::GetRootTransform)
-	.OnUpdateSelection(DebuggerInstance, &FDebuggerInstance::OnUpdateSelection)
-	.OnUpdate(DebuggerInstance, &FDebuggerInstance::OnUpdate)
-	// Rewind debugger callbacks gathered from shared object
-	.World_Raw(this, &FDebugger::GetWorld)
-	.IsRecording_Raw(this, &FDebugger::IsRecording)
-	.IsPIESimulating_Raw(this, &FDebugger::IsPIESimulating)
-	.RecordingDuration_Raw(this, &FDebugger::GetRecordingDuration)
-	.OnViewClosed_Raw(this, &FDebugger::OnViewClosed);
+		.ViewModel_Static(&FDebugger::GetViewModel, InAnimInstanceId)
+		.OnViewClosed_Static(&FDebugger::OnViewClosed);
 
 	return DebuggerView;
 }
diff --git a/Engine/Plugins/Experimental/Animation/PoseSearch/Source/Editor/Private/PoseSearchDebugger.h b/Engine/Plugins/Experimental/Animation/PoseSearch/Source/Editor/Private/PoseSearchDebugger.h
--- a/Engine/Plugins/Experimental/Animation/PoseSearch/Source/Editor/Private/PoseSearchDebugger.h
+++ b/Engine/Plugins/Experimental/Animation/PoseSearch/Source/Editor/Private/PoseSearchDebugger.h
@@ -3,6 +3,7 @@
 #pragma once
 
 #include "CoreMinimal.h"
+#include "Components/PoseableMeshComponent.h"
 #include "RewindDebuggerInterface/Public/IRewindDebuggerExtension.h"
 #include "RewindDebuggerInterface/Public/IRewindDebuggerView.h"
 #include "RewindDebuggerInterface/Public/IRewindDebuggerViewCreator.h"
@@ -10,6 +11,9 @@
 #include "PoseSearch/PoseSearch.h"
 #include "PoseSearchDebugger.generated.h"
 
+struct FPoseSearchDatabaseSequence;
+class UAnimSequence;
+
 namespace TraceServices { class IAnalysisSession; }
 enum class EPoseSearchFeatureDomain;
 class UPoseSearchDatabase;
@@ -34,6 +38,33 @@ namespace UE::PoseSearch
 } // namespace UE::PoseSearch
 
 
+UCLASS()
+class UPoseSearchMeshComponent : public UPoseableMeshComponent
+{
+	GENERATED_BODY()
+public:
+
+	struct FUpdateContext
+	{
+		const UAnimSequence* Sequence = nullptr;
+		float SequenceStartTime = 0.0f;
+		float SequenceTime = 0.0f;
+		bool bLoop = false;
+		bool bMirrored = false;
+		const UMirrorDataTable* MirrorDataTable = nullptr;
+		TCustomBoneIndexArray<FCompactPoseBoneIndex, FCompactPoseBoneIndex>* CompactPoseMirrorBones = nullptr;
+		TCustomBoneIndexArray<FQuat, FCompactPoseBoneIndex>* ComponentSpaceRefRotations = nullptr;
+	};
+
+	void Refresh();
+	void ResetToStart();
+	void UpdatePose(const FUpdateContext& UpdateContext);
+	void Initialize(const FTransform& InComponentToWorld);
+	FTransform StartingTransform = FTransform::Identity;
+	FTransform LastRootMotionDelta = FTransform::Identity;
+};
+
+
 /**
  * Used by the reflection UObject to encompass a set of feature vectors
  */
@@ -146,6 +177,12 @@ public:
 	UPROPERTY(EditAnywhere, Category="Draw Options", Meta=(DisplayName="Selected Pose"))
 	FPoseSearchDebuggerFeatureDrawOptions SelectedPoseDrawOptions;
 
+	UPROPERTY(EditAnywhere, Category = "Draw Options")
+	bool bDrawActiveSkeleton = true;
+	
+	UPROPERTY(EditAnywhere, Category = "Draw Options")
+	bool bDrawSelectedSkeleton = true;
+
     UPROPERTY(VisibleAnywhere, Category="Pose Vectors")
 	FPoseSearchDebuggerPoseVector QueryPoseVector;
     	
@@ -162,7 +199,24 @@ public:
 
 namespace UE::PoseSearch {
 
-DECLARE_DELEGATE(FOnPoseSelectionChanged)
+/** Draw flags for the view's debug draw */
+enum class ESkeletonDrawFlags : uint32
+{
+	None              = 0,
+	ActivePose	  = 1 << 0,
+	SelectedPose  = 1 << 1,
+	AnimSequence  = 1 << 2
+};
+ENUM_CLASS_FLAGS(ESkeletonDrawFlags)
+
+struct FSkeletonDrawParams
+{
+	ESkeletonDrawFlags Flags = ESkeletonDrawFlags::None;
+};
+
+/** Sets model selection data on row selection */
+DECLARE_DELEGATE_TwoParams(FOnPoseSelectionChanged, int32, float)
+class FDebuggerViewModel;
 
 /**
  * Database panel view widget of the PoseSearch debugger
@@ -221,6 +275,7 @@ private:
 	/** Called when the text in the filter box is modified to update the filtering */
 	void OnFilterTextChanged(const FText& SearchText);
 
+	/** Row selection to update model view */
 	void OnDatabaseRowSelectionChanged(TSharedPtr<FDebuggerDatabaseRowData> Row, ESelectInfo::Type SelectInfo);
 
 	/** Generates a database row widget for the given data */
@@ -289,14 +344,15 @@ class SDebuggerDetailsView : public SCompoundWidget
 		SLATE_ARGUMENT(TWeakPtr<class SDebuggerView>, Parent)
 	SLATE_END_ARGS()
 
+	SDebuggerDetailsView() = default;
+	virtual ~SDebuggerDetailsView() override;
+	
 	void Construct(const FArguments& InArgs);
 	void Update(const FTraceMotionMatchingStateMessage& State, const UPoseSearchDatabase& Database) const;
 
 	/** Get a const version of our reflection object */
 	const TObjectPtr<UPoseSearchDebuggerReflection>& GetReflection() const { return Reflection; }
 	
-	virtual ~SDebuggerDetailsView() override;
-
 private:
 	/** Update our details view object with new state information */
 	void UpdateReflection(const FTraceMotionMatchingStateMessage& State, const UPoseSearchDatabase& Database) const;
@@ -309,16 +365,6 @@ private:
 	/** Last updated reflection data relative to MM state */
 	TObjectPtr<UPoseSearchDebuggerReflection> Reflection = nullptr;
 };
-
-/** Callback to update the debugger when node is actively selected */
-DECLARE_DELEGATE_OneParam(FOnUpdateSelection, int32 NodeId);
-
-/**
- * Callback to update the debugger when view update occurs, updating the motion matching state
- * relative to the selected anim instance.
- */
-DECLARE_DELEGATE_OneParam(FOnUpdate, uint64 AnimInstanceId);
-
 /** Callback to relay closing of the view to destroy the debugger instance */
 DECLARE_DELEGATE_OneParam(FOnViewClosed, uint64 AnimInstanceId);
 
@@ -328,33 +374,21 @@ DECLARE_DELEGATE_OneParam(FOnViewClosed, uint64 AnimInstanceId);
 class SDebuggerView : public IRewindDebuggerView
 {
 public:
-	SDebuggerView() = default;
-
 	SLATE_BEGIN_ARGS(SDebuggerView){}
-		SLATE_ATTRIBUTE(const FTraceMotionMatchingStateMessage*, MotionMatchingState)
-		SLATE_ATTRIBUTE(const UPoseSearchDatabase*, PoseSearchDatabase)
-		SLATE_ATTRIBUTE(const TArray<int32>*, MotionMatchingNodeIds)
-		SLATE_ATTRIBUTE(bool, IsPIESimulating)
-		SLATE_ATTRIBUTE(bool, IsRecording)
-		SLATE_ATTRIBUTE(double, RecordingDuration)
-		SLATE_ATTRIBUTE(int32, NodesNum)
-		SLATE_ATTRIBUTE(const UWorld*, World)
-		SLATE_ATTRIBUTE(const FTransform*, RootTransform)
-		SLATE_EVENT(FOnUpdateSelection, OnUpdateSelection)
-		SLATE_EVENT(FOnUpdate, OnUpdate)
+		SLATE_ATTRIBUTE(TSharedPtr<FDebuggerViewModel>, ViewModel)
 		SLATE_EVENT(FOnViewClosed, OnViewClosed)
 	SLATE_END_ARGS()
 
+	SDebuggerView() = default;
+    virtual ~SDebuggerView() override;
+
 	void Construct(const FArguments& InArgs, uint64 InAnimInstanceId);
 	virtual void SetTimeMarker(double InTimeMarker) override;
 	virtual void Tick(const FGeometry& AllottedGeometry, const double InCurrentTime, const float InDeltaTime) override;
 	virtual FName GetName() const override;
 	virtual uint64 GetObjectId() const override;
-
 	TArray<TSharedRef<FDebuggerDatabaseRowData>> GetSelectedDatabaseRows() const;
 
-	virtual ~SDebuggerView() override;
-
 private:
 	/** Called each frame to draw features of the query vector & database selections */
 	void DrawFeatures(const UWorld& DebuggerWorld, const FTraceMotionMatchingStateMessage& State, const UPoseSearchDatabase& Database, const FTransform& Transform) const;
@@ -373,49 +407,22 @@ private:
 	/** Callback when a button in the selection view is clicked */
 	FReply OnUpdateNodeSelection(int32 InSelectedNodeId);
 
-	void OnPoseSelectionChanged();
+	void OnPoseSelectionChanged(int32 PoseIdx, float Time);
+
+	/** Button interaction to toggle play / stop of the anim sequence */
+	FReply TogglePlaySelectedSequences() const;
 
 	/** Generates the message view relaying that there is no data */
 	TSharedRef<SWidget> GenerateNoDataMessageView();
 
 	/** Generates the return button to go back to the selection mode */
-	TSharedRef<SWidget> GenerateReturnButtonView();
+	TSharedRef<SHorizontalBox> GenerateReturnButtonView();
 
 	/** Generates the entire node debugger widget, including database and details view */
 	TSharedRef<SWidget> GenerateNodeDebuggerView();
 
-	/** Gets all MM nodes being traced in this frame */
-	TAttribute<const TArray<int32>*> MotionMatchingNodeIds;
-	
-	/** Retrieves the MM state from the debugger */
-	TAttribute<const FTraceMotionMatchingStateMessage*> MotionMatchingState;
-	
-	/** Retrieves the PoseSearch Database from the debugger */
-	TAttribute<const UPoseSearchDatabase*> PoseSearchDatabase;
-	
-	/** Whether the game is un-paused and currently simulating */
-	TAttribute<bool> IsPIESimulating;
-	
-	/** Whether the Rewind Debugger is currently recording gameplay */
-	TAttribute<bool> IsRecording;
-	
- 	/** Length of the recorded sequence (if any) */
-	TAttribute<double> RecordingDuration;
-	
- 	/** Number of active nodes at the current frame */
-	TAttribute<int32> NodesNum;
-	
-	/** Current world the Rewind Debugger is functioning in */
-	TAttribute<const UWorld*> World;
-	
-	/** Current Component to World Space transform of the Skeletal Mesh Component */
-	TAttribute<const FTransform*> RootTransform;
-	
-	/** Update current debugger data when node selection is changed */
-	FOnUpdateSelection OnUpdateSelection;
-	
-	/** Update current debugger data when update occurs */
-	FOnUpdate OnUpdate;
+	/** Pointer to the debugger instance / model for this view */
+	TAttribute<TSharedPtr<FDebuggerViewModel>> ViewModel;
 	
 	/** Destroy the debugger instanced when closed */
 	FOnViewClosed OnViewClosed;
@@ -473,12 +480,12 @@ private:
 	static constexpr int32 ConsecutiveFramesUpdateThreshold = 10;
 };
 
-class FDebuggerInstance : public TSharedFromThis<FDebuggerInstance>
+class FDebuggerViewModel : public TSharedFromThis<FDebuggerViewModel>
 {
 public:
-	explicit FDebuggerInstance(uint64 InAnimInstanceId);
+	explicit FDebuggerViewModel(uint64 InAnimInstanceId);
+	virtual ~FDebuggerViewModel();
 
-private:
 	// Used for view callbacks
     const FTraceMotionMatchingStateMessage* GetMotionMatchingState() const;
     const UPoseSearchDatabase* GetPoseSearchDatabase() const;
@@ -487,13 +494,50 @@ private:
 	const FTransform* GetRootTransform() const;
 
 	/** Update motion matching states for frame */
-	void OnUpdate(uint64 InAnimInstanceId);
+	void OnUpdate();
 	
 	/** Updates active motion matching state based on node selection */
-	void OnUpdateSelection(int32 InNodeId);
+	void OnUpdateNodeSelection(int32 InNodeId);
+
+	/** Updates internal Skeletal Mesh Component depending on input */
+	void OnDraw(FSkeletonDrawParams& DrawParams);
+	
+	/** Get an animation sequence from the sequence ID of the active database */
+	const FPoseSearchDatabaseSequence* GetAnimSequence(int32 SequenceIdx) const;
+	
+	/** Sets the selected pose skeleton*/
+	void ShowSelectedSkeleton(int32 PoseIdx, float Time);
+	
+	/** Clears the selected pose skeleton */
+	void ClearSelectedSkeleton();
+	
+	/** Plays the selected row upon button press */
+	void PlaySelection(int32 PoseIdx, float Time);
+
+	/** Stops the playing selection upon button press */
+	void StopSelection();
+	
+	/** If there is a sequence selection playing */
+	bool IsPlayingSelections() const { return SequenceData.bActive; }
+	
+	/** Play Rate of the sequence selection player */
+	void ChangePlayRate(float PlayRate) { SequencePlayRate = PlayRate; }
+	
+	/** Current play rate of the sequence selection player */
+	float GetPlayRate() const { return SequencePlayRate; }
+	
+	/** Callback to reset debug skeletons for the active world */
+	void OnWorldCleanup(UWorld* InWorld, bool bSessionEnded, bool bCleanupResources);
+
+	/** Updates the current playing sequence */
+	void UpdateAnimSequence();
+
+private:
+	/** Update the list of states for this frame */
+	void UpdateFromTimeline();
 
-	/** Update the list of states for this frame */	
-	void UpdateMotionMatchingStates(uint64 InAnimInstanceId);
+	/** Populates arrays used for mirroring the animation pose */
+	void FillCompactPoseAndComponentRefRotations();
 	
 	/** List of all Node IDs associated with motion matching states */
 	TArray<int32> NodeIds;
@@ -504,12 +548,79 @@ private:
 	/** Currently active motion matching state based on node selection in the view */
 	const FTraceMotionMatchingStateMessage* ActiveMotionMatchingState = nullptr;
 
-	/** Stored Rewind Debugger instance provided by the extension */
-	const IRewindDebugger* RewindDebugger = nullptr;
+	/** Current Skeletal Mesh Component Id for the AnimInstance */
+	uint64 SkeletalMeshComponentId = 0;
+
+	/** Currently active root transform on the skeletal mesh */
+	const FTransform* RootTransform = nullptr;
+
+	/** Pointer to the active rewind debugger in the scene */
+	TAttribute<const IRewindDebugger*> RewindDebugger;
 
 	/** Anim Instance associated with this debugger instance */
 	uint64 AnimInstanceId = 0;
 
+	/** Compact pose format of Mirror Bone Map */
+	TCustomBoneIndexArray<FCompactPoseBoneIndex, FCompactPoseBoneIndex> CompactPoseMirrorBones;
+
+	/** Pre-calculated component space rotations of reference pose */
+	TCustomBoneIndexArray<FQuat, FCompactPoseBoneIndex> ComponentSpaceRefRotations;
+
+	/** Debug visualization skeleton actor */
+	struct FSkeleton
+	{
+		/** Actor object for the skeleton */
+		TWeakObjectPtr<AActor> Actor = nullptr;
+
+		/** Derived skeletal mesh for setting the skeleton in the scene */
+		UPoseSearchMeshComponent* Component = nullptr;
+	
+		/** Active sequence index being used for this setting this skeleton */
+		int32 SequenceIdx = 0;
+		
+		/** Time in the sequence this skeleton is accessing */
+		float Time = 0.0f;
+
+		bool bMirrored = false;
+	};
+	
+	/** Index for each type of skeleton we store for debug visualization */
+	enum ESkeletonIndex
+	{
+		ActivePose = 0,
+		SelectedPose,
+		AnimSequence,
+
+		Num
+	};
+
+	/** Skeleton container for each type */
+	TArray<FSkeleton, TFixedAllocator<ESkeletonIndex::Num>> Skeletons;
+
+	/** Whether the skeletons have been initialized for this world */
+	bool bSkeletonsInitialized = false;
+	
+	/** If we currently have a selection active in the view */
+	bool bSelecting = false;
+	
+	/** Data of the active playing sequence */
+	struct FSequence
+	{
+		/** How long to stop upon reaching target */
+		static constexpr float StopDuration = 2.0f;
+		/** Time since the start of play */
+		float AccumulatedTime = 0.0f;
+
+		/** Start time of the sequence */
+		float StartTime = 0.0f;
+		/** If we are currently playing a sequence */
+		bool bActive = false;
+
+	} SequenceData;
+	
+	/** Current play rate of the sequence selection player */
+	float SequencePlayRate = 1.0f;
+
 	/** Limits some public API */
 	friend class FDebugger;
 };
@@ -529,21 +640,30 @@ public:
 	static void Shutdown();
 	static const FName ModularFeatureName;
 
+	// Shared data from the Rewind Debugger singleton
+	static bool IsPIESimulating();
+	static bool IsRecording();
+	static double GetRecordingDuration();
+	static UWorld* GetWorld();
+	static const IRewindDebugger* GetRewindDebugger();
+
 	/** Generates the slate debugger view widget */
 	TSharedPtr<SDebuggerView> GenerateInstance(uint64 InAnimInstanceId);
+
 private:
-	bool IsPIESimulating() const;
-	bool IsRecording() const;
-	double GetRecordingDuration() const;
-	const UWorld* GetWorld() const;
-	void OnViewClosed(uint64 InAnimInstanceId);
+	/** Removes the reference from the model array when closed, destroying the model */
+	static void OnViewClosed(uint64 InAnimInstanceId);
+
+	/** Acquire view model from the array */
+	static TSharedPtr<FDebuggerViewModel> GetViewModel(uint64 InAnimInstanceId);
 
 	/** Last stored Rewind Debugger */
 	const IRewindDebugger* RewindDebugger = nullptr;
 
 	/** List of all active debugger instances */
-	TArray<TSharedRef<FDebuggerInstance>> DebuggerInstances;
+	TArray<TSharedRef<FDebuggerViewModel>> ViewModels;
 	
+	/** Internal instance */
 	static FDebugger* Debugger;
 };
 
