PoseSearch: Never jump to a candidate pose within the same looping animation regardless of how far away it is

diff --git a/Engine/Plugins/Experimental/Animation/PoseSearch/Source/Runtime/Private/PoseSearchLibrary.cpp b/Engine/Plugins/Experimental/Animation/PoseSearch/Source/Runtime/Private/PoseSearchLibrary.cpp
--- a/Engine/Plugins/Experimental/Animation/PoseSearch/Source/Runtime/Private/PoseSearchLibrary.cpp
+++ b/Engine/Plugins/Experimental/Animation/PoseSearch/Source/Runtime/Private/PoseSearchLibrary.cpp
@@ -149,19 +149,23 @@ void UpdateMotionMatchingState(const FAnimationUpdateContext& Context
 				// Consider the search result better if it is more similar to the query than the current pose we're playing back from the database
 				bool bBetterPose = Result.Dissimilarity * (1.0f + (Settings.MinPercentImprovement / 100.0f)) < CurrentDissimilarity;
 
-				// We'll ignore the candidate pose if it is too near to our current pose
+				// We'll ignore the candidate pose if it is from the same animation and too near to the current pose
 				bool bNearbyPose = false;
-				const FPoseSearchIndexAsset* StateSearchIndexAsset = 
-					InOutMotionMatchingState.GetCurrentSearchIndexAsset();
+				const FPoseSearchIndexAsset* StateSearchIndexAsset =  InOutMotionMatchingState.GetCurrentSearchIndexAsset();
 				if (StateSearchIndexAsset == Result.SearchIndexAsset)
 				{
 					const FPoseSearchDatabaseSequence& ResultDbSequence = Database->GetSourceAsset(Result.SearchIndexAsset);
-					bNearbyPose = FMath::Abs(InOutMotionMatchingState.AssetPlayerTime - Result.TimeOffsetSeconds) < Settings.PoseJumpThresholdTime;
-					if (!bNearbyPose && ResultDbSequence.bLoopAnimation)
+
+					// Consider the candidate pose nearby if the animation loops, regardless of how far away the pose is
+					if (ResultDbSequence.bLoopAnimation)
+					{
+						bNearbyPose = true;
+					}
+
+					// Otherwise consider the pose nearby if it is within the pose jump threshold
+					else
 					{
-						const FPoseSearchDatabaseSequence& StateDbSequence = Database->GetSourceAsset(StateSearchIndexAsset);
-						const float AssetLength = StateDbSequence.Sequence->GetPlayLength();
-						bNearbyPose = FMath::Abs(AssetLength - InOutMotionMatchingState.AssetPlayerTime - Result.TimeOffsetSeconds) < Settings.PoseJumpThresholdTime;
+						bNearbyPose = FMath::Abs(InOutMotionMatchingState.AssetPlayerTime - Result.TimeOffsetSeconds) < Settings.PoseJumpThresholdTime;
 					}
 				}
 
