MotionMatchingFFDemo - pose search debugger now shows simulation and animation velocities

diff --git a/Engine/Plugins/Experimental/Animation/PoseSearch/Source/Editor/Private/PoseSearchDebugger.cpp b/Engine/Plugins/Experimental/Animation/PoseSearch/Source/Editor/Private/PoseSearchDebugger.cpp
--- a/Engine/Plugins/Experimental/Animation/PoseSearch/Source/Editor/Private/PoseSearchDebugger.cpp
+++ b/Engine/Plugins/Experimental/Animation/PoseSearch/Source/Editor/Private/PoseSearchDebugger.cpp
@@ -1068,7 +1068,20 @@ void SDebuggerDetailsView::UpdateReflection(const FTraceMotionMatchingStateMessa
 	Reflection->CurrentDatabaseName = Database.GetName();
 	Reflection->ElapsedPoseJumpTime = State.ElapsedPoseJumpTime;
 	Reflection->bFollowUpAnimation = EnumHasAnyFlags(State.Flags, FTraceMotionMatchingState::EFlags::FollowupAnimation);
-	
+
+	Reflection->AssetPlayerSequenceName = FString();
+	if (const FPoseSearchIndexAsset* IndexAsset = Database.SearchIndex.FindAssetForPose(State.DbPoseIdx))
+	{
+		Reflection->AssetPlayerSequenceName = Database.GetSourceAsset(IndexAsset).Sequence->GetName();
+	}
+
+	Reflection->AssetPlayerTime = State.AssetPlayerTime;
+	Reflection->LastDeltaTime = State.DeltaTime;
+	Reflection->SimLinearVelocity = State.SimLinearVelocity;
+	Reflection->SimAngularVelocity = State.SimAngularVelocity;
+	Reflection->AnimLinearVelocity = State.AnimLinearVelocity;
+	Reflection->AnimAngularVelocity = State.AnimAngularVelocity;
+
 	FFeatureVectorReader Reader;
 	Reader.Init(&Schema->Layout);
 
diff --git a/Engine/Plugins/Experimental/Animation/PoseSearch/Source/Editor/Private/PoseSearchDebugger.h b/Engine/Plugins/Experimental/Animation/PoseSearch/Source/Editor/Private/PoseSearchDebugger.h
--- a/Engine/Plugins/Experimental/Animation/PoseSearch/Source/Editor/Private/PoseSearchDebugger.h
+++ b/Engine/Plugins/Experimental/Animation/PoseSearch/Source/Editor/Private/PoseSearchDebugger.h
@@ -4,6 +4,7 @@
 
 #include "CoreMinimal.h"
 #include "Components/PoseableMeshComponent.h"
+#include "CustomBoneIndexArray.h"
 #include "RewindDebuggerInterface/Public/IRewindDebuggerExtension.h"
 #include "RewindDebuggerInterface/Public/IRewindDebuggerView.h"
 #include "RewindDebuggerInterface/Public/IRewindDebuggerViewCreator.h"
@@ -170,6 +171,27 @@ public:
 	/** Whether it is playing the loop following the expended animation runway */
 	UPROPERTY(VisibleAnywhere, Category="Motion Matching State")
 	bool bFollowUpAnimation = false;
+	
+	UPROPERTY(VisibleAnywhere, Category = "Motion Matching State", Meta = (DisplayName = "Asset Player Sequence"))
+	FString AssetPlayerSequenceName = "";
+
+	UPROPERTY(VisibleAnywhere, Category = "Motion Matching State")
+	float AssetPlayerTime = 0.0f;
+
+	UPROPERTY(VisibleAnywhere, Category = "Motion Matching State")
+	float LastDeltaTime = 0.0f;
+
+	UPROPERTY(VisibleAnywhere, Category = "Motion Matching State")
+	float SimLinearVelocity;
+
+	UPROPERTY(VisibleAnywhere, Category = "Motion Matching State")
+	float SimAngularVelocity;
+
+	UPROPERTY(VisibleAnywhere, Category = "Motion Matching State")
+	float AnimLinearVelocity;
+
+	UPROPERTY(VisibleAnywhere, Category = "Motion Matching State")
+	float AnimAngularVelocity;
 
 	UPROPERTY(EditAnywhere, Category="Draw Options", Meta=(DisplayName="Query"))
 	FPoseSearchDebuggerFeatureDrawOptions QueryDrawOptions;
diff --git a/Engine/Plugins/Experimental/Animation/PoseSearch/Source/Editor/Private/Trace/PoseSearchTraceAnalyzer.cpp b/Engine/Plugins/Experimental/Animation/PoseSearch/Source/Editor/Private/Trace/PoseSearchTraceAnalyzer.cpp
--- a/Engine/Plugins/Experimental/Animation/PoseSearch/Source/Editor/Private/Trace/PoseSearchTraceAnalyzer.cpp
+++ b/Engine/Plugins/Experimental/Animation/PoseSearch/Source/Editor/Private/Trace/PoseSearchTraceAnalyzer.cpp
@@ -48,6 +48,13 @@ bool FTraceAnalyzer::OnEvent(uint16 RouteId, EStyle Style, const FOnEventContext
 				Message.QueryVector = EventData.GetArrayView<float>("QueryVector");
                 Message.QueryVectorNormalized = EventData.GetArrayView<float>("QueryVectorNormalized");
 
+				Message.AssetPlayerTime = EventData.GetValue<float>("AssetPlayerTime");
+				Message.DeltaTime = EventData.GetValue<float>("DeltaTime");
+				Message.SimLinearVelocity = EventData.GetValue<float>("SimLinearVelocity");
+				Message.SimAngularVelocity = EventData.GetValue<float>("SimAngularVelocity");
+				Message.AnimLinearVelocity = EventData.GetValue<float>("AnimLinearVelocity");
+				Message.AnimAngularVelocity = EventData.GetValue<float>("AnimAngularVelocity");
+
 				TArrayView<const float> ChannelWeightScales = EventData.GetArrayView<float>("ChannelWeightScales");
 				TArrayView<const float> HistoryWeightScales = EventData.GetArrayView<float>("HistoryWeightScales");
 				TArrayView<const float> PredictionWeightScales = EventData.GetArrayView<float>("PredictionWeightScales");
diff --git a/Engine/Plugins/Experimental/Animation/PoseSearch/Source/Editor/Private/Trace/PoseSearchTraceProvider.h b/Engine/Plugins/Experimental/Animation/PoseSearch/Source/Editor/Private/Trace/PoseSearchTraceProvider.h
--- a/Engine/Plugins/Experimental/Animation/PoseSearch/Source/Editor/Private/Trace/PoseSearchTraceProvider.h
+++ b/Engine/Plugins/Experimental/Animation/PoseSearch/Source/Editor/Private/Trace/PoseSearchTraceProvider.h
@@ -30,6 +30,12 @@ struct FTraceMotionMatchingStateMessage : FTraceMessage
 	FPoseSearchDynamicWeightParams Weights;
 	int32 DbPoseIdx;
 	uint64 DatabaseId;
+	float AssetPlayerTime = 0.0f;
+	float DeltaTime = 0.0f;
+	float SimLinearVelocity = 0.0f;
+	float SimAngularVelocity = 0.0f;
+	float AnimLinearVelocity = 0.0f;
+	float AnimAngularVelocity = 0.0f;
 };
 
 
diff --git a/Engine/Plugins/Experimental/Animation/PoseSearch/Source/Runtime/Private/PoseSearchLibrary.cpp b/Engine/Plugins/Experimental/Animation/PoseSearch/Source/Runtime/Private/PoseSearchLibrary.cpp
--- a/Engine/Plugins/Experimental/Animation/PoseSearch/Source/Runtime/Private/PoseSearchLibrary.cpp
+++ b/Engine/Plugins/Experimental/Animation/PoseSearch/Source/Runtime/Private/PoseSearchLibrary.cpp
@@ -58,6 +58,52 @@ static void RequestInertialBlend(const FAnimationUpdateContext& Context, float B
 	}
 }
 
+void ComputeCurrentVelocityMagnitudes(
+	const FMotionMatchingState& InOutMotionMatchingState, 
+	const FTrajectorySampleRange& Trajectory, 
+	const float DeltaTime, 
+	float& SimLinearVelocity, float& SimAngularVelocity, 
+	float& AnimLinearVelocity, float& AnimAngularVelocity)
+{
+	if (DeltaTime > SMALL_NUMBER)
+	{
+		// simulation
+
+		int32 FirstIdx = 0;
+		const FTrajectorySample PrevSample = FTrajectorySampleRange::IterSampleTrajectory(
+			Trajectory.Samples,
+			ETrajectorySampleDomain::Time,
+			-DeltaTime, FirstIdx);
+
+		const FTrajectorySample CurrSample = FTrajectorySampleRange::IterSampleTrajectory(
+			Trajectory.Samples,
+			ETrajectorySampleDomain::Time,
+			0.0f, FirstIdx);
+
+		const FTransform RelativeMotion = CurrSample.Transform.GetRelativeTransform(PrevSample.Transform);
+
+		SimLinearVelocity = RelativeMotion.GetTranslation().Size() / DeltaTime;
+		SimAngularVelocity = FMath::RadiansToDegrees(RelativeMotion.GetRotation().GetAngle()) / DeltaTime;
+
+		// animation
+		if (const FPoseSearchIndexAsset* IndexAsset = InOutMotionMatchingState.GetCurrentSearchIndexAsset())
+		{
+			const FPoseSearchDatabaseSequence& DbSequence =
+				InOutMotionMatchingState.CurrentDatabase->GetSourceAsset(IndexAsset);
+			const UAnimSequence* SourceSequence = DbSequence.Sequence;
+
+			const FTransform AnimRootMotion = SourceSequence->ExtractRootMotion(
+				InOutMotionMatchingState.AssetPlayerTime,
+				-DeltaTime,
+				DbSequence.bLoopAnimation);
+
+			AnimLinearVelocity = AnimRootMotion.GetTranslation().Size() / DeltaTime;
+			AnimAngularVelocity = FMath::RadiansToDegrees(AnimRootMotion.GetRotation().GetAngle()) / DeltaTime;
+		}
+
+	}
+}
+
 void UpdateMotionMatchingState(const FAnimationUpdateContext& Context
 	, const UPoseSearchDatabase* Database
 	, const FTrajectorySampleRange& Trajectory
@@ -199,6 +245,18 @@ void UpdateMotionMatchingState(const FAnimationUpdateContext& Context
 #if UE_POSE_SEARCH_TRACE_ENABLED
 	if (InOutMotionMatchingState.DbPoseIdx != INDEX_NONE)
 	{
+
+		float SimLinearVelocity = -1.0f;
+		float SimAngularVelocity = -1.0f;
+		float AnimLinearVelocity = -1.0f;
+		float AnimAngularVelocity = -1.0f;
+		ComputeCurrentVelocityMagnitudes(
+			InOutMotionMatchingState, 
+			Trajectory, 
+			DeltaTime, 
+			SimLinearVelocity, SimAngularVelocity, 
+			AnimLinearVelocity, AnimAngularVelocity);
+
 		UE::PoseSearch::FTraceMotionMatchingState TraceState;
 		if ((InOutMotionMatchingState.Flags & EMotionMatchingFlags::JumpedToPose) == EMotionMatchingFlags::JumpedToPose)
 		{
@@ -212,6 +270,13 @@ void UpdateMotionMatchingState(const FAnimationUpdateContext& Context
 		TraceState.Weights = Settings.Weights;
 		TraceState.DbPoseIdx = InOutMotionMatchingState.DbPoseIdx;
 		TraceState.DatabaseId = FObjectTrace::GetObjectId(Database);
+
+		TraceState.AssetPlayerTime = InOutMotionMatchingState.AssetPlayerTime;
+		TraceState.DeltaTime = DeltaTime;
+		TraceState.SimLinearVelocity = SimLinearVelocity;
+		TraceState.SimAngularVelocity = SimAngularVelocity;
+		TraceState.AnimLinearVelocity = AnimLinearVelocity;
+		TraceState.AnimAngularVelocity = AnimAngularVelocity;
 		UE_TRACE_POSE_SEARCH_MOTION_MATCHING_STATE(Context, TraceState)
 	}
 #endif
diff --git a/Engine/Plugins/Experimental/Animation/PoseSearch/Source/Runtime/Private/Trace/PoseSearchTraceLogger.cpp b/Engine/Plugins/Experimental/Animation/PoseSearch/Source/Runtime/Private/Trace/PoseSearchTraceLogger.cpp
--- a/Engine/Plugins/Experimental/Animation/PoseSearch/Source/Runtime/Private/Trace/PoseSearchTraceLogger.cpp
+++ b/Engine/Plugins/Experimental/Animation/PoseSearch/Source/Runtime/Private/Trace/PoseSearchTraceLogger.cpp
@@ -24,6 +24,12 @@ UE_TRACE_EVENT_BEGIN(PoseSearch, MotionMatchingState)
 	UE_TRACE_EVENT_FIELD(float[], HistoryWeightScales)
 	UE_TRACE_EVENT_FIELD(float[], PredictionWeightScales)
 	UE_TRACE_EVENT_FIELD(bool, DebugDisableWeights)
+	UE_TRACE_EVENT_FIELD(float, AssetPlayerTime)
+	UE_TRACE_EVENT_FIELD(float, DeltaTime)
+	UE_TRACE_EVENT_FIELD(float, SimLinearVelocity)
+	UE_TRACE_EVENT_FIELD(float, SimAngularVelocity)
+	UE_TRACE_EVENT_FIELD(float, AnimLinearVelocity)
+	UE_TRACE_EVENT_FIELD(float, AnimAngularVelocity)
 UE_TRACE_EVENT_END()
 
 namespace UE { namespace PoseSearch {
@@ -91,7 +97,13 @@ void FTraceMotionMatchingState::Output(const FAnimationBaseContext& InContext, c
 		<< MotionMatchingState.ChannelWeightScales(ChannelWeightScales, UE_ARRAY_COUNT(ChannelWeightScales))
 		<< MotionMatchingState.HistoryWeightScales(HistoryWeightScales, UE_ARRAY_COUNT(HistoryWeightScales))
 		<< MotionMatchingState.PredictionWeightScales(PredictionWeightScales, UE_ARRAY_COUNT(PredictionWeightScales))
-		<< MotionMatchingState.DebugDisableWeights(State.Weights.bDebugDisableWeights);
+		<< MotionMatchingState.DebugDisableWeights(State.Weights.bDebugDisableWeights)
+		<< MotionMatchingState.AssetPlayerTime(State.AssetPlayerTime)
+		<< MotionMatchingState.DeltaTime(State.DeltaTime)
+		<< MotionMatchingState.SimLinearVelocity(State.SimLinearVelocity)
+		<< MotionMatchingState.SimAngularVelocity(State.SimAngularVelocity)
+		<< MotionMatchingState.AnimLinearVelocity(State.AnimLinearVelocity)
+		<< MotionMatchingState.AnimAngularVelocity(State.AnimAngularVelocity);
 }
 
 }}
diff --git a/Engine/Plugins/Experimental/Animation/PoseSearch/Source/Runtime/Private/Trace/PoseSearchTraceLogger.h b/Engine/Plugins/Experimental/Animation/PoseSearch/Source/Runtime/Private/Trace/PoseSearchTraceLogger.h
--- a/Engine/Plugins/Experimental/Animation/PoseSearch/Source/Runtime/Private/Trace/PoseSearchTraceLogger.h
+++ b/Engine/Plugins/Experimental/Animation/PoseSearch/Source/Runtime/Private/Trace/PoseSearchTraceLogger.h
@@ -61,6 +61,14 @@ struct POSESEARCH_API FTraceMotionMatchingState
 	 */
 	uint64 SkeletalMeshComponentId = 0;
 
+	float AssetPlayerTime = 0.0f;
+	float DeltaTime = 0.0f;
+	float SimLinearVelocity = 0.0f;
+	float SimAngularVelocity = 0.0f;
+	float AnimLinearVelocity = 0.0f;
+	float AnimAngularVelocity = 0.0f;
+
+
 	/** Output the current state info to the logger */
 	static void Output(const FAnimationBaseContext& InContext, const FTraceMotionMatchingState& State);
 	
