MotionMatchingFFDemo - Debugger update to show the cost vector difference between the selected pose and the active pose. A negative value means the selected pose cost for that feature is lower. A positive value means the selected pose cost for that feature is higher.

diff --git a/Engine/Plugins/Experimental/Animation/PoseSearch/Source/Editor/Private/PoseSearchDebugger.cpp b/Engine/Plugins/Experimental/Animation/PoseSearch/Source/Editor/Private/PoseSearchDebugger.cpp
--- a/Engine/Plugins/Experimental/Animation/PoseSearch/Source/Editor/Private/PoseSearchDebugger.cpp
+++ b/Engine/Plugins/Experimental/Animation/PoseSearch/Source/Editor/Private/PoseSearchDebugger.cpp
@@ -585,6 +585,20 @@ void SDebuggerDatabaseView::Update(const FTraceMotionMatchingStateMessage& State
 	UpdateRows(State, Database);
 }
 
+const TSharedRef<FDebuggerDatabaseRowData>& SDebuggerDatabaseView::GetPoseIdxDatabaseRow(int32 PoseIdx) const
+{
+	const TSharedRef<FDebuggerDatabaseRowData>* RowPtr = UnfilteredDatabaseRows.FindByPredicate(
+		[=](TSharedRef<FDebuggerDatabaseRowData>& Row)
+	{
+		return Row->PoseIdx == PoseIdx;
+	});
+
+	check(RowPtr != nullptr);
+
+	return *RowPtr;
+}
+
+
 void SDebuggerDatabaseView::RefreshColumns()
 {
 	using namespace DebuggerDatabaseColumns;
@@ -1111,6 +1125,18 @@ void SDebuggerDetailsView::UpdateReflection(const FTraceMotionMatchingStateMessa
 			//Database.SearchIndex.InverseNormalize(Pose);
 			Reader.SetValues(Pose);
 			Reflection->CostVector.ExtractFeatures(Reader);
+
+			const TSharedRef<FDebuggerDatabaseRowData>& ActiveRow = 
+				DebuggerView->GetPoseIdxDatabaseRow(State.DbPoseIdx);
+
+			TArray<float> ActiveCostDifference(Pose);
+			for (int i = 0; i < ActiveCostDifference.Num(); ++i)
+			{
+				ActiveCostDifference[i] -= ActiveRow->PoseCostInfo.CostVector[i];
+			}
+
+			Reader.SetValues(ActiveCostDifference);
+			Reflection->CostVectorDifference.ExtractFeatures(Reader);
 		}
 	}
 }
@@ -1322,6 +1348,11 @@ TArray<TSharedRef<FDebuggerDatabaseRowData>> SDebuggerView::GetSelectedDatabaseR
 	return DatabaseView->GetDatabaseRows()->GetSelectedItems();
 }
 
+const TSharedRef<FDebuggerDatabaseRowData>& SDebuggerView::GetPoseIdxDatabaseRow(int32 PoseIdx) const
+{
+	return DatabaseView->GetPoseIdxDatabaseRow(PoseIdx);
+}
+
 void SDebuggerView::DrawFeatures(
 	const UWorld& DebuggerWorld,
 	const FTraceMotionMatchingStateMessage& State,
diff --git a/Engine/Plugins/Experimental/Animation/PoseSearch/Source/Editor/Private/PoseSearchDebugger.h b/Engine/Plugins/Experimental/Animation/PoseSearch/Source/Editor/Private/PoseSearchDebugger.h
--- a/Engine/Plugins/Experimental/Animation/PoseSearch/Source/Editor/Private/PoseSearchDebugger.h
+++ b/Engine/Plugins/Experimental/Animation/PoseSearch/Source/Editor/Private/PoseSearchDebugger.h
@@ -216,6 +216,11 @@ public:
 
 	UPROPERTY(VisibleAnywhere, Category="Pose Vectors")
 	FPoseSearchDebuggerPoseVector CostVector;
+
+	// Cost comparison of selected and active poses. A negative value indicates the cost in the selected pose is lower;
+	// a positive value indicates the cost in the selected pose is higher.
+	UPROPERTY(VisibleAnywhere, Category="Pose Vectors")
+	FPoseSearchDebuggerPoseVector CostVectorDifference;
 };
 
 
@@ -256,6 +261,7 @@ class SDebuggerDatabaseView : public SCompoundWidget
 
 	const TSharedPtr<SListView<TSharedRef<FDebuggerDatabaseRowData>>>& GetActiveRow() const { return ActiveView.ListView; }
 	const TSharedPtr<SListView<TSharedRef<FDebuggerDatabaseRowData>>>& GetDatabaseRows() const { return FilteredDatabaseView.ListView; }
+	const TSharedRef<FDebuggerDatabaseRowData>& GetPoseIdxDatabaseRow(int32 PoseIdx) const;
 
 	/** Used by database rows to acquire column-specific information */
 	using FColumnMap = TMap<FName, TSharedRef<DebuggerDatabaseColumns::IColumn>>;
@@ -410,6 +416,7 @@ public:
 	virtual FName GetName() const override;
 	virtual uint64 GetObjectId() const override;
 	TArray<TSharedRef<FDebuggerDatabaseRowData>> GetSelectedDatabaseRows() const;
+	const TSharedRef<FDebuggerDatabaseRowData>& GetPoseIdxDatabaseRow(int32 PoseIdx) const;
 
 private:
 	/** Called each frame to draw features of the query vector & database selections */
