PoseSearch: Fix pose dissimilarity ratio comparisons
- Cost modifiers were making scores negative in some cases which was breaking the MinPercentImprovement feature and causing looping animations to sometimes restart spuriously
- Now pose comparisons return an FPoseCost structure which separates pose vector dissimilarity from modifiers
- MinPercentImprovement now only considers the pose vector dissimilarity
- Removed loop reentrancy guard to reevaluate behavior with fixed pose comparison
- Renamed FMotionMatchingContinuityParams to FMotionMatchingPoseStepper and encapsulated stepping behavior
- Pulled common behavior into FMotionMatchingState::JumpToPose
- Simplified column predicates in debugger
- Fixed follow up anim trace state
- Renamed FPoseCostInfo to FPoseCostDetails

diff --git a/Engine/Plugins/Experimental/Animation/PoseSearch/Source/Editor/Private/PoseSearchDebugger.cpp b/Engine/Plugins/Experimental/Animation/PoseSearch/Source/Editor/Private/PoseSearchDebugger.cpp
--- a/Engine/Plugins/Experimental/Animation/PoseSearch/Source/Editor/Private/PoseSearchDebugger.cpp
+++ b/Engine/Plugins/Experimental/Animation/PoseSearch/Source/Editor/Private/PoseSearchDebugger.cpp
@@ -208,9 +208,9 @@ class FDebuggerDatabaseRowData : public TSharedFromThis<FDebuggerDatabaseRowData
 public:
 	FDebuggerDatabaseRowData() = default;
 	
-	float GetPoseCost() const { return PoseCostInfo.ChannelCosts.Num() ? PoseCostInfo.ChannelCosts[0] : 0.0f; }
-	float GetTrajectoryCost() const { return PoseCostInfo.ChannelCosts.Num() >= 3 ? PoseCostInfo.ChannelCosts[1] + PoseCostInfo.ChannelCosts[2] : 0.0f; }
-	float GetAddendsCost() const { return PoseCostInfo.NotifyCostAddend + PoseCostInfo.MirrorMismatchCostAddend; }
+	float GetPoseCost() const { return PoseCostDetails.ChannelCosts.Num() ? PoseCostDetails.ChannelCosts[0] : 0.0f; }
+	float GetTrajectoryCost() const { return PoseCostDetails.ChannelCosts.Num() >= 3 ? PoseCostDetails.ChannelCosts[1] + PoseCostDetails.ChannelCosts[2] : 0.0f; }
+	float GetAddendsCost() const { return PoseCostDetails.NotifyCostAddend + PoseCostDetails.MirrorMismatchCostAddend; }
 
 	int32 PoseIdx = 0;
 	FString AnimSequenceName = "";
@@ -218,7 +218,7 @@ public:
 	int32 AnimFrame = 0;
 	float Time = 0.0f;
 	bool bMirrored = false;
-	FPoseCostInfo PoseCostInfo;
+	FPoseCostDetails PoseCostDetails;
 };
 
 namespace DebuggerDatabaseColumns
@@ -246,10 +246,7 @@ namespace DebuggerDatabaseColumns
 		using FRowDataRef = TSharedRef<FDebuggerDatabaseRowData>;
 		using FSortPredicate = TFunctionRef<bool(const FRowDataRef&, const FRowDataRef&)>;
 		
-		/** Sort predicate to sort list in ascending order by this column */
-		virtual const FSortPredicate& GetSortPredicateAscending() const = 0;
-		/** Sort predicate to sort list in descending order by this column */
-		virtual const FSortPredicate& GetSortPredicateDescending() const = 0;
+		virtual FSortPredicate GetSortPredicate() const = 0;
 		
 		TSharedRef<STextBlock> GenerateTextWidget(const FRowDataRef& RowData) const
 		{
@@ -271,16 +268,9 @@ namespace DebuggerDatabaseColumns
 		static const FName Name;
 		virtual FName GetName() const override { return Name; }
 
-		virtual const FSortPredicate& GetSortPredicateAscending() const override
+		virtual FSortPredicate GetSortPredicate() const override
 		{
-			static FSortPredicate Predicate = [](const FRowDataRef& Row0, const FRowDataRef& Row1) -> bool { return Row0->PoseIdx < Row1->PoseIdx; };
-			return Predicate;
-		}
-		
-		virtual const FSortPredicate& GetSortPredicateDescending() const override
-		{
-			static FSortPredicate Predicate = [](const FRowDataRef& Row0, const FRowDataRef& Row1) -> bool { return Row0->PoseIdx >= Row1->PoseIdx; };
-			return Predicate;
+			return [](const FRowDataRef& Row0, const FRowDataRef& Row1) -> bool { return Row0->PoseIdx < Row1->PoseIdx; };
 		}
 
 		virtual FText GetRowText(const FRowDataRef& Row) const override
@@ -296,18 +286,12 @@ namespace DebuggerDatabaseColumns
 		static const FName Name;
 		virtual FName GetName() const override { return Name; }
 		
-		virtual const FSortPredicate& GetSortPredicateAscending() const override
+		virtual FSortPredicate GetSortPredicate() const override
 		{
 			static FSortPredicate Predicate = [](const FRowDataRef& Row0, const FRowDataRef& Row1) -> bool { return Row0->AnimSequenceName < Row1->AnimSequenceName; };
 			return Predicate;
 		}
 
-		virtual const FSortPredicate& GetSortPredicateDescending() const override
-		{
-			static FSortPredicate Predicate = [](const FRowDataRef& Row0, const FRowDataRef& Row1) -> bool { return Row0->AnimSequenceName >= Row1->AnimSequenceName; };
-			return Predicate;
-		}
-
 		virtual FText GetRowText(const FRowDataRef& Row) const override
 		{
 			return FText::FromString(Row->AnimSequenceName);
@@ -321,17 +305,11 @@ namespace DebuggerDatabaseColumns
 		static const FName Name;
 		virtual FName GetName() const override { return Name; }
 		
-		virtual const FSortPredicate& GetSortPredicateAscending() const override
+		virtual FSortPredicate GetSortPredicate() const override
 		{
 			static FSortPredicate Predicate = [](const FRowDataRef& Row0, const FRowDataRef& Row1) -> bool { return Row0->Time < Row1->Time; };
 			return Predicate;
 		}
-		
-		virtual const FSortPredicate& GetSortPredicateDescending() const override
-		{
-			static FSortPredicate Predicate = [](const FRowDataRef& Row0, const FRowDataRef& Row1) -> bool { return Row0->Time >= Row1->Time; };
-			return Predicate;
-		}
 
 		virtual FText GetRowText(const FRowDataRef& Row) const override
 		{
@@ -354,50 +332,36 @@ namespace DebuggerDatabaseColumns
 		static const FName Name;
 		virtual FName GetName() const override { return Name; }
 		
-		virtual const FSortPredicate& GetSortPredicateAscending() const override
+		virtual FSortPredicate GetSortPredicate() const override
 		{
-			static FSortPredicate Predicate = [](const FRowDataRef& Row0, const FRowDataRef& Row1) -> bool { return Row0->PoseCostInfo.TotalCost < Row1->PoseCostInfo.TotalCost; };
-			return Predicate;
-		}
-
-		virtual const FSortPredicate& GetSortPredicateDescending() const override
-		{
-			static FSortPredicate Predicate = [](const FRowDataRef& Row0, const FRowDataRef& Row1) -> bool { return Row0->PoseCostInfo.TotalCost >= Row1->PoseCostInfo.TotalCost; };
-			return Predicate;
+			return [](const FRowDataRef& Row0, const FRowDataRef& Row1) -> bool { return Row0->PoseCostDetails.PoseCost < Row1->PoseCostDetails.PoseCost; };
 		}
 		
 		virtual FText GetRowText(const FRowDataRef& Row) const override
         {
-        	return FText::AsNumber(Row->PoseCostInfo.TotalCost);
+        	return FText::AsNumber(Row->PoseCostDetails.PoseCost.TotalCost);
         }
 	};
 	const FName FCost::Name = "Cost";
 
 
-	struct FPoseCost : IColumn
+	struct FPoseCostColumn : IColumn
 	{
 		using IColumn::IColumn;
 		static const FName Name;
 		virtual FName GetName() const override { return Name; }
 		
-		virtual const FSortPredicate& GetSortPredicateAscending() const override
+		virtual FSortPredicate GetSortPredicate() const override
 		{
-			static FSortPredicate Predicate = [](const FRowDataRef& Row0, const FRowDataRef& Row1) -> bool { return Row0->GetPoseCost() < Row1->GetPoseCost(); };
-			return Predicate;
+			return [](const FRowDataRef& Row0, const FRowDataRef& Row1) -> bool { return Row0->GetPoseCost() < Row1->GetPoseCost(); };
 		}
 
-		virtual const FSortPredicate& GetSortPredicateDescending() const override
-		{
-			static FSortPredicate Predicate = [](const FRowDataRef& Row0, const FRowDataRef& Row1) -> bool { return Row0->GetPoseCost() >= Row1->GetPoseCost(); };
-			return Predicate;
-		}
-		
 		virtual FText GetRowText(const FRowDataRef& Row) const override
 		{
 			return FText::AsNumber(Row->GetPoseCost());
 		}
 	};
-	const FName FPoseCost::Name = "Pose Cost";
+	const FName FPoseCostColumn::Name = "Pose Cost";
 
 
 	struct FTrajectoryCost : IColumn
@@ -406,16 +370,9 @@ namespace DebuggerDatabaseColumns
 		static const FName Name;
 		virtual FName GetName() const override { return Name; }
 		
-		virtual const FSortPredicate& GetSortPredicateAscending() const override
-		{
-			static FSortPredicate Predicate = [](const FRowDataRef& Row0, const FRowDataRef& Row1) -> bool { return Row0->GetTrajectoryCost() < Row1->GetTrajectoryCost(); };
-			return Predicate;
-		}
-
-		virtual const FSortPredicate& GetSortPredicateDescending() const override
+		virtual FSortPredicate GetSortPredicate() const override
 		{
-			static FSortPredicate Predicate = [](const FRowDataRef& Row0, const FRowDataRef& Row1) -> bool { return Row0->GetTrajectoryCost() >= Row1->GetTrajectoryCost(); };
-			return Predicate;
+			return [](const FRowDataRef& Row0, const FRowDataRef& Row1) -> bool { return Row0->GetTrajectoryCost() < Row1->GetTrajectoryCost(); };
 		}
 		
 		virtual FText GetRowText(const FRowDataRef& Row) const override
@@ -431,22 +388,12 @@ namespace DebuggerDatabaseColumns
 		static const FName Name;
 		virtual FName GetName() const override { return Name; }
 
-		virtual const FSortPredicate& GetSortPredicateAscending() const override
+		virtual FSortPredicate GetSortPredicate() const override
 		{
-			static FSortPredicate Predicate = [](const FRowDataRef& Row0, const FRowDataRef& Row1) -> bool 
+			return [](const FRowDataRef& Row0, const FRowDataRef& Row1) -> bool 
 			{ 
 				return Row0->GetAddendsCost() < Row1->GetAddendsCost();
 			};
-			return Predicate;
-		}
-
-		virtual const FSortPredicate& GetSortPredicateDescending() const override
-		{
-			static FSortPredicate Predicate = [](const FRowDataRef& Row0, const FRowDataRef& Row1) -> bool 
-			{ 
-				return Row0->GetAddendsCost() >= Row1->GetAddendsCost();
-			};
-			return Predicate;
 		}
 
 		virtual FText GetRowText(const FRowDataRef& Row) const override
@@ -462,22 +409,12 @@ namespace DebuggerDatabaseColumns
 		static const FName Name;
 		virtual FName GetName() const override { return Name; }
 
-		virtual const FSortPredicate& GetSortPredicateAscending() const override
+		virtual FSortPredicate GetSortPredicate() const override
 		{
-			static FSortPredicate Predicate = [](const FRowDataRef& Row0, const FRowDataRef& Row1) -> bool 
+			return [](const FRowDataRef& Row0, const FRowDataRef& Row1) -> bool 
 			{ 
 				return Row0->bMirrored < Row1->bMirrored; 
 			};
-			return Predicate;
-		}
-
-		virtual const FSortPredicate& GetSortPredicateDescending() const override
-		{
-			static FSortPredicate Predicate = [](const FRowDataRef& Row0, const FRowDataRef& Row1) -> bool 
-			{ 
-				return Row0->bMirrored >= Row1->bMirrored; 
-			};
-			return Predicate;
 		}
 
 		virtual FText GetRowText(const FRowDataRef& Row) const override
@@ -698,11 +635,16 @@ void SDebuggerDatabaseView::SortDatabaseRows()
 {
 	if (SortMode == EColumnSortMode::Ascending)
 	{
-		UnfilteredDatabaseRows.Sort(Columns[SortColumn]->GetSortPredicateAscending());
+		UnfilteredDatabaseRows.Sort(Columns[SortColumn]->GetSortPredicate());
 	}
 	else if (SortMode == EColumnSortMode::Descending)
 	{
-		UnfilteredDatabaseRows.Sort(Columns[SortColumn]->GetSortPredicateDescending());
+		auto DescendingPredicate = [this](const auto& Lhs, const auto& Rhs) -> bool
+		{
+			return !Columns[SortColumn]->GetSortPredicate()(Lhs, Rhs);
+		};
+
+		UnfilteredDatabaseRows.Sort(DescendingPredicate);
 	}
 }
 
@@ -794,7 +736,7 @@ void SDebuggerDatabaseView::UpdateRows(const FTraceMotionMatchingStateMessage& S
 	{
 		const int32 PoseIdx = Row->PoseIdx;
 
-		ComparePoses(PoseIdx, SearchContext, Row->PoseCostInfo);
+		ComparePoses(PoseIdx, SearchContext, Row->PoseCostDetails);
 
 		// If we are on the active pose for the frame
 		if (PoseIdx == State.DbPoseIdx)
@@ -834,7 +776,7 @@ void SDebuggerDatabaseView::Construct(const FArguments& InArgs)
 	// Construct all column types
 	AddColumn(MakeShared<FAnimSequenceName>(0));
 	AddColumn(MakeShared<FCost>(1));
-	AddColumn(MakeShared<FPoseCost>(2));
+	AddColumn(MakeShared<FPoseCostColumn>(2));
 	AddColumn(MakeShared<FTrajectoryCost>(3));
 	AddColumn(MakeShared<FCostModifier>(4));
 	AddColumn(MakeShared<FFrame>(5));
@@ -1121,7 +1063,7 @@ void SDebuggerDetailsView::UpdateReflection(const FTraceMotionMatchingStateMessa
 			Reader.SetValues(Pose);
 			Reflection->SelectedPoseVector.ExtractFeatures(Reader);
 
-			Pose = Selected->PoseCostInfo.CostVector;
+			Pose = Selected->PoseCostDetails.CostVector;
 			//Database.SearchIndex.InverseNormalize(Pose);
 			Reader.SetValues(Pose);
 			Reflection->CostVector.ExtractFeatures(Reader);
@@ -1132,7 +1074,7 @@ void SDebuggerDetailsView::UpdateReflection(const FTraceMotionMatchingStateMessa
 			TArray<float> ActiveCostDifference(Pose);
 			for (int i = 0; i < ActiveCostDifference.Num(); ++i)
 			{
-				ActiveCostDifference[i] -= ActiveRow->PoseCostInfo.CostVector[i];
+				ActiveCostDifference[i] -= ActiveRow->PoseCostDetails.CostVector[i];
 			}
 
 			Reader.SetValues(ActiveCostDifference);
diff --git a/Engine/Plugins/Experimental/Animation/PoseSearch/Source/Runtime/Private/PoseSearch.cpp b/Engine/Plugins/Experimental/Animation/PoseSearch/Source/Runtime/Private/PoseSearch.cpp
--- a/Engine/Plugins/Experimental/Animation/PoseSearch/Source/Runtime/Private/PoseSearch.cpp
+++ b/Engine/Plugins/Experimental/Animation/PoseSearch/Source/Runtime/Private/PoseSearch.cpp
@@ -3993,7 +3993,7 @@ FSearchResult Search(FSearchContext& SearchContext)
 		return Result;
 	}
 
-	float BestPoseDissimilarity = MAX_flt;
+	FPoseCost BestPoseCost;
 	int32 BestPoseIdx = INDEX_NONE;
 
 	for (int32 PoseIdx = 0; PoseIdx != SearchIndex->NumPoses; ++PoseIdx)
@@ -4005,18 +4005,18 @@ FSearchResult Search(FSearchContext& SearchContext)
 			continue;
 		}
 
-		float PoseDissimilarity = ComparePoses(PoseIdx, SearchContext);
+		FPoseCost PoseCost = ComparePoses(PoseIdx, SearchContext);
 
-		if (PoseDissimilarity < BestPoseDissimilarity)
+		if (PoseCost < BestPoseCost)
 		{
-			BestPoseDissimilarity = PoseDissimilarity;
+			BestPoseCost = PoseCost;
 			BestPoseIdx = PoseIdx;
 		}
 	}
 
 	ensure(BestPoseIdx != INDEX_NONE);
 
-	Result.Dissimilarity = BestPoseDissimilarity;
+	Result.PoseCost = BestPoseCost;
 	Result.PoseIdx = BestPoseIdx;
 	Result.SearchIndexAsset = SearchIndex->FindAssetForPose(BestPoseIdx);
 	Result.TimeOffsetSeconds = SearchIndex->GetTimeOffset(BestPoseIdx, Result.SearchIndexAsset);
@@ -4053,19 +4053,20 @@ void ComputePoseCostAddends(
 	OutNotifyAddend = PoseMetadata.CostAddend;
 }
 
-float ComparePoses(int32 PoseIdx, FSearchContext& SearchContext)
+FPoseCost ComparePoses(int32 PoseIdx, FSearchContext& SearchContext)
 {
+	FPoseCost Result;
+
 	TArrayView<const float> PoseValues = SearchContext.GetSearchIndex()->GetPoseValues(PoseIdx);
 	if (!ensure(PoseValues.Num() == SearchContext.QueryValues.Num()))
 	{
-		return MAX_flt;
+		return Result;
 	}
 
-	float Dissimilarity;
 	if (SearchContext.WeightsContext)
 	{
 		const FPoseSearchWeights* WeightsSet = SearchContext.WeightsContext->GetGroupWeights(0);
-		Dissimilarity = CompareFeatureVectors(
+		Result.Dissimilarity = CompareFeatureVectors(
 			PoseValues.Num(), 
 			PoseValues.GetData(), 
 			SearchContext.QueryValues.GetData(),
@@ -4073,38 +4074,39 @@ float ComparePoses(int32 PoseIdx, FSearchContext& SearchContext)
 	}
 	else
 	{
-		Dissimilarity = 
-			CompareFeatureVectors(PoseValues.Num(), PoseValues.GetData(), SearchContext.QueryValues.GetData());
+		Result.Dissimilarity = CompareFeatureVectors(PoseValues.Num(), PoseValues.GetData(), SearchContext.QueryValues.GetData());
 	}
 
 	float NotifyAddend = 0.0f;
 	float MirrorMismatchAddend = 0.0f;
 	ComputePoseCostAddends(PoseIdx, SearchContext, NotifyAddend, MirrorMismatchAddend);
-	Dissimilarity += (NotifyAddend + MirrorMismatchAddend);
+	Result.CostAddend = NotifyAddend + MirrorMismatchAddend;
+	Result.TotalCost = Result.Dissimilarity + Result.CostAddend;
 	
-	return Dissimilarity;
+	return Result;
 }
 
-float ComparePoses(int32 PoseIdx, FSearchContext& SearchContext, FPoseCostInfo& OutPoseCostInfo)
+FPoseCost ComparePoses(int32 PoseIdx, FSearchContext& SearchContext, FPoseCostDetails& OutPoseCostDetails)
 {
 	using namespace Eigen;
 
+	FPoseCost Result;
+
 	TArrayView<const float> PoseValues = SearchContext.GetSearchIndex()->GetPoseValues(PoseIdx);
 	const int32 Dims = PoseValues.Num();
 	if (!ensure(Dims == SearchContext.QueryValues.Num()))
 	{
-		return MAX_flt;
+		return Result;
 	}
 
-	OutPoseCostInfo.CostVector.SetNum(Dims);
+	OutPoseCostDetails.CostVector.SetNum(Dims);
 
 	// Setup Eigen views onto our vectors
-	auto OutCostVector = Map<ArrayXf>(OutPoseCostInfo.CostVector.GetData(), Dims);
+	auto OutCostVector = Map<ArrayXf>(OutPoseCostDetails.CostVector.GetData(), Dims);
 	auto PoseVector = Map<const ArrayXf>(PoseValues.GetData(), Dims);
 	auto QueryVector = Map<const ArrayXf>(SearchContext.QueryValues.GetData(), Dims);
 	
 	// Compute weighted squared difference vector
-	float Dissimilarity;
 	if (SearchContext.WeightsContext)
 	{
 		const FPoseSearchWeights* WeightsSet = SearchContext.WeightsContext->GetGroupWeights(0);
@@ -4113,41 +4115,39 @@ float ComparePoses(int32 PoseIdx, FSearchContext& SearchContext, FPoseCostInfo&
 		auto WeightsVector = Map<const ArrayXf>(WeightsSet->Weights.GetData(), Dims);
 
 		OutCostVector = WeightsVector * (PoseVector - QueryVector).square();
-		Dissimilarity = OutCostVector.sum();
+		Result.Dissimilarity = OutCostVector.sum();
 	}
 	else
 	{
 		OutCostVector = (PoseVector - QueryVector).square();
-		Dissimilarity = OutCostVector.sum();
+		Result.Dissimilarity = OutCostVector.sum();
 	}
 
-	ComputePoseCostAddends(
-		PoseIdx, 
-		SearchContext, 
-		OutPoseCostInfo.NotifyCostAddend, 
-		OutPoseCostInfo.MirrorMismatchCostAddend);
+	float NotifyAddend = 0.0f;
+	float MirrorMismatchAddend = 0.0f;
+	ComputePoseCostAddends(PoseIdx, SearchContext, NotifyAddend, MirrorMismatchAddend);
+	Result.CostAddend = NotifyAddend + MirrorMismatchAddend;
+	Result.TotalCost = Result.Dissimilarity + Result.CostAddend;
 
-	const float TotalDissimilarity = 
-		Dissimilarity + 
-		OutPoseCostInfo.NotifyCostAddend + 
-		OutPoseCostInfo.MirrorMismatchCostAddend;
+	OutPoseCostDetails.NotifyCostAddend = NotifyAddend;
+	OutPoseCostDetails.MirrorMismatchCostAddend = MirrorMismatchAddend;
 
 	// Verify this math agrees with the runtime pose comparator
 	checkSlow(FMath::IsNearlyEqual(
-		TotalDissimilarity,
-		ComparePoses(PoseIdx, SearchContext), 
+		Result.TotalCost,
+		ComparePoses(PoseIdx, SearchContext).TotalCost, 
 		KINDA_SMALL_NUMBER));
 
 	
 	// Output cost details
-	OutPoseCostInfo.TotalCost = TotalDissimilarity;
-	CalcChannelCosts(SearchContext.GetSearchIndex()->Schema, OutPoseCostInfo.CostVector, OutPoseCostInfo.ChannelCosts);
+	OutPoseCostDetails.PoseCost = Result;
+	CalcChannelCosts(SearchContext.GetSearchIndex()->Schema, OutPoseCostDetails.CostVector, OutPoseCostDetails.ChannelCosts);
 
 	// Verify channel cost decomposition agrees with pose comparator
-	auto OutChannelCosts = Map<const ArrayXf>(OutPoseCostInfo.ChannelCosts.GetData(), OutPoseCostInfo.ChannelCosts.Num());
-	checkSlow(FMath::IsNearlyEqual(Dissimilarity, OutChannelCosts.sum(), KINDA_SMALL_NUMBER));
+	auto OutChannelCosts = Map<const ArrayXf>(OutPoseCostDetails.ChannelCosts.GetData(), OutPoseCostDetails.ChannelCosts.Num());
+	checkSlow(FMath::IsNearlyEqual(Result.Dissimilarity, OutChannelCosts.sum(), KINDA_SMALL_NUMBER));
 
-	return TotalDissimilarity;
+	return Result;
 }
 
 
@@ -4211,7 +4211,7 @@ UE::Anim::IPoseSearchProvider::FSearchResult FModule::Search(const FAnimationBas
 	SearchContext.QueryValues = QueryBuilder.GetNormalizedValues();
 	::UE::PoseSearch::FSearchResult Result = ::UE::PoseSearch::Search(SearchContext);
 
-	ProviderResult.Dissimilarity = Result.Dissimilarity;
+	ProviderResult.Dissimilarity = Result.PoseCost.TotalCost;
 	ProviderResult.PoseIdx = Result.PoseIdx;
 	ProviderResult.TimeOffsetSeconds = Result.TimeOffsetSeconds;
 	return ProviderResult;
diff --git a/Engine/Plugins/Experimental/Animation/PoseSearch/Source/Runtime/Private/PoseSearchLibrary.cpp b/Engine/Plugins/Experimental/Animation/PoseSearch/Source/Runtime/Private/PoseSearchLibrary.cpp
--- a/Engine/Plugins/Experimental/Animation/PoseSearch/Source/Runtime/Private/PoseSearchLibrary.cpp
+++ b/Engine/Plugins/Experimental/Animation/PoseSearch/Source/Runtime/Private/PoseSearchLibrary.cpp
@@ -10,53 +10,7 @@
 
 DEFINE_LOG_CATEGORY_STATIC(LogPoseSearchLibrary, Verbose, All);
 
-void FMotionMatchingState::InitNewDatabaseSearch(const UPoseSearchDatabase* Database, float SearchThrottleTime)
-{
-	DbPoseIdx = INDEX_NONE;
-	SearchIndexAssetIdx = INDEX_NONE;
-	ElapsedPoseJumpTime = SearchThrottleTime;
-	AssetPlayerTime = 0.0f;
-	CurrentDatabase = Database;
-}
-
-void FMotionMatchingState::ComposeQuery(const UPoseSearchDatabase* Database, const FTrajectorySampleRange& Trajectory)
-{
-	FPoseSearchFeatureVectorBuilder Goal;
-	Goal.Init(Database->Schema);
-	Goal.BuildFromTrajectory(Trajectory);
-
-	// Merge goal features into the query vector
-	if (ComposedQuery.IsCompatible(Goal))
-	{
-		ComposedQuery.MergeReplace(Goal);
-	}
-
-	ComposedQuery.Normalize(Database->SearchIndex);
-}
-
-void FMotionMatchingState::JumpToPose(const UE::PoseSearch::FSearchResult& Result)
-{
-	// Remember which pose and sequence we're playing from the database
-	DbPoseIdx = Result.PoseIdx;
-	SearchIndexAssetIdx = CurrentDatabase->SearchIndex.FindAssetIndex(Result.SearchIndexAsset);
-
-	ElapsedPoseJumpTime = 0.0f;
-	AssetPlayerTime = Result.TimeOffsetSeconds;
-}
-
-static void RequestInertialBlend(const FAnimationUpdateContext& Context, float BlendTime)
-{
-	// Use inertial blending to smooth over the transition
-	// It would be cool in the future to adjust the blend time by amount of dissimilarity, but we'll need a standardized distance metric first.
-	if (BlendTime > 0.0f)
-	{
-		UE::Anim::IInertializationRequester* InertializationRequester = Context.GetMessage<UE::Anim::IInertializationRequester>();
-		if (InertializationRequester)
-		{
-			InertializationRequester->RequestInertialization(BlendTime);
-		}
-	}
-}
+namespace UE::PoseSearch {
 
 void ComputeCurrentVelocityMagnitudes(
 	const FMotionMatchingState& InOutMotionMatchingState, 
@@ -104,13 +58,145 @@ void ComputeCurrentVelocityMagnitudes(
 	}
 }
 
-void UpdateMotionMatchingState(const FAnimationUpdateContext& Context
-	, const UPoseSearchDatabase* Database
-	, const FTrajectorySampleRange& Trajectory
-	, const FMotionMatchingSettings& Settings
-	, FMotionMatchingState& InOutMotionMatchingState
+//////////////////////////////////////////////////////////////////////////
+// FMotionMatchingPoseStepper
+void FMotionMatchingPoseStepper::Update(const FAnimationUpdateContext& UpdateContext, const FMotionMatchingState& State)
+{
+	if (State.DbPoseIdx == INDEX_NONE)
+	{
+		return;
+	}
+
+	check(State.SearchIndexAssetIdx != INDEX_NONE);
+
+	const FPoseSearchIndexAsset* SearchIndexAsset = State.GetCurrentSearchIndexAsset();
+	const FPoseSearchDatabaseSequence& DbSequence = State.CurrentDatabase->GetSourceAsset(SearchIndexAsset);
+	const float AssetLength = DbSequence.Sequence->GetPlayLength();
+	const float DeltaTime = UpdateContext.GetDeltaTime();
+
+	float SteppedAssetTime = State.AssetPlayerTime;
+	ETypeAdvanceAnim AdvanceType = FAnimationRuntime::AdvanceTime(
+		DbSequence.bLoopAnimation,
+		DeltaTime,
+		SteppedAssetTime,
+		AssetLength);
+
+	if (AdvanceType != ETAA_Finished)
+	{
+		// if the sampling interval doesn't contain the new time, there's no pose to return.
+		if (SearchIndexAsset->SamplingInterval.Contains(SteppedAssetTime))
+		{
+			// we can continue ticking the same sequence forward
+			Result.PoseIdx = State.CurrentDatabase->GetPoseIndexFromAssetTime(SteppedAssetTime, SearchIndexAsset);
+			Result.SearchIndexAsset = SearchIndexAsset;
+			Result.TimeOffsetSeconds = State.AssetPlayerTime;
+		}
+	}
+	else
+	{
+		// check if there's a follow-up that can be used
+		int32 FollowUpDbSequenceIdx = State.CurrentDatabase->Sequences.IndexOfByPredicate(
+			[&](const FPoseSearchDatabaseSequence& Entry)
+			{
+				return Entry.Sequence == DbSequence.FollowUpSequence;
+			});
+
+		int32 FollowUpSearchIndexAssetIdx = State.CurrentDatabase->SearchIndex.Assets.IndexOfByPredicate(
+			[&](const FPoseSearchIndexAsset& Entry)
+			{
+				const bool bIsMatch =
+					Entry.SourceAssetIdx == FollowUpDbSequenceIdx &&
+					Entry.bMirrored == SearchIndexAsset->bMirrored &&
+					Entry.SamplingInterval.Contains(0.0f);
+				return bIsMatch;
+			});
+
+		if (FollowUpSearchIndexAssetIdx != INDEX_NONE)
+		{
+			const FPoseSearchIndexAsset* FollowUpSearchIndexAsset =
+				&State.CurrentDatabase->SearchIndex.Assets[FollowUpSearchIndexAssetIdx];
+			const float FollowUpAssetTime = State.AssetPlayerTime + DeltaTime - AssetLength;
+			const int32 FollowUpPoseIdx = State.CurrentDatabase->GetPoseIndexFromAssetTime(
+				FollowUpAssetTime,
+				FollowUpSearchIndexAsset);
+			const FFloatInterval SamplingRange = FollowUpSearchIndexAsset->SamplingInterval;
+
+			Result.PoseIdx = FollowUpPoseIdx;
+			Result.SearchIndexAsset = FollowUpSearchIndexAsset;
+			Result.TimeOffsetSeconds =
+				SamplingRange.Min +
+				(State.CurrentDatabase->Schema->SamplingInterval *
+					(Result.PoseIdx - FollowUpSearchIndexAsset->FirstPoseIdx));
+			bJumpRequired = true;
+		}
+	}
+}
+
+} // namespace UE::PoseSearch
+
+
+void FMotionMatchingState::InitNewDatabaseSearch(const UPoseSearchDatabase* Database, float SearchThrottleTime)
+{
+	DbPoseIdx = INDEX_NONE;
+	SearchIndexAssetIdx = INDEX_NONE;
+	ElapsedPoseJumpTime = SearchThrottleTime;
+	AssetPlayerTime = 0.0f;
+	CurrentDatabase = Database;
+}
+
+void FMotionMatchingState::ComposeQuery(const UPoseSearchDatabase* Database, const FTrajectorySampleRange& Trajectory)
+{
+	FPoseSearchFeatureVectorBuilder Goal;
+	Goal.Init(Database->Schema);
+	Goal.BuildFromTrajectory(Trajectory);
+
+	// Merge goal features into the query vector
+	if (ComposedQuery.IsCompatible(Goal))
+	{
+		ComposedQuery.MergeReplace(Goal);
+	}
+
+	ComposedQuery.Normalize(Database->SearchIndex);
+}
+
+static void RequestInertialBlend(const FAnimationUpdateContext& Context, float BlendTime)
+{
+	// Use inertial blending to smooth over the transition
+	// It would be cool in the future to adjust the blend time by amount of dissimilarity, but we'll need a standardized distance metric first.
+	if (BlendTime > 0.0f)
+	{
+		UE::Anim::IInertializationRequester* InertializationRequester = Context.GetMessage<UE::Anim::IInertializationRequester>();
+		if (InertializationRequester)
+		{
+			InertializationRequester->RequestInertialization(BlendTime);
+		}
+	}
+}
+
+void FMotionMatchingState::JumpToPose(const FAnimationUpdateContext& Context, const FMotionMatchingSettings& Settings, const UE::PoseSearch::FSearchResult& Result)
+{
+	// Remember which pose and sequence we're playing from the database
+	DbPoseIdx = Result.PoseIdx;
+	SearchIndexAssetIdx = CurrentDatabase->SearchIndex.FindAssetIndex(Result.SearchIndexAsset);
+
+	ElapsedPoseJumpTime = 0.0f;
+	AssetPlayerTime = Result.TimeOffsetSeconds;
+
+	const float JumpBlendTime = ComputeJumpBlendTime(Result, Settings);
+	RequestInertialBlend(Context, JumpBlendTime);
+	Flags |= EMotionMatchingFlags::JumpedToPose;
+}
+
+void UpdateMotionMatchingState(
+	const FAnimationUpdateContext& Context,
+	const UPoseSearchDatabase* Database,
+	const FTrajectorySampleRange& Trajectory,
+	const FMotionMatchingSettings& Settings,
+	FMotionMatchingState& InOutMotionMatchingState
 )
 {
+	using namespace UE::PoseSearch;
+
 	InOutMotionMatchingState.Flags = EMotionMatchingFlags::None;
 	const float DeltaTime = Context.GetDeltaTime();
 
@@ -127,27 +213,28 @@ void UpdateMotionMatchingState(const FAnimationUpdateContext& Context
 		}
 
 		// Step the pose forward
-		UE::PoseSearch::FMotionMatchingContinuityParams ContinuityParameters = 
-			InOutMotionMatchingState.ComputeContinuityParameters(Context);
-		const bool bCanContinue = ContinuityParameters.IsValid();
-		if (bCanContinue)
+		FMotionMatchingPoseStepper PoseStepper;
+		PoseStepper.Update(Context, InOutMotionMatchingState);
+		if (PoseStepper.CanContinue())
 		{
-			InOutMotionMatchingState.DbPoseIdx = ContinuityParameters.Result.PoseIdx;
+			InOutMotionMatchingState.DbPoseIdx = PoseStepper.Result.PoseIdx;
 			InOutMotionMatchingState.SearchIndexAssetIdx = 
-				InOutMotionMatchingState.CurrentDatabase->SearchIndex.FindAssetIndex(ContinuityParameters.Result.SearchIndexAsset);
+				InOutMotionMatchingState.CurrentDatabase->SearchIndex.FindAssetIndex(PoseStepper.Result.SearchIndexAsset);
 		}
 
-		// build query
+		// Build the search query
 		if (InOutMotionMatchingState.DbPoseIdx != INDEX_NONE)
 		{
+			// Copy search query directly from the database if we have an active pose
 			InOutMotionMatchingState.ComposedQuery.CopyFromSearchIndex(Database->SearchIndex, InOutMotionMatchingState.DbPoseIdx);
 		}
 		else
 		{
-			UE::PoseSearch::IPoseHistoryProvider* PoseHistoryProvider = Context.GetMessage<UE::PoseSearch::IPoseHistoryProvider>();
+			// When we don't have an active pose, initialize the search query from the pose history provider
+			IPoseHistoryProvider* PoseHistoryProvider = Context.GetMessage<IPoseHistoryProvider>();
 			if (PoseHistoryProvider)
 			{
-				UE::PoseSearch::FPoseHistory& History = PoseHistoryProvider->GetPoseHistory();
+				FPoseHistory& History = PoseHistoryProvider->GetPoseHistory();
 				InOutMotionMatchingState.ComposedQuery.TrySetPoseFeatures(
 					&History, 
 					Context.AnimInstanceProxy->GetRequiredBones());
@@ -157,7 +244,7 @@ void UpdateMotionMatchingState(const FAnimationUpdateContext& Context
 		// Update features in the query with the latest inputs
 		InOutMotionMatchingState.ComposeQuery(Database, Trajectory);
 
-		UE::PoseSearch::FSearchContext SearchContext;
+		FSearchContext SearchContext;
 		SearchContext.SetSource(InOutMotionMatchingState.CurrentDatabase.Get());
 		SearchContext.QueryValues = InOutMotionMatchingState.ComposedQuery.GetNormalizedValues();
 		SearchContext.WeightsContext = &InOutMotionMatchingState.WeightsContext;
@@ -173,67 +260,62 @@ void UpdateMotionMatchingState(const FAnimationUpdateContext& Context
 		InOutMotionMatchingState.WeightsContext.Update(Settings.Weights, Database);
 
 		// Determine how much the updated query vector deviates from the current pose vector
-		float CurrentDissimilarity = MAX_flt;
+		FPoseCost CurrentPoseCost;
 		if (InOutMotionMatchingState.DbPoseIdx != INDEX_NONE)
 		{
-			CurrentDissimilarity = UE::PoseSearch::ComparePoses(InOutMotionMatchingState.DbPoseIdx, SearchContext);
+			CurrentPoseCost = ComparePoses(InOutMotionMatchingState.DbPoseIdx, SearchContext);
 		}
 
 		// Search the database for the nearest match to the updated query vector
-		UE::PoseSearch::FSearchResult Result = UE::PoseSearch::Search(SearchContext);
-		if (Result.IsValid() && ((InOutMotionMatchingState.ElapsedPoseJumpTime >= Settings.SearchThrottleTime) || !bCanContinue))
+		FSearchResult Result = Search(SearchContext);
+
+		if (Result.IsValid())
 		{
-			if (!bCanContinue)
+			// Jump to the candidate pose if the stepper couldn't continue to the next frame
+			if (!PoseStepper.CanContinue())
 			{
-				const float JumpBlendTime = InOutMotionMatchingState.ComputeJumpBlendTime(Result, Settings);
-				InOutMotionMatchingState.JumpToPose(Result);
-				RequestInertialBlend(Context, JumpBlendTime);
-				InOutMotionMatchingState.Flags |= EMotionMatchingFlags::JumpedToPose;
+				InOutMotionMatchingState.JumpToPose(Context, Settings, Result);
 			}
-			else
+
+			// Consider the candidate pose when enough time has elapsed since the last pose jump
+			else if ((InOutMotionMatchingState.ElapsedPoseJumpTime >= Settings.SearchThrottleTime))
 			{
 				// Consider the search result better if it is more similar to the query than the current pose we're playing back from the database
-				bool bBetterPose = Result.Dissimilarity * (1.0f + (Settings.MinPercentImprovement / 100.0f)) < CurrentDissimilarity;
+				check(Result.PoseCost.Dissimilarity >= 0.0f && CurrentPoseCost.Dissimilarity >= 0.0f);
+				bool bBetterPose = Result.PoseCost.Dissimilarity * (1.0f + (Settings.MinPercentImprovement / 100.0f)) < CurrentPoseCost.Dissimilarity;
 
-				// We'll ignore the candidate pose if it is from the same animation and too near to the current pose
+				// Ignore the candidate poses from the same anim when they are too near to the current pose
 				bool bNearbyPose = false;
-				const FPoseSearchIndexAsset* StateSearchIndexAsset =  InOutMotionMatchingState.GetCurrentSearchIndexAsset();
+				const FPoseSearchIndexAsset* StateSearchIndexAsset = InOutMotionMatchingState.GetCurrentSearchIndexAsset();
 				if (StateSearchIndexAsset == Result.SearchIndexAsset)
 				{
 					const FPoseSearchDatabaseSequence& ResultDbSequence = Database->GetSourceAsset(Result.SearchIndexAsset);
+					bNearbyPose = FMath::Abs(InOutMotionMatchingState.AssetPlayerTime - Result.TimeOffsetSeconds) < Settings.PoseJumpThresholdTime;
 
-					// Consider the candidate pose nearby if the animation loops, regardless of how far away the pose is
-					if (ResultDbSequence.bLoopAnimation)
-					{
-						bNearbyPose = true;
-					}
-
-					// Otherwise consider the pose nearby if it is within the pose jump threshold
-					else
+					// Handle looping anims when checking for the pose being too close
+					if (!bNearbyPose && ResultDbSequence.bLoopAnimation)
 					{
-						bNearbyPose = FMath::Abs(InOutMotionMatchingState.AssetPlayerTime - Result.TimeOffsetSeconds) < Settings.PoseJumpThresholdTime;
+						const float AssetLength = ResultDbSequence.Sequence->GetPlayLength();
+						bNearbyPose = FMath::Abs(AssetLength - InOutMotionMatchingState.AssetPlayerTime - Result.TimeOffsetSeconds) < Settings.PoseJumpThresholdTime;
 					}
 				}
 
 				// Start playback from the candidate pose if we determined it was a better option
 				if (bBetterPose && !bNearbyPose)
 				{
-					const float JumpBlendTime = InOutMotionMatchingState.ComputeJumpBlendTime(Result, Settings);
-					InOutMotionMatchingState.JumpToPose(Result);
-					RequestInertialBlend(Context, JumpBlendTime);
-					InOutMotionMatchingState.Flags |= EMotionMatchingFlags::JumpedToPose;
+					InOutMotionMatchingState.JumpToPose(Context, Settings, Result);
 				}
 			}
 		}
 
-		// Continue with the follow up sequence if we're finishing a one shot anim
+		// Jump to the pose stepper's next anim if we didn't jump to a pose as a result of the search above
 		if (!(InOutMotionMatchingState.Flags & EMotionMatchingFlags::JumpedToPose)
-			&& bCanContinue
-			&& ContinuityParameters.bJumpRequired)
+			&& PoseStepper.CanContinue()
+			&& PoseStepper.bJumpRequired)
 		{
-			InOutMotionMatchingState.JumpToPose(ContinuityParameters.Result);
-			RequestInertialBlend(Context, Settings.BlendTime);
-			InOutMotionMatchingState.Flags |= EMotionMatchingFlags::JumpedToPose;
+			const FPoseSearchDatabaseSequence& DbSequence = Database->GetSourceAsset(PoseStepper.Result.SearchIndexAsset);
+			InOutMotionMatchingState.JumpToPose(Context, Settings, PoseStepper.Result);
+			InOutMotionMatchingState.Flags |= EMotionMatchingFlags::JumpedToFollowUp;
 		}
 	}
 
@@ -257,10 +339,10 @@ void UpdateMotionMatchingState(const FAnimationUpdateContext& Context
 			SimLinearVelocity, SimAngularVelocity, 
 			AnimLinearVelocity, AnimAngularVelocity);
 
-		UE::PoseSearch::FTraceMotionMatchingState TraceState;
-		if ((InOutMotionMatchingState.Flags & EMotionMatchingFlags::JumpedToPose) == EMotionMatchingFlags::JumpedToPose)
+		FTraceMotionMatchingState TraceState;
+		if (EnumHasAnyFlags(InOutMotionMatchingState.Flags, EMotionMatchingFlags::JumpedToFollowUp))
 		{
-			TraceState.Flags |= UE::PoseSearch::FTraceMotionMatchingState::EFlags::FollowupAnimation;
+			TraceState.Flags |= FTraceMotionMatchingState::EFlags::FollowupAnimation;
 		}
 
 		TraceState.ElapsedPoseJumpTime = InOutMotionMatchingState.ElapsedPoseJumpTime;
@@ -282,12 +364,13 @@ void UpdateMotionMatchingState(const FAnimationUpdateContext& Context
 #endif
 }
 
-void UPoseSearchLibrary::UpdateMotionMatchingForSequencePlayer(const FAnimUpdateContext& AnimUpdateContext
-	, const FSequencePlayerReference& SequencePlayer
-	, const UPoseSearchDatabase* Database
-	, const FTrajectorySampleRange& Trajectory
-	, const FMotionMatchingSettings& Settings
-	, UPARAM(ref) FMotionMatchingState& InOutMotionMatchingState
+void UPoseSearchLibrary::UpdateMotionMatchingForSequencePlayer(
+	const FAnimUpdateContext& AnimUpdateContext,
+	const FSequencePlayerReference& SequencePlayer,
+	const UPoseSearchDatabase* Database,
+	const FTrajectorySampleRange& Trajectory,
+	const FMotionMatchingSettings& Settings,
+	UPARAM(ref) FMotionMatchingState& InOutMotionMatchingState
 )
 {
 	if (const FAnimationUpdateContext* AnimationUpdateContext = AnimUpdateContext.GetContext())
@@ -322,79 +405,6 @@ void UPoseSearchLibrary::UpdateMotionMatchingForSequencePlayer(const FAnimUpdate
 	}
 }
 
-UE::PoseSearch::FMotionMatchingContinuityParams FMotionMatchingState::ComputeContinuityParameters(const FAnimationUpdateContext& Context) const
-{
-	UE::PoseSearch::FMotionMatchingContinuityParams ContinuityParameters;
-
-	if (DbPoseIdx != INDEX_NONE)
-	{
-		check(SearchIndexAssetIdx != INDEX_NONE);
-
-		const FPoseSearchIndexAsset* SearchIndexAsset = GetCurrentSearchIndexAsset();
-		const FPoseSearchDatabaseSequence& DbSequence = CurrentDatabase->GetSourceAsset(SearchIndexAsset);
-		const float AssetLength = DbSequence.Sequence->GetPlayLength();
-		const float DeltaTime = Context.GetDeltaTime();
-
-		float ContinuityAssetTime = AssetPlayerTime;
-		ETypeAdvanceAnim ContinuityAdvanceType = FAnimationRuntime::AdvanceTime(
-			DbSequence.bLoopAnimation,
-			DeltaTime,
-			ContinuityAssetTime,
-			AssetLength);
-
-		if (ContinuityAdvanceType != ETAA_Finished)
-		{
-			// if the sampling interval doesn't contain the new time, there's no pose to return.
-			if (SearchIndexAsset->SamplingInterval.Contains(ContinuityAssetTime))
-			{
-				// we can continue ticking the same sequence forward
-				ContinuityParameters.Result.PoseIdx = CurrentDatabase->GetPoseIndexFromAssetTime(ContinuityAssetTime, SearchIndexAsset);
-				ContinuityParameters.Result.SearchIndexAsset = SearchIndexAsset;
-				ContinuityParameters.Result.TimeOffsetSeconds = AssetPlayerTime;
-			}
-		}
-		else
-		{
-			// check if there's a follow-up that can be used
-			int32 FollowUpDbSequenceIdx = CurrentDatabase->Sequences.IndexOfByPredicate(
-				[&](const FPoseSearchDatabaseSequence& Entry)
-			{
-				return Entry.Sequence == DbSequence.FollowUpSequence;
-			});
-
-			int32 FollowUpSearchIndexAssetIdx = CurrentDatabase->SearchIndex.Assets.IndexOfByPredicate(
-				[&](const FPoseSearchIndexAsset& Entry)
-			{
-				const bool bIsMatch = 
-					Entry.SourceAssetIdx == FollowUpDbSequenceIdx && 
-					Entry.bMirrored == SearchIndexAsset->bMirrored &&
-					Entry.SamplingInterval.Contains(0.0f);
-				return bIsMatch;
-			});
-
-			if (FollowUpSearchIndexAssetIdx != INDEX_NONE)
-			{
-				const FPoseSearchIndexAsset* FollowUpSearchIndexAsset = 
-					&CurrentDatabase->SearchIndex.Assets[FollowUpSearchIndexAssetIdx];
-				const float FollowUpAssetTime = AssetPlayerTime + DeltaTime - AssetLength;
-				const int32 FollowUpPoseIdx = CurrentDatabase->GetPoseIndexFromAssetTime(
-					FollowUpAssetTime,
-					FollowUpSearchIndexAsset);
-				const FFloatInterval SamplingRange = FollowUpSearchIndexAsset->SamplingInterval;
-
-				ContinuityParameters.Result.PoseIdx = FollowUpPoseIdx;
-				ContinuityParameters.Result.SearchIndexAsset = FollowUpSearchIndexAsset;
-				ContinuityParameters.Result.TimeOffsetSeconds =
-					SamplingRange.Min + 
-					(CurrentDatabase->Schema->SamplingInterval *
-					 (ContinuityParameters.Result.PoseIdx - FollowUpSearchIndexAsset->FirstPoseIdx));
-				ContinuityParameters.bJumpRequired = true;
-			}
-		}
-	}
-	return ContinuityParameters;
-}
-
 const FPoseSearchIndexAsset* FMotionMatchingState::GetCurrentSearchIndexAsset() const
 {
 	if (!CurrentDatabase->SearchIndex.Assets.IsValidIndex(SearchIndexAssetIdx))
@@ -407,9 +417,11 @@ const FPoseSearchIndexAsset* FMotionMatchingState::GetCurrentSearchIndexAsset()
 
 float FMotionMatchingState::ComputeJumpBlendTime(
 	const UE::PoseSearch::FSearchResult& Result, 
-	const FMotionMatchingSettings& Settings) const
+	const FMotionMatchingSettings& Settings
+) const
 {
 	const FPoseSearchIndexAsset* SearchIndexAsset = GetCurrentSearchIndexAsset();
+
 	// Use alternate blend time when changing between mirrored and unmirrored
 	float JumpBlendTime = Settings.BlendTime;
 	if ((SearchIndexAsset != nullptr) && (Settings.MirrorChangeBlendTime > 0.0f))
diff --git a/Engine/Plugins/Experimental/Animation/PoseSearch/Source/Runtime/Public/PoseSearch/PoseSearch.h b/Engine/Plugins/Experimental/Animation/PoseSearch/Source/Runtime/Public/PoseSearch/PoseSearch.h
--- a/Engine/Plugins/Experimental/Animation/PoseSearch/Source/Runtime/Public/PoseSearch/PoseSearch.h
+++ b/Engine/Plugins/Experimental/Animation/PoseSearch/Source/Runtime/Public/PoseSearch/PoseSearch.h
@@ -904,12 +904,21 @@ struct POSESEARCH_API FDebugDrawParams
 	const UPoseSearchSchema* GetSchema() const;
 };
 
+struct FPoseCost
+{
+	float Dissimilarity = MAX_flt;
+	float CostAddend = 0.0f;
+	float TotalCost = MAX_flt;
+	bool operator<(const FPoseCost& Other) const { return TotalCost < Other.TotalCost; }
+
+};
+
 struct FSearchResult
 {
+	FPoseCost PoseCost;
 	int32 PoseIdx = INDEX_NONE;
 	const FPoseSearchIndexAsset* SearchIndexAsset = nullptr;
 	float TimeOffsetSeconds = 0.0f;
-	float Dissimilarity = MAX_flt;
 
 	bool IsValid() const { return PoseIdx >= 0; }
 };
@@ -933,6 +942,7 @@ private:
 	float MirrorMismatchCost = 0.0f;
 };
 
+
 /**
 * Visualize pose search debug information
 *
@@ -980,20 +990,21 @@ POSESEARCH_API FSearchResult Search(FSearchContext& SearchContext);
 * 
 * @return Dissimilarity between the two poses
 */
-POSESEARCH_API float ComparePoses(int32 PoseIdx, FSearchContext& SearchContext);
+
+POSESEARCH_API FPoseCost ComparePoses(int32 PoseIdx, FSearchContext& SearchContext);
 
 
 /**
  * Cost details for pose analysis in the rewind debugger
  */
-struct FPoseCostInfo
+struct FPoseCostDetails
 {
-	// The dissimilarity of this pose compared to the query
-	float TotalCost = 0.0f;
+	FPoseCost PoseCost;
 
 	// Contribution from ModifyCost anim notify
 	float NotifyCostAddend = 0.0f;
 
+	// Contribution from mirroring cost
 	float MirrorMismatchCostAddend = 0.0f;
 
 	// Cost breakdown per channel (e.g. pose cost, time-based trajectory cost, distance-based trajectory cost, etc.)
@@ -1010,9 +1021,10 @@ struct FPoseCostInfo
 * @param SearchIndex		The search index containing the pose to compare to the query
 * @param PoseIdx			The index of the pose in the search index to compare to the query
 * @param SearchContext		Structure containing search parameters
-* &param OutPoseCostInfo	Cost details for analysis in the debugger
+* &param OutPoseCostDetails	Cost details for analysis in the debugger
 *
 * @return Dissimilarity between the two poses
 */
-POSESEARCH_API float ComparePoses(int32 PoseIdx, FSearchContext& SearchContext, FPoseCostInfo& OutPoseCostInfo);
+POSESEARCH_API FPoseCost ComparePoses(int32 PoseIdx, FSearchContext& SearchContext, FPoseCostDetails& OutPoseCostDetails);
+
 }} // namespace UE::PoseSearch
\ No newline at end of file
diff --git a/Engine/Plugins/Experimental/Animation/PoseSearch/Source/Runtime/Public/PoseSearch/PoseSearchLibrary.h b/Engine/Plugins/Experimental/Animation/PoseSearch/Source/Runtime/Public/PoseSearch/PoseSearchLibrary.h
--- a/Engine/Plugins/Experimental/Animation/PoseSearch/Source/Runtime/Public/PoseSearch/PoseSearchLibrary.h
+++ b/Engine/Plugins/Experimental/Animation/PoseSearch/Source/Runtime/Public/PoseSearch/PoseSearchLibrary.h
@@ -15,12 +15,12 @@
 
 namespace UE::PoseSearch
 {
-	struct FMotionMatchingContinuityParams
+	struct FMotionMatchingPoseStepper
 	{
 		FSearchResult Result;
 		bool bJumpRequired = false;
 
-		bool IsValid()
+		bool CanContinue() const
 		{
 			return Result.IsValid();
 		}
@@ -30,6 +30,8 @@ namespace UE::PoseSearch
 			Result = UE::PoseSearch::FSearchResult();
 			bJumpRequired = false;
 		}
+
+		void Update(const FAnimationUpdateContext& UpdateContext, const struct FMotionMatchingState& State);
 	};
 }
 
@@ -37,7 +39,8 @@ UENUM(BlueprintType, Category="Motion Trajectory", meta=(Bitflags, UseEnumValues
 enum class EMotionMatchingFlags : uint8
 {
 	None = 0 UMETA(Hidden),
-	JumpedToPose = 1 << 0 // Signals that motion matching has made a significant deviation in the selected sequence/pose index
+	JumpedToPose = 1 << 0,		// Signals that motion matching has made a significant deviation in the selected sequence/pose index
+	JumpedToFollowUp = 1 << 1,	// Motion matching chose the follow up animation of the prior sequence
 };
 ENUM_CLASS_FLAGS(EMotionMatchingFlags);
 
@@ -84,10 +87,7 @@ struct POSESEARCH_API FMotionMatchingState
 	void ComposeQuery(const UPoseSearchDatabase* Database, const FTrajectorySampleRange& Trajectory);
 
 	// Internally stores the 'jump' to a new pose/sequence index and asset time for evaluation
-	void JumpToPose(const UE::PoseSearch::FSearchResult& Result);
-
-	// Updates DbPoseIdx to track DeltaTime and jumps to a follow-up sequence when available
-	UE::PoseSearch::FMotionMatchingContinuityParams ComputeContinuityParameters(const FAnimationUpdateContext& Context) const;
+	void JumpToPose(const FAnimationUpdateContext& Context, const FMotionMatchingSettings& Settings, const UE::PoseSearch::FSearchResult& Result);
 
 	const FPoseSearchIndexAsset* GetCurrentSearchIndexAsset() const;
 
@@ -134,7 +134,6 @@ struct POSESEARCH_API FMotionMatchingState
 * @param Trajectory					Input motion trajectory samples for pose search queries
 * @param Settings					Input motion matching algorithm configuration settings
 * @param InOutMotionMatchingState	Input/Output encapsulated motion matching algorithm and state
-* @param bOutJumpedToPose			Output motion matching decision to 'jump' to a new pose/sequence index in the database
 */
 POSESEARCH_API void UpdateMotionMatchingState(const FAnimationUpdateContext& Context
 	, const UPoseSearchDatabase* Database
