[UE-139997] Fix debugger support for switching databases

diff --git a/Engine/Plugins/Experimental/Animation/PoseSearch/Source/Editor/Private/PoseSearchDebugger.cpp b/Engine/Plugins/Experimental/Animation/PoseSearch/Source/Editor/Private/PoseSearchDebugger.cpp
--- a/Engine/Plugins/Experimental/Animation/PoseSearch/Source/Editor/Private/PoseSearchDebugger.cpp
+++ b/Engine/Plugins/Experimental/Animation/PoseSearch/Source/Editor/Private/PoseSearchDebugger.cpp
@@ -687,7 +687,9 @@ void SDebuggerDatabaseView::FilterDatabaseRows()
 void SDebuggerDatabaseView::CreateRows(const UPoseSearchDatabase& Database)
 {
 	const int32 NumPoses = Database.SearchIndex.NumPoses;
-	UnfilteredDatabaseRows.Reserve(NumPoses);
+	UnfilteredDatabaseRows.Reset(NumPoses);
+
+	RowsSourceDatabase = &Database;
 
 	// Build database rows
 	for(const FPoseSearchIndexAsset& SearchIndexAsset : Database.SearchIndex.Assets)
@@ -707,14 +709,16 @@ void SDebuggerDatabaseView::CreateRows(const UPoseSearchDatabase& Database)
 		}
 	}
 
+	ActiveView.Rows.Reset();
 	ActiveView.Rows.Add(MakeShared<FDebuggerDatabaseRowData>());
 }
 
 void SDebuggerDatabaseView::UpdateRows(const FTraceMotionMatchingStateMessage& State, const UPoseSearchDatabase& Database)
 {
-	if (UnfilteredDatabaseRows.IsEmpty())
+	const bool bNewDatabase = RowsSourceDatabase != &Database;
+
+	if (bNewDatabase || UnfilteredDatabaseRows.IsEmpty())
 	{
-		check(ActiveView.Rows.IsEmpty());
 		CreateRows(Database);
 	}
 	check(ActiveView.Rows.Num() == 1);
@@ -747,6 +751,11 @@ void SDebuggerDatabaseView::UpdateRows(const FTraceMotionMatchingStateMessage& S
 
 	SortDatabaseRows();
 	FilterDatabaseRows();
+
+	if (bNewDatabase)
+	{
+		FilteredDatabaseView.ListView->ClearSelection();
+	}
 }
 
 TSharedRef<ITableRow> SDebuggerDatabaseView::HandleGenerateDatabaseRow(TSharedRef<FDebuggerDatabaseRowData> Item, const TSharedRef<STableViewBase>& OwnerTable) const
@@ -1176,7 +1185,7 @@ void SDebuggerView::Tick(const FGeometry& AllottedGeometry, const double InCurre
 
 	TSharedPtr<FDebuggerViewModel> Model = ViewModel.Get();
 
-	bool bNeedUpdate = false;
+	bool bNeedUpdate = Model->NeedsUpdate();
 
 	// We haven't reached the update point yet
 	if (CurrentConsecutiveFrames < ConsecutiveFramesUpdateThreshold)
@@ -1782,6 +1791,13 @@ const FTransform* FDebuggerViewModel::GetRootTransform() const
 	return RootTransform;
 }
 
+bool FDebuggerViewModel::NeedsUpdate() const
+{
+	const UPoseSearchDatabase* NewDatabase = GetPoseSearchDatabase();
+	const bool bDatabaseChanged = NewDatabase != CurrentDatabase;
+	return bDatabaseChanged;
+}
+
 void FDebuggerViewModel::OnUpdate()
 {
 	if (!bSkeletonsInitialized)
@@ -1828,6 +1844,13 @@ void FDebuggerViewModel::OnUpdateNodeSelection(int32 InNodeId)
 	{
 		Skeletons[ActivePose].SequenceIdx = ActiveMotionMatchingState->DbPoseIdx;
 	}
+
+	const UPoseSearchDatabase* NewDatabase = GetPoseSearchDatabase();
+	if (NewDatabase != CurrentDatabase)
+	{
+		ClearSelectedSkeleton();
+		CurrentDatabase = NewDatabase;
+	}
 }
 
 void FDebuggerViewModel::OnDraw(FSkeletonDrawParams& DrawParams)
diff --git a/Engine/Plugins/Experimental/Animation/PoseSearch/Source/Editor/Private/PoseSearchDebugger.h b/Engine/Plugins/Experimental/Animation/PoseSearch/Source/Editor/Private/PoseSearchDebugger.h
--- a/Engine/Plugins/Experimental/Animation/PoseSearch/Source/Editor/Private/PoseSearchDebugger.h
+++ b/Engine/Plugins/Experimental/Animation/PoseSearch/Source/Editor/Private/PoseSearchDebugger.h
@@ -353,6 +353,8 @@ private:
 	/** All database poses */
 	TArray<TSharedRef<FDebuggerDatabaseRowData>> UnfilteredDatabaseRows;
 
+	TWeakObjectPtr<const UPoseSearchDatabase> RowsSourceDatabase = nullptr;
+
 	/** Database listing for filtered poses */
 	FTable FilteredDatabaseView;
 
@@ -522,6 +524,9 @@ public:
 	int32 GetNodesNum() const;
 	const FTransform* GetRootTransform() const;
 
+	/** Checks if Update must be called */
+	bool NeedsUpdate() const;
+
 	/** Update motion matching states for frame */
 	void OnUpdate();
 	
@@ -577,6 +582,8 @@ private:
 	/** Currently active motion matching state based on node selection in the view */
 	const FTraceMotionMatchingStateMessage* ActiveMotionMatchingState = nullptr;
 
+	TWeakObjectPtr<const UPoseSearchDatabase> CurrentDatabase = nullptr;
+
 	/** Current Skeletal Mesh Component Id for the AnimInstance */
 	uint64 SkeletalMeshComponentId = 0;
 
