PoseSearchDebugger: now the anim sequence name is a link to the asset

diff --git a/Engine/Plugins/Experimental/Animation/PoseSearch/Source/Editor/Private/PoseSearchDebugger.cpp b/Engine/Plugins/Experimental/Animation/PoseSearch/Source/Editor/Private/PoseSearchDebugger.cpp
--- a/Engine/Plugins/Experimental/Animation/PoseSearch/Source/Editor/Private/PoseSearchDebugger.cpp
+++ b/Engine/Plugins/Experimental/Animation/PoseSearch/Source/Editor/Private/PoseSearchDebugger.cpp
@@ -15,6 +15,7 @@
 #include "Widgets/Layout/SWidgetSwitcher.h"
 #include "Widgets/SBoxPanel.h"
 #include "Widgets/Input/SButton.h"
+#include "Widgets/Input/SHyperlink.h"
 #include "Widgets/Input/SNumericEntryBox.h"
 #include "Widgets/Input/SSearchBox.h"
 #include "Widgets/Images/SImage.h"
@@ -214,6 +215,7 @@ public:
 
 	int32 PoseIdx = 0;
 	FString AnimSequenceName = "";
+	FString AnimSequencePath = "";
 	int32 DbSequenceIdx = 0;
 	int32 AnimFrame = 0;
 	float Time = 0.0f;
@@ -231,9 +233,9 @@ namespace DebuggerDatabaseColumns
 			, bEnabled(InEnabled)
 		{
 		}
-		
+
 		virtual ~IColumn() = default;
-		
+
 		/** Sorted left to right based on this index */
 		int32 SortIndex = 0;
 		/** Current width, starts at 1 to be evenly spaced between all columns */
@@ -242,13 +244,26 @@ namespace DebuggerDatabaseColumns
 		bool bEnabled = false;
 
 		virtual FName GetName() const = 0;
-		
+
 		using FRowDataRef = TSharedRef<FDebuggerDatabaseRowData>;
 		using FSortPredicate = TFunctionRef<bool(const FRowDataRef&, const FRowDataRef&)>;
-		
+
 		virtual FSortPredicate GetSortPredicate() const = 0;
+
+		virtual TSharedRef<SWidget> GenerateWidget(const FRowDataRef& RowData) const = 0;
+	};
+
+	/** Column struct to represent each column in the debugger database */
+	struct ITextColumn : IColumn
+	{
+		explicit ITextColumn(int32 InSortIndex, bool InEnabled = true)
+			: IColumn(InSortIndex, InEnabled)
+		{
+		}
+
+		virtual ~ITextColumn() = default;
 		
-		TSharedRef<STextBlock> GenerateTextWidget(const FRowDataRef& RowData) const
+		virtual TSharedRef<SWidget> GenerateWidget(const FRowDataRef& RowData) const override
 		{
 			static FSlateFontInfo RowFont = FEditorStyle::Get().GetFontStyle("DetailsView.CategoryTextStyle");
 			
@@ -262,9 +277,9 @@ namespace DebuggerDatabaseColumns
 		virtual FText GetRowText(const FRowDataRef& Row) const = 0;
 	};
 
-	struct FPoseIdx : IColumn
+	struct FPoseIdx : ITextColumn
 	{
-		using IColumn::IColumn;
+		using ITextColumn::ITextColumn;
 		static const FName Name;
 		virtual FName GetName() const override { return Name; }
 
@@ -292,16 +307,23 @@ namespace DebuggerDatabaseColumns
 			return Predicate;
 		}
 
-		virtual FText GetRowText(const FRowDataRef& Row) const override
+		virtual TSharedRef<SWidget> GenerateWidget(const FRowDataRef& RowData) const override
 		{
-			return FText::FromString(Row->AnimSequenceName);
+			return SNew(SHyperlink)
+				.Text(FText::FromString(RowData->AnimSequenceName))
+				.TextStyle(&FCoreStyle::Get().GetWidgetStyle<FTextBlockStyle>("SmallText"))
+				.ToolTipText(FText::Format(LOCTEXT("AssetHyperlinkTooltipFormat", "Open asset '{0}'"), FText::FromString(RowData->AnimSequencePath)))
+				.OnNavigate_Lambda([AnimSequencePath = RowData->AnimSequencePath]()
+			{
+				GEditor->GetEditorSubsystem<UAssetEditorSubsystem>()->OpenEditorForAsset(AnimSequencePath);
+			});
 		}
 	};
 	const FName FAnimSequenceName::Name = "Sequence";
 
-	struct FFrame : IColumn
+	struct FFrame : ITextColumn
 	{
-		using IColumn::IColumn;
+		using ITextColumn::ITextColumn;
 		static const FName Name;
 		virtual FName GetName() const override { return Name; }
 		
@@ -326,9 +348,9 @@ namespace DebuggerDatabaseColumns
 	};
 	const FName FFrame::Name = "Frame";
 
-	struct FCost : IColumn
+	struct FCost : ITextColumn
 	{
-		using IColumn::IColumn;
+		using ITextColumn::ITextColumn;
 		static const FName Name;
 		virtual FName GetName() const override { return Name; }
 		
@@ -345,9 +367,9 @@ namespace DebuggerDatabaseColumns
 	const FName FCost::Name = "Cost";
 
 
-	struct FPoseCostColumn : IColumn
+	struct FPoseCostColumn : ITextColumn
 	{
-		using IColumn::IColumn;
+		using ITextColumn::ITextColumn;
 		static const FName Name;
 		virtual FName GetName() const override { return Name; }
 		
@@ -364,9 +386,9 @@ namespace DebuggerDatabaseColumns
 	const FName FPoseCostColumn::Name = "Pose Cost";
 
 
-	struct FTrajectoryCost : IColumn
+	struct FTrajectoryCost : ITextColumn
 	{
-		using IColumn::IColumn;
+		using ITextColumn::ITextColumn;
 		static const FName Name;
 		virtual FName GetName() const override { return Name; }
 		
@@ -382,9 +404,9 @@ namespace DebuggerDatabaseColumns
 	};
 	const FName FTrajectoryCost::Name = "Trajectory Cost";
 
-	struct FCostModifier : IColumn
+	struct FCostModifier : ITextColumn
 	{
-		using IColumn::IColumn;
+		using ITextColumn::ITextColumn;
 		static const FName Name;
 		virtual FName GetName() const override { return Name; }
 
@@ -403,9 +425,9 @@ namespace DebuggerDatabaseColumns
 	};
 	const FName FCostModifier::Name = "Cost Modifier";
 
-	struct FMirrored : IColumn
+	struct FMirrored : ITextColumn
 	{
-		using IColumn::IColumn;
+		using ITextColumn::ITextColumn;
 		static const FName Name;
 		virtual FName GetName() const override { return Name; }
 
@@ -466,7 +488,7 @@ class SDebuggerDatabaseRow : public SMultiColumnTableRow<TSharedRef<FDebuggerDat
 		const TSharedRef<DebuggerDatabaseColumns::IColumn>& Column = (*ColumnMap.Get())[InColumnName];
 		
 		static FSlateFontInfo NormalFont = FEditorStyle::Get().GetFontStyle("DetailsView.CategoryTextStyle");
-		const TSharedRef<SWidget> Widget = Column->GenerateTextWidget(Row.ToSharedRef());
+		const TSharedRef<SWidget> Widget = Column->GenerateWidget(Row.ToSharedRef());
 		
 		return
 			SNew(SBorder)
@@ -703,6 +725,7 @@ void SDebuggerDatabaseView::CreateRows(const UPoseSearchDatabase& Database)
 			TSharedRef<FDebuggerDatabaseRowData> Row = UnfilteredDatabaseRows.Add_GetRef(MakeShared<FDebuggerDatabaseRowData>());
 			Row->PoseIdx = PoseIdx;
 			Row->AnimSequenceName = DbSequence.Sequence->GetName();
+			Row->AnimSequencePath = DbSequence.Sequence->GetPathName();
 			Row->Time = Time;
 			Row->AnimFrame = DbSequence.Sequence->GetFrameAtTime(Time);
 			Row->bMirrored = SearchIndexAsset.bMirrored;
