Build fix: added missing includes.

diff --git a/Engine/Plugins/Experimental/Animation/PoseSearch/Source/Editor/Private/PoseSearchDebugger.cpp b/Engine/Plugins/Experimental/Animation/PoseSearch/Source/Editor/Private/PoseSearchDebugger.cpp
--- a/Engine/Plugins/Experimental/Animation/PoseSearch/Source/Editor/Private/PoseSearchDebugger.cpp
+++ b/Engine/Plugins/Experimental/Animation/PoseSearch/Source/Editor/Private/PoseSearchDebugger.cpp
@@ -21,6 +21,8 @@
 #include "Widgets/Images/SImage.h"
 #include "Styling/SlateIconFinder.h"
 #include "TraceServices/Model/Frames.h"
+#include "Subsystems/AssetEditorSubsystem.h"
+#include "Editor.h"
 #include "Algo/AllOf.h"
 
 #define LOCTEXT_NAMESPACE "PoseSearchDebugger"
