PoseSearch:
- Log AnimBP error when database is invalid
- Updated default values for trajectory weight and MinPercentImprovement

diff --git a/Engine/Plugins/Experimental/Animation/PoseSearch/Source/Runtime/Private/AnimNode_MotionMatching.cpp b/Engine/Plugins/Experimental/Animation/PoseSearch/Source/Runtime/Private/AnimNode_MotionMatching.cpp
--- a/Engine/Plugins/Experimental/Animation/PoseSearch/Source/Runtime/Private/AnimNode_MotionMatching.cpp
+++ b/Engine/Plugins/Experimental/Animation/PoseSearch/Source/Runtime/Private/AnimNode_MotionMatching.cpp
@@ -22,7 +22,7 @@ void FAnimNode_MotionMatching::Initialize_AnyThread(const FAnimationInitializeCo
 
 	GetEvaluateGraphExposedInputs().Execute(Context);
 
-	MotionMatchingState.InitNewDatabaseSearch(Database, Settings.SearchThrottleTime);
+	MotionMatchingState.InitNewDatabaseSearch(Database, Settings.SearchThrottleTime, nullptr /*OutError*/);
 
 	Source.SetLinkNode(&MirrorNode);
 	MirrorNode.SetSourceLinkNode(&SequencePlayerNode);
diff --git a/Engine/Plugins/Experimental/Animation/PoseSearch/Source/Runtime/Private/PoseSearch.cpp b/Engine/Plugins/Experimental/Animation/PoseSearch/Source/Runtime/Private/PoseSearch.cpp
--- a/Engine/Plugins/Experimental/Animation/PoseSearch/Source/Runtime/Private/PoseSearch.cpp
+++ b/Engine/Plugins/Experimental/Animation/PoseSearch/Source/Runtime/Private/PoseSearch.cpp
@@ -580,6 +580,14 @@ void UPoseSearchSchema::ResolveBoneReferences()
 }
 
 
+//////////////////////////////////////////////////////////////////////////
+// FPoseSearchChannelWeightParams
+FPoseSearchWeightParams::FPoseSearchWeightParams()
+{
+	TrajectoryWeight.ChannelWeight = 2.5f;
+}
+
+
 //////////////////////////////////////////////////////////////////////////
 // FPoseSearchChannelWeightParams
 
diff --git a/Engine/Plugins/Experimental/Animation/PoseSearch/Source/Runtime/Private/PoseSearchLibrary.cpp b/Engine/Plugins/Experimental/Animation/PoseSearch/Source/Runtime/Private/PoseSearchLibrary.cpp
--- a/Engine/Plugins/Experimental/Animation/PoseSearch/Source/Runtime/Private/PoseSearchLibrary.cpp
+++ b/Engine/Plugins/Experimental/Animation/PoseSearch/Source/Runtime/Private/PoseSearchLibrary.cpp
@@ -8,6 +8,8 @@
 #include "PoseSearch/AnimNode_MotionMatching.h"
 #include "Trace/PoseSearchTraceLogger.h"
 
+#define LOCTEXT_NAMESPACE "PoseSearchLibrary"
+
 DEFINE_LOG_CATEGORY_STATIC(LogPoseSearchLibrary, Verbose, All);
 
 namespace UE::PoseSearch {
@@ -60,6 +62,7 @@ void ComputeCurrentVelocityMagnitudes(
 
 //////////////////////////////////////////////////////////////////////////
 // FMotionMatchingPoseStepper
+
 void FMotionMatchingPoseStepper::Update(const FAnimationUpdateContext& UpdateContext, const FMotionMatchingState& State)
 {
 	if (State.DbPoseIdx == INDEX_NONE)
@@ -135,13 +138,44 @@ void FMotionMatchingPoseStepper::Update(const FAnimationUpdateContext& UpdateCon
 } // namespace UE::PoseSearch
 
 
-void FMotionMatchingState::InitNewDatabaseSearch(const UPoseSearchDatabase* Database, float SearchThrottleTime)
+//////////////////////////////////////////////////////////////////////////
+// FMotionMatchingState
+
+bool FMotionMatchingState::InitNewDatabaseSearch(
+	const UPoseSearchDatabase* Database,
+	float SearchThrottleTime,
+	FText* OutError
+)
 {
-	DbPoseIdx = INDEX_NONE;
-	SearchIndexAssetIdx = INDEX_NONE;
-	ElapsedPoseJumpTime = SearchThrottleTime;
-	AssetPlayerTime = 0.0f;
-	CurrentDatabase = Database;
+	bool bValidDatabase = Database && Database->IsValidForSearch();
+
+	if (bValidDatabase)
+	{
+		DbPoseIdx = INDEX_NONE;
+		SearchIndexAssetIdx = INDEX_NONE;
+		ElapsedPoseJumpTime = SearchThrottleTime;
+		AssetPlayerTime = 0.0f;
+		CurrentDatabase = Database;
+
+		if (!ComposedQuery.IsInitializedForSchema(Database->Schema))
+		{
+			ComposedQuery.Init(Database->Schema);
+		}
+	}
+
+	if (!bValidDatabase && OutError)
+	{
+		if (Database)
+		{
+			*OutError = FText::Format(LOCTEXT("InvalidDatabase", "Invalid database for motion matching. Try re-saving {0}."), FText::FromString(Database->GetPathName()));
+		}
+		else
+		{
+			*OutError = LOCTEXT("NoDatabase", "No database provided for motion matching.");
+		}
+	}
+
+	return bValidDatabase;
 }
 
 void FMotionMatchingState::ComposeQuery(const UPoseSearchDatabase* Database, const FTrajectorySampleRange& Trajectory)
@@ -198,125 +232,122 @@ void UpdateMotionMatchingState(
 	using namespace UE::PoseSearch;
 
 	InOutMotionMatchingState.Flags = EMotionMatchingFlags::None;
-	const float DeltaTime = Context.GetDeltaTime();
-
-	if (Database && Database->IsValidForSearch())
+	if (InOutMotionMatchingState.CurrentDatabase != Database)
 	{
-		if (InOutMotionMatchingState.CurrentDatabase != Database)
+		FText InitError;
+		if (!InOutMotionMatchingState.InitNewDatabaseSearch(Database, Settings.SearchThrottleTime, &InitError))
 		{
-			InOutMotionMatchingState.InitNewDatabaseSearch(Database, Settings.SearchThrottleTime);
+			Context.LogMessage(EMessageSeverity::Error, InitError);
+			return;
 		}
+	}
 
-		if (!InOutMotionMatchingState.ComposedQuery.IsInitializedForSchema(Database->Schema))
-		{
-			InOutMotionMatchingState.ComposedQuery.Init(Database->Schema);
-		}
+	const float DeltaTime = Context.GetDeltaTime();
 
-		// Step the pose forward
-		FMotionMatchingPoseStepper PoseStepper;
-		PoseStepper.Update(Context, InOutMotionMatchingState);
-		if (PoseStepper.CanContinue())
-		{
-			InOutMotionMatchingState.DbPoseIdx = PoseStepper.Result.PoseIdx;
-			InOutMotionMatchingState.SearchIndexAssetIdx = 
-				InOutMotionMatchingState.CurrentDatabase->SearchIndex.FindAssetIndex(PoseStepper.Result.SearchIndexAsset);
-		}
+	// Step the pose forward
+	FMotionMatchingPoseStepper PoseStepper;
+	PoseStepper.Update(Context, InOutMotionMatchingState);
+	if (PoseStepper.CanContinue())
+	{
+		InOutMotionMatchingState.DbPoseIdx = PoseStepper.Result.PoseIdx;
+		InOutMotionMatchingState.SearchIndexAssetIdx = 
+			InOutMotionMatchingState.CurrentDatabase->SearchIndex.FindAssetIndex(PoseStepper.Result.SearchIndexAsset);
+	}
 
-		// Build the search query
-		if (InOutMotionMatchingState.DbPoseIdx != INDEX_NONE)
-		{
-			// Copy search query directly from the database if we have an active pose
-			InOutMotionMatchingState.ComposedQuery.CopyFromSearchIndex(Database->SearchIndex, InOutMotionMatchingState.DbPoseIdx);
-		}
-		else
+	// Build the search query
+	if (InOutMotionMatchingState.DbPoseIdx != INDEX_NONE)
+	{
+		// Copy search query directly from the database if we have an active pose
+		InOutMotionMatchingState.ComposedQuery.CopyFromSearchIndex(Database->SearchIndex, InOutMotionMatchingState.DbPoseIdx);
+	}
+	else
+	{
+		// When we don't have an active pose, initialize the search query from the pose history provider
+		IPoseHistoryProvider* PoseHistoryProvider = Context.GetMessage<IPoseHistoryProvider>();
+		if (PoseHistoryProvider)
 		{
-			// When we don't have an active pose, initialize the search query from the pose history provider
-			IPoseHistoryProvider* PoseHistoryProvider = Context.GetMessage<IPoseHistoryProvider>();
-			if (PoseHistoryProvider)
-			{
-				FPoseHistory& History = PoseHistoryProvider->GetPoseHistory();
-				InOutMotionMatchingState.ComposedQuery.TrySetPoseFeatures(
-					&History, 
-					Context.AnimInstanceProxy->GetRequiredBones());
-			}
+			FPoseHistory& History = PoseHistoryProvider->GetPoseHistory();
+			InOutMotionMatchingState.ComposedQuery.TrySetPoseFeatures(
+				&History, 
+				Context.AnimInstanceProxy->GetRequiredBones());
 		}
+	}
 
-		// Update features in the query with the latest inputs
-		InOutMotionMatchingState.ComposeQuery(Database, Trajectory);
+	// Update features in the query with the latest inputs
+	InOutMotionMatchingState.ComposeQuery(Database, Trajectory);
 
-		FSearchContext SearchContext;
-		SearchContext.SetSource(InOutMotionMatchingState.CurrentDatabase.Get());
-		SearchContext.QueryValues = InOutMotionMatchingState.ComposedQuery.GetNormalizedValues();
-		SearchContext.WeightsContext = &InOutMotionMatchingState.WeightsContext;
-		if (const FPoseSearchIndexAsset* CurrentIndexAsset = InOutMotionMatchingState.GetCurrentSearchIndexAsset())
-		{
-			SearchContext.QueryMirrorRequest =
-				CurrentIndexAsset->bMirrored ?
-				EPoseSearchBooleanRequest::TrueValue :
-				EPoseSearchBooleanRequest::FalseValue;
-		}
+	FSearchContext SearchContext;
+	SearchContext.SetSource(InOutMotionMatchingState.CurrentDatabase.Get());
+	SearchContext.QueryValues = InOutMotionMatchingState.ComposedQuery.GetNormalizedValues();
+	SearchContext.WeightsContext = &InOutMotionMatchingState.WeightsContext;
+	if (const FPoseSearchIndexAsset* CurrentIndexAsset = InOutMotionMatchingState.GetCurrentSearchIndexAsset())
+	{
+		SearchContext.QueryMirrorRequest =
+			CurrentIndexAsset->bMirrored ?
+			EPoseSearchBooleanRequest::TrueValue :
+			EPoseSearchBooleanRequest::FalseValue;
+	}
 
-		// Update weight groups
-		InOutMotionMatchingState.WeightsContext.Update(Settings.Weights, Database);
+	// Update weight groups
+	InOutMotionMatchingState.WeightsContext.Update(Settings.Weights, Database);
+
+	// Determine how much the updated query vector deviates from the current pose vector
+	FPoseCost CurrentPoseCost;
+	if (InOutMotionMatchingState.DbPoseIdx != INDEX_NONE)
+	{
+		CurrentPoseCost = ComparePoses(InOutMotionMatchingState.DbPoseIdx, SearchContext);
+	}
 
-		// Determine how much the updated query vector deviates from the current pose vector
-		FPoseCost CurrentPoseCost;
-		if (InOutMotionMatchingState.DbPoseIdx != INDEX_NONE)
+	// Search the database for the nearest match to the updated query vector
+	FSearchResult Result = Search(SearchContext);
+
+	if (Result.IsValid())
+	{
+		// Jump to the candidate pose if the stepper couldn't continue to the next frame
+		if (!PoseStepper.CanContinue())
 		{
-			CurrentPoseCost = ComparePoses(InOutMotionMatchingState.DbPoseIdx, SearchContext);
+			InOutMotionMatchingState.JumpToPose(Context, Settings, Result);
 		}
 
-		// Search the database for the nearest match to the updated query vector
-		FSearchResult Result = Search(SearchContext);
-
-		if (Result.IsValid())
+		// Consider the candidate pose when enough time has elapsed since the last pose jump
+		else if ((InOutMotionMatchingState.ElapsedPoseJumpTime >= Settings.SearchThrottleTime))
 		{
-			// Jump to the candidate pose if the stepper couldn't continue to the next frame
-			if (!PoseStepper.CanContinue())
+			// Consider the search result better if it is more similar to the query than the current pose we're playing back from the database
+			check(Result.PoseCost.Dissimilarity >= 0.0f && CurrentPoseCost.Dissimilarity >= 0.0f);
+			bool bBetterPose = Result.PoseCost.Dissimilarity * (1.0f + (Settings.MinPercentImprovement / 100.0f)) < CurrentPoseCost.Dissimilarity;
+
+			// Ignore the candidate poses from the same anim when they are too near to the current pose
+			bool bNearbyPose = false;
+			const FPoseSearchIndexAsset* StateSearchIndexAsset = InOutMotionMatchingState.GetCurrentSearchIndexAsset();
+			if (StateSearchIndexAsset == Result.SearchIndexAsset)
 			{
-				InOutMotionMatchingState.JumpToPose(Context, Settings, Result);
-			}
+				const FPoseSearchDatabaseSequence& ResultDbSequence = Database->GetSourceAsset(Result.SearchIndexAsset);
+				bNearbyPose = FMath::Abs(InOutMotionMatchingState.AssetPlayerTime - Result.TimeOffsetSeconds) < Settings.PoseJumpThresholdTime;
 
-			// Consider the candidate pose when enough time has elapsed since the last pose jump
-			else if ((InOutMotionMatchingState.ElapsedPoseJumpTime >= Settings.SearchThrottleTime))
-			{
-				// Consider the search result better if it is more similar to the query than the current pose we're playing back from the database
-				check(Result.PoseCost.Dissimilarity >= 0.0f && CurrentPoseCost.Dissimilarity >= 0.0f);
-				bool bBetterPose = Result.PoseCost.Dissimilarity * (1.0f + (Settings.MinPercentImprovement / 100.0f)) < CurrentPoseCost.Dissimilarity;
-
-				// Ignore the candidate poses from the same anim when they are too near to the current pose
-				bool bNearbyPose = false;
-				const FPoseSearchIndexAsset* StateSearchIndexAsset = InOutMotionMatchingState.GetCurrentSearchIndexAsset();
-				if (StateSearchIndexAsset == Result.SearchIndexAsset)
+				// Handle looping anims when checking for the pose being too close
+				if (!bNearbyPose && ResultDbSequence.bLoopAnimation)
 				{
-					const FPoseSearchDatabaseSequence& ResultDbSequence = Database->GetSourceAsset(Result.SearchIndexAsset);
-					bNearbyPose = FMath::Abs(InOutMotionMatchingState.AssetPlayerTime - Result.TimeOffsetSeconds) < Settings.PoseJumpThresholdTime;
-
-					// Handle looping anims when checking for the pose being too close
-					if (!bNearbyPose && ResultDbSequence.bLoopAnimation)
-					{
-						const float AssetLength = ResultDbSequence.Sequence->GetPlayLength();
-						bNearbyPose = FMath::Abs(AssetLength - InOutMotionMatchingState.AssetPlayerTime - Result.TimeOffsetSeconds) < Settings.PoseJumpThresholdTime;
-					}
+					const float AssetLength = ResultDbSequence.Sequence->GetPlayLength();
+					bNearbyPose = FMath::Abs(AssetLength - InOutMotionMatchingState.AssetPlayerTime - Result.TimeOffsetSeconds) < Settings.PoseJumpThresholdTime;
 				}
+			}
 
-				// Start playback from the candidate pose if we determined it was a better option
-				if (bBetterPose && !bNearbyPose)
-				{
-					InOutMotionMatchingState.JumpToPose(Context, Settings, Result);
-				}
+			// Start playback from the candidate pose if we determined it was a better option
+			if (bBetterPose && !bNearbyPose)
+			{
+				InOutMotionMatchingState.JumpToPose(Context, Settings, Result);
 			}
 		}
+	}
 
-		// Jump to the pose stepper's next anim if we didn't jump to a pose as a result of the search above
-		if (!(InOutMotionMatchingState.Flags & EMotionMatchingFlags::JumpedToPose)
-			&& PoseStepper.CanContinue()
-			&& PoseStepper.bJumpRequired)
-		{
-			const FPoseSearchDatabaseSequence& DbSequence = Database->GetSourceAsset(PoseStepper.Result.SearchIndexAsset);
-			InOutMotionMatchingState.JumpToPose(Context, Settings, PoseStepper.Result);
-			InOutMotionMatchingState.Flags |= EMotionMatchingFlags::JumpedToFollowUp;
-		}
+	// Jump to the pose stepper's next anim if we didn't jump to a pose as a result of the search above
+	if (!(InOutMotionMatchingState.Flags & EMotionMatchingFlags::JumpedToPose)
+		&& PoseStepper.CanContinue()
+		&& PoseStepper.bJumpRequired)
+	{
+		const FPoseSearchDatabaseSequence& DbSequence = Database->GetSourceAsset(PoseStepper.Result.SearchIndexAsset);
+		InOutMotionMatchingState.JumpToPose(Context, Settings, PoseStepper.Result);
+		InOutMotionMatchingState.Flags |= EMotionMatchingFlags::JumpedToFollowUp;
 	}
 
 	if (!(InOutMotionMatchingState.Flags & EMotionMatchingFlags::JumpedToPose))
@@ -434,3 +465,5 @@ float FMotionMatchingState::ComputeJumpBlendTime(
 
 	return JumpBlendTime;
 }
+
+#undef LOCTEXT_NAMESPACE
\ No newline at end of file
diff --git a/Engine/Plugins/Experimental/Animation/PoseSearch/Source/Runtime/Public/PoseSearch/PoseSearch.h b/Engine/Plugins/Experimental/Animation/PoseSearch/Source/Runtime/Public/PoseSearch/PoseSearch.h
--- a/Engine/Plugins/Experimental/Animation/PoseSearch/Source/Runtime/Public/PoseSearch/PoseSearch.h
+++ b/Engine/Plugins/Experimental/Animation/PoseSearch/Source/Runtime/Public/PoseSearch/PoseSearch.h
@@ -540,6 +540,8 @@ struct POSESEARCH_API FPoseSearchWeightParams
 
 	UPROPERTY(BlueprintReadWrite, EditAnywhere, Category = "Settings")
 	FPoseSearchChannelWeightParams TrajectoryWeight;
+
+	FPoseSearchWeightParams();
 };
 
 USTRUCT(BlueprintType, Category="Animation|Pose Search")
diff --git a/Engine/Plugins/Experimental/Animation/PoseSearch/Source/Runtime/Public/PoseSearch/PoseSearchLibrary.h b/Engine/Plugins/Experimental/Animation/PoseSearch/Source/Runtime/Public/PoseSearch/PoseSearchLibrary.h
--- a/Engine/Plugins/Experimental/Animation/PoseSearch/Source/Runtime/Public/PoseSearch/PoseSearchLibrary.h
+++ b/Engine/Plugins/Experimental/Animation/PoseSearch/Source/Runtime/Public/PoseSearch/PoseSearchLibrary.h
@@ -70,9 +70,8 @@ struct POSESEARCH_API FMotionMatchingSettings
 	float SearchThrottleTime = 0.1f;
 
 	// How much better the search result must be compared to the current pose in order to jump to it
-	// Note: This feature won't work quite as advertised until search data rescaling is implemented
 	UPROPERTY(EditAnywhere, BlueprintReadWrite, Category=Settings, meta=(ClampMin="0", ClampMax="100"))
-	float MinPercentImprovement = 100.f;
+	float MinPercentImprovement = 40.0f;
 };
 
 USTRUCT(BlueprintType, Category="Animation|Pose Search")
@@ -81,7 +80,7 @@ struct POSESEARCH_API FMotionMatchingState
 	GENERATED_BODY()
 
 	// Initializes the minimum required motion matching state
-	void InitNewDatabaseSearch(const UPoseSearchDatabase* Database, float SearchThrottleTime);
+	bool InitNewDatabaseSearch(const UPoseSearchDatabase* Database, float SearchThrottleTime, FText* OutError);
 
 	// Adds trajectory prediction and history information to ComposedQuery
 	void ComposeQuery(const UPoseSearchDatabase* Database, const FTrajectorySampleRange& Trajectory);
