MotionMatchingFFDemo - Added groups to pose search database
- Each group specifies Name, Weights (more settings in the future)
- Each sequence specifies the groups associated to it (could be more than 1)
- Initial implementation relies on user entering matching names for groups
- Added prototype customizations to the pose search database so we get a better picture of the UI

diff --git a/Engine/Config/BaseEngine.ini b/Engine/Config/BaseEngine.ini
--- a/Engine/Config/BaseEngine.ini
+++ b/Engine/Config/BaseEngine.ini
@@ -346,6 +346,8 @@ NearClipPlane=10.0
 +PropertyRedirects=(OldName="MultiLineEditableTextBox.bAutoWrapText",NewName="MultiLineEditableTextBox.AutoWrapText")
 +PropertyRedirects=(OldName="AnimNode_SequenceEvaluator.bShouldLoopWhenInSyncGroup",NewName="AnimNode_SequenceEvaluator.bShouldLoop")
 
++PropertyRedirects=(OldName="PoseSearchDatabase.Weights",NewName="PoseSearchDatabase.DefaultWeights")
+
 +EnumRedirects=(OldName="EControllerHand",NewName="/Script/InputCore.EControllerHand")
 +EnumRedirects=(OldName="AnimPhysConstraintType",NewName="AnimPhysAngularConstraintType")
 +EnumRedirects=(OldName="AnimPhysAxisType",NewName="AnimPhysLinearConstraintType")
diff --git a/Engine/Plugins/Experimental/Animation/PoseSearch/Source/Editor/PoseSearchEditor.Build.cs b/Engine/Plugins/Experimental/Animation/PoseSearch/Source/Editor/PoseSearchEditor.Build.cs
--- a/Engine/Plugins/Experimental/Animation/PoseSearch/Source/Editor/PoseSearchEditor.Build.cs
+++ b/Engine/Plugins/Experimental/Animation/PoseSearch/Source/Editor/PoseSearchEditor.Build.cs
@@ -40,7 +40,12 @@ public class PoseSearchEditor : ModuleRules
 		PrivateDependencyModuleNames.AddRange(
 			new string[]
 			{
-				"BlueprintGraph"
+				"BlueprintGraph",
+				"PropertyEditor",
+				"SlateCore",
+				"Slate",
+				"EditorStyle",
+				"DetailCustomizations"
 			}
 		);
 	}
diff --git a/Engine/Plugins/Experimental/Animation/PoseSearch/Source/Editor/Private/PoseSearchCustomization.cpp b/Engine/Plugins/Experimental/Animation/PoseSearch/Source/Editor/Private/PoseSearchCustomization.cpp
new file mode 100644
--- /dev/null
+++ b/Engine/Plugins/Experimental/Animation/PoseSearch/Source/Editor/Private/PoseSearchCustomization.cpp
@@ -0,0 +1,171 @@
+// Copyright Epic Games, Inc. All Rights Reserved.
+
+#include "PoseSearchCustomization.h"
+#include "PoseSearch/PoseSearch.h"
+#include "Engine/GameViewportClient.h"
+#include "AssetData.h"
+#include "EditorClassUtils.h"
+#include "IPropertyUtilities.h"
+#include "PropertyHandle.h"
+#include "PropertyCustomizationHelpers.h"
+#include "ScopedTransaction.h"
+#include "Engine/EngineBaseTypes.h"
+#include "UObject/UnrealType.h"
+#include "EditorStyleSet.h"
+#include "Widgets/DeclarativeSyntaxSupport.h"
+#include "Widgets/SBoxPanel.h"
+#include "Widgets/Layout/SWrapBox.h"
+#include "Widgets/Layout/SBox.h"
+#include "Widgets/Text/STextBlock.h"
+#include "Widgets/Images/SImage.h"
+#include "Widgets/Input/SButton.h"
+#include "IDetailChildrenBuilder.h"
+#include "DetailLayoutBuilder.h"
+#include "DetailWidgetRow.h"
+#include "DetailCategoryBuilder.h"
+#include "IDetailsView.h"
+#include "ObjectEditorUtils.h"
+
+#define LOCTEXT_NAMESPACE "PoseSearchCustomization"
+
+
+
+
+///////////////////////////////////////////////////////////////////////////////////////////////////////////////////////
+/// FPoseSearchDatabaseGroupCustomization
+
+void FPoseSearchDatabaseGroupCustomization::CustomizeHeader(TSharedRef<IPropertyHandle> InStructPropertyHandle, FDetailWidgetRow& HeaderRow, IPropertyTypeCustomizationUtils& StructCustomizationUtils)
+{
+	TSharedPtr<IPropertyHandle> NamePropertyHandle = InStructPropertyHandle->GetChildHandle(GET_MEMBER_NAME_CHECKED(FPoseSearchDatabaseGroup, Name));
+	FText NamePropertyText;
+	NamePropertyHandle->GetValueAsFormattedText(NamePropertyText);
+
+	HeaderRow
+	.NameContent()
+	[
+		SNew(STextBlock)
+		.Text(NamePropertyText)
+	]
+	.ValueContent()
+	[
+		InStructPropertyHandle->CreatePropertyValueWidget(false)
+	];
+
+	const FSimpleDelegate OnValueChanged = FSimpleDelegate::CreateLambda([&StructCustomizationUtils]()
+	{
+		StructCustomizationUtils.GetPropertyUtilities()->ForceRefresh();
+	});
+
+	InStructPropertyHandle->SetOnChildPropertyValueChanged(OnValueChanged);
+}
+
+void FPoseSearchDatabaseGroupCustomization::CustomizeChildren(TSharedRef<IPropertyHandle> InStructPropertyHandle, IDetailChildrenBuilder& StructBuilder, IPropertyTypeCustomizationUtils& StructCustomizationUtils)
+{
+	uint32 NumChildren;
+	InStructPropertyHandle->GetNumChildren(NumChildren);
+
+	for (uint32 ChildIndex = 0; ChildIndex < NumChildren; ++ChildIndex)
+	{
+		StructBuilder.AddProperty(InStructPropertyHandle->GetChildHandle(ChildIndex).ToSharedRef());
+	}
+}
+
+
+///////////////////////////////////////////////////////////////////////////////////////////////////////////////////////
+/// FPoseSearchDatabaseSequenceCustomization
+
+void FPoseSearchDatabaseSequenceCustomization::CustomizeHeader(TSharedRef<IPropertyHandle> InStructPropertyHandle, FDetailWidgetRow& HeaderRow, IPropertyTypeCustomizationUtils& StructCustomizationUtils)
+{
+	FString SequenceNameString;
+	FString GroupsString; // = TEXT("Groups: ");
+
+	TArray<UObject*> Objects;
+	InStructPropertyHandle->GetOuterObjects(Objects);
+
+	if (Objects.Num() == 1)
+	{
+		FPoseSearchDatabaseSequence* PoseSearchDatabaseSequence = (FPoseSearchDatabaseSequence*)InStructPropertyHandle->GetValueBaseAddress((uint8*)Objects[0]);
+		
+		SequenceNameString = PoseSearchDatabaseSequence->Sequence->GetName();
+		
+		const int32 NumGroups = PoseSearchDatabaseSequence->Groups.Num();
+		if (NumGroups == 0)
+		{
+			GroupsString.Append(TEXT("Default"));
+		}
+		else
+		{
+			for (int i = 0; i < NumGroups; ++i)
+			{
+				GroupsString.Append(PoseSearchDatabaseSequence->Groups[i].ToString());
+				if (i < NumGroups - 1)
+				{
+					GroupsString.Append(TEXT(" | "));
+				}
+			}
+		}
+	}
+
+	HeaderRow
+	.NameContent()
+	[
+		SNew(STextBlock)
+		.Text(FText::FromString(SequenceNameString))
+	]
+	.ValueContent()
+	[
+		SNew(SHorizontalBox)
+		+ SHorizontalBox::Slot()
+		.VAlign(VAlign_Center)
+		[
+			SNew(STextBlock)
+			.TextStyle(&FCoreStyle::Get().GetWidgetStyle<FTextBlockStyle>("SmallText"))
+			.Text(FText::FromString(GroupsString))
+		]
+	];
+
+	const FSimpleDelegate OnValueChanged = FSimpleDelegate::CreateLambda([&StructCustomizationUtils]()
+	{
+		StructCustomizationUtils.GetPropertyUtilities()->ForceRefresh();
+	});
+
+	InStructPropertyHandle->SetOnChildPropertyValueChanged(OnValueChanged);
+}
+
+void FPoseSearchDatabaseSequenceCustomization::CustomizeChildren(TSharedRef<IPropertyHandle> InStructPropertyHandle, IDetailChildrenBuilder& StructBuilder, IPropertyTypeCustomizationUtils& StructCustomizationUtils)
+{
+	uint32 NumChildren;
+	InStructPropertyHandle->GetNumChildren(NumChildren);
+
+	for (uint32 ChildIndex = 0; ChildIndex < NumChildren; ++ChildIndex)
+	{
+		StructBuilder.AddProperty(InStructPropertyHandle->GetChildHandle(ChildIndex).ToSharedRef());
+	}
+}
+
+
+///////////////////////////////////////////////////////////////////////////////////////////////////////////////////////
+/// FPoseSearchDatabaseDetails
+
+//TSharedRef<IDetailCustomization> FPoseSearchDatabaseDetails::MakeInstance()
+//{
+//	return MakeShareable(new FPoseSearchDatabaseDetails);
+//}
+//
+//void FPoseSearchDatabaseDetails::CustomizeDetails(IDetailLayoutBuilder& DetailBuilder)
+//{
+//	//TSharedPtr<IPropertyHandle> PropertyHandle = DetailBuilder.GetProperty("Sequences", UPoseSearchDatabase::StaticClass());
+//	//if (PropertyHandle.IsValid() && PropertyHandle->IsValidHandle())
+//	//{
+//	//	auto Refresh = FSimpleDelegate::CreateLambda([&DetailBuilder]()
+//	//	{
+//	//		DetailBuilder.ForceRefreshDetails();
+//	//	});
+//
+//	//	PropertyHandle->SetOnPropertyValueChanged(Refresh);
+//	//	PropertyHandle->SetOnChildPropertyValueChanged(Refresh);
+//	//}
+//}
+
+#undef LOCTEXT_NAMESPACE
+
diff --git a/Engine/Plugins/Experimental/Animation/PoseSearch/Source/Editor/Private/PoseSearchDebugger.cpp b/Engine/Plugins/Experimental/Animation/PoseSearch/Source/Editor/Private/PoseSearchDebugger.cpp
--- a/Engine/Plugins/Experimental/Animation/PoseSearch/Source/Editor/Private/PoseSearchDebugger.cpp
+++ b/Engine/Plugins/Experimental/Animation/PoseSearch/Source/Editor/Private/PoseSearchDebugger.cpp
@@ -312,13 +312,18 @@ namespace DebuggerDatabaseColumns
 		virtual TSharedRef<SWidget> GenerateWidget(const FRowDataRef& RowData) const override
 		{
 			return SNew(SHyperlink)
-				.Text(FText::FromString(RowData->AnimSequenceName))
+				.Text_Lambda([RowData]() -> FText { return FText::FromString(RowData->AnimSequenceName); })
 				.TextStyle(&FCoreStyle::Get().GetWidgetStyle<FTextBlockStyle>("SmallText"))
-				.ToolTipText(FText::Format(LOCTEXT("AssetHyperlinkTooltipFormat", "Open asset '{0}'"), FText::FromString(RowData->AnimSequencePath)))
-				.OnNavigate_Lambda([AnimSequencePath = RowData->AnimSequencePath]()
-			{
-				GEditor->GetEditorSubsystem<UAssetEditorSubsystem>()->OpenEditorForAsset(AnimSequencePath);
-			});
+				.ToolTipText_Lambda([RowData]() -> FText 
+					{ 
+						return FText::Format(
+							LOCTEXT("AssetHyperlinkTooltipFormat", "Open asset '{0}'"), 
+							FText::FromString(RowData->AnimSequencePath)); 
+					})
+				.OnNavigate_Lambda([RowData]()
+					{
+						GEditor->GetEditorSubsystem<UAssetEditorSubsystem>()->OpenEditorForAsset(RowData->AnimSequencePath);
+					});
 		}
 	};
 	const FName FAnimSequenceName::Name = "Sequence";
@@ -1865,12 +1870,19 @@ void FDebuggerViewModel::OnUpdateNodeSelection(int32 InNodeId)
 			break;
 		}
 	}
+
+	const UPoseSearchDatabase* NewDatabase = GetPoseSearchDatabase();
+
 	if (ActiveMotionMatchingState)
 	{
-		Skeletons[ActivePose].SequenceIdx = ActiveMotionMatchingState->DbPoseIdx;
+		//todo: this seems wrong, sequenceidx == poseidx? let's test doing it right
+		//Skeletons[ActivePose].SequenceIdx = ActiveMotionMatchingState->DbPoseIdx;
+
+		const FPoseSearchIndexAsset* Asset = 
+			NewDatabase->SearchIndex.FindAssetForPose(ActiveMotionMatchingState->DbPoseIdx);
+		Skeletons[ActivePose].SequenceIdx = Asset->SourceAssetIdx;
 	}
 
-	const UPoseSearchDatabase* NewDatabase = GetPoseSearchDatabase();
 	if (NewDatabase != CurrentDatabase)
 	{
 		ClearSelectedSkeleton();
diff --git a/Engine/Plugins/Experimental/Animation/PoseSearch/Source/Editor/Private/PoseSearchEditor.cpp b/Engine/Plugins/Experimental/Animation/PoseSearch/Source/Editor/Private/PoseSearchEditor.cpp
--- a/Engine/Plugins/Experimental/Animation/PoseSearch/Source/Editor/Private/PoseSearchEditor.cpp
+++ b/Engine/Plugins/Experimental/Animation/PoseSearch/Source/Editor/Private/PoseSearchEditor.cpp
@@ -1,9 +1,11 @@
 // Copyright Epic Games, Inc. All Rights Reserved.
 
 #include "PoseSearchEditor.h"
+#include "PoseSearchCustomization.h"
 #include "Animation/AnimSequence.h"
 #include "Modules/ModuleManager.h"
 #include "Editor.h"
+#include "PropertyEditorModule.h"
 #include "Subsystems/AssetEditorSubsystem.h"
 #include "IAnimationEditor.h"
 #include "IPersonaToolkit.h"
@@ -70,6 +72,18 @@ private:
 	TSharedPtr<FDebuggerViewCreator> DebuggerViewCreator;
 	/** Enables dedicated PoseSearch trace module */
 	TSharedPtr<FTraceModule> TraceModule;
+	
+private:
+	void RegisterPropertyTypeCustomizations();
+	void RegisterObjectCustomizations();
+	void UnregisterCustomizations();
+	void RegisterCustomClassLayout(FName ClassName, FOnGetDetailCustomizationInstance DetailLayoutDelegate);
+	void RegisterCustomPropertyTypeLayout(FName PropertyTypeName, FOnGetPropertyTypeCustomizationInstance PropertyTypeLayoutDelegate);
+
+private:
+	/** List of registered class that we must unregister when the module shuts down */
+	TSet<FName> RegisteredClassNames;
+	TSet<FName> RegisteredPropertyTypes;
 };
 
 void FEditorModule::StartupModule()
@@ -90,6 +104,14 @@ void FEditorModule::StartupModule()
 			ECVF_Default
 		));
 	}
+
+	FPropertyEditorModule& PropertyModule = FModuleManager::GetModuleChecked<FPropertyEditorModule>("PropertyEditor");
+
+	RegisterPropertyTypeCustomizations();
+	RegisterObjectCustomizations();
+
+	PropertyModule.NotifyCustomizationModuleChanged();
+
 }
 
 void FEditorModule::ShutdownModule()
@@ -100,10 +122,76 @@ void FEditorModule::ShutdownModule()
 	}
 	ConsoleCommands.Empty();
 	
+	UnregisterCustomizations();
+	
 	IModularFeatures::Get().UnregisterModularFeature(TraceServices::ModuleFeatureName, TraceModule.Get());
 	FDebugger::Shutdown();
 }
 
+void FEditorModule::RegisterCustomClassLayout(FName ClassName, FOnGetDetailCustomizationInstance DetailLayoutDelegate)
+{
+	check(ClassName != NAME_None);
+
+	RegisteredClassNames.Add(ClassName);
+
+	static FName PropertyEditor("PropertyEditor");
+	FPropertyEditorModule& PropertyModule = FModuleManager::GetModuleChecked<FPropertyEditorModule>(PropertyEditor);
+	PropertyModule.RegisterCustomClassLayout(ClassName, DetailLayoutDelegate);
+}
+
+
+void FEditorModule::RegisterCustomPropertyTypeLayout(FName PropertyTypeName, FOnGetPropertyTypeCustomizationInstance PropertyTypeLayoutDelegate)
+{
+	check(PropertyTypeName != NAME_None);
+
+	RegisteredPropertyTypes.Add(PropertyTypeName);
+
+	static FName PropertyEditor("PropertyEditor");
+	FPropertyEditorModule& PropertyModule = FModuleManager::GetModuleChecked<FPropertyEditorModule>(PropertyEditor);
+	PropertyModule.RegisterCustomPropertyTypeLayout(PropertyTypeName, PropertyTypeLayoutDelegate);
+}
+
+
+void FEditorModule::RegisterPropertyTypeCustomizations()
+{
+	RegisterCustomPropertyTypeLayout("PoseSearchDatabaseSequence", FOnGetPropertyTypeCustomizationInstance::CreateStatic(&FPoseSearchDatabaseSequenceCustomization::MakeInstance));
+	RegisterCustomPropertyTypeLayout("PoseSearchDatabaseGroup", FOnGetPropertyTypeCustomizationInstance::CreateStatic(&FPoseSearchDatabaseGroupCustomization::MakeInstance));
+}
+
+void FEditorModule::RegisterObjectCustomizations()
+{
+//	RegisterCustomClassLayout("PoseSearchDatabase", FOnGetDetailCustomizationInstance::CreateStatic(&FPoseSearchDatabaseDetails::MakeInstance));
+}
+
+void FEditorModule::UnregisterCustomizations()
+{
+	if (FModuleManager::Get().IsModuleLoaded("PropertyEditor"))
+	{
+		FPropertyEditorModule& PropertyModule = FModuleManager::GetModuleChecked<FPropertyEditorModule>("PropertyEditor");
+
+		// Unregister all classes customized by name
+		for (auto It = RegisteredClassNames.CreateConstIterator(); It; ++It)
+		{
+			if (It->IsValid())
+			{
+				PropertyModule.UnregisterCustomClassLayout(*It);
+			}
+		}
+
+		// Unregister all structures
+		for (auto It = RegisteredPropertyTypes.CreateConstIterator(); It; ++It)
+		{
+			if (It->IsValid())
+			{
+				PropertyModule.UnregisterCustomPropertyTypeLayout(*It);
+			}
+		}
+
+		PropertyModule.NotifyCustomizationModuleChanged();
+	}
+}
+
+
 }} // namespace UE::PoseSearch
 
 IMPLEMENT_MODULE(UE::PoseSearch::FEditorModule, PoseSearchEditor);
diff --git a/Engine/Plugins/Experimental/Animation/PoseSearch/Source/Editor/Public/PoseSearchCustomization.h b/Engine/Plugins/Experimental/Animation/PoseSearch/Source/Editor/Public/PoseSearchCustomization.h
new file mode 100644
--- /dev/null
+++ b/Engine/Plugins/Experimental/Animation/PoseSearch/Source/Editor/Public/PoseSearchCustomization.h
@@ -0,0 +1,62 @@
+// Copyright Epic Games, Inc. All Rights Reserved.
+
+#pragma once
+
+#include "CoreMinimal.h"
+#include "IPropertyTypeCustomization.h"
+#include "UObject/WeakObjectPtr.h"
+#include "Layout/Visibility.h"
+#include "Input/Reply.h"
+#include "Widgets/SWidget.h"
+#include "IDetailCustomization.h"
+
+class IPropertyHandle;
+class IDetailLayoutBuilder;
+
+/**
+ * Customizes FPoseSearchDatabaseGroup to show the name of the group as the header
+ */
+class FPoseSearchDatabaseGroupCustomization : public IPropertyTypeCustomization
+{
+public:
+	static TSharedRef<IPropertyTypeCustomization> MakeInstance()
+	{
+		return MakeShareable(new FPoseSearchDatabaseGroupCustomization);
+	}
+
+	/** IPropertyTypeCustomization interface */
+	virtual void CustomizeHeader(TSharedRef<class IPropertyHandle> InStructPropertyHandle, class FDetailWidgetRow& HeaderRow, IPropertyTypeCustomizationUtils& StructCustomizationUtils) override;
+	virtual void CustomizeChildren(TSharedRef<class IPropertyHandle> InStructPropertyHandle, class IDetailChildrenBuilder& StructBuilder, IPropertyTypeCustomizationUtils& StructCustomizationUtils) override;
+};
+
+/**
+ * Customizes FPoseSearchDatabaseSequence to show the name of the sequence and the groups in the header
+ */
+class FPoseSearchDatabaseSequenceCustomization : public IPropertyTypeCustomization
+{
+public:
+	static TSharedRef<IPropertyTypeCustomization> MakeInstance()
+	{
+		return MakeShareable(new FPoseSearchDatabaseSequenceCustomization);
+	}
+
+	/** IPropertyTypeCustomization interface */
+	virtual void CustomizeHeader(TSharedRef<class IPropertyHandle> InStructPropertyHandle, class FDetailWidgetRow& HeaderRow, IPropertyTypeCustomizationUtils& StructCustomizationUtils) override;
+	virtual void CustomizeChildren(TSharedRef<class IPropertyHandle> InStructPropertyHandle, class IDetailChildrenBuilder& StructBuilder, IPropertyTypeCustomizationUtils& StructCustomizationUtils) override;
+};
+
+
+
+/**
+ * Customizes UPoseSearchDatabase to improve workflow
+ */
+//class FPoseSearchDatabaseDetails : public IDetailCustomization
+//{
+//public:
+//	// Creates an instance of FObjectDetails
+//	static TSharedRef<IDetailCustomization> MakeInstance();
+//
+//	// IDetailCustomization interface
+//	virtual void CustomizeDetails(IDetailLayoutBuilder& DetailLayout) override;
+//	// End of IDetailCustomization interface
+//};
diff --git a/Engine/Plugins/Experimental/Animation/PoseSearch/Source/Runtime/Private/PoseSearch.cpp b/Engine/Plugins/Experimental/Animation/PoseSearch/Source/Runtime/Private/PoseSearch.cpp
--- a/Engine/Plugins/Experimental/Animation/PoseSearch/Source/Runtime/Private/PoseSearch.cpp
+++ b/Engine/Plugins/Experimental/Animation/PoseSearch/Source/Runtime/Private/PoseSearch.cpp
@@ -898,17 +898,28 @@ void FPoseSearchWeightsContext::Update(
 
 	if (bRecomputeWeights)
 	{
-		int32 NumGroups = ActiveDatabase ? 1 : 0;
+		ComputedDefaultGroupWeights.Init(Database->DefaultWeights, Database->Schema, ActiveWeights);
+
+		int32 NumGroups = ActiveDatabase ? ActiveDatabase->Groups.Num() : 0;
 		ComputedGroupWeights.SetNum(NumGroups);
-		for (FPoseSearchWeights& GroupWeights: ComputedGroupWeights)
+
+		for (int GroupIdx = 0; GroupIdx < NumGroups; ++GroupIdx)
 		{
-			GroupWeights.Init(Database->Weights, Database->Schema, ActiveWeights);
+			ComputedGroupWeights[GroupIdx].Init(
+				Database->Groups[GroupIdx].Weights, 
+				Database->Schema, 
+				ActiveWeights);
 		}
 	}
 }
 
 const FPoseSearchWeights* FPoseSearchWeightsContext::GetGroupWeights(int32 WeightsGroupIdx) const
 {
+	if (WeightsGroupIdx == INDEX_NONE)
+	{
+		return &ComputedDefaultGroupWeights;
+	}
+
 	if (ComputedGroupWeights.IsValidIndex(WeightsGroupIdx))
 	{
 		return &ComputedGroupWeights[WeightsGroupIdx];
@@ -1259,6 +1270,8 @@ bool UPoseSearchDatabase::TryInitSearchIndexAssets()
 	SearchIndex.Assets.Empty();
 	bool bAnyMirrored = false;
 	TArray<FFloatRange> ValidRanges;
+	TArray<int32> GroupIndices;
+	TArray<int32> BadGroupSequenceIndices;
 	for (int32 SequenceIdx = 0; SequenceIdx < Sequences.Num(); ++SequenceIdx)
 	{
 		const FPoseSearchDatabaseSequence& Sequence = Sequences[SequenceIdx];
@@ -1269,27 +1282,58 @@ bool UPoseSearchDatabase::TryInitSearchIndexAssets()
 			Sequence.MirrorOption == EPoseSearchMirrorOption::MirroredOnly ||
 			Sequence.MirrorOption == EPoseSearchMirrorOption::UnmirroredAndMirrored;
 
-		ValidRanges.Reset();
-		FindValidSequenceIntervals(Sequence, ValidRanges);
-		for (const FFloatRange& Range : ValidRanges)
+		GroupIndices.Reset();
+		for (const FName& SequenceGroupName : Sequence.Groups)
 		{
-			if (bAddUnmirrored)
+			const int32 SequenceGroupIndex = Groups.IndexOfByPredicate([&](const FPoseSearchDatabaseGroup& DatabaseGroup)
+			{
+				return DatabaseGroup.Name == SequenceGroupName;
+			});
+
+			// we don't add INDEX_NONE because index none represents a choice to use the default group by not adding
+			// any group identifiers. If an added identifier doesn't match, that's an error. In the future this
+			// should be made robust enough to prevent these errors from happening
+			if (SequenceGroupIndex == INDEX_NONE)
+			{
+				BadGroupSequenceIndices.Add(SequenceIdx);
+			}
+			else
 			{
-				SearchIndex.Assets.Add(
-					FPoseSearchIndexAsset(
-						SequenceIdx,
-						false,
-						FFloatInterval(Range.GetLowerBoundValue(), Range.GetUpperBoundValue())));
+				GroupIndices.Add(SequenceGroupIndex);
 			}
+		}
 
-			if (bAddMirrored)
+		if (GroupIndices.Num() == 0)
+		{
+			GroupIndices.Add(INDEX_NONE);
+		}
+
+		for (int32 GroupIndex : GroupIndices)
+		{
+			ValidRanges.Reset();
+			FindValidSequenceIntervals(Sequence, ValidRanges);
+			for (const FFloatRange& Range : ValidRanges)
 			{
-				SearchIndex.Assets.Add(
-					FPoseSearchIndexAsset(
-						SequenceIdx,
-						true,
-						FFloatInterval(Range.GetLowerBoundValue(), Range.GetUpperBoundValue())));
-				bAnyMirrored = true;
+				if (bAddUnmirrored)
+				{
+					SearchIndex.Assets.Add(
+						FPoseSearchIndexAsset(
+							GroupIndex,
+							SequenceIdx,
+							false,
+							FFloatInterval(Range.GetLowerBoundValue(), Range.GetUpperBoundValue())));
+				}
+
+				if (bAddMirrored)
+				{
+					SearchIndex.Assets.Add(
+						FPoseSearchIndexAsset(
+							GroupIndex,
+							SequenceIdx,
+							true,
+							FFloatInterval(Range.GetLowerBoundValue(), Range.GetUpperBoundValue())));
+					bAnyMirrored = true;
+				}
 			}
 		}
 	}
@@ -1306,6 +1350,16 @@ bool UPoseSearchDatabase::TryInitSearchIndexAssets()
 		return false;
 	}
 
+	for (int32 BadGroupSequenceIdx : BadGroupSequenceIndices)
+	{
+		UE_LOG(
+			LogPoseSearch,
+			Warning,
+			TEXT("Database %s, sequence %s is asking for a group that doesn't exist"),
+			*GetNameSafe(this),
+			*GetNameSafe(Sequences[BadGroupSequenceIdx].Sequence));
+	}
+
 	return true;
 }
 
@@ -3046,8 +3100,8 @@ void FSequenceIndexer::AddMetadata(int32 SampleIdx)
 
 	const bool bBlockTransition =
 		!Input.MainSequence->Input.bLoopable &&
-		(SampleTime < Input.BlockTransitionParameters.SequenceStartInterval || 
-		 SampleTime > SequenceLength - Input.BlockTransitionParameters.SequenceEndInterval);
+		(SampleTime < Input.RequestedSamplingRange.Min + Input.BlockTransitionParameters.SequenceStartInterval ||
+		 SampleTime > Input.RequestedSamplingRange.Max - Input.BlockTransitionParameters.SequenceEndInterval);
 
 	if (bBlockTransition)
 	{
@@ -4004,21 +4058,25 @@ FSearchResult Search(FSearchContext& SearchContext)
 	FPoseCost BestPoseCost;
 	int32 BestPoseIdx = INDEX_NONE;
 
-	for (int32 PoseIdx = 0; PoseIdx != SearchIndex->NumPoses; ++PoseIdx)
+	for (const FPoseSearchIndexAsset& Asset : SearchIndex->Assets)
 	{
-		const FPoseSearchPoseMetadata& Metadata = SearchIndex->PoseMetadata[PoseIdx];
-
-		if (EnumHasAnyFlags(Metadata.Flags, EPoseSearchPoseFlags::BlockTransition))
+		const int32 EndIndex = Asset.FirstPoseIdx + Asset.NumPoses;
+		for (int32 PoseIdx = Asset.FirstPoseIdx; PoseIdx < EndIndex; ++PoseIdx)
 		{
-			continue;
-		}
+			const FPoseSearchPoseMetadata& Metadata = SearchIndex->PoseMetadata[PoseIdx];
+
+			if (EnumHasAnyFlags(Metadata.Flags, EPoseSearchPoseFlags::BlockTransition))
+			{
+				continue;
+			}
 
-		FPoseCost PoseCost = ComparePoses(PoseIdx, SearchContext);
+			FPoseCost PoseCost = ComparePoses(PoseIdx, SearchContext, Asset.SourceGroupIdx);
 
-		if (PoseCost < BestPoseCost)
-		{
-			BestPoseCost = PoseCost;
-			BestPoseIdx = PoseIdx;
+			if (PoseCost < BestPoseCost)
+			{
+				BestPoseCost = PoseCost;
+				BestPoseIdx = PoseIdx;
+			}
 		}
 	}
 
@@ -4061,7 +4119,8 @@ void ComputePoseCostAddends(
 	OutNotifyAddend = PoseMetadata.CostAddend;
 }
 
-FPoseCost ComparePoses(int32 PoseIdx, FSearchContext& SearchContext)
+
+FPoseCost ComparePoses(int32 PoseIdx, FSearchContext& SearchContext, int32 GroupIdx /* = INDEX_NONE */)
 {
 	FPoseCost Result;
 
@@ -4073,10 +4132,11 @@ FPoseCost ComparePoses(int32 PoseIdx, FSearchContext& SearchContext)
 
 	if (SearchContext.WeightsContext)
 	{
-		const FPoseSearchWeights* WeightsSet = SearchContext.WeightsContext->GetGroupWeights(0);
+		const FPoseSearchWeights* WeightsSet =
+			SearchContext.WeightsContext->GetGroupWeights(GroupIdx);
 		Result.Dissimilarity = CompareFeatureVectors(
-			PoseValues.Num(), 
-			PoseValues.GetData(), 
+			PoseValues.Num(),
+			PoseValues.GetData(),
 			SearchContext.QueryValues.GetData(),
 			WeightsSet->Weights.GetData());
 	}
@@ -4090,10 +4150,11 @@ FPoseCost ComparePoses(int32 PoseIdx, FSearchContext& SearchContext)
 	ComputePoseCostAddends(PoseIdx, SearchContext, NotifyAddend, MirrorMismatchAddend);
 	Result.CostAddend = NotifyAddend + MirrorMismatchAddend;
 	Result.TotalCost = Result.Dissimilarity + Result.CostAddend;
-	
+
 	return Result;
 }
 
+
 FPoseCost ComparePoses(int32 PoseIdx, FSearchContext& SearchContext, FPoseCostDetails& OutPoseCostDetails)
 {
 	using namespace Eigen;
@@ -4117,7 +4178,9 @@ FPoseCost ComparePoses(int32 PoseIdx, FSearchContext& SearchContext, FPoseCostDe
 	// Compute weighted squared difference vector
 	if (SearchContext.WeightsContext)
 	{
-		const FPoseSearchWeights* WeightsSet = SearchContext.WeightsContext->GetGroupWeights(0);
+		const FPoseSearchIndexAsset* SearchIndexAsset = SearchContext.GetSearchIndex()->FindAssetForPose(PoseIdx);
+		const FPoseSearchWeights* WeightsSet = 
+			SearchContext.WeightsContext->GetGroupWeights(SearchIndexAsset->SourceGroupIdx);
 		check(WeightsSet);
 		check(WeightsSet->Weights.Num() == Dims);
 		auto WeightsVector = Map<const ArrayXf>(WeightsSet->Weights.GetData(), Dims);
diff --git a/Engine/Plugins/Experimental/Animation/PoseSearch/Source/Runtime/Public/PoseSearch/PoseSearch.h b/Engine/Plugins/Experimental/Animation/PoseSearch/Source/Runtime/Public/PoseSearch/PoseSearch.h
--- a/Engine/Plugins/Experimental/Animation/PoseSearch/Source/Runtime/Public/PoseSearch/PoseSearch.h
+++ b/Engine/Plugins/Experimental/Animation/PoseSearch/Source/Runtime/Public/PoseSearch/PoseSearch.h
@@ -328,12 +328,20 @@ public:
 	FPoseSearchIndexAsset()
 	{}
 
-	FPoseSearchIndexAsset(int32 InSourceAssetIdx, bool bInMirrored, const FFloatInterval& InSamplingInterval)
-		: SourceAssetIdx(InSourceAssetIdx)
+	FPoseSearchIndexAsset(
+		int32 InSourceGroupIdx, 
+		int32 InSourceAssetIdx, 
+		bool bInMirrored, 
+		const FFloatInterval& InSamplingInterval)
+		: SourceGroupIdx(InSourceGroupIdx)
+		, SourceAssetIdx(InSourceAssetIdx)
 		, bMirrored(bInMirrored)
 		, SamplingInterval(InSamplingInterval)
 	{}
 
+	UPROPERTY()
+	int32 SourceGroupIdx = INDEX_NONE;
+
 	// Index of the source asset in search index's container (i.e. UPoseSearchDatabase)
 	UPROPERTY()
 	int32 SourceAssetIdx = INDEX_NONE;
@@ -597,6 +605,9 @@ private:
 	UPROPERTY(Transient)
 	FPoseSearchDynamicWeightParams DynamicWeights;
 
+	UPROPERTY(Transient)
+	FPoseSearchWeights ComputedDefaultGroupWeights;
+
 	UPROPERTY(Transient)
 	TArray<FPoseSearchWeights> ComputedGroupWeights;
 };
@@ -650,9 +661,24 @@ struct POSESEARCH_API FPoseSearchDatabaseSequence
 	UPROPERTY(EditAnywhere, Category="Sequence")
 	bool bLoopFollowUpAnimation = false;
 
+	UPROPERTY(EditAnywhere, Category = "Group")
+	TArray<FName> Groups;
+
 	FFloatInterval GetEffectiveSamplingRange() const;
 };
 
+USTRUCT()
+struct POSESEARCH_API FPoseSearchDatabaseGroup
+{
+	GENERATED_BODY()
+public:
+	UPROPERTY(EditAnywhere, Category = "Database")
+	FName Name;
+
+	UPROPERTY(EditAnywhere, Category = "Database")
+	FPoseSearchWeightParams Weights;
+};
+
 /** A data asset for indexing a collection of animation sequences. */
 UCLASS(BlueprintType, Category = "Animation|Pose Search")
 class POSESEARCH_API UPoseSearchDatabase : public UDataAsset
@@ -664,7 +690,7 @@ public:
 	const UPoseSearchSchema* Schema;
 
 	UPROPERTY(EditAnywhere, Category = "Database")
-	FPoseSearchWeightParams Weights;
+	FPoseSearchWeightParams DefaultWeights;
 
 	// If there's a mirroring mismatch between the currently playing sequence and a search candidate, this cost will be 
 	// added to the candidate, making it less likely to be selected
@@ -677,6 +703,9 @@ public:
 	UPROPERTY(EditAnywhere, Category = "Database")
 	FPoseSearchBlockTransitionParameters BlockTransitionParameters;
 
+	UPROPERTY(EditAnywhere, Category = "Database")
+	TArray<FPoseSearchDatabaseGroup> Groups;
+
 	// Drag and drop animations here to add them in bulk to Sequences
 	UPROPERTY(EditAnywhere, Category = "Database", DisplayName="Drag And Drop Anims Here")
 	TArray<TObjectPtr<UAnimSequence>> SimpleSequences;
@@ -699,7 +728,6 @@ public:
 	const FPoseSearchDatabaseSequence& GetSourceAsset(const FPoseSearchIndexAsset* SearchIndexAsset) const;
 
 public: // UObject
-
 	virtual void PreSave(FObjectPreSaveContext ObjectSaveContext) override;
 
 #if WITH_EDITOR
@@ -983,18 +1011,18 @@ POSESEARCH_API bool BuildIndex(UPoseSearchDatabase* Database);
 */
 POSESEARCH_API FSearchResult Search(FSearchContext& SearchContext);
 
+
 /**
 * Evaluate pose comparison metric between a pose in the search index and an input query
-* 
-* @param SearchIndex	The search index containing the pose to compare to the query
-* @param PoseIdx		The index of the pose in the search index to compare to the query
-* @param SearchContext	Structure containing search parameters
-* 
+*
+* @param PoseIdx			The index of the pose in the search index to compare to the query
+* @param SearchContext		Structure containing search parameters
+* @param GroupIdx			Indicated which weight to use when evaluating dissimilarity
+*
 * @return Dissimilarity between the two poses
 */
 
-POSESEARCH_API FPoseCost ComparePoses(int32 PoseIdx, FSearchContext& SearchContext);
-
+POSESEARCH_API FPoseCost ComparePoses(int32 PoseIdx, FSearchContext& SearchContext, int32 GroupIdx = INDEX_NONE);
 
 /**
  * Cost details for pose analysis in the rewind debugger
