MotionMatchingFFDemo - Non-unity build fix

diff --git a/Engine/Plugins/Experimental/Animation/PoseSearch/Source/Editor/Private/PoseSearchCustomization.cpp b/Engine/Plugins/Experimental/Animation/PoseSearch/Source/Editor/Private/PoseSearchCustomization.cpp
--- a/Engine/Plugins/Experimental/Animation/PoseSearch/Source/Editor/Private/PoseSearchCustomization.cpp
+++ b/Engine/Plugins/Experimental/Animation/PoseSearch/Source/Editor/Private/PoseSearchCustomization.cpp
@@ -25,6 +25,7 @@
 #include "DetailCategoryBuilder.h"
 #include "IDetailsView.h"
 #include "ObjectEditorUtils.h"
+#include "Animation/AnimSequence.h"
 
 #define LOCTEXT_NAMESPACE "PoseSearchCustomization"
 
