MotionMatchingFFDemo - Implemented Database Group filtering through FGameplayTagQuery

diff --git a/Engine/Plugins/Experimental/Animation/PoseSearch/Source/Editor/PoseSearchEditor.Build.cs b/Engine/Plugins/Experimental/Animation/PoseSearch/Source/Editor/PoseSearchEditor.Build.cs
--- a/Engine/Plugins/Experimental/Animation/PoseSearch/Source/Editor/PoseSearchEditor.Build.cs
+++ b/Engine/Plugins/Experimental/Animation/PoseSearch/Source/Editor/PoseSearchEditor.Build.cs
@@ -18,6 +18,7 @@ public class PoseSearchEditor : ModuleRules
 				"CoreUObject",
 				"Engine",
 				"PoseSearch",
+				"GameplayTags",
 				
 				// Trace-related dependencies
 				"TraceLog",
diff --git a/Engine/Plugins/Experimental/Animation/PoseSearch/Source/Editor/Private/PoseSearchCustomization.cpp b/Engine/Plugins/Experimental/Animation/PoseSearch/Source/Editor/Private/PoseSearchCustomization.cpp
--- a/Engine/Plugins/Experimental/Animation/PoseSearch/Source/Editor/Private/PoseSearchCustomization.cpp
+++ b/Engine/Plugins/Experimental/Animation/PoseSearch/Source/Editor/Private/PoseSearchCustomization.cpp
@@ -37,15 +37,26 @@
 
 void FPoseSearchDatabaseGroupCustomization::CustomizeHeader(TSharedRef<IPropertyHandle> InStructPropertyHandle, FDetailWidgetRow& HeaderRow, IPropertyTypeCustomizationUtils& StructCustomizationUtils)
 {
-	TSharedPtr<IPropertyHandle> NamePropertyHandle = InStructPropertyHandle->GetChildHandle(GET_MEMBER_NAME_CHECKED(FPoseSearchDatabaseGroup, Name));
 	FText NamePropertyText;
-	NamePropertyHandle->GetValueAsFormattedText(NamePropertyText);
+
+	TArray<UObject*> Objects;
+	InStructPropertyHandle->GetOuterObjects(Objects);
+	if (Objects.Num() == 1)
+	{
+		FPoseSearchDatabaseGroup* PoseSearchDatabaseGroup = (FPoseSearchDatabaseGroup*)InStructPropertyHandle->GetValueBaseAddress((uint8*)Objects[0]);
+		NamePropertyText = FText::FromString(PoseSearchDatabaseGroup->Tag.ToString());
+	}
+	else
+	{
+		TSharedPtr<IPropertyHandle> NamePropertyHandle = InStructPropertyHandle->GetChildHandle(GET_MEMBER_NAME_CHECKED(FPoseSearchDatabaseGroup, Tag));
+		NamePropertyHandle->GetValueAsFormattedText(NamePropertyText);
+	}
 
 	HeaderRow
 	.NameContent()
 	[
 		SNew(STextBlock)
-		.Text(NamePropertyText)
+		.Text(NamePropertyText) 
 	]
 	.ValueContent()
 	[
@@ -64,10 +75,10 @@ void FPoseSearchDatabaseGroupCustomization::CustomizeChildren(TSharedRef<IProper
 {
 	uint32 NumChildren;
 	InStructPropertyHandle->GetNumChildren(NumChildren);
-
 	for (uint32 ChildIndex = 0; ChildIndex < NumChildren; ++ChildIndex)
 	{
-		StructBuilder.AddProperty(InStructPropertyHandle->GetChildHandle(ChildIndex).ToSharedRef());
+		TSharedPtr<IPropertyHandle> ChildPropertyHandle = InStructPropertyHandle->GetChildHandle(ChildIndex);
+		StructBuilder.AddProperty(ChildPropertyHandle.ToSharedRef());
 	}
 }
 
@@ -89,7 +100,7 @@ void FPoseSearchDatabaseSequenceCustomization::CustomizeHeader(TSharedRef<IPrope
 		
 		SequenceNameString = PoseSearchDatabaseSequence->Sequence->GetName();
 		
-		const int32 NumGroups = PoseSearchDatabaseSequence->Groups.Num();
+		const int32 NumGroups = PoseSearchDatabaseSequence->GroupTags.Num();
 		if (NumGroups == 0)
 		{
 			GroupsString.Append(TEXT("Default"));
@@ -98,7 +109,7 @@ void FPoseSearchDatabaseSequenceCustomization::CustomizeHeader(TSharedRef<IPrope
 		{
 			for (int i = 0; i < NumGroups; ++i)
 			{
-				GroupsString.Append(PoseSearchDatabaseSequence->Groups[i].ToString());
+				GroupsString.Append(PoseSearchDatabaseSequence->GroupTags.GetByIndex(i).ToString());
 				if (i < NumGroups - 1)
 				{
 					GroupsString.Append(TEXT(" | "));
diff --git a/Engine/Plugins/Experimental/Animation/PoseSearch/Source/Runtime/PoseSearch.Build.cs b/Engine/Plugins/Experimental/Animation/PoseSearch/Source/Runtime/PoseSearch.Build.cs
--- a/Engine/Plugins/Experimental/Animation/PoseSearch/Source/Runtime/PoseSearch.Build.cs
+++ b/Engine/Plugins/Experimental/Animation/PoseSearch/Source/Runtime/PoseSearch.Build.cs
@@ -35,7 +35,8 @@ public class PoseSearch : ModuleRules
 				"AnimationCore",
 				"TraceLog",
 				"AnimGraphRuntime",
-				"DynamicPlayRate"
+				"DynamicPlayRate",
+				"GameplayTags"
 			}
 		);
 
diff --git a/Engine/Plugins/Experimental/Animation/PoseSearch/Source/Runtime/Private/AnimNode_MotionMatching.cpp b/Engine/Plugins/Experimental/Animation/PoseSearch/Source/Runtime/Private/AnimNode_MotionMatching.cpp
--- a/Engine/Plugins/Experimental/Animation/PoseSearch/Source/Runtime/Private/AnimNode_MotionMatching.cpp
+++ b/Engine/Plugins/Experimental/Animation/PoseSearch/Source/Runtime/Private/AnimNode_MotionMatching.cpp
@@ -48,6 +48,7 @@ void FAnimNode_MotionMatching::UpdateAssetPlayer(const FAnimationUpdateContext&
 	// Execute core motion matching algorithm and retain across frame state
 	UpdateMotionMatchingState(Context
 		, Database
+		, bUseDatabaseTagQuery ? &DatabaseTagQuery : nullptr
 		, Trajectory
 		, Settings
 		, MotionMatchingState
diff --git a/Engine/Plugins/Experimental/Animation/PoseSearch/Source/Runtime/Private/PoseSearch.cpp b/Engine/Plugins/Experimental/Animation/PoseSearch/Source/Runtime/Private/PoseSearch.cpp
--- a/Engine/Plugins/Experimental/Animation/PoseSearch/Source/Runtime/Private/PoseSearch.cpp
+++ b/Engine/Plugins/Experimental/Animation/PoseSearch/Source/Runtime/Private/PoseSearch.cpp
@@ -1283,11 +1283,11 @@ bool UPoseSearchDatabase::TryInitSearchIndexAssets()
 			Sequence.MirrorOption == EPoseSearchMirrorOption::UnmirroredAndMirrored;
 
 		GroupIndices.Reset();
-		for (const FName& SequenceGroupName : Sequence.Groups)
+		for (const FGameplayTag& SequenceGroupTag : Sequence.GroupTags)
 		{
 			const int32 SequenceGroupIndex = Groups.IndexOfByPredicate([&](const FPoseSearchDatabaseGroup& DatabaseGroup)
 			{
-				return DatabaseGroup.Name == SequenceGroupName;
+				return DatabaseGroup.Tag == SequenceGroupTag;
 			});
 
 			// we don't add INDEX_NONE because index none represents a choice to use the default group by not adding
@@ -1297,7 +1297,7 @@ bool UPoseSearchDatabase::TryInitSearchIndexAssets()
 			{
 				BadGroupSequenceIndices.Add(SequenceIdx);
 			}
-			else
+			else if (Groups[SequenceGroupIndex].bUseGroupWeights)
 			{
 				GroupIndices.Add(SequenceGroupIndex);
 			}
@@ -3125,7 +3125,6 @@ void FSequenceIndexer::AddMetadata(int32 SampleIdx)
 	}
 }
 
-
 //////////////////////////////////////////////////////////////////////////
 // PoseSearch API
 
@@ -4058,8 +4057,19 @@ FSearchResult Search(FSearchContext& SearchContext)
 	FPoseCost BestPoseCost;
 	int32 BestPoseIdx = INDEX_NONE;
 
+	const UPoseSearchDatabase* Database = SearchContext.GetSourceDatabase();
+
 	for (const FPoseSearchIndexAsset& Asset : SearchIndex->Assets)
 	{
+		if (Database && SearchContext.DatabaseTagQuery)
+		{
+			const FPoseSearchDatabaseSequence& DbSequence = Database->Sequences[Asset.SourceAssetIdx];
+			if (!SearchContext.DatabaseTagQuery->Matches(DbSequence.GroupTags))
+			{
+				continue;
+			}
+		}
+
 		const int32 EndIndex = Asset.FirstPoseIdx + Asset.NumPoses;
 		for (int32 PoseIdx = Asset.FirstPoseIdx; PoseIdx < EndIndex; ++PoseIdx)
 		{
diff --git a/Engine/Plugins/Experimental/Animation/PoseSearch/Source/Runtime/Private/PoseSearchLibrary.cpp b/Engine/Plugins/Experimental/Animation/PoseSearch/Source/Runtime/Private/PoseSearchLibrary.cpp
--- a/Engine/Plugins/Experimental/Animation/PoseSearch/Source/Runtime/Private/PoseSearchLibrary.cpp
+++ b/Engine/Plugins/Experimental/Animation/PoseSearch/Source/Runtime/Private/PoseSearchLibrary.cpp
@@ -224,6 +224,7 @@ void FMotionMatchingState::JumpToPose(const FAnimationUpdateContext& Context, co
 void UpdateMotionMatchingState(
 	const FAnimationUpdateContext& Context,
 	const UPoseSearchDatabase* Database,
+	const FGameplayTagQuery* DatabaseTagQuery,
 	const FTrajectorySampleRange& Trajectory,
 	const FMotionMatchingSettings& Settings,
 	FMotionMatchingState& InOutMotionMatchingState
@@ -280,6 +281,7 @@ void UpdateMotionMatchingState(
 	SearchContext.SetSource(InOutMotionMatchingState.CurrentDatabase.Get());
 	SearchContext.QueryValues = InOutMotionMatchingState.ComposedQuery.GetNormalizedValues();
 	SearchContext.WeightsContext = &InOutMotionMatchingState.WeightsContext;
+	SearchContext.DatabaseTagQuery = DatabaseTagQuery;
 	if (const FPoseSearchIndexAsset* CurrentIndexAsset = InOutMotionMatchingState.GetCurrentSearchIndexAsset())
 	{
 		SearchContext.QueryMirrorRequest =
@@ -412,7 +414,13 @@ void UPoseSearchLibrary::UpdateMotionMatchingForSequencePlayer(
 			InOutMotionMatchingState.AssetPlayerTime = SequencePlayerNode->GetAccumulatedTime();
 
 			// Execute core motion matching algorithm and retain across frame state
-			UpdateMotionMatchingState(*AnimationUpdateContext, Database, Trajectory, Settings, InOutMotionMatchingState);
+			UpdateMotionMatchingState(
+				*AnimationUpdateContext, 
+				Database, 
+				nullptr, // DatabaseTagQuery
+				Trajectory, 
+				Settings, 
+				InOutMotionMatchingState);
 
 			// If a new pose is requested, jump to the pose by updating the embedded sequence player node
 			if ((InOutMotionMatchingState.Flags & EMotionMatchingFlags::JumpedToPose) == EMotionMatchingFlags::JumpedToPose)
diff --git a/Engine/Plugins/Experimental/Animation/PoseSearch/Source/Runtime/Public/PoseSearch/AnimNode_MotionMatching.h b/Engine/Plugins/Experimental/Animation/PoseSearch/Source/Runtime/Public/PoseSearch/AnimNode_MotionMatching.h
--- a/Engine/Plugins/Experimental/Animation/PoseSearch/Source/Runtime/Public/PoseSearch/AnimNode_MotionMatching.h
+++ b/Engine/Plugins/Experimental/Animation/PoseSearch/Source/Runtime/Public/PoseSearch/AnimNode_MotionMatching.h
@@ -25,6 +25,14 @@ public:
 	UPROPERTY(EditAnywhere, BlueprintReadWrite, Category=Settings, meta=(PinShownByDefault))
 	TObjectPtr<const UPoseSearchDatabase> Database = nullptr;
 
+
+	UPROPERTY(EditAnywhere, BlueprintReadWrite, Category = Settings, meta = (PinHiddenByDefault))
+	bool bUseDatabaseTagQuery = false;
+
+	// Query used to filter database groups which can be searched
+	UPROPERTY(EditAnywhere, BlueprintReadWrite, Category = Settings, meta = (PinShownByDefault))
+	FGameplayTagQuery DatabaseTagQuery;
+
 	// Motion trajectory samples for pose search queries
 	UPROPERTY(EditAnywhere, BlueprintReadWrite, Category=Settings, meta=(PinShownByDefault))
 	FTrajectorySampleRange Trajectory;
diff --git a/Engine/Plugins/Experimental/Animation/PoseSearch/Source/Runtime/Public/PoseSearch/PoseSearch.h b/Engine/Plugins/Experimental/Animation/PoseSearch/Source/Runtime/Public/PoseSearch/PoseSearch.h
--- a/Engine/Plugins/Experimental/Animation/PoseSearch/Source/Runtime/Public/PoseSearch/PoseSearch.h
+++ b/Engine/Plugins/Experimental/Animation/PoseSearch/Source/Runtime/Public/PoseSearch/PoseSearch.h
@@ -15,6 +15,7 @@
 #include "Animation/MotionTrajectoryTypes.h"
 #include "AlphaBlend.h"
 #include "BoneIndices.h"
+#include "GameplayTagContainer.h"
 #include "Interfaces/Interface_BoneReferenceSkeletonProvider.h"
 
 #include "PoseSearch.generated.h"
@@ -662,7 +663,7 @@ struct POSESEARCH_API FPoseSearchDatabaseSequence
 	bool bLoopFollowUpAnimation = false;
 
 	UPROPERTY(EditAnywhere, Category = "Group")
-	TArray<FName> Groups;
+	FGameplayTagContainer GroupTags;
 
 	FFloatInterval GetEffectiveSamplingRange() const;
 };
@@ -671,11 +672,15 @@ USTRUCT()
 struct POSESEARCH_API FPoseSearchDatabaseGroup
 {
 	GENERATED_BODY()
+
 public:
-	UPROPERTY(EditAnywhere, Category = "Database")
-	FName Name;
+	UPROPERTY(EditAnywhere, Category = "Settings")
+	FGameplayTag Tag;
 
-	UPROPERTY(EditAnywhere, Category = "Database")
+	UPROPERTY(EditAnywhere, Category = "Settings")
+	bool bUseGroupWeights = false;
+
+	UPROPERTY(EditAnywhere, Category = "Settings", meta=(EditCondition="bUseGroupWeights", EditConditionHides))
 	FPoseSearchWeightParams Weights;
 };
 
@@ -742,7 +747,6 @@ public:
 	bool TryInitSearchIndexAssets();
 };
 
-
 /** 
 * Helper object for writing features into a float buffer according to a feature vector layout.
 * Keeps track of which features are present, allowing the feature vector to be built up piecemeal.
@@ -958,12 +962,14 @@ struct POSESEARCH_API FSearchContext
 	TArrayView<const float> QueryValues;
 	EPoseSearchBooleanRequest QueryMirrorRequest = EPoseSearchBooleanRequest::Indifferent;
 	const FPoseSearchWeightsContext* WeightsContext = nullptr;
+	const FGameplayTagQuery* DatabaseTagQuery = nullptr;
 	FDebugDrawParams DebugDrawParams;
 
 	void SetSource(const UPoseSearchDatabase* InSourceDatabase);
 	void SetSource(const UAnimSequenceBase* InSourceSequence);
 	const FPoseSearchIndex* GetSearchIndex() const;
 	float GetMirrorMismatchCost() const;
+	const UPoseSearchDatabase* GetSourceDatabase() { return SourceDatabase; }
 
 private:
 	const UPoseSearchDatabase* SourceDatabase = nullptr;
@@ -972,7 +978,6 @@ private:
 	float MirrorMismatchCost = 0.0f;
 };
 
-
 /**
 * Visualize pose search debug information
 *
diff --git a/Engine/Plugins/Experimental/Animation/PoseSearch/Source/Runtime/Public/PoseSearch/PoseSearchLibrary.h b/Engine/Plugins/Experimental/Animation/PoseSearch/Source/Runtime/Public/PoseSearch/PoseSearchLibrary.h
--- a/Engine/Plugins/Experimental/Animation/PoseSearch/Source/Runtime/Public/PoseSearch/PoseSearchLibrary.h
+++ b/Engine/Plugins/Experimental/Animation/PoseSearch/Source/Runtime/Public/PoseSearch/PoseSearchLibrary.h
@@ -136,6 +136,7 @@ struct POSESEARCH_API FMotionMatchingState
 */
 POSESEARCH_API void UpdateMotionMatchingState(const FAnimationUpdateContext& Context
 	, const UPoseSearchDatabase* Database
+	, const FGameplayTagQuery* DatabaseTagQuery
 	, const FTrajectorySampleRange& Trajectory
 	, const FMotionMatchingSettings& Settings
 	, FMotionMatchingState& InOutMotionMatchingState
