MotionMatchingFFDemo - Updated Pose Search Debugger to allow the choice to see all or only filtered sequence groups

diff --git a/Engine/Plugins/Experimental/Animation/PoseSearch/Source/Editor/Private/PoseSearchDebugger.cpp b/Engine/Plugins/Experimental/Animation/PoseSearch/Source/Editor/Private/PoseSearchDebugger.cpp
--- a/Engine/Plugins/Experimental/Animation/PoseSearch/Source/Editor/Private/PoseSearchDebugger.cpp
+++ b/Engine/Plugins/Experimental/Animation/PoseSearch/Source/Editor/Private/PoseSearchDebugger.cpp
@@ -660,6 +660,17 @@ void SDebuggerDatabaseView::OnDatabaseRowSelectionChanged(
 	}
 }
 
+ECheckBoxState SDebuggerDatabaseView::IsSequenceFilterEnabled() const
+{
+	return bSequenceFilterEnabled ? ECheckBoxState::Checked : ECheckBoxState::Unchecked;
+}
+
+void SDebuggerDatabaseView::OnSequenceFilterEnabledChanged(ECheckBoxState NewState)
+{
+	bSequenceFilterEnabled = NewState == ECheckBoxState::Checked;
+	FilterDatabaseRows();
+}
+
 void SDebuggerDatabaseView::SortDatabaseRows()
 {
 	if (SortMode == EColumnSortMode::Ascending)
@@ -689,23 +700,29 @@ void SDebuggerDatabaseView::FilterDatabaseRows()
 	{
 		for (const auto& UnfilteredRow : UnfilteredDatabaseRows)
 		{
-			FilteredDatabaseView.Rows.Add(UnfilteredRow);
+			if (!bSequenceFilterEnabled || DatabaseSequenceFilter[UnfilteredRow->DbSequenceIdx])
+			{
+				FilteredDatabaseView.Rows.Add(UnfilteredRow);
+			}
 		}
 	}
 	else
 	{
 		for (const auto& UnfilteredRow : UnfilteredDatabaseRows)
 		{
-			bool bMatchesAllTokens = Algo::AllOf(
-				Tokens,
-				[&](FString Token)
+			if (!bSequenceFilterEnabled || DatabaseSequenceFilter[UnfilteredRow->DbSequenceIdx])
 			{
-				return UnfilteredRow->AnimSequenceName.Contains(Token);
-			});
+				bool bMatchesAllTokens = Algo::AllOf(
+					Tokens,
+					[&](FString Token)
+				{
+					return UnfilteredRow->AnimSequenceName.Contains(Token);
+				});
 
-			if (bMatchesAllTokens)
-			{
-				FilteredDatabaseView.Rows.Add(UnfilteredRow);
+				if (bMatchesAllTokens)
+				{
+					FilteredDatabaseView.Rows.Add(UnfilteredRow);
+				}
 			}
 		}
 	}
@@ -733,6 +750,7 @@ void SDebuggerDatabaseView::CreateRows(const UPoseSearchDatabase& Database)
 			Row->PoseIdx = PoseIdx;
 			Row->AnimSequenceName = DbSequence.Sequence->GetName();
 			Row->AnimSequencePath = DbSequence.Sequence->GetPathName();
+			Row->DbSequenceIdx = SearchIndexAsset.SourceAssetIdx;
 			Row->Time = Time;
 			Row->AnimFrame = DbSequence.Sequence->GetFrameAtTime(Time);
 			Row->bMirrored = SearchIndexAsset.bMirrored;
@@ -779,6 +797,8 @@ void SDebuggerDatabaseView::UpdateRows(const FTraceMotionMatchingStateMessage& S
 		}
 	}
 
+	DatabaseSequenceFilter = State.DatabaseSequenceFilter;
+
 	SortDatabaseRows();
 	FilterDatabaseRows();
 
@@ -980,8 +1000,24 @@ void SDebuggerDatabaseView::Construct(const FArguments& InArgs)
 			.Padding(0.0f, 0.0f, 0.0f, 5.0f)
 			.AutoHeight()
 			[
-				SAssignNew(FilterBox, SSearchBox)
-				.OnTextChanged(this, &SDebuggerDatabaseView::OnFilterTextChanged)
+				SNew(SHorizontalBox)
+				+ SHorizontalBox::Slot()
+				.AutoWidth()
+				.Padding(10, 5, 10, 5)
+				[
+					SNew(SCheckBox)
+					.IsChecked(this, &SDebuggerDatabaseView::IsSequenceFilterEnabled)
+					.OnCheckStateChanged(this, &SDebuggerDatabaseView::OnSequenceFilterEnabledChanged)
+					[
+						SNew(STextBlock)
+						.Text(LOCTEXT("PoseSearchDebuggerGroupFiltering", "Apply Group Filtering"))
+					]
+				]
+				+ SHorizontalBox::Slot()
+				[
+					SAssignNew(FilterBox, SSearchBox)
+					.OnTextChanged(this, &SDebuggerDatabaseView::OnFilterTextChanged)
+				]
 			]
 		
 			+ SVerticalBox::Slot()
diff --git a/Engine/Plugins/Experimental/Animation/PoseSearch/Source/Editor/Private/PoseSearchDebugger.h b/Engine/Plugins/Experimental/Animation/PoseSearch/Source/Editor/Private/PoseSearchDebugger.h
--- a/Engine/Plugins/Experimental/Animation/PoseSearch/Source/Editor/Private/PoseSearchDebugger.h
+++ b/Engine/Plugins/Experimental/Animation/PoseSearch/Source/Editor/Private/PoseSearchDebugger.h
@@ -306,6 +306,12 @@ private:
 	/** Row selection to update model view */
 	void OnDatabaseRowSelectionChanged(TSharedPtr<FDebuggerDatabaseRowData> Row, ESelectInfo::Type SelectInfo);
 
+	/** Informs the widget if sequence filtering is enabled */
+	ECheckBoxState IsSequenceFilterEnabled() const;
+
+	/** Updates the state of sequence filtering */
+	void OnSequenceFilterEnabledChanged(ECheckBoxState NewState);
+
 	/** Generates a database row widget for the given data */
 	TSharedRef<ITableRow> HandleGenerateDatabaseRow(TSharedRef<FDebuggerDatabaseRowData> Item, const TSharedRef<STableViewBase>& OwnerTable) const;
 	
@@ -354,6 +360,7 @@ private:
 	TArray<TSharedRef<FDebuggerDatabaseRowData>> UnfilteredDatabaseRows;
 
 	TWeakObjectPtr<const UPoseSearchDatabase> RowsSourceDatabase = nullptr;
+	TArrayView<const bool> DatabaseSequenceFilter;
 
 	/** Database listing for filtered poses */
 	FTable FilteredDatabaseView;
@@ -363,6 +370,9 @@ private:
 
 	/** Text used to filter DatabaseView */
 	FText FilterText;
+
+	/** True if only sequences that pass the sequence group query are being displayed */
+	bool bSequenceFilterEnabled = true;
 };
 
 /**
diff --git a/Engine/Plugins/Experimental/Animation/PoseSearch/Source/Editor/Private/Trace/PoseSearchTraceAnalyzer.cpp b/Engine/Plugins/Experimental/Animation/PoseSearch/Source/Editor/Private/Trace/PoseSearchTraceAnalyzer.cpp
--- a/Engine/Plugins/Experimental/Animation/PoseSearch/Source/Editor/Private/Trace/PoseSearchTraceAnalyzer.cpp
+++ b/Engine/Plugins/Experimental/Animation/PoseSearch/Source/Editor/Private/Trace/PoseSearchTraceAnalyzer.cpp
@@ -54,6 +54,7 @@ bool FTraceAnalyzer::OnEvent(uint16 RouteId, EStyle Style, const FOnEventContext
 				Message.SimAngularVelocity = EventData.GetValue<float>("SimAngularVelocity");
 				Message.AnimLinearVelocity = EventData.GetValue<float>("AnimLinearVelocity");
 				Message.AnimAngularVelocity = EventData.GetValue<float>("AnimAngularVelocity");
+				Message.DatabaseSequenceFilter = EventData.GetArrayView<bool>("DatabaseSequenceFilter");
 
 				TArrayView<const float> ChannelWeightScales = EventData.GetArrayView<float>("ChannelWeightScales");
 				TArrayView<const float> HistoryWeightScales = EventData.GetArrayView<float>("HistoryWeightScales");
diff --git a/Engine/Plugins/Experimental/Animation/PoseSearch/Source/Editor/Private/Trace/PoseSearchTraceProvider.h b/Engine/Plugins/Experimental/Animation/PoseSearch/Source/Editor/Private/Trace/PoseSearchTraceProvider.h
--- a/Engine/Plugins/Experimental/Animation/PoseSearch/Source/Editor/Private/Trace/PoseSearchTraceProvider.h
+++ b/Engine/Plugins/Experimental/Animation/PoseSearch/Source/Editor/Private/Trace/PoseSearchTraceProvider.h
@@ -24,18 +24,19 @@ struct FTraceMessage
 struct FTraceMotionMatchingStateMessage : FTraceMessage
 {
 	FTraceMotionMatchingState::EFlags Flags;
-	float ElapsedPoseJumpTime;
+	float ElapsedPoseJumpTime = -1.0f;
 	TArray<float> QueryVector;
 	TArray<float> QueryVectorNormalized;
 	FPoseSearchDynamicWeightParams Weights;
-	int32 DbPoseIdx;
-	uint64 DatabaseId;
+	int32 DbPoseIdx = INDEX_NONE;
+	uint64 DatabaseId = INDEX_NONE;
 	float AssetPlayerTime = 0.0f;
 	float DeltaTime = 0.0f;
 	float SimLinearVelocity = 0.0f;
 	float SimAngularVelocity = 0.0f;
 	float AnimLinearVelocity = 0.0f;
 	float AnimAngularVelocity = 0.0f;
+	TArray<bool> DatabaseSequenceFilter;
 };
 
 
diff --git a/Engine/Plugins/Experimental/Animation/PoseSearch/Source/Runtime/Private/PoseSearchLibrary.cpp b/Engine/Plugins/Experimental/Animation/PoseSearch/Source/Runtime/Private/PoseSearchLibrary.cpp
--- a/Engine/Plugins/Experimental/Animation/PoseSearch/Source/Runtime/Private/PoseSearchLibrary.cpp
+++ b/Engine/Plugins/Experimental/Animation/PoseSearch/Source/Runtime/Private/PoseSearchLibrary.cpp
@@ -60,6 +60,29 @@ void ComputeCurrentVelocityMagnitudes(
 	}
 }
 
+void ComputeDatabaseSequenceFilter(
+	const UPoseSearchDatabase* Database, 
+	const FGameplayTagQuery* Query, 
+	TArray<bool>& OutDbSequenceFilter)
+{
+	OutDbSequenceFilter.SetNum(Database->Sequences.Num());
+
+	if (Query)
+	{
+		for (int SeqIdx = 0; SeqIdx < Database->Sequences.Num(); ++SeqIdx)
+		{
+			OutDbSequenceFilter[SeqIdx] = Query->Matches(Database->Sequences[SeqIdx].GroupTags);
+		}
+	}
+	else
+	{
+		for (int SeqIdx = 0; SeqIdx < Database->Sequences.Num(); ++SeqIdx)
+		{
+			OutDbSequenceFilter[SeqIdx] = true;
+		}
+	}
+}
+
 //////////////////////////////////////////////////////////////////////////
 // FMotionMatchingPoseStepper
 
@@ -372,6 +395,9 @@ void UpdateMotionMatchingState(
 			SimLinearVelocity, SimAngularVelocity, 
 			AnimLinearVelocity, AnimAngularVelocity);
 
+		TArray<bool> DatabaseSequenceFilter;
+		ComputeDatabaseSequenceFilter(Database, DatabaseTagQuery, DatabaseSequenceFilter);
+
 		FTraceMotionMatchingState TraceState;
 		if (EnumHasAnyFlags(InOutMotionMatchingState.Flags, EMotionMatchingFlags::JumpedToFollowUp))
 		{
@@ -392,6 +418,7 @@ void UpdateMotionMatchingState(
 		TraceState.SimAngularVelocity = SimAngularVelocity;
 		TraceState.AnimLinearVelocity = AnimLinearVelocity;
 		TraceState.AnimAngularVelocity = AnimAngularVelocity;
+		TraceState.DatabaseSequenceFilter = DatabaseSequenceFilter;
 		UE_TRACE_POSE_SEARCH_MOTION_MATCHING_STATE(Context, TraceState)
 	}
 #endif
diff --git a/Engine/Plugins/Experimental/Animation/PoseSearch/Source/Runtime/Private/Trace/PoseSearchTraceLogger.cpp b/Engine/Plugins/Experimental/Animation/PoseSearch/Source/Runtime/Private/Trace/PoseSearchTraceLogger.cpp
--- a/Engine/Plugins/Experimental/Animation/PoseSearch/Source/Runtime/Private/Trace/PoseSearchTraceLogger.cpp
+++ b/Engine/Plugins/Experimental/Animation/PoseSearch/Source/Runtime/Private/Trace/PoseSearchTraceLogger.cpp
@@ -30,6 +30,7 @@ UE_TRACE_EVENT_BEGIN(PoseSearch, MotionMatchingState)
 	UE_TRACE_EVENT_FIELD(float, SimAngularVelocity)
 	UE_TRACE_EVENT_FIELD(float, AnimLinearVelocity)
 	UE_TRACE_EVENT_FIELD(float, AnimAngularVelocity)
+	UE_TRACE_EVENT_FIELD(bool[], DatabaseSequenceFilter)
 UE_TRACE_EVENT_END()
 
 namespace UE { namespace PoseSearch {
@@ -103,7 +104,8 @@ void FTraceMotionMatchingState::Output(const FAnimationBaseContext& InContext, c
 		<< MotionMatchingState.SimLinearVelocity(State.SimLinearVelocity)
 		<< MotionMatchingState.SimAngularVelocity(State.SimAngularVelocity)
 		<< MotionMatchingState.AnimLinearVelocity(State.AnimLinearVelocity)
-		<< MotionMatchingState.AnimAngularVelocity(State.AnimAngularVelocity);
+		<< MotionMatchingState.AnimAngularVelocity(State.AnimAngularVelocity)
+		<< MotionMatchingState.DatabaseSequenceFilter(State.DatabaseSequenceFilter.GetData(), State.DatabaseSequenceFilter.Num());
 }
 
 }}
diff --git a/Engine/Plugins/Experimental/Animation/PoseSearch/Source/Runtime/Private/Trace/PoseSearchTraceLogger.h b/Engine/Plugins/Experimental/Animation/PoseSearch/Source/Runtime/Private/Trace/PoseSearchTraceLogger.h
--- a/Engine/Plugins/Experimental/Animation/PoseSearch/Source/Runtime/Private/Trace/PoseSearchTraceLogger.h
+++ b/Engine/Plugins/Experimental/Animation/PoseSearch/Source/Runtime/Private/Trace/PoseSearchTraceLogger.h
@@ -47,6 +47,9 @@ struct POSESEARCH_API FTraceMotionMatchingState
 	/** Runtime weights */
 	FPoseSearchDynamicWeightParams Weights;
 
+	/** If true, groups are being filtered by DatabaseGroupQuery */
+	TArrayView<const bool> DatabaseSequenceFilter;
+
 	/** Index of the pose in our database */
 	int32 DbPoseIdx = 0;
 
