PoseSearch: Fixed bugs when poses are cached/retrieved with FPoseHistory

diff --git a/Engine/Plugins/Experimental/Animation/PoseSearch/Source/Runtime/Private/PoseSearch.cpp b/Engine/Plugins/Experimental/Animation/PoseSearch/Source/Runtime/Private/PoseSearch.cpp
--- a/Engine/Plugins/Experimental/Animation/PoseSearch/Source/Runtime/Private/PoseSearch.cpp
+++ b/Engine/Plugins/Experimental/Animation/PoseSearch/Source/Runtime/Private/PoseSearch.cpp
@@ -1566,9 +1566,9 @@ bool FPoseSearchFeatureVectorBuilder::TrySetPoseFeatures(UE::PoseSearch::FPoseHi
 			Feature.SchemaBoneIdx = SchemaBoneIdx;
 
 			int32 SkeletonBoneIndex = Schema->BoneIndices[SchemaBoneIdx];
-			FCompactPoseBoneIndex CompactBoneIndex = BoneContainer.GetCompactPoseIndexFromSkeletonIndex(SkeletonBoneIndex);
-			const FTransform& Transform = ComponentPose[CompactBoneIndex.GetInt()];
-			const FTransform& PrevTransform = ComponentPrevPose[CompactBoneIndex.GetInt()];
+
+			const FTransform& Transform = ComponentPose[SkeletonBoneIndex];
+			const FTransform& PrevTransform = ComponentPrevPose[SkeletonBoneIndex];
 			SetTransform(Feature, Transform);
 			SetTransformVelocity(Feature, Transform, PrevTransform, History->GetSampleTimeInterval());
 		}
@@ -1710,15 +1710,18 @@ namespace UE { namespace PoseSearch {
 static void CopyCompactToSkeletonPose(const FCompactPose& Pose, TArray<FTransform>& OutLocalTransforms)
 {
 	const FBoneContainer& BoneContainer = Pose.GetBoneContainer();
-	const FReferenceSkeleton& RefSkeleton = BoneContainer.GetReferenceSkeleton();
+	const USkeleton* SkeletonAsset = BoneContainer.GetSkeletonAsset();
+	check(SkeletonAsset);
+
+	const FReferenceSkeleton& RefSkeleton = SkeletonAsset->GetReferenceSkeleton();
 	TArrayView<const FTransform> RefSkeletonTransforms = MakeArrayView(RefSkeleton.GetRefBonePose());
+	const int32 NumSkeletonBones = RefSkeleton.GetNum();
 
-	const int32 NumSkeletonBones = BoneContainer.GetNumBones();
 	OutLocalTransforms.SetNum(NumSkeletonBones);
 
 	for (auto SkeletonBoneIdx = FSkeletonPoseBoneIndex(0); SkeletonBoneIdx != NumSkeletonBones; ++SkeletonBoneIdx)
 	{
-		FCompactPoseBoneIndex CompactBoneIdx = BoneContainer.GetCompactPoseIndexFromSkeletonIndex(SkeletonBoneIdx.GetInt());
+		FCompactPoseBoneIndex CompactBoneIdx = BoneContainer.GetCompactPoseIndexFromSkeletonPoseIndex(SkeletonBoneIdx);
 		OutLocalTransforms[SkeletonBoneIdx.GetInt()] = CompactBoneIdx.IsValid() ? Pose[CompactBoneIdx] : RefSkeletonTransforms[SkeletonBoneIdx.GetInt()];
 	}
 }
@@ -2083,10 +2086,6 @@ private:
 void FAnimSamplingContext::Init(const UPoseSearchSchema* Schema)
 {
 	MirrorDataTable = Schema->MirrorDataTable;
-	// ***
-	// Note this is a change for the UPoseSearchSequenceMetaData flavor of BuildIndex
-	// Was previously doing Sequence->GetSkeleton() instead of Schema->GetSkeleton()
-	// Hmm, is that okay?
 	BoneContainer.InitializeTo(Schema->BoneIndicesWithParents, FCurveEvaluationOption(false), *Schema->Skeleton);
 	FillCompactPoseAndComponentRefRotations();
 }
