PoseSearch:
- Marked classes as experimental
- Fixed motion matching node debug rendering lingering in the scene after exiting a state
- Updated mirroring editor names

diff --git a/Engine/Plugins/Experimental/Animation/MotionTrajectory/Source/MotionTrajectory/Public/MotionTrajectory.h b/Engine/Plugins/Experimental/Animation/MotionTrajectory/Source/MotionTrajectory/Public/MotionTrajectory.h
--- a/Engine/Plugins/Experimental/Animation/MotionTrajectory/Source/MotionTrajectory/Public/MotionTrajectory.h
+++ b/Engine/Plugins/Experimental/Animation/MotionTrajectory/Source/MotionTrajectory/Public/MotionTrajectory.h
@@ -43,7 +43,7 @@ struct MOTIONTRAJECTORY_API FMotionTrajectorySettings
 };
 
 // Abstract component interface for the minimum Motion Trajectory prediction and history API
-UCLASS(ClassGroup=Movement, abstract, BlueprintType, Category="Motion Trajectory")
+UCLASS(ClassGroup=Movement, abstract, BlueprintType, Category="Motion Trajectory", Experimental)
 class MOTIONTRAJECTORY_API UMotionTrajectoryComponent : public UActorComponent
 {
 	GENERATED_BODY()
diff --git a/Engine/Plugins/Experimental/Animation/PoseSearch/Source/Editor/Public/AnimGraphNode_MotionMatching.h b/Engine/Plugins/Experimental/Animation/PoseSearch/Source/Editor/Public/AnimGraphNode_MotionMatching.h
--- a/Engine/Plugins/Experimental/Animation/PoseSearch/Source/Editor/Public/AnimGraphNode_MotionMatching.h
+++ b/Engine/Plugins/Experimental/Animation/PoseSearch/Source/Editor/Public/AnimGraphNode_MotionMatching.h
@@ -9,7 +9,7 @@
 #include "AnimGraphNode_MotionMatching.generated.h"
 
 
-UCLASS(MinimalAPI)
+UCLASS(MinimalAPI, Experimental)
 class UAnimGraphNode_MotionMatching : public UAnimGraphNode_AssetPlayerBase
 {
 	GENERATED_BODY()
diff --git a/Engine/Plugins/Experimental/Animation/PoseSearch/Source/Editor/Public/AnimGraphNode_PoseSearchHistoryCollector.h b/Engine/Plugins/Experimental/Animation/PoseSearch/Source/Editor/Public/AnimGraphNode_PoseSearchHistoryCollector.h
--- a/Engine/Plugins/Experimental/Animation/PoseSearch/Source/Editor/Public/AnimGraphNode_PoseSearchHistoryCollector.h
+++ b/Engine/Plugins/Experimental/Animation/PoseSearch/Source/Editor/Public/AnimGraphNode_PoseSearchHistoryCollector.h
@@ -9,7 +9,7 @@
 #include "AnimGraphNode_PoseSearchHistoryCollector.generated.h"
 
 
-UCLASS(MinimalAPI)
+UCLASS(MinimalAPI, Experimental)
 class UAnimGraphNode_PoseSearchHistoryCollector : public UAnimGraphNode_Base
 {
 	GENERATED_BODY()
diff --git a/Engine/Plugins/Experimental/Animation/PoseSearch/Source/Runtime/Private/AnimNode_MotionMatching.cpp b/Engine/Plugins/Experimental/Animation/PoseSearch/Source/Runtime/Private/AnimNode_MotionMatching.cpp
--- a/Engine/Plugins/Experimental/Animation/PoseSearch/Source/Runtime/Private/AnimNode_MotionMatching.cpp
+++ b/Engine/Plugins/Experimental/Animation/PoseSearch/Source/Runtime/Private/AnimNode_MotionMatching.cpp
@@ -34,6 +34,10 @@ void FAnimNode_MotionMatching::Evaluate_AnyThread(FPoseContext& Output)
 	DECLARE_SCOPE_HIERARCHICAL_COUNTER_ANIMNODE(Evaluate_AnyThread);
 
 	Source.Evaluate(Output);
+
+#if WITH_EDITORONLY_DATA
+	bWasEvaluated = true;
+#endif
 }
 
 void FAnimNode_MotionMatching::UpdateAssetPlayer(const FAnimationUpdateContext& Context)
@@ -46,12 +50,13 @@ void FAnimNode_MotionMatching::UpdateAssetPlayer(const FAnimationUpdateContext&
 	MotionMatchingState.AssetPlayerTime = SequencePlayerNode.GetAccumulatedTime();
 
 	// Execute core motion matching algorithm and retain across frame state
-	UpdateMotionMatchingState(Context
-		, Database
-		, bUseDatabaseTagQuery ? &DatabaseTagQuery : nullptr
-		, Trajectory
-		, Settings
-		, MotionMatchingState
+	UpdateMotionMatchingState(
+		Context,
+		Database,
+		bUseDatabaseTagQuery ? &DatabaseTagQuery : nullptr,
+		Trajectory,
+		Settings,
+		MotionMatchingState
 	);
 
 	// If a new pose is requested, jump to the pose by updating the embedded sequence player node
@@ -62,20 +67,21 @@ void FAnimNode_MotionMatching::UpdateAssetPlayer(const FAnimationUpdateContext&
 		SequencePlayerNode.SetSequence(ResultDbSequence.Sequence);
 		SequencePlayerNode.SetAccumulatedTime(MotionMatchingState.AssetPlayerTime);
 		SequencePlayerNode.SetLoopAnimation(ResultDbSequence.bLoopAnimation);
-		SequencePlayerNode.SetPlayRate(1.f);
+		SequencePlayerNode.SetPlayRate(1.0f);
 
 		MirrorNode.SetMirrorDataTable(Database->Schema->MirrorDataTable.Get());
 		MirrorNode.SetMirror(SearchIndexAsset->bMirrored);
 	}
 
 	// Optionally applying dynamic play rate adjustment to chosen sequences based on predictive motion analysis
-	const float PlayRate = DynamicPlayRateAdjustment(Context
-		, Trajectory
-		, DynamicPlayRateSettings
-		, SequencePlayerNode.GetSequence()
-		, SequencePlayerNode.GetAccumulatedTime()
-		, SequencePlayerNode.GetPlayRate()
-		, SequencePlayerNode.GetLoopAnimation()
+	const float PlayRate = DynamicPlayRateAdjustment(
+		Context,
+		Trajectory,
+		DynamicPlayRateSettings,
+		SequencePlayerNode.GetSequence(),
+		SequencePlayerNode.GetAccumulatedTime(),
+		SequencePlayerNode.GetPlayRate(),
+		SequencePlayerNode.GetLoopAnimation()
 	);
 
 	SequencePlayerNode.SetPlayRate(PlayRate);
@@ -94,7 +100,7 @@ bool FAnimNode_MotionMatching::HasPreUpdate() const
 void FAnimNode_MotionMatching::PreUpdate(const UAnimInstance* InAnimInstance)
 {
 #if WITH_EDITORONLY_DATA
-	if (bDebugDraw)
+	if (bWasEvaluated && bDebugDraw)
 	{
 		USkeletalMeshComponent* SkeletalMeshComponent = InAnimInstance->GetSkelMeshComponent();
 		check(SkeletalMeshComponent);
@@ -117,6 +123,8 @@ void FAnimNode_MotionMatching::PreUpdate(const UAnimInstance* InAnimInstance)
 
 		UE::PoseSearch::Draw(DrawParams);
 	}
+
+	bWasEvaluated = false;
 #endif
 }
 
diff --git a/Engine/Plugins/Experimental/Animation/PoseSearch/Source/Runtime/Private/PoseSearch.cpp b/Engine/Plugins/Experimental/Animation/PoseSearch/Source/Runtime/Private/PoseSearch.cpp
--- a/Engine/Plugins/Experimental/Animation/PoseSearch/Source/Runtime/Private/PoseSearch.cpp
+++ b/Engine/Plugins/Experimental/Animation/PoseSearch/Source/Runtime/Private/PoseSearch.cpp
@@ -584,7 +584,7 @@ void UPoseSearchSchema::ResolveBoneReferences()
 // FPoseSearchChannelWeightParams
 FPoseSearchWeightParams::FPoseSearchWeightParams()
 {
-	TrajectoryWeight.ChannelWeight = 2.5f;
+	TrajectoryWeight.ChannelWeight = 3.0f;
 }
 
 
diff --git a/Engine/Plugins/Experimental/Animation/PoseSearch/Source/Runtime/Public/PoseSearch/AnimNode_MotionMatching.h b/Engine/Plugins/Experimental/Animation/PoseSearch/Source/Runtime/Public/PoseSearch/AnimNode_MotionMatching.h
--- a/Engine/Plugins/Experimental/Animation/PoseSearch/Source/Runtime/Public/PoseSearch/AnimNode_MotionMatching.h
+++ b/Engine/Plugins/Experimental/Animation/PoseSearch/Source/Runtime/Public/PoseSearch/AnimNode_MotionMatching.h
@@ -94,5 +94,8 @@ private:
 	// If true, "Relevant anim" nodes that look for the highest weighted animation in a state will ignore this node
 	UPROPERTY(EditAnywhere, Category=Relevancy, meta=(FoldProperty, PinHiddenByDefault))
 	bool bIgnoreForRelevancyTest = false;
+
+	// Whether this node was evaluated last frame
+	bool bWasEvaluated = false;
 #endif // WITH_EDITORONLY_DATA
 };
\ No newline at end of file
diff --git a/Engine/Plugins/Experimental/Animation/PoseSearch/Source/Runtime/Public/PoseSearch/PoseSearch.h b/Engine/Plugins/Experimental/Animation/PoseSearch/Source/Runtime/Public/PoseSearch/PoseSearch.h
--- a/Engine/Plugins/Experimental/Animation/PoseSearch/Source/Runtime/Public/PoseSearch/PoseSearch.h
+++ b/Engine/Plugins/Experimental/Animation/PoseSearch/Source/Runtime/Public/PoseSearch/PoseSearch.h
@@ -165,7 +165,7 @@ struct POSESEARCH_API FPoseSearchBone
 /**
 * Specifies the format of a pose search index. At runtime, queries are built according to the schema for searching.
 */
-UCLASS(BlueprintType, Category = "Animation|Pose Search")
+UCLASS(BlueprintType, Category = "Animation|Pose Search", Experimental)
 class POSESEARCH_API UPoseSearchSchema : public UDataAsset, public IBoneReferenceSkeletonProvider
 {
 	GENERATED_BODY()
@@ -442,7 +442,7 @@ struct FPoseSearchBlockTransitionParameters
 };
 
 /** Animation metadata object for indexing a single animation. */
-UCLASS(BlueprintType, Category = "Animation|Pose Search")
+UCLASS(BlueprintType, Category = "Animation|Pose Search", Experimental)
 class POSESEARCH_API UPoseSearchSequenceMetaData : public UAnimMetaData
 {
 	GENERATED_BODY()
@@ -616,9 +616,9 @@ private:
 UENUM()
 enum class EPoseSearchMirrorOption : int32
 {
-	UnmirroredOnly,
-	MirroredOnly,
-	UnmirroredAndMirrored,
+	UnmirroredOnly UMETA(DisplayName="Original Only"),
+	MirroredOnly UMETA(DisplayName="Mirrored Only"),
+	UnmirroredAndMirrored UMETA(DisplayName="Original and Mirrored"),
 	
 	Num UMETA(Hidden),
 	Invalid = Num UMETA(Hidden)
@@ -685,7 +685,7 @@ public:
 };
 
 /** A data asset for indexing a collection of animation sequences. */
-UCLASS(BlueprintType, Category = "Animation|Pose Search")
+UCLASS(BlueprintType, Category = "Animation|Pose Search", Experimental)
 class POSESEARCH_API UPoseSearchDatabase : public UDataAsset
 {
 	GENERATED_BODY()
