PoseSearch: Fixed crashes when a new sequence is added to a db or an MM node without a db is linked into a graph

diff --git a/Engine/Plugins/Experimental/Animation/PoseSearch/Source/Editor/Private/PoseSearchCustomization.cpp b/Engine/Plugins/Experimental/Animation/PoseSearch/Source/Editor/Private/PoseSearchCustomization.cpp
--- a/Engine/Plugins/Experimental/Animation/PoseSearch/Source/Editor/Private/PoseSearchCustomization.cpp
+++ b/Engine/Plugins/Experimental/Animation/PoseSearch/Source/Editor/Private/PoseSearchCustomization.cpp
@@ -88,7 +88,7 @@ void FPoseSearchDatabaseGroupCustomization::CustomizeChildren(TSharedRef<IProper
 
 void FPoseSearchDatabaseSequenceCustomization::CustomizeHeader(TSharedRef<IPropertyHandle> InStructPropertyHandle, FDetailWidgetRow& HeaderRow, IPropertyTypeCustomizationUtils& StructCustomizationUtils)
 {
-	FString SequenceNameString;
+	FText SequenceNameText;
 	FString GroupsString; // = TEXT("Groups: ");
 
 	TArray<UObject*> Objects;
@@ -97,8 +97,16 @@ void FPoseSearchDatabaseSequenceCustomization::CustomizeHeader(TSharedRef<IPrope
 	if (Objects.Num() == 1)
 	{
 		FPoseSearchDatabaseSequence* PoseSearchDatabaseSequence = (FPoseSearchDatabaseSequence*)InStructPropertyHandle->GetValueBaseAddress((uint8*)Objects[0]);
+		check(PoseSearchDatabaseSequence);
 		
-		SequenceNameString = PoseSearchDatabaseSequence->Sequence->GetName();
+		if (PoseSearchDatabaseSequence->Sequence)
+		{
+			SequenceNameText = FText::FromName(PoseSearchDatabaseSequence->Sequence->GetFName());
+		}
+		else
+		{
+			SequenceNameText = LOCTEXT("NewSequenceLabel", "New Sequence");
+		}
 		
 		const int32 NumGroups = PoseSearchDatabaseSequence->GroupTags.Num();
 		if (NumGroups == 0)
@@ -122,7 +130,7 @@ void FPoseSearchDatabaseSequenceCustomization::CustomizeHeader(TSharedRef<IPrope
 	.NameContent()
 	[
 		SNew(STextBlock)
-		.Text(FText::FromString(SequenceNameString))
+		.Text(SequenceNameText)
 	]
 	.ValueContent()
 	[
diff --git a/Engine/Plugins/Experimental/Animation/PoseSearch/Source/Runtime/Private/PoseSearchLibrary.cpp b/Engine/Plugins/Experimental/Animation/PoseSearch/Source/Runtime/Private/PoseSearchLibrary.cpp
--- a/Engine/Plugins/Experimental/Animation/PoseSearch/Source/Runtime/Private/PoseSearchLibrary.cpp
+++ b/Engine/Plugins/Experimental/Animation/PoseSearch/Source/Runtime/Private/PoseSearchLibrary.cpp
@@ -255,6 +255,12 @@ void UpdateMotionMatchingState(
 {
 	using namespace UE::PoseSearch;
 
+	if (!Database)
+	{
+		Context.LogMessage(EMessageSeverity::Error, LOCTEXT("NoDatabase", "No database provided for motion matching."));
+		return;
+	}
+
 	InOutMotionMatchingState.Flags = EMotionMatchingFlags::None;
 	if (InOutMotionMatchingState.CurrentDatabase != Database)
 	{
