PoseSearch: Fixed uninitialized SearchIndexAsset members in UPoseSearchSequenceMetaData
#8902

diff --git a/Engine/Plugins/Experimental/Animation/PoseSearch/Source/Runtime/Private/PoseSearch.cpp b/Engine/Plugins/Experimental/Animation/PoseSearch/Source/Runtime/Private/PoseSearch.cpp
--- a/Engine/Plugins/Experimental/Animation/PoseSearch/Source/Runtime/Private/PoseSearch.cpp
+++ b/Engine/Plugins/Experimental/Animation/PoseSearch/Source/Runtime/Private/PoseSearch.cpp
@@ -100,6 +100,17 @@ FLinearColor GetColorForFeature(FPoseSearchFeatureDesc Feature, const FPoseSearc
 	return ColorHSV.HSVToLinearRGB();
 }
 
+static FFloatInterval GetEffectiveSamplingRange(const UAnimSequenceBase* Sequence, FFloatInterval RequestedSamplingRange)
+{
+	const bool bSampleAll = (RequestedSamplingRange.Min == 0.0f) && (RequestedSamplingRange.Max == 0.0f);
+	const float SequencePlayLength = Sequence->GetPlayLength();
+	FFloatInterval Range;
+	Range.Min = bSampleAll ? 0.0f : RequestedSamplingRange.Min;
+	Range.Max = bSampleAll ? SequencePlayLength : FMath::Min(SequencePlayLength, RequestedSamplingRange.Max);
+	return Range;
+}
+
+
 /**
 * Algo::LowerBound adapted to TIndexedContainerIterator for use with indexable but not necessarily contiguous containers. Used here with TRingBuffer.
 *
@@ -1100,16 +1111,12 @@ bool UPoseSearchSequenceMetaData::IsValidForSearch() const
 	return IsValidForIndexing() && SearchIndex.IsValid();
 }
 
+
 //////////////////////////////////////////////////////////////////////////
 // FPoseSearchDatabaseSequence
 FFloatInterval FPoseSearchDatabaseSequence::GetEffectiveSamplingRange() const
 {
-	const bool bSampleAll = (SamplingRange.Min == 0.0f) && (SamplingRange.Max == 0.0f);
-	const float SequencePlayLength = Sequence->GetPlayLength();
-	FFloatInterval Range;
-	Range.Min = bSampleAll ? 0.0f : SamplingRange.Min;
-	Range.Max = bSampleAll ? SequencePlayLength : FMath::Min(SequencePlayLength, SamplingRange.Max);
-	return Range;
+	return UE::PoseSearch::GetEffectiveSamplingRange(Sequence, SamplingRange);
 }
 
 //////////////////////////////////////////////////////////////////////////
@@ -3884,20 +3891,22 @@ bool BuildIndex(const UAnimSequence* Sequence, UPoseSearchSequenceMetaData* Sequ
 	IndexerInput.SamplingContext = &SamplingContext;
 	IndexerInput.MainSequence = &Sampler;
 	IndexerInput.Schema = SequenceMetaData->Schema;
-	IndexerInput.RequestedSamplingRange = SequenceMetaData->SamplingRange;
+	IndexerInput.RequestedSamplingRange = GetEffectiveSamplingRange(Sequence, SequenceMetaData->SamplingRange);
 	Indexer.Init(IndexerInput);
 	Indexer.Process();
 
-	SequenceMetaData->SearchIndex.Values = Indexer.Output.FeatureVectorTable;
-	SequenceMetaData->SearchIndex.NumPoses = Indexer.Output.NumIndexedPoses;
-	SequenceMetaData->SearchIndex.Schema = SequenceMetaData->Schema;
-
 	SequenceMetaData->SearchIndex.Assets.Empty();
 	FPoseSearchIndexAsset SearchIndexAsset;
 	SearchIndexAsset.SourceAssetIdx = 0;
 	SearchIndexAsset.FirstPoseIdx = 0;
 	SearchIndexAsset.NumPoses = Indexer.Output.NumIndexedPoses;
+	SearchIndexAsset.SamplingInterval = IndexerInput.RequestedSamplingRange;
+
+	SequenceMetaData->SearchIndex.Values = Indexer.Output.FeatureVectorTable;
+	SequenceMetaData->SearchIndex.NumPoses = Indexer.Output.NumIndexedPoses;
+	SequenceMetaData->SearchIndex.Schema = SequenceMetaData->Schema;
 	SequenceMetaData->SearchIndex.Assets.Add(SearchIndexAsset);
+	SequenceMetaData->SearchIndex.PoseMetadata = Indexer.Output.PoseMetadata;
 
 	PreprocessSearchIndex(&SequenceMetaData->SearchIndex);
 
