PoseSearch: Fixed cost details check by including the pose's source group and relaxing score difference tolerance. Misc cleanup in FPoseSearchWeights::Init.

diff --git a/Engine/Plugins/Experimental/Animation/PoseSearch/Source/Runtime/Private/PoseSearch.cpp b/Engine/Plugins/Experimental/Animation/PoseSearch/Source/Runtime/Private/PoseSearch.cpp
--- a/Engine/Plugins/Experimental/Animation/PoseSearch/Source/Runtime/Private/PoseSearch.cpp
+++ b/Engine/Plugins/Experimental/Animation/PoseSearch/Source/Runtime/Private/PoseSearch.cpp
@@ -641,10 +641,6 @@ void FPoseSearchWeights::Init(const FPoseSearchWeightParams& WeightParams, const
 		return;
 	}
 
-	
-	FMemMark MemMark(FMemStack::Get());
-
-
 	// Setup channel indexable weight params
 	constexpr int ChannelNum = MaxChannels;
 	const FPoseSearchChannelWeightParams* ChannelWeightParams[ChannelNum];
@@ -659,7 +655,7 @@ void FPoseSearchWeights::Init(const FPoseSearchWeightParams& WeightParams, const
 	ChannelDynamicWeightParams[ChannelIdxTrajectoryDistance] = &DynamicWeightParams.TrajectoryDynamicWeights;
 
 	// Normalize channel weights
-	Array<float, ChannelNum, 1> NormalizedChannelWeights;
+	Eigen::Array<float, ChannelNum, 1> NormalizedChannelWeights;
 	for (int ChannelIdx = 0; ChannelIdx != ChannelNum; ++ChannelIdx)
 	{
 		NormalizedChannelWeights[ChannelIdx] = ChannelWeightParams[ChannelIdx]->ChannelWeight * ChannelDynamicWeightParams[ChannelIdx]->ChannelWeightScale;
@@ -688,17 +684,17 @@ void FPoseSearchWeights::Init(const FPoseSearchWeightParams& WeightParams, const
 	}
 
 	// WeightsByFeature is indexed by FeatureIdx in a Layout
-	TArray<float, TMemStackAllocator<>> WeightsByFeatureStorage;
+	TArray<float, TInlineAllocator<32>> WeightsByFeatureStorage;
 	WeightsByFeatureStorage.SetNum(Schema->Layout.Features.Num());
-	Map<ArrayXf> WeightsByFeature(WeightsByFeatureStorage.GetData(), WeightsByFeatureStorage.Num());
+	Eigen::Map<ArrayXf> WeightsByFeature(WeightsByFeatureStorage.GetData(), WeightsByFeatureStorage.Num());
 
 	// HorizonWeightsBySample is indexed by the channel's sample offsets in a Schema
-	TArray<float, TMemStackAllocator<>> HorizonWeightsBySampleStorage;
+	TArray<float, TInlineAllocator<16>> HorizonWeightsBySampleStorage;
 	HorizonWeightsBySampleStorage.SetNum(MaxChannelSampleOffsets);
-	Map<ArrayXf> HorizonWeightsBySample(HorizonWeightsBySampleStorage.GetData(), HorizonWeightsBySampleStorage.Num());
+	Eigen::Map<ArrayXf> HorizonWeightsBySample(HorizonWeightsBySampleStorage.GetData(), HorizonWeightsBySampleStorage.Num());
 
 	// WeightsByType is indexed by feature type
-	Array<float, (int)EPoseSearchFeatureType::Num, 1> WeightsByType;
+	Eigen::Array<float, (int)EPoseSearchFeatureType::Num, 1> WeightsByType;
 
 
 	// Determine each channel's feature weights
@@ -761,7 +757,7 @@ void FPoseSearchWeights::Init(const FPoseSearchWeightParams& WeightParams, const
 
 
 		// Initialize horizon weights
-		Array<float, 1, EHorizon::Num> NormalizedHorizonWeights;
+		Eigen::Array<float, 1, EHorizon::Num> NormalizedHorizonWeights;
 		NormalizedHorizonWeights.setConstant(0.0f);
 
 		if (!HorizonSampleIdxRanges[EHorizon::History].IsEmpty())
@@ -831,7 +827,7 @@ void FPoseSearchWeights::Init(const FPoseSearchWeightParams& WeightParams, const
 
 
 		// Now set this channel's weights for every feature in each horizon
-		Array<float, 1, EHorizon::Num> HorizonSums;
+		Eigen::Array<float, 1, EHorizon::Num> HorizonSums;
 		HorizonSums = 0.0f;
 		for (int FeatureIdx = INDEX_NONE; Schema->Layout.EnumerateBy(ChannelIdx, EPoseSearchFeatureType::Invalid, FeatureIdx); /*empty*/)
 		{
@@ -2901,7 +2897,7 @@ FSequenceIndexer::FSampleInfo FSequenceIndexer::GetSampleInfo(float SampleTime)
 		while (CyclesRemaining--)
 		{
 			Sample.RootTransform = RootMotionPerCycle * Sample.RootTransform;
-         			Sample.RootDistance += RootDistancePerCycle;
+         	Sample.RootDistance += RootDistancePerCycle;
 		}
 
 		Sample.RootTransform = RootMotionRemainder * Sample.RootTransform;
@@ -3619,6 +3615,7 @@ static void PreprocessSearchIndexNormalize(FPoseSearchIndex* SearchIndex)
 	SampleMeanMap = SampleMean.cast<float>();
 
 #if UE_POSE_SEARCH_EIGEN_DEBUG
+	FString MeanDevationsStr = EigenMatrixToString(MeanDeviations);
 	FString PoseMtxOriginalStr = EigenMatrixToString(PoseMatrixOriginal);
 	FString PoseMtxStr = EigenMatrixToString(PoseMatrix);
 	FString TransformationStr = EigenMatrixToString(TransformMap);
@@ -4194,9 +4191,9 @@ FPoseCost ComparePoses(int32 PoseIdx, FSearchContext& SearchContext, FPoseCostDe
 	auto QueryVector = Map<const ArrayXf>(SearchContext.QueryValues.GetData(), Dims);
 	
 	// Compute weighted squared difference vector
+	const FPoseSearchIndexAsset* SearchIndexAsset = SearchContext.GetSearchIndex()->FindAssetForPose(PoseIdx);
 	if (SearchContext.WeightsContext)
 	{
-		const FPoseSearchIndexAsset* SearchIndexAsset = SearchContext.GetSearchIndex()->FindAssetForPose(PoseIdx);
 		const FPoseSearchWeights* WeightsSet = 
 			SearchContext.WeightsContext->GetGroupWeights(SearchIndexAsset->SourceGroupIdx);
 		check(WeightsSet);
@@ -4212,29 +4209,31 @@ FPoseCost ComparePoses(int32 PoseIdx, FSearchContext& SearchContext, FPoseCostDe
 		Result.Dissimilarity = OutCostVector.sum();
 	}
 
+	// Output result
 	float NotifyAddend = 0.0f;
 	float MirrorMismatchAddend = 0.0f;
 	ComputePoseCostAddends(PoseIdx, SearchContext, NotifyAddend, MirrorMismatchAddend);
 	Result.CostAddend = NotifyAddend + MirrorMismatchAddend;
 	Result.TotalCost = Result.Dissimilarity + Result.CostAddend;
 
+	// Output cost details
 	OutPoseCostDetails.NotifyCostAddend = NotifyAddend;
 	OutPoseCostDetails.MirrorMismatchCostAddend = MirrorMismatchAddend;
-
-	// Verify this math agrees with the runtime pose comparator
-	checkSlow(FMath::IsNearlyEqual(
-		Result.TotalCost,
-		ComparePoses(PoseIdx, SearchContext).TotalCost, 
-		KINDA_SMALL_NUMBER));
-
-	
-	// Output cost details
 	OutPoseCostDetails.PoseCost = Result;
 	CalcChannelCosts(SearchContext.GetSearchIndex()->Schema, OutPoseCostDetails.CostVector, OutPoseCostDetails.ChannelCosts);
 
-	// Verify channel cost decomposition agrees with pose comparator
-	auto OutChannelCosts = Map<const ArrayXf>(OutPoseCostDetails.ChannelCosts.GetData(), OutPoseCostDetails.ChannelCosts.Num());
-	checkSlow(FMath::IsNearlyEqual(Result.Dissimilarity, OutChannelCosts.sum(), KINDA_SMALL_NUMBER));
+
+#if DO_GUARD_SLOW
+	{
+		// Verify details pose comparator agrees with runtime pose comparator
+		FPoseCost RuntimeComparatorCost = ComparePoses(PoseIdx, SearchContext, SearchIndexAsset->SourceGroupIdx);
+		checkSlow(FMath::IsNearlyEqual(Result.TotalCost, RuntimeComparatorCost.TotalCost, 1e-3f));
+
+		// Verify channel cost decomposition agrees with runtime pose comparator
+		auto OutChannelCosts = Map<const ArrayXf>(OutPoseCostDetails.ChannelCosts.GetData(), OutPoseCostDetails.ChannelCosts.Num());
+		checkSlow(FMath::IsNearlyEqual(OutChannelCosts.sum(), RuntimeComparatorCost.Dissimilarity, 1e-3f));
+	}
+#endif
 
 	return Result;
 }
