Fixed properties of bones in the feature layout being incorrectly specified when bone indices in schema were not already in sorted order.

diff --git a/Engine/Plugins/Experimental/Animation/PoseSearch/Source/Runtime/Private/PoseSearch.cpp b/Engine/Plugins/Experimental/Animation/PoseSearch/Source/Runtime/Private/PoseSearch.cpp
--- a/Engine/Plugins/Experimental/Animation/PoseSearch/Source/Runtime/Private/PoseSearch.cpp
+++ b/Engine/Plugins/Experimental/Animation/PoseSearch/Source/Runtime/Private/PoseSearch.cpp
@@ -568,16 +568,17 @@ void UPoseSearchSchema::ResolveBoneReferences()
 		Bone.Reference.Initialize(Skeleton);
 	}
 
-	// Fill out bone index array and sort by bone index
+	// Fill out bone index array
 	BoneIndices.SetNum(SampledBones.Num());
 	for (int32 Index = 0; Index != SampledBones.Num(); ++Index)
 	{
 		BoneIndices[Index] = SampledBones[Index].Reference.BoneIndex;
 	}
-	BoneIndices.Sort();
 
-	// Build separate index array with parent indices guaranteed to be present
+	// Build separate index array with parent indices guaranteed to be present. Sort for EnsureParentsPresent.
 	BoneIndicesWithParents = BoneIndices;
+	BoneIndicesWithParents.Sort();
+
 	if (Skeleton)
 	{
 		FAnimationRuntime::EnsureParentsPresent(BoneIndicesWithParents, Skeleton->GetReferenceSkeleton());
