PoseSearch:
- Changed MinPercentImprovement math to be a true percent change calculation to more accurately reflect its name and also reduce spurious pose changes. This reduced pose changes by about 60% in the MM test gym.
- Fixed current pose cost calculation not factoring in its source group weights
- Fixed current pose comparison not accounting for cost modifiers
- Also made it so searching is only done when needed instead of every frame.
- Changed default SequenceEndInterval to 0. This was causing index construction for pose matching to trim the end of sequences.
#8918

diff --git a/Engine/Plugins/Experimental/Animation/PoseSearch/Source/Runtime/Private/PoseSearch.cpp b/Engine/Plugins/Experimental/Animation/PoseSearch/Source/Runtime/Private/PoseSearch.cpp
--- a/Engine/Plugins/Experimental/Animation/PoseSearch/Source/Runtime/Private/PoseSearch.cpp
+++ b/Engine/Plugins/Experimental/Animation/PoseSearch/Source/Runtime/Private/PoseSearch.cpp
@@ -2395,7 +2395,7 @@ FTransform FSequenceSampler::ExtrapolateRootMotion(float SampleStart, float Samp
 		ExtrapolatedRootMotion = AbsTimeSampleToExtrapolate * ExtrapolatedRootMotion;
 	}
 	
-	// and a blend with identify for whatever is left
+	// and a blend with identity for whatever is left
 	FTransform RemainingExtrapolatedRootMotion;
 	RemainingExtrapolatedRootMotion.Blend(
 		FTransform::Identity, 
@@ -4136,11 +4136,17 @@ void ComputePoseCostAddends(
 }
 
 
-FPoseCost ComparePoses(int32 PoseIdx, FSearchContext& SearchContext, int32 GroupIdx /* = INDEX_NONE */)
+FPoseCost ComparePoses(int32 PoseIdx, FSearchContext& SearchContext, int32 GroupIdx)
 {
 	FPoseCost Result;
 
-	TArrayView<const float> PoseValues = SearchContext.GetSearchIndex()->GetPoseValues(PoseIdx);
+	const FPoseSearchIndex* SearchIndex = SearchContext.GetSearchIndex();
+	if (!ensure(SearchIndex))
+	{
+		return Result;
+	}
+
+	TArrayView<const float> PoseValues = SearchIndex->GetPoseValues(PoseIdx);
 	if (!ensure(PoseValues.Num() == SearchContext.QueryValues.Num()))
 	{
 		return Result;
@@ -4148,6 +4154,17 @@ FPoseCost ComparePoses(int32 PoseIdx, FSearchContext& SearchContext, int32 Group
 
 	if (SearchContext.WeightsContext)
 	{
+		if (GroupIdx == INDEX_NONE)
+		{
+			const FPoseSearchIndexAsset* SearchIndexAsset = SearchIndex->FindAssetForPose(PoseIdx);
+			if (!ensure(SearchIndexAsset))
+			{
+				return Result;
+			}
+
+			GroupIdx = SearchIndexAsset->SourceGroupIdx;
+		}
+
 		const FPoseSearchWeights* WeightsSet =
 			SearchContext.WeightsContext->GetGroupWeights(GroupIdx);
 		Result.Dissimilarity = CompareFeatureVectors(
diff --git a/Engine/Plugins/Experimental/Animation/PoseSearch/Source/Runtime/Private/PoseSearchLibrary.cpp b/Engine/Plugins/Experimental/Animation/PoseSearch/Source/Runtime/Private/PoseSearchLibrary.cpp
--- a/Engine/Plugins/Experimental/Animation/PoseSearch/Source/Runtime/Private/PoseSearchLibrary.cpp
+++ b/Engine/Plugins/Experimental/Animation/PoseSearch/Source/Runtime/Private/PoseSearchLibrary.cpp
@@ -322,53 +322,65 @@ void UpdateMotionMatchingState(
 	// Update weight groups
 	InOutMotionMatchingState.WeightsContext.Update(Settings.Weights, Database);
 
-	// Determine how much the updated query vector deviates from the current pose vector
-	FPoseCost CurrentPoseCost;
-	if (InOutMotionMatchingState.DbPoseIdx != INDEX_NONE)
+	// Jump to the candidate pose if the stepper couldn't continue to the next frame
+	if (!PoseStepper.CanContinue())
 	{
-		CurrentPoseCost = ComparePoses(InOutMotionMatchingState.DbPoseIdx, SearchContext);
+		FSearchResult Result = Search(SearchContext);
+		InOutMotionMatchingState.JumpToPose(Context, Settings, Result);
 	}
 
-	// Search the database for the nearest match to the updated query vector
-	FSearchResult Result = Search(SearchContext);
-
-	if (Result.IsValid())
+	// Do a search when enough time has elapsed since the last pose jump
+	else if ((InOutMotionMatchingState.ElapsedPoseJumpTime >= Settings.SearchThrottleTime))
 	{
-		// Jump to the candidate pose if the stepper couldn't continue to the next frame
-		if (!PoseStepper.CanContinue())
+		// Determine how much the updated query vector deviates from the current pose vector
+		FPoseCost CurrentPoseCost;
+		if (InOutMotionMatchingState.DbPoseIdx != INDEX_NONE)
 		{
-			InOutMotionMatchingState.JumpToPose(Context, Settings, Result);
+			const FPoseSearchIndexAsset* SearchIndexAsset = &Database->SearchIndex.Assets[InOutMotionMatchingState.SearchIndexAssetIdx];
+			CurrentPoseCost = ComparePoses(InOutMotionMatchingState.DbPoseIdx, SearchContext, SearchIndexAsset->SourceGroupIdx);
 		}
 
-		// Consider the candidate pose when enough time has elapsed since the last pose jump
-		else if ((InOutMotionMatchingState.ElapsedPoseJumpTime >= Settings.SearchThrottleTime))
+		// Search the database for the nearest match to the updated query vector
+		FSearchResult Result = Search(SearchContext);
+
+		// Consider the search result better if it is more similar to the query than the current pose we're playing back from the database
+		check(Result.PoseCost.Dissimilarity >= 0.0f);
+		bool bBetterPose = true;
+		if (CurrentPoseCost.IsValid())
 		{
-			// Consider the search result better if it is more similar to the query than the current pose we're playing back from the database
-			check(Result.PoseCost.Dissimilarity >= 0.0f && CurrentPoseCost.Dissimilarity >= 0.0f);
-			bool bBetterPose = Result.PoseCost.Dissimilarity * (1.0f + (Settings.MinPercentImprovement / 100.0f)) < CurrentPoseCost.Dissimilarity;
-
-			// Ignore the candidate poses from the same anim when they are too near to the current pose
-			bool bNearbyPose = false;
-			const FPoseSearchIndexAsset* StateSearchIndexAsset = InOutMotionMatchingState.GetCurrentSearchIndexAsset();
-			if (StateSearchIndexAsset == Result.SearchIndexAsset)
+			if ((CurrentPoseCost.TotalCost <= Result.PoseCost.TotalCost) || (CurrentPoseCost.Dissimilarity <= Result.PoseCost.Dissimilarity))
 			{
-				const FPoseSearchDatabaseSequence& ResultDbSequence = Database->GetSourceAsset(Result.SearchIndexAsset);
-				bNearbyPose = FMath::Abs(InOutMotionMatchingState.AssetPlayerTime - Result.TimeOffsetSeconds) < Settings.PoseJumpThresholdTime;
-
-				// Handle looping anims when checking for the pose being too close
-				if (!bNearbyPose && ResultDbSequence.bLoopAnimation)
-				{
-					const float AssetLength = ResultDbSequence.Sequence->GetPlayLength();
-					bNearbyPose = FMath::Abs(AssetLength - InOutMotionMatchingState.AssetPlayerTime - Result.TimeOffsetSeconds) < Settings.PoseJumpThresholdTime;
-				}
+				bBetterPose = false;
 			}
+			else
+			{
+				checkSlow(CurrentPoseCost.Dissimilarity > 0.0f && CurrentPoseCost.Dissimilarity > Result.PoseCost.Dissimilarity);
+				const float RelativeSimilarityGain = -1.0f * (Result.PoseCost.Dissimilarity - CurrentPoseCost.Dissimilarity) / CurrentPoseCost.Dissimilarity;
+				bBetterPose = RelativeSimilarityGain >= Settings.MinPercentImprovement / 100.0f;
+			}
+		}
 
-			// Start playback from the candidate pose if we determined it was a better option
-			if (bBetterPose && !bNearbyPose)
+		// Ignore the candidate poses from the same anim when they are too near to the current pose
+		bool bNearbyPose = false;
+		const FPoseSearchIndexAsset* StateSearchIndexAsset = InOutMotionMatchingState.GetCurrentSearchIndexAsset();
+		if (StateSearchIndexAsset == Result.SearchIndexAsset)
+		{
+			const FPoseSearchDatabaseSequence& ResultDbSequence = Database->GetSourceAsset(Result.SearchIndexAsset);
+			bNearbyPose = FMath::Abs(InOutMotionMatchingState.AssetPlayerTime - Result.TimeOffsetSeconds) < Settings.PoseJumpThresholdTime;
+
+			// Handle looping anims when checking for the pose being too close
+			if (!bNearbyPose && ResultDbSequence.bLoopAnimation)
 			{
-				InOutMotionMatchingState.JumpToPose(Context, Settings, Result);
+				const float AssetLength = ResultDbSequence.Sequence->GetPlayLength();
+				bNearbyPose = FMath::Abs(AssetLength - InOutMotionMatchingState.AssetPlayerTime - Result.TimeOffsetSeconds) < Settings.PoseJumpThresholdTime;
 			}
 		}
+
+		// Start playback from the candidate pose if we determined it was a better option
+		if (bBetterPose && !bNearbyPose)
+		{
+			InOutMotionMatchingState.JumpToPose(Context, Settings, Result);
+		}
 	}
 
 	// Jump to the pose stepper's next anim if we didn't jump to a pose as a result of the search above
diff --git a/Engine/Plugins/Experimental/Animation/PoseSearch/Source/Runtime/Public/PoseSearch/PoseSearch.h b/Engine/Plugins/Experimental/Animation/PoseSearch/Source/Runtime/Public/PoseSearch/PoseSearch.h
--- a/Engine/Plugins/Experimental/Animation/PoseSearch/Source/Runtime/Public/PoseSearch/PoseSearch.h
+++ b/Engine/Plugins/Experimental/Animation/PoseSearch/Source/Runtime/Public/PoseSearch/PoseSearch.h
@@ -437,8 +437,7 @@ struct FPoseSearchBlockTransitionParameters
 	// Excluding the end of sequences help ensure an exact future trajectory, and also prevents the selection of
 	// a sequence which will end too soon to be worth selecting.
 	UPROPERTY(EditAnywhere, Category = "Settings")
-	float SequenceEndInterval = 0.2f;
-
+	float SequenceEndInterval = 0.0f;
 };
 
 /** Animation metadata object for indexing a single animation. */
@@ -706,7 +705,7 @@ public:
 	FPoseSearchExtrapolationParameters ExtrapolationParameters;
 
 	UPROPERTY(EditAnywhere, Category = "Database")
-	FPoseSearchBlockTransitionParameters BlockTransitionParameters;
+	FPoseSearchBlockTransitionParameters BlockTransitionParameters = {0.0f, 0.2f};
 
 	UPROPERTY(EditAnywhere, Category = "Database")
 	TArray<FPoseSearchDatabaseGroup> Groups;
@@ -944,6 +943,7 @@ struct FPoseCost
 	float CostAddend = 0.0f;
 	float TotalCost = MAX_flt;
 	bool operator<(const FPoseCost& Other) const { return TotalCost < Other.TotalCost; }
+	bool IsValid() const { return TotalCost != MAX_flt; }
 
 };
 
@@ -1022,12 +1022,12 @@ POSESEARCH_API FSearchResult Search(FSearchContext& SearchContext);
 *
 * @param PoseIdx			The index of the pose in the search index to compare to the query
 * @param SearchContext		Structure containing search parameters
-* @param GroupIdx			Indicated which weight to use when evaluating dissimilarity
+* @param GroupIdx			Indicates the group for this pose's source asset. Specify INDEX_NONE to lookup the pose's group.
 *
 * @return Dissimilarity between the two poses
 */
 
-POSESEARCH_API FPoseCost ComparePoses(int32 PoseIdx, FSearchContext& SearchContext, int32 GroupIdx = INDEX_NONE);
+POSESEARCH_API FPoseCost ComparePoses(int32 PoseIdx, FSearchContext& SearchContext, int32 GroupIdx);
 
 /**
  * Cost details for pose analysis in the rewind debugger
