Blend space support for Motion Matching

diff --git a/Engine/Plugins/Experimental/Animation/PoseSearch/Source/Editor/Private/PoseSearchDebugger.cpp b/Engine/Plugins/Experimental/Animation/PoseSearch/Source/Editor/Private/PoseSearchDebugger.cpp
--- a/Engine/Plugins/Experimental/Animation/PoseSearch/Source/Editor/Private/PoseSearchDebugger.cpp
+++ b/Engine/Plugins/Experimental/Animation/PoseSearch/Source/Editor/Private/PoseSearchDebugger.cpp
@@ -6,7 +6,9 @@
 #include "PoseSearch/PoseSearch.h"
 #include "Animation/AnimationPoseData.h"
 #include "Animation/AnimSequence.h"
+#include "Animation/BlendSpace.h"
 #include "Animation/MirrorDataTable.h"
+#include "Animation/AnimRootMotionProvider.h"
 #include "RewindDebuggerInterface/Public/IRewindDebugger.h"
 #include "ObjectTrace.h"
 #include "TraceServices/Model/AnalysisSession.h"
@@ -27,7 +29,6 @@
 
 #define LOCTEXT_NAMESPACE "PoseSearchDebugger"
 
-
 void FPoseSearchDebuggerPoseVectorChannel::Reset()
 {
 	Positions.Reset();
@@ -144,17 +145,74 @@ void UPoseSearchMeshComponent::UpdatePose(const FUpdateContext& UpdateContext)
 	UE::Anim::FStackAttributeContainer Attributes;
 	FAnimationPoseData PoseData(CompactPose, Curve, Attributes);
 
-	float AdvancedTime = UpdateContext.SequenceStartTime;
-	FAnimationRuntime::AdvanceTime(
-		UpdateContext.bLoop,
-		UpdateContext.SequenceTime - UpdateContext.SequenceStartTime,
-		AdvancedTime,
-		UpdateContext.Sequence->GetPlayLength());
+	if (UpdateContext.Type == ESearchIndexAssetType::Sequence)
+	{
+		float AdvancedTime = UpdateContext.StartTime;
+
+		FAnimationRuntime::AdvanceTime(
+			UpdateContext.bLoop,
+			UpdateContext.Time - UpdateContext.StartTime,
+			AdvancedTime,
+			UpdateContext.Sequence->GetPlayLength());
 
-	FAnimExtractContext ExtractionCtx;
-	ExtractionCtx.CurrentTime = AdvancedTime;
+		FAnimExtractContext ExtractionCtx;
+		ExtractionCtx.CurrentTime = AdvancedTime;
+
+		UpdateContext.Sequence->GetAnimationPose(PoseData, ExtractionCtx);
+	}
+	else if (UpdateContext.Type == ESearchIndexAssetType::BlendSpace)
+	{
+		TArray<FBlendSampleData> BlendSamples;
+		int32 TriangulationIndex = 0;
+		UpdateContext.BlendSpace->GetSamplesFromBlendInput(UpdateContext.BlendParameters, BlendSamples, TriangulationIndex, true);
+		
+		float PlayLength = UpdateContext.BlendSpace->GetAnimationLengthFromSampleData(BlendSamples);
+		
+		float PreviousTime = UpdateContext.StartTime * PlayLength;
+		float CurrentTime = UpdateContext.Time * PlayLength;
+
+		float AdvancedTime = PreviousTime;
+		FAnimationRuntime::AdvanceTime(
+			UpdateContext.bLoop,
+			CurrentTime - PreviousTime,
+			CurrentTime,
+			PlayLength);
+		
+		FDeltaTimeRecord DeltaTimeRecord;
+		DeltaTimeRecord.Set(PreviousTime, AdvancedTime - PreviousTime);
+		FAnimExtractContext ExtractionCtx(AdvancedTime, true, DeltaTimeRecord, UpdateContext.bLoop);
+
+		for (int32 BlendSampleIdex = 0; BlendSampleIdex < BlendSamples.Num(); BlendSampleIdex++)
+		{
+			float Scale = BlendSamples[BlendSampleIdex].Animation->GetPlayLength() / PlayLength;
+
+			FDeltaTimeRecord BlendSampleDeltaTimeRecord;
+			BlendSampleDeltaTimeRecord.Set(DeltaTimeRecord.GetPrevious() * Scale, DeltaTimeRecord.Delta * Scale);
+
+			BlendSamples[BlendSampleIdex].DeltaTimeRecord = BlendSampleDeltaTimeRecord;
+			BlendSamples[BlendSampleIdex].PreviousTime = PreviousTime * Scale;
+			BlendSamples[BlendSampleIdex].Time = AdvancedTime * Scale;
+		}
+
+		UpdateContext.BlendSpace->GetAnimationPose(BlendSamples, ExtractionCtx, PoseData);
+	}
+	else
+	{
+		checkNoEntry();
+	}
+
+	LastRootMotionDelta = FTransform::Identity;
+
+	const UE::Anim::IAnimRootMotionProvider* RootMotionProvider = UE::Anim::IAnimRootMotionProvider::Get();
+
+	if (ensureMsgf(RootMotionProvider, TEXT("Could not get Root Motion Provider.")))
+	{
+		if (ensureMsgf(RootMotionProvider->HasRootMotion(Attributes), TEXT("Blend Space had no Root Motion Attribute.")))
+		{
+			RootMotionProvider->ExtractRootMotion(Attributes, LastRootMotionDelta);
+		}
+	}
 
-	UpdateContext.Sequence->GetAnimationPose(PoseData, ExtractionCtx);
 	if (UpdateContext.bMirrored)
 	{
 		FAnimationRuntime::MirrorPose(
@@ -179,11 +237,6 @@ void UPoseSearchMeshComponent::UpdatePose(const FUpdateContext& UpdateContext)
 		SetBoneTransformByName(BoneName, BoneTransform, EBoneSpaces::ComponentSpace);
 	}
 
-	LastRootMotionDelta = UpdateContext.Sequence->ExtractRootMotion(
-		UpdateContext.SequenceStartTime, 
-		UpdateContext.SequenceTime - UpdateContext.SequenceStartTime,
-		UpdateContext.bLoop);
-
 	if (UpdateContext.bMirrored)
 	{
 		const EAxis::Type MirrorAxis = UpdateContext.MirrorDataTable->MirrorAxis;
@@ -215,13 +268,15 @@ public:
 	float GetTrajectoryCost() const { return PoseCostDetails.ChannelCosts.Num() >= 3 ? PoseCostDetails.ChannelCosts[1] + PoseCostDetails.ChannelCosts[2] : 0.0f; }
 	float GetAddendsCost() const { return PoseCostDetails.NotifyCostAddend + PoseCostDetails.MirrorMismatchCostAddend; }
 
+	ESearchIndexAssetType AssetType = ESearchIndexAssetType::Invalid;
 	int32 PoseIdx = 0;
-	FString AnimSequenceName = "";
-	FString AnimSequencePath = "";
-	int32 DbSequenceIdx = 0;
+	FString AssetName = "";
+	FString AssetPath = "";
+	int32 DbAssetIdx = 0;
 	int32 AnimFrame = 0;
-	float Time = 0.0f;
+	float AssetTime = 0.0f;
 	bool bMirrored = false;
+	FVector BlendParameters = FVector::Zero();
 	FPoseCostDetails PoseCostDetails;
 };
 
@@ -297,7 +352,7 @@ namespace DebuggerDatabaseColumns
 	};
 	const FName FPoseIdx::Name = "Pose Index";
 
-	struct FAnimSequenceName : IColumn
+	struct FAssetName : IColumn
 	{
 		using IColumn::IColumn;
 		static const FName Name;
@@ -305,28 +360,57 @@ namespace DebuggerDatabaseColumns
 		
 		virtual FSortPredicate GetSortPredicate() const override
 		{
-			static FSortPredicate Predicate = [](const FRowDataRef& Row0, const FRowDataRef& Row1) -> bool { return Row0->AnimSequenceName < Row1->AnimSequenceName; };
-			return Predicate;
+			return [](const FRowDataRef& Row0, const FRowDataRef& Row1) -> bool { return Row0->AssetName < Row1->AssetName; };
 		}
 
 		virtual TSharedRef<SWidget> GenerateWidget(const FRowDataRef& RowData) const override
 		{
 			return SNew(SHyperlink)
-				.Text_Lambda([RowData]() -> FText { return FText::FromString(RowData->AnimSequenceName); })
+				.Text_Lambda([RowData]() -> FText { return FText::FromString(RowData->AssetName); })
 				.TextStyle(&FCoreStyle::Get().GetWidgetStyle<FTextBlockStyle>("SmallText"))
 				.ToolTipText_Lambda([RowData]() -> FText 
 					{ 
 						return FText::Format(
 							LOCTEXT("AssetHyperlinkTooltipFormat", "Open asset '{0}'"), 
-							FText::FromString(RowData->AnimSequencePath)); 
+							FText::FromString(RowData->AssetPath)); 
 					})
 				.OnNavigate_Lambda([RowData]()
 					{
-						GEditor->GetEditorSubsystem<UAssetEditorSubsystem>()->OpenEditorForAsset(RowData->AnimSequencePath);
+						GEditor->GetEditorSubsystem<UAssetEditorSubsystem>()->OpenEditorForAsset(RowData->AssetPath);
 					});
 		}
 	};
-	const FName FAnimSequenceName::Name = "Sequence";
+	const FName FAssetName::Name = "Asset";
+
+	struct FAssetType : ITextColumn
+	{
+		using ITextColumn::ITextColumn;
+		static const FName Name;
+		virtual FName GetName() const override { return Name; }
+
+		virtual FSortPredicate GetSortPredicate() const override
+		{
+			return [](const FRowDataRef& Row0, const FRowDataRef& Row1) -> bool { return Row0->AssetType < Row1->AssetType; };
+		}
+
+		virtual FText GetRowText(const FRowDataRef& Row) const override
+		{
+			if (Row->AssetType == ESearchIndexAssetType::Sequence)
+			{
+				return FText::FromString(TEXT("Sequence"));
+			}
+			else if (Row->AssetType == ESearchIndexAssetType::BlendSpace)
+			{
+				return FText::FromString(TEXT("BlendSpace"));
+			}
+			else
+			{
+				checkNoEntry();
+				return FText::FromString(TEXT(""));
+			}
+		}
+	};
+	const FName FAssetType::Name = "Type";
 
 	struct FFrame : ITextColumn
 	{
@@ -336,8 +420,7 @@ namespace DebuggerDatabaseColumns
 		
 		virtual FSortPredicate GetSortPredicate() const override
 		{
-			static FSortPredicate Predicate = [](const FRowDataRef& Row0, const FRowDataRef& Row1) -> bool { return Row0->Time < Row1->Time; };
-			return Predicate;
+			return [](const FRowDataRef& Row0, const FRowDataRef& Row1) -> bool { return Row0->AssetTime < Row1->AssetTime; };
 		}
 
 		virtual FText GetRowText(const FRowDataRef& Row) const override
@@ -346,11 +429,25 @@ namespace DebuggerDatabaseColumns
 				.SetUseGrouping(false)
 				.SetMaximumFractionalDigits(2);
 
-			return FText::Format(
-				FText::FromString("{0} ({1})"),
-				FText::AsNumber(Row->AnimFrame, &FNumberFormattingOptions::DefaultNoGrouping()),
-				FText::AsNumber(Row->Time, &TimeFormattingOptions)
-			);
+			if (Row->AssetType == ESearchIndexAssetType::Sequence)
+			{
+				return FText::Format(
+					FText::FromString("{0} ({1})"),
+					FText::AsNumber(Row->AnimFrame, &FNumberFormattingOptions::DefaultNoGrouping()),
+					FText::AsNumber(Row->AssetTime, &TimeFormattingOptions));
+			}
+			else if (Row->AssetType == ESearchIndexAssetType::BlendSpace)
+			{
+				// There is no frame index associated with a blendspace
+				return FText::Format(
+					FText::FromString("({0})"),
+					FText::AsNumber(Row->AssetTime, &TimeFormattingOptions));
+			}
+			else
+			{
+				checkNoEntry();
+				return FText();
+			}
 		}
 	};
 	const FName FFrame::Name = "Frame";
@@ -419,10 +516,7 @@ namespace DebuggerDatabaseColumns
 
 		virtual FSortPredicate GetSortPredicate() const override
 		{
-			return [](const FRowDataRef& Row0, const FRowDataRef& Row1) -> bool 
-			{ 
-				return Row0->GetAddendsCost() < Row1->GetAddendsCost();
-			};
+			return [](const FRowDataRef& Row0, const FRowDataRef& Row1) -> bool { return Row0->GetAddendsCost() < Row1->GetAddendsCost(); };
 		}
 
 		virtual FText GetRowText(const FRowDataRef& Row) const override
@@ -440,10 +534,7 @@ namespace DebuggerDatabaseColumns
 
 		virtual FSortPredicate GetSortPredicate() const override
 		{
-			return [](const FRowDataRef& Row0, const FRowDataRef& Row1) -> bool 
-			{ 
-				return Row0->bMirrored < Row1->bMirrored; 
-			};
+			return [](const FRowDataRef& Row0, const FRowDataRef& Row1) -> bool { return Row0->bMirrored < Row1->bMirrored; };
 		}
 
 		virtual FText GetRowText(const FRowDataRef& Row) const override
@@ -452,6 +543,45 @@ namespace DebuggerDatabaseColumns
 		}
 	};
 	const FName FMirrored::Name = "Mirrored";
+
+	struct FBlendParameters : ITextColumn
+	{
+		using ITextColumn::ITextColumn;
+		static const FName Name;
+		virtual FName GetName() const override { return Name; }
+
+		virtual FSortPredicate GetSortPredicate() const override
+		{
+			return [](const FRowDataRef& Row0, const FRowDataRef& Row1) -> bool
+			{
+				return (
+					Row0->BlendParameters[0] < Row1->BlendParameters[0] ||
+					Row0->BlendParameters[1] < Row1->BlendParameters[1]);
+			};
+		}
+
+		virtual FText GetRowText(const FRowDataRef& Row) const override
+		{
+			if (Row->AssetType == ESearchIndexAssetType::Sequence)
+			{
+				return FText::FromString(TEXT("-"));
+			}
+			else if (Row->AssetType == ESearchIndexAssetType::BlendSpace)
+			{
+				return FText::Format(
+					LOCTEXT("Blend Parameters", "({0}, {1})"),
+					FText::AsNumber(Row->BlendParameters[0]),
+					FText::AsNumber(Row->BlendParameters[1]));
+			}
+			else
+			{
+				checkNoEntry();
+				return FText();
+			}
+		}
+	};
+	const FName FBlendParameters::Name = "Blend Parameters";
+
 }
 
 
@@ -656,18 +786,18 @@ void SDebuggerDatabaseView::OnDatabaseRowSelectionChanged(
 {
 	if (Row.IsValid())
 	{
-		OnPoseSelectionChanged.ExecuteIfBound(Row->PoseIdx, Row->Time);
+		OnPoseSelectionChanged.ExecuteIfBound(Row->PoseIdx, Row->AssetTime);
 	}
 }
 
-ECheckBoxState SDebuggerDatabaseView::IsSequenceFilterEnabled() const
+ECheckBoxState SDebuggerDatabaseView::IsAssetFilterEnabled() const
 {
-	return bSequenceFilterEnabled ? ECheckBoxState::Checked : ECheckBoxState::Unchecked;
+	return bAssetFilterEnabled ? ECheckBoxState::Checked : ECheckBoxState::Unchecked;
 }
 
-void SDebuggerDatabaseView::OnSequenceFilterEnabledChanged(ECheckBoxState NewState)
+void SDebuggerDatabaseView::OnAssetFilterEnabledChanged(ECheckBoxState NewState)
 {
-	bSequenceFilterEnabled = NewState == ECheckBoxState::Checked;
+	bAssetFilterEnabled = NewState == ECheckBoxState::Checked;
 	FilterDatabaseRows();
 }
 
@@ -700,7 +830,9 @@ void SDebuggerDatabaseView::FilterDatabaseRows()
 	{
 		for (const auto& UnfilteredRow : UnfilteredDatabaseRows)
 		{
-			if (!bSequenceFilterEnabled || DatabaseSequenceFilter[UnfilteredRow->DbSequenceIdx])
+			if (!bAssetFilterEnabled || 
+				(UnfilteredRow->AssetType == ESearchIndexAssetType::Sequence && DatabaseSequenceFilter[UnfilteredRow->DbAssetIdx]) ||
+				(UnfilteredRow->AssetType == ESearchIndexAssetType::BlendSpace && DatabaseBlendSpaceFilter[UnfilteredRow->DbAssetIdx]))
 			{
 				FilteredDatabaseView.Rows.Add(UnfilteredRow);
 			}
@@ -710,13 +842,15 @@ void SDebuggerDatabaseView::FilterDatabaseRows()
 	{
 		for (const auto& UnfilteredRow : UnfilteredDatabaseRows)
 		{
-			if (!bSequenceFilterEnabled || DatabaseSequenceFilter[UnfilteredRow->DbSequenceIdx])
+			if (!bAssetFilterEnabled ||
+				(UnfilteredRow->AssetType == ESearchIndexAssetType::Sequence && DatabaseSequenceFilter[UnfilteredRow->DbAssetIdx]) ||
+				(UnfilteredRow->AssetType == ESearchIndexAssetType::BlendSpace && DatabaseBlendSpaceFilter[UnfilteredRow->DbAssetIdx]))
 			{
 				bool bMatchesAllTokens = Algo::AllOf(
 					Tokens,
 					[&](FString Token)
 				{
-					return UnfilteredRow->AnimSequenceName.Contains(Token);
+					return UnfilteredRow->AssetName.Contains(Token);
 				});
 
 				if (bMatchesAllTokens)
@@ -740,21 +874,51 @@ void SDebuggerDatabaseView::CreateRows(const UPoseSearchDatabase& Database)
 	// Build database rows
 	for(const FPoseSearchIndexAsset& SearchIndexAsset : Database.SearchIndex.Assets)
 	{
-		const FPoseSearchDatabaseSequence& DbSequence = Database.GetSourceAsset(&SearchIndexAsset);
-		const int32 LastPoseIdx = SearchIndexAsset.FirstPoseIdx + SearchIndexAsset.NumPoses;
-		for (int32 PoseIdx = SearchIndexAsset.FirstPoseIdx; PoseIdx != LastPoseIdx; ++PoseIdx)
+		if (SearchIndexAsset.Type == ESearchIndexAssetType::Sequence)
 		{
-			float Time = Database.GetTimeOffset(PoseIdx);
-
-			TSharedRef<FDebuggerDatabaseRowData> Row = UnfilteredDatabaseRows.Add_GetRef(MakeShared<FDebuggerDatabaseRowData>());
-			Row->PoseIdx = PoseIdx;
-			Row->AnimSequenceName = DbSequence.Sequence->GetName();
-			Row->AnimSequencePath = DbSequence.Sequence->GetPathName();
-			Row->DbSequenceIdx = SearchIndexAsset.SourceAssetIdx;
-			Row->Time = Time;
-			Row->AnimFrame = DbSequence.Sequence->GetFrameAtTime(Time);
-			Row->bMirrored = SearchIndexAsset.bMirrored;
+			const FPoseSearchDatabaseSequence& DbSequence = Database.GetSequenceSourceAsset(&SearchIndexAsset);
+			const int32 LastPoseIdx = SearchIndexAsset.FirstPoseIdx + SearchIndexAsset.NumPoses;
+			for (int32 PoseIdx = SearchIndexAsset.FirstPoseIdx; PoseIdx != LastPoseIdx; ++PoseIdx)
+			{
+				float Time = Database.GetAssetTime(PoseIdx);
+
+				TSharedRef<FDebuggerDatabaseRowData> Row = UnfilteredDatabaseRows.Add_GetRef(MakeShared<FDebuggerDatabaseRowData>());
+				Row->PoseIdx = PoseIdx;
+				Row->AssetType = ESearchIndexAssetType::Sequence;
+				Row->AssetName = DbSequence.Sequence->GetName();
+				Row->AssetPath = DbSequence.Sequence->GetPathName();
+				Row->DbAssetIdx = SearchIndexAsset.SourceAssetIdx;
+				Row->AssetTime = Time;
+				Row->AnimFrame = DbSequence.Sequence->GetFrameAtTime(Time);
+				Row->bMirrored = SearchIndexAsset.bMirrored;
+				Row->BlendParameters = FVector::Zero();
+			}
+		}
+		else if (SearchIndexAsset.Type == ESearchIndexAssetType::BlendSpace)
+		{
+			const FPoseSearchDatabaseBlendSpace& DbBlendSpace = Database.GetBlendSpaceSourceAsset(&SearchIndexAsset);
+			const int32 LastPoseIdx = SearchIndexAsset.FirstPoseIdx + SearchIndexAsset.NumPoses;
+			for (int32 PoseIdx = SearchIndexAsset.FirstPoseIdx; PoseIdx != LastPoseIdx; ++PoseIdx)
+			{
+				float Time = Database.GetAssetTime(PoseIdx);
+
+				TSharedRef<FDebuggerDatabaseRowData> Row = UnfilteredDatabaseRows.Add_GetRef(MakeShared<FDebuggerDatabaseRowData>());
+				Row->PoseIdx = PoseIdx;
+				Row->AssetType = ESearchIndexAssetType::BlendSpace;
+				Row->AssetName = DbBlendSpace.BlendSpace->GetName();
+				Row->AssetPath = DbBlendSpace.BlendSpace->GetPathName();
+				Row->DbAssetIdx = SearchIndexAsset.SourceAssetIdx;
+				Row->AssetTime = Time;
+				Row->AnimFrame = 0; // There is no frame index associated with a blendspace
+				Row->bMirrored = SearchIndexAsset.bMirrored;
+				Row->BlendParameters = SearchIndexAsset.BlendParameters;
+			}
+		}
+		else
+		{
+			checkNoEntry();
 		}
+
 	}
 
 	ActiveView.Rows.Reset();
@@ -798,6 +962,7 @@ void SDebuggerDatabaseView::UpdateRows(const FTraceMotionMatchingStateMessage& S
 	}
 
 	DatabaseSequenceFilter = State.DatabaseSequenceFilter;
+	DatabaseBlendSpaceFilter = State.DatabaseBlendSpaceFilter;
 
 	SortDatabaseRows();
 	FilterDatabaseRows();
@@ -833,14 +998,16 @@ void SDebuggerDatabaseView::Construct(const FArguments& InArgs)
 
 	// @TODO: Support runtime reordering of these indices
 	// Construct all column types
-	AddColumn(MakeShared<FAnimSequenceName>(0));
-	AddColumn(MakeShared<FCost>(1));
-	AddColumn(MakeShared<FPoseCostColumn>(2));
-	AddColumn(MakeShared<FTrajectoryCost>(3));
-	AddColumn(MakeShared<FCostModifier>(4));
-	AddColumn(MakeShared<FFrame>(5));
-	AddColumn(MakeShared<FMirrored>(6));
-	AddColumn(MakeShared<FPoseIdx>(7));
+	AddColumn(MakeShared<FAssetName>(0));
+	AddColumn(MakeShared<FAssetType>(1));
+	AddColumn(MakeShared<FCost>(2));
+	AddColumn(MakeShared<FPoseCostColumn>(3));
+	AddColumn(MakeShared<FTrajectoryCost>(4));
+	AddColumn(MakeShared<FCostModifier>(5));
+	AddColumn(MakeShared<FFrame>(6));
+	AddColumn(MakeShared<FMirrored>(7));
+	AddColumn(MakeShared<FBlendParameters>(8));
+	AddColumn(MakeShared<FPoseIdx>(9));
 
 	// Active Row
 	ActiveView.HeaderRow = SNew(SHeaderRow);
@@ -1006,8 +1173,8 @@ void SDebuggerDatabaseView::Construct(const FArguments& InArgs)
 				.Padding(10, 5, 10, 5)
 				[
 					SNew(SCheckBox)
-					.IsChecked(this, &SDebuggerDatabaseView::IsSequenceFilterEnabled)
-					.OnCheckStateChanged(this, &SDebuggerDatabaseView::OnSequenceFilterEnabledChanged)
+					.IsChecked(this, &SDebuggerDatabaseView::IsAssetFilterEnabled)
+					.OnCheckStateChanged(this, &SDebuggerDatabaseView::OnAssetFilterEnabledChanged)
 					[
 						SNew(STextBlock)
 						.Text(LOCTEXT("PoseSearchDebuggerGroupFiltering", "Apply Group Filtering"))
@@ -1100,10 +1267,10 @@ void SDebuggerDetailsView::UpdateReflection(const FTraceMotionMatchingStateMessa
 	Reflection->ElapsedPoseJumpTime = State.ElapsedPoseJumpTime;
 	Reflection->bFollowUpAnimation = EnumHasAnyFlags(State.Flags, FTraceMotionMatchingState::EFlags::FollowupAnimation);
 
-	Reflection->AssetPlayerSequenceName = FString();
+	Reflection->AssetPlayerAssetName = FString();
 	if (const FPoseSearchIndexAsset* IndexAsset = Database.SearchIndex.FindAssetForPose(State.DbPoseIdx))
 	{
-		Reflection->AssetPlayerSequenceName = Database.GetSourceAsset(IndexAsset).Sequence->GetName();
+		Reflection->AssetPlayerAssetName = Database.GetSourceAssetName(IndexAsset);
 	}
 
 	Reflection->AssetPlayerTime = State.AssetPlayerTime;
@@ -1288,7 +1455,7 @@ void SDebuggerView::Tick(const FGeometry& AllottedGeometry, const double InCurre
 		bUpdated = true;
 	}
 
-	Model->UpdateAnimSequence();
+	Model->UpdateAsset();
 	
 	// Draw visualization every tick
 	DrawVisualization();
@@ -1467,7 +1634,7 @@ void SDebuggerView::DrawFeatures(
 		SkeletonDrawParams.Flags |= ESkeletonDrawFlags::ActivePose;
 	}
 
-	SkeletonDrawParams.Flags |= ESkeletonDrawFlags::AnimSequence;
+	SkeletonDrawParams.Flags |= ESkeletonDrawFlags::Asset;
 
 	ViewModel.Get()->OnDraw(SkeletonDrawParams);
 }
@@ -1525,7 +1692,7 @@ void SDebuggerView::OnPoseSelectionChanged(int32 PoseIdx, float Time)
 	else
 	{
 		Model->ShowSelectedSkeleton(PoseIdx, Time);
-		// Stop sequence player when switching selections
+		// Stop asset player when switching selections
 		Model->StopSelection();
 	}
 }
@@ -1538,7 +1705,7 @@ FReply SDebuggerView::OnUpdateNodeSelection(int32 InSelectedNodeId)
 	return FReply::Handled();
 }
 
-FReply SDebuggerView::TogglePlaySelectedSequences() const
+FReply SDebuggerView::TogglePlaySelectedAssets() const
 {
 	const TSharedPtr<SListView<TSharedRef<FDebuggerDatabaseRowData>>>& DatabaseRows = DatabaseView->GetDatabaseRows();
 	TArray<TSharedRef<FDebuggerDatabaseRowData>> Selected = DatabaseRows->GetSelectedItems();
@@ -1548,7 +1715,7 @@ FReply SDebuggerView::TogglePlaySelectedSequences() const
 		if (!Selected.IsEmpty())
 		{
 			// @TODO: Make functional with multiple poses being selected
-			ViewModel.Get()->PlaySelection(Selected[0]->PoseIdx, Selected[0]->Time);
+			ViewModel.Get()->PlaySelection(Selected[0]->PoseIdx, Selected[0]->AssetTime);
 		}
 	}
 	else
@@ -1640,7 +1807,7 @@ TSharedRef<SWidget> SDebuggerView::GenerateNodeDebuggerView()
 		.HAlign(HAlign_Fill)
 		.ButtonStyle(FEditorStyle::Get(), "Button")
 		.ContentPadding( FMargin(5, 0) )
-		.OnClicked(this, &SDebuggerView::TogglePlaySelectedSequences)
+		.OnClicked(this, &SDebuggerView::TogglePlaySelectedAssets)
 		[
 			SNew(SHorizontalBox)
 			// Icon
@@ -1660,7 +1827,7 @@ TSharedRef<SWidget> SDebuggerView::GenerateNodeDebuggerView()
 			.AutoWidth()
 			[
 				SNew(STextBlock)
-				.Text_Lambda([this] { return ViewModel.Get()->IsPlayingSelections() ? FText::FromString("Stop Selected Sequence") : FText::FromString("Play Selected Sequence"); })
+				.Text_Lambda([this] { return ViewModel.Get()->IsPlayingSelections() ? FText::FromString("Stop Selected Asset") : FText::FromString("Play Selected Asset"); })
 				.Justification(ETextJustify::Center)
 			]
 		]
@@ -1679,7 +1846,7 @@ TSharedRef<SWidget> SDebuggerView::GenerateNodeDebuggerView()
 		.Padding(0, 5, 0, 0)
 		[
 			SNew(STextBlock)
-			.Text(FText::FromString("Sequence Play Rate: "))
+			.Text(FText::FromString("Asset Play Rate: "))
 		]
 		+ SHorizontalBox::Slot()
 		.AutoWidth()
@@ -1816,6 +1983,20 @@ const FPoseSearchDatabaseSequence* FDebuggerViewModel::GetAnimSequence(int32 Seq
 	return nullptr;
 }
 
+const FPoseSearchDatabaseBlendSpace* FDebuggerViewModel::GetBlendSpace(int32 BlendSpaceIdx) const
+{
+	const UPoseSearchDatabase* Database = GetPoseSearchDatabase();
+	if (Database && Database->BlendSpaces.IsValidIndex(BlendSpaceIdx))
+	{
+		const FPoseSearchDatabaseBlendSpace* DatabaseBlendSpace = &Database->BlendSpaces[BlendSpaceIdx];
+		if (DatabaseBlendSpace)
+		{
+			return DatabaseBlendSpace;
+		}
+	}
+	return nullptr;
+}
+
 void FDebuggerViewModel::ShowSelectedSkeleton(int32 PoseIdx, float Time)
 {
 	const UPoseSearchDatabase* Database = GetPoseSearchDatabase();
@@ -1831,10 +2012,12 @@ void FDebuggerViewModel::ShowSelectedSkeleton(int32 PoseIdx, float Time)
 
 	Component->ResetToStart();
 	bSelecting = true;
-	const FPoseSearchIndexAsset* Asset = Database->SearchIndex.FindAssetForPose(PoseIdx);
-	Skeletons[SelectedPose].SequenceIdx = Asset->SourceAssetIdx;
+	const FPoseSearchIndexAsset* IndexAsset = Database->SearchIndex.FindAssetForPose(PoseIdx);
+	Skeletons[SelectedPose].Type = IndexAsset->Type;
 	Skeletons[SelectedPose].Time = Time;
-	Skeletons[SelectedPose].bMirrored = Asset->bMirrored;
+	Skeletons[SelectedPose].bMirrored = IndexAsset->bMirrored;
+	Skeletons[SelectedPose].AssetIdx = IndexAsset->SourceAssetIdx;
+	Skeletons[SelectedPose].BlendParameters = IndexAsset->BlendParameters;
 }
 
 void FDebuggerViewModel::ClearSelectedSkeleton()
@@ -1911,12 +2094,15 @@ void FDebuggerViewModel::OnUpdateNodeSelection(int32 InNodeId)
 
 	if (ActiveMotionMatchingState)
 	{
-		//todo: this seems wrong, sequenceidx == poseidx? let's test doing it right
-		//Skeletons[ActivePose].SequenceIdx = ActiveMotionMatchingState->DbPoseIdx;
+		const FPoseSearchIndexAsset* IndexAsset = NewDatabase->SearchIndex.FindAssetForPose(ActiveMotionMatchingState->DbPoseIdx);
 
-		const FPoseSearchIndexAsset* Asset = 
-			NewDatabase->SearchIndex.FindAssetForPose(ActiveMotionMatchingState->DbPoseIdx);
-		Skeletons[ActivePose].SequenceIdx = Asset->SourceAssetIdx;
+		if (IndexAsset)
+		{
+			Skeletons[Asset].Type = IndexAsset->Type;
+			Skeletons[Asset].bMirrored = IndexAsset->bMirrored;
+			Skeletons[Asset].AssetIdx = IndexAsset->SourceAssetIdx;
+			Skeletons[Asset].BlendParameters = IndexAsset->BlendParameters;
+		}
 	}
 
 	if (NewDatabase != CurrentDatabase)
@@ -1954,33 +2140,71 @@ void FDebuggerViewModel::OnDraw(FSkeletonDrawParams& DrawParams)
 	if (bDrawSelectedPose)
 	{
 		UPoseSearchMeshComponent* Component = Skeletons[SelectedPose].Component;
-		const FPoseSearchDatabaseSequence* DatabaseSequence = GetAnimSequence(Skeletons[SelectedPose].SequenceIdx);
-		UAnimSequence* Sequence = DatabaseSequence->Sequence;
 
-		UpdateContext.Sequence = Sequence;
-		UpdateContext.SequenceStartTime = Skeletons[SelectedPose].Time;
-		UpdateContext.SequenceTime = Skeletons[SelectedPose].Time;
-		UpdateContext.bMirrored = Skeletons[SelectedPose].bMirrored;
-		UpdateContext.bLoop = DatabaseSequence->bLoopAnimation;
+		if (Skeletons[SelectedPose].Type == ESearchIndexAssetType::Sequence)
+		{
+			const FPoseSearchDatabaseSequence* DatabaseSequence = GetAnimSequence(Skeletons[SelectedPose].AssetIdx);
+
+			UpdateContext.Type = ESearchIndexAssetType::Sequence;
+			UpdateContext.Sequence = DatabaseSequence->Sequence;
+			UpdateContext.StartTime = Skeletons[SelectedPose].Time;
+			UpdateContext.Time = Skeletons[SelectedPose].Time;
+			UpdateContext.bMirrored = Skeletons[SelectedPose].bMirrored;
+			UpdateContext.bLoop = DatabaseSequence->bLoopAnimation;
+		}
+		else if (Skeletons[SelectedPose].Type == ESearchIndexAssetType::BlendSpace)
+		{
+			const FPoseSearchDatabaseBlendSpace* DatabaseBlendSpace = GetBlendSpace(Skeletons[SelectedPose].AssetIdx);
+
+			UpdateContext.Type = ESearchIndexAssetType::BlendSpace;
+			UpdateContext.BlendSpace = DatabaseBlendSpace->BlendSpace;
+			UpdateContext.StartTime = Skeletons[SelectedPose].Time;
+			UpdateContext.Time = Skeletons[SelectedPose].Time;
+			UpdateContext.bMirrored = Skeletons[SelectedPose].bMirrored;
+			UpdateContext.bLoop = DatabaseBlendSpace->bLoopAnimation;
+			UpdateContext.BlendParameters = Skeletons[SelectedPose].BlendParameters;
+		}
+		else
+		{
+			checkNoEntry();
+		}
 
 		Component->UpdatePose(UpdateContext);
 	}
 
-	const bool bDrawSequence = EnumHasAnyFlags(DrawParams.Flags, ESkeletonDrawFlags::AnimSequence);
-	if (bDrawSequence && SequenceData.bActive)
+	const bool bDrawAsset = EnumHasAnyFlags(DrawParams.Flags, ESkeletonDrawFlags::Asset);
+	if (bDrawAsset && AssetData.bActive)
 	{
-		UPoseSearchMeshComponent* Component = Skeletons[AnimSequence].Component;
+		UPoseSearchMeshComponent* Component = Skeletons[Asset].Component;
 		SetDrawSkeleton(Component, true);
 
-		const FPoseSearchDatabaseSequence* DatabaseSequence = GetAnimSequence(Skeletons[AnimSequence].SequenceIdx);
-		UAnimSequence* Sequence = DatabaseSequence->Sequence;
-
-		UpdateContext.Sequence = Sequence;
-		UpdateContext.SequenceStartTime = SequenceData.StartTime;
-		UpdateContext.SequenceTime = Skeletons[AnimSequence].Time;
-		UpdateContext.bMirrored = Skeletons[SelectedPose].bMirrored;
-		UpdateContext.bLoop = DatabaseSequence->bLoopAnimation;
-
+		if (Skeletons[SelectedPose].Type == ESearchIndexAssetType::Sequence)
+		{
+			const FPoseSearchDatabaseSequence* DatabaseSequence = GetAnimSequence(Skeletons[Asset].AssetIdx);
+
+			UpdateContext.Type = ESearchIndexAssetType::Sequence;
+			UpdateContext.Sequence = DatabaseSequence->Sequence;
+			UpdateContext.StartTime = Skeletons[Asset].Time;
+			UpdateContext.Time = Skeletons[Asset].Time;
+			UpdateContext.bMirrored = Skeletons[Asset].bMirrored;
+			UpdateContext.bLoop = DatabaseSequence->bLoopAnimation;
+		}
+		else if (Skeletons[SelectedPose].Type == ESearchIndexAssetType::BlendSpace)
+		{
+			const FPoseSearchDatabaseBlendSpace* DatabaseBlendSpace = GetBlendSpace(Skeletons[Asset].AssetIdx);
+
+			UpdateContext.Type = ESearchIndexAssetType::BlendSpace;
+			UpdateContext.BlendSpace = DatabaseBlendSpace->BlendSpace;
+			UpdateContext.StartTime = Skeletons[Asset].Time;
+			UpdateContext.Time = Skeletons[Asset].Time;
+			UpdateContext.bMirrored = Skeletons[Asset].bMirrored;
+			UpdateContext.bLoop = DatabaseBlendSpace->bLoopAnimation;
+			UpdateContext.BlendParameters = Skeletons[Asset].BlendParameters;
+		}
+		else
+		{
+			checkNoEntry();
+		}
 
 		Component->UpdatePose(UpdateContext);
 	}
@@ -2042,13 +2266,13 @@ void FDebuggerViewModel::UpdateFromTimeline()
 			}
 			UPoseSearchMeshComponent* ActiveComponent = Skeletons[ActivePose].Component;
 			UPoseSearchMeshComponent* SelectedComponent = Skeletons[SelectedPose].Component;
-			UPoseSearchMeshComponent* SequenceComponent = Skeletons[AnimSequence].Component;
+			UPoseSearchMeshComponent* AssetComponent = Skeletons[Asset].Component;
 			USkeletalMesh* SkeletalMesh = TSoftObjectPtr<USkeletalMesh>(FSoftObjectPath(SkeletalMeshObjectInfo->PathName)).LoadSynchronous();
 			if (SkeletalMesh)
 			{
 				ActiveComponent->SetSkeletalMesh(SkeletalMesh);
 				SelectedComponent->SetSkeletalMesh(SkeletalMesh);
-				SequenceComponent->SetSkeletalMesh(SkeletalMesh);
+				AssetComponent->SetSkeletalMesh(SkeletalMesh);
 			}
 			FTransform ComponentWorldTransform;
 			// Active skeleton is simply the traced bone transforms
@@ -2057,15 +2281,15 @@ void FDebuggerViewModel::UpdateFromTimeline()
 			ActiveComponent->SetDebugDrawColor(FLinearColor::Green);
 			SelectedComponent->SetDebugDrawColor(FLinearColor::Blue);
 			SelectedComponent->Initialize(ComponentWorldTransform);
-			SequenceComponent->SetDebugDrawColor(FLinearColor::Red);
-			SequenceComponent->Initialize(ComponentWorldTransform);
+			AssetComponent->SetDebugDrawColor(FLinearColor::Red);
+			AssetComponent->Initialize(ComponentWorldTransform);
 
 			return TraceServices::EEventEnumerate::Stop;
 		});
 	});
 }
 
-void FDebuggerViewModel::UpdateAnimSequence()
+void FDebuggerViewModel::UpdateAsset()
 {
 	if (!IsPlayingSelections())
 	{
@@ -2073,71 +2297,89 @@ void FDebuggerViewModel::UpdateAnimSequence()
 	}
 
 	const UPoseSearchDatabase* Database = GetPoseSearchDatabase();
-	const FPoseSearchDatabaseSequence* DatabaseSequence = GetAnimSequence(Skeletons[AnimSequence].SequenceIdx);
-	UAnimSequence* Sequence = DatabaseSequence->Sequence;
-	FSkeleton& SequenceSkeleton = Skeletons[AnimSequence];
-	UPoseSearchMeshComponent* Component = SequenceSkeleton.Component;
-	auto RestartSequence = [&]()
+
+	FSkeleton& AssetSkeleton = Skeletons[Asset];
+	UPoseSearchMeshComponent* Component = AssetSkeleton.Component;
+
+	auto RestartAsset = [&]()
 	{
 		Component->ResetToStart();
-		SequenceData.AccumulatedTime = 0.0;
-		SequenceSkeleton.Time = SequenceData.StartTime;
+		AssetData.AccumulatedTime = 0.0;
+		AssetSkeleton.Time = AssetData.StartTime;
 	};
 
-	if (Sequence)
+	UAnimationAsset* AnimAsset = nullptr;
+	bool bAssetLooping = false;
 
+	if (AssetSkeleton.Type == ESearchIndexAssetType::Sequence)
+	{
+		const FPoseSearchDatabaseSequence* DatabaseSequence = GetAnimSequence(AssetSkeleton.AssetIdx);
+		AnimAsset = DatabaseSequence->Sequence;
+		bAssetLooping = DatabaseSequence->bLoopAnimation;
+	}
+	else if (AssetSkeleton.Type == ESearchIndexAssetType::BlendSpace)
+	{
+		const FPoseSearchDatabaseBlendSpace* DatabaseBlendSpace = GetBlendSpace(AssetSkeleton.AssetIdx);
+		AnimAsset = DatabaseBlendSpace->BlendSpace;
+		bAssetLooping = DatabaseBlendSpace->bLoopAnimation;
+	}
+	else
+	{
+		checkNoEntry();
+	}
 
+	if (AnimAsset)
 	{
-		const float DT = static_cast<float>(FApp::GetDeltaTime()) * SequencePlayRate;
-		const float PlayLength = Sequence->GetPlayLength();
+		const float DT = static_cast<float>(FApp::GetDeltaTime()) * AssetPlayRate;
+		const float PlayLength = AnimAsset->GetPlayLength();
 
 		const float DistanceHorizon = Database->Schema->GetTrajectoryFutureDistanceHorizon();
 		const bool bExceededDistanceHorizon = Component->LastRootMotionDelta.GetTranslation().Size() > DistanceHorizon;
 		const float TimeHorizon = Database->Schema->GetTrajectoryFutureTimeHorizon();
-		const bool bExceededTimeHorizon = (SequenceSkeleton.Time - SequenceData.StartTime) > TimeHorizon;
+		const bool bExceededTimeHorizon = (AssetSkeleton.Time - AssetData.StartTime) > TimeHorizon;
 		const bool bExceededHorizon = bExceededDistanceHorizon && bExceededTimeHorizon;
-		if (DatabaseSequence->bLoopAnimation)
+		if (bAssetLooping)
 		{
 			if (bExceededHorizon)
 			{
-				// Delay before restarting the sequence to give the user some idea of where it would land
-				if (SequenceData.AccumulatedTime > SequenceData.StopDuration)
+				// Delay before restarting the asset to give the user some idea of where it would land
+				if (AssetData.AccumulatedTime > AssetData.StopDuration)
 				{
-					RestartSequence();
+					RestartAsset();
 				}
 				else
 				{
-					SequenceData.AccumulatedTime += DT;
+					AssetData.AccumulatedTime += DT;
 				}
 				return;
 			}
 
-			SequenceSkeleton.Time += DT;
-			SequenceData.AccumulatedTime += DT;
+			AssetSkeleton.Time += DT;
+			AssetData.AccumulatedTime += DT;
 		}
 		else
 		{
-			// Used to cap the sequence, but avoid modding when updating the pose
+			// Used to cap the asset, but avoid modding when updating the pose
 			static constexpr float LengthOffset = 0.001f;
-			const bool bFinishedSequence = SequenceSkeleton.Time >= PlayLength - LengthOffset;
+			const bool bFinishedAsset = AssetSkeleton.Time >= PlayLength - LengthOffset;
 
-			// Sequence player reached end of clip or reached distance horizon of trajectory vector
-			if (bFinishedSequence || bExceededHorizon)
+			// Asset player reached end of clip or reached distance horizon of trajectory vector
+			if (bFinishedAsset || bExceededHorizon)
 			{
-				// Delay before restarting the sequence to give the user some idea of where it would land
-				if (SequenceData.AccumulatedTime > SequenceData.StopDuration)
+				// Delay before restarting the asset to give the user some idea of where it would land
+				if (AssetData.AccumulatedTime > AssetData.StopDuration)
 				{
-					RestartSequence();
+					RestartAsset();
 				}
 				else
 				{
-					SequenceData.AccumulatedTime += DT;
+					AssetData.AccumulatedTime += DT;
 				}
 			}
 			else
 			{
-				// If we haven't finished, update the play time capped by the anim sequence (not looping)
-				SequenceSkeleton.Time += DT;
+				// If we haven't finished, update the play time capped by the anim asset (not looping)
+				AssetSkeleton.Time += DT;
 			}
 		}
 	}
@@ -2174,31 +2416,36 @@ void FDebuggerViewModel::PlaySelection(int32 PoseIdx, float Time)
 	{
 		return;
 	}
-	UPoseSearchMeshComponent* Component = Skeletons[AnimSequence].Component;
+	UPoseSearchMeshComponent* Component = Skeletons[Asset].Component;
 	if (!Component)
 	{
 		return;
 	}
 
-	const FPoseSearchIndexAsset* Asset = Database->SearchIndex.FindAssetForPose(PoseIdx);
-
 	Component->ResetToStart();
-	Skeletons[AnimSequence].SequenceIdx = Asset->SourceAssetIdx;
-	Skeletons[AnimSequence].Time = SequenceData.StartTime = Time;
-	Skeletons[AnimSequence].bMirrored = Asset->bMirrored;
-	SequenceData.AccumulatedTime = 0.0f;
-	SequenceData.bActive = true;
+	
+	const FPoseSearchIndexAsset* IndexAsset = Database->SearchIndex.FindAssetForPose(PoseIdx);
+
+	Skeletons[Asset].Type = IndexAsset->Type;
+	Skeletons[Asset].AssetIdx = IndexAsset->SourceAssetIdx;
+	Skeletons[Asset].Time = Time;
+	Skeletons[Asset].bMirrored = IndexAsset->bMirrored;
+	Skeletons[Asset].BlendParameters = IndexAsset->BlendParameters;
+
+	AssetData.StartTime = Time;
+	AssetData.AccumulatedTime = 0.0f;
+	AssetData.bActive = true;
 }
 void FDebuggerViewModel::StopSelection()
 {
-	UPoseSearchMeshComponent* Component = Skeletons[AnimSequence].Component;
+	UPoseSearchMeshComponent* Component = Skeletons[Asset].Component;
 	if (!Component)
 	{
 		return;
 	}
 
-	SequenceData = {};
-	// @TODO: Make more functionality rely on checking if it should draw the sequence
+	AssetData = {};
+	// @TODO: Make more functionality rely on checking if it should draw the asset
 	Component->SetDrawDebugSkeleton(false);
 }
 void FDebuggerViewModel::OnWorldCleanup(UWorld* InWorld, bool bSessionEnded, bool bCleanupResources)
diff --git a/Engine/Plugins/Experimental/Animation/PoseSearch/Source/Editor/Private/PoseSearchDebugger.h b/Engine/Plugins/Experimental/Animation/PoseSearch/Source/Editor/Private/PoseSearchDebugger.h
--- a/Engine/Plugins/Experimental/Animation/PoseSearch/Source/Editor/Private/PoseSearchDebugger.h
+++ b/Engine/Plugins/Experimental/Animation/PoseSearch/Source/Editor/Private/PoseSearchDebugger.h
@@ -13,7 +13,9 @@
 #include "PoseSearchDebugger.generated.h"
 
 struct FPoseSearchDatabaseSequence;
+struct FPoseSearchDatabaseBlendSpace;
 class UAnimSequence;
+class UBlendSpace;
 
 namespace TraceServices { class IAnalysisSession; }
 enum class EPoseSearchFeatureDomain;
@@ -47,11 +49,14 @@ public:
 
 	struct FUpdateContext
 	{
+		ESearchIndexAssetType Type = ESearchIndexAssetType::Invalid;
 		const UAnimSequence* Sequence = nullptr;
-		float SequenceStartTime = 0.0f;
-		float SequenceTime = 0.0f;
+		const UBlendSpace* BlendSpace = nullptr;
+		float StartTime = 0.0f;
+		float Time = 0.0f;
 		bool bLoop = false;
 		bool bMirrored = false;
+		FVector BlendParameters = FVector::Zero();
 		const UMirrorDataTable* MirrorDataTable = nullptr;
 		TCustomBoneIndexArray<FCompactPoseBoneIndex, FCompactPoseBoneIndex>* CompactPoseMirrorBones = nullptr;
 		TCustomBoneIndexArray<FQuat, FCompactPoseBoneIndex>* ComponentSpaceRefRotations = nullptr;
@@ -172,8 +177,8 @@ public:
 	UPROPERTY(VisibleAnywhere, Category="Motion Matching State")
 	bool bFollowUpAnimation = false;
 	
-	UPROPERTY(VisibleAnywhere, Category = "Motion Matching State", Meta = (DisplayName = "Asset Player Sequence"))
-	FString AssetPlayerSequenceName = "";
+	UPROPERTY(VisibleAnywhere, Category = "Motion Matching State")
+	FString AssetPlayerAssetName = "";
 
 	UPROPERTY(VisibleAnywhere, Category = "Motion Matching State")
 	float AssetPlayerTime = 0.0f;
@@ -229,10 +234,10 @@ namespace UE::PoseSearch {
 /** Draw flags for the view's debug draw */
 enum class ESkeletonDrawFlags : uint32
 {
-	None              = 0,
-	ActivePose	  = 1 << 0,
-	SelectedPose  = 1 << 1,
-	AnimSequence  = 1 << 2
+	None			= 0,
+	ActivePose		= 1 << 0,
+	SelectedPose	= 1 << 1,
+	Asset			= 1 << 2,
 };
 ENUM_CLASS_FLAGS(ESkeletonDrawFlags)
 
@@ -306,11 +311,11 @@ private:
 	/** Row selection to update model view */
 	void OnDatabaseRowSelectionChanged(TSharedPtr<FDebuggerDatabaseRowData> Row, ESelectInfo::Type SelectInfo);
 
-	/** Informs the widget if sequence filtering is enabled */
-	ECheckBoxState IsSequenceFilterEnabled() const;
+	/** Informs the widget if asset filtering is enabled */
+	ECheckBoxState IsAssetFilterEnabled() const;
 
-	/** Updates the state of sequence filtering */
-	void OnSequenceFilterEnabledChanged(ECheckBoxState NewState);
+	/** Updates the state of asset filtering */
+	void OnAssetFilterEnabledChanged(ECheckBoxState NewState);
 
 	/** Generates a database row widget for the given data */
 	TSharedRef<ITableRow> HandleGenerateDatabaseRow(TSharedRef<FDebuggerDatabaseRowData> Item, const TSharedRef<STableViewBase>& OwnerTable) const;
@@ -361,6 +366,7 @@ private:
 
 	TWeakObjectPtr<const UPoseSearchDatabase> RowsSourceDatabase = nullptr;
 	TArrayView<const bool> DatabaseSequenceFilter;
+	TArrayView<const bool> DatabaseBlendSpaceFilter;
 
 	/** Database listing for filtered poses */
 	FTable FilteredDatabaseView;
@@ -371,8 +377,8 @@ private:
 	/** Text used to filter DatabaseView */
 	FText FilterText;
 
-	/** True if only sequences that pass the sequence group query are being displayed */
-	bool bSequenceFilterEnabled = true;
+	/** True if only asets that pass the asset group query are being displayed */
+	bool bAssetFilterEnabled = true;
 };
 
 /**
@@ -450,8 +456,8 @@ private:
 
 	void OnPoseSelectionChanged(int32 PoseIdx, float Time);
 
-	/** Button interaction to toggle play / stop of the anim sequence */
-	FReply TogglePlaySelectedSequences() const;
+	/** Button interaction to toggle play / stop of the asset */
+	FReply TogglePlaySelectedAssets() const;
 
 	/** Generates the message view relaying that there is no data */
 	TSharedRef<SWidget> GenerateNoDataMessageView();
@@ -549,6 +555,9 @@ public:
 	/** Get an animation sequence from the sequence ID of the active database */
 	const FPoseSearchDatabaseSequence* GetAnimSequence(int32 SequenceIdx) const;
 	
+	/** Get a blendspace from the blend space ID of the active database */
+	const FPoseSearchDatabaseBlendSpace* GetBlendSpace(int32 BlendSpaceIdx) const;
+
 	/** Sets the selected pose skeleton*/
 	void ShowSelectedSkeleton(int32 PoseIdx, float Time);
 	
@@ -561,20 +570,20 @@ public:
 	/** Stops the playing selection upon button press */
 	void StopSelection();
 	
-	/** If there is a sequence selection playing */
-	bool IsPlayingSelections() const { return SequenceData.bActive; }
+	/** If there is an asset selection playing */
+	bool IsPlayingSelections() const { return AssetData.bActive; }
 	
-	/** Play Rate of the sequence selection player */
-	void ChangePlayRate(float PlayRate) { SequencePlayRate = PlayRate; }
+	/** Play Rate of the asset selection player */
+	void ChangePlayRate(float PlayRate) { AssetPlayRate = PlayRate; }
 	
-	/** Current play rate of the sequence selection player */
-	float GetPlayRate() const { return SequencePlayRate; }
+	/** Current play rate of the asset selection player */
+	float GetPlayRate() const { return AssetPlayRate; }
 	
 	/** Callback to reset debug skeletons for the active world */
 	void OnWorldCleanup(UWorld* InWorld, bool bSessionEnded, bool bCleanupResources);
 
-	/** Updates the current playing sequence */
-	void UpdateAnimSequence();
+	/** Updates the current playing asset */
+	void UpdateAsset();
 
 private:
 	/** Update the list of states for this frame */
@@ -620,14 +629,21 @@ private:
 
 		/** Derived skeletal mesh for setting the skeleton in the scene */
 		UPoseSearchMeshComponent* Component = nullptr;
-	
-		/** Active sequence index being used for this setting this skeleton */
-		int32 SequenceIdx = 0;
+		
+		/** Type of the asset being played */
+		ESearchIndexAssetType Type = ESearchIndexAssetType::Invalid;
+
+		/** Active asset index being used for this setting this skeleton */
+		int32 AssetIdx = 0;
 		
 		/** Time in the sequence this skeleton is accessing */
 		float Time = 0.0f;
 
+		/** If this asset should be mirrored */
 		bool bMirrored = false;
+
+		/** Blend Parameters if asset is a BlendSpace */
+		FVector BlendParameters = FVector::Zero();
 	};
 	
 	/** Index for each type of skeleton we store for debug visualization */
@@ -635,7 +651,7 @@ private:
 	{
 		ActivePose = 0,
 		SelectedPose,
-		AnimSequence,
+		Asset,
 
 		Num
 	};
@@ -649,23 +665,25 @@ private:
 	/** If we currently have a selection active in the view */
 	bool bSelecting = false;
 	
-	/** Data of the active playing sequence */
-	struct FSequence
+	/** Data of the active playing asset */
+	struct FAsset
 	{
 		/** How long to stop upon reaching target */
 		static constexpr float StopDuration = 2.0f;
+
 		/** Time since the start of play */
 		float AccumulatedTime = 0.0f;
 
-		/** Start time of the sequence */
+		/** Start time of the asset */
 		float StartTime = 0.0f;
-		/** If we are currently playing a sequence */
+
+		/** If we are currently playing a asset */
 		bool bActive = false;
 
-	} SequenceData;
+	} AssetData;
 	
-	/** Current play rate of the sequence selection player */
-	float SequencePlayRate = 1.0f;
+	/** Current play rate of the asset selection player */
+	float AssetPlayRate = 1.0f;
 
 	/** Limits some public API */
 	friend class FDebugger;
diff --git a/Engine/Plugins/Experimental/Animation/PoseSearch/Source/Editor/Private/Trace/PoseSearchTraceAnalyzer.cpp b/Engine/Plugins/Experimental/Animation/PoseSearch/Source/Editor/Private/Trace/PoseSearchTraceAnalyzer.cpp
--- a/Engine/Plugins/Experimental/Animation/PoseSearch/Source/Editor/Private/Trace/PoseSearchTraceAnalyzer.cpp
+++ b/Engine/Plugins/Experimental/Animation/PoseSearch/Source/Editor/Private/Trace/PoseSearchTraceAnalyzer.cpp
@@ -55,6 +55,7 @@ bool FTraceAnalyzer::OnEvent(uint16 RouteId, EStyle Style, const FOnEventContext
 				Message.AnimLinearVelocity = EventData.GetValue<float>("AnimLinearVelocity");
 				Message.AnimAngularVelocity = EventData.GetValue<float>("AnimAngularVelocity");
 				Message.DatabaseSequenceFilter = EventData.GetArrayView<bool>("DatabaseSequenceFilter");
+				Message.DatabaseBlendSpaceFilter = EventData.GetArrayView<bool>("DatabaseBlendSpaceFilter");
 
 				TArrayView<const float> ChannelWeightScales = EventData.GetArrayView<float>("ChannelWeightScales");
 				TArrayView<const float> HistoryWeightScales = EventData.GetArrayView<float>("HistoryWeightScales");
diff --git a/Engine/Plugins/Experimental/Animation/PoseSearch/Source/Editor/Private/Trace/PoseSearchTraceProvider.h b/Engine/Plugins/Experimental/Animation/PoseSearch/Source/Editor/Private/Trace/PoseSearchTraceProvider.h
--- a/Engine/Plugins/Experimental/Animation/PoseSearch/Source/Editor/Private/Trace/PoseSearchTraceProvider.h
+++ b/Engine/Plugins/Experimental/Animation/PoseSearch/Source/Editor/Private/Trace/PoseSearchTraceProvider.h
@@ -37,6 +37,7 @@ struct FTraceMotionMatchingStateMessage : FTraceMessage
 	float AnimLinearVelocity = 0.0f;
 	float AnimAngularVelocity = 0.0f;
 	TArray<bool> DatabaseSequenceFilter;
+	TArray<bool> DatabaseBlendSpaceFilter;
 };
 
 
diff --git a/Engine/Plugins/Experimental/Animation/PoseSearch/Source/Runtime/Private/AnimNode_MotionMatching.cpp b/Engine/Plugins/Experimental/Animation/PoseSearch/Source/Runtime/Private/AnimNode_MotionMatching.cpp
--- a/Engine/Plugins/Experimental/Animation/PoseSearch/Source/Runtime/Private/AnimNode_MotionMatching.cpp
+++ b/Engine/Plugins/Experimental/Animation/PoseSearch/Source/Runtime/Private/AnimNode_MotionMatching.cpp
@@ -6,6 +6,7 @@
 #include "Animation/AnimNode_Inertialization.h"
 #include "Animation/AnimSequence.h"
 #include "Animation/MotionTrajectoryTypes.h"
+#include "Animation/AnimRootMotionProvider.h"
 #include "DynamicPlayRate/DynamicPlayRateLibrary.h"
 #include "PoseSearch/AnimNode_PoseSearchHistoryCollector.h"
 #include "PoseSearch/PoseSearch.h"
@@ -24,8 +25,12 @@ void FAnimNode_MotionMatching::Initialize_AnyThread(const FAnimationInitializeCo
 
 	MotionMatchingState.InitNewDatabaseSearch(Database, Settings.SearchThrottleTime, nullptr /*OutError*/);
 
+	CurrentAssetPlayerNode = &SequencePlayerNode;
+
+	MirrorNode.SetSourceLinkNode(CurrentAssetPlayerNode);
+	MirrorNode.SetMirrorDataTable(Database->Schema->MirrorDataTable.Get());
+
 	Source.SetLinkNode(&MirrorNode);
-	MirrorNode.SetSourceLinkNode(&SequencePlayerNode);
 	Source.Initialize(Context);
 }
 
@@ -38,6 +43,18 @@ void FAnimNode_MotionMatching::Evaluate_AnyThread(FPoseContext& Output)
 #if WITH_EDITORONLY_DATA
 	bWasEvaluated = true;
 #endif
+
+#if UE_POSE_SEARCH_TRACE_ENABLED
+	const UE::Anim::IAnimRootMotionProvider* RootMotionProvider = UE::Anim::IAnimRootMotionProvider::Get();
+
+	if (ensureMsgf(RootMotionProvider, TEXT("Could not get Root Motion Provider.")))
+	{
+		if (ensureMsgf(RootMotionProvider->HasRootMotion(Output.CustomAttributes), TEXT("Motion Matching Output no Root Motion Attribute.")))
+		{
+			RootMotionProvider->ExtractRootMotion(Output.CustomAttributes, MotionMatchingState.RootMotionTransformDelta);
+		}
+	}
+#endif
 }
 
 void FAnimNode_MotionMatching::UpdateAssetPlayer(const FAnimationUpdateContext& Context)
@@ -46,8 +63,8 @@ void FAnimNode_MotionMatching::UpdateAssetPlayer(const FAnimationUpdateContext&
 
 	GetEvaluateGraphExposedInputs().Execute(Context);
 
-	// Update with the sequence player's current time.
-	MotionMatchingState.AssetPlayerTime = SequencePlayerNode.GetAccumulatedTime();
+	// Update the current time.
+	MotionMatchingState.AssetPlayerTime = CurrentAssetPlayerNode->GetAccumulatedTime();
 
 	// Execute core motion matching algorithm and retain across frame state
 	UpdateMotionMatchingState(
@@ -59,32 +76,57 @@ void FAnimNode_MotionMatching::UpdateAssetPlayer(const FAnimationUpdateContext&
 		MotionMatchingState
 	);
 
+	const FPoseSearchIndexAsset* SearchIndexAsset = MotionMatchingState.GetCurrentSearchIndexAsset();
+
 	// If a new pose is requested, jump to the pose by updating the embedded sequence player node
-	if ((MotionMatchingState.Flags & EMotionMatchingFlags::JumpedToPose) == EMotionMatchingFlags::JumpedToPose)
+	if (SearchIndexAsset && (MotionMatchingState.Flags & EMotionMatchingFlags::JumpedToPose) == EMotionMatchingFlags::JumpedToPose)
 	{
-		const FPoseSearchIndexAsset* SearchIndexAsset = MotionMatchingState.GetCurrentSearchIndexAsset();
-		const FPoseSearchDatabaseSequence& ResultDbSequence = Database->GetSourceAsset(SearchIndexAsset);
-		SequencePlayerNode.SetSequence(ResultDbSequence.Sequence);
-		SequencePlayerNode.SetAccumulatedTime(MotionMatchingState.AssetPlayerTime);
-		SequencePlayerNode.SetLoopAnimation(ResultDbSequence.bLoopAnimation);
-		SequencePlayerNode.SetPlayRate(1.0f);
-
-		MirrorNode.SetMirrorDataTable(Database->Schema->MirrorDataTable.Get());
+		if (SearchIndexAsset->Type == ESearchIndexAssetType::Sequence)
+		{
+			CurrentAssetPlayerNode = &SequencePlayerNode;
+
+			const FPoseSearchDatabaseSequence& ResultDbSequence = Database->GetSequenceSourceAsset(SearchIndexAsset);
+			SequencePlayerNode.SetAccumulatedTime(MotionMatchingState.AssetPlayerTime);
+			SequencePlayerNode.SetSequence(ResultDbSequence.Sequence);
+			SequencePlayerNode.SetLoopAnimation(ResultDbSequence.bLoopAnimation);
+			SequencePlayerNode.SetPlayRate(1.0f);
+		}
+		else if (SearchIndexAsset->Type == ESearchIndexAssetType::BlendSpace)
+		{
+			CurrentAssetPlayerNode = &BlendSpacePlayerNode;
+
+			const FPoseSearchDatabaseBlendSpace& ResultDbBlendSpace = Database->GetBlendSpaceSourceAsset(SearchIndexAsset);
+			BlendSpacePlayerNode.SetAccumulatedTime(MotionMatchingState.AssetPlayerTime);
+			BlendSpacePlayerNode.SetBlendSpace(ResultDbBlendSpace.BlendSpace);
+			BlendSpacePlayerNode.SetLoop(ResultDbBlendSpace.bLoopAnimation);
+			BlendSpacePlayerNode.SetPlayRate(1.0f);
+			BlendSpacePlayerNode.SetPosition(SearchIndexAsset->BlendParameters);
+		}
+		else
+		{
+			checkNoEntry();
+		}
+
+		MirrorNode.SetSourceLinkNode(CurrentAssetPlayerNode);
 		MirrorNode.SetMirror(SearchIndexAsset->bMirrored);
 	}
 
-	// Optionally applying dynamic play rate adjustment to chosen sequences based on predictive motion analysis
-	const float PlayRate = DynamicPlayRateAdjustment(
-		Context,
-		Trajectory,
-		DynamicPlayRateSettings,
-		SequencePlayerNode.GetSequence(),
-		SequencePlayerNode.GetAccumulatedTime(),
-		SequencePlayerNode.GetPlayRate(),
-		SequencePlayerNode.GetLoopAnimation()
-	);
+	if (SearchIndexAsset && SearchIndexAsset->Type == ESearchIndexAssetType::Sequence)
+	{
+		// Optionally applying dynamic play rate adjustment to chosen sequences based on predictive motion analysis
+		const float PlayRate = DynamicPlayRateAdjustment(
+			Context,
+			Trajectory,
+			DynamicPlayRateSettings,
+			SequencePlayerNode.GetSequence(),
+			SequencePlayerNode.GetAccumulatedTime(),
+			SequencePlayerNode.GetPlayRate(),
+			SequencePlayerNode.GetLoopAnimation()
+		);
+
+		SequencePlayerNode.SetPlayRate(PlayRate);
+	}
 
-	SequencePlayerNode.SetPlayRate(PlayRate);
 	Source.Update(Context);
 }
 
@@ -136,27 +178,27 @@ void FAnimNode_MotionMatching::GatherDebugData(FNodeDebugData& DebugData)
 // FAnimNode_AssetPlayerBase interface
 float FAnimNode_MotionMatching::GetAccumulatedTime() const
 {
-	return SequencePlayerNode.GetAccumulatedTime();
+	return CurrentAssetPlayerNode->GetAccumulatedTime();
 }
 
 UAnimationAsset* FAnimNode_MotionMatching::GetAnimAsset() const
 {
-	return SequencePlayerNode.GetAnimAsset();
+	return CurrentAssetPlayerNode->GetAnimAsset();
 }
 
 float FAnimNode_MotionMatching::GetCurrentAssetLength() const
 {
-	return SequencePlayerNode.GetCurrentAssetLength();
+	return CurrentAssetPlayerNode->GetCurrentAssetLength();
 }
 
 float FAnimNode_MotionMatching::GetCurrentAssetTime() const
 {
-	return SequencePlayerNode.GetCurrentAssetTime();
+	return CurrentAssetPlayerNode->GetCurrentAssetLength();
 }
 
 float FAnimNode_MotionMatching::GetCurrentAssetTimePlayRateAdjusted() const
 {
-	return SequencePlayerNode.GetCurrentAssetTimePlayRateAdjusted();
+	return CurrentAssetPlayerNode->GetCurrentAssetTimePlayRateAdjusted();
 }
 
 bool FAnimNode_MotionMatching::GetIgnoreForRelevancyTest() const
diff --git a/Engine/Plugins/Experimental/Animation/PoseSearch/Source/Runtime/Private/PoseSearch.cpp b/Engine/Plugins/Experimental/Animation/PoseSearch/Source/Runtime/Private/PoseSearch.cpp
--- a/Engine/Plugins/Experimental/Animation/PoseSearch/Source/Runtime/Private/PoseSearch.cpp
+++ b/Engine/Plugins/Experimental/Animation/PoseSearch/Source/Runtime/Private/PoseSearch.cpp
@@ -27,6 +27,9 @@
 #include "Trace/PoseSearchTraceLogger.h"
 #include "UObject/ObjectSaveContext.h"
 #include "Misc/MemStack.h"
+#include "Animation/BuiltInAttributeTypes.h"
+#include "Animation/AnimRootMotionProvider.h"
+#include "Animation/BlendSpace1D.h"
 
 IMPLEMENT_ANIMGRAPH_MESSAGE(UE::PoseSearch::IPoseHistoryProvider);
 
@@ -85,7 +88,7 @@ static inline float CompareFeatureVectors(int32 NumValues, const float* A, const
 	return (float)Dissimilarity;
 }
 
-FLinearColor GetColorForFeature(FPoseSearchFeatureDesc Feature, const FPoseSearchFeatureVectorLayout* Layout)
+static FLinearColor GetColorForFeature(FPoseSearchFeatureDesc Feature, const FPoseSearchFeatureVectorLayout* Layout)
 {
 	const float FeatureIdx = Layout->Features.IndexOfByKey(Feature);
 	const float FeatureCountIdx = Layout->Features.Num() - 1;
@@ -194,7 +197,7 @@ static constexpr FFeatureTypeTraits FeatureTypeTraits[] =
 	{ EPoseSearchFeatureType::ForwardVector, 3 },
 };
 
-FFeatureTypeTraits GetFeatureTypeTraits(EPoseSearchFeatureType Type)
+static FFeatureTypeTraits GetFeatureTypeTraits(EPoseSearchFeatureType Type)
 {
 	// Could allow external registration to a TSet of traits in the future
 	// For now just use a simple local array
@@ -967,29 +970,49 @@ const FPoseSearchIndexAsset* FPoseSearchIndex::FindAssetForPose(int32 PoseIdx) c
 	return Assets.FindByPredicate(Predicate);
 }
 
-float FPoseSearchIndex::GetTimeOffset(int32 PoseIdx, const FPoseSearchIndexAsset* Asset) const
+float FPoseSearchIndex::GetAssetTime(int32 PoseIdx, const FPoseSearchIndexAsset* Asset) const
 {
-	const FPoseSearchIndexAsset* ValidAsset = Asset;
-	if (!ValidAsset)
+	if (!Asset)
 	{
-		ValidAsset = FindAssetForPose(PoseIdx);
-		if (!ValidAsset)
+		Asset = FindAssetForPose(PoseIdx);
+		if (!Asset)
 		{
 			UE_LOG(LogPoseSearch, Error, TEXT("Couldn't find asset for pose %i in database"), PoseIdx);
 			return -1.0f;
 		}
 	}
-	if (!ValidAsset->IsPoseInRange(PoseIdx))
+
+	if (!Asset->IsPoseInRange(PoseIdx))
 	{
 		UE_LOG(LogPoseSearch, Error, TEXT("Pose %i out of range in database"), PoseIdx);
 		return -1.0f;
 	}
 
-	const FFloatInterval SamplingRange = ValidAsset->SamplingInterval;
-	float TimeOffsetSeconds = FMath::Min(
-		SamplingRange.Min + Schema->SamplingInterval * (PoseIdx - ValidAsset->FirstPoseIdx), 
-		SamplingRange.Max);
-	return TimeOffsetSeconds;
+	if (Asset->Type == ESearchIndexAssetType::Sequence)
+	{
+		const FFloatInterval SamplingRange = Asset->SamplingInterval;
+
+		float AssetTime = FMath::Min(
+			SamplingRange.Min + Schema->SamplingInterval * (PoseIdx - Asset->FirstPoseIdx),
+			SamplingRange.Max);
+		return AssetTime;
+	}
+	else if (Asset->Type == ESearchIndexAssetType::BlendSpace)
+	{
+		const FFloatInterval SamplingRange = Asset->SamplingInterval;
+
+		// For BlendSpaces the AssetTime is in the range [0, 1] while the Sampling Range
+		// is in real time (seconds)
+		float AssetTime = FMath::Min(
+			SamplingRange.Min + Schema->SamplingInterval * (PoseIdx - Asset->FirstPoseIdx),
+			SamplingRange.Max) / (Asset->NumPoses * Schema->SamplingInterval);
+		return AssetTime;
+	}
+	else
+	{
+		checkNoEntry();
+		return -1.0f;
+	}
 }
 
 
@@ -1119,18 +1142,19 @@ FFloatInterval FPoseSearchDatabaseSequence::GetEffectiveSamplingRange() const
 //////////////////////////////////////////////////////////////////////////
 // UPoseSearchDatabase
 
-int32 UPoseSearchDatabase::GetPoseIndexFromAssetTime(float AssetTime, const FPoseSearchIndexAsset* SearchIndexAsset) const
+int32 UPoseSearchDatabase::GetPoseIndexFromTime(float Time, const FPoseSearchIndexAsset* SearchIndexAsset) const
 {
-	const FPoseSearchDatabaseSequence& DbSequence = Sequences[SearchIndexAsset->SourceAssetIdx];
-
 	FFloatInterval Range = SearchIndexAsset->SamplingInterval;
 
-	if (Range.Contains(AssetTime))
+	if (Range.Contains(Time))
 	{
-		int32 PoseOffset = FMath::RoundToInt(Schema->SampleRate * (AssetTime - Range.Min));
+		int32 PoseOffset = FMath::RoundToInt(Schema->SampleRate * (Time - Range.Min));
+		
+		check(PoseOffset >= 0);
+
 		if (PoseOffset >= SearchIndexAsset->NumPoses)
 		{
-			if (DbSequence.bLoopAnimation)
+			if (IsSourceAssetLooping(SearchIndexAsset))
 			{
 				PoseOffset -= SearchIndexAsset->NumPoses;
 			}
@@ -1139,6 +1163,7 @@ int32 UPoseSearchDatabase::GetPoseIndexFromAssetTime(float AssetTime, const FPos
 				PoseOffset = SearchIndexAsset->NumPoses - 1;
 			}
 		}
+
 		int32 PoseIdx = SearchIndexAsset->FirstPoseIdx + PoseOffset;
 		return PoseIdx;
 	}
@@ -1146,32 +1171,76 @@ int32 UPoseSearchDatabase::GetPoseIndexFromAssetTime(float AssetTime, const FPos
 	return INDEX_NONE;
 }
 
-float UPoseSearchDatabase::GetTimeOffset(int32 PoseIdx, const FPoseSearchIndexAsset* Asset) const
+float UPoseSearchDatabase::GetAssetTime(int32 PoseIdx, const FPoseSearchIndexAsset* Asset) const
 {
-	return SearchIndex.GetTimeOffset(PoseIdx, Asset);
+	float AssetTime = SearchIndex.GetAssetTime(PoseIdx, Asset);
+	return AssetTime;
 }
 
-const FPoseSearchDatabaseSequence& UPoseSearchDatabase::GetSourceAsset(const FPoseSearchIndexAsset* SearchIndexAsset) const
+const FPoseSearchDatabaseSequence& UPoseSearchDatabase::GetSequenceSourceAsset(const FPoseSearchIndexAsset* SearchIndexAsset) const
 {
+	check(SearchIndexAsset->Type == ESearchIndexAssetType::Sequence);
 	return Sequences[SearchIndexAsset->SourceAssetIdx];
 }
 
-int32 UPoseSearchDatabase::FindSequenceForPose(int32 PoseIdx) const
+const FPoseSearchDatabaseBlendSpace& UPoseSearchDatabase::GetBlendSpaceSourceAsset(const FPoseSearchIndexAsset* SearchIndexAsset) const
+{
+	check(SearchIndexAsset->Type == ESearchIndexAssetType::BlendSpace);
+	return BlendSpaces[SearchIndexAsset->SourceAssetIdx];
+}
+
+const bool UPoseSearchDatabase::IsSourceAssetLooping(const FPoseSearchIndexAsset* SearchIndexAsset) const
 {
-	const FPoseSearchIndexAsset* Asset = SearchIndex.FindAssetForPose(PoseIdx);
-	return Asset->SourceAssetIdx;
+	if (SearchIndexAsset->Type == ESearchIndexAssetType::Sequence)
+	{
+		return GetSequenceSourceAsset(SearchIndexAsset).bLoopAnimation;
+	}
+	else if (SearchIndexAsset->Type == ESearchIndexAssetType::BlendSpace)
+	{
+		return GetBlendSpaceSourceAsset(SearchIndexAsset).bLoopAnimation;
+	}
+	else
+	{
+		checkNoEntry();
+		return false;
+	}
 }
 
-float UPoseSearchDatabase::GetSequenceLength(int32 DbSequenceIdx) const
+const FGameplayTagContainer* UPoseSearchDatabase::GetSourceAssetGroupTags(const FPoseSearchIndexAsset* SearchIndexAsset) const
 {
-	return Sequences[DbSequenceIdx].Sequence->GetPlayLength();
+	if (SearchIndexAsset->Type == ESearchIndexAssetType::Sequence)
+	{
+		return &GetSequenceSourceAsset(SearchIndexAsset).GroupTags;
+	}
+	else if (SearchIndexAsset->Type == ESearchIndexAssetType::BlendSpace)
+	{
+		return &GetBlendSpaceSourceAsset(SearchIndexAsset).GroupTags;
+	}
+	else
+	{
+		checkNoEntry();
+		return nullptr;
+	}
 }
 
-bool UPoseSearchDatabase::DoesSequenceLoop(int32 DbSequenceIdx) const
+const FString UPoseSearchDatabase::GetSourceAssetName(const FPoseSearchIndexAsset* SearchIndexAsset) const
 {
-	return Sequences[DbSequenceIdx].bLoopAnimation;
+	if (SearchIndexAsset->Type == ESearchIndexAssetType::Sequence)
+	{
+		return GetSequenceSourceAsset(SearchIndexAsset).Sequence->GetName();
+	}
+	else if (SearchIndexAsset->Type == ESearchIndexAssetType::BlendSpace)
+	{
+		return GetBlendSpaceSourceAsset(SearchIndexAsset).BlendSpace->GetName();
+	}
+	else
+	{
+		checkNoEntry();
+		return FString();
+	}
 }
 
+
 bool UPoseSearchDatabase::IsValidForIndexing() const
 {
 	bool bValid = Schema && Schema->IsValid() && !Sequences.IsEmpty();
@@ -1225,7 +1294,26 @@ void UPoseSearchDatabase::CollectSimpleSequences()
 	SimpleSequences.Reset();
 }
 
-void FindValidSequenceIntervals(const FPoseSearchDatabaseSequence& DbSequence, TArray<FFloatRange>& ValidRanges)
+void UPoseSearchDatabase::CollectSimpleBlendSpaces()
+{
+	for (auto& SimpleBlendSpace : SimpleBlendSpaces)
+	{
+		auto Predicate = [&SimpleBlendSpace](FPoseSearchDatabaseBlendSpace& DbBlendSpace) -> bool
+		{
+			return DbBlendSpace.BlendSpace == SimpleBlendSpace;
+		};
+
+		if (!BlendSpaces.ContainsByPredicate(Predicate))
+		{
+			FPoseSearchDatabaseBlendSpace& DbBlendSpace = BlendSpaces.AddDefaulted_GetRef();
+			DbBlendSpace.BlendSpace = SimpleBlendSpace;
+		}
+	}
+
+	SimpleBlendSpaces.Reset();
+}
+
+static void FindValidSequenceIntervals(const FPoseSearchDatabaseSequence& DbSequence, TArray<FFloatRange>& ValidRanges)
 {
 	const UAnimSequence* Sequence = DbSequence.Sequence;
 	check(DbSequence.Sequence);
@@ -1269,13 +1357,101 @@ void FindValidSequenceIntervals(const FPoseSearchDatabaseSequence& DbSequence, T
 	}
 }
 
+static inline void CollectGroupIndices(
+	TArrayView<const FPoseSearchDatabaseGroup> Groups,
+	const FGameplayTagContainer& GroupTags,
+	const int32 Index,
+	TArray<int32>& GroupIndices,
+	TArray<int32>& BadGroupIndices)
+{
+	GroupIndices.Reset();
+
+	for (const FGameplayTag& GroupTag : GroupTags)
+	{
+		const int32 GroupIndex = Groups.IndexOfByPredicate([&](const FPoseSearchDatabaseGroup& DatabaseGroup)
+		{
+			return DatabaseGroup.Tag == GroupTag;
+		});
+
+		// we don't add INDEX_NONE because index none represents a choice to use the default group by not adding
+		// any group identifiers. If an added identifier doesn't match, that's an error. In the future this
+		// should be made robust enough to prevent these errors from happening
+		if (GroupIndex == INDEX_NONE)
+		{
+			BadGroupIndices.Add(Index);
+		}
+		else if (Groups[GroupIndex].bUseGroupWeights)
+		{
+			GroupIndices.Add(GroupIndex);
+		}
+	}
+
+	if (GroupIndices.Num() == 0)
+	{
+		GroupIndices.Add(INDEX_NONE);
+	}
+}
+
+void FPoseSearchDatabaseBlendSpace::GetBlendSpaceParameterSampleRanges(
+	int32& HorizontalBlendNum,
+	int32& VerticalBlendNum,
+	float& HorizontalBlendMin,
+	float& HorizontalBlendMax,
+	float& VerticalBlendMin,
+	float& VerticalBlendMax) const
+{
+	HorizontalBlendNum = bUseGridForSampling ? BlendSpace->GetBlendParameter(0).GridNum + 1 : FMath::Max(NumberOfHorizontalSamples, 1);
+	VerticalBlendNum = bUseGridForSampling ? BlendSpace->GetBlendParameter(1).GridNum + 1 : FMath::Max(NumberOfVerticalSamples, 1);
+
+	check(HorizontalBlendNum >= 1 && VerticalBlendNum >= 1);
+
+	HorizontalBlendMin = BlendSpace->GetBlendParameter(0).Min;
+	HorizontalBlendMax = BlendSpace->GetBlendParameter(0).Max;
+
+	VerticalBlendMin = BlendSpace->GetBlendParameter(1).Min;
+	VerticalBlendMax = BlendSpace->GetBlendParameter(1).Max;
+
+	if (BlendSpace->IsA<UBlendSpace1D>())
+	{
+		VerticalBlendNum = 1;
+		VerticalBlendMin = 0.0;
+		VerticalBlendMax = 0.0;
+	}
+}
+
+static FVector BlendParameterForSampleRanges(
+	int32 HorizontalBlendIndex,
+	int32 VerticalBlendIndex,
+	int32 HorizontalBlendNum,
+	int32 VerticalBlendNum,
+	float HorizontalBlendMin,
+	float HorizontalBlendMax,
+	float VerticalBlendMin,
+	float VerticalBlendMax)
+{
+	return FVector(
+		HorizontalBlendNum > 1 ? 
+			HorizontalBlendMin + (HorizontalBlendMax - HorizontalBlendMin) * 
+			((float)HorizontalBlendIndex) / (HorizontalBlendNum - 1) : 
+		HorizontalBlendMin,
+		VerticalBlendNum > 1 ? 
+			VerticalBlendMin + (VerticalBlendMax - VerticalBlendMin) * 
+			((float)VerticalBlendIndex) / (VerticalBlendNum - 1) : 
+		VerticalBlendMin,
+		0.0f);
+}
+
 bool UPoseSearchDatabase::TryInitSearchIndexAssets()
 {
 	SearchIndex.Assets.Empty();
+	
 	bool bAnyMirrored = false;
+	
 	TArray<FFloatRange> ValidRanges;
 	TArray<int32> GroupIndices;
-	TArray<int32> BadGroupSequenceIndices;
+	TArray<int32> BadSequenceGroupIndices;
+	TArray<int32> BadBlendSpaceGroupIndices;
+
 	for (int32 SequenceIdx = 0; SequenceIdx < Sequences.Num(); ++SequenceIdx)
 	{
 		const FPoseSearchDatabaseSequence& Sequence = Sequences[SequenceIdx];
@@ -1286,31 +1462,12 @@ bool UPoseSearchDatabase::TryInitSearchIndexAssets()
 			Sequence.MirrorOption == EPoseSearchMirrorOption::MirroredOnly ||
 			Sequence.MirrorOption == EPoseSearchMirrorOption::UnmirroredAndMirrored;
 
-		GroupIndices.Reset();
-		for (const FGameplayTag& SequenceGroupTag : Sequence.GroupTags)
-		{
-			const int32 SequenceGroupIndex = Groups.IndexOfByPredicate([&](const FPoseSearchDatabaseGroup& DatabaseGroup)
-			{
-				return DatabaseGroup.Tag == SequenceGroupTag;
-			});
-
-			// we don't add INDEX_NONE because index none represents a choice to use the default group by not adding
-			// any group identifiers. If an added identifier doesn't match, that's an error. In the future this
-			// should be made robust enough to prevent these errors from happening
-			if (SequenceGroupIndex == INDEX_NONE)
-			{
-				BadGroupSequenceIndices.Add(SequenceIdx);
-			}
-			else if (Groups[SequenceGroupIndex].bUseGroupWeights)
-			{
-				GroupIndices.Add(SequenceGroupIndex);
-			}
-		}
-
-		if (GroupIndices.Num() == 0)
-		{
-			GroupIndices.Add(INDEX_NONE);
-		}
+		CollectGroupIndices(
+			Groups,
+			Sequence.GroupTags,
+			SequenceIdx,
+			GroupIndices,
+			BadSequenceGroupIndices);
 
 		for (int32 GroupIndex : GroupIndices)
 		{
@@ -1322,6 +1479,7 @@ bool UPoseSearchDatabase::TryInitSearchIndexAssets()
 				{
 					SearchIndex.Assets.Add(
 						FPoseSearchIndexAsset(
+							ESearchIndexAssetType::Sequence,
 							GroupIndex,
 							SequenceIdx,
 							false,
@@ -1332,6 +1490,7 @@ bool UPoseSearchDatabase::TryInitSearchIndexAssets()
 				{
 					SearchIndex.Assets.Add(
 						FPoseSearchIndexAsset(
+							ESearchIndexAssetType::Sequence,
 							GroupIndex,
 							SequenceIdx,
 							true,
@@ -1342,6 +1501,87 @@ bool UPoseSearchDatabase::TryInitSearchIndexAssets()
 		}
 	}
 
+	TArray<FBlendSampleData> BlendSamples;
+
+	for (int32 BlendSpaceIdx = 0; BlendSpaceIdx < BlendSpaces.Num(); ++BlendSpaceIdx)
+	{
+		const FPoseSearchDatabaseBlendSpace& BlendSpace = BlendSpaces[BlendSpaceIdx];
+		bool bAddUnmirrored =
+			BlendSpace.MirrorOption == EPoseSearchMirrorOption::UnmirroredOnly ||
+			BlendSpace.MirrorOption == EPoseSearchMirrorOption::UnmirroredAndMirrored;
+		bool bAddMirrored =
+			BlendSpace.MirrorOption == EPoseSearchMirrorOption::MirroredOnly ||
+			BlendSpace.MirrorOption == EPoseSearchMirrorOption::UnmirroredAndMirrored;
+
+		CollectGroupIndices(
+			Groups,
+			BlendSpace.GroupTags,
+			BlendSpaceIdx,
+			GroupIndices,
+			BadBlendSpaceGroupIndices);
+
+		for (int32 GroupIndex : GroupIndices)
+		{
+			int32 HorizontalBlendNum, VerticalBlendNum;
+			float HorizontalBlendMin, HorizontalBlendMax, VerticalBlendMin, VerticalBlendMax;
+
+			BlendSpace.GetBlendSpaceParameterSampleRanges(
+				HorizontalBlendNum,
+				VerticalBlendNum,
+				HorizontalBlendMin,
+				HorizontalBlendMax,
+				VerticalBlendMin,
+				VerticalBlendMax);
+
+			for (int32 HorizontalIndex = 0; HorizontalIndex < HorizontalBlendNum; HorizontalIndex++)
+			{
+				for (int32 VerticalIndex = 0; VerticalIndex < VerticalBlendNum; VerticalIndex++)
+				{
+					FVector BlendParameters = BlendParameterForSampleRanges(
+						HorizontalIndex,
+						VerticalIndex,
+						HorizontalBlendNum,
+						VerticalBlendNum,
+						HorizontalBlendMin,
+						HorizontalBlendMax,
+						VerticalBlendMin,
+						VerticalBlendMax);
+						
+					int32 TriangulationIndex = 0;
+					BlendSpace.BlendSpace->GetSamplesFromBlendInput(BlendParameters, BlendSamples, TriangulationIndex, true);
+
+					float PlayLength = BlendSpace.BlendSpace->GetAnimationLengthFromSampleData(BlendSamples);
+
+					if (bAddUnmirrored)
+					{
+						SearchIndex.Assets.Add(
+							FPoseSearchIndexAsset(
+								ESearchIndexAssetType::BlendSpace,
+								GroupIndex,
+								BlendSpaceIdx,
+								false,
+								FFloatInterval(0.0f, PlayLength),
+								BlendParameters));
+					}
+
+					if (bAddMirrored)
+					{
+						SearchIndex.Assets.Add(
+							FPoseSearchIndexAsset(
+								ESearchIndexAssetType::BlendSpace,
+								GroupIndex,
+								BlendSpaceIdx,
+								true,
+								FFloatInterval(0.0f, PlayLength),
+								BlendParameters));
+						bAnyMirrored = true;
+					}
+				}
+			}
+		}
+	}
+
+
 	if (bAnyMirrored && !Schema->MirrorDataTable)
 	{
 		UE_LOG(
@@ -1354,7 +1594,7 @@ bool UPoseSearchDatabase::TryInitSearchIndexAssets()
 		return false;
 	}
 
-	for (int32 BadGroupSequenceIdx : BadGroupSequenceIndices)
+	for (int32 BadGroupSequenceIdx : BadSequenceGroupIndices)
 	{
 		UE_LOG(
 			LogPoseSearch,
@@ -1364,6 +1604,16 @@ bool UPoseSearchDatabase::TryInitSearchIndexAssets()
 			*GetNameSafe(Sequences[BadGroupSequenceIdx].Sequence));
 	}
 
+	for (int32 BadGroupBlendSpaceIdx : BadBlendSpaceGroupIndices)
+	{
+		UE_LOG(
+			LogPoseSearch,
+			Warning,
+			TEXT("Database %s, blendspace %s is asking for a group that doesn't exist"),
+			*GetNameSafe(this),
+			*GetNameSafe(BlendSpaces[BadGroupBlendSpaceIdx].BlendSpace));
+	}
+
 	return true;
 }
 
@@ -1396,6 +1646,14 @@ void UPoseSearchDatabase::PostEditChangeProperty(struct FPropertyChangedEvent& P
 			CollectSimpleSequences();
 		}
 	}
+
+	if (PropertyChangedEvent.Property->GetFName() == GET_MEMBER_NAME_CHECKED(UPoseSearchDatabase, SimpleBlendSpaces))
+	{
+		if (!SimpleBlendSpaces.IsEmpty())
+		{
+			CollectSimpleBlendSpaces();
+		}
+	}
 }
 #endif // WITH_EDITOR
 
@@ -2125,125 +2383,510 @@ void FAnimSamplingContext::FillCompactPoseAndComponentRefRotations()
 }
 
 //////////////////////////////////////////////////////////////////////////
-// FSequenceSampler
+// IAssetSampler
+
+struct IAssetSampler
+{
+public:
+	virtual ~IAssetSampler() {};
+
+	virtual float GetPlayLength() const = 0;
+	virtual bool IsLoopable() const = 0;
+
+	// Gets the time associated with a particular root distance travelled
+	virtual float GetTimeFromRootDistance(float Distance) const = 0;
+
+	// Gets the total root distance travelled 
+	virtual float GetTotalRootDistance() const = 0;
+
+	// Gets the final root transformation at the end of the asset's playback time
+	virtual FTransform GetTotalRootTransform() const = 0;
+
+	// Extracts pose for this asset for a given context
+	virtual void ExtractPose(const FAnimExtractContext& ExtractionCtx, FAnimationPoseData& OutAnimPoseData) const = 0;
+
+	// Extracts the accumulated root distance at the given time, using the extremities of the sequence to extrapolate 
+	// beyond the sequence limits when Time is less than zero or greater than the sequence length
+	virtual float ExtractRootDistance(float Time) const = 0;
+
+	// Extracts root transform at the given time, using the extremities of the sequence to extrapolate beyond the 
+	// sequence limits when Time is less than zero or greater than the sequence length.
+	virtual FTransform ExtractRootTransform(float Time) const = 0;
+
+	// Extracts notify states inheriting from UAnimNotifyState_PoseSearchBase present in the sequence at Time.
+	// The function does not empty NotifyStates before adding new notifies!
+	virtual void ExtractPoseSearchNotifyStates(float Time, TArray<class UAnimNotifyState_PoseSearchBase*>& NotifyStates) const = 0;
+};
+
+// Uses distance delta between NextRootDistanceIndex and NextRootDistanceIndex - 1 and extrapolates it to ExtrapolationTime
+static float ExtrapolateAccumulatedRootDistance(
+	int32 SamplingRate,
+	TArrayView<const float> AccumulatedRootDistance,
+	int32 NextRootDistanceIndex, 
+	float ExtrapolationTime,
+	const FPoseSearchExtrapolationParameters& ExtrapolationParameters)
+{
+	check(NextRootDistanceIndex > 0 && NextRootDistanceIndex < AccumulatedRootDistance.Num());
+
+	const float DistanceDelta =
+		AccumulatedRootDistance[NextRootDistanceIndex] -
+		AccumulatedRootDistance[NextRootDistanceIndex - 1];
+	const float Speed = DistanceDelta * SamplingRate;
+	const float ExtrapolationSpeed = Speed >= ExtrapolationParameters.LinearSpeedThreshold ?
+		Speed : 0.0f;
+	const float ExtrapolatedDistance = ExtrapolationSpeed * ExtrapolationTime;
+
+	return ExtrapolatedDistance;
+}
+
+static float ExtractAccumulatedRootDistance(
+	int32 SamplingRate,
+	TArrayView<const float> AccumulatedRootDistance,
+	float PlayLength,
+	float Time,
+	const FPoseSearchExtrapolationParameters& ExtrapolationParameters)
+{
+	const float ClampedTime = FMath::Clamp(Time, 0.0f, PlayLength);
+
+	// Find the distance sample that corresponds with the time and split into whole and partial parts
+	float IntegralDistanceSample;
+	float DistanceAlpha = FMath::Modf(ClampedTime * SamplingRate, &IntegralDistanceSample);
+	float DistanceIdx = (int32)IntegralDistanceSample;
+
+	// Verify the distance offset and any residual portion would be in bounds
+	check(DistanceIdx + (DistanceAlpha > 0.0f ? 1 : 0) < AccumulatedRootDistance.Num());
+
+	// Look up the distance and interpolate between distance samples if necessary
+	float Distance = AccumulatedRootDistance[DistanceIdx];
+	if (DistanceAlpha > 0.0f)
+	{
+		float NextDistance = AccumulatedRootDistance[DistanceIdx + 1];
+		Distance = FMath::Lerp(Distance, NextDistance, DistanceAlpha);
+	}
+
+	const float ExtrapolationTime = Time - ClampedTime;
+
+	if (ExtrapolationTime != 0.0f)
+	{
+		// If extrapolationTime is not zero, we extrapolate the beginning or the end of the animation to estimate
+		// the root distance.
+		const int32 DistIdx = (ExtrapolationTime > 0.0f) ? AccumulatedRootDistance.Num() - 1 : 1;
+		const float ExtrapolatedDistance = ExtrapolateAccumulatedRootDistance(
+			SamplingRate,
+			AccumulatedRootDistance,
+			DistIdx,
+			ExtrapolationTime,
+			ExtrapolationParameters);
+		Distance += ExtrapolatedDistance;
+	}
+
+	return Distance;
+}
+
+static FTransform ExtrapolateRootMotion(
+	FTransform SampleToExtrapolate,
+	float SampleStart, 
+	float SampleEnd, 
+	float ExtrapolationTime,
+	const FPoseSearchExtrapolationParameters& ExtrapolationParameters)
+{
+	const float SampleDelta = SampleEnd - SampleStart;
+	check(!FMath::IsNearlyZero(SampleDelta));
+
+	const FVector LinearVelocityToExtrapolate = SampleToExtrapolate.GetTranslation() / SampleDelta;
+	const float LinearSpeedToExtrapolate = LinearVelocityToExtrapolate.Size();
+	const bool bCanExtrapolateTranslation =
+		LinearSpeedToExtrapolate >= ExtrapolationParameters.LinearSpeedThreshold;
+
+	const float AngularSpeedToExtrapolateRad = SampleToExtrapolate.GetRotation().GetAngle() / SampleDelta;
+	const bool bCanExtrapolateRotation =
+		FMath::RadiansToDegrees(AngularSpeedToExtrapolateRad) >= ExtrapolationParameters.AngularSpeedThreshold;
+
+	if (!bCanExtrapolateTranslation && !bCanExtrapolateRotation)
+	{
+		return FTransform::Identity;
+	}
+
+	if (!bCanExtrapolateTranslation)
+	{
+		SampleToExtrapolate.SetTranslation(FVector::ZeroVector);
+	}
+
+	if (!bCanExtrapolateRotation)
+	{
+		SampleToExtrapolate.SetRotation(FQuat::Identity);
+	}
+
+	// converting ExtrapolationTime to a positive number to avoid dealing with the negative extrapolation and inverting
+	// transforms later on.
+	const float AbsExtrapolationTime = FMath::Abs(ExtrapolationTime);
+	const float AbsSampleDelta = FMath::Abs(SampleDelta);
+	const FTransform AbsTimeSampleToExtrapolate =
+		ExtrapolationTime >= 0.0f ? SampleToExtrapolate : SampleToExtrapolate.Inverse();
+
+	// because we're extrapolating rotation, the extrapolation must be integrated over time
+	const float SampleMultiplier = AbsExtrapolationTime / AbsSampleDelta;
+	float IntegralNumSamples;
+	float RemainingSampleFraction = FMath::Modf(SampleMultiplier, &IntegralNumSamples);
+	int32 NumSamples = (int32)IntegralNumSamples;
+
+	// adding full samples to the extrapolated root motion
+	FTransform ExtrapolatedRootMotion = FTransform::Identity;
+	for (int i = 0; i < NumSamples; ++i)
+	{
+		ExtrapolatedRootMotion = AbsTimeSampleToExtrapolate * ExtrapolatedRootMotion;
+	}
+
+	// and a blend with identity for whatever is left
+	FTransform RemainingExtrapolatedRootMotion;
+	RemainingExtrapolatedRootMotion.Blend(
+		FTransform::Identity,
+		AbsTimeSampleToExtrapolate,
+		RemainingSampleFraction);
+
+	ExtrapolatedRootMotion = RemainingExtrapolatedRootMotion * ExtrapolatedRootMotion;
+	return ExtrapolatedRootMotion;
+}
+
+
+//////////////////////////////////////////////////////////////////////////
+// FSequenceSampler
+
+struct FSequenceSampler : public IAssetSampler
+{
+public:
+	struct FInput
+	{
+		const UAnimSequence* Sequence = nullptr;
+		bool bLoopable = false;
+		int32 RootDistanceSamplingRate = 60;
+		FPoseSearchExtrapolationParameters ExtrapolationParameters;
+	} Input;
+
+	void Init(const FInput& Input);
+	void Process();
+
+	float GetPlayLength() const override { return Input.Sequence->GetPlayLength(); };
+	bool IsLoopable() const override { return Input.bLoopable; };
+
+	float GetTimeFromRootDistance(float Distance) const override;
+
+	float GetTotalRootDistance() const override { return TotalRootDistance; };
+	FTransform GetTotalRootTransform() const override { return TotalRootTransform; }
+
+	virtual void ExtractPose(const FAnimExtractContext& ExtractionCtx, FAnimationPoseData& OutAnimPoseData) const override;
+	virtual float ExtractRootDistance(float Time) const override;
+	virtual FTransform ExtractRootTransform(float Time) const override;
+	virtual void ExtractPoseSearchNotifyStates(float Time, TArray<class UAnimNotifyState_PoseSearchBase*>& NotifyStates) const override;
+
+private:
+	float TotalRootDistance = 0.0f;
+	FTransform TotalRootTransform = FTransform::Identity;
+	TArray<float> AccumulatedRootDistance;
+
+	void ProcessRootDistance();
+};
+
+void FSequenceSampler::Init(const FInput& InInput)
+{
+	check(InInput.Sequence);
+
+	Input = InInput;
+}
+
+void FSequenceSampler::Process()
+{
+	ProcessRootDistance();
+}
+
+float FSequenceSampler::GetTimeFromRootDistance(float Distance) const
+{
+	int32 NextSampleIdx = 1;
+	int32 PrevSampleIdx = 0;
+	if (Distance > 0.0f)
+	{
+		// Search for the distance value. Because the values will be extrapolated if necessary
+		// LowerBound might go past the end of the array, in which case the last valid index is used
+		int32 ClipDistanceLowerBoundIndex = Algo::LowerBound(AccumulatedRootDistance, Distance);
+		NextSampleIdx = FMath::Min(
+			ClipDistanceLowerBoundIndex,
+			AccumulatedRootDistance.Num() - 1);
+
+		// Compute distance interpolation amount
+		PrevSampleIdx = FMath::Max(0, NextSampleIdx - 1);
+	}
+
+	float NextDistance = AccumulatedRootDistance[NextSampleIdx];
+	float PrevDistance = AccumulatedRootDistance[PrevSampleIdx];
+	float DistanceSampleAlpha = FMath::GetRangePct(PrevDistance, NextDistance, Distance);
+
+	// Convert to time
+	float ClipTime = (float(NextSampleIdx) - (1.0f - DistanceSampleAlpha)) / Input.RootDistanceSamplingRate;
+	return ClipTime;
+}
+
+void FSequenceSampler::ExtractPose(const FAnimExtractContext& ExtractionCtx, FAnimationPoseData& OutAnimPoseData) const
+{
+	Input.Sequence->GetAnimationPose(OutAnimPoseData, ExtractionCtx);
+}
+
+FTransform FSequenceSampler::ExtractRootTransform(float Time) const
+{
+	if (Input.bLoopable)
+	{
+		FTransform LoopableRootTransform = Input.Sequence->ExtractRootMotion(0.0f, Time, true);
+		return LoopableRootTransform;
+	}
+
+	const float ExtrapolationSampleTime = Input.ExtrapolationParameters.SampleTime;
+
+	const float PlayLength = Input.Sequence->GetPlayLength();
+	const float ClampedTime = FMath::Clamp(Time, 0.0f, PlayLength);
+	const float ExtrapolationTime = Time - ClampedTime;
+
+	FTransform RootTransform = FTransform::Identity;
+
+	// If Time is less than zero, ExtrapolationTime will be negative. In this case, we extrapolate the beginning of the 
+	// animation to estimate where the root would be at Time
+	if (ExtrapolationTime < -SMALL_NUMBER)
+	{
+		FTransform SampleToExtrapolate = Input.Sequence->ExtractRootMotionFromRange(0.0f, ExtrapolationSampleTime);
+
+		const FTransform ExtrapolatedRootMotion = ExtrapolateRootMotion(
+			SampleToExtrapolate,
+			0.0f, ExtrapolationSampleTime, 
+			ExtrapolationTime,
+			Input.ExtrapolationParameters);
+		RootTransform = ExtrapolatedRootMotion;
+	}
+	else
+	{
+		RootTransform = Input.Sequence->ExtractRootMotionFromRange(0.0f, ClampedTime);
+
+		// If Time is greater than PlayLength, ExtrapolationTIme will be a positive number. In this case, we extrapolate
+		// the end of the animation to estimate where the root would be at Time
+		if (ExtrapolationTime > SMALL_NUMBER)
+		{
+			FTransform SampleToExtrapolate = Input.Sequence->ExtractRootMotionFromRange(PlayLength - ExtrapolationSampleTime, PlayLength);
+
+			const FTransform ExtrapolatedRootMotion = ExtrapolateRootMotion(
+				SampleToExtrapolate,
+				PlayLength - ExtrapolationSampleTime, PlayLength,
+				ExtrapolationTime,
+				Input.ExtrapolationParameters);
+			RootTransform = ExtrapolatedRootMotion * RootTransform;
+		}
+	}
+
+	return RootTransform;
+}
+
+float FSequenceSampler::ExtractRootDistance(float Time) const
+{
+	return ExtractAccumulatedRootDistance(
+		Input.RootDistanceSamplingRate,
+		AccumulatedRootDistance,
+		Input.Sequence->GetPlayLength(),
+		Time,
+		Input.ExtrapolationParameters);
+}
+
+void FSequenceSampler::ExtractPoseSearchNotifyStates(
+	float Time, 
+	TArray<UAnimNotifyState_PoseSearchBase*>& NotifyStates) const
+{
+	// getting pose search notifies in an interval of size ExtractionInterval, centered on Time
+	constexpr float ExtractionInterval = 1.0f / 120.0f;
+	FAnimNotifyContext NotifyContext;
+	Input.Sequence->GetAnimNotifies(Time - (ExtractionInterval * 0.5f), ExtractionInterval, NotifyContext);
+
+	// check which notifies actually overlap Time and are of the right base type
+	for (const FAnimNotifyEventReference& EventReference : NotifyContext.ActiveNotifies)
+	{
+		const FAnimNotifyEvent* NotifyEvent = EventReference.GetNotify();
+		if (!NotifyEvent)
+		{
+			continue;
+		}
+
+		if (NotifyEvent->GetTriggerTime() > Time ||
+			NotifyEvent->GetEndTriggerTime() < Time)
+		{
+			continue;
+		}
+
+		UAnimNotifyState_PoseSearchBase* PoseSearchAnimNotify = 
+			Cast<UAnimNotifyState_PoseSearchBase>(NotifyEvent->NotifyStateClass);
+		if (PoseSearchAnimNotify)
+		{
+			NotifyStates.Add(PoseSearchAnimNotify);
+		}
+	}
+}
+
+void FSequenceSampler::ProcessRootDistance()
+{
+	// Note the distance sampling interval is independent of the schema's sampling interval
+	const float DistanceSamplingInterval = 1.0f / Input.RootDistanceSamplingRate;
+
+	const FTransform InitialRootTransform = Input.Sequence->ExtractRootTrackTransform(0.0f, nullptr);
+
+	uint32 NumDistanceSamples = FMath::CeilToInt(Input.Sequence->GetPlayLength() * Input.RootDistanceSamplingRate) + 1;
+	AccumulatedRootDistance.Reserve(NumDistanceSamples);
+
+	// Build a distance lookup table by sampling root motion at a fixed rate and accumulating
+	// absolute translation deltas. During indexing we'll bsearch this table and interpolate
+	// between samples in order to convert distance offsets to time offsets.
+	// See also FSequenceIndexer::AddTrajectoryDistanceFeatures().
+
+	double TotalAccumulatedRootDistance = 0.0;
+	FTransform LastRootTransform = InitialRootTransform;
+	float SampleTime = 0.0f;
+	for (int32 SampleIdx = 0; SampleIdx != NumDistanceSamples; ++SampleIdx)
+	{
+		SampleTime = FMath::Min(SampleIdx * DistanceSamplingInterval, Input.Sequence->GetPlayLength());
+
+		FTransform RootTransform = Input.Sequence->ExtractRootTrackTransform(SampleTime, nullptr);
+		FTransform LocalRootMotion = RootTransform.GetRelativeTransform(LastRootTransform);
+		LastRootTransform = RootTransform;
+
+		TotalAccumulatedRootDistance += LocalRootMotion.GetTranslation().Size();
+		AccumulatedRootDistance.Add((float)TotalAccumulatedRootDistance);
+	}
+
+	// Verify we sampled the final frame of the clip
+	check(SampleTime == Input.Sequence->GetPlayLength());
+
+	// Also emit root motion summary info to help with sample wrapping in 
+	// FSequenceIndexer::GetSampleTimeFromDistance() and FSequenceIndexer::GetSampleInfo()
+	TotalRootTransform = LastRootTransform.GetRelativeTransform(InitialRootTransform);
+	TotalRootDistance = AccumulatedRootDistance.Last();
+}
+
+//////////////////////////////////////////////////////////////////////////
+// FBlendSpaceSampler
 
-struct FSequenceSampler
+struct FBlendSpaceSampler : public IAssetSampler
 {
 public:
 	struct FInput
 	{
 		const FAnimSamplingContext* SamplingContext = nullptr;
-		const UPoseSearchSchema* Schema = nullptr;
-		const UAnimSequence* Sequence = nullptr;
+		const UBlendSpace* BlendSpace = nullptr;
 		bool bLoopable = false;
-		int32 DistanceSamplingRate = 60;
+		int32 RootDistanceSamplingRate = 60;
+		int32 RootTransformSamplingRate = 60;
 		FPoseSearchExtrapolationParameters ExtrapolationParameters;
+		FVector BlendParameters;
 	} Input;
 
-	struct FOutput
-	{
-		TArray<float> AccumulatedRootDistance;
-		int32 NumDistanceSamples = 0;
-		float PlayLength = 0.0f;
-		float TotalRootDistance = 0.0f;
-		FTransform TotalRootMotion = FTransform::Identity;
-	} Output;
-
-	void Reset();
 	void Init(const FInput& Input);
+
 	void Process();
 
-	// Extracts pose from input sequence and mirrors it if necessary
-	void ExtractPose(const FAnimExtractContext& ExtractionCtx, bool bMirrored, FAnimationPoseData& OutAnimPoseData) const;
+	float GetPlayLength() const override { return PlayLength; };
+	bool IsLoopable() const override { return Input.bLoopable; };
 
-	// Extracts root transform at the given time, using the extremities of the sequence to extrapolate beyond the 
-	// sequence limits when Time is less than zero or greater than the sequence length.
-	FTransform ExtractRootTransform(float Time) const;
+	float GetTimeFromRootDistance(float Distance) const override;
 
-	// Extracts the accumulated root distance at the given time, using the extremities of the sequence to extrapolate 
-	// beyond the sequence limits when Time is less than zero or greater than the sequence length
-	float ExtractRootDistance(float Time) const;
+	float GetTotalRootDistance() const override { return TotalRootDistance; };
+	FTransform GetTotalRootTransform() const override { return TotalRootTransform; }
 
-	// Extracts notify states inheriting from UAnimNotifyState_PoseSearchBase present in the sequence at Time.
-	// The function does not empty NotifyStates before adding new notifies!
-	void ExtractPoseSearchNotifyStates(float Time, TArray<class UAnimNotifyState_PoseSearchBase*>& NotifyStates) const;
+	virtual void ExtractPose(const FAnimExtractContext& ExtractionCtx, FAnimationPoseData& OutAnimPoseData) const override;
+	virtual float ExtractRootDistance(float Time) const override;
+	virtual FTransform ExtractRootTransform(float Time) const override;
+	virtual void ExtractPoseSearchNotifyStates(float Time, TArray<class UAnimNotifyState_PoseSearchBase*>& NotifyStates) const override;
 
 private:
-	void Reserve();
-	void ProcessRootMotion();
-
-	// Samples sequence and adjusts obtained root motion to ExtrapolationTime
-	FTransform ExtrapolateRootMotion(float SampleStart, float SampleEnd, float ExtrapolationTime) const;
-
-	// Uses distance delta between NextRootDistanceIndex and NextRootDistanceIndex - 1 and extrapolates it to 
-	// ExtrapolationTime
-	float ExtrapolateRootDistance(int32 NextRootDistanceIndex, float ExtrapolationTime) const;
-
+	float PlayLength = 0.0f;
+	float TotalRootDistance = 0.0f;
+	FTransform TotalRootTransform = FTransform::Identity;
+	TArray<float> AccumulatedRootDistance;
+	TArray<FTransform> AccumulatedRootTransform;
+	
+	void ProcessPlayLength();
+	void ProcessRootDistance();
+	void ProcessRootTransform();
+
+	// Extracts the pre-computed blend space root transform. ProcessRootTransform must be run first.
+	FTransform ExtractBlendSpaceRootTrackTransform(float Time) const;
+	FTransform ExtractBlendSpaceRootMotion(float StartTime, float DeltaTime, bool bAllowLooping) const;
+	FTransform ExtractBlendSpaceRootMotionFromRange(float StartTrackPosition, float EndTrackPosition) const;
 };
 
-void FSequenceSampler::Init(const FInput& InInput)
+void FBlendSpaceSampler::Init(const FInput& InInput)
 {
-	check(InInput.Schema);
-	check(InInput.Schema->IsValid());
-	check(InInput.Sequence);
-
-	Reset();
+	check(InInput.BlendSpace);
 
 	Input = InInput;
-
-	Output.PlayLength = Input.Sequence->GetPlayLength();
-	Output.NumDistanceSamples = FMath::CeilToInt(Output.PlayLength * Input.DistanceSamplingRate) + 1;
-
-	Reserve();
 }
 
-void FSequenceSampler::Reset()
+void FBlendSpaceSampler::Process()
 {
-	Input = FInput();
+	FMemMark Mark(FMemStack::Get());
 
-	Output.NumDistanceSamples = 0;
-	Output.PlayLength = 0.0f;
-	Output.TotalRootDistance = 0.0f;
-	Output.AccumulatedRootDistance.Reset(0);
+	ProcessPlayLength();
+	ProcessRootTransform();
+	ProcessRootDistance();
 }
 
-void FSequenceSampler::Reserve()
+float FBlendSpaceSampler::GetTimeFromRootDistance(float Distance) const
 {
-	Output.AccumulatedRootDistance.Reserve(Output.NumDistanceSamples);
-}
+	int32 NextSampleIdx = 1;
+	int32 PrevSampleIdx = 0;
+	if (Distance > 0.0f)
+	{
+		// Search for the distance value. Because the values will be extrapolated if necessary
+		// LowerBound might go past the end of the array, in which case the last valid index is used
+		int32 ClipDistanceLowerBoundIndex = Algo::LowerBound(AccumulatedRootDistance, Distance);
+		NextSampleIdx = FMath::Min(
+			ClipDistanceLowerBoundIndex,
+			AccumulatedRootDistance.Num() - 1);
 
-void FSequenceSampler::Process()
-{
-	ProcessRootMotion();
+		// Compute distance interpolation amount
+		PrevSampleIdx = FMath::Max(0, NextSampleIdx - 1);
+	}
+
+	float NextDistance = AccumulatedRootDistance[NextSampleIdx];
+	float PrevDistance = AccumulatedRootDistance[PrevSampleIdx];
+	float DistanceSampleAlpha = FMath::GetRangePct(PrevDistance, NextDistance, Distance);
+
+	// Convert to time
+	float ClipTime = (float(NextSampleIdx) - (1.0f - DistanceSampleAlpha)) / Input.RootDistanceSamplingRate;
+	return ClipTime;
 }
 
-void FSequenceSampler::ExtractPose(const FAnimExtractContext& ExtractionCtx, bool bMirrored, FAnimationPoseData& OutAnimPoseData) const
+void FBlendSpaceSampler::ExtractPose(const FAnimExtractContext& ExtractionCtx, FAnimationPoseData& OutAnimPoseData) const
 {
-	check(!bMirrored || Input.Schema->MirrorDataTable);
-	Input.Sequence->GetAnimationPose(OutAnimPoseData, ExtractionCtx);
-	if (bMirrored)
+	TArray<FBlendSampleData> BlendSamples;
+	int32 TriangulationIndex = 0;
+	Input.BlendSpace->GetSamplesFromBlendInput(Input.BlendParameters, BlendSamples, TriangulationIndex, true);
+
+	for (int32 BlendSampleIdex = 0; BlendSampleIdex < BlendSamples.Num(); BlendSampleIdex++)
 	{
-		FAnimationRuntime::MirrorPose(
-			OutAnimPoseData.GetPose(),
-			Input.Schema->MirrorDataTable->MirrorAxis,
-			Input.SamplingContext->CompactPoseMirrorBones,
-			Input.SamplingContext->ComponentSpaceRefRotations
-		);
-		// Note curves and attributes are not used during the indexing process and therefore don't need to be mirrored
+		float Scale = BlendSamples[BlendSampleIdex].Animation->GetPlayLength() / PlayLength;
+
+		FDeltaTimeRecord BlendSampleDeltaTimeRecord;
+		BlendSampleDeltaTimeRecord.Set(ExtractionCtx.DeltaTimeRecord.GetPrevious() * Scale, ExtractionCtx.DeltaTimeRecord.Delta * Scale);
+
+		BlendSamples[BlendSampleIdex].DeltaTimeRecord = BlendSampleDeltaTimeRecord;
+		BlendSamples[BlendSampleIdex].PreviousTime = ExtractionCtx.DeltaTimeRecord.GetPrevious() * Scale;
+		BlendSamples[BlendSampleIdex].Time = ExtractionCtx.CurrentTime * Scale;
 	}
+
+	Input.BlendSpace->GetAnimationPose(BlendSamples, ExtractionCtx, OutAnimPoseData);
 }
 
-FTransform FSequenceSampler::ExtractRootTransform(float Time) const
+FTransform FBlendSpaceSampler::ExtractRootTransform(float Time) const
 {
 	if (Input.bLoopable)
 	{
-		FTransform LoopableRootTransform = Input.Sequence->ExtractRootMotion(0.0f, Time, true);
+		FTransform LoopableRootTransform = ExtractBlendSpaceRootMotion(0.0f, Time, true);
 		return LoopableRootTransform;
 	}
 
 	const float ExtrapolationSampleTime = Input.ExtrapolationParameters.SampleTime;
 
-	const float PlayLength = Input.Sequence->GetPlayLength();
 	const float ClampedTime = FMath::Clamp(Time, 0.0f, PlayLength);
 	const float ExtrapolationTime = Time - ClampedTime;
 
@@ -2253,22 +2896,30 @@ FTransform FSequenceSampler::ExtractRootTransform(float Time) const
 	// animation to estimate where the root would be at Time
 	if (ExtrapolationTime < -SMALL_NUMBER)
 	{
+		FTransform SampleToExtrapolate = ExtractBlendSpaceRootMotionFromRange(0.0f, ExtrapolationSampleTime);
+
 		const FTransform ExtrapolatedRootMotion = ExtrapolateRootMotion(
-			0.0f, ExtrapolationSampleTime, 
-			ExtrapolationTime);
+			SampleToExtrapolate,
+			0.0f, ExtrapolationSampleTime,
+			ExtrapolationTime,
+			Input.ExtrapolationParameters);
 		RootTransform = ExtrapolatedRootMotion;
 	}
 	else
 	{
-		RootTransform = Input.Sequence->ExtractRootMotionFromRange(0.0f, ClampedTime);
+		RootTransform = ExtractBlendSpaceRootMotionFromRange(0.0f, ClampedTime);
 
 		// If Time is greater than PlayLength, ExtrapolationTIme will be a positive number. In this case, we extrapolate
 		// the end of the animation to estimate where the root would be at Time
 		if (ExtrapolationTime > SMALL_NUMBER)
 		{
+			FTransform SampleToExtrapolate = ExtractBlendSpaceRootMotionFromRange(PlayLength - ExtrapolationSampleTime, PlayLength);
+
 			const FTransform ExtrapolatedRootMotion = ExtrapolateRootMotion(
+				SampleToExtrapolate,
 				PlayLength - ExtrapolationSampleTime, PlayLength,
-				ExtrapolationTime);
+				ExtrapolationTime,
+				Input.ExtrapolationParameters);
 			RootTransform = ExtrapolatedRootMotion * RootTransform;
 		}
 	}
@@ -2276,189 +2927,276 @@ FTransform FSequenceSampler::ExtractRootTransform(float Time) const
 	return RootTransform;
 }
 
-float FSequenceSampler::ExtractRootDistance(float Time) const
+float FBlendSpaceSampler::ExtractRootDistance(float Time) const
 {
-	const float ClampedTime = FMath::Clamp(Time, 0.0f, Output.PlayLength);
-
-	// Find the distance sample that corresponds with the time and split into whole and partial parts
-	float IntegralDistanceSample;
-	float DistanceAlpha = FMath::Modf(ClampedTime * Input.DistanceSamplingRate, &IntegralDistanceSample);
-	float DistanceIdx = (int32)IntegralDistanceSample;
-
-	// Verify the distance offset and any residual portion would be in bounds
-	check(DistanceIdx + (DistanceAlpha > 0.0f ? 1 : 0) < Output.AccumulatedRootDistance.Num());
-	
-	// Look up the distance and interpolate between distance samples if necessary
-	float Distance = Output.AccumulatedRootDistance[DistanceIdx];
-	if (DistanceAlpha > 0.0f)
-	{
-		float NextDistance = Output.AccumulatedRootDistance[DistanceIdx + 1];
-		Distance = FMath::Lerp(Distance, NextDistance, DistanceAlpha);
-	}
-
-	const float ExtrapolationTime = Time - ClampedTime;
+	return ExtractAccumulatedRootDistance(
+		Input.RootDistanceSamplingRate,
+		AccumulatedRootDistance,
+		PlayLength,
+		Time,
+		Input.ExtrapolationParameters);
+}
 
-	if (ExtrapolationTime != 0.0f)
+static int32 GetHighestWeightSample(const TArray<struct FBlendSampleData>& SampleDataList)
+{
+	int32 HighestWeightIndex = 0;
+	float HighestWeight = SampleDataList[HighestWeightIndex].GetClampedWeight();
+	for (int32 I = 1; I < SampleDataList.Num(); I++)
 	{
-		// If extrapolationTime is not zero, we extrapolate the beginning or the end of the animation to estimate
-		// the root distance.
-		const int32 DistIdx = (ExtrapolationTime > 0.0f) ? Output.AccumulatedRootDistance.Num() - 1 : 1;
-		const float ExtrapolatedDistance = ExtrapolateRootDistance(DistIdx, ExtrapolationTime);
-		Distance += ExtrapolatedDistance;
+		if (SampleDataList[I].GetClampedWeight() > HighestWeight)
+		{
+			HighestWeightIndex = I;
+			HighestWeight = SampleDataList[I].GetClampedWeight();
+		}
 	}
-
-	return Distance;
+	return HighestWeightIndex;
 }
 
-void FSequenceSampler::ExtractPoseSearchNotifyStates(
-	float Time, 
+void FBlendSpaceSampler::ExtractPoseSearchNotifyStates(
+	float Time,
 	TArray<UAnimNotifyState_PoseSearchBase*>& NotifyStates) const
 {
-	// getting pose search notifies in an interval of size ExtractionInterval, centered on Time
-	constexpr float ExtractionInterval = 1.0f / 120.0f;
-	FAnimNotifyContext NotifyContext;
-	Input.Sequence->GetAnimNotifies(Time - (ExtractionInterval * 0.5f), ExtractionInterval, NotifyContext);
-
-	// check which notifies actually overlap Time and are of the right base type
-	for (const FAnimNotifyEventReference& EventReference : NotifyContext.ActiveNotifies)
+	if (Input.BlendSpace->NotifyTriggerMode == ENotifyTriggerMode::HighestWeightedAnimation)
 	{
-		const FAnimNotifyEvent* NotifyEvent = EventReference.GetNotify();
-		if (!NotifyEvent)
-		{
-			continue;
-		}
+		// Set up blend samples
+		TArray<FBlendSampleData> BlendSamples;
+		int32 TriangulationIndex = 0;
+		Input.BlendSpace->GetSamplesFromBlendInput(Input.BlendParameters, BlendSamples, TriangulationIndex, true);
 
-		if (NotifyEvent->GetTriggerTime() > Time ||
-			NotifyEvent->GetEndTriggerTime() < Time)
-		{
-			continue;
-		}
+		// Find highest weighted
+		const int32 HighestWeightIndex = GetHighestWeightSample(BlendSamples);
 
-		UAnimNotifyState_PoseSearchBase* PoseSearchAnimNotify = 
-			Cast<UAnimNotifyState_PoseSearchBase>(NotifyEvent->NotifyStateClass);
-		if (PoseSearchAnimNotify)
+		check(HighestWeightIndex != -1);
+
+		// getting pose search notifies in an interval of size ExtractionInterval, centered on Time
+		constexpr float ExtractionInterval = 1.0f / 120.0f;
+
+		float SampleTime = Time * (BlendSamples[HighestWeightIndex].Animation->GetPlayLength() / PlayLength);
+
+		// Get notifies for highest weighted
+		FAnimNotifyContext NotifyContext;
+		BlendSamples[HighestWeightIndex].Animation->GetAnimNotifies(
+			(SampleTime - (ExtractionInterval * 0.5f)),
+			ExtractionInterval, 
+			NotifyContext);
+
+		// check which notifies actually overlap Time and are of the right base type
+		for (const FAnimNotifyEventReference& EventReference : NotifyContext.ActiveNotifies)
 		{
-			NotifyStates.Add(PoseSearchAnimNotify);
+			const FAnimNotifyEvent* NotifyEvent = EventReference.GetNotify();
+			if (!NotifyEvent)
+			{
+				continue;
+			}
+
+			if (NotifyEvent->GetTriggerTime() > SampleTime ||
+				NotifyEvent->GetEndTriggerTime() < SampleTime)
+			{
+				continue;
+			}
+
+			UAnimNotifyState_PoseSearchBase* PoseSearchAnimNotify =
+				Cast<UAnimNotifyState_PoseSearchBase>(NotifyEvent->NotifyStateClass);
+			if (PoseSearchAnimNotify)
+			{
+				NotifyStates.Add(PoseSearchAnimNotify);
+			}
 		}
 	}
 }
 
+FTransform FBlendSpaceSampler::ExtractBlendSpaceRootTrackTransform(float Time) const
+{
+	checkf(AccumulatedRootTransform.Num() > 0, TEXT("ProcessRootTransform must be run first"));
+
+	int32 Index = Time * Input.RootTransformSamplingRate;
+	int32 FirstIndexClamped = FMath::Clamp(Index + 0, 0, AccumulatedRootTransform.Num() - 1);
+	int32 SecondIndexClamped = FMath::Clamp(Index + 1, 0, AccumulatedRootTransform.Num() - 1);
+	float Alpha = FMath::Fmod(Time * Input.RootTransformSamplingRate, 1.0f);
+	FTransform OutputTransform;
+	OutputTransform.Blend(
+		AccumulatedRootTransform[FirstIndexClamped],
+		AccumulatedRootTransform[SecondIndexClamped],
+		Alpha);
+
+	return OutputTransform;
+}
 
-FTransform FSequenceSampler::ExtrapolateRootMotion(float SampleStart, float SampleEnd, float ExtrapolationTime) const
+FTransform FBlendSpaceSampler::ExtractBlendSpaceRootMotionFromRange(float StartTrackPosition, float EndTrackPosition) const
 {
-	const float SampleDelta = SampleEnd - SampleStart;
-	check(!FMath::IsNearlyZero(SampleDelta));
+	checkf(AccumulatedRootTransform.Num() > 0, TEXT("ProcessRootTransform must be run first"));
 
-	FTransform SampleToExtrapolate = Input.Sequence->ExtractRootMotionFromRange(SampleStart, SampleEnd);
+	FTransform RootTransformRefPose = ExtractBlendSpaceRootTrackTransform(0.0f);
 
-	const FVector LinearVelocityToExtrapolate = SampleToExtrapolate.GetTranslation() / SampleDelta;
-	const float LinearSpeedToExtrapolate = LinearVelocityToExtrapolate.Size();
-	const bool bCanExtrapolateTranslation = 
-		LinearSpeedToExtrapolate >= Input.ExtrapolationParameters.LinearSpeedThreshold;
+	FTransform StartTransform = ExtractBlendSpaceRootTrackTransform(StartTrackPosition);
+	FTransform EndTransform = ExtractBlendSpaceRootTrackTransform(EndTrackPosition);
 
-	const float AngularSpeedToExtrapolateRad = SampleToExtrapolate.GetRotation().GetAngle() / SampleDelta;
-	const bool bCanExtrapolateRotation = 
-		FMath::RadiansToDegrees(AngularSpeedToExtrapolateRad) >= Input.ExtrapolationParameters.AngularSpeedThreshold;
+	// Transform to Component Space
+	const FTransform RootToComponent = RootTransformRefPose.Inverse();
+	StartTransform = RootToComponent * StartTransform;
+	EndTransform = RootToComponent * EndTransform;
 
-	if (!bCanExtrapolateTranslation && !bCanExtrapolateRotation)
-	{
-		return FTransform::Identity;
-	}
+	return EndTransform.GetRelativeTransform(StartTransform);
+}
 
-	if (!bCanExtrapolateTranslation)
-	{
-		SampleToExtrapolate.SetTranslation(FVector::ZeroVector);
-	}
+FTransform FBlendSpaceSampler::ExtractBlendSpaceRootMotion(float StartTime, float DeltaTime, bool bAllowLooping) const
+{
+	FRootMotionMovementParams RootMotionParams;
 
-	if (!bCanExtrapolateRotation)
+	if (DeltaTime != 0.f)
 	{
-		SampleToExtrapolate.SetRotation(FQuat::Identity);
-	}
+		bool const bPlayingBackwards = (DeltaTime < 0.f);
 
-	// converting ExtrapolationTime to a positive number to avoid dealing with the negative extrapolation and inverting
-	// transforms later on.
-	const float AbsExtrapolationTime = FMath::Abs(ExtrapolationTime);
-	const float AbsSampleDelta = FMath::Abs(SampleDelta);
-	const FTransform AbsTimeSampleToExtrapolate = 
-		ExtrapolationTime >= 0.0f ? SampleToExtrapolate : SampleToExtrapolate.Inverse();
+		float PreviousPosition = StartTime;
+		float CurrentPosition = StartTime;
+		float DesiredDeltaMove = DeltaTime;
 
-	// because we're extrapolating rotation, the extrapolation must be integrated over time
-	const float SampleMultiplier = AbsExtrapolationTime / AbsSampleDelta;
-	float IntegralNumSamples;
-	float RemainingSampleFraction = FMath::Modf(SampleMultiplier, &IntegralNumSamples);
-	int32 NumSamples = (int32)IntegralNumSamples;
+		do
+		{
+			// Disable looping here. Advance to desired position, or beginning / end of animation 
+			const ETypeAdvanceAnim AdvanceType = FAnimationRuntime::AdvanceTime(false, DesiredDeltaMove, CurrentPosition, PlayLength);
 
-	// adding full samples to the extrapolated root motion
-	FTransform ExtrapolatedRootMotion = FTransform::Identity;
-	for (int i = 0; i < NumSamples; ++i)
-	{
-		ExtrapolatedRootMotion = AbsTimeSampleToExtrapolate * ExtrapolatedRootMotion;
+			// Verify position assumptions
+			//ensureMsgf(bPlayingBackwards ? (CurrentPosition <= PreviousPosition) : (CurrentPosition >= PreviousPosition), TEXT("in Animation %s(Skeleton %s) : bPlayingBackwards(%d), PreviousPosition(%0.2f), Current Position(%0.2f)"),
+			//	*GetName(), *GetNameSafe(GetSkeleton()), bPlayingBackwards, PreviousPosition, CurrentPosition);
+
+			RootMotionParams.Accumulate(ExtractBlendSpaceRootMotionFromRange(PreviousPosition, CurrentPosition));
+
+			// If we've hit the end of the animation, and we're allowed to loop, keep going.
+			if ((AdvanceType == ETAA_Finished) && bAllowLooping)
+			{
+				const float ActualDeltaMove = (CurrentPosition - PreviousPosition);
+				DesiredDeltaMove -= ActualDeltaMove;
+
+				PreviousPosition = bPlayingBackwards ? PlayLength : 0.f;
+				CurrentPosition = PreviousPosition;
+			}
+			else
+			{
+				break;
+			}
+		} while (true);
 	}
-	
-	// and a blend with identity for whatever is left
-	FTransform RemainingExtrapolatedRootMotion;
-	RemainingExtrapolatedRootMotion.Blend(
-		FTransform::Identity, 
-		AbsTimeSampleToExtrapolate, 
-		RemainingSampleFraction);
 
-	ExtrapolatedRootMotion = RemainingExtrapolatedRootMotion * ExtrapolatedRootMotion;
-	return ExtrapolatedRootMotion;
+	return RootMotionParams.GetRootMotionTransform();
 }
 
-float FSequenceSampler::ExtrapolateRootDistance(int32 NextRootDistanceIndex, float ExtrapolationTime) const
+void FBlendSpaceSampler::ProcessPlayLength()
 {
-	check(NextRootDistanceIndex > 0 && NextRootDistanceIndex < Output.AccumulatedRootDistance.Num());
+	TArray<FBlendSampleData> BlendSamples;
+	int32 TriangulationIndex = 0;
+	Input.BlendSpace->GetSamplesFromBlendInput(Input.BlendParameters, BlendSamples, TriangulationIndex, true);
 
-	const float DistanceDelta = 
-		Output.AccumulatedRootDistance[NextRootDistanceIndex] - 
-		Output.AccumulatedRootDistance[NextRootDistanceIndex - 1];
-	const float Speed = DistanceDelta * Input.DistanceSamplingRate;
-	const float ExtrapolationSpeed = Speed >= Input.ExtrapolationParameters.LinearSpeedThreshold ?
-		Speed : 0.0f;
-	const float ExtrapolatedDistance = ExtrapolationSpeed * ExtrapolationTime;
+	PlayLength = Input.BlendSpace->GetAnimationLengthFromSampleData(BlendSamples);
 
-	return ExtrapolatedDistance;
+	checkf(PlayLength > 0.0f, TEXT("Blendspace has zero play length"));
 }
 
+void FBlendSpaceSampler::ProcessRootTransform()
+{
+	// Pre-compute root motion
+
+	int32 NumRootSamples = FMath::Max(PlayLength * Input.RootTransformSamplingRate + 1, 1);
+	AccumulatedRootTransform.SetNumUninitialized(NumRootSamples);
+
+	TArray<FBlendSampleData> BlendSamples;
+	int32 TriangulationIndex = 0;
+	Input.BlendSpace->GetSamplesFromBlendInput(Input.BlendParameters, BlendSamples, TriangulationIndex, true);
+
+	FTransform RootMotionAccumulation = FTransform::Identity;
+
+	AccumulatedRootTransform[0] = RootMotionAccumulation;
+
+	for (int32 SampleIdx = 1; SampleIdx < NumRootSamples; ++SampleIdx)
+	{
+		float PreviousTime = float(SampleIdx - 1) / Input.RootTransformSamplingRate;
+		float CurrentTime = float(SampleIdx - 0) / Input.RootTransformSamplingRate;
+
+		FDeltaTimeRecord DeltaTimeRecord;
+		DeltaTimeRecord.Set(PreviousTime, CurrentTime - PreviousTime);
+		FAnimExtractContext ExtractionCtx(CurrentTime, true, DeltaTimeRecord, Input.bLoopable);
+
+		for (int32 BlendSampleIdex = 0; BlendSampleIdex < BlendSamples.Num(); BlendSampleIdex++)
+		{
+			float Scale = BlendSamples[BlendSampleIdex].Animation->GetPlayLength() / PlayLength;
+
+			FDeltaTimeRecord BlendSampleDeltaTimeRecord;
+			BlendSampleDeltaTimeRecord.Set(DeltaTimeRecord.GetPrevious() * Scale, DeltaTimeRecord.Delta * Scale);
+
+			BlendSamples[BlendSampleIdex].DeltaTimeRecord = BlendSampleDeltaTimeRecord;
+			BlendSamples[BlendSampleIdex].PreviousTime = PreviousTime * Scale;
+			BlendSamples[BlendSampleIdex].Time = CurrentTime * Scale;
+		}
+
+		FCompactPose Pose;
+		FBlendedCurve BlendedCurve;
+		Anim::FStackAttributeContainer StackAttributeContainer;
+		FAnimationPoseData AnimPoseData(Pose, BlendedCurve, StackAttributeContainer);
+
+		Pose.SetBoneContainer(&Input.SamplingContext->BoneContainer);
+		BlendedCurve.InitFrom(Input.SamplingContext->BoneContainer);
+
+		Input.BlendSpace->GetAnimationPose(BlendSamples, ExtractionCtx, AnimPoseData);
+
+		const Anim::IAnimRootMotionProvider* RootMotionProvider = Anim::IAnimRootMotionProvider::Get();
+
+		if (ensureMsgf(RootMotionProvider, TEXT("Could not get Root Motion Provider.")))
+		{
+			if (ensureMsgf(RootMotionProvider->HasRootMotion(StackAttributeContainer), TEXT("Blend Space had no Root Motion Attribute.")))
+			{
+				FTransform RootMotionDelta;
+				RootMotionProvider->ExtractRootMotion(StackAttributeContainer, RootMotionDelta);
+
+				RootMotionAccumulation = RootMotionDelta * RootMotionAccumulation;
+			}
+		}
+
+		AccumulatedRootTransform[SampleIdx] = RootMotionAccumulation;
+	}
+}
 
-void FSequenceSampler::ProcessRootMotion()
+void FBlendSpaceSampler::ProcessRootDistance()
 {
+	checkf(AccumulatedRootTransform.Num() > 0, TEXT("ProcessRootTransform must be run first"));
+
 	// Note the distance sampling interval is independent of the schema's sampling interval
-	const float DistanceSamplingInterval = 1.0f / Input.DistanceSamplingRate;
+	const float DistanceSamplingInterval = 1.0f / Input.RootDistanceSamplingRate;
 
-	const FTransform InitialRootTransform = Input.Sequence->ExtractRootTrackTransform(0.0f, nullptr);
+	const FTransform InitialRootTransform = FTransform::Identity;
+
+	uint32 NumDistanceSamples = FMath::CeilToInt(PlayLength * Input.RootDistanceSamplingRate) + 1;
+	AccumulatedRootDistance.Reserve(NumDistanceSamples);
 
 	// Build a distance lookup table by sampling root motion at a fixed rate and accumulating
 	// absolute translation deltas. During indexing we'll bsearch this table and interpolate
 	// between samples in order to convert distance offsets to time offsets.
 	// See also FSequenceIndexer::AddTrajectoryDistanceFeatures().
-	double AccumulatedRootDistance = 0.0;
+	double TotalAccumulatedRootDistance = 0.0;
 	FTransform LastRootTransform = InitialRootTransform;
 	float SampleTime = 0.0f;
-	for (int32 SampleIdx = 0; SampleIdx != Output.NumDistanceSamples; ++SampleIdx)
+	for (int32 SampleIdx = 0; SampleIdx != NumDistanceSamples; ++SampleIdx)
 	{
-		SampleTime = FMath::Min(SampleIdx * DistanceSamplingInterval, Output.PlayLength);
+		SampleTime = FMath::Min(SampleIdx * DistanceSamplingInterval, PlayLength);
 
-		FTransform RootTransform = Input.Sequence->ExtractRootTrackTransform(SampleTime, nullptr);
+		FTransform RootTransform = ExtractBlendSpaceRootTrackTransform(SampleTime);
 		FTransform LocalRootMotion = RootTransform.GetRelativeTransform(LastRootTransform);
 		LastRootTransform = RootTransform;
 
-		AccumulatedRootDistance += LocalRootMotion.GetTranslation().Size();
-		Output.AccumulatedRootDistance.Add((float)AccumulatedRootDistance);
+		TotalAccumulatedRootDistance += LocalRootMotion.GetTranslation().Size();
+		AccumulatedRootDistance.Add((float)TotalAccumulatedRootDistance);
 	}
 
 	// Verify we sampled the final frame of the clip
-	check(SampleTime == Input.Sequence->GetPlayLength());
+	check(SampleTime == PlayLength);
 
 	// Also emit root motion summary info to help with sample wrapping in 
 	// FSequenceIndexer::GetSampleTimeFromDistance() and FSequenceIndexer::GetSampleInfo()
-	Output.TotalRootMotion = LastRootTransform.GetRelativeTransform(InitialRootTransform);
-	Output.TotalRootDistance = Output.AccumulatedRootDistance.Last();
+	TotalRootTransform = LastRootTransform.GetRelativeTransform(InitialRootTransform);
+	TotalRootDistance = AccumulatedRootDistance.Last();
 }
 
 //////////////////////////////////////////////////////////////////////////
-// FSequenceIndexer helpers
+// FIndexer helpers
 
 struct FSamplingParam
 {
@@ -2518,9 +3256,9 @@ static FSamplingParam WrapOrClampSamplingParam(bool bCanWrap, float SamplingPara
 
 
 //////////////////////////////////////////////////////////////////////////
-// FSequenceIndexer
+// FIndexer
 
-class FSequenceIndexer
+class FIndexer
 {
 public:
 
@@ -2528,9 +3266,9 @@ public:
 	{
 		const FAnimSamplingContext* SamplingContext = nullptr;
 		const UPoseSearchSchema* Schema = nullptr;
-		const FSequenceSampler* MainSequence = nullptr;
-		const FSequenceSampler* LeadInSequence = nullptr;
-		const FSequenceSampler* FollowUpSequence = nullptr;
+		const IAssetSampler* MainSampler = nullptr;
+		const IAssetSampler* LeadInSampler = nullptr;
+		const IAssetSampler* FollowUpSampler = nullptr;
 		bool bMirrored = false;
 		FFloatInterval RequestedSamplingRange = FFloatInterval(0.0f, 0.0f);
 		FPoseSearchBlockTransitionParameters BlockTransitionParameters;
@@ -2555,7 +3293,7 @@ private:
 
 	struct FSampleInfo
 	{
-		const FSequenceSampler* Clip = nullptr;
+		const IAssetSampler* Clip = nullptr;
 		FTransform RootTransform;
 		float ClipTime = 0.0f;
 		float RootDistance = 0.0f;
@@ -2577,7 +3315,7 @@ private:
 	void AddMetadata(int32 SampleIdx);
 };
 
-void FSequenceIndexer::Reset()
+void FIndexer::Reset()
 {
 	Output.FirstIndexedSample = 0;
 	Output.LastIndexedSample = 0;
@@ -2587,17 +3325,17 @@ void FSequenceIndexer::Reset()
 	Output.PoseMetadata.Reset(0);
 }
 
-void FSequenceIndexer::Reserve()
+void FIndexer::Reserve()
 {
 	Output.FeatureVectorTable.SetNumZeroed(Input.Schema->Layout.NumFloats * Output.NumIndexedPoses);
 	Output.PoseMetadata.SetNum(Output.NumIndexedPoses);
 }
 
-void FSequenceIndexer::Init(const FInput& InSettings)
+void FIndexer::Init(const FInput& InSettings)
 {
 	check(InSettings.Schema);
 	check(InSettings.Schema->IsValid());
-	check(InSettings.MainSequence);
+	check(InSettings.MainSampler);
 
 	Input = InSettings;
 
@@ -2609,7 +3347,7 @@ void FSequenceIndexer::Init(const FInput& InSettings)
 	Reserve();
 }
 
-void FSequenceIndexer::Process()
+void FIndexer::Process()
 {
 	for (int32 SampleIdx = Output.FirstIndexedSample; SampleIdx <= Output.LastIndexedSample; ++SampleIdx)
 	{
@@ -2626,12 +3364,12 @@ void FSequenceIndexer::Process()
 	}
 }
 
-void FSequenceIndexer::SampleBegin(int32 SampleIdx)
+void FIndexer::SampleBegin(int32 SampleIdx)
 {
 	FeatureVector.Init(Input.Schema);
 }
 
-void FSequenceIndexer::SampleEnd(int32 SampleIdx)
+void FIndexer::SampleEnd(int32 SampleIdx)
 {
 	check(FeatureVector.IsComplete());
 
@@ -2648,79 +3386,51 @@ void FSequenceIndexer::SampleEnd(int32 SampleIdx)
 	Output.PoseMetadata[PoseIdx] = Metadata;
 }
 
-const float FSequenceIndexer::GetSampleTimeFromDistance(float SampleDistance) const
+const float FIndexer::GetSampleTimeFromDistance(float SampleDistance) const
 {
-	auto CanWrapDistanceSamples = [](const FSequenceSampler* Sampler) -> bool
+	auto CanWrapDistanceSamples = [](const IAssetSampler* Sampler) -> bool
 	{
 		constexpr float SMALL_ROOT_DISTANCE = 1.0f;
-		return Sampler->Input.bLoopable && Sampler->Output.TotalRootDistance > SMALL_ROOT_DISTANCE;
+		return Sampler->IsLoopable() && Sampler->GetTotalRootDistance() > SMALL_ROOT_DISTANCE;
 	};
 
-	auto ClipTimeFromDistance = [](const FSequenceSampler* Sampler, float ClipDistance) -> float
-	{
-		int32 NextSampleIdx = 1;
-		int32 PrevSampleIdx = 0;
-		if (ClipDistance > 0.0f)
-		{
-			// Search for the distance value. Because the values will be extrapolated if necessary
-			// LowerBound might go past the end of the array, in which case the last valid index is used
-			int32 ClipDistanceLowerBoundIndex = Algo::LowerBound(Sampler->Output.AccumulatedRootDistance, ClipDistance);
-			NextSampleIdx = FMath::Min(
-				ClipDistanceLowerBoundIndex,
-				Sampler->Output.AccumulatedRootDistance.Num() - 1);
-
-			// Compute distance interpolation amount
-			PrevSampleIdx = FMath::Max(0, NextSampleIdx - 1);
-		}
-
-		float NextDistance = Sampler->Output.AccumulatedRootDistance[NextSampleIdx];
-		float PrevDistance = Sampler->Output.AccumulatedRootDistance[PrevSampleIdx];
-		float DistanceSampleAlpha = FMath::GetRangePct(PrevDistance, NextDistance, ClipDistance);
-
-		// Convert to time
-		float ClipTime = (float(NextSampleIdx) - (1.0f - DistanceSampleAlpha)) / Sampler->Input.DistanceSamplingRate;
-		return ClipTime;
-	};
-
-	float MainTotalDistance = Input.MainSequence->Output.TotalRootDistance;
-	bool bMainCanWrap = CanWrapDistanceSamples(Input.MainSequence);
+	float MainTotalDistance = Input.MainSampler->GetTotalRootDistance();
+	bool bMainCanWrap = CanWrapDistanceSamples(Input.MainSampler);
 
 	float SampleTime = MAX_flt;
 
 	if (!bMainCanWrap)
 	{
 		// Use the lead in anim if we would have to clamp to the beginning of the main anim
-		if (Input.LeadInSequence && (SampleDistance < 0.0f))
+		if (Input.LeadInSampler && (SampleDistance < 0.0f))
 		{
-			const FSequenceSampler::FOutput& ClipData = Input.LeadInSequence->Output;
+			const IAssetSampler* ClipSampler = Input.LeadInSampler;
 
-			bool bLeadInCanWrap = CanWrapDistanceSamples(Input.LeadInSequence);
-			float LeadRelativeDistance = SampleDistance + ClipData.TotalRootDistance;
-			FSamplingParam SamplingParam = WrapOrClampSamplingParam(bLeadInCanWrap, ClipData.TotalRootDistance, LeadRelativeDistance);
+			bool bLeadInCanWrap = CanWrapDistanceSamples(Input.LeadInSampler);
+			float LeadRelativeDistance = SampleDistance + ClipSampler->GetTotalRootDistance();
+			FSamplingParam SamplingParam = WrapOrClampSamplingParam(bLeadInCanWrap, ClipSampler->GetTotalRootDistance(), LeadRelativeDistance);
 
-			float ClipTime = ClipTimeFromDistance(
-				Input.LeadInSequence, 
+			float ClipTime = ClipSampler->GetTimeFromRootDistance(
 				SamplingParam.WrappedParam + SamplingParam.Extrapolation);
 
 			// Make the lead in clip time relative to the main sequence again and unwrap
-			SampleTime = -((SamplingParam.NumCycles * ClipData.PlayLength) + (ClipData.PlayLength - ClipTime));
+			SampleTime = -((SamplingParam.NumCycles * ClipSampler->GetPlayLength()) + (ClipSampler->GetPlayLength() - ClipTime));
 		}
 
 		// Use the follow up anim if we would have clamp to the end of the main anim
-		else if (Input.FollowUpSequence && (SampleDistance > MainTotalDistance))
+		else if (Input.FollowUpSampler && (SampleDistance > MainTotalDistance))
 		{
-			const FSequenceSampler::FOutput& ClipData = Input.FollowUpSequence->Output;
+			const IAssetSampler* ClipSampler = Input.FollowUpSampler;
 
-			bool bFollowUpCanWrap = CanWrapDistanceSamples(Input.FollowUpSequence);
+			bool bFollowUpCanWrap = CanWrapDistanceSamples(Input.FollowUpSampler);
 			float FollowRelativeDistance = SampleDistance - MainTotalDistance;
-			FSamplingParam SamplingParam = WrapOrClampSamplingParam(bFollowUpCanWrap, ClipData.TotalRootDistance, FollowRelativeDistance);
+			FSamplingParam SamplingParam = WrapOrClampSamplingParam(bFollowUpCanWrap, ClipSampler->GetTotalRootDistance(), FollowRelativeDistance);
 
-			float ClipTime = ClipTimeFromDistance(
-				Input.FollowUpSequence, 
+			float ClipTime = ClipSampler->GetTimeFromRootDistance(
 				SamplingParam.WrappedParam + SamplingParam.Extrapolation);
 
 			// Make the follow up clip time relative to the main sequence again and unwrap
-			SampleTime = Input.MainSequence->Output.PlayLength + SamplingParam.NumCycles * ClipData.PlayLength + ClipTime;
+			SampleTime = Input.MainSampler->GetPlayLength() + SamplingParam.NumCycles * ClipSampler->GetPlayLength() + ClipTime;
 		}
 	}
 
@@ -2728,19 +3438,16 @@ const float FSequenceIndexer::GetSampleTimeFromDistance(float SampleDistance) co
 	// The main anim sample may have been wrapped or clamped
 	if (SampleTime == MAX_flt)
 	{
-		const FSequenceSampler::FOutput& ClipData = Input.MainSequence->Output;
-
 		float MainRelativeDistance = SampleDistance;
 		if (SampleDistance < 0.0f && bMainCanWrap)
 		{
 			// In this case we're sampling a loop backwards, so MainRelativeDistance must adjust so the number of cycles 
 			// is counted correctly.
-			MainRelativeDistance += ClipData.TotalRootDistance;
+			MainRelativeDistance += Input.MainSampler->GetTotalRootDistance();
 		}
 
 		FSamplingParam SamplingParam = WrapOrClampSamplingParam(bMainCanWrap, MainTotalDistance, MainRelativeDistance);
-		float ClipTime = ClipTimeFromDistance(
-			Input.MainSequence, 
+		float ClipTime = Input.MainSampler->GetTimeFromRootDistance(
 			SamplingParam.WrappedParam + SamplingParam.Extrapolation);
 
 		// Unwrap the main clip time
@@ -2748,11 +3455,11 @@ const float FSequenceIndexer::GetSampleTimeFromDistance(float SampleDistance) co
 		{
 			if (SampleDistance < 0.0f)
 			{
-				SampleTime = -((SamplingParam.NumCycles * ClipData.PlayLength) + (ClipData.PlayLength - ClipTime));
+				SampleTime = -((SamplingParam.NumCycles * Input.MainSampler->GetPlayLength()) + (Input.MainSampler->GetPlayLength() - ClipTime));
 			}
 			else
 			{
-				SampleTime = SamplingParam.NumCycles * ClipData.PlayLength + ClipTime;
+				SampleTime = SamplingParam.NumCycles * Input.MainSampler->GetPlayLength() + ClipTime;
 			}
 		}
 		else
@@ -2764,7 +3471,7 @@ const float FSequenceIndexer::GetSampleTimeFromDistance(float SampleDistance) co
 	return SampleTime;
 }
 
-FSequenceIndexer::FSampleInfo FSequenceIndexer::GetSampleInfo(float SampleTime) const
+FIndexer::FSampleInfo FIndexer::GetSampleInfo(float SampleTime) const
 {
 	FSampleInfo Sample;
 
@@ -2774,33 +3481,33 @@ FSequenceIndexer::FSampleInfo FSequenceIndexer::GetSampleInfo(float SampleTime)
 	float RootDistanceLast = 0.0f;
 	float RootDistanceInitial = 0.0f;
 
-	auto CanWrapTimeSamples = [](const FSequenceSampler* Sampler) -> bool
+	auto CanWrapTimeSamples = [](const IAssetSampler* Sampler) -> bool
 	{
-		return Sampler->Input.bLoopable;
+		return Sampler->IsLoopable();
 	};
 
-	float MainPlayLength = Input.MainSequence->Output.PlayLength;
-	bool bMainCanWrap = CanWrapTimeSamples(Input.MainSequence);
+	float MainPlayLength = Input.MainSampler->GetPlayLength();
+	bool bMainCanWrap = CanWrapTimeSamples(Input.MainSampler);
 
 	FSamplingParam SamplingParam;
 	if (!bMainCanWrap)
 	{
 		// Use the lead in anim if we would have to clamp to the beginning of the main anim
-		if (Input.LeadInSequence && (SampleTime < 0.0f))
+		if (Input.LeadInSampler && (SampleTime < 0.0f))
 		{
-			const FSequenceSampler::FOutput& ClipData = Input.LeadInSequence->Output;
+			const IAssetSampler* ClipSampler = Input.LeadInSampler;
 
-			bool bLeadInCanWrap = CanWrapTimeSamples(Input.LeadInSequence);
-			float LeadRelativeTime = SampleTime + ClipData.PlayLength;
-			SamplingParam = WrapOrClampSamplingParam(bLeadInCanWrap, ClipData.PlayLength, LeadRelativeTime);
+			bool bLeadInCanWrap = CanWrapTimeSamples(Input.LeadInSampler);
+			float LeadRelativeTime = SampleTime + ClipSampler->GetPlayLength();
+			SamplingParam = WrapOrClampSamplingParam(bLeadInCanWrap, ClipSampler->GetPlayLength(), LeadRelativeTime);
 
-			Sample.Clip = Input.LeadInSequence;
+			Sample.Clip = Input.LeadInSampler;
 
 			check(SamplingParam.Extrapolation <= 0.0f);
 			if (SamplingParam.Extrapolation < 0.0f)
 			{
-				RootMotionInitial = Input.LeadInSequence->Output.TotalRootMotion.Inverse();
-				RootDistanceInitial = -Input.LeadInSequence->Output.TotalRootDistance;
+				RootMotionInitial = Input.LeadInSampler->GetTotalRootTransform().Inverse();
+				RootDistanceInitial = -Input.LeadInSampler->GetTotalRootDistance();
 			}
 			else
 			{
@@ -2808,26 +3515,26 @@ FSequenceIndexer::FSampleInfo FSequenceIndexer::GetSampleInfo(float SampleTime)
 				RootDistanceInitial = 0.0f;
 			}
 
-			RootMotionLast = Input.LeadInSequence->Output.TotalRootMotion;
-			RootDistanceLast = Input.LeadInSequence->Output.TotalRootDistance;
+			RootMotionLast = Input.LeadInSampler->GetTotalRootTransform();
+			RootDistanceLast = Input.LeadInSampler->GetTotalRootDistance();
 		}
 
 		// Use the follow up anim if we would have clamp to the end of the main anim
-		else if (Input.FollowUpSequence && (SampleTime > MainPlayLength))
+		else if (Input.FollowUpSampler && (SampleTime > MainPlayLength))
 		{
-			const FSequenceSampler::FOutput& ClipData = Input.FollowUpSequence->Output;
+			const IAssetSampler* ClipSampler = Input.FollowUpSampler;
 
-			bool bFollowUpCanWrap = CanWrapTimeSamples(Input.FollowUpSequence);
+			bool bFollowUpCanWrap = CanWrapTimeSamples(Input.FollowUpSampler);
 			float FollowRelativeTime = SampleTime - MainPlayLength;
-			SamplingParam = WrapOrClampSamplingParam(bFollowUpCanWrap, ClipData.PlayLength, FollowRelativeTime);
+			SamplingParam = WrapOrClampSamplingParam(bFollowUpCanWrap, ClipSampler->GetPlayLength(), FollowRelativeTime);
 
-			Sample.Clip = Input.FollowUpSequence;
+			Sample.Clip = Input.FollowUpSampler;
 
-			RootMotionInitial = Input.MainSequence->Output.TotalRootMotion;
-			RootDistanceInitial = Input.MainSequence->Output.TotalRootDistance;
+			RootMotionInitial = Input.MainSampler->GetTotalRootTransform();
+			RootDistanceInitial = Input.MainSampler->GetTotalRootDistance();
 
-			RootMotionLast = Input.FollowUpSequence->Output.TotalRootMotion;
-			RootDistanceLast = Input.FollowUpSequence->Output.TotalRootDistance;
+			RootMotionLast = Input.FollowUpSampler->GetTotalRootTransform();
+			RootDistanceLast = Input.FollowUpSampler->GetTotalRootDistance();
 		}
 	}
 
@@ -2845,13 +3552,13 @@ FSequenceIndexer::FSampleInfo FSequenceIndexer::GetSampleInfo(float SampleTime)
 
 		SamplingParam = WrapOrClampSamplingParam(bMainCanWrap, MainPlayLength, MainRelativeTime);
 
-		Sample.Clip = Input.MainSequence;
+		Sample.Clip = Input.MainSampler;
 
 		RootMotionInitial = FTransform::Identity;
 		RootDistanceInitial = 0.0f;
 
-		RootMotionLast = Input.MainSequence->Output.TotalRootMotion;
-		RootDistanceLast = Input.MainSequence->Output.TotalRootDistance;
+		RootMotionLast = Input.MainSampler->GetTotalRootTransform();
+		RootDistanceLast = Input.MainSampler->GetTotalRootDistance();
 	}
 
 
@@ -2908,7 +3615,7 @@ FSequenceIndexer::FSampleInfo FSequenceIndexer::GetSampleInfo(float SampleTime)
 	return Sample;
 }
 
-FSequenceIndexer::FSampleInfo FSequenceIndexer::GetSampleInfoRelative(float SampleTime, const FSampleInfo& Origin) const
+FIndexer::FSampleInfo FIndexer::GetSampleInfoRelative(float SampleTime, const FSampleInfo& Origin) const
 {
 	FSampleInfo Sample = GetSampleInfo(SampleTime);
 	Sample.RootTransform.SetToRelativeTransform(Origin.RootTransform);
@@ -2916,12 +3623,12 @@ FSequenceIndexer::FSampleInfo FSequenceIndexer::GetSampleInfoRelative(float Samp
 	return Sample;
 }
 
-FTransform FSequenceIndexer::MirrorTransform(const FTransform& Transform) const
+FTransform FIndexer::MirrorTransform(const FTransform& Transform) const
 {
 	return Input.bMirrored ? Input.SamplingContext->MirrorTransform(Transform) : Transform;
 }
 
-void FSequenceIndexer::AddPoseFeatures(int32 SampleIdx)
+void FIndexer::AddPoseFeatures(int32 SampleIdx)
 {
 	// This function samples the instantaneous pose at time t as well as the pose's velocity and acceleration at time t.
 	// Symmetric finite differences are used to approximate derivatives:
@@ -2946,6 +3653,11 @@ void FSequenceIndexer::AddPoseFeatures(int32 SampleIdx)
 		Pose.SetBoneContainer(&Input.SamplingContext->BoneContainer);
 	}
 
+	for (FBlendedCurve& Curve : UnusedCurves)
+	{
+		Curve.InitFrom(Input.SamplingContext->BoneContainer);
+	}
+
 	FAnimationPoseData AnimPoseData[NumFiniteDiffTerms] = 
 	{
 		{Poses[0], UnusedCurves[0], UnusedAtrributes[0]},
@@ -2953,13 +3665,10 @@ void FSequenceIndexer::AddPoseFeatures(int32 SampleIdx)
 		{Poses[2], UnusedCurves[2], UnusedAtrributes[2]},
 	};
 
-	FAnimExtractContext ExtractionCtx;
-	ExtractionCtx.bExtractRootMotion = true;
-
 	FPoseSearchFeatureDesc Feature;
 	Feature.Domain = EPoseSearchFeatureDomain::Time;
 
-	float SampleTime = FMath::Min(SampleIdx * Input.Schema->SamplingInterval, Input.MainSequence->Output.PlayLength);
+	float SampleTime = FMath::Min(SampleIdx * Input.Schema->SamplingInterval, Input.MainSampler->GetPlayLength());
 	FSampleInfo Origin = GetSampleInfo(SampleTime);
 	
 	for (int32 SchemaSubsampleIdx = 0; SchemaSubsampleIdx != Input.Schema->PoseSampleTimes.Num(); ++SchemaSubsampleIdx)
@@ -2978,8 +3687,26 @@ void FSequenceIndexer::AddPoseFeatures(int32 SampleIdx)
 		// Get pose samples
 		for (int32 Term = 0; Term != NumFiniteDiffTerms; ++Term)
 		{
-			ExtractionCtx.CurrentTime = Samples[Term].ClipTime;
-			Samples[Term].Clip->ExtractPose(ExtractionCtx, Input.bMirrored, AnimPoseData[Term]);
+			float CurrentTime = Samples[Term].ClipTime;
+			float PreviousTime = CurrentTime - Input.SamplingContext->FiniteDelta;
+
+			FDeltaTimeRecord DeltaTimeRecord;
+			DeltaTimeRecord.Set(PreviousTime, CurrentTime - PreviousTime);
+			FAnimExtractContext ExtractionCtx(CurrentTime, true, DeltaTimeRecord, Samples[Term].Clip->IsLoopable());
+
+			Samples[Term].Clip->ExtractPose(ExtractionCtx, AnimPoseData[Term]);
+
+			if (Input.bMirrored)
+			{
+				FAnimationRuntime::MirrorPose(
+					AnimPoseData[Term].GetPose(),
+					Input.Schema->MirrorDataTable->MirrorAxis,
+					Input.SamplingContext->CompactPoseMirrorBones,
+					Input.SamplingContext->ComponentSpaceRefRotations
+				);
+				// Note curves and attributes are not used during the indexing process and therefore don't need to be mirrored
+			}
+
 			ComponentSpacePoses[Term].InitPose(Poses[Term]);
 		}
 
@@ -3009,7 +3736,7 @@ void FSequenceIndexer::AddPoseFeatures(int32 SampleIdx)
 	}
 }
 
-void FSequenceIndexer::AddTrajectoryTimeFeatures(int32 SampleIdx)
+void FIndexer::AddTrajectoryTimeFeatures(int32 SampleIdx)
 {
 	// This function samples the instantaneous trajectory at time t as well as the trajectory's velocity and acceleration at time t.
 	// Symmetric finite differences are used to approximate derivatives:
@@ -3022,7 +3749,7 @@ void FSequenceIndexer::AddTrajectoryTimeFeatures(int32 SampleIdx)
 	Feature.Domain = EPoseSearchFeatureDomain::Time;
 	Feature.SchemaBoneIdx = FPoseSearchFeatureDesc::TrajectoryBoneIndex;
 
-	float SampleTime = FMath::Min(SampleIdx * Input.Schema->SamplingInterval, Input.MainSequence->Output.PlayLength);
+	float SampleTime = FMath::Min(SampleIdx * Input.Schema->SamplingInterval, Input.MainSampler->GetPlayLength());
 	FSampleInfo Origin = GetSampleInfo(SampleTime);
 	
 	for (int32 SchemaSubsampleIdx = 0; SchemaSubsampleIdx != Input.Schema->TrajectorySampleTimes.Num(); ++SchemaSubsampleIdx)
@@ -3051,7 +3778,7 @@ void FSequenceIndexer::AddTrajectoryTimeFeatures(int32 SampleIdx)
 	}
 }
 
-void FSequenceIndexer::AddTrajectoryDistanceFeatures(int32 SampleIdx)
+void FIndexer::AddTrajectoryDistanceFeatures(int32 SampleIdx)
 {
 	// This function is very similar to AddTrajectoryTimeFeatures, but samples are taken in the distance domain
 	// instead of time domain.
@@ -3060,7 +3787,7 @@ void FSequenceIndexer::AddTrajectoryDistanceFeatures(int32 SampleIdx)
 	Feature.Domain = EPoseSearchFeatureDomain::Distance;
 	Feature.SchemaBoneIdx = FPoseSearchFeatureDesc::TrajectoryBoneIndex;
 
-	float SampleTime = FMath::Min(SampleIdx * Input.Schema->SamplingInterval, Input.MainSequence->Output.PlayLength);
+	float SampleTime = FMath::Min(SampleIdx * Input.Schema->SamplingInterval, Input.MainSampler->GetPlayLength());
 	FSampleInfo Origin = GetSampleInfo(SampleTime);
 	
 	FQuat RootReferenceRot = Input.SamplingContext->BoneContainer.GetRefPoseTransform(FCompactPoseBoneIndex(0)).GetRotation();
@@ -3094,15 +3821,15 @@ void FSequenceIndexer::AddTrajectoryDistanceFeatures(int32 SampleIdx)
 	}
 }
 
-void FSequenceIndexer::AddMetadata(int32 SampleIdx)
+void FIndexer::AddMetadata(int32 SampleIdx)
 {
-	const float SequenceLength = Input.MainSequence->Output.PlayLength;
+	const float SequenceLength = Input.MainSampler->GetPlayLength();
 	const float SampleTime = FMath::Min(SampleIdx * Input.Schema->SamplingInterval, SequenceLength);
 
 	Metadata = FPoseSearchPoseMetadata();
 
 	const bool bBlockTransition =
-		!Input.MainSequence->Input.bLoopable &&
+		!Input.MainSampler->IsLoopable() &&
 		(SampleTime < Input.RequestedSamplingRange.Min + Input.BlockTransitionParameters.SequenceStartInterval ||
 		 SampleTime > Input.RequestedSamplingRange.Max - Input.BlockTransitionParameters.SequenceEndInterval);
 
@@ -3112,7 +3839,7 @@ void FSequenceIndexer::AddMetadata(int32 SampleIdx)
 	}
 
 	TArray<UAnimNotifyState_PoseSearchBase*> NotifyStates;
-	Input.MainSequence->ExtractPoseSearchNotifyStates(SampleTime, NotifyStates);
+	Input.MainSampler->ExtractPoseSearchNotifyStates(SampleTime, NotifyStates);
 	for (const UAnimNotifyState_PoseSearchBase* PoseSearchNotify : NotifyStates)
 	{
 		if (PoseSearchNotify->GetClass()->IsChildOf<UAnimNotifyState_PoseSearchBlockTransition>())
@@ -3472,7 +4199,8 @@ static void PreprocessSearchIndexNone(FPoseSearchIndex* SearchIndex)
 	SampleMeanMap = VectorXf::Zero(NumDimensions);
 }
 
-inline Eigen::VectorXd ComputeFeatureMeanDeviations(const Eigen::MatrixXd& CenteredPoseMatrix, const FPoseSearchFeatureVectorLayout& Layout) {
+static inline Eigen::VectorXd ComputeFeatureMeanDeviations(const Eigen::MatrixXd& CenteredPoseMatrix, const FPoseSearchFeatureVectorLayout& Layout)
+{
 	using namespace Eigen;
 
 	const int32 NumPoses = CenteredPoseMatrix.cols();
@@ -3876,18 +4604,16 @@ bool BuildIndex(const UAnimSequence* Sequence, UPoseSearchSequenceMetaData* Sequ
 
 	FSequenceSampler Sampler;
 	FSequenceSampler::FInput SamplerInput;
-	SamplerInput.SamplingContext = &SamplingContext;
-	SamplerInput.Schema = SequenceMetaData->Schema;
 	SamplerInput.ExtrapolationParameters = SequenceMetaData->ExtrapolationParameters;
 	SamplerInput.Sequence = Sequence;
 	SamplerInput.bLoopable = false;
 	Sampler.Init(SamplerInput);
 	Sampler.Process();
 
-	FSequenceIndexer Indexer;
-	FSequenceIndexer::FInput IndexerInput;
+	FIndexer Indexer;
+	FIndexer::FInput IndexerInput;
 	IndexerInput.SamplingContext = &SamplingContext;
-	IndexerInput.MainSequence = &Sampler;
+	IndexerInput.MainSampler = &Sampler;
 	IndexerInput.Schema = SequenceMetaData->Schema;
 	IndexerInput.RequestedSamplingRange = GetEffectiveSamplingRange(Sequence, SequenceMetaData->SamplingRange);
 	Indexer.Init(IndexerInput);
@@ -3928,11 +4654,12 @@ bool BuildIndex(UPoseSearchDatabase* Database)
 	FAnimSamplingContext SamplingContext;
 	SamplingContext.Init(Database->Schema);
 
-	// Prepare animation preprocessing tasks
+	// Prepare samplers for all sequences
+
 	TArray<FSequenceSampler> SequenceSamplers;
 	TMap<const UAnimSequence*, int32> SequenceSamplerMap;
 
-	auto AddSampler = [&](const UAnimSequence* Sequence, bool bLoopable)
+	auto AddSequenceSampler = [&](const UAnimSequence* Sequence, bool bLoopable)
 	{
 		if (!SequenceSamplerMap.Contains(Sequence))
 		{
@@ -3940,8 +4667,6 @@ bool BuildIndex(UPoseSearchDatabase* Database)
 			SequenceSamplerMap.Add(Sequence, SequenceSamplerIdx);
 
 			FSequenceSampler::FInput Input;
-			Input.SamplingContext = &SamplingContext;
-			Input.Schema = Database->Schema;
 			Input.ExtrapolationParameters = Database->ExtrapolationParameters;
 			Input.Sequence = Sequence;
 			Input.bLoopable = bLoopable;
@@ -3953,65 +4678,136 @@ bool BuildIndex(UPoseSearchDatabase* Database)
 	{
 		if (DbSequence.Sequence)
 		{
-			AddSampler(DbSequence.Sequence, DbSequence.bLoopAnimation);
+			AddSequenceSampler(DbSequence.Sequence, DbSequence.bLoopAnimation);
 		}
 
 		if (DbSequence.LeadInSequence)
 		{
-			AddSampler(DbSequence.LeadInSequence, DbSequence.bLoopLeadInAnimation);
+			AddSequenceSampler(DbSequence.LeadInSequence, DbSequence.bLoopLeadInAnimation);
 		}
 
 		if (DbSequence.FollowUpSequence)
 		{
-			AddSampler(DbSequence.FollowUpSequence, DbSequence.bLoopFollowUpAnimation);
+			AddSequenceSampler(DbSequence.FollowUpSequence, DbSequence.bLoopFollowUpAnimation);
 		}
 	}
 
-	// Preprocess animations independently
 	ParallelFor(SequenceSamplers.Num(), [&SequenceSamplers](int32 SamplerIdx){ SequenceSamplers[SamplerIdx].Process(); });
 
+	// Prepare samplers for all blend spaces
+
+	TArray<FBlendSpaceSampler> BlendSpaceSamplers;
+	TMap<TPair<const UBlendSpace*,FVector>, int32> BlendSpaceSamplerMap;
+
+	for (const FPoseSearchDatabaseBlendSpace& DbBlendSpace : Database->BlendSpaces)
+	{
+		if (DbBlendSpace.BlendSpace)
+		{
+			int32 HorizontalBlendNum, VerticalBlendNum;
+			float HorizontalBlendMin, HorizontalBlendMax, VerticalBlendMin, VerticalBlendMax;
+
+			DbBlendSpace.GetBlendSpaceParameterSampleRanges(
+				HorizontalBlendNum,
+				VerticalBlendNum,
+				HorizontalBlendMin,
+				HorizontalBlendMax,
+				VerticalBlendMin,
+				VerticalBlendMax);
+
+			for (int32 HorizontalIndex = 0; HorizontalIndex < HorizontalBlendNum; HorizontalIndex++)
+			{
+				for (int32 VerticalIndex = 0; VerticalIndex < VerticalBlendNum; VerticalIndex++)
+				{
+					FVector BlendParameters = BlendParameterForSampleRanges(
+						HorizontalIndex,
+						VerticalIndex,
+						HorizontalBlendNum,
+						VerticalBlendNum,
+						HorizontalBlendMin,
+						HorizontalBlendMax,
+						VerticalBlendMin,
+						VerticalBlendMax);
+
+					if (!BlendSpaceSamplerMap.Contains({ DbBlendSpace.BlendSpace, BlendParameters }))
+					{
+						int32 BlendSpaceSamplerIdx = BlendSpaceSamplers.AddDefaulted();
+						BlendSpaceSamplerMap.Add({ DbBlendSpace.BlendSpace, BlendParameters }, BlendSpaceSamplerIdx);
+
+						FBlendSpaceSampler::FInput Input;
+						Input.SamplingContext = &SamplingContext;
+						Input.ExtrapolationParameters = Database->ExtrapolationParameters;
+						Input.BlendSpace = DbBlendSpace.BlendSpace;
+						Input.bLoopable = DbBlendSpace.bLoopAnimation;
+						Input.BlendParameters = BlendParameters;
+
+						BlendSpaceSamplers[BlendSpaceSamplerIdx].Init(Input);
+					}
+				}
+			}
+		}
+	}
+
+	ParallelFor(BlendSpaceSamplers.Num(), [&BlendSpaceSamplers](int32 SamplerIdx) { BlendSpaceSamplers[SamplerIdx].Process(); });
 
-	auto GetSampler = [&](const UAnimSequence* Sequence) -> const FSequenceSampler*
+	// Prepare sequences for indexing
+
+	TArray<FIndexer> Indexers;
+	Indexers.Reserve(Database->SearchIndex.Assets.Num());
+
+	auto GetSequenceSampler = [&](const UAnimSequence* Sequence) -> const FSequenceSampler*
 	{
 		return Sequence ? &SequenceSamplers[SequenceSamplerMap[Sequence]] : nullptr;
 	};
 
-	// Prepare animation indexing tasks
-	TArray<FSequenceIndexer> Indexers;
-	Indexers.Reserve(Database->SearchIndex.Assets.Num());
+	auto GetBlendSpaceSampler = [&](const UBlendSpace* BlendSpace, const FVector BlendParameters) -> const FBlendSpaceSampler*
+	{
+		return BlendSpace ? &BlendSpaceSamplers[BlendSpaceSamplerMap[{BlendSpace, BlendParameters}]] : nullptr;
+	};
+
 	for (int32 AssetIdx = 0; AssetIdx != Database->SearchIndex.Assets.Num(); ++AssetIdx)
 	{
 		const FPoseSearchIndexAsset& SearchIndexAsset = Database->SearchIndex.Assets[AssetIdx];
-		const FPoseSearchDatabaseSequence& DbSequence = Database->Sequences[SearchIndexAsset.SourceAssetIdx];
-		const float SequenceLength = DbSequence.Sequence->GetPlayLength();
 
-		FSequenceIndexer::FInput Input;
+		FIndexer::FInput Input;
 		Input.SamplingContext = &SamplingContext;
-		Input.MainSequence = GetSampler(DbSequence.Sequence);
-		Input.LeadInSequence = 
-			SearchIndexAsset.SamplingInterval.Min == 0.0f ? GetSampler(DbSequence.LeadInSequence) : nullptr;
-		Input.FollowUpSequence = 
-			SearchIndexAsset.SamplingInterval.Max == SequenceLength ? GetSampler(DbSequence.FollowUpSequence) : nullptr;
+
 		Input.Schema = Database->Schema;
 		Input.BlockTransitionParameters = Database->BlockTransitionParameters;
 		Input.RequestedSamplingRange = SearchIndexAsset.SamplingInterval;
 		Input.bMirrored = SearchIndexAsset.bMirrored;
 
-		FSequenceIndexer& Indexer = Indexers.AddDefaulted_GetRef();
+		if (SearchIndexAsset.Type == ESearchIndexAssetType::Sequence)
+		{
+			const FPoseSearchDatabaseSequence& DbSequence = Database->GetSequenceSourceAsset(&SearchIndexAsset);
+			const float SequenceLength = DbSequence.Sequence->GetPlayLength();
+			Input.MainSampler = GetSequenceSampler(DbSequence.Sequence);
+			Input.LeadInSampler =
+				SearchIndexAsset.SamplingInterval.Min == 0.0f ? GetSequenceSampler(DbSequence.LeadInSequence) : nullptr;
+			Input.FollowUpSampler =
+				SearchIndexAsset.SamplingInterval.Max == SequenceLength ? GetSequenceSampler(DbSequence.FollowUpSequence) : nullptr;
+		}
+		else if (SearchIndexAsset.Type == ESearchIndexAssetType::BlendSpace)
+		{
+			const FPoseSearchDatabaseBlendSpace& DbBlendSpace = Database->GetBlendSpaceSourceAsset(&SearchIndexAsset);
+			Input.MainSampler = GetBlendSpaceSampler(DbBlendSpace.BlendSpace, SearchIndexAsset.BlendParameters);
+		}
+		else
+		{
+			checkNoEntry();
+		}
+
+		FIndexer& Indexer = Indexers.AddDefaulted_GetRef();
 		Indexer.Init(Input);
 	}
 
-	// Index animations independently
 	ParallelFor(Indexers.Num(), [&Indexers](int32 SequenceIdx){ Indexers[SequenceIdx].Process(); });
 
-
-
 	// Write index info to sequence and count up total poses and storage required
 	int32 TotalPoses = 0;
 	int32 TotalFloats = 0;
 	for (int32 AssetIdx = 0; AssetIdx != Database->SearchIndex.Assets.Num(); ++AssetIdx)
 	{
-		const FSequenceIndexer::FOutput& Output = Indexers[AssetIdx].Output;
+		const FIndexer::FOutput& Output = Indexers[AssetIdx].Output;
 
 		FPoseSearchIndexAsset& SearchIndexAsset = Database->SearchIndex.Assets[AssetIdx];
 		SearchIndexAsset.NumPoses = Output.NumIndexedPoses;
@@ -4024,9 +4820,10 @@ bool BuildIndex(UPoseSearchDatabase* Database)
 	// Join animation data into a single search index
 	Database->SearchIndex.Values.Reset(TotalFloats);
 	Database->SearchIndex.PoseMetadata.Reset(TotalPoses);
-	for (const FSequenceIndexer& Indexer : Indexers)
+
+	for (const FIndexer& Indexer : Indexers)
 	{
-		const FSequenceIndexer::FOutput& Output = Indexer.Output;
+		const FIndexer::FOutput& Output = Indexer.Output;
 		Database->SearchIndex.Values.Append(Output.FeatureVectorTable.GetData(), Output.FeatureVectorTable.Num());
 		Database->SearchIndex.PoseMetadata.Append(Output.PoseMetadata);
 	}
@@ -4069,8 +4866,7 @@ FSearchResult Search(FSearchContext& SearchContext)
 	{
 		if (Database && SearchContext.DatabaseTagQuery)
 		{
-			const FPoseSearchDatabaseSequence& DbSequence = Database->Sequences[Asset.SourceAssetIdx];
-			if (!SearchContext.DatabaseTagQuery->Matches(DbSequence.GroupTags))
+			if (!SearchContext.DatabaseTagQuery->Matches(*Database->GetSourceAssetGroupTags(&Asset)))
 			{
 				continue;
 			}
@@ -4101,7 +4897,7 @@ FSearchResult Search(FSearchContext& SearchContext)
 	Result.PoseCost = BestPoseCost;
 	Result.PoseIdx = BestPoseIdx;
 	Result.SearchIndexAsset = SearchIndex->FindAssetForPose(BestPoseIdx);
-	Result.TimeOffsetSeconds = SearchIndex->GetTimeOffset(BestPoseIdx, Result.SearchIndexAsset);
+	Result.AssetTime = SearchIndex->GetAssetTime(BestPoseIdx, Result.SearchIndexAsset);
 
 	SearchContext.DebugDrawParams.PoseVector = SearchContext.QueryValues;
 	SearchContext.DebugDrawParams.PoseIdx = Result.PoseIdx;
@@ -4110,7 +4906,7 @@ FSearchResult Search(FSearchContext& SearchContext)
 	return Result;
 }
 
-void ComputePoseCostAddends(
+static void ComputePoseCostAddends(
 	int32 PoseIdx, 
 	FSearchContext& SearchContext, 
 	float& OutNotifyAddend, 
@@ -4319,7 +5115,7 @@ UE::Anim::IPoseSearchProvider::FSearchResult FModule::Search(const FAnimationBas
 
 	ProviderResult.Dissimilarity = Result.PoseCost.TotalCost;
 	ProviderResult.PoseIdx = Result.PoseIdx;
-	ProviderResult.TimeOffsetSeconds = Result.TimeOffsetSeconds;
+	ProviderResult.TimeOffsetSeconds = Result.AssetTime;
 	return ProviderResult;
 }
 
diff --git a/Engine/Plugins/Experimental/Animation/PoseSearch/Source/Runtime/Private/PoseSearchLibrary.cpp b/Engine/Plugins/Experimental/Animation/PoseSearch/Source/Runtime/Private/PoseSearchLibrary.cpp
--- a/Engine/Plugins/Experimental/Animation/PoseSearch/Source/Runtime/Private/PoseSearchLibrary.cpp
+++ b/Engine/Plugins/Experimental/Animation/PoseSearch/Source/Runtime/Private/PoseSearchLibrary.cpp
@@ -5,62 +5,15 @@
 #include "Animation/AnimSequence.h"
 #include "Animation/AnimNode_Inertialization.h"
 #include "Animation/AnimNode_SequencePlayer.h"
+#include "Animation/AnimRootMotionProvider.h"
 #include "PoseSearch/AnimNode_MotionMatching.h"
 #include "Trace/PoseSearchTraceLogger.h"
 
 #define LOCTEXT_NAMESPACE "PoseSearchLibrary"
 
-DEFINE_LOG_CATEGORY_STATIC(LogPoseSearchLibrary, Verbose, All);
-
 namespace UE::PoseSearch {
 
-void ComputeCurrentVelocityMagnitudes(
-	const FMotionMatchingState& InOutMotionMatchingState, 
-	const FTrajectorySampleRange& Trajectory, 
-	const float DeltaTime, 
-	float& SimLinearVelocity, float& SimAngularVelocity, 
-	float& AnimLinearVelocity, float& AnimAngularVelocity)
-{
-	if (DeltaTime > SMALL_NUMBER)
-	{
-		// simulation
-
-		int32 FirstIdx = 0;
-		const FTrajectorySample PrevSample = FTrajectorySampleRange::IterSampleTrajectory(
-			Trajectory.Samples,
-			ETrajectorySampleDomain::Time,
-			-DeltaTime, FirstIdx);
-
-		const FTrajectorySample CurrSample = FTrajectorySampleRange::IterSampleTrajectory(
-			Trajectory.Samples,
-			ETrajectorySampleDomain::Time,
-			0.0f, FirstIdx);
-
-		const FTransform RelativeMotion = CurrSample.Transform.GetRelativeTransform(PrevSample.Transform);
-
-		SimLinearVelocity = RelativeMotion.GetTranslation().Size() / DeltaTime;
-		SimAngularVelocity = FMath::RadiansToDegrees(RelativeMotion.GetRotation().GetAngle()) / DeltaTime;
-
-		// animation
-		if (const FPoseSearchIndexAsset* IndexAsset = InOutMotionMatchingState.GetCurrentSearchIndexAsset())
-		{
-			const FPoseSearchDatabaseSequence& DbSequence =
-				InOutMotionMatchingState.CurrentDatabase->GetSourceAsset(IndexAsset);
-			const UAnimSequence* SourceSequence = DbSequence.Sequence;
-
-			const FTransform AnimRootMotion = SourceSequence->ExtractRootMotion(
-				InOutMotionMatchingState.AssetPlayerTime,
-				-DeltaTime,
-				DbSequence.bLoopAnimation);
-
-			AnimLinearVelocity = AnimRootMotion.GetTranslation().Size() / DeltaTime;
-			AnimAngularVelocity = FMath::RadiansToDegrees(AnimRootMotion.GetRotation().GetAngle()) / DeltaTime;
-		}
-
-	}
-}
-
-void ComputeDatabaseSequenceFilter(
+static void ComputeDatabaseSequenceFilter(
 	const UPoseSearchDatabase* Database, 
 	const FGameplayTagQuery* Query, 
 	TArray<bool>& OutDbSequenceFilter)
@@ -83,6 +36,29 @@ void ComputeDatabaseSequenceFilter(
 	}
 }
 
+static void ComputeDatabaseBlendSpaceFilter(
+	const UPoseSearchDatabase* Database,
+	const FGameplayTagQuery* Query,
+	TArray<bool>& OutDbBlendSpaceFilter)
+{
+	OutDbBlendSpaceFilter.SetNum(Database->BlendSpaces.Num());
+
+	if (Query)
+	{
+		for (int BlendSpaceIdx = 0; BlendSpaceIdx < Database->BlendSpaces.Num(); ++BlendSpaceIdx)
+		{
+			OutDbBlendSpaceFilter[BlendSpaceIdx] = Query->Matches(Database->BlendSpaces[BlendSpaceIdx].GroupTags);
+		}
+	}
+	else
+	{
+		for (int BlendSpaceIdx = 0; BlendSpaceIdx < Database->BlendSpaces.Num(); ++BlendSpaceIdx)
+		{
+			OutDbBlendSpaceFilter[BlendSpaceIdx] = true;
+		}
+	}
+}
+
 //////////////////////////////////////////////////////////////////////////
 // FMotionMatchingPoseStepper
 
@@ -95,67 +71,110 @@ void FMotionMatchingPoseStepper::Update(const FAnimationUpdateContext& UpdateCon
 
 	check(State.SearchIndexAssetIdx != INDEX_NONE);
 
+	Result.PoseIdx = INDEX_NONE;
+
 	const FPoseSearchIndexAsset* SearchIndexAsset = State.GetCurrentSearchIndexAsset();
-	const FPoseSearchDatabaseSequence& DbSequence = State.CurrentDatabase->GetSourceAsset(SearchIndexAsset);
-	const float AssetLength = DbSequence.Sequence->GetPlayLength();
-	const float DeltaTime = UpdateContext.GetDeltaTime();
-
-	float SteppedAssetTime = State.AssetPlayerTime;
-	ETypeAdvanceAnim AdvanceType = FAnimationRuntime::AdvanceTime(
-		DbSequence.bLoopAnimation,
-		DeltaTime,
-		SteppedAssetTime,
-		AssetLength);
-
-	if (AdvanceType != ETAA_Finished)
+
+	if (SearchIndexAsset->Type == ESearchIndexAssetType::Sequence)
 	{
-		// if the sampling interval doesn't contain the new time, there's no pose to return.
-		if (SearchIndexAsset->SamplingInterval.Contains(SteppedAssetTime))
+		const FPoseSearchDatabaseSequence& DbSequence = State.CurrentDatabase->GetSequenceSourceAsset(SearchIndexAsset);
+		const float AssetLength = DbSequence.Sequence->GetPlayLength();
+		const float DeltaTime = UpdateContext.GetDeltaTime();
+
+		float SteppedTime = State.AssetPlayerTime;
+		ETypeAdvanceAnim AdvanceType = FAnimationRuntime::AdvanceTime(
+			DbSequence.bLoopAnimation,
+			DeltaTime,
+			SteppedTime,
+			AssetLength);
+
+		if (AdvanceType != ETAA_Finished)
+		{
+			// if the sampling interval doesn't contain the new time, there's no pose to return.
+			if (SearchIndexAsset->SamplingInterval.Contains(SteppedTime))
+			{
+				// we can continue ticking the same sequence forward
+				Result.PoseIdx = State.CurrentDatabase->GetPoseIndexFromTime(State.AssetPlayerTime, SearchIndexAsset);
+				Result.SearchIndexAsset = SearchIndexAsset;
+				Result.AssetTime = State.AssetPlayerTime;
+			}
+		}
+		else
 		{
-			// we can continue ticking the same sequence forward
-			Result.PoseIdx = State.CurrentDatabase->GetPoseIndexFromAssetTime(SteppedAssetTime, SearchIndexAsset);
-			Result.SearchIndexAsset = SearchIndexAsset;
-			Result.TimeOffsetSeconds = State.AssetPlayerTime;
+			// check if there's a follow-up that can be used
+			int32 FollowUpDbSequenceIdx = State.CurrentDatabase->Sequences.IndexOfByPredicate(
+				[&](const FPoseSearchDatabaseSequence& Entry)
+				{
+					return Entry.Sequence == DbSequence.FollowUpSequence;
+				});
+
+			int32 FollowUpSearchIndexAssetIdx = State.CurrentDatabase->SearchIndex.Assets.IndexOfByPredicate(
+				[&](const FPoseSearchIndexAsset& Entry)
+				{
+					const bool bIsMatch =
+						Entry.SourceAssetIdx == FollowUpDbSequenceIdx &&
+						Entry.bMirrored == SearchIndexAsset->bMirrored &&
+						Entry.SamplingInterval.Contains(0.0f);
+					return bIsMatch;
+				});
+
+			if (FollowUpSearchIndexAssetIdx != INDEX_NONE)
+			{
+				const FPoseSearchIndexAsset* FollowUpSearchIndexAsset =
+					&State.CurrentDatabase->SearchIndex.Assets[FollowUpSearchIndexAssetIdx];
+				const float FollowUpAssetTime = State.AssetPlayerTime + DeltaTime - AssetLength;
+				const int32 FollowUpPoseIdx = State.CurrentDatabase->GetPoseIndexFromTime(
+					FollowUpAssetTime,
+					FollowUpSearchIndexAsset);
+				const FFloatInterval SamplingRange = FollowUpSearchIndexAsset->SamplingInterval;
+
+				Result.PoseIdx = FollowUpPoseIdx;
+				Result.SearchIndexAsset = FollowUpSearchIndexAsset;
+				Result.AssetTime =
+					SamplingRange.Min +
+					(State.CurrentDatabase->Schema->SamplingInterval *
+						(Result.PoseIdx - FollowUpSearchIndexAsset->FirstPoseIdx));
+				bJumpRequired = true;
+			}
 		}
 	}
-	else
+	else if (SearchIndexAsset->Type == ESearchIndexAssetType::BlendSpace)
 	{
-		// check if there's a follow-up that can be used
-		int32 FollowUpDbSequenceIdx = State.CurrentDatabase->Sequences.IndexOfByPredicate(
-			[&](const FPoseSearchDatabaseSequence& Entry)
-			{
-				return Entry.Sequence == DbSequence.FollowUpSequence;
-			});
-
-		int32 FollowUpSearchIndexAssetIdx = State.CurrentDatabase->SearchIndex.Assets.IndexOfByPredicate(
-			[&](const FPoseSearchIndexAsset& Entry)
-			{
-				const bool bIsMatch =
-					Entry.SourceAssetIdx == FollowUpDbSequenceIdx &&
-					Entry.bMirrored == SearchIndexAsset->bMirrored &&
-					Entry.SamplingInterval.Contains(0.0f);
-				return bIsMatch;
-			});
-
-		if (FollowUpSearchIndexAssetIdx != INDEX_NONE)
+		const FPoseSearchDatabaseBlendSpace& DbBlendSpace = State.CurrentDatabase->GetBlendSpaceSourceAsset(SearchIndexAsset);
+		const float DeltaTime = UpdateContext.GetDeltaTime();
+
+		TArray<FBlendSampleData> BlendSamples;
+		int32 TriangulationIndex = 0;
+		DbBlendSpace.BlendSpace->GetSamplesFromBlendInput(SearchIndexAsset->BlendParameters, BlendSamples, TriangulationIndex, true);
+
+		float PlayLength = DbBlendSpace.BlendSpace->GetAnimationLengthFromSampleData(BlendSamples);
+
+		// Asset player time for blendspaces is normalized [0, 1] so we need to convert 
+		// to a real time before we advance it
+		float RealTime = State.AssetPlayerTime * PlayLength;		
+		float SteppedTime = RealTime;
+		ETypeAdvanceAnim AdvanceType = FAnimationRuntime::AdvanceTime(
+			DbBlendSpace.bLoopAnimation,
+			DeltaTime,
+			SteppedTime,
+			PlayLength);
+
+		if (AdvanceType != ETAA_Finished)
 		{
-			const FPoseSearchIndexAsset* FollowUpSearchIndexAsset =
-				&State.CurrentDatabase->SearchIndex.Assets[FollowUpSearchIndexAssetIdx];
-			const float FollowUpAssetTime = State.AssetPlayerTime + DeltaTime - AssetLength;
-			const int32 FollowUpPoseIdx = State.CurrentDatabase->GetPoseIndexFromAssetTime(
-				FollowUpAssetTime,
-				FollowUpSearchIndexAsset);
-			const FFloatInterval SamplingRange = FollowUpSearchIndexAsset->SamplingInterval;
-
-			Result.PoseIdx = FollowUpPoseIdx;
-			Result.SearchIndexAsset = FollowUpSearchIndexAsset;
-			Result.TimeOffsetSeconds =
-				SamplingRange.Min +
-				(State.CurrentDatabase->Schema->SamplingInterval *
-					(Result.PoseIdx - FollowUpSearchIndexAsset->FirstPoseIdx));
-			bJumpRequired = true;
+			// if the sampling interval doesn't contain the new time, there's no pose to return.
+			if (SearchIndexAsset->SamplingInterval.Contains(SteppedTime))
+			{
+				// we can continue ticking the same sequence forward
+				Result.PoseIdx = State.CurrentDatabase->GetPoseIndexFromTime(RealTime, SearchIndexAsset);
+				Result.SearchIndexAsset = SearchIndexAsset;
+				Result.AssetTime = State.AssetPlayerTime;
+			}
 		}
 	}
+	else
+	{
+		checkNoEntry();
+	}
 }
 
 } // namespace UE::PoseSearch
@@ -237,7 +256,7 @@ void FMotionMatchingState::JumpToPose(const FAnimationUpdateContext& Context, co
 	SearchIndexAssetIdx = CurrentDatabase->SearchIndex.FindAssetIndex(Result.SearchIndexAsset);
 
 	ElapsedPoseJumpTime = 0.0f;
-	AssetPlayerTime = Result.TimeOffsetSeconds;
+	AssetPlayerTime = Result.AssetTime;
 
 	const float JumpBlendTime = ComputeJumpBlendTime(Result, Settings);
 	RequestInertialBlend(Context, JumpBlendTime);
@@ -365,14 +384,15 @@ void UpdateMotionMatchingState(
 		const FPoseSearchIndexAsset* StateSearchIndexAsset = InOutMotionMatchingState.GetCurrentSearchIndexAsset();
 		if (StateSearchIndexAsset == Result.SearchIndexAsset)
 		{
-			const FPoseSearchDatabaseSequence& ResultDbSequence = Database->GetSourceAsset(Result.SearchIndexAsset);
-			bNearbyPose = FMath::Abs(InOutMotionMatchingState.AssetPlayerTime - Result.TimeOffsetSeconds) < Settings.PoseJumpThresholdTime;
+			// We need to check in terms of PoseIdx rather than AssetTime because
+			// for blendspaces, AssetTime is not in seconds, but in the normalized range 
+			// [0, 1] so comparing to `PoseJumpThresholdTime` will not make sense		
+			bNearbyPose = FMath::Abs(InOutMotionMatchingState.DbPoseIdx - Result.PoseIdx) * Database->Schema->SamplingInterval < Settings.PoseJumpThresholdTime;
 
 			// Handle looping anims when checking for the pose being too close
-			if (!bNearbyPose && ResultDbSequence.bLoopAnimation)
+			if (!bNearbyPose && Database->IsSourceAssetLooping(StateSearchIndexAsset))
 			{
-				const float AssetLength = ResultDbSequence.Sequence->GetPlayLength();
-				bNearbyPose = FMath::Abs(AssetLength - InOutMotionMatchingState.AssetPlayerTime - Result.TimeOffsetSeconds) < Settings.PoseJumpThresholdTime;
+				bNearbyPose = FMath::Abs(StateSearchIndexAsset->NumPoses - InOutMotionMatchingState.DbPoseIdx - Result.PoseIdx) * Database->Schema->SamplingInterval < Settings.PoseJumpThresholdTime;
 			}
 		}
 
@@ -388,7 +408,6 @@ void UpdateMotionMatchingState(
 		&& PoseStepper.CanContinue()
 		&& PoseStepper.bJumpRequired)
 	{
-		const FPoseSearchDatabaseSequence& DbSequence = Database->GetSourceAsset(PoseStepper.Result.SearchIndexAsset);
 		InOutMotionMatchingState.JumpToPose(Context, Settings, PoseStepper.Result);
 		InOutMotionMatchingState.Flags |= EMotionMatchingFlags::JumpedToFollowUp;
 	}
@@ -401,21 +420,49 @@ void UpdateMotionMatchingState(
 #if UE_POSE_SEARCH_TRACE_ENABLED
 	if (InOutMotionMatchingState.DbPoseIdx != INDEX_NONE)
 	{
+		float SimLinearVelocity, SimAngularVelocity, AnimLinearVelocity, AnimAngularVelocity;
+
+		if (DeltaTime > SMALL_NUMBER)
+		{
+			// simulation
 
-		float SimLinearVelocity = -1.0f;
-		float SimAngularVelocity = -1.0f;
-		float AnimLinearVelocity = -1.0f;
-		float AnimAngularVelocity = -1.0f;
-		ComputeCurrentVelocityMagnitudes(
-			InOutMotionMatchingState, 
-			Trajectory, 
-			DeltaTime, 
-			SimLinearVelocity, SimAngularVelocity, 
-			AnimLinearVelocity, AnimAngularVelocity);
+			int32 FirstIdx = 0;
+			const FTrajectorySample PrevSample = FTrajectorySampleRange::IterSampleTrajectory(
+				Trajectory.Samples,
+				ETrajectorySampleDomain::Time,
+				-DeltaTime, FirstIdx);
+
+			const FTrajectorySample CurrSample = FTrajectorySampleRange::IterSampleTrajectory(
+				Trajectory.Samples,
+				ETrajectorySampleDomain::Time,
+				0.0f, FirstIdx);
+
+			const FTransform SimDelta = CurrSample.Transform.GetRelativeTransform(PrevSample.Transform);
+
+			SimLinearVelocity = SimDelta.GetTranslation().Size() / DeltaTime;
+			SimAngularVelocity = FMath::RadiansToDegrees(SimDelta.GetRotation().GetAngle()) / DeltaTime;
+
+			// animation
+
+			const FTransform AnimDelta = InOutMotionMatchingState.RootMotionTransformDelta;
+
+			AnimLinearVelocity = AnimDelta.GetTranslation().Size() / DeltaTime;
+			AnimAngularVelocity = FMath::RadiansToDegrees(AnimDelta.GetRotation().GetAngle()) / DeltaTime;
+		}
+		else
+		{
+			SimLinearVelocity = 0.0f;
+			SimAngularVelocity = 0.0f;
+			AnimLinearVelocity = 0.0f;
+			AnimAngularVelocity = 0.0f;
+		}
 
 		TArray<bool> DatabaseSequenceFilter;
 		ComputeDatabaseSequenceFilter(Database, DatabaseTagQuery, DatabaseSequenceFilter);
 
+		TArray<bool> DatabaseBlendSpaceFilter;
+		ComputeDatabaseBlendSpaceFilter(Database, DatabaseTagQuery, DatabaseBlendSpaceFilter);
+
 		FTraceMotionMatchingState TraceState;
 		if (EnumHasAnyFlags(InOutMotionMatchingState.Flags, EMotionMatchingFlags::JumpedToFollowUp))
 		{
@@ -437,58 +484,12 @@ void UpdateMotionMatchingState(
 		TraceState.AnimLinearVelocity = AnimLinearVelocity;
 		TraceState.AnimAngularVelocity = AnimAngularVelocity;
 		TraceState.DatabaseSequenceFilter = DatabaseSequenceFilter;
+		TraceState.DatabaseBlendSpaceFilter = DatabaseBlendSpaceFilter;
 		UE_TRACE_POSE_SEARCH_MOTION_MATCHING_STATE(Context, TraceState)
 	}
 #endif
 }
 
-void UPoseSearchLibrary::UpdateMotionMatchingForSequencePlayer(
-	const FAnimUpdateContext& AnimUpdateContext,
-	const FSequencePlayerReference& SequencePlayer,
-	const UPoseSearchDatabase* Database,
-	const FTrajectorySampleRange& Trajectory,
-	const FMotionMatchingSettings& Settings,
-	UPARAM(ref) FMotionMatchingState& InOutMotionMatchingState
-)
-{
-	if (const FAnimationUpdateContext* AnimationUpdateContext = AnimUpdateContext.GetContext())
-	{
-		if (FAnimNode_SequencePlayer* SequencePlayerNode = SequencePlayer.GetAnimNodePtr<FAnimNode_SequencePlayer>())
-		{
-			// Update with the sequence player's current time.
-			InOutMotionMatchingState.AssetPlayerTime = SequencePlayerNode->GetAccumulatedTime();
-
-			// Execute core motion matching algorithm and retain across frame state
-			UpdateMotionMatchingState(
-				*AnimationUpdateContext, 
-				Database, 
-				nullptr, // DatabaseTagQuery
-				Trajectory, 
-				Settings, 
-				InOutMotionMatchingState);
-
-			// If a new pose is requested, jump to the pose by updating the embedded sequence player node
-			if ((InOutMotionMatchingState.Flags & EMotionMatchingFlags::JumpedToPose) == EMotionMatchingFlags::JumpedToPose)
-			{
-				const FPoseSearchIndexAsset* SearchIndexAsset = &Database->SearchIndex.Assets[InOutMotionMatchingState.SearchIndexAssetIdx];
-				const FPoseSearchDatabaseSequence& ResultDbSequence = Database->GetSourceAsset(SearchIndexAsset);
-				SequencePlayerNode->SetSequence(Cast<UAnimSequenceBase>(ResultDbSequence.Sequence));
-				SequencePlayerNode->SetAccumulatedTime(InOutMotionMatchingState.AssetPlayerTime);
-				SequencePlayerNode->SetLoopAnimation(ResultDbSequence.bLoopAnimation);
-				SequencePlayerNode->SetPlayRate(1.f);
-			}
-		}
-		else
-		{
-			UE_LOG(LogPoseSearchLibrary, Warning, TEXT("UpdateMotionMatchingForSequencePlayer must be called on a Sequence Player node"));
-		}
-	}
-	else
-	{
-		UE_LOG(LogPoseSearchLibrary, Warning, TEXT("UpdateMotionMatchingForSequencePlayer called with invalid context"));
-	}
-}
-
 const FPoseSearchIndexAsset* FMotionMatchingState::GetCurrentSearchIndexAsset() const
 {
 	if (!CurrentDatabase->SearchIndex.Assets.IsValidIndex(SearchIndexAssetIdx))
diff --git a/Engine/Plugins/Experimental/Animation/PoseSearch/Source/Runtime/Private/Trace/PoseSearchTraceLogger.cpp b/Engine/Plugins/Experimental/Animation/PoseSearch/Source/Runtime/Private/Trace/PoseSearchTraceLogger.cpp
--- a/Engine/Plugins/Experimental/Animation/PoseSearch/Source/Runtime/Private/Trace/PoseSearchTraceLogger.cpp
+++ b/Engine/Plugins/Experimental/Animation/PoseSearch/Source/Runtime/Private/Trace/PoseSearchTraceLogger.cpp
@@ -31,7 +31,8 @@ UE_TRACE_EVENT_BEGIN(PoseSearch, MotionMatchingState)
 	UE_TRACE_EVENT_FIELD(float, AnimLinearVelocity)
 	UE_TRACE_EVENT_FIELD(float, AnimAngularVelocity)
 	UE_TRACE_EVENT_FIELD(bool[], DatabaseSequenceFilter)
-UE_TRACE_EVENT_END()
+	UE_TRACE_EVENT_FIELD(bool[], DatabaseBlendSpaceFilter)
+	UE_TRACE_EVENT_END()
 
 namespace UE { namespace PoseSearch {
 
@@ -105,7 +106,8 @@ void FTraceMotionMatchingState::Output(const FAnimationBaseContext& InContext, c
 		<< MotionMatchingState.SimAngularVelocity(State.SimAngularVelocity)
 		<< MotionMatchingState.AnimLinearVelocity(State.AnimLinearVelocity)
 		<< MotionMatchingState.AnimAngularVelocity(State.AnimAngularVelocity)
-		<< MotionMatchingState.DatabaseSequenceFilter(State.DatabaseSequenceFilter.GetData(), State.DatabaseSequenceFilter.Num());
+		<< MotionMatchingState.DatabaseSequenceFilter(State.DatabaseSequenceFilter.GetData(), State.DatabaseSequenceFilter.Num())
+		<< MotionMatchingState.DatabaseBlendSpaceFilter(State.DatabaseBlendSpaceFilter.GetData(), State.DatabaseBlendSpaceFilter.Num());
 }
 
 }}
diff --git a/Engine/Plugins/Experimental/Animation/PoseSearch/Source/Runtime/Private/Trace/PoseSearchTraceLogger.h b/Engine/Plugins/Experimental/Animation/PoseSearch/Source/Runtime/Private/Trace/PoseSearchTraceLogger.h
--- a/Engine/Plugins/Experimental/Animation/PoseSearch/Source/Runtime/Private/Trace/PoseSearchTraceLogger.h
+++ b/Engine/Plugins/Experimental/Animation/PoseSearch/Source/Runtime/Private/Trace/PoseSearchTraceLogger.h
@@ -50,6 +50,9 @@ struct POSESEARCH_API FTraceMotionMatchingState
 	/** If true, groups are being filtered by DatabaseGroupQuery */
 	TArrayView<const bool> DatabaseSequenceFilter;
 
+	/** If true, groups are being filtered by DatabaseGroupQuery */
+	TArrayView<const bool> DatabaseBlendSpaceFilter;
+
 	/** Index of the pose in our database */
 	int32 DbPoseIdx = 0;
 
diff --git a/Engine/Plugins/Experimental/Animation/PoseSearch/Source/Runtime/Public/PoseSearch/AnimNode_MotionMatching.h b/Engine/Plugins/Experimental/Animation/PoseSearch/Source/Runtime/Public/PoseSearch/AnimNode_MotionMatching.h
--- a/Engine/Plugins/Experimental/Animation/PoseSearch/Source/Runtime/Public/PoseSearch/AnimNode_MotionMatching.h
+++ b/Engine/Plugins/Experimental/Animation/PoseSearch/Source/Runtime/Public/PoseSearch/AnimNode_MotionMatching.h
@@ -4,6 +4,7 @@
 
 #include "Animation/AnimNode_AssetPlayerBase.h"
 #include "Animation/AnimNode_SequencePlayer.h"
+#include "AnimNodes/AnimNode_BlendSpacePlayer.h"
 #include "AnimNodes/AnimNode_Mirror.h"
 #include "Animation/MotionTrajectoryTypes.h"
 #include "DynamicPlayRate/DynamicPlayRateLibrary.h"
@@ -69,12 +70,19 @@ private:
 	// Embedded sequence player node for playing animations from the motion matching database
 	FAnimNode_SequencePlayer_Standalone SequencePlayerNode;
 
+	// Embedded blendspace player node for playing blendspaces from the motion matching database
+	FAnimNode_BlendSpacePlayer_Standalone BlendSpacePlayerNode;
+
 	// Embedded mirror node to handle mirroring if the pose search results in a mirrored sequence
-	FAnimNode_Mirror MirrorNode;
+	FAnimNode_Mirror_Standalone MirrorNode;
 
 	// Encapsulated motion matching algorithm and internal state
 	FMotionMatchingState MotionMatchingState;
-	
+
+	// Current Asset Player Node
+	FAnimNode_AssetPlayerBase* CurrentAssetPlayerNode = &SequencePlayerNode;
+
+
 	// FAnimNode_AssetPlayerBase
 protected:
 	// FAnimNode_AssetPlayerBase interface
diff --git a/Engine/Plugins/Experimental/Animation/PoseSearch/Source/Runtime/Public/PoseSearch/PoseSearch.h b/Engine/Plugins/Experimental/Animation/PoseSearch/Source/Runtime/Public/PoseSearch/PoseSearch.h
--- a/Engine/Plugins/Experimental/Animation/PoseSearch/Source/Runtime/Public/PoseSearch/PoseSearch.h
+++ b/Engine/Plugins/Experimental/Animation/PoseSearch/Source/Runtime/Public/PoseSearch/PoseSearch.h
@@ -21,6 +21,7 @@
 #include "PoseSearch.generated.h"
 
 class UAnimSequence;
+class UBlendSpace;
 struct FCompactPose;
 struct FPoseContext;
 struct FReferenceSkeleton;
@@ -317,6 +318,15 @@ struct POSESEARCH_API FPoseSearchPoseMetadata
 	float CostAddend = 0.0f;
 };
 
+
+UENUM()
+enum class ESearchIndexAssetType : int32
+{
+	Invalid,
+	Sequence,
+	BlendSpace,
+};
+
 /**
 * Information about a source animation asset used by a search index.
 * Some source animation entries may generate multiple FPoseSearchIndexAsset entries.
@@ -330,16 +340,25 @@ public:
 	{}
 
 	FPoseSearchIndexAsset(
+		ESearchIndexAssetType InType,
 		int32 InSourceGroupIdx, 
 		int32 InSourceAssetIdx, 
 		bool bInMirrored, 
-		const FFloatInterval& InSamplingInterval)
-		: SourceGroupIdx(InSourceGroupIdx)
+		const FFloatInterval& InSamplingInterval,
+		FVector InBlendParameters = FVector::Zero())
+		: Type(InType)
+		, SourceGroupIdx(InSourceGroupIdx)
 		, SourceAssetIdx(InSourceAssetIdx)
 		, bMirrored(bInMirrored)
+		, BlendParameters(InBlendParameters)
 		, SamplingInterval(InSamplingInterval)
 	{}
 
+	// Default to Sequence for now for backward compatibility but
+	// at some point we might want to change this to Invalid.
+	UPROPERTY()
+	ESearchIndexAssetType Type = ESearchIndexAssetType::Sequence;
+
 	UPROPERTY()
 	int32 SourceGroupIdx = INDEX_NONE;
 
@@ -350,6 +369,9 @@ public:
 	UPROPERTY()
 	bool bMirrored = false;
 
+	UPROPERTY()
+	FVector BlendParameters = FVector::Zero();
+
 	UPROPERTY()
 	FFloatInterval SamplingInterval;
 
@@ -398,7 +420,7 @@ struct POSESEARCH_API FPoseSearchIndex
 
 	int32 FindAssetIndex(const FPoseSearchIndexAsset* Asset) const;
 	const FPoseSearchIndexAsset* FindAssetForPose(int32 PoseIdx) const;
-	float GetTimeOffset(int32 PoseIdx, const FPoseSearchIndexAsset* Asset) const;
+	float GetAssetTime(int32 PoseIdx, const FPoseSearchIndexAsset* Asset) const;
 
 	void Reset();
 
@@ -667,6 +689,46 @@ struct POSESEARCH_API FPoseSearchDatabaseSequence
 	FFloatInterval GetEffectiveSamplingRange() const;
 };
 
+/** An blend space entry in a UPoseSearchDatabase. */
+USTRUCT(BlueprintType, Category = "Animation|Pose Search")
+struct POSESEARCH_API FPoseSearchDatabaseBlendSpace
+{
+	GENERATED_BODY()
+
+	UPROPERTY(EditAnywhere, Category = "BlendSpace")
+	TObjectPtr<UBlendSpace> BlendSpace = nullptr;
+
+	UPROPERTY(EditAnywhere, Category = "BlendSpace")
+	bool bLoopAnimation = false;
+
+	UPROPERTY(EditAnywhere, Category = "BlendSpace")
+	EPoseSearchMirrorOption MirrorOption = EPoseSearchMirrorOption::UnmirroredOnly;
+
+	// If to use the blendspace grid locations as parameter sample locations.
+	// When enabled, NumberOfHorizontalSamples and NumberOfVerticalSamples are ignored.
+	UPROPERTY(EditAnywhere, Category = "BlendSpace")
+	bool bUseGridForSampling = true;
+
+	UPROPERTY(EditAnywhere, Category = "BlendSpace", meta = (ClampMin = "1", UIMin = "1", UIMax = "25"))
+	int32 NumberOfHorizontalSamples = 5;
+
+	UPROPERTY(EditAnywhere, Category = "BlendSpace", meta = (ClampMin = "1", UIMin = "1", UIMax = "25"))
+	int32 NumberOfVerticalSamples = 5;
+
+	UPROPERTY(EditAnywhere, Category = "Group")
+	FGameplayTagContainer GroupTags;
+
+public:
+
+	void GetBlendSpaceParameterSampleRanges(
+		int32& HorizontalBlendNum,
+		int32& VerticalBlendNum,
+		float& HorizontalBlendMin,
+		float& HorizontalBlendMax,
+		float& VerticalBlendMin,
+		float& VerticalBlendMax) const;
+};
+
 USTRUCT()
 struct POSESEARCH_API FPoseSearchDatabaseGroup
 {
@@ -705,7 +767,7 @@ public:
 	FPoseSearchExtrapolationParameters ExtrapolationParameters;
 
 	UPROPERTY(EditAnywhere, Category = "Database")
-	FPoseSearchBlockTransitionParameters BlockTransitionParameters = {0.0f, 0.2f};
+	FPoseSearchBlockTransitionParameters BlockTransitionParameters = { 0.0f, 0.2f };
 
 	UPROPERTY(EditAnywhere, Category = "Database")
 	TArray<FPoseSearchDatabaseGroup> Groups;
@@ -717,19 +779,27 @@ public:
 	UPROPERTY(EditAnywhere, Category="Database")
 	TArray<FPoseSearchDatabaseSequence> Sequences;
 
+	// Drag and drop blendspaces here to add them in bulk to Blend Spaces
+	UPROPERTY(EditAnywhere, Category = "Database", DisplayName = "Drag And Drop Blend Spaces Here")
+	TArray<TObjectPtr<UBlendSpace>> SimpleBlendSpaces;
+
+	UPROPERTY(EditAnywhere, Category = "Database")
+	TArray<FPoseSearchDatabaseBlendSpace> BlendSpaces;
+
 	UPROPERTY()
 	FPoseSearchIndex SearchIndex;
 
-	int32 FindSequenceForPose(int32 PoseIdx) const;
-	float GetSequenceLength(int32 DbSequenceIdx) const;
-	bool DoesSequenceLoop(int32 DbSequenceIdx) const;
-
 	bool IsValidForIndexing() const;
 	bool IsValidForSearch() const;
 
-	int32 GetPoseIndexFromAssetTime(float AssetTime, const FPoseSearchIndexAsset* SearchIndexAsset) const;
-	float GetTimeOffset(int32 PoseIdx, const FPoseSearchIndexAsset* SearchIndexAsset = nullptr) const;
-	const FPoseSearchDatabaseSequence& GetSourceAsset(const FPoseSearchIndexAsset* SearchIndexAsset) const;
+	int32 GetPoseIndexFromTime(float AssetTime, const FPoseSearchIndexAsset* SearchIndexAsset) const;
+	float GetAssetTime(int32 PoseIdx, const FPoseSearchIndexAsset* SearchIndexAsset = nullptr) const;
+
+	const FPoseSearchDatabaseSequence& GetSequenceSourceAsset(const FPoseSearchIndexAsset* SearchIndexAsset) const;
+	const FPoseSearchDatabaseBlendSpace& GetBlendSpaceSourceAsset(const FPoseSearchIndexAsset* SearchIndexAsset) const;
+	const bool IsSourceAssetLooping(const FPoseSearchIndexAsset* SearchIndexAsset) const;
+	const FGameplayTagContainer* GetSourceAssetGroupTags(const FPoseSearchIndexAsset* SearchIndexAsset) const;
+	const FString GetSourceAssetName(const FPoseSearchIndexAsset* SearchIndexAsset) const;
 
 public: // UObject
 	virtual void PreSave(FObjectPreSaveContext ObjectSaveContext) override;
@@ -740,6 +810,7 @@ public: // UObject
 
 private:
 	void CollectSimpleSequences();
+	void CollectSimpleBlendSpaces();
 
 public:
 	// Populates the FPoseSearchIndex::Assets array by evaluating the data in the Sequences array
@@ -952,9 +1023,9 @@ struct FSearchResult
 	FPoseCost PoseCost;
 	int32 PoseIdx = INDEX_NONE;
 	const FPoseSearchIndexAsset* SearchIndexAsset = nullptr;
-	float TimeOffsetSeconds = 0.0f;
+	float AssetTime = 0.0f;
 
-	bool IsValid() const { return PoseIdx >= 0; }
+	bool IsValid() const { return PoseIdx != INDEX_NONE; }
 };
 
 struct POSESEARCH_API FSearchContext
diff --git a/Engine/Plugins/Experimental/Animation/PoseSearch/Source/Runtime/Public/PoseSearch/PoseSearchLibrary.h b/Engine/Plugins/Experimental/Animation/PoseSearch/Source/Runtime/Public/PoseSearch/PoseSearchLibrary.h
--- a/Engine/Plugins/Experimental/Animation/PoseSearch/Source/Runtime/Public/PoseSearch/PoseSearchLibrary.h
+++ b/Engine/Plugins/Experimental/Animation/PoseSearch/Source/Runtime/Public/PoseSearch/PoseSearchLibrary.h
@@ -123,6 +123,11 @@ struct POSESEARCH_API FMotionMatchingState
 	// Evaluation flags relevant to the state of motion matching
 	UPROPERTY(VisibleAnywhere, BlueprintReadOnly, Category=State)
 	EMotionMatchingFlags Flags = EMotionMatchingFlags::None;
+
+	// Root motion delta for currently playing animation. Only required
+	// when UE_POSE_SEARCH_TRACE_ENABLED is active
+	UPROPERTY(Transient)
+	FTransform RootMotionTransformDelta = FTransform::Identity;
 };
 
 /**
@@ -141,22 +146,3 @@ POSESEARCH_API void UpdateMotionMatchingState(const FAnimationUpdateContext& Con
 	, const FMotionMatchingSettings& Settings
 	, FMotionMatchingState& InOutMotionMatchingState
 );
-
-// Exposes common pose search operations to Blueprint and native code
-UCLASS(Category="Animation|Pose Search")
-class POSESEARCH_API UPoseSearchLibrary : public UBlueprintFunctionLibrary
-{
-	GENERATED_BODY()
-
-public:
-
-	// Encapsulates and computes the motion matching algorithm for a sequence player
-	UFUNCTION(BlueprintPure, Category="Pose Search", meta=(BlueprintThreadSafe, AutoCreateRefTerm="Trajectory, Settings"))
-	static void UpdateMotionMatchingForSequencePlayer(const FAnimUpdateContext& AnimUpdateContext
-		, const FSequencePlayerReference& SequencePlayer
-		, const UPoseSearchDatabase* Database
-		, const FTrajectorySampleRange& Trajectory
-		, const FMotionMatchingSettings& Settings
-		, UPARAM(ref) FMotionMatchingState& InOutMotionMatchingState
-	);
-};
\ No newline at end of file
diff --git a/Engine/Source/Runtime/AnimGraphRuntime/Private/AnimGraphRuntimeTrace.cpp b/Engine/Source/Runtime/AnimGraphRuntime/Private/AnimGraphRuntimeTrace.cpp
--- a/Engine/Source/Runtime/AnimGraphRuntime/Private/AnimGraphRuntimeTrace.cpp
+++ b/Engine/Source/Runtime/AnimGraphRuntime/Private/AnimGraphRuntimeTrace.cpp
@@ -23,7 +23,7 @@ UE_TRACE_EVENT_BEGIN(Animation, BlendSpacePlayer)
 	UE_TRACE_EVENT_FIELD(float, FilteredPositionZ)
 	UE_TRACE_EVENT_END()
 
-void FAnimGraphRuntimeTrace::OutputBlendSpacePlayer(const FAnimationBaseContext& InContext, const FAnimNode_BlendSpacePlayer& InNode)
+void FAnimGraphRuntimeTrace::OutputBlendSpacePlayer(const FAnimationBaseContext& InContext, const FAnimNode_BlendSpacePlayerBase& InNode)
 {
 	bool bEventEnabled = UE_TRACE_CHANNELEXPR_IS_ENABLED(AnimationChannel);
 	if (!bEventEnabled)
diff --git a/Engine/Source/Runtime/AnimGraphRuntime/Private/AnimNodes/AnimNode_BlendSpacePlayer.cpp b/Engine/Source/Runtime/AnimGraphRuntime/Private/AnimNodes/AnimNode_BlendSpacePlayer.cpp
--- a/Engine/Source/Runtime/AnimGraphRuntime/Private/AnimNodes/AnimNode_BlendSpacePlayer.cpp
+++ b/Engine/Source/Runtime/AnimGraphRuntime/Private/AnimNodes/AnimNode_BlendSpacePlayer.cpp
@@ -9,9 +9,9 @@
 #include "Animation/AnimSyncScope.h"
 
 /////////////////////////////////////////////////////
-// FAnimNode_BlendSpacePlayer
+// FAnimNode_BlendSpacePlayerBase
 
-float FAnimNode_BlendSpacePlayer::GetCurrentAssetTime() const
+float FAnimNode_BlendSpacePlayerBase::GetCurrentAssetTime() const
 {
 	if(const FBlendSampleData* HighestWeightedSample = GetHighestWeightedSample())
 	{
@@ -22,13 +22,13 @@
 	return 0.0f;
 }
 
-float FAnimNode_BlendSpacePlayer::GetCurrentAssetTimePlayRateAdjusted() const
+float FAnimNode_BlendSpacePlayerBase::GetCurrentAssetTimePlayRateAdjusted() const
 {
 	float Length = GetCurrentAssetLength();
 	return GetPlayRate() < 0.0f ? Length - InternalTimeAccumulator * Length : Length * InternalTimeAccumulator;
 }
 
-float FAnimNode_BlendSpacePlayer::GetCurrentAssetLength() const
+float FAnimNode_BlendSpacePlayerBase::GetCurrentAssetLength() const
 {
 	if(const FBlendSampleData* HighestWeightedSample = GetHighestWeightedSample())
 	{
@@ -44,7 +44,7 @@
 	return 0.0f;
 }
 
-void FAnimNode_BlendSpacePlayer::Initialize_AnyThread(const FAnimationInitializeContext& Context)
+void FAnimNode_BlendSpacePlayerBase::Initialize_AnyThread(const FAnimationInitializeContext& Context)
 {
 	DECLARE_SCOPE_HIERARCHICAL_COUNTER_ANIMNODE(Initialize_AnyThread)
 	FAnimNode_AssetPlayerBase::Initialize_AnyThread(Context);
@@ -56,19 +56,19 @@
 	PreviousBlendSpace = GetBlendSpace();
 }
 
-void FAnimNode_BlendSpacePlayer::CacheBones_AnyThread(const FAnimationCacheBonesContext& Context)
+void FAnimNode_BlendSpacePlayerBase::CacheBones_AnyThread(const FAnimationCacheBonesContext& Context)
 {
 	DECLARE_SCOPE_HIERARCHICAL_COUNTER_ANIMNODE(CacheBones_AnyThread)
 }
 
-void FAnimNode_BlendSpacePlayer::UpdateAssetPlayer(const FAnimationUpdateContext& Context)
+void FAnimNode_BlendSpacePlayerBase::UpdateAssetPlayer(const FAnimationUpdateContext& Context)
 {
 	GetEvaluateGraphExposedInputs().Execute(Context);
 
 	UpdateInternal(Context);
 }
 
-void FAnimNode_BlendSpacePlayer::UpdateInternal(const FAnimationUpdateContext& Context)
+void FAnimNode_BlendSpacePlayerBase::UpdateInternal(const FAnimationUpdateContext& Context)
 {
 	DECLARE_SCOPE_HIERARCHICAL_COUNTER_ANIMNODE(UpdateInternal)
 
@@ -114,7 +114,7 @@
 	TRACE_ANIM_NODE_VALUE(Context, TEXT("Playback Time"), InternalTimeAccumulator);
 }
 
-void FAnimNode_BlendSpacePlayer::Evaluate_AnyThread(FPoseContext& Output)
+void FAnimNode_BlendSpacePlayerBase::Evaluate_AnyThread(FPoseContext& Output)
 {
 	DECLARE_SCOPE_HIERARCHICAL_COUNTER_ANIMNODE(Evaluate_AnyThread) 
 
@@ -130,7 +130,7 @@
 	}
 }
 
-void FAnimNode_BlendSpacePlayer::GatherDebugData(FNodeDebugData& DebugData)
+void FAnimNode_BlendSpacePlayerBase::GatherDebugData(FNodeDebugData& DebugData)
 {
 	DECLARE_SCOPE_HIERARCHICAL_COUNTER_ANIMNODE(GatherDebugData)
 	FString DebugLine = DebugData.GetNodeName(this);
@@ -144,19 +144,19 @@
 	}
 }
 
-float FAnimNode_BlendSpacePlayer::GetTimeFromEnd(float CurrentTime) const
+float FAnimNode_BlendSpacePlayerBase::GetTimeFromEnd(float CurrentTime) const
 {
 	// Blend-spaces use normalized time value
 	const float PlayLength = 1.0f;
 	return GetBlendSpace() != nullptr ? PlayLength - CurrentTime : 0.0f;
 }
 
-UAnimationAsset* FAnimNode_BlendSpacePlayer::GetAnimAsset() const
+UAnimationAsset* FAnimNode_BlendSpacePlayerBase::GetAnimAsset() const
 {
 	return GetBlendSpace();
 }
 
-const FBlendSampleData* FAnimNode_BlendSpacePlayer::GetHighestWeightedSample() const
+const FBlendSampleData* FAnimNode_BlendSpacePlayerBase::GetHighestWeightedSample() const
 {
 	if(BlendSampleDataCache.Num() == 0)
 	{
@@ -176,7 +176,7 @@
 	return HighestSample;
 }
 
-void FAnimNode_BlendSpacePlayer::Reinitialize(bool bResetTime)
+void FAnimNode_BlendSpacePlayerBase::Reinitialize(bool bResetTime)
 {
 	BlendSampleDataCache.Empty();
 	if(bResetTime)
@@ -198,51 +198,8 @@
 	}
 }
 
-bool FAnimNode_BlendSpacePlayer::SetBlendSpace(UBlendSpace* InBlendSpace)
-{
-#if WITH_EDITORONLY_DATA
-	BlendSpace = InBlendSpace;
-	GET_MUTABLE_ANIM_NODE_DATA(TObjectPtr<UBlendSpace>, BlendSpace) = InBlendSpace;
-#endif
-	
-	if(TObjectPtr<UBlendSpace>* BlendSpacePtr = GET_INSTANCE_ANIM_NODE_DATA_PTR(TObjectPtr<UBlendSpace>, BlendSpace))
-	{
-		*BlendSpacePtr = InBlendSpace;
-		return true;
-	}
+/////////////////////////////////////////////////////
+// FAnimNode_BlendSpacePlayer
 
-	return false;
-}
-
-FVector FAnimNode_BlendSpacePlayer::GetPosition() const
-{
-	return FVector(GET_ANIM_NODE_DATA(float, X), GET_ANIM_NODE_DATA(float, Y), 0.0f);
-}
-
-float FAnimNode_BlendSpacePlayer::GetPlayRate() const
-{
-	return GET_ANIM_NODE_DATA(float, PlayRate);
-}
-
-bool FAnimNode_BlendSpacePlayer::GetLoop() const
-{
-	return GET_ANIM_NODE_DATA(bool, bLoop);
-}
-
-bool FAnimNode_BlendSpacePlayer::ShouldResetPlayTimeWhenBlendSpaceChanges() const
-{
-	return GET_ANIM_NODE_DATA(bool, bResetPlayTimeWhenBlendSpaceChanges);
-}
-
-float FAnimNode_BlendSpacePlayer::GetStartPosition() const
-{
-	return GET_ANIM_NODE_DATA(float, StartPosition);
-}
-
-UBlendSpace* FAnimNode_BlendSpacePlayer::GetBlendSpace() const
-{
-	return GET_ANIM_NODE_DATA(TObjectPtr<UBlendSpace>, BlendSpace);
-}
-
 FName FAnimNode_BlendSpacePlayer::GetGroupName() const
 {
@@ -323,3 +280,105 @@
 
 	return false;
 }
+
+
+bool FAnimNode_BlendSpacePlayer::SetBlendSpace(UBlendSpace* InBlendSpace)
+{
+#if WITH_EDITORONLY_DATA
+	BlendSpace = InBlendSpace;
+	GET_MUTABLE_ANIM_NODE_DATA(TObjectPtr<UBlendSpace>, BlendSpace) = InBlendSpace;
+#endif
+
+	if (TObjectPtr<UBlendSpace>* BlendSpacePtr = GET_INSTANCE_ANIM_NODE_DATA_PTR(TObjectPtr<UBlendSpace>, BlendSpace))
+	{
+		*BlendSpacePtr = InBlendSpace;
+		return true;
+	}
+
+	return false;
+}
+
+FVector FAnimNode_BlendSpacePlayer::GetPosition() const
+{
+	return FVector(GET_ANIM_NODE_DATA(float, X), GET_ANIM_NODE_DATA(float, Y), 0.0f);
+}
+
+bool FAnimNode_BlendSpacePlayer::SetPosition(FVector InPosition)
+{
+#if WITH_EDITORONLY_DATA
+	X = InPosition[0];
+	Y = InPosition[1];
+	GET_MUTABLE_ANIM_NODE_DATA(float, X) = InPosition[0];
+	GET_MUTABLE_ANIM_NODE_DATA(float, Y) = InPosition[1];
+#endif
+
+	float* XPtr = GET_INSTANCE_ANIM_NODE_DATA_PTR(float, X);
+	float* YPtr = GET_INSTANCE_ANIM_NODE_DATA_PTR(float, Y);
+
+	if (XPtr && YPtr)
+	{
+		*XPtr = InPosition[0];
+		*YPtr = InPosition[1];
+		return true;
+	}
+
+	return false;
+}
+
+
+float FAnimNode_BlendSpacePlayer::GetPlayRate() const
+{
+	return GET_ANIM_NODE_DATA(float, PlayRate);
+}
+
+bool FAnimNode_BlendSpacePlayer::SetPlayRate(float InPlayRate)
+{
+#if WITH_EDITORONLY_DATA
+	PlayRate = InPlayRate;
+	GET_MUTABLE_ANIM_NODE_DATA(float, PlayRate) = InPlayRate;
+#endif
+
+	if (float* PlayRatePtr = GET_INSTANCE_ANIM_NODE_DATA_PTR(float, PlayRate))
+	{
+		*PlayRatePtr = InPlayRate;
+		return true;
+	}
+
+	return false;
+}
+
+bool FAnimNode_BlendSpacePlayer::GetLoop() const
+{
+	return GET_ANIM_NODE_DATA(bool, bLoop);
+}
+
+bool FAnimNode_BlendSpacePlayer::SetLoop(bool bInLoop)
+{
+#if WITH_EDITORONLY_DATA
+	bLoop = bInLoop;
+	GET_MUTABLE_ANIM_NODE_DATA(bool, bLoop) = bInLoop;
+#endif
+
+	if (bool* bLoopPtr = GET_INSTANCE_ANIM_NODE_DATA_PTR(bool, bLoop))
+	{
+		*bLoopPtr = bInLoop;
+		return true;
+	}
+
+	return false;
+}
+
+bool FAnimNode_BlendSpacePlayer::ShouldResetPlayTimeWhenBlendSpaceChanges() const
+{
+	return GET_ANIM_NODE_DATA(bool, bResetPlayTimeWhenBlendSpaceChanges);
+}
+
+float FAnimNode_BlendSpacePlayer::GetStartPosition() const
+{
+	return GET_ANIM_NODE_DATA(float, StartPosition);
+}
+
+UBlendSpace* FAnimNode_BlendSpacePlayer::GetBlendSpace() const
+{
+	return GET_ANIM_NODE_DATA(TObjectPtr<UBlendSpace>, BlendSpace);
+}
diff --git a/Engine/Source/Runtime/AnimGraphRuntime/Public/AnimGraphRuntimeTrace.h b/Engine/Source/Runtime/AnimGraphRuntime/Public/AnimGraphRuntimeTrace.h
--- a/Engine/Source/Runtime/AnimGraphRuntime/Public/AnimGraphRuntimeTrace.h
+++ b/Engine/Source/Runtime/AnimGraphRuntime/Public/AnimGraphRuntimeTrace.h
@@ -8,13 +8,13 @@
 #if ANIM_TRACE_ENABLED
 
 struct FAnimationBaseContext;
-struct FAnimNode_BlendSpacePlayer;
+struct FAnimNode_BlendSpacePlayerBase;
 struct FAnimNode_BlendSpaceGraphBase;
 
 struct FAnimGraphRuntimeTrace
 {
 	/** Helper function to output debug info for blendspace player nodes */
-	ANIMGRAPHRUNTIME_API static void OutputBlendSpacePlayer(const FAnimationBaseContext& InContext, const FAnimNode_BlendSpacePlayer& InNode);
+	ANIMGRAPHRUNTIME_API static void OutputBlendSpacePlayer(const FAnimationBaseContext& InContext, const FAnimNode_BlendSpacePlayerBase& InNode);
 
 	/** Helper function to output debug info for blendspace nodes */
 	ANIMGRAPHRUNTIME_API static void OutputBlendSpace(const FAnimationBaseContext& InContext, const FAnimNode_BlendSpaceGraphBase& InNode);
diff --git a/Engine/Source/Runtime/AnimGraphRuntime/Public/AnimNodes/AnimNode_BlendSpacePlayer.h b/Engine/Source/Runtime/AnimGraphRuntime/Public/AnimNodes/AnimNode_BlendSpacePlayer.h
--- a/Engine/Source/Runtime/AnimGraphRuntime/Public/AnimNodes/AnimNode_BlendSpacePlayer.h
+++ b/Engine/Source/Runtime/AnimGraphRuntime/Public/AnimNodes/AnimNode_BlendSpacePlayer.h
@@ -12,36 +12,124 @@
 
 //@TODO: Comment
 USTRUCT(BlueprintInternalUseOnly)
-struct ANIMGRAPHRUNTIME_API FAnimNode_BlendSpacePlayer : public FAnimNode_AssetPlayerBase
+struct ANIMGRAPHRUNTIME_API FAnimNode_BlendSpacePlayerBase : public FAnimNode_AssetPlayerBase
 {
 	GENERATED_BODY()
 
- 	friend class UAnimGraphNode_BlendSpacePlayer;
- 	friend class UAnimGraphNode_BlendSpaceEvaluator;
-	friend class UAnimGraphNode_RotationOffsetBlendSpace;
-	friend class UAnimGraphNode_AimOffsetLookAt;
+protected:
+	// Filter used to dampen coordinate changes
+	FBlendFilter BlendFilter;
 
+	// Cache of samples used to determine blend weights
+	TArray<FBlendSampleData> BlendSampleDataCache;
+
+	/** Previous position in the triangulation/segmentation */
+	int32 CachedTriangulationIndex = -1;
+
+	UPROPERTY(Transient)
+	TObjectPtr<UBlendSpace> PreviousBlendSpace = nullptr;
+
+public:	
+
+	// FAnimNode_AssetPlayerBase interface
+	virtual float GetCurrentAssetTime() const override;
+	virtual float GetCurrentAssetTimePlayRateAdjusted() const override;
+	virtual float GetCurrentAssetLength() const override;
+	virtual UAnimationAsset* GetAnimAsset() const override;
+	// End of FAnimNode_AssetPlayerBase interface
+
+	// FAnimNode_Base interface
+	virtual void Initialize_AnyThread(const FAnimationInitializeContext& Context) override;
+	virtual void CacheBones_AnyThread(const FAnimationCacheBonesContext& Context) override;
+	virtual void UpdateAssetPlayer(const FAnimationUpdateContext& Context) override;
+	virtual void Evaluate_AnyThread(FPoseContext& Output) override;
+	virtual void GatherDebugData(FNodeDebugData& DebugData) override;
+	// End of FAnimNode_Base interface
+
+	// Get the amount of time from the end
+	float GetTimeFromEnd(float CurrentTime) const;
+
 	// @return the current sample coordinates after going through the filtering
 	FVector GetFilteredPosition() const { return BlendFilter.GetFilterLastOutput(); }
+
+public:
+
+	// Get the blendspace asset to play
+	virtual UBlendSpace* GetBlendSpace() const PURE_VIRTUAL(FAnimNode_BlendSpacePlayerBase::GetBlendSpace, return nullptr;);
+
+	// Get the coordinates that are currently being sampled by the blendspace
+	virtual FVector GetPosition() const PURE_VIRTUAL(FAnimNode_BlendSpacePlayerBase::GetPosition, return FVector::Zero(););
+
+	// The start position in [0, 1] to use when initializing. When looping, play will still jump back to the beginning when reaching the end.
+	virtual float GetStartPosition() const PURE_VIRTUAL(FAnimNode_BlendSpacePlayerBase::GetStartPosition, return 0.0f;);
+
+	// Get the play rate multiplier. Can be negative, which will cause the animation to play in reverse.
+	virtual float GetPlayRate() const PURE_VIRTUAL(FAnimNode_BlendSpacePlayerBase::GetPlayRate, return 1.0f;);
+
+	// Should the animation loop back to the start when it reaches the end?
+	virtual bool GetLoop() const PURE_VIRTUAL(FAnimNode_BlendSpacePlayerBase::GetLoop, return true;);
+
+	// Get whether we should reset the current play time when the blend space changes
+	virtual bool ShouldResetPlayTimeWhenBlendSpaceChanges() const PURE_VIRTUAL(FAnimNode_BlendSpacePlayerBase::ShouldResetPlayTimeWhenBlendSpaceChanges, return true;);
+
+	// An evaluator will be setting the play rate to zero and setting the time explicitly. ShouldTeleportToTime indicates whether we should jump to that time, or move to it playing out root motion and events etc.
+	virtual bool ShouldTeleportToTime() const { return false; }
+
+	// Indicates if we are an evaluator - i.e. will be setting the time explicitly rather than letting it play out
+	virtual bool IsEvaluator() const { return false; }
+
+	// Set the blendspace asset to play
+	virtual bool SetBlendSpace(UBlendSpace* InBlendSpace) PURE_VIRTUAL(FAnimNode_BlendSpacePlayerBase::SetBlendSpace, return false;);
+
+	// Set the coordinates that are currently being sampled by the blendspace
+	virtual bool SetPosition(FVector InPosition) PURE_VIRTUAL(FAnimNode_BlendSpacePlayerBase::SetPosition, return false;);
+
+	// Set the play rate multiplier. Can be negative, which will cause the animation to play in reverse.
+	virtual bool SetPlayRate(float InPlayRate) PURE_VIRTUAL(FAnimNode_BlendSpacePlayerBase::SetPlayRate, return false;);
+
+	// Set if the animation should loop back to the start when it reaches the end?
+	virtual bool SetLoop(bool bInLoop) PURE_VIRTUAL(FAnimNode_BlendSpacePlayerBase::SetLoop, return false;);
+
+protected:
+	void UpdateInternal(const FAnimationUpdateContext& Context);
+
+private:
+	void Reinitialize(bool bResetTime = true);
+
+	const FBlendSampleData* GetHighestWeightedSample() const;
+};
+
+
+//@TODO: Comment
+USTRUCT(BlueprintInternalUseOnly)
+struct ANIMGRAPHRUNTIME_API FAnimNode_BlendSpacePlayer : public FAnimNode_BlendSpacePlayerBase
+{
+	GENERATED_BODY()
+
+	friend class UAnimGraphNode_BlendSpacePlayer;
+	friend class UAnimGraphNode_BlendSpaceEvaluator;
+	friend class UAnimGraphNode_RotationOffsetBlendSpace;
+	friend class UAnimGraphNode_AimOffsetLookAt;
+
 private:
 
 #if WITH_EDITORONLY_DATA
 	// The group name (NAME_None if it is not part of any group)
-	UPROPERTY(EditAnywhere, Category=Sync, meta=(FoldProperty))
+	UPROPERTY(EditAnywhere, Category = Sync, meta = (FoldProperty))
 	FName GroupName = NAME_None;
 
 	// The role this player can assume within the group (ignored if GroupIndex is INDEX_NONE)
-	UPROPERTY(EditAnywhere, Category=Sync, meta=(FoldProperty))
+	UPROPERTY(EditAnywhere, Category = Sync, meta = (FoldProperty))
 	TEnumAsByte<EAnimGroupRole::Type> GroupRole = EAnimGroupRole::CanBeLeader;
 
 	// How synchronization is determined
-	UPROPERTY(EditAnywhere, Category=Sync, meta=(FoldProperty))
+	UPROPERTY(EditAnywhere, Category = Sync, meta = (FoldProperty))
 	EAnimSyncMethod Method = EAnimSyncMethod::DoNotSync;
 
 	// If true, "Relevant anim" nodes that look for the highest weighted animation in a state will ignore this node
-	UPROPERTY(EditAnywhere, Category=Relevancy, meta=(FoldProperty, PinHiddenByDefault))
+	UPROPERTY(EditAnywhere, Category = Relevancy, meta = (FoldProperty, PinHiddenByDefault))
 	bool bIgnoreForRelevancyTest = false;
-	
+
 	// The X coordinate to sample in the blendspace
 	UPROPERTY(EditAnywhere, Category = Coordinates, meta = (PinShownByDefault, FoldProperty))
 	float X = 0.0f;
@@ -63,7 +151,7 @@
 	bool bResetPlayTimeWhenBlendSpaceChanges = true;
 
 	// The start position in [0, 1] to use when initializing. When looping, play will still jump back to the beginning when reaching the end.
-	UPROPERTY(EditAnywhere, Category=Settings, meta = (DefaultValue = "0.f", PinHiddenByDefault, FoldProperty))
+	UPROPERTY(EditAnywhere, Category = Settings, meta = (DefaultValue = "0.f", PinHiddenByDefault, FoldProperty))
 	float StartPosition = 0.0f;
 
 	// The blendspace asset to play
@@ -71,27 +159,9 @@
 	TObjectPtr<UBlendSpace> BlendSpace = nullptr;
 #endif
 
-protected:
-	// Filter used to dampen coordinate changes
-	FBlendFilter BlendFilter;
+public:
 
-	// Cache of samples used to determine blend weights
-	TArray<FBlendSampleData> BlendSampleDataCache;
-
-	/** Previous position in the triangulation/segmentation */
-	int32 CachedTriangulationIndex = -1;
-
-	UPROPERTY(Transient)
-	TObjectPtr<UBlendSpace> PreviousBlendSpace = nullptr;
-
-public:	
-	FAnimNode_BlendSpacePlayer() = default;
-
 	// FAnimNode_AssetPlayerBase interface
-	virtual float GetCurrentAssetTime() const override;
-	virtual float GetCurrentAssetTimePlayRateAdjusted() const override;
-	virtual float GetCurrentAssetLength() const override;
-	virtual UAnimationAsset* GetAnimAsset() const override;
 	virtual FName GetGroupName() const override;
 	virtual EAnimGroupRole::Type GetGroupRole() const override;
 	virtual EAnimSyncMethod GetGroupMethod() const override;
@@ -102,48 +172,95 @@
 	virtual bool SetIgnoreForRelevancyTest(bool bInIgnoreForRelevancyTest) override;
 	// End of FAnimNode_AssetPlayerBase interface
 
-	// FAnimNode_Base interface
-	virtual void Initialize_AnyThread(const FAnimationInitializeContext& Context) override;
-	virtual void CacheBones_AnyThread(const FAnimationCacheBonesContext& Context) override;
-	virtual void UpdateAssetPlayer(const FAnimationUpdateContext& Context) override;
-	virtual void Evaluate_AnyThread(FPoseContext& Output) override;
-	virtual void GatherDebugData(FNodeDebugData& DebugData) override;
-	// End of FAnimNode_Base interface
+	// FAnimNode_BlendSpacePlayerBase interface
+	virtual UBlendSpace* GetBlendSpace() const override;
+	virtual FVector GetPosition() const override;
+	virtual float GetStartPosition() const override;
+	virtual float GetPlayRate() const override;
+	virtual bool GetLoop() const override;
+	virtual bool ShouldResetPlayTimeWhenBlendSpaceChanges() const override;
+	virtual bool SetBlendSpace(UBlendSpace* InBlendSpace) override;
+	virtual bool SetPosition(FVector InPosition) override;
+	virtual bool SetPlayRate(float InPlayRate) override;
+	virtual bool SetLoop(bool bInLoop) override;
+	// End of FAnimNode_BlendSpacePlayerBase interface
+};
 
-	float GetTimeFromEnd(float CurrentTime) const;
+//@TODO: Comment
+USTRUCT(BlueprintInternalUseOnly)
+struct ANIMGRAPHRUNTIME_API FAnimNode_BlendSpacePlayer_Standalone : public FAnimNode_BlendSpacePlayerBase
+{
+	GENERATED_BODY()
 
-	// Set the blendspace asset to play
-	bool SetBlendSpace(UBlendSpace* InBlendSpace);
+private:
 
-	// Get the coordinates that are currently being sampled by the blendspace
-	virtual FVector GetPosition() const;
+	// The group name (NAME_None if it is not part of any group)
+	UPROPERTY(EditAnywhere, Category = Sync)
+	FName GroupName = NAME_None;
 
-	// Get the play rate multiplier. Can be negative, which will cause the animation to play in reverse.
-	virtual float GetPlayRate() const;
+	// The role this player can assume within the group (ignored if GroupIndex is INDEX_NONE)
+	UPROPERTY(EditAnywhere, Category = Sync)
+	TEnumAsByte<EAnimGroupRole::Type> GroupRole = EAnimGroupRole::CanBeLeader;
 
-	// An evaluator will be setting the play rate to zero and setting the time explicitly. ShouldTeleportToTime indicates whether we should jump to that time, or move to it playing out root motion and events etc.
-	virtual bool ShouldTeleportToTime() const { return false; }
+	// How synchronization is determined
+	UPROPERTY(EditAnywhere, Category = Sync)
+	EAnimSyncMethod Method = EAnimSyncMethod::DoNotSync;
 
-	// Indicates if we are an evaluator - i.e. will be setting the time explicitly rather than letting it play out
-	virtual bool IsEvaluator() const { return false; }
+	// If true, "Relevant anim" nodes that look for the highest weighted animation in a state will ignore this node
+	UPROPERTY(EditAnywhere, Category = Relevancy, meta = (PinHiddenByDefault))
+	bool bIgnoreForRelevancyTest = false;
 
+	// The X coordinate to sample in the blendspace
+	UPROPERTY(EditAnywhere, Category = Coordinates, meta = (PinShownByDefault))
+	float X = 0.0f;
+
+	// The Y coordinate to sample in the blendspace
+	UPROPERTY(EditAnywhere, Category = Coordinates, meta = (PinShownByDefault))
+	float Y = 0.0f;
+
+	// The play rate multiplier. Can be negative, which will cause the animation to play in reverse.
+	UPROPERTY(EditAnywhere, Category = Settings, meta = (DefaultValue = "1.0", PinHiddenByDefault))
+	float PlayRate = 1.0f;
+
 	// Should the animation loop back to the start when it reaches the end?
-	bool GetLoop() const;
+	UPROPERTY(EditAnywhere, Category = Settings, meta = (DefaultValue = "true", PinHiddenByDefault))
+	bool bLoop = true;
 
-	// Get whether we should reset the current play time when the blend space changes
-	bool ShouldResetPlayTimeWhenBlendSpaceChanges() const;
+	// Whether we should reset the current play time when the blend space changes
+	UPROPERTY(EditAnywhere, Category = Settings, meta = (PinHiddenByDefault))
+	bool bResetPlayTimeWhenBlendSpaceChanges = true;
 
 	// The start position in [0, 1] to use when initializing. When looping, play will still jump back to the beginning when reaching the end.
-	float GetStartPosition() const;
+	UPROPERTY(EditAnywhere, Category = Settings, meta = (DefaultValue = "0.f", PinHiddenByDefault))
+	float StartPosition = 0.0f;
 
-	// Get the blendspace asset to play
-	UBlendSpace* GetBlendSpace() const;
+	// The blendspace asset to play
+	UPROPERTY(EditAnywhere, Category = Settings, meta = (PinHiddenByDefault))
+	TObjectPtr<UBlendSpace> BlendSpace = nullptr;
 
-protected:
-	void UpdateInternal(const FAnimationUpdateContext& Context);
+public:
 
-private:
-	void Reinitialize(bool bResetTime = true);
+	// FAnimNode_AssetPlayerBase interface
+	virtual FName GetGroupName() const override { return GroupName; }
+	virtual EAnimGroupRole::Type GetGroupRole() const override { return GroupRole; }
+	virtual EAnimSyncMethod GetGroupMethod() const override { return Method; }
+	virtual bool GetIgnoreForRelevancyTest() const override { return bIgnoreForRelevancyTest; }
+	virtual bool SetGroupName(FName InGroupName) override { GroupName = InGroupName; return true; }
+	virtual bool SetGroupRole(EAnimGroupRole::Type InRole) override { GroupRole = InRole; return true; }
+	virtual bool SetGroupMethod(EAnimSyncMethod InMethod) override { Method = InMethod; return true; }
+	virtual bool SetIgnoreForRelevancyTest(bool bInIgnoreForRelevancyTest) override { bIgnoreForRelevancyTest = bInIgnoreForRelevancyTest; return true; }
+	// End of FAnimNode_AssetPlayerBase interface
 
-	const FBlendSampleData* GetHighestWeightedSample() const;
+	// FAnimNode_BlendSpacePlayerBase interface
+	virtual UBlendSpace* GetBlendSpace() const override { return BlendSpace; }
+	virtual FVector GetPosition() const override { return FVector(X, Y, 0.0); }
+	virtual float GetStartPosition() const override { return StartPosition; }
+	virtual float GetPlayRate() const override { return PlayRate; }
+	virtual bool GetLoop() const override { return bLoop; }
+	virtual bool ShouldResetPlayTimeWhenBlendSpaceChanges() const override { return bResetPlayTimeWhenBlendSpaceChanges; }
+	virtual bool SetBlendSpace(UBlendSpace* InBlendSpace) override { BlendSpace = InBlendSpace; return true; }
+	virtual bool SetPosition(FVector InPosition) override { X = InPosition[0]; Y = InPosition[1]; return true; }
+	virtual bool SetPlayRate(float InPlayRate) override { PlayRate = InPlayRate; return true; }
+	virtual bool SetLoop(bool bInLoop) override { bLoop = bInLoop; return true; }
+	// End of FAnimNode_BlendSpacePlayerBase interface
 };
diff --git a/Engine/Source/Runtime/Engine/Classes/Animation/BlendSpace.h b/Engine/Source/Runtime/Engine/Classes/Animation/BlendSpace.h
--- a/Engine/Source/Runtime/Engine/Classes/Animation/BlendSpace.h
+++ b/Engine/Source/Runtime/Engine/Classes/Animation/BlendSpace.h
@@ -518,6 +518,9 @@ public:
 	*/
 	ENGINE_API bool GetSamplesFromBlendInput(const FVector &BlendInput, TArray<FBlendSampleData> & OutSampleDataList, int32& InOutCachedTriangulationIndex, bool bCombineAnimations) const;
 
+	/** Utility function to calculate animation length from sample data list **/
+	ENGINE_API float GetAnimationLengthFromSampleData(const TArray<FBlendSampleData>& SampleDataList) const;
+
 	/** 
 	 * Initialize BlendSpace filtering for runtime. Filtering supports multiple dimensions, defaulting to
 	 * two (since we don't have 3D BlendSpaces yet) 
@@ -645,9 +648,6 @@ protected:
 
 	void TickFollowerSamples(TArray<FBlendSampleData> &SampleDataList, const int32 HighestWeightIndex, FAnimAssetTickContext &Context, bool bResetMarkerDataOnFollowers, const UMirrorDataTable* MirrorDataTable = nullptr) const;
 
-	/** Utility function to calculate animation length from sample data list **/
-	float GetAnimationLengthFromSampleData(const TArray<FBlendSampleData> & SampleDataList) const;
-
 	/** Returns the blend input clamped to the valid range, unless that axis has been set to wrap in which case no clamping is done **/
 	FVector GetClampedBlendInput(const FVector& BlendInput) const;
 	
