Fixed crash in FAnimNode_MotionMatching when the pose search database is null

diff --git a/Engine/Plugins/Experimental/Animation/PoseSearch/Source/Runtime/Private/AnimNode_MotionMatching.cpp b/Engine/Plugins/Experimental/Animation/PoseSearch/Source/Runtime/Private/AnimNode_MotionMatching.cpp
--- a/Engine/Plugins/Experimental/Animation/PoseSearch/Source/Runtime/Private/AnimNode_MotionMatching.cpp
+++ b/Engine/Plugins/Experimental/Animation/PoseSearch/Source/Runtime/Private/AnimNode_MotionMatching.cpp
@@ -28,7 +28,10 @@ void FAnimNode_MotionMatching::Initialize_AnyThread(const FAnimationInitializeCo
 	CurrentAssetPlayerNode = &SequencePlayerNode;
 
 	MirrorNode.SetSourceLinkNode(CurrentAssetPlayerNode);
-	MirrorNode.SetMirrorDataTable(Database->Schema->MirrorDataTable.Get());
+	if (Database && Database->Schema)
+	{
+		MirrorNode.SetMirrorDataTable(Database->Schema->MirrorDataTable.Get());
+	}
 
 	Source.SetLinkNode(&MirrorNode);
 	Source.Initialize(Context);
@@ -76,6 +79,11 @@ void FAnimNode_MotionMatching::UpdateAssetPlayer(const FAnimationUpdateContext&
 		MotionMatchingState
 	);
 
+	if (MotionMatchingState.CurrentDatabase.IsValid() && MotionMatchingState.CurrentDatabase->Schema)
+	{
+		MirrorNode.SetMirrorDataTable(MotionMatchingState.CurrentDatabase->Schema->MirrorDataTable.Get());
+	}
+
 	const FPoseSearchIndexAsset* SearchIndexAsset = MotionMatchingState.GetCurrentSearchIndexAsset();
 
 	// If a new pose is requested, jump to the pose by updating the embedded sequence player node
diff --git a/Engine/Plugins/Experimental/Animation/PoseSearch/Source/Runtime/Private/PoseSearchLibrary.cpp b/Engine/Plugins/Experimental/Animation/PoseSearch/Source/Runtime/Private/PoseSearchLibrary.cpp
--- a/Engine/Plugins/Experimental/Animation/PoseSearch/Source/Runtime/Private/PoseSearchLibrary.cpp
+++ b/Engine/Plugins/Experimental/Animation/PoseSearch/Source/Runtime/Private/PoseSearchLibrary.cpp
@@ -492,7 +492,8 @@ void UpdateMotionMatchingState(
 
 const FPoseSearchIndexAsset* FMotionMatchingState::GetCurrentSearchIndexAsset() const
 {
-	if (!CurrentDatabase->SearchIndex.Assets.IsValidIndex(SearchIndexAssetIdx))
+	if (!CurrentDatabase.IsValid() ||
+		!CurrentDatabase->SearchIndex.Assets.IsValidIndex(SearchIndexAssetIdx))
 	{
 		return nullptr;
 	}
