Motion Matching - Moved Pose Search Database Search Index to the DDC using lower level DDC api

diff --git a/Engine/Plugins/Experimental/Animation/PoseSearch/Source/Editor/Private/PoseSearchDebugger.cpp b/Engine/Plugins/Experimental/Animation/PoseSearch/Source/Editor/Private/PoseSearchDebugger.cpp
--- a/Engine/Plugins/Experimental/Animation/PoseSearch/Source/Editor/Private/PoseSearchDebugger.cpp
+++ b/Engine/Plugins/Experimental/Animation/PoseSearch/Source/Editor/Private/PoseSearchDebugger.cpp
@@ -866,13 +866,13 @@ void SDebuggerDatabaseView::FilterDatabaseRows()
 
 void SDebuggerDatabaseView::CreateRows(const UPoseSearchDatabase& Database)
 {
-	const int32 NumPoses = Database.SearchIndex.NumPoses;
+	const int32 NumPoses = Database.GetSearchIndex()->NumPoses;
 	UnfilteredDatabaseRows.Reset(NumPoses);
 
 	RowsSourceDatabase = &Database;
 
 	// Build database rows
-	for(const FPoseSearchIndexAsset& SearchIndexAsset : Database.SearchIndex.Assets)
+	for(const FPoseSearchIndexAsset& SearchIndexAsset : Database.GetSearchIndex()->Assets)
 	{
 		if (SearchIndexAsset.Type == ESearchIndexAssetType::Sequence)
 		{
@@ -942,7 +942,7 @@ void SDebuggerDatabaseView::UpdateRows(const FTraceMotionMatchingStateMessage& S
 	SearchContext.SetSource(&Database);
 	SearchContext.QueryValues = State.QueryVectorNormalized;
 	SearchContext.WeightsContext = &StateWeights;
-	if (const FPoseSearchIndexAsset* CurrentIndexAsset = Database.SearchIndex.FindAssetForPose(State.DbPoseIdx))
+	if (const FPoseSearchIndexAsset* CurrentIndexAsset = Database.GetSearchIndex()->FindAssetForPose(State.DbPoseIdx))
 	{
 		SearchContext.QueryMirrorRequest = CurrentIndexAsset->bMirrored ? 
 			EPoseSearchBooleanRequest::TrueValue : EPoseSearchBooleanRequest::FalseValue;
@@ -1268,7 +1268,7 @@ void SDebuggerDetailsView::UpdateReflection(const FTraceMotionMatchingStateMessa
 	Reflection->bFollowUpAnimation = EnumHasAnyFlags(State.Flags, FTraceMotionMatchingState::EFlags::FollowupAnimation);
 
 	Reflection->AssetPlayerAssetName = FString();
-	if (const FPoseSearchIndexAsset* IndexAsset = Database.SearchIndex.FindAssetForPose(State.DbPoseIdx))
+	if (const FPoseSearchIndexAsset* IndexAsset = Database.GetSearchIndex()->FindAssetForPose(State.DbPoseIdx))
 	{
 		Reflection->AssetPlayerAssetName = Database.GetSourceAssetName(IndexAsset);
 	}
@@ -1288,8 +1288,8 @@ void SDebuggerDetailsView::UpdateReflection(const FTraceMotionMatchingStateMessa
 	Reflection->QueryPoseVector.ExtractFeatures(Reader);
 
 	// Active pose
-	TArray<float> Pose(Database.SearchIndex.GetPoseValues(State.DbPoseIdx));
-	Database.SearchIndex.InverseNormalize(Pose);
+	TArray<float> Pose(Database.GetSearchIndex()->GetPoseValues(State.DbPoseIdx));
+	Database.GetSearchIndex()->InverseNormalize(Pose);
 	Reader.SetValues(Pose);
 	Reflection->ActivePoseVector.ExtractFeatures(Reader);
 
@@ -1300,8 +1300,8 @@ void SDebuggerDetailsView::UpdateReflection(const FTraceMotionMatchingStateMessa
 		if (!SelectedRows.IsEmpty())
 		{
 			const TSharedRef<FDebuggerDatabaseRowData>& Selected = SelectedRows[0];
-			Pose = Database.SearchIndex.GetPoseValues(Selected->PoseIdx);
-			Database.SearchIndex.InverseNormalize(Pose);
+			Pose = Database.GetSearchIndex()->GetPoseValues(Selected->PoseIdx);
+			Database.GetSearchIndex()->InverseNormalize(Pose);
 			Reader.SetValues(Pose);
 			Reflection->SelectedPoseVector.ExtractFeatures(Reader);
 
@@ -2012,7 +2012,7 @@ void FDebuggerViewModel::ShowSelectedSkeleton(int32 PoseIdx, float Time)
 
 	Component->ResetToStart();
 	bSelecting = true;
-	const FPoseSearchIndexAsset* IndexAsset = Database->SearchIndex.FindAssetForPose(PoseIdx);
+	const FPoseSearchIndexAsset* IndexAsset = Database->GetSearchIndex()->FindAssetForPose(PoseIdx);
 	Skeletons[SelectedPose].Type = IndexAsset->Type;
 	Skeletons[SelectedPose].Time = Time;
 	Skeletons[SelectedPose].bMirrored = IndexAsset->bMirrored;
@@ -2094,7 +2094,7 @@ void FDebuggerViewModel::OnUpdateNodeSelection(int32 InNodeId)
 
 	if (ActiveMotionMatchingState)
 	{
-		const FPoseSearchIndexAsset* IndexAsset = NewDatabase->SearchIndex.FindAssetForPose(ActiveMotionMatchingState->DbPoseIdx);
+		const FPoseSearchIndexAsset* IndexAsset = NewDatabase->GetSearchIndex()->FindAssetForPose(ActiveMotionMatchingState->DbPoseIdx);
 
 		if (IndexAsset)
 		{
@@ -2424,7 +2424,7 @@ void FDebuggerViewModel::PlaySelection(int32 PoseIdx, float Time)
 
 	Component->ResetToStart();
 	
-	const FPoseSearchIndexAsset* IndexAsset = Database->SearchIndex.FindAssetForPose(PoseIdx);
+	const FPoseSearchIndexAsset* IndexAsset = Database->GetSearchIndex()->FindAssetForPose(PoseIdx);
 
 	Skeletons[Asset].Type = IndexAsset->Type;
 	Skeletons[Asset].AssetIdx = IndexAsset->SourceAssetIdx;
diff --git a/Engine/Plugins/Experimental/Animation/PoseSearch/Source/Runtime/PoseSearch.Build.cs b/Engine/Plugins/Experimental/Animation/PoseSearch/Source/Runtime/PoseSearch.Build.cs
--- a/Engine/Plugins/Experimental/Animation/PoseSearch/Source/Runtime/PoseSearch.Build.cs
+++ b/Engine/Plugins/Experimental/Animation/PoseSearch/Source/Runtime/PoseSearch.Build.cs
@@ -48,5 +48,14 @@ public class PoseSearch : ModuleRules
 				"Engine"
 			}
 		);
+
+		if (Target.bCompileAgainstEditor)
+		{
+			PrivateDependencyModuleNames.AddRange(
+				new string[] {
+					"DerivedDataCache"
+				}
+			);
+		}
 	}
 }
diff --git a/Engine/Plugins/Experimental/Animation/PoseSearch/Source/Runtime/Private/PoseSearch.cpp b/Engine/Plugins/Experimental/Animation/PoseSearch/Source/Runtime/Private/PoseSearch.cpp
--- a/Engine/Plugins/Experimental/Animation/PoseSearch/Source/Runtime/Private/PoseSearch.cpp
+++ b/Engine/Plugins/Experimental/Animation/PoseSearch/Source/Runtime/Private/PoseSearch.cpp
@@ -2,6 +2,7 @@
 
 #include "PoseSearch/PoseSearch.h"
 #include "PoseSearch/PoseSearchAnimNotifies.h"
+#include "PoseSearchDerivedData.h"
 #include "PoseSearchEigenHelper.h"
 
 #include "Algo/BinarySearch.h"
@@ -27,10 +28,16 @@
 #include "Trace/PoseSearchTraceLogger.h"
 #include "UObject/ObjectSaveContext.h"
 #include "Misc/MemStack.h"
+
 #include "Animation/BuiltInAttributeTypes.h"
 #include "Animation/AnimRootMotionProvider.h"
 #include "Animation/BlendSpace1D.h"
 
+#if WITH_EDITOR
+#include "AssetRegistry/AssetRegistryModule.h"
+#include "AssetRegistry/ARFilter.h"
+#endif // WITH_EDITOR
+
 IMPLEMENT_ANIMGRAPH_MESSAGE(UE::PoseSearch::IPoseHistoryProvider);
 
 DEFINE_LOG_CATEGORY(LogPoseSearch);
@@ -230,7 +237,6 @@ static void CalcChannelCosts(const UPoseSearchSchema* Schema, TArrayView<const f
 		}
 	}
 }
-
 }} // namespace UE::PoseSearch
 
 
@@ -1142,6 +1148,26 @@ FFloatInterval FPoseSearchDatabaseSequence::GetEffectiveSamplingRange() const
 //////////////////////////////////////////////////////////////////////////
 // UPoseSearchDatabase
 
+FPoseSearchIndex* UPoseSearchDatabase::GetSearchIndex()
+{
+	if (PrivateDerivedData == nullptr)
+	{
+		return nullptr;
+	}
+
+	return &PrivateDerivedData->SearchIndex;
+}
+
+const FPoseSearchIndex* UPoseSearchDatabase::GetSearchIndex() const
+{
+	if (PrivateDerivedData == nullptr)
+	{
+		return nullptr;
+	}
+
+	return &PrivateDerivedData->SearchIndex;
+}
+
 int32 UPoseSearchDatabase::GetPoseIndexFromTime(float Time, const FPoseSearchIndexAsset* SearchIndexAsset) const
 {
 	FFloatInterval Range = SearchIndexAsset->SamplingInterval;
@@ -1173,7 +1199,7 @@ int32 UPoseSearchDatabase::GetPoseIndexFromTime(float Time, const FPoseSearchInd
 
 float UPoseSearchDatabase::GetAssetTime(int32 PoseIdx, const FPoseSearchIndexAsset* Asset) const
 {
-	float AssetTime = SearchIndex.GetAssetTime(PoseIdx, Asset);
+	float AssetTime = GetSearchIndex()->GetAssetTime(PoseIdx, Asset);
 	return AssetTime;
 }
 
@@ -1272,7 +1298,8 @@ bool UPoseSearchDatabase::IsValidForIndexing() const
 
 bool UPoseSearchDatabase::IsValidForSearch() const
 {
-	return IsValidForIndexing() && SearchIndex.IsValid();
+	const FPoseSearchIndex* SearchIndex = GetSearchIndex();
+	return IsValidForIndexing() && SearchIndex && SearchIndex->IsValid();
 }
 
 void UPoseSearchDatabase::CollectSimpleSequences()
@@ -1441,9 +1468,9 @@ static FVector BlendParameterForSampleRanges(
 		0.0f);
 }
 
-bool UPoseSearchDatabase::TryInitSearchIndexAssets()
+bool UPoseSearchDatabase::TryInitSearchIndexAssets(FPoseSearchIndex& OutSearchIndex)
 {
-	SearchIndex.Assets.Empty();
+	OutSearchIndex.Assets.Empty();
 	
 	bool bAnyMirrored = false;
 	
@@ -1477,7 +1504,7 @@ bool UPoseSearchDatabase::TryInitSearchIndexAssets()
 			{
 				if (bAddUnmirrored)
 				{
-					SearchIndex.Assets.Add(
+					OutSearchIndex.Assets.Add(
 						FPoseSearchIndexAsset(
 							ESearchIndexAssetType::Sequence,
 							GroupIndex,
@@ -1488,7 +1515,7 @@ bool UPoseSearchDatabase::TryInitSearchIndexAssets()
 
 				if (bAddMirrored)
 				{
-					SearchIndex.Assets.Add(
+					OutSearchIndex.Assets.Add(
 						FPoseSearchIndexAsset(
 							ESearchIndexAssetType::Sequence,
 							GroupIndex,
@@ -1554,7 +1581,7 @@ bool UPoseSearchDatabase::TryInitSearchIndexAssets()
 
 					if (bAddUnmirrored)
 					{
-						SearchIndex.Assets.Add(
+						OutSearchIndex.Assets.Add(
 							FPoseSearchIndexAsset(
 								ESearchIndexAssetType::BlendSpace,
 								GroupIndex,
@@ -1566,7 +1593,7 @@ bool UPoseSearchDatabase::TryInitSearchIndexAssets()
 
 					if (bAddMirrored)
 					{
-						SearchIndex.Assets.Add(
+						OutSearchIndex.Assets.Add(
 							FPoseSearchIndexAsset(
 								ESearchIndexAssetType::BlendSpace,
 								GroupIndex,
@@ -1590,7 +1617,7 @@ bool UPoseSearchDatabase::TryInitSearchIndexAssets()
 			TEXT("Database %s is asking for mirrored sequences but MirrorDataBase is null in %s"),
 			*GetNameSafe(this), 
 			*GetNameSafe(Schema));
-		SearchIndex.Assets.Empty();
+		OutSearchIndex.Assets.Empty();
 		return false;
 	}
 
@@ -1617,21 +1644,92 @@ bool UPoseSearchDatabase::TryInitSearchIndexAssets()
 	return true;
 }
 
-void UPoseSearchDatabase::PreSave(FObjectPreSaveContext ObjectSaveContext)
+void UPoseSearchDatabase::PostLoad()
 {
-	SearchIndex.Reset();
+#if WITH_EDITOR
+	// todo: might want to have a different check to see if we need the data or not
+	if (FApp::CanEverRender())
+	{
+		if (!PrivateDerivedData)
+		{
+			BeginCacheDerivedData();
+		}
+	}
+#endif
+
+	Super::PostLoad();
+}
 
 #if WITH_EDITOR
-	if (!IsTemplate())
+void UPoseSearchDatabase::BeginCacheDerivedData()
+{
+	bool bPerformCache = true;
+			
+	using namespace UE::DerivedData;
+	if (PrivateDerivedData)
+	{
+		const FIoHash ExistingDerivedDataHash = PrivateDerivedData->FetchOrBuildDerivedDataKey;
+		if (!ExistingDerivedDataHash.IsZero())
+		{
+			if (ExistingDerivedDataHash == UE::PoseSearch::FPoseSearchDatabaseAsyncCacheTask::CreateKey(*this))
+			{
+				bPerformCache = false;
+			}
+		}
+	}
+
+	if (bPerformCache)
+	{
+		if (!PrivateDerivedData)
+		{
+			PrivateDerivedData = new FPoseSearchDatabaseDerivedData();
+		}
+
+		PrivateDerivedData->Cache(*this, false);
+	}
+}
+#endif // WITH_EDITOR
+
+
+void UPoseSearchDatabase::PostSaveRoot(FObjectPostSaveRootContext ObjectSaveContext)
+{
+	Super::PostSaveRoot(ObjectSaveContext);
+
+#if WITH_EDITOR
+	if (!IsTemplate() && !ObjectSaveContext.IsProceduralSave())
 	{
 		if (IsValidForIndexing())
 		{
-			UE::PoseSearch::BuildIndex(this);
+			if (!PrivateDerivedData)
+			{
+				PrivateDerivedData = new FPoseSearchDatabaseDerivedData();
+			}
+
+			PrivateDerivedData->Cache(*this, true);
 		}
 	}
 #endif
+}
 
-	Super::PreSave(ObjectSaveContext);
+void UPoseSearchDatabase::Serialize(FArchive& Ar)
+{
+	using namespace UE::PoseSearch;
+
+ 	Super::Serialize(Ar);
+
+	if (Ar.IsFilterEditorOnly())
+	{
+		check(Ar.IsLoading() || (Ar.IsCooking() && IsDerivedDataValid()));
+		const FPoseSearchIndex* SearchIndex = GetSearchIndex();
+		Ar << *SearchIndex;
+	}
+}
+
+bool UPoseSearchDatabase::IsDerivedDataValid()
+{
+	const FPoseSearchIndex* SearchIndex = GetSearchIndex();
+	bool bIsValid = SearchIndex && SearchIndex->IsValid();
+	return bIsValid;
 }
 
 #if WITH_EDITOR
@@ -1655,6 +1753,36 @@ void UPoseSearchDatabase::PostEditChangeProperty(struct FPropertyChangedEvent& P
 		}
 	}
 }
+
+void UPoseSearchDatabase::BeginCacheForCookedPlatformData(const ITargetPlatform* TargetPlatform)
+{
+	Super::BeginCacheForCookedPlatformData(TargetPlatform);
+
+	BeginCacheDerivedData();
+}
+
+bool UPoseSearchDatabase::IsCachedCookedPlatformDataLoaded(const ITargetPlatform* TargetPlatform)
+{
+	if (!PrivateDerivedData)
+	{
+		PrivateDerivedData = new FPoseSearchDatabaseDerivedData();
+		PrivateDerivedData->Cache(*this, true);
+		return false;
+	}
+
+	if (PrivateDerivedData->AsyncTask && PrivateDerivedData->AsyncTask->Poll())
+	{
+		PrivateDerivedData->FinishCache();
+	}
+
+	if (PrivateDerivedData->AsyncTask)
+	{
+		return false;
+	}
+
+	return true;
+}
+
 #endif // WITH_EDITOR
 
 
@@ -2244,7 +2372,7 @@ const FPoseSearchIndex* FDebugDrawParams::GetSearchIndex() const
 {
 	if (Database)
 	{
-		return &Database->SearchIndex;
+		return Database->GetSearchIndex();
 	}
 
 	if (SequenceMetaData)
@@ -2284,7 +2412,7 @@ void FSearchContext::SetSource(const UPoseSearchDatabase* InSourceDatabase)
 	{
 		if (ensure(SourceDatabase->IsValidForSearch()))
 		{
-			SearchIndex = &SourceDatabase->SearchIndex;
+			SearchIndex = SourceDatabase->GetSearchIndex();
 			DebugDrawParams.Database = SourceDatabase;
 			MirrorMismatchCost = SourceDatabase->MirroringMismatchCost;
 		}
@@ -4637,7 +4765,7 @@ bool BuildIndex(const UAnimSequence* Sequence, UPoseSearchSequenceMetaData* Sequ
 	return true;
 }
 
-bool BuildIndex(UPoseSearchDatabase* Database)
+bool BuildIndex(UPoseSearchDatabase* Database, FPoseSearchIndex& OutSearchIndex)
 {
 	check(Database);
 
@@ -4646,7 +4774,9 @@ bool BuildIndex(UPoseSearchDatabase* Database)
 		return false;
 	}
 
-	if (!Database->TryInitSearchIndexAssets())
+	OutSearchIndex.Schema = Database->Schema;
+
+	if (!Database->TryInitSearchIndexAssets(OutSearchIndex))
 	{
 		return false;
 	}
@@ -4752,7 +4882,7 @@ bool BuildIndex(UPoseSearchDatabase* Database)
 	// Prepare sequences for indexing
 
 	TArray<FIndexer> Indexers;
-	Indexers.Reserve(Database->SearchIndex.Assets.Num());
+	Indexers.Reserve(Database->GetSearchIndex()->Assets.Num());
 
 	auto GetSequenceSampler = [&](const UAnimSequence* Sequence) -> const FSequenceSampler*
 	{
@@ -4763,10 +4893,10 @@ bool BuildIndex(UPoseSearchDatabase* Database)
 	{
 		return BlendSpace ? &BlendSpaceSamplers[BlendSpaceSamplerMap[{BlendSpace, BlendParameters}]] : nullptr;
 	};
-
-	for (int32 AssetIdx = 0; AssetIdx != Database->SearchIndex.Assets.Num(); ++AssetIdx)
+	Indexers.Reserve(Database->GetSearchIndex()->Assets.Num());
+	for (int32 AssetIdx = 0; AssetIdx != OutSearchIndex.Assets.Num(); ++AssetIdx)
 	{
-		const FPoseSearchIndexAsset& SearchIndexAsset = Database->SearchIndex.Assets[AssetIdx];
+		const FPoseSearchIndexAsset& SearchIndexAsset = OutSearchIndex.Assets[AssetIdx];
 
 		FIndexer::FInput Input;
 		Input.SamplingContext = &SamplingContext;
@@ -4805,11 +4935,11 @@ bool BuildIndex(UPoseSearchDatabase* Database)
 	// Write index info to sequence and count up total poses and storage required
 	int32 TotalPoses = 0;
 	int32 TotalFloats = 0;
-	for (int32 AssetIdx = 0; AssetIdx != Database->SearchIndex.Assets.Num(); ++AssetIdx)
+	for (int32 AssetIdx = 0; AssetIdx != OutSearchIndex.Assets.Num(); ++AssetIdx)
 	{
 		const FIndexer::FOutput& Output = Indexers[AssetIdx].Output;
 
-		FPoseSearchIndexAsset& SearchIndexAsset = Database->SearchIndex.Assets[AssetIdx];
+		FPoseSearchIndexAsset& SearchIndexAsset = OutSearchIndex.Assets[AssetIdx];
 		SearchIndexAsset.NumPoses = Output.NumIndexedPoses;
 		SearchIndexAsset.FirstPoseIdx = TotalPoses;
 
@@ -4818,20 +4948,20 @@ bool BuildIndex(UPoseSearchDatabase* Database)
 	}
 
 	// Join animation data into a single search index
-	Database->SearchIndex.Values.Reset(TotalFloats);
-	Database->SearchIndex.PoseMetadata.Reset(TotalPoses);
+	OutSearchIndex.Values.Reset(TotalFloats);
+	OutSearchIndex.PoseMetadata.Reset(TotalPoses);
 
 	for (const FIndexer& Indexer : Indexers)
 	{
 		const FIndexer::FOutput& Output = Indexer.Output;
-		Database->SearchIndex.Values.Append(Output.FeatureVectorTable.GetData(), Output.FeatureVectorTable.Num());
-		Database->SearchIndex.PoseMetadata.Append(Output.PoseMetadata);
+		OutSearchIndex.Values.Append(Output.FeatureVectorTable.GetData(), Output.FeatureVectorTable.Num());
+		OutSearchIndex.PoseMetadata.Append(Output.PoseMetadata);
 	}
 
-	Database->SearchIndex.NumPoses = TotalPoses;
-	Database->SearchIndex.Schema = Database->Schema;
+	OutSearchIndex.NumPoses = TotalPoses;
+	OutSearchIndex.Schema = Database->Schema;
 
-	PreprocessSearchIndex(&Database->SearchIndex);
+	PreprocessSearchIndex(&OutSearchIndex);
 
 	return true;
 }
@@ -5064,6 +5194,11 @@ public: // IModuleInterface
 
 public: // IPoseSearchProvider
 	virtual UE::Anim::IPoseSearchProvider::FSearchResult Search(const FAnimationBaseContext& GraphContext, const UAnimSequenceBase* Sequence) override;
+
+private:
+#if WITH_EDITOR
+	void OnObjectSaved(UObject* SavedObject, FObjectPreSaveContext SaveContext);
+#endif // WITH_EDITOR
 };
 
 void FModule::StartupModule()
@@ -5074,6 +5209,10 @@ void FModule::StartupModule()
 	// Enable the PoseSearch trace channel
 	UE::Trace::ToggleChannel(*FTraceLogger::Name.ToString(), true);
 #endif
+
+#if WITH_EDITOR
+	FCoreUObjectDelegates::OnObjectPreSave.AddRaw(this, &FModule::OnObjectSaved);
+#endif // WITH_EDITOR
 }
 
 void FModule::ShutdownModule()
@@ -5119,6 +5258,95 @@ UE::Anim::IPoseSearchProvider::FSearchResult FModule::Search(const FAnimationBas
 	return ProviderResult;
 }
 
+#if WITH_EDITOR
+
+void GetPoseSearchDatabaseAssetDataList(TArray<FAssetData>& OutPoseSearchDatabaseAssetDataList)
+{
+	FAssetRegistryModule& AssetRegistryModule =
+		FModuleManager::LoadModuleChecked<FAssetRegistryModule>("AssetRegistry");
+	FARFilter Filter;
+	Filter.bRecursiveClasses = true;
+	Filter.ClassNames.Add(UPoseSearchDatabase::StaticClass()->GetFName());
+
+	OutPoseSearchDatabaseAssetDataList.Reset();
+	AssetRegistryModule.Get().GetAssets(Filter, OutPoseSearchDatabaseAssetDataList);
+}
+
+void ForEachPoseSearchDatabase(bool bLoadAssets, TFunctionRef<void(UPoseSearchDatabase&)> InFunction)
+{
+	TArray<FAssetData> PoseSearchDatabaseAssetDataList;
+	GetPoseSearchDatabaseAssetDataList(PoseSearchDatabaseAssetDataList);
+	for (const auto& PoseSearchDbAssetData : PoseSearchDatabaseAssetDataList)
+	{
+		if (UPoseSearchDatabase* PoseSearchDb = 
+			Cast<UPoseSearchDatabase>(PoseSearchDbAssetData.FastGetAsset(bLoadAssets)))
+		{
+			InFunction(*PoseSearchDb);
+		}
+	}
+}
+
+void FModule::OnObjectSaved(UObject* SavedObject, FObjectPreSaveContext SaveContext)
+{
+	if (UAnimSequence* SavedSequence = Cast<UAnimSequence>(SavedObject))
+	{
+		ForEachPoseSearchDatabase(false, [SavedSequence](UPoseSearchDatabase& PoseSearchDb)
+		{
+			bool bSequenceFound =
+				PoseSearchDb.Sequences.ContainsByPredicate([SavedSequence](FPoseSearchDatabaseSequence& DbSequence)
+			{
+				bool bIsMatch =
+					SavedSequence == DbSequence.Sequence ||
+					SavedSequence == DbSequence.LeadInSequence ||
+					SavedSequence == DbSequence.FollowUpSequence;
+				return bIsMatch;
+			});
+
+			if (bSequenceFound)
+			{
+				PoseSearchDb.BeginCacheDerivedData();
+			}
+		});
+	}
+	else if (UBlendSpace* SavedBlendSpace = Cast<UBlendSpace>(SavedObject))
+	{
+		ForEachPoseSearchDatabase(false, [SavedBlendSpace](UPoseSearchDatabase& PoseSearchDb)
+		{
+			bool bBlendSpaceFound = PoseSearchDb.BlendSpaces.ContainsByPredicate(
+					[SavedBlendSpace](FPoseSearchDatabaseBlendSpace& DbBlendSpace)
+			{
+				return SavedBlendSpace == DbBlendSpace.BlendSpace;
+			});
+
+			if (bBlendSpaceFound)
+			{
+				PoseSearchDb.BeginCacheDerivedData();
+			}
+		});
+	}
+	else if (UPoseSearchSchema* SavedSchema = Cast<UPoseSearchSchema>(SavedObject))
+	{
+		ForEachPoseSearchDatabase(false, [SavedSchema](UPoseSearchDatabase& PoseSearchDb)
+		{
+			if (PoseSearchDb.Schema == SavedSchema)
+			{
+				PoseSearchDb.BeginCacheDerivedData();
+			}
+		});
+	}
+	else if (USkeleton* SavedSkeleton = Cast<USkeleton>(SavedObject))
+	{
+		ForEachPoseSearchDatabase(false, [SavedSkeleton](UPoseSearchDatabase& PoseSearchDb)
+		{
+			if (PoseSearchDb.Schema && PoseSearchDb.Schema->Skeleton == SavedSkeleton)
+			{
+				PoseSearchDb.BeginCacheDerivedData();
+			}
+		});
+	}
+}
+#endif // WITH_EDITOR
+
 }} // namespace UE::PoseSearch
 
 IMPLEMENT_MODULE(UE::PoseSearch::FModule, PoseSearch)
\ No newline at end of file
diff --git a/Engine/Plugins/Experimental/Animation/PoseSearch/Source/Runtime/Private/PoseSearchDerivedData.cpp b/Engine/Plugins/Experimental/Animation/PoseSearch/Source/Runtime/Private/PoseSearchDerivedData.cpp
new file mode 100644
--- /dev/null
+++ b/Engine/Plugins/Experimental/Animation/PoseSearch/Source/Runtime/Private/PoseSearchDerivedData.cpp
@@ -0,0 +1,408 @@
+// Copyright Epic Games, Inc. All Rights Reserved.
+
+#include "PoseSearchDerivedData.h"
+#include "PoseSearch/PoseSearchAnimNotifies.h"
+
+#if WITH_EDITOR
+#include "Serialization/BulkDataRegistry.h"
+#endif // WITH_EDITOR
+
+#if WITH_EDITOR
+namespace UE::PoseSearch
+{
+	const UE::DerivedData::FValueId FPoseSearchDatabaseAsyncCacheTask::Id = UE::DerivedData::FValueId::FromName("Data");
+	const UE::DerivedData::FCacheBucket FPoseSearchDatabaseAsyncCacheTask::Bucket("PoseSearchDatabase");
+}
+#endif // WITH_EDITOR
+
+#if WITH_EDITOR
+
+void FPoseSearchDatabaseDerivedData::Cache(UPoseSearchDatabase& Database, bool bForceRebuild)
+{
+	CancelCache();
+	if (Database.IsValidForIndexing())
+	{
+		CreateDatabaseBuildTask(Database, bForceRebuild);
+	}
+}
+
+void FPoseSearchDatabaseDerivedData::CancelCache()
+{
+	if (AsyncTask)
+	{
+		AsyncTask->Cancel();
+	}
+
+	FinishCache();
+}
+
+void FPoseSearchDatabaseDerivedData::FinishCache()
+{
+	if (AsyncTask)
+	{
+		AsyncTask->Wait();
+		delete AsyncTask;
+		AsyncTask = nullptr;
+	}
+}
+
+void FPoseSearchDatabaseDerivedData::CreateDatabaseBuildTask(UPoseSearchDatabase& Database, bool bForceRebuild)
+{
+	AsyncTask = new UE::PoseSearch::FPoseSearchDatabaseAsyncCacheTask(Database, *this, bForceRebuild);
+}
+
+#endif // WITH_EDITOR
+
+namespace UE::PoseSearch
+{
+
+#if WITH_EDITOR
+
+	FPoseSearchDatabaseAsyncCacheTask::FPoseSearchDatabaseAsyncCacheTask(
+		UPoseSearchDatabase& InDatabase,
+		FPoseSearchDatabaseDerivedData& InDerivedData,
+		bool bForceRebuild)
+		: Owner(UE::DerivedData::EPriority::Normal)
+		, DerivedData(InDerivedData)
+		, Database(InDatabase)
+	{
+		using namespace UE::DerivedData;
+
+		FIoHash DerivedDataKey = CreateKey(Database);
+		DerivedData.FetchOrBuildDerivedDataKey = DerivedDataKey;
+
+		if (bForceRebuild)
+		{
+			DerivedData.DerivedDataKey = { Bucket, DerivedDataKey };
+			BuildAndWrite();
+		}
+		else
+		{
+			BeginCache();
+		}
+	}
+
+	bool FPoseSearchDatabaseAsyncCacheTask::Cancel()
+	{
+		Owner.Cancel();
+		return true;
+	}
+
+	void FPoseSearchDatabaseAsyncCacheTask::Wait()
+	{
+		Owner.Wait();
+	}
+
+	bool FPoseSearchDatabaseAsyncCacheTask::Poll() const
+	{
+		return Owner.Poll();
+	}
+
+
+	void FPoseSearchDatabaseAsyncCacheTask::BeginCache()
+	{
+		using namespace UE::DerivedData;
+
+		TArray<FCacheGetRequest> CacheRequests;
+		ECachePolicy CachePolicy = ECachePolicy::Default;
+		const FCacheKey CacheKey{ Bucket, DerivedData.FetchOrBuildDerivedDataKey };
+		CacheRequests.Add({ {Database.GetPathName()}, CacheKey, CachePolicy });
+		GetCache().Get(CacheRequests, Owner, [this](FCacheGetResponse&& Response)
+		{
+			OnGetComplete(MoveTemp(Response));
+		});
+	}
+
+	void FPoseSearchDatabaseAsyncCacheTask::OnGetComplete(UE::DerivedData::FCacheGetResponse&& Response)
+	{
+		using namespace UE::DerivedData;
+
+		DerivedData.DerivedDataKey = Response.Record.GetKey();
+		if (Response.Status == EStatus::Ok)
+		{
+			BuildIndexFromCacheRecord(MoveTemp(Response.Record));
+		}
+		else if (Response.Status == EStatus::Error)
+		{
+			BuildAndWrite();
+		}
+	}
+
+	void FPoseSearchDatabaseAsyncCacheTask::BuildAndWrite()
+	{
+		GetRequestOwner().LaunchTask(TEXT("PoseSearchDatabaseBuild"), [this]
+		{
+			if (!GetRequestOwner().IsCanceled())
+			{
+				DerivedData.SearchIndex.Reset();
+				DerivedData.SearchIndex.Schema = Database.Schema;
+				const bool bIndexReady = BuildIndex(&Database, DerivedData.SearchIndex);
+
+				WriteIndexToCache();
+			}
+		});
+	}
+
+	void FPoseSearchDatabaseAsyncCacheTask::WriteIndexToCache()
+	{
+		using namespace UE::DerivedData;
+
+		TArray<uint8> RawBytes;
+		FMemoryWriter Writer(RawBytes);
+		Writer << DerivedData.SearchIndex;
+		FSharedBuffer RawData = MakeSharedBufferFromArray(MoveTemp(RawBytes));
+
+		FCacheRecordBuilder Builder(DerivedData.DerivedDataKey);
+		Builder.AddValue(Id, RawData);
+
+		Owner.KeepAlive();
+		GetCache().Put({ { { Database.GetPathName() }, Builder.Build() } }, Owner);
+	}
+
+	void FPoseSearchDatabaseAsyncCacheTask::BuildIndexFromCacheRecord(UE::DerivedData::FCacheRecord&& CacheRecord)
+	{
+		using namespace UE::DerivedData;
+
+		DerivedData.SearchIndex.Reset();
+		DerivedData.SearchIndex.Schema = Database.Schema;
+
+		FSharedBuffer RawData = CacheRecord.GetValue(Id).GetData().Decompress();
+		FMemoryReaderView Reader(RawData);
+		Reader << DerivedData.SearchIndex;
+	}
+
+	FIoHash FPoseSearchDatabaseAsyncCacheTask::CreateKey(const UPoseSearchDatabase& Database)
+	{
+		FBlake3 Writer;
+
+		const FGuid VersionGuid =  FDevSystemGuids::GetSystemGuid(FDevSystemGuids::Get().POSESEARCHDB_DERIVEDDATA_VER);
+		Writer.Update(MakeMemoryView(VersionGuid.ToString()));
+
+		AddPoseSearchSchemaToWriter(Database.Schema, Writer);
+
+		for (int32 i = 0; i < Database.Sequences.Num(); ++i)
+		{
+			AddDbSequenceToWriter(Database.Sequences[i], Writer);
+		}
+
+		for (int32 i = 0; i < Database.BlendSpaces.Num(); ++i)
+		{
+			AddDbBlendSpaceToWriter(Database.BlendSpaces[i], Writer);
+		}
+
+		return Writer.Finalize();
+	}
+
+	void FPoseSearchDatabaseAsyncCacheTask::AddPoseSearchSchemaToWriter(
+		const UPoseSearchSchema* Schema, 
+		FBlake3& InOutWriter)
+	{
+		if (IsValid(Schema))
+		{
+			InOutWriter.Update(&Schema->bUseTrajectoryVelocities, sizeof(Schema->bUseTrajectoryVelocities));
+			InOutWriter.Update(&Schema->bUseTrajectoryPositions, sizeof(Schema->bUseTrajectoryPositions));
+			InOutWriter.Update(&Schema->bUseTrajectoryForwardVectors, sizeof(Schema->bUseTrajectoryForwardVectors));
+			InOutWriter.Update(MakeMemoryView(Schema->SampledBones));
+			InOutWriter.Update(MakeMemoryView(Schema->PoseSampleTimes));
+			InOutWriter.Update(MakeMemoryView(Schema->TrajectorySampleTimes));
+			InOutWriter.Update(MakeMemoryView(Schema->TrajectorySampleDistances));
+			InOutWriter.Update(&Schema->DataPreprocessor, sizeof(Schema->DataPreprocessor));
+			InOutWriter.Update(&Schema->EffectiveDataPreprocessor, sizeof(Schema->EffectiveDataPreprocessor));
+			InOutWriter.Update(&Schema->SamplingInterval, sizeof(Schema->SamplingInterval));
+			InOutWriter.Update(MakeMemoryView(Schema->BoneIndicesWithParents));
+		}
+	}
+
+	void FPoseSearchDatabaseAsyncCacheTask::AddDbSequenceToWriter(
+		const FPoseSearchDatabaseSequence& DbSequence, 
+		FBlake3& InOutWriter)
+	{
+		// Main Sequence
+		AddRawSequenceToWriter(DbSequence.Sequence, InOutWriter);
+		InOutWriter.Update(&DbSequence.SamplingRange, sizeof(DbSequence.SamplingRange));
+		InOutWriter.Update(&DbSequence.bLoopAnimation, sizeof(DbSequence.bLoopAnimation));
+		InOutWriter.Update(&DbSequence.MirrorOption, sizeof(DbSequence.MirrorOption));
+
+		// Lead in sequence
+		AddRawSequenceToWriter(DbSequence.LeadInSequence, InOutWriter);
+		InOutWriter.Update(&DbSequence.bLoopLeadInAnimation, sizeof(DbSequence.bLoopLeadInAnimation));
+
+		// Follow up sequence
+		AddRawSequenceToWriter(DbSequence.FollowUpSequence, InOutWriter);
+		InOutWriter.Update(&DbSequence.bLoopFollowUpAnimation, sizeof(DbSequence.bLoopFollowUpAnimation));
+
+		// Tags
+		InOutWriter.Update(&DbSequence.GroupTags, sizeof(DbSequence.GroupTags));
+
+		// Notifies
+		AddPoseSearchNotifiesToWriter(DbSequence.Sequence, InOutWriter);
+	}
+
+	void FPoseSearchDatabaseAsyncCacheTask::AddRawSequenceToWriter(
+		const UAnimSequence* Sequence, 
+		FBlake3& InOutWriter)
+	{
+		if (Sequence)
+		{
+			InOutWriter.Update(MakeMemoryView(Sequence->GetName()));
+			InOutWriter.Update(MakeMemoryView(Sequence->GetRawDataGuid().ToString()));
+		}
+	}
+
+	void FPoseSearchDatabaseAsyncCacheTask::AddPoseSearchNotifiesToWriter(
+		const UAnimSequence* Sequence, 
+		FBlake3& InOutWriter)
+	{
+		if (!Sequence)
+		{
+			return;
+		}
+
+		FAnimNotifyContext NotifyContext;
+		Sequence->GetAnimNotifies(0.0f, Sequence->GetPlayLength(), NotifyContext);
+
+		for (const FAnimNotifyEventReference& EventReference : NotifyContext.ActiveNotifies)
+		{
+			const FAnimNotifyEvent* NotifyEvent = EventReference.GetNotify();
+			if (!NotifyEvent)
+			{
+				continue;
+			}
+
+			if (NotifyEvent->NotifyStateClass->IsA<UAnimNotifyState_PoseSearchBase>())
+			{
+				const float StartTime = NotifyEvent->GetTriggerTime();
+				const float EndTime = NotifyEvent->GetEndTriggerTime();
+				InOutWriter.Update(&StartTime, sizeof(float));
+				InOutWriter.Update(&EndTime, sizeof(float));
+
+				if (const auto ModifyCostNotifyState =
+					Cast<const UAnimNotifyState_PoseSearchModifyCost>(NotifyEvent->NotifyStateClass))
+				{
+					InOutWriter.Update(&ModifyCostNotifyState->CostAddend, sizeof(float));
+				}
+			}
+		}
+	}
+
+	void FPoseSearchDatabaseAsyncCacheTask::AddDbBlendSpaceToWriter(
+		const FPoseSearchDatabaseBlendSpace& DbBlendSpace, 
+		FBlake3& InOutWriter)
+	{
+		if (IsValid(DbBlendSpace.BlendSpace))
+		{
+			const TArray<FBlendSample>& BlendSpaceSamples = DbBlendSpace.BlendSpace->GetBlendSamples();
+			for (const FBlendSample& Sample : BlendSpaceSamples)
+			{
+				AddRawSequenceToWriter(Sample.Animation, InOutWriter);
+				InOutWriter.Update(&Sample.SampleValue, sizeof(Sample.SampleValue));
+				InOutWriter.Update(&Sample.RateScale, sizeof(Sample.RateScale));
+			}
+
+			InOutWriter.Update(&DbBlendSpace.bLoopAnimation, sizeof(DbBlendSpace.bLoopAnimation));
+			InOutWriter.Update(&DbBlendSpace.MirrorOption, sizeof(DbBlendSpace.MirrorOption));
+			InOutWriter.Update(&DbBlendSpace.bUseGridForSampling, sizeof(DbBlendSpace.bUseGridForSampling));
+			InOutWriter.Update(&DbBlendSpace.NumberOfHorizontalSamples, sizeof(DbBlendSpace.NumberOfHorizontalSamples));
+			InOutWriter.Update(&DbBlendSpace.NumberOfVerticalSamples, sizeof(DbBlendSpace.NumberOfVerticalSamples));
+			InOutWriter.Update(&DbBlendSpace.GroupTags, sizeof(DbBlendSpace.GroupTags));
+		}
+	}
+
+#endif // WITH_EDITOR
+
+	FArchive& operator<<(FArchive& Ar, FPoseSearchIndexPreprocessInfo& Info)
+	{
+		int32 NumTransformationMatrix = 0;
+
+		if (Ar.IsSaving())
+		{
+			NumTransformationMatrix = Info.TransformationMatrix.Num();
+		}
+
+		Ar << Info.NumDimensions;
+		Ar << NumTransformationMatrix;
+
+		if (Ar.IsLoading())
+		{
+			Info.TransformationMatrix.SetNumUninitialized(NumTransformationMatrix);
+			Info.InverseTransformationMatrix.SetNumUninitialized(NumTransformationMatrix);
+			Info.SampleMean.SetNumUninitialized(Info.NumDimensions);
+		}
+
+		if (Info.TransformationMatrix.Num() > 0)
+		{
+			Ar.Serialize(
+				&Info.TransformationMatrix[0],
+				Info.TransformationMatrix.Num() * Info.TransformationMatrix.GetTypeSize());
+		}
+
+		if (Info.InverseTransformationMatrix.Num() > 0)
+		{
+			Ar.Serialize(
+				&Info.InverseTransformationMatrix[0],
+				Info.InverseTransformationMatrix.Num() * Info.InverseTransformationMatrix.GetTypeSize());
+		}
+
+		if (Info.SampleMean.Num() > 0)
+		{
+			Ar.Serialize(&Info.SampleMean[0], Info.SampleMean.Num() * Info.SampleMean.GetTypeSize());
+		}
+
+		return Ar;
+	}
+
+	FArchive& operator<<(FArchive& Ar, const FPoseSearchIndexPreprocessInfo& Info)
+	{
+		check(Ar.IsSaving());
+		return (Ar << const_cast<FPoseSearchIndexPreprocessInfo&>(Info));
+	}
+
+
+	FArchive& operator<<(FArchive& Ar, FPoseSearchIndex& Index)
+	{
+		int32 NumValues = 0;
+		int32 NumAssets = 0;
+
+		if (Ar.IsSaving())
+		{
+			NumValues = Index.Values.Num();
+			NumAssets = Index.Assets.Num();
+		}
+
+		Ar << Index.NumPoses;
+		Ar << NumValues;
+		Ar << NumAssets;
+
+		if (Ar.IsLoading())
+		{
+			Index.Values.SetNumUninitialized(NumValues);
+			Index.PoseMetadata.SetNumUninitialized(Index.NumPoses);
+			Index.Assets.SetNumUninitialized(NumAssets);
+		}
+
+		if (Index.Values.Num() > 0)
+		{
+			Ar.Serialize(&Index.Values[0], Index.Values.Num() * Index.Values.GetTypeSize());
+		}
+
+		if (Index.PoseMetadata.Num() > 0)
+		{
+			Ar.Serialize(&Index.PoseMetadata[0], Index.PoseMetadata.Num() * Index.PoseMetadata.GetTypeSize());
+		}
+
+		if (Index.Assets.Num() > 0)
+		{
+			Ar.Serialize(&Index.Assets[0], Index.Assets.Num() * Index.Assets.GetTypeSize());
+		}
+
+		Ar << Index.PreprocessInfo;
+
+		return Ar;
+	}
+
+	FArchive& operator<<(FArchive& Ar, const FPoseSearchIndex& Index)
+	{
+		check(Ar.IsSaving());
+		return (Ar << const_cast<FPoseSearchIndex&>(Index));
+	}
+}
\ No newline at end of file
diff --git a/Engine/Plugins/Experimental/Animation/PoseSearch/Source/Runtime/Private/PoseSearchDerivedData.h b/Engine/Plugins/Experimental/Animation/PoseSearch/Source/Runtime/Private/PoseSearchDerivedData.h
new file mode 100644
--- /dev/null
+++ b/Engine/Plugins/Experimental/Animation/PoseSearch/Source/Runtime/Private/PoseSearchDerivedData.h
@@ -0,0 +1,114 @@
+// Copyright Epic Games, Inc. All Rights Reserved.
+
+#pragma once
+
+#include "PoseSearch/PoseSearch.h"
+
+#if WITH_EDITORONLY_DATA
+
+#include "DerivedDataCache.h"
+#include "DerivedDataRequestOwner.h"
+
+#endif // WITH_EDITORONLY_DATA
+
+#include "PoseSearchDerivedData.generated.h"
+
+namespace UE::DerivedData
+{
+	struct FCacheGetResponse;
+	class FCacheRecord;
+}
+
+USTRUCT()
+struct FPoseSearchDatabaseDerivedData
+{
+	GENERATED_BODY()
+
+	FPoseSearchIndex SearchIndex;
+
+#if WITH_EDITORONLY_DATA
+	UE::DerivedData::FCacheKey DerivedDataKey;
+
+
+	/** This is the key for the FetchOrBuild variant of our Cache. We assume that uniqueness
+	*	for that is equivalent to uniqueness if we use both FetchFirst and FetchOrBuild. This
+	*	is used as the key in to CookedPlatformData, as well as to determine if we are already
+	*	cooking the data the editor needs in CachePlatformData. 
+	*	Note that since this is read on the game thread constantly in CachePlatformData, it
+	*	must be written to on the game thread to avoid false recaches.
+	*/
+	FIoHash FetchOrBuildDerivedDataKey;
+
+	/** Async cache task if one is outstanding. */
+	UE::PoseSearch::FPoseSearchDatabaseAsyncCacheTask* AsyncTask;
+
+
+#endif // WITH_EDITORONLY_DATA
+
+#if WITH_EDITOR
+	// Build search index and adds result to DDC
+	void Cache(UPoseSearchDatabase& Database, bool bForceRebuild);
+
+	void CancelCache();
+
+	void FinishCache();
+
+	void CreateDatabaseBuildTask(UPoseSearchDatabase& Database, bool bForceRebuild);
+#endif // WITH_EDITOR
+};
+
+namespace UE::PoseSearch
+{
+#if WITH_EDITOR
+	struct FPoseSearchDatabaseAsyncCacheTask
+	{
+
+	public:
+		FPoseSearchDatabaseAsyncCacheTask(
+			UPoseSearchDatabase& InDatabase,
+			FPoseSearchDatabaseDerivedData& InDerivedData,
+			bool bForceRebuild);
+
+		bool Cancel();
+		void Wait();
+		bool Poll() const;
+
+	private:
+		void BeginCache();
+		void OnGetComplete(UE::DerivedData::FCacheGetResponse&& Response);
+		void BuildAndWrite();
+		void WriteIndexToCache();
+		void BuildIndexFromCacheRecord(UE::DerivedData::FCacheRecord&& CacheRecord);
+
+		UE::DerivedData::IRequestOwner& GetRequestOwner() { return Owner; }
+
+	private:
+		UE::DerivedData::FRequestOwner Owner;
+		FPoseSearchDatabaseDerivedData& DerivedData;
+		UPoseSearchDatabase& Database;
+
+		static const UE::DerivedData::FValueId Id;
+		static const UE::DerivedData::FCacheBucket Bucket;
+
+	public:
+		static FIoHash CreateKey(const UPoseSearchDatabase& Database);
+		static void AddPoseSearchSchemaToWriter(const UPoseSearchSchema* Schema, FBlake3& InOutWriter);
+		static void AddDbSequenceToWriter(const FPoseSearchDatabaseSequence& DbSequence, FBlake3& InOutWriter);
+		static void AddRawSequenceToWriter(const UAnimSequence* Sequence, FBlake3& InOutWriter);
+		static void AddPoseSearchNotifiesToWriter(const UAnimSequence* Sequence, FBlake3& InOutWriter);
+		static void AddDbBlendSpaceToWriter(const FPoseSearchDatabaseBlendSpace& DbBlendSpace, FBlake3& InOutWriter);
+	};
+#endif // WITH_EDITOR
+
+	// Serialization for FPoseSearchIndexPreprocessInfo.
+	FArchive& operator<<(FArchive& Ar, FPoseSearchIndexPreprocessInfo& Info);
+
+	// Serialization for const FPoseSearchIndexPreprocessInfo.
+	FArchive& operator<<(FArchive& Ar, const FPoseSearchIndexPreprocessInfo& Info);
+
+	// Serialization for FPoseSearchIndex.
+	FArchive& operator<<(FArchive& Ar, FPoseSearchIndex& Index);
+
+	// Serialization for const FPoseSearchIndex.
+	FArchive& operator<<(FArchive& Ar, const FPoseSearchIndex& Index);
+}
\ No newline at end of file
diff --git a/Engine/Plugins/Experimental/Animation/PoseSearch/Source/Runtime/Private/PoseSearchLibrary.cpp b/Engine/Plugins/Experimental/Animation/PoseSearch/Source/Runtime/Private/PoseSearchLibrary.cpp
--- a/Engine/Plugins/Experimental/Animation/PoseSearch/Source/Runtime/Private/PoseSearchLibrary.cpp
+++ b/Engine/Plugins/Experimental/Animation/PoseSearch/Source/Runtime/Private/PoseSearchLibrary.cpp
@@ -64,7 +64,7 @@ static void ComputeDatabaseBlendSpaceFilter(
 
 void FMotionMatchingPoseStepper::Update(const FAnimationUpdateContext& UpdateContext, const FMotionMatchingState& State)
 {
-	if (State.DbPoseIdx == INDEX_NONE)
+	if (State.DbPoseIdx == INDEX_NONE || !State.CurrentDatabase.IsValid() || !(State.CurrentDatabase->GetSearchIndex()))
 	{
 		return;
 	}
@@ -108,7 +108,7 @@ void FMotionMatchingPoseStepper::Update(const FAnimationUpdateContext& UpdateCon
 					return Entry.Sequence == DbSequence.FollowUpSequence;
 				});
 
-			int32 FollowUpSearchIndexAssetIdx = State.CurrentDatabase->SearchIndex.Assets.IndexOfByPredicate(
+			int32 FollowUpSearchIndexAssetIdx = State.CurrentDatabase->GetSearchIndex()->Assets.IndexOfByPredicate(
 				[&](const FPoseSearchIndexAsset& Entry)
 				{
 					const bool bIsMatch =
@@ -121,7 +121,7 @@ void FMotionMatchingPoseStepper::Update(const FAnimationUpdateContext& UpdateCon
 			if (FollowUpSearchIndexAssetIdx != INDEX_NONE)
 			{
 				const FPoseSearchIndexAsset* FollowUpSearchIndexAsset =
-					&State.CurrentDatabase->SearchIndex.Assets[FollowUpSearchIndexAssetIdx];
+				&State.CurrentDatabase->GetSearchIndex()->Assets[FollowUpSearchIndexAssetIdx];
 				const float FollowUpAssetTime = State.AssetPlayerTime + DeltaTime - AssetLength;
 				const int32 FollowUpPoseIdx = State.CurrentDatabase->GetPoseIndexFromTime(
 					FollowUpAssetTime,
@@ -232,7 +232,7 @@ void FMotionMatchingState::ComposeQuery(const UPoseSearchDatabase* Database, con
 		ComposedQuery.MergeReplace(Goal);
 	}
 
-	ComposedQuery.Normalize(Database->SearchIndex);
+	ComposedQuery.Normalize(*Database->GetSearchIndex());
 }
 
 static void RequestInertialBlend(const FAnimationUpdateContext& Context, float BlendTime)
@@ -253,7 +253,7 @@ void FMotionMatchingState::JumpToPose(const FAnimationUpdateContext& Context, co
 {
 	// Remember which pose and sequence we're playing from the database
 	DbPoseIdx = Result.PoseIdx;
-	SearchIndexAssetIdx = CurrentDatabase->SearchIndex.FindAssetIndex(Result.SearchIndexAsset);
+	SearchIndexAssetIdx = CurrentDatabase->GetSearchIndex()->FindAssetIndex(Result.SearchIndexAsset);
 
 	ElapsedPoseJumpTime = 0.0f;
 	AssetPlayerTime = Result.AssetTime;
@@ -300,14 +300,16 @@ void UpdateMotionMatchingState(
 	{
 		InOutMotionMatchingState.DbPoseIdx = PoseStepper.Result.PoseIdx;
 		InOutMotionMatchingState.SearchIndexAssetIdx = 
-			InOutMotionMatchingState.CurrentDatabase->SearchIndex.FindAssetIndex(PoseStepper.Result.SearchIndexAsset);
+			InOutMotionMatchingState.CurrentDatabase->GetSearchIndex()->FindAssetIndex(PoseStepper.Result.SearchIndexAsset);
 	}
 
 	// Build the search query
 	if (InOutMotionMatchingState.DbPoseIdx != INDEX_NONE)
 	{
 		// Copy search query directly from the database if we have an active pose
-		InOutMotionMatchingState.ComposedQuery.CopyFromSearchIndex(Database->SearchIndex, InOutMotionMatchingState.DbPoseIdx);
+		InOutMotionMatchingState.ComposedQuery.CopyFromSearchIndex(
+			*Database->GetSearchIndex(), 
+			InOutMotionMatchingState.DbPoseIdx);
 	}
 	else
 	{
@@ -355,7 +357,7 @@ void UpdateMotionMatchingState(
 		FPoseCost CurrentPoseCost;
 		if (InOutMotionMatchingState.DbPoseIdx != INDEX_NONE)
 		{
-			const FPoseSearchIndexAsset* SearchIndexAsset = &Database->SearchIndex.Assets[InOutMotionMatchingState.SearchIndexAssetIdx];
+			const FPoseSearchIndexAsset* SearchIndexAsset = &Database->GetSearchIndex()->Assets[InOutMotionMatchingState.SearchIndexAssetIdx];
 			CurrentPoseCost = ComparePoses(InOutMotionMatchingState.DbPoseIdx, SearchContext, SearchIndexAsset->SourceGroupIdx);
 		}
 
@@ -492,13 +494,14 @@ void UpdateMotionMatchingState(
 
 const FPoseSearchIndexAsset* FMotionMatchingState::GetCurrentSearchIndexAsset() const
 {
-	if (!CurrentDatabase.IsValid() ||
-		!CurrentDatabase->SearchIndex.Assets.IsValidIndex(SearchIndexAssetIdx))
+	if (!CurrentDatabase.Get() || 
+		!CurrentDatabase->GetSearchIndex() ||
+		!CurrentDatabase->GetSearchIndex()->Assets.IsValidIndex(SearchIndexAssetIdx))
 	{
 		return nullptr;
 	}
 
-	return &CurrentDatabase->SearchIndex.Assets[SearchIndexAssetIdx];
+	return &CurrentDatabase->GetSearchIndex()->Assets[SearchIndexAssetIdx];
 }
 
 float FMotionMatchingState::ComputeJumpBlendTime(
diff --git a/Engine/Plugins/Experimental/Animation/PoseSearch/Source/Runtime/Public/PoseSearch/PoseSearch.h b/Engine/Plugins/Experimental/Animation/PoseSearch/Source/Runtime/Public/PoseSearch/PoseSearch.h
--- a/Engine/Plugins/Experimental/Animation/PoseSearch/Source/Runtime/Public/PoseSearch/PoseSearch.h
+++ b/Engine/Plugins/Experimental/Animation/PoseSearch/Source/Runtime/Public/PoseSearch/PoseSearch.h
@@ -28,11 +28,12 @@ struct FReferenceSkeleton;
 
 DECLARE_LOG_CATEGORY_EXTERN(LogPoseSearch, Log, All);
 
-namespace UE { namespace PoseSearch {
+struct FPoseSearchDatabaseDerivedData;
 
+namespace UE::PoseSearch {
 class FPoseHistory;
-
-}}
+struct FPoseSearchDatabaseAsyncCacheTask;
+}
 
 UENUM()
 enum class EPoseSearchFeatureType : int32
@@ -786,8 +787,8 @@ public:
 	UPROPERTY(EditAnywhere, Category = "Database")
 	TArray<FPoseSearchDatabaseBlendSpace> BlendSpaces;
 
-	UPROPERTY()
-	FPoseSearchIndex SearchIndex;
+	FPoseSearchIndex* GetSearchIndex();
+	const FPoseSearchIndex* GetSearchIndex() const;
 
 	bool IsValidForIndexing() const;
 	bool IsValidForSearch() const;
@@ -802,10 +803,14 @@ public:
 	const FString GetSourceAssetName(const FPoseSearchIndexAsset* SearchIndexAsset) const;
 
 public: // UObject
-	virtual void PreSave(FObjectPreSaveContext ObjectSaveContext) override;
+	virtual void PostLoad() override;
+	virtual void PostSaveRoot(FObjectPostSaveRootContext ObjectSaveContext) override;
+	virtual void Serialize(FArchive& Ar) override;
 
 #if WITH_EDITOR
 	virtual void PostEditChangeProperty(struct FPropertyChangedEvent& PropertyChangedEvent) override;
+	virtual void BeginCacheForCookedPlatformData(const ITargetPlatform* TargetPlatform) override;
+	virtual bool IsCachedCookedPlatformDataLoaded(const ITargetPlatform* TargetPlatform) override;
 #endif
 
 private:
@@ -814,7 +819,17 @@ private:
 
 public:
 	// Populates the FPoseSearchIndex::Assets array by evaluating the data in the Sequences array
-	bool TryInitSearchIndexAssets();
+	bool TryInitSearchIndexAssets(FPoseSearchIndex& OutSearchIndex);
+
+private:
+	FPoseSearchDatabaseDerivedData* PrivateDerivedData;
+
+public:
+#if WITH_EDITOR
+	void BeginCacheDerivedData();
+#endif // WITH_EDITOR
+
+	bool IsDerivedDataValid();
 };
 
 /** 
@@ -1075,7 +1090,7 @@ POSESEARCH_API bool BuildIndex(const UAnimSequence* Sequence, UPoseSearchSequenc
 * 
 * @return Whether the index was built successfully
 */
-POSESEARCH_API bool BuildIndex(UPoseSearchDatabase* Database);
+POSESEARCH_API bool BuildIndex(UPoseSearchDatabase* Database, FPoseSearchIndex& OutSearchIndex);
 
 
 /**
diff --git a/Engine/Source/Runtime/Core/Private/UObject/DevObjectVersion.cpp b/Engine/Source/Runtime/Core/Private/UObject/DevObjectVersion.cpp
--- a/Engine/Source/Runtime/Core/Private/UObject/DevObjectVersion.cpp
+++ b/Engine/Source/Runtime/Core/Private/UObject/DevObjectVersion.cpp
@@ -124,6 +124,7 @@ FDevSystemGuids::FDevSystemGuids()
 	, Niagara_LatestScriptCompileVersion(0x6D32B8EE, 0x909FCA7E, 0xA5CE4F17, 0x066A5F25)
 	, SkeletalMeshDerivedDataVersion(0x9B5F4544, 0x76D7481C, 0x9AD3F614, 0xA6C07904)
 	, STATICMESH_DERIVEDDATA_VER(0x5B7A05A7, 0x7371440C, 0xA314621A, 0x9B1E08BC)
+	, POSESEARCHDB_DERIVEDDATA_VER(0x389117E4, 0x807A4CC0, 0x9F37C2E6, 0xD808A78D)
 {
 }
 
diff --git a/Engine/Source/Runtime/Core/Private/UObject/UE5MainStreamObjectVersion.cpp b/Engine/Source/Runtime/Core/Private/UObject/UE5MainStreamObjectVersion.cpp
--- a/Engine/Source/Runtime/Core/Private/UObject/UE5MainStreamObjectVersion.cpp
+++ b/Engine/Source/Runtime/Core/Private/UObject/UE5MainStreamObjectVersion.cpp
@@ -17,6 +17,7 @@ TMap<FGuid, FGuid> FUE5MainStreamObjectVersion::GetSystemGuids()
 	SystemGuids.Add(DevGuids.Niagara_LatestScriptCompileVersion, FGuid("078A8ED467E24F94BC5DDB4BE2A4EB8E"));
 	SystemGuids.Add(DevGuids.SkeletalMeshDerivedDataVersion, FGuid("FB7B2FB546DE49B2A2F88B7F69763CEE"));
 	SystemGuids.Add(DevGuids.STATICMESH_DERIVEDDATA_VER, FGuid("98F7D79A5811013825E75BBCC41ED3E9"));
+	SystemGuids.Add(DevGuids.POSESEARCHDB_DERIVEDDATA_VER, FGuid("4E595C2AC5E947D6BA9ABC874353E5BC"));
 
 	return SystemGuids;
 }
diff --git a/Engine/Source/Runtime/Core/Public/UObject/DevObjectVersion.h b/Engine/Source/Runtime/Core/Public/UObject/DevObjectVersion.h
--- a/Engine/Source/Runtime/Core/Public/UObject/DevObjectVersion.h
+++ b/Engine/Source/Runtime/Core/Public/UObject/DevObjectVersion.h
@@ -26,6 +26,7 @@ struct CORE_API FDevSystemGuids
 	const FGuid Niagara_LatestScriptCompileVersion;
 	const FGuid SkeletalMeshDerivedDataVersion;
 	const FGuid STATICMESH_DERIVEDDATA_VER;
+	const FGuid POSESEARCHDB_DERIVEDDATA_VER;
 };
 
 class CORE_API FDevVersionRegistration :  public FCustomVersionRegistration
