Fixed non-unity build.

diff --git a/Engine/Plugins/Experimental/Animation/PoseSearch/Source/Runtime/Private/PoseSearchDerivedData.cpp b/Engine/Plugins/Experimental/Animation/PoseSearch/Source/Runtime/Private/PoseSearchDerivedData.cpp
--- a/Engine/Plugins/Experimental/Animation/PoseSearch/Source/Runtime/Private/PoseSearchDerivedData.cpp
+++ b/Engine/Plugins/Experimental/Animation/PoseSearch/Source/Runtime/Private/PoseSearchDerivedData.cpp
@@ -3,6 +3,9 @@
 #include "PoseSearchDerivedData.h"
 #include "PoseSearch/PoseSearchAnimNotifies.h"
 
+#include "Animation/AnimSequence.h"
+#include "Animation/BlendSpace.h"
+
 #if WITH_EDITOR
 #include "Serialization/BulkDataRegistry.h"
 #endif // WITH_EDITOR
