Allow MotionMatching node to be reset after not being ticked, pulling pose from PoseHistory node as intended. Added accountability for the root delta in the bone velocities returned by the PoseHistory node. Updated finite differencing in the indexer to account for certain edge cases where samples are clipped.

diff --git a/Engine/Plugins/Experimental/Animation/PoseSearch/Source/Editor/Private/PoseSearchDebugger.cpp b/Engine/Plugins/Experimental/Animation/PoseSearch/Source/Editor/Private/PoseSearchDebugger.cpp
--- a/Engine/Plugins/Experimental/Animation/PoseSearch/Source/Editor/Private/PoseSearchDebugger.cpp
+++ b/Engine/Plugins/Experimental/Animation/PoseSearch/Source/Editor/Private/PoseSearchDebugger.cpp
@@ -1573,6 +1573,11 @@ void SDebuggerView::DrawFeatures(
 			InDrawParams.Flags |= EDebugDrawFlags::IncludeTrajectory;
 		}
 
+		if (Options.bDrawBoneNames)
+		{
+			InDrawParams.Flags |= EDebugDrawFlags::DrawBoneNames;
+		}
+
 		if (Options.bDrawSampleLabels)
 		{
 			InDrawParams.Flags |= EDebugDrawFlags::DrawSampleLabels;
diff --git a/Engine/Plugins/Experimental/Animation/PoseSearch/Source/Editor/Private/PoseSearchDebugger.h b/Engine/Plugins/Experimental/Animation/PoseSearch/Source/Editor/Private/PoseSearchDebugger.h
--- a/Engine/Plugins/Experimental/Animation/PoseSearch/Source/Editor/Private/PoseSearchDebugger.h
+++ b/Engine/Plugins/Experimental/Animation/PoseSearch/Source/Editor/Private/PoseSearchDebugger.h
@@ -149,6 +149,9 @@ struct FPoseSearchDebuggerFeatureDrawOptions
 	UPROPERTY(EditAnywhere, Category="Draw Options", Meta=(EditCondition="!bDisable"))
     bool bDrawTrajectoryFeatures = true;
 
+	UPROPERTY(EditAnywhere, Category = "Draw Options", Meta = (EditCondition = "!bDisable"))
+	bool bDrawBoneNames = true;
+
 	UPROPERTY(EditAnywhere, Category = "Draw Options", Meta = (EditCondition = "!bDisable"))
 	bool bDrawSampleLabels = true;
 
diff --git a/Engine/Plugins/Experimental/Animation/PoseSearch/Source/Runtime/Private/AnimNode_MotionMatching.cpp b/Engine/Plugins/Experimental/Animation/PoseSearch/Source/Runtime/Private/AnimNode_MotionMatching.cpp
--- a/Engine/Plugins/Experimental/Animation/PoseSearch/Source/Runtime/Private/AnimNode_MotionMatching.cpp
+++ b/Engine/Plugins/Experimental/Animation/PoseSearch/Source/Runtime/Private/AnimNode_MotionMatching.cpp
@@ -66,8 +66,16 @@ void FAnimNode_MotionMatching::UpdateAssetPlayer(const FAnimationUpdateContext&
 
 	GetEvaluateGraphExposedInputs().Execute(Context);
 
-	// Update the current time.
-	MotionMatchingState.AssetPlayerTime = CurrentAssetPlayerNode->GetAccumulatedTime();
+	// If we just became relevant and haven't been initialized yet, then reset motion matching state, otherwise update the current time using the player.
+	if (bResetOnBecomingRelevant && UpdateCounter.HasEverBeenUpdated() && !UpdateCounter.WasSynchronizedCounter(Context.AnimInstanceProxy->GetUpdateCounter()))
+	{
+		MotionMatchingState.Reset();
+	}
+	else
+	{
+		MotionMatchingState.AssetPlayerTime = CurrentAssetPlayerNode->GetAccumulatedTime();
+	}
+	UpdateCounter.SynchronizeWith(Context.AnimInstanceProxy->GetUpdateCounter());
 
 	// Execute core motion matching algorithm and retain across frame state
 	UpdateMotionMatchingState(
diff --git a/Engine/Plugins/Experimental/Animation/PoseSearch/Source/Runtime/Private/AnimNode_PoseSearchHistoryCollector.cpp b/Engine/Plugins/Experimental/Animation/PoseSearch/Source/Runtime/Private/AnimNode_PoseSearchHistoryCollector.cpp
--- a/Engine/Plugins/Experimental/Animation/PoseSearch/Source/Runtime/Private/AnimNode_PoseSearchHistoryCollector.cpp
+++ b/Engine/Plugins/Experimental/Animation/PoseSearch/Source/Runtime/Private/AnimNode_PoseSearchHistoryCollector.cpp
@@ -62,7 +62,19 @@ void FAnimNode_PoseSearchHistoryCollector::Evaluate_AnyThread(FPoseContext& Outp
 
 	Source.Evaluate(Output);
 
-	PoseHistory.Update(Output.AnimInstanceProxy->GetDeltaSeconds(), Output);
+	FText ErrorText;
+
+	if (!PoseHistory.Update(
+		Output.AnimInstanceProxy->GetDeltaSeconds(),
+		Output,
+		Output.AnimInstanceProxy->GetComponentTransform(),
+		&ErrorText,
+		bUseRootMotion ?
+		UE::PoseSearch::EPoseHistoryRootUpdateMode::RootMotionDelta :
+		UE::PoseSearch::EPoseHistoryRootUpdateMode::ComponentTransformDelta))
+	{
+		Output.LogMessage(EMessageSeverity::Error, ErrorText);
+	}
 }
 
 void FAnimNode_PoseSearchHistoryCollector::Update_AnyThread(const FAnimationUpdateContext& Context)
diff --git a/Engine/Plugins/Experimental/Animation/PoseSearch/Source/Runtime/Private/PoseSearch.cpp b/Engine/Plugins/Experimental/Animation/PoseSearch/Source/Runtime/Private/PoseSearch.cpp
--- a/Engine/Plugins/Experimental/Animation/PoseSearch/Source/Runtime/Private/PoseSearch.cpp
+++ b/Engine/Plugins/Experimental/Animation/PoseSearch/Source/Runtime/Private/PoseSearch.cpp
@@ -40,6 +40,8 @@
 
 IMPLEMENT_ANIMGRAPH_MESSAGE(UE::PoseSearch::IPoseHistoryProvider);
 
+#define LOCTEXT_NAMESPACE "PoseSearch"
+
 DEFINE_LOG_CATEGORY(LogPoseSearch);
 
 namespace UE { namespace PoseSearch {
@@ -1815,28 +1817,34 @@ void FPoseSearchFeatureVectorBuilder::ResetFeatures()
 	FeaturesAdded.Init(false, Schema->Layout.Features.Num());
 }
 
-void FPoseSearchFeatureVectorBuilder::SetTransform(FPoseSearchFeatureDesc Element, const FTransform& Transform)
+void FPoseSearchFeatureVectorBuilder::SetTransform(FPoseSearchFeatureDesc Feature, const FTransform& Transform)
 {
-	SetPosition(Element, Transform.GetTranslation());
-	SetRotation(Element, Transform.GetRotation());
+	SetPosition(Feature, Transform.GetTranslation());
+	SetRotation(Feature, Transform.GetRotation());
 }
 
-void FPoseSearchFeatureVectorBuilder::SetTransformVelocity(FPoseSearchFeatureDesc Element, const FTransform& Transform, const FTransform& PrevTransform, float DeltaTime)
+void FPoseSearchFeatureVectorBuilder::SetTransformVelocity(FPoseSearchFeatureDesc Feature, const FTransform& Transform, const FTransform& PrevTransform, float DeltaTime)
 {
-	SetLinearVelocity(Element, Transform, PrevTransform, DeltaTime);
-	SetAngularVelocity(Element, Transform, PrevTransform, DeltaTime);
+	SetLinearVelocity(Feature, Transform, PrevTransform, DeltaTime);
+	SetAngularVelocity(Feature, Transform, PrevTransform, DeltaTime);
 }
 
-void FPoseSearchFeatureVectorBuilder::SetPosition(FPoseSearchFeatureDesc Element, const FVector& Position)
+void FPoseSearchFeatureVectorBuilder::SetTransformVelocity(FPoseSearchFeatureDesc Feature, const FTransform& NextTransform, const FTransform& Transform, const FTransform& PrevTransform, float DeltaTime)
 {
-	Element.Type = EPoseSearchFeatureType::Position;
-	SetVector(Element, Position);
+	SetLinearVelocity(Feature, NextTransform, Transform, PrevTransform, DeltaTime);
+	SetAngularVelocity(Feature, NextTransform, Transform, PrevTransform, DeltaTime);
 }
 
-void FPoseSearchFeatureVectorBuilder::SetRotation(FPoseSearchFeatureDesc Element, const FQuat& Rotation)
+void FPoseSearchFeatureVectorBuilder::SetPosition(FPoseSearchFeatureDesc Feature, const FVector& Position)
 {
-	Element.Type = EPoseSearchFeatureType::Rotation;
-	int32 ElementIndex = Schema->Layout.Features.Find(Element);
+	Feature.Type = EPoseSearchFeatureType::Position;
+	SetVector(Feature, Position);
+}
+
+void FPoseSearchFeatureVectorBuilder::SetRotation(FPoseSearchFeatureDesc Feature, const FQuat& Rotation)
+{
+	Feature.Type = EPoseSearchFeatureType::Rotation;
+	int32 ElementIndex = Schema->Layout.Features.Find(Feature);
 	if (ElementIndex >= 0)
 	{
 		FVector X = Rotation.GetAxisX();
@@ -1858,56 +1866,64 @@ void FPoseSearchFeatureVectorBuilder::SetRotation(FPoseSearchFeatureDesc Element
 		}
 	}
 
-	Element.Type = EPoseSearchFeatureType::ForwardVector;
-	SetVector(Element, Rotation.GetAxisY());
+	Feature.Type = EPoseSearchFeatureType::ForwardVector;
+	SetVector(Feature, Rotation.GetAxisY());
 }
 
-void FPoseSearchFeatureVectorBuilder::SetLinearVelocity(FPoseSearchFeatureDesc Element, const FTransform& Transform, const FTransform& PrevTransform, float DeltaTime)
+void FPoseSearchFeatureVectorBuilder::SetLinearVelocity(FPoseSearchFeatureDesc Feature, const FTransform& Transform, const FTransform& PrevTransform, float DeltaTime)
 {
-	Element.Type = EPoseSearchFeatureType::LinearVelocity;
+	Feature.Type = EPoseSearchFeatureType::LinearVelocity;
 	FVector LinearVelocity = (Transform.GetTranslation() - PrevTransform.GetTranslation()) / DeltaTime;
-	SetVector(Element, LinearVelocity);
+	SetVector(Feature, LinearVelocity);
 }
 
-void FPoseSearchFeatureVectorBuilder::SetAngularVelocity(FPoseSearchFeatureDesc Element, const FTransform& Transform, const FTransform& PrevTransform, float DeltaTime)
+void FPoseSearchFeatureVectorBuilder::SetLinearVelocity(FPoseSearchFeatureDesc Feature, const FTransform& NextTransform, const FTransform& Transform, const FTransform& PrevTransform, float DeltaTime)
 {
-	Element.Type = EPoseSearchFeatureType::AngularVelocity;
-	int32 ElementIndex = Schema->Layout.Features.Find(Element);
-	if (ElementIndex >= 0)
-	{
-		FQuat Q0 = PrevTransform.GetRotation();
-		FQuat Q1 = Transform.GetRotation();
-		Q1.EnforceShortestArcWith(Q0);
+	Feature.Type = EPoseSearchFeatureType::LinearVelocity;
+	FVector LinearVelocityNext = (NextTransform.GetTranslation() - Transform.GetTranslation()) / DeltaTime;
+	FVector LinearVelocityPrev = (Transform.GetTranslation() - PrevTransform.GetTranslation()) / DeltaTime;
+	SetVector(Feature, (LinearVelocityNext + LinearVelocityPrev) / 2.0f);
+}
 
-		// Given angular velocity vector w, quaternion differentiation can be represented as
-		//   dq/dt = (w * q)/2
-		// Solve for w
-		//   w = 2 * dq/dt * q^-1
-		// And let dq/dt be expressed as the finite difference
-		//   dq/dt = (q(t+h) - q(t)) / h
-		FQuat DQDt = (Q1 - Q0) / DeltaTime;
-		FQuat QInv = Q0.Inverse();
-		FQuat W = (DQDt * QInv) * 2.0f;
+static inline FVector QuaternionAngularVelocity(const FQuat& Rotation, const FQuat& PrevRotation, float DeltaTime)
+{
+	FQuat Q0 = PrevRotation;
+	FQuat Q1 = Rotation;
+	Q1.EnforceShortestArcWith(Q0);
 
-		FVector AngularVelocity(W.X, W.Y, W.Z);
+	// Given angular velocity vector w, quaternion differentiation can be represented as
+	//   dq/dt = (w * q)/2
+	// Solve for w
+	//   w = 2 * dq/dt * q^-1
+	// And let dq/dt be expressed as the finite difference
+	//   dq/dt = (q(t+h) - q(t)) / h
+	FQuat DQDt = (Q1 - Q0) / DeltaTime;
+	FQuat QInv = Q0.Inverse();
+	FQuat W = (DQDt * QInv) * 2.0f;
 
-		const FPoseSearchFeatureDesc& FoundElement = Schema->Layout.Features[ElementIndex];
+	FVector AngularVelocity(W.X, W.Y, W.Z);
 
-		Values[FoundElement.ValueOffset + 0] = AngularVelocity[0];
-		Values[FoundElement.ValueOffset + 1] = AngularVelocity[1];
-		Values[FoundElement.ValueOffset + 2] = AngularVelocity[2];
+	return AngularVelocity;
+}
 
-		if (!FeaturesAdded[ElementIndex])
-		{
-			FeaturesAdded[ElementIndex] = true;
-			++NumFeaturesAdded;
-		}
-	}
+void FPoseSearchFeatureVectorBuilder::SetAngularVelocity(FPoseSearchFeatureDesc Feature, const FTransform& Transform, const FTransform& PrevTransform, float DeltaTime)
+{
+	Feature.Type = EPoseSearchFeatureType::AngularVelocity;
+	FVector AngularVelocity = QuaternionAngularVelocity(Transform.GetRotation(), PrevTransform.GetRotation(), DeltaTime);
+	SetVector(Feature, AngularVelocity);
 }
 
-void FPoseSearchFeatureVectorBuilder::SetVector(FPoseSearchFeatureDesc Element, const FVector& Vector)
+void FPoseSearchFeatureVectorBuilder::SetAngularVelocity(FPoseSearchFeatureDesc Feature, const FTransform& NextTransform, const FTransform& Transform, const FTransform& PrevTransform, float DeltaTime)
 {
-	int32 ElementIndex = Schema->Layout.Features.Find(Element);
+	Feature.Type = EPoseSearchFeatureType::AngularVelocity;
+	FVector AngularVelocityNext = QuaternionAngularVelocity(NextTransform.GetRotation(), Transform.GetRotation(), DeltaTime);
+	FVector AngularVelocityPrev = QuaternionAngularVelocity(Transform.GetRotation(), PrevTransform.GetRotation(), DeltaTime);
+	SetVector(Feature, (AngularVelocityNext + AngularVelocityPrev) / 2.0f);
+}
+
+void FPoseSearchFeatureVectorBuilder::SetVector(FPoseSearchFeatureDesc Feature, const FVector& Vector)
+{
+	int32 ElementIndex = Schema->Layout.Features.Find(Feature);
 	if (ElementIndex >= 0)
 	{
 		const FPoseSearchFeatureDesc& FoundElement = Schema->Layout.Features[ElementIndex];
@@ -1951,6 +1967,8 @@ bool FPoseSearchFeatureVectorBuilder::TrySetPoseFeatures(UE::PoseSearch::FPoseHi
 
 		TArrayView<const FTransform> ComponentPose = History->GetComponentPoseSample();
 		TArrayView<const FTransform> ComponentPrevPose = History->GetPrevComponentPoseSample();
+		FTransform RootTransform = History->GetRootTransformSample();
+		FTransform RootTransformPrev = History->GetPrevRootTransformSample();
 		for (int32 SchemaBoneIdx = 0; SchemaBoneIdx != Schema->BoneIndices.Num(); ++SchemaBoneIdx)
 		{
 			Feature.SchemaBoneIdx = SchemaBoneIdx;
@@ -1958,7 +1976,7 @@ bool FPoseSearchFeatureVectorBuilder::TrySetPoseFeatures(UE::PoseSearch::FPoseHi
 			int32 SkeletonBoneIndex = Schema->BoneIndices[SchemaBoneIdx];
 
 			const FTransform& Transform = ComponentPose[SkeletonBoneIndex];
-			const FTransform& PrevTransform = ComponentPrevPose[SkeletonBoneIndex];
+			const FTransform& PrevTransform = ComponentPrevPose[SkeletonBoneIndex] * (RootTransformPrev * RootTransform.Inverse());
 			SetTransform(Feature, Transform);
 			SetTransformVelocity(Feature, Transform, PrevTransform, History->GetSampleTimeInterval());
 		}
@@ -2130,7 +2148,7 @@ void FPoseHistory::Init(const FPoseHistory& History)
 	TimeHorizon = History.TimeHorizon;
 }
 
-bool FPoseHistory::TrySampleLocalPose(float SecondsAgo, const TArray<FBoneIndexType>& RequiredBones, TArray<FTransform>& LocalPose)
+bool FPoseHistory::TrySampleLocalPose(float SecondsAgo, const TArray<FBoneIndexType>& RequiredBones, TArray<FTransform>& LocalPose, FTransform& RootTransform)
 {
 	int32 NextIdx = LowerBound(Knots.begin(), Knots.end(), SecondsAgo, TGreater<>());
 	if (NextIdx <= 0 || NextIdx >= Knots.Num())
@@ -2168,16 +2186,18 @@ bool FPoseHistory::TrySampleLocalPose(float SecondsAgo, const TArray<FBoneIndexT
 		Alpha,
 		RequiredBones);
 
+	RootTransform.Blend(PrevPose.RootTransform, NextPose.RootTransform, Alpha);
+
 	return true;
 }
 
 bool FPoseHistory::TrySamplePose(float SecondsAgo, const FReferenceSkeleton& RefSkeleton, const TArray<FBoneIndexType>& RequiredBones)
 {
 	// Compute local space pose at requested time
-	bool bSampled = TrySampleLocalPose(SecondsAgo, RequiredBones, SampledLocalPose);
+	bool bSampled = TrySampleLocalPose(SecondsAgo, RequiredBones, SampledLocalPose, SampledRootTransform);
 
 	// Compute local space pose one sample interval in the past
-	bSampled = bSampled && TrySampleLocalPose(SecondsAgo + GetSampleTimeInterval(), RequiredBones, SampledPrevLocalPose);
+	bSampled = bSampled && TrySampleLocalPose(SecondsAgo + GetSampleTimeInterval(), RequiredBones, SampledPrevLocalPose, SampledPrevRootTransform);
 
 	// Convert local to component space
 	if (bSampled)
@@ -2189,7 +2209,7 @@ bool FPoseHistory::TrySamplePose(float SecondsAgo, const FReferenceSkeleton& Ref
 	return bSampled;
 }
 
-void FPoseHistory::Update(float SecondsElapsed, const FPoseContext& PoseContext)
+bool FPoseHistory::Update(float SecondsElapsed, const FPoseContext& PoseContext, FTransform ComponentTransform, FText* OutError, EPoseHistoryRootUpdateMode UpdateMode)
 {
 	// Age our elapsed times
 	for (float& Knot : Knots)
@@ -2230,6 +2250,58 @@ void FPoseHistory::Update(float SecondsElapsed, const FPoseContext& PoseContext)
 	Knots.Last() = 0.0f;
 	FPose& CurrentPose = Poses.Last();
 	CopyCompactToSkeletonPose(PoseContext.Pose, CurrentPose.LocalTransforms);
+
+	// Initialize with Previous Root Transform or Identity
+	CurrentPose.RootTransform = Poses.Num() > 1 ? Poses[Poses.Num() - 2].RootTransform : FTransform::Identity;
+	
+	// Update using either AniumRootMotionProvider or Component Transform
+	if (UpdateMode == EPoseHistoryRootUpdateMode::RootMotionDelta)
+	{
+		const UE::Anim::IAnimRootMotionProvider* RootMotionProvider = UE::Anim::IAnimRootMotionProvider::Get();
+
+		if (RootMotionProvider)
+		{
+			if (RootMotionProvider->HasRootMotion(PoseContext.CustomAttributes))
+			{
+				FTransform RootMotionDelta = FTransform::Identity;
+				RootMotionProvider->ExtractRootMotion(PoseContext.CustomAttributes, RootMotionDelta);
+
+				CurrentPose.RootTransform = RootMotionDelta * CurrentPose.RootTransform;
+			}
+#if WITH_EDITORONLY_DATA	
+			else
+			{
+				if (OutError)
+				{
+					*OutError = LOCTEXT("PoseHistoryRootMotionProviderError",
+						"Input to Pose History has no Root Motion Attribute. Try disabling 'Use Root Motion'.");
+				}
+				return false;
+			}
+#endif
+		}
+#if WITH_EDITORONLY_DATA	
+		else
+		{
+			if (OutError)
+			{
+				*OutError = LOCTEXT("PoseHistoryRootMotionAttributeError",
+					"Could not get Root Motion Provider. Try disabling 'Use Root Motion'.");
+			}
+			return false;
+		}
+#endif
+	}
+	else if (UpdateMode == EPoseHistoryRootUpdateMode::ComponentTransformDelta)
+	{
+		CurrentPose.RootTransform = ComponentTransform;
+	}
+	else
+	{
+		checkNoEntry();
+	}
+
+	return true;
 }
 
 float FPoseHistory::GetSampleTimeInterval() const
@@ -2865,7 +2937,7 @@ void FSequenceSampler::ProcessRootDistance()
 	// Build a distance lookup table by sampling root motion at a fixed rate and accumulating
 	// absolute translation deltas. During indexing we'll bsearch this table and interpolate
 	// between samples in order to convert distance offsets to time offsets.
-	// See also FSequenceIndexer::AddTrajectoryDistanceFeatures().
+	// See also FIndexer::AddTrajectoryDistanceFeatures().
 
 	double TotalAccumulatedRootDistance = 0.0;
 	FTransform LastRootTransform = InitialRootTransform;
@@ -2886,7 +2958,7 @@ void FSequenceSampler::ProcessRootDistance()
 	check(SampleTime == Input.Sequence->GetPlayLength());
 
 	// Also emit root motion summary info to help with sample wrapping in 
-	// FSequenceIndexer::GetSampleTimeFromDistance() and FSequenceIndexer::GetSampleInfo()
+	// FIndexer::GetSampleTimeFromDistance() and FIndexer::GetSampleInfo()
 	TotalRootTransform = LastRootTransform.GetRelativeTransform(InitialRootTransform);
 	TotalRootDistance = AccumulatedRootDistance.Last();
 }
@@ -3298,7 +3370,7 @@ void FBlendSpaceSampler::ProcessRootDistance()
 	// Build a distance lookup table by sampling root motion at a fixed rate and accumulating
 	// absolute translation deltas. During indexing we'll bsearch this table and interpolate
 	// between samples in order to convert distance offsets to time offsets.
-	// See also FSequenceIndexer::AddTrajectoryDistanceFeatures().
+	// See also FIndexer::AddTrajectoryDistanceFeatures().
 	double TotalAccumulatedRootDistance = 0.0;
 	FTransform LastRootTransform = InitialRootTransform;
 	float SampleTime = 0.0f;
@@ -3318,7 +3390,7 @@ void FBlendSpaceSampler::ProcessRootDistance()
 	check(SampleTime == PlayLength);
 
 	// Also emit root motion summary info to help with sample wrapping in 
-	// FSequenceIndexer::GetSampleTimeFromDistance() and FSequenceIndexer::GetSampleInfo()
+	// FIndexer::GetSampleTimeFromDistance() and FIndexer::GetSampleInfo()
 	TotalRootTransform = LastRootTransform.GetRelativeTransform(InitialRootTransform);
 	TotalRootDistance = AccumulatedRootDistance.Last();
 }
@@ -3330,7 +3402,6 @@ struct FSamplingParam
 {
 	float WrappedParam = 0.0f;
 	int32 NumCycles = 0;
-	bool bClamped = false;
 	
 	// If the animation can't loop, WrappedParam contains the clamped value and whatever is left is stored here
 	float Extrapolation = 0.0f;
@@ -3376,7 +3447,6 @@ static FSamplingParam WrapOrClampSamplingParam(bool bCanWrap, float SamplingPara
 		check(!bCanWrap);
 		Result.Extrapolation = Result.WrappedParam - ParamClamped;
 		Result.WrappedParam = ParamClamped;
-		Result.bClamped = true;
 	}
 	
 	return Result;
@@ -3425,6 +3495,7 @@ private:
 		FTransform RootTransform;
 		float ClipTime = 0.0f;
 		float RootDistance = 0.0f;
+		bool bClamped = false;
 
 		bool IsValid() const { return Clip != nullptr; }
 	};
@@ -3692,6 +3763,7 @@ FIndexer::FSampleInfo FIndexer::GetSampleInfo(float SampleTime) const
 
 	if (FMath::Abs(SamplingParam.Extrapolation) > SMALL_NUMBER)
 	{
+		Sample.bClamped = true;
 		Sample.ClipTime = SamplingParam.WrappedParam + SamplingParam.Extrapolation;
 		const FTransform ClipRootMotion = Sample.Clip->ExtractRootTransform(Sample.ClipTime);
 		const float ClipDistance = Sample.Clip->ExtractRootDistance(Sample.ClipTime);
@@ -3858,8 +3930,23 @@ void FIndexer::AddPoseFeatures(int32 SampleIdx)
 
 			// Add properties to the feature vector for the pose at SampleIdx
 			FeatureVector.SetTransform(Feature, BoneTransforms[1]);
-			FeatureVector.SetTransformVelocity(Feature, BoneTransforms[2], BoneTransforms[0], 2.0f * Input.SamplingContext->FiniteDelta);
-			//FeatureVector.SetTransformAccleration(Feature, BoneTransforms[2], BoneTransforms[1], BoneTransforms[0], FMath::Square(Input.SamplingContext->FiniteDelta));
+
+			// We can get a better finite difference if we ignore samples that have
+			// been clamped at either side of the clip. However, if the central sample 
+			// itself is clamped, or there are no samples that are clamped, we can just 
+			// use the central difference as normal.
+			if (Samples[0].bClamped && !Samples[1].bClamped && !Samples[2].bClamped)
+			{
+				FeatureVector.SetTransformVelocity(Feature, BoneTransforms[2], BoneTransforms[1], Input.SamplingContext->FiniteDelta);
+			}
+			else if (Samples[2].bClamped && !Samples[1].bClamped && !Samples[0].bClamped)
+			{
+				FeatureVector.SetTransformVelocity(Feature, BoneTransforms[1], BoneTransforms[0], Input.SamplingContext->FiniteDelta);
+			}
+			else
+			{
+				FeatureVector.SetTransformVelocity(Feature, BoneTransforms[2], BoneTransforms[1], BoneTransforms[0], Input.SamplingContext->FiniteDelta);
+			}
 		}
 	}
 }
@@ -3901,8 +3988,23 @@ void FIndexer::AddTrajectoryTimeFeatures(int32 SampleIdx)
 
 		// Add properties to the feature vector for the pose at SampleIdx
 		FeatureVector.SetTransform(Feature, MirroredRoots[1]);
-		FeatureVector.SetTransformVelocity(Feature, MirroredRoots[2], MirroredRoots[0], 2.0f * Input.SamplingContext->FiniteDelta);
-		//FeatureVector.SetTransformAcceleration(Feature, MirroredRoots[2], MirroredRoots[1], MirroredRoots[0], FMath::Square(Input.SamplingContext->FiniteDelta));
+
+		// We can get a better finite difference if we ignore samples that have
+		// been clamped at either side of the clip. However, if the central sample 
+		// itself is clamped, or there are no samples that are clamped, we can just 
+		// use the central difference as normal.
+		if (Samples[0].bClamped && !Samples[1].bClamped && !Samples[2].bClamped)
+		{
+			FeatureVector.SetTransformVelocity(Feature, MirroredRoots[2], MirroredRoots[1], Input.SamplingContext->FiniteDelta);
+		}
+		else if (Samples[2].bClamped && !Samples[1].bClamped && !Samples[0].bClamped)
+		{
+			FeatureVector.SetTransformVelocity(Feature, MirroredRoots[1], MirroredRoots[0], Input.SamplingContext->FiniteDelta);
+		}
+		else
+		{
+			FeatureVector.SetTransformVelocity(Feature, MirroredRoots[2], MirroredRoots[1], MirroredRoots[0], Input.SamplingContext->FiniteDelta);
+		}
 	}
 }
 
@@ -3944,8 +4046,23 @@ void FIndexer::AddTrajectoryDistanceFeatures(int32 SampleIdx)
 
 		// Add properties to the feature vector for the pose at SampleIdx
 		FeatureVector.SetTransform(Feature, MirroredRoots[1]);
-		FeatureVector.SetTransformVelocity(Feature, MirroredRoots[2], MirroredRoots[0], 2.0f * Input.SamplingContext->FiniteDelta);
-		//FeatureVector.SetTransformAcceleration(Feature, MirroredRoots[2], MirroredRoots[1], MirroredRoots[0], FMath::Square(Input.SamplingContext->FiniteDelta));
+
+		// We can get a better finite difference if we ignore samples that have
+		// been clamped at either side of the clip. However, if the central sample 
+		// itself is clamped, or there are no samples that are clamped, we can just 
+		// use the central difference as normal.
+		if (Samples[0].bClamped && !Samples[1].bClamped && !Samples[2].bClamped)
+		{
+			FeatureVector.SetTransformVelocity(Feature, MirroredRoots[2], MirroredRoots[1], Input.SamplingContext->FiniteDelta);
+		}
+		else if (Samples[2].bClamped && !Samples[1].bClamped && !Samples[0].bClamped)
+		{
+			FeatureVector.SetTransformVelocity(Feature, MirroredRoots[1], MirroredRoots[0], Input.SamplingContext->FiniteDelta);
+		}
+		else
+		{
+			FeatureVector.SetTransformVelocity(Feature, MirroredRoots[2], MirroredRoots[1], MirroredRoots[0], Input.SamplingContext->FiniteDelta);
+		}
 	}
 }
 
@@ -4180,6 +4297,14 @@ static void DrawPoseFeatures(const FDebugDrawParams& DrawParams, const FFeatureV
 				{
 					DrawDebugSphere(DrawParams.World, BonePos, DrawDebugSphereSize, DrawDebugSphereSegments, Color, bPersistent, LifeTime, DepthPriority);
 				}
+
+				if (EnumHasAnyFlags(DrawParams.Flags, EDebugDrawFlags::DrawBoneNames))
+				{
+					DrawDebugString(
+						DrawParams.World, BonePos + FVector(0.0, 0.0, 10.0),
+						Schema->SampledBones[SchemaBoneIdx].Reference.BoneName.ToString(),
+						nullptr, Color, LifeTime, false, 1.0f);
+				}
 			}
 
 			FVector BoneVel;
@@ -5349,4 +5474,6 @@ void FModule::OnObjectSaved(UObject* SavedObject, FObjectPreSaveContext SaveCont
 
 }} // namespace UE::PoseSearch
 
+#undef LOCTEXT_NAMESPACE
+
 IMPLEMENT_MODULE(UE::PoseSearch::FModule, PoseSearch)
\ No newline at end of file
diff --git a/Engine/Plugins/Experimental/Animation/PoseSearch/Source/Runtime/Private/PoseSearchLibrary.cpp b/Engine/Plugins/Experimental/Animation/PoseSearch/Source/Runtime/Private/PoseSearchLibrary.cpp
--- a/Engine/Plugins/Experimental/Animation/PoseSearch/Source/Runtime/Private/PoseSearchLibrary.cpp
+++ b/Engine/Plugins/Experimental/Animation/PoseSearch/Source/Runtime/Private/PoseSearchLibrary.cpp
@@ -193,16 +193,8 @@ bool FMotionMatchingState::InitNewDatabaseSearch(
 
 	if (bValidDatabase)
 	{
-		DbPoseIdx = INDEX_NONE;
-		SearchIndexAssetIdx = INDEX_NONE;
-		ElapsedPoseJumpTime = SearchThrottleTime;
-		AssetPlayerTime = 0.0f;
 		CurrentDatabase = Database;
-
-		if (!ComposedQuery.IsInitializedForSchema(Database->Schema))
-		{
-			ComposedQuery.Init(Database->Schema);
-		}
+		Reset();
 	}
 
 	if (!bValidDatabase && OutError)
@@ -220,6 +212,23 @@ bool FMotionMatchingState::InitNewDatabaseSearch(
 	return bValidDatabase;
 }
 
+void FMotionMatchingState::Reset()
+{
+	check(CurrentDatabase.IsValid() && CurrentDatabase->IsValidForSearch());
+
+	DbPoseIdx = INDEX_NONE;
+	SearchIndexAssetIdx = INDEX_NONE;
+	AssetPlayerTime = 0.0f;
+	// Set the elapsed time to INFINITY to trigger a search right away
+	ElapsedPoseJumpTime = INFINITY;
+
+	if (!ComposedQuery.IsInitializedForSchema(CurrentDatabase->Schema))
+	{
+		ComposedQuery.Init(CurrentDatabase->Schema);
+	}
+
+}
+
 void FMotionMatchingState::ComposeQuery(const UPoseSearchDatabase* Database, const FTrajectorySampleRange& Trajectory)
 {
 	FPoseSearchFeatureVectorBuilder Goal;
@@ -313,15 +322,24 @@ void UpdateMotionMatchingState(
 	}
 	else
 	{
-		// When we don't have an active pose, initialize the search query from the pose history provider
+		bool bPoseHistoryFound = false;
+
+		// When we don't have an active pose in the database, initialize the search query from the pose history provider
 		IPoseHistoryProvider* PoseHistoryProvider = Context.GetMessage<IPoseHistoryProvider>();
 		if (PoseHistoryProvider)
 		{
 			FPoseHistory& History = PoseHistoryProvider->GetPoseHistory();
-			InOutMotionMatchingState.ComposedQuery.TrySetPoseFeatures(
-				&History, 
+
+			bPoseHistoryFound = InOutMotionMatchingState.ComposedQuery.TrySetPoseFeatures(
+				&History,
 				Context.AnimInstanceProxy->GetRequiredBones());
 		}
+
+		// If no pose history provider or no pose found in history, then just reset/zero
+		if (!bPoseHistoryFound)
+		{
+			InOutMotionMatchingState.ComposedQuery.ResetFeatures();
+		}
 	}
 
 	// Update features in the query with the latest inputs
diff --git a/Engine/Plugins/Experimental/Animation/PoseSearch/Source/Runtime/Public/PoseSearch/AnimNode_MotionMatching.h b/Engine/Plugins/Experimental/Animation/PoseSearch/Source/Runtime/Public/PoseSearch/AnimNode_MotionMatching.h
--- a/Engine/Plugins/Experimental/Animation/PoseSearch/Source/Runtime/Public/PoseSearch/AnimNode_MotionMatching.h
+++ b/Engine/Plugins/Experimental/Animation/PoseSearch/Source/Runtime/Public/PoseSearch/AnimNode_MotionMatching.h
@@ -46,6 +46,11 @@ public:
 	UPROPERTY(EditAnywhere, BlueprintReadWrite, Category=Settings, meta=(PinHiddenByDefault))
 	FMotionMatchingSettings Settings;
 
+	// Reset the motion matching state if we have become relevant to the graph
+	// after not being ticked on the previous frame(s)
+	UPROPERTY(EditAnywhere, BlueprintReadWrite, Category = Settings, meta = (PinHiddenByDefault))
+	bool bResetOnBecomingRelevant = true;
+
 #if WITH_EDITORONLY_DATA
 	UPROPERTY(EditAnywhere, Category=Debug, meta = (PinShownByDefault))
 	bool bDebugDraw = false;
@@ -82,6 +87,8 @@ private:
 	// Current Asset Player Node
 	FAnimNode_AssetPlayerBase* CurrentAssetPlayerNode = &SequencePlayerNode;
 
+	// Update Counter for detecting being relevant
+	FGraphTraversalCounter UpdateCounter;
 
 	// FAnimNode_AssetPlayerBase
 protected:
diff --git a/Engine/Plugins/Experimental/Animation/PoseSearch/Source/Runtime/Public/PoseSearch/AnimNode_PoseSearchHistoryCollector.h b/Engine/Plugins/Experimental/Animation/PoseSearch/Source/Runtime/Public/PoseSearch/AnimNode_PoseSearchHistoryCollector.h
--- a/Engine/Plugins/Experimental/Animation/PoseSearch/Source/Runtime/Public/PoseSearch/AnimNode_PoseSearchHistoryCollector.h
+++ b/Engine/Plugins/Experimental/Animation/PoseSearch/Source/Runtime/Public/PoseSearch/AnimNode_PoseSearchHistoryCollector.h
@@ -23,6 +23,13 @@ public:
 	UPROPERTY(EditAnywhere, BlueprintReadWrite, Category=Settings, meta=(ClampMin="0"))
 	float PoseDuration = 4.f;
 
+	// If this node should compute velocities using the playing animation root motion, or 
+	// the actual movement of the component in the world. While using root motion can 
+	// be more stable, not all input animations support it, and it may produce bad results 
+	// if the current movement in the world does not match closely the animation data.
+	UPROPERTY(EditAnywhere, BlueprintReadWrite, Category = Settings, meta = (PinHiddenByDefault))
+	bool bUseRootMotion = true;
+
 public:
 	FAnimNode_PoseSearchHistoryCollector() { }
 
diff --git a/Engine/Plugins/Experimental/Animation/PoseSearch/Source/Runtime/Public/PoseSearch/PoseSearch.h b/Engine/Plugins/Experimental/Animation/PoseSearch/Source/Runtime/Public/PoseSearch/PoseSearch.h
--- a/Engine/Plugins/Experimental/Animation/PoseSearch/Source/Runtime/Public/PoseSearch/PoseSearch.h
+++ b/Engine/Plugins/Experimental/Animation/PoseSearch/Source/Runtime/Public/PoseSearch/PoseSearch.h
@@ -854,10 +854,13 @@ public:
 
 	void SetTransform(FPoseSearchFeatureDesc Feature, const FTransform& Transform);
 	void SetTransformVelocity(FPoseSearchFeatureDesc Feature, const FTransform& Transform, const FTransform& PrevTransform, float DeltaTime);
+	void SetTransformVelocity(FPoseSearchFeatureDesc Feature, const FTransform& NextTransform, const FTransform& Transform, const FTransform& PrevTransform, float DeltaTime);
 	void SetPosition(FPoseSearchFeatureDesc Feature, const FVector& Translation);
 	void SetRotation(FPoseSearchFeatureDesc Feature, const FQuat& Rotation);
 	void SetLinearVelocity(FPoseSearchFeatureDesc Feature, const FTransform& Transform, const FTransform& PrevTransform, float DeltaTime);
+	void SetLinearVelocity(FPoseSearchFeatureDesc Feature, const FTransform& NextTransform, const FTransform& Transform, const FTransform& PrevTransform, float DeltaTime);
 	void SetAngularVelocity(FPoseSearchFeatureDesc Feature, const FTransform& Transform, const FTransform& PrevTransform, float DeltaTime);
+	void SetAngularVelocity(FPoseSearchFeatureDesc Feature, const FTransform& NextTransform, const FTransform& Transform, const FTransform& PrevTransform, float DeltaTime);
 	void SetVector(FPoseSearchFeatureDesc Feature, const FVector& Vector);
 	void BuildFromTrajectory(const FTrajectorySampleRange& Trajectory);
 	bool TrySetPoseFeatures(UE::PoseSearch::FPoseHistory* History, const FBoneContainer& BoneContainer);
@@ -890,6 +893,13 @@ private:
 
 namespace UE { namespace PoseSearch {
 
+UENUM()
+enum class EPoseHistoryRootUpdateMode : int32
+{
+	RootMotionDelta,
+	ComponentTransformDelta,
+};
+
 /** Records poses over time in a ring buffer. FFeatureVectorBuilder uses this to sample from the present or past poses according to the search schema. */
 class POSESEARCH_API FPoseHistory
 {
@@ -898,20 +908,29 @@ public:
 	void Init(const FPoseHistory& History);
 	bool TrySamplePose(float SecondsAgo, const FReferenceSkeleton& RefSkeleton, const TArray<FBoneIndexType>& RequiredBones);
 
-	void Update(float SecondsElapsed, const FPoseContext& PoseContext);
+	bool Update(
+		float SecondsElapsed, 
+		const FPoseContext& PoseContext, 
+		FTransform ComponentTransform, 
+		FText* OutError,
+		EPoseHistoryRootUpdateMode UpdateMode = EPoseHistoryRootUpdateMode::RootMotionDelta);
+
 	float GetSampleTimeInterval() const;
 	TArrayView<const FTransform> GetLocalPoseSample() const { return SampledLocalPose; }
 	TArrayView<const FTransform> GetComponentPoseSample() const { return SampledComponentPose; }
 	TArrayView<const FTransform> GetPrevLocalPoseSample() const { return SampledPrevLocalPose; }
 	TArrayView<const FTransform> GetPrevComponentPoseSample() const { return SampledPrevComponentPose; }
+	const FTransform& GetRootTransformSample() const { return SampledRootTransform; }
+	const FTransform& GetPrevRootTransformSample() const { return SampledPrevRootTransform; }
 	float GetTimeHorizon() const { return TimeHorizon; }
 	FPoseSearchFeatureVectorBuilder& GetQueryBuilder() { return QueryBuilder; }
 
 private:
-	bool TrySampleLocalPose(float Time, const TArray<FBoneIndexType>& RequiredBones, TArray<FTransform>& LocalPose);
+	bool TrySampleLocalPose(float Time, const TArray<FBoneIndexType>& RequiredBones, TArray<FTransform>& LocalPose, FTransform& RootTransform);
 
 	struct FPose
 	{
+		FTransform RootTransform;
 		TArray<FTransform> LocalTransforms;
 	};
 
@@ -921,6 +940,9 @@ private:
 	TArray<FTransform> SampledComponentPose;
 	TArray<FTransform> SampledPrevLocalPose;
 	TArray<FTransform> SampledPrevComponentPose;
+	FTransform SampledRootTransform;
+	FTransform SampledPrevRootTransform;
+
 	FPoseSearchFeatureVectorBuilder QueryBuilder;
 
 	float TimeHorizon = 0.0f;
@@ -991,6 +1013,9 @@ enum class EDebugDrawFlags : uint32
 
 	// Fade colors
 	DrawSamplesWithColorGradient = 1 << 5,
+
+	// Draw Bone Names
+	DrawBoneNames = 1 << 6,
 };
 ENUM_CLASS_FLAGS(EDebugDrawFlags);
 
@@ -999,7 +1024,7 @@ struct POSESEARCH_API FDebugDrawParams
 	const UWorld* World = nullptr;
 	const UPoseSearchDatabase* Database = nullptr;
 	const UPoseSearchSequenceMetaData* SequenceMetaData = nullptr;
-	EDebugDrawFlags Flags = EDebugDrawFlags::IncludeAllFeatures;
+	EDebugDrawFlags Flags = EDebugDrawFlags::IncludeAllFeatures | EDebugDrawFlags::DrawBoneNames;
 
 	float DefaultLifeTime = 5.0f;
 	float PointSize = 1.0f;
diff --git a/Engine/Plugins/Experimental/Animation/PoseSearch/Source/Runtime/Public/PoseSearch/PoseSearchLibrary.h b/Engine/Plugins/Experimental/Animation/PoseSearch/Source/Runtime/Public/PoseSearch/PoseSearchLibrary.h
--- a/Engine/Plugins/Experimental/Animation/PoseSearch/Source/Runtime/Public/PoseSearch/PoseSearchLibrary.h
+++ b/Engine/Plugins/Experimental/Animation/PoseSearch/Source/Runtime/Public/PoseSearch/PoseSearchLibrary.h
@@ -82,6 +82,9 @@ struct POSESEARCH_API FMotionMatchingState
 	// Initializes the minimum required motion matching state
 	bool InitNewDatabaseSearch(const UPoseSearchDatabase* Database, float SearchThrottleTime, FText* OutError);
 
+	// Reset the state to a default state using the current Database
+	void Reset();
+
 	// Adds trajectory prediction and history information to ComposedQuery
 	void ComposeQuery(const UPoseSearchDatabase* Database, const FTrajectorySampleRange& Trajectory);
 
