Motion Matching - implemented first version of Pose Search Database Editor

diff --git a/Engine/Plugins/Experimental/Animation/PoseSearch/Source/Editor/PoseSearchEditor.Build.cs b/Engine/Plugins/Experimental/Animation/PoseSearch/Source/Editor/PoseSearchEditor.Build.cs
--- a/Engine/Plugins/Experimental/Animation/PoseSearch/Source/Editor/PoseSearchEditor.Build.cs
+++ b/Engine/Plugins/Experimental/Animation/PoseSearch/Source/Editor/PoseSearchEditor.Build.cs
@@ -46,7 +46,9 @@ public class PoseSearchEditor : ModuleRules
 				"SlateCore",
 				"Slate",
 				"EditorStyle",
-				"DetailCustomizations"
+				"DetailCustomizations",
+				"AdvancedPreviewScene",
+				"EditorFramework"
 			}
 		);
 	}
diff --git a/Engine/Plugins/Experimental/Animation/PoseSearch/Source/Editor/Private/PoseSearchDatabaseEdMode.cpp b/Engine/Plugins/Experimental/Animation/PoseSearch/Source/Editor/Private/PoseSearchDatabaseEdMode.cpp
new file mode 100644
--- /dev/null
+++ b/Engine/Plugins/Experimental/Animation/PoseSearch/Source/Editor/Private/PoseSearchDatabaseEdMode.cpp
@@ -0,0 +1,90 @@
+// Copyright Epic Games, Inc. All Rights Reserved.
+
+#include "PoseSearchDatabaseEdMode.h"
+#include "PoseSearchDatabaseViewportClient.h"
+#include "PoseSearchDatabaseEditorToolkit.h"
+#include "PoseSearchDatabaseViewModel.h"
+#include "PoseSearch/PoseSearch.h"
+
+#include "EngineUtils.h"
+
+const FEditorModeID FPoseSearchDatabaseEdMode::EdModeId = TEXT("PoseSearchDatabaseEdMode");
+
+FPoseSearchDatabaseEdMode::FPoseSearchDatabaseEdMode()
+{
+}
+
+FPoseSearchDatabaseEdMode::~FPoseSearchDatabaseEdMode()
+{
+}
+
+void FPoseSearchDatabaseEdMode::Tick(FEditorViewportClient* ViewportClient, float DeltaTime)
+{
+	FEdMode::Tick(ViewportClient, DeltaTime);
+
+	FPoseSearchDatabaseViewportClient* PoseSearchDbViewportClient = 
+		static_cast<FPoseSearchDatabaseViewportClient*>(ViewportClient);
+	if (!ViewModel && PoseSearchDbViewportClient)
+	{
+		ViewModel = PoseSearchDbViewportClient->GetAssetEditorToolkit()->GetViewModel();
+	}
+}
+
+void FPoseSearchDatabaseEdMode::Render(const FSceneView* View, FViewport* Viewport, FPrimitiveDrawInterface* PDI)
+{
+	FEdMode::Render(View, Viewport, PDI);
+}
+
+bool FPoseSearchDatabaseEdMode::HandleClick(
+	FEditorViewportClient* InViewportClient, 
+	HHitProxy* HitProxy, 
+	const FViewportClick& Click)
+{
+	return FEdMode::HandleClick(InViewportClient, HitProxy, Click);
+}
+
+
+bool FPoseSearchDatabaseEdMode::StartTracking(FEditorViewportClient* InViewportClient, FViewport* InViewport)
+{
+	return FEdMode::StartTracking(InViewportClient, InViewport);
+}
+
+bool FPoseSearchDatabaseEdMode::EndTracking(FEditorViewportClient* InViewportClient, FViewport* InViewport)
+{
+	return FEdMode::EndTracking(InViewportClient, InViewport);
+}
+
+bool FPoseSearchDatabaseEdMode::InputDelta(FEditorViewportClient* InViewportClient, FViewport* InViewport, FVector& InDrag, FRotator& InRot, FVector& InScale)
+{
+	return FEdMode::InputDelta(InViewportClient, InViewport, InDrag, InRot, InScale);
+}
+
+bool FPoseSearchDatabaseEdMode::InputKey(FEditorViewportClient* ViewportClient, FViewport* Viewport, FKey Key, EInputEvent Event)
+{
+	return FEdMode::InputKey(ViewportClient, Viewport, Key, Event);
+}
+
+bool FPoseSearchDatabaseEdMode::AllowWidgetMove()
+{
+	return FEdMode::ShouldDrawWidget();
+}
+
+bool FPoseSearchDatabaseEdMode::ShouldDrawWidget() const
+{
+	return FEdMode::ShouldDrawWidget();
+}
+
+bool FPoseSearchDatabaseEdMode::GetCustomDrawingCoordinateSystem(FMatrix& InMatrix, void* InData)
+{
+	return FEdMode::GetCustomDrawingCoordinateSystem(InMatrix, InData);
+}
+
+bool FPoseSearchDatabaseEdMode::GetCustomInputCoordinateSystem(FMatrix& InMatrix, void* InData)
+{
+	return FEdMode::GetCustomDrawingCoordinateSystem(InMatrix, InData);
+}
+
+FVector FPoseSearchDatabaseEdMode::GetWidgetLocation() const
+{
+	return FEdMode::GetWidgetLocation();
+}
diff --git a/Engine/Plugins/Experimental/Animation/PoseSearch/Source/Editor/Private/PoseSearchDatabaseEditorCommands.cpp b/Engine/Plugins/Experimental/Animation/PoseSearch/Source/Editor/Private/PoseSearchDatabaseEditorCommands.cpp
new file mode 100644
--- /dev/null
+++ b/Engine/Plugins/Experimental/Animation/PoseSearch/Source/Editor/Private/PoseSearchDatabaseEditorCommands.cpp
@@ -0,0 +1,13 @@
+// Copyright Epic Games, Inc. All Rights Reserved.
+
+#include "PoseSearchDatabaseEditorCommands.h"
+
+#define LOCTEXT_NAMESPACE "PoseSearchDatabaseEditorCommands"
+
+void FPoseSearchDatabaseEditorCommands::RegisterCommands()
+{
+	UI_COMMAND(ResetPreviewScene, "Reset Scene", "Reset Scene", EUserInterfaceActionType::Button, FInputChord());
+	UI_COMMAND(BuildSearchIndex, "Build Index", "Build Index", EUserInterfaceActionType::Button, FInputChord());
+}
+
+#undef LOCTEXT_NAMESPACE
diff --git a/Engine/Plugins/Experimental/Animation/PoseSearch/Source/Editor/Private/PoseSearchDatabaseEditorToolkit.cpp b/Engine/Plugins/Experimental/Animation/PoseSearch/Source/Editor/Private/PoseSearchDatabaseEditorToolkit.cpp
new file mode 100644
--- /dev/null
+++ b/Engine/Plugins/Experimental/Animation/PoseSearch/Source/Editor/Private/PoseSearchDatabaseEditorToolkit.cpp
@@ -0,0 +1,323 @@
+// Copyright Epic Games, Inc. All Rights Reserved.
+
+#include "PoseSearchDatabaseEditorToolkit.h"
+#include "SPoseSearchDatabaseViewport.h"
+#include "PoseSearchDatabasePreviewScene.h"
+#include "PoseSearchDatabaseEditorCommands.h"
+#include "PoseSearchDatabaseViewModel.h"
+#include "PoseSearchEditor.h"
+#include "PoseSearch/PoseSearch.h"
+#include "GameFramework/WorldSettings.h"
+#include "AdvancedPreviewSceneModule.h"
+#include "IStructureDetailsView.h"
+#include "PropertyEditorModule.h"
+#include "Modules/ModuleManager.h"
+#include "Widgets/Input/SButton.h"
+#include "Widgets/Docking/SDockTab.h"
+
+#define LOCTEXT_NAMESPACE "PoseSearchDatabaseEditorToolkit"
+
+const FName PoseSearchDatabaseEditorAppName = FName(TEXT("PoseSearchDatabaseEditorApp"));
+
+// Tab identifiers
+struct FPoseSearchDatabaseEditorTabs
+{
+	static const FName AssetDetailsID;
+	static const FName ViewportID;
+	static const FName PreviewSettingsID;
+};
+
+const FName FPoseSearchDatabaseEditorTabs::AssetDetailsID(TEXT("PoseSearchDatabaseEditorAssetDetailsTabID"));
+const FName FPoseSearchDatabaseEditorTabs::ViewportID(TEXT("PoseSearchDatabaseEditorViewportTabID"));
+const FName FPoseSearchDatabaseEditorTabs::PreviewSettingsID(TEXT("PoseSearchDatabaseEditorPreviewSettingsTabID"));
+
+FPoseSearchDatabaseEditorToolkit::FPoseSearchDatabaseEditorToolkit()
+{
+}
+
+FPoseSearchDatabaseEditorToolkit::~FPoseSearchDatabaseEditorToolkit()
+{
+}
+
+UPoseSearchDatabase* FPoseSearchDatabaseEditorToolkit::GetPoseSearchDatabase() const
+{
+	return ViewModel.IsValid() ? ViewModel->GetPoseSearchDatabase() : nullptr;
+}
+
+void FPoseSearchDatabaseEditorToolkit::ResetPreviewScene()
+{
+	ViewModel->RestartAnimations();
+}
+
+void FPoseSearchDatabaseEditorToolkit::BuildSearchIndex()
+{
+	ViewModel->BuildSearchIndex();
+}
+
+void FPoseSearchDatabaseEditorToolkit::InitAssetEditor(
+	const EToolkitMode::Type Mode,
+	const TSharedPtr<IToolkitHost>& InitToolkitHost,
+	UPoseSearchDatabase* DatabaseAsset)
+{
+	// Bind Commands
+	BindCommands();
+
+	// Create Preview Scene
+	if (!PreviewScene.IsValid())
+	{
+		PreviewScene = MakeShareable(
+			new FPoseSearchDatabasePreviewScene(
+				FPreviewScene::ConstructionValues()
+				.AllowAudioPlayback(true)
+				.ShouldSimulatePhysics(true)
+				.ForceUseMovementComponentInNonGameWorld(true),
+				StaticCastSharedRef<FPoseSearchDatabaseEditorToolkit>(AsShared())));
+
+		//Temporary fix for missing attached assets - MDW (Copied from FPersonaToolkit::CreatePreviewScene)
+		PreviewScene->GetWorld()->GetWorldSettings()->SetIsTemporarilyHiddenInEditor(false);
+	}
+
+	// Create viewport widget
+	FPoseSearchDatabaseViewportRequiredArgs ViewportArgs(
+		StaticCastSharedRef<FPoseSearchDatabaseEditorToolkit>(AsShared()), 
+		PreviewScene.ToSharedRef());
+	ViewportWidget = SNew(SPoseSearchDatabaseViewport, ViewportArgs);
+
+	ViewModel = MakeShared<FPoseSearchDatabaseViewModel>();
+	ViewModel->Initialize(DatabaseAsset, PreviewScene.ToSharedRef());
+
+	// Create Asset Details widget
+	FPropertyEditorModule& PropertyModule = FModuleManager::LoadModuleChecked<FPropertyEditorModule>("PropertyEditor");
+
+	FDetailsViewArgs Args;
+	Args.bHideSelectionTip = true;
+	Args.NotifyHook = this;
+
+	EditingAssetWidget = PropertyModule.CreateDetailView(Args);
+	EditingAssetWidget->SetObject(DatabaseAsset);
+	EditingAssetWidget->OnFinishedChangingProperties().AddSP(
+		this, 
+		&FPoseSearchDatabaseEditorToolkit::OnFinishedChangingProperties);
+
+	// Define Editor Layout
+	const TSharedRef<FTabManager::FLayout> StandaloneDefaultLayout = 
+		FTabManager::NewLayout("Standalone_PoseSearchDatabaseEditor_Layout_v0.01")
+		->AddArea
+		(
+			FTabManager::NewPrimaryArea()->SetOrientation(Orient_Vertical)
+			->Split
+			(
+				FTabManager::NewSplitter()->SetOrientation(Orient_Vertical)->SetSizeCoefficient(0.9f)
+				->Split
+				(
+					//FTabManager::NewSplitter()->SetOrientation(Orient_Horizontal)->SetSizeCoefficient(0.9f)
+					//->Split
+					//(
+					//	FTabManager::NewStack()
+					//	->SetSizeCoefficient(0.65f)
+					//	->AddTab(FPoseSearchDatabaseEditorTabs::ViewportID, ETabState::OpenedTab)
+					//	->SetHideTabWell(true)
+					//)
+					//->Split
+					//(
+					//	FTabManager::NewSplitter()->SetOrientation(Orient_Vertical)
+					//	->Split
+					//	(
+							FTabManager::NewStack()
+							->SetSizeCoefficient(0.3f)
+							->AddTab(FPoseSearchDatabaseEditorTabs::AssetDetailsID, ETabState::OpenedTab)
+					//		->AddTab(FPoseSearchDatabaseEditorTabs::PreviewSettingsID, ETabState::OpenedTab)
+					//	)
+					//)
+				)
+			)
+		);
+
+	const bool bCreateDefaultStandaloneMenuParam = true;
+	const bool bCreateDefaultToolbarParam = true;
+	const bool bIsToolbarFocusableParam = false;
+	FAssetEditorToolkit::InitAssetEditor(
+		Mode, 
+		InitToolkitHost, 
+		PoseSearchDatabaseEditorAppName, 
+		StandaloneDefaultLayout, 
+		bCreateDefaultStandaloneMenuParam, 
+		bCreateDefaultToolbarParam, 
+		DatabaseAsset, 
+		bIsToolbarFocusableParam);
+
+	ExtendToolbar();
+
+	RegenerateMenusAndToolbars();
+}
+
+void FPoseSearchDatabaseEditorToolkit::BindCommands()
+{
+	const FPoseSearchDatabaseEditorCommands& Commands = FPoseSearchDatabaseEditorCommands::Get();
+
+	ToolkitCommands->MapAction(
+		Commands.ResetPreviewScene,
+		FExecuteAction::CreateSP(this, &FPoseSearchDatabaseEditorToolkit::ResetPreviewScene),
+		EUIActionRepeatMode::RepeatDisabled);
+
+	ToolkitCommands->MapAction(
+		Commands.BuildSearchIndex,
+		FExecuteAction::CreateSP(this, &FPoseSearchDatabaseEditorToolkit::BuildSearchIndex),
+		EUIActionRepeatMode::RepeatDisabled);
+}
+
+void FPoseSearchDatabaseEditorToolkit::ExtendToolbar()
+{
+	TSharedPtr<FExtender> ToolbarExtender = MakeShareable(new FExtender);
+
+	AddToolbarExtender(ToolbarExtender);
+
+	ToolbarExtender->AddToolBarExtension(
+		"Asset",
+		EExtensionHook::After,
+		GetToolkitCommands(),
+		FToolBarExtensionDelegate::CreateSP(this, &FPoseSearchDatabaseEditorToolkit::FillToolbar));
+}
+
+void FPoseSearchDatabaseEditorToolkit::FillToolbar(FToolBarBuilder& ToolbarBuilder)
+{
+	ToolbarBuilder.AddToolBarButton(
+		FPoseSearchDatabaseEditorCommands::Get().ResetPreviewScene,
+		NAME_None,
+		TAttribute<FText>(),
+		TAttribute<FText>(),
+		FSlateIcon(FAppStyle::Get().GetStyleSetName(), "Icons.Refresh"));
+
+	ToolbarBuilder.AddToolBarButton(
+		FPoseSearchDatabaseEditorCommands::Get().BuildSearchIndex,
+		NAME_None,
+		TAttribute<FText>(),
+		TAttribute<FText>(),
+		FSlateIcon());
+}
+
+void FPoseSearchDatabaseEditorToolkit::RegisterTabSpawners(const TSharedRef<FTabManager>& InTabManager)
+{
+	WorkspaceMenuCategory = InTabManager->AddLocalWorkspaceMenuCategory(
+		LOCTEXT("WorkspaceMenu_PoseSearchDbEditor", "Pose Search Database Editor"));
+	auto WorkspaceMenuCategoryRef = WorkspaceMenuCategory.ToSharedRef();
+
+	FAssetEditorToolkit::RegisterTabSpawners(InTabManager);
+
+	//InTabManager->RegisterTabSpawner(
+	//	FPoseSearchDatabaseEditorTabs::ViewportID, 
+	//	FOnSpawnTab::CreateSP(this, &FPoseSearchDatabaseEditorToolkit::SpawnTab_Viewport))
+	//	.SetDisplayName(LOCTEXT("ViewportTab", "Viewport"))
+	//	.SetGroup(WorkspaceMenuCategoryRef)
+	//	.SetIcon(FSlateIcon(FEditorStyle::GetStyleSetName(), "GraphEditor.EventGraph_16x"));
+
+	InTabManager->RegisterTabSpawner(
+		FPoseSearchDatabaseEditorTabs::AssetDetailsID, 
+		FOnSpawnTab::CreateSP(this, &FPoseSearchDatabaseEditorToolkit::SpawnTab_AssetDetails))
+		.SetDisplayName(LOCTEXT("AssetDetailsTab", "AssetDetails"))
+		.SetGroup(WorkspaceMenuCategoryRef)
+		.SetIcon(FSlateIcon(FEditorStyle::GetStyleSetName(), "LevelEditor.Tabs.Details"));
+
+	//InTabManager->RegisterTabSpawner(
+	//	FPoseSearchDatabaseEditorTabs::PreviewSettingsID, 
+	//	FOnSpawnTab::CreateSP(this, &FPoseSearchDatabaseEditorToolkit::SpawnTab_PreviewSettings))
+	//	.SetDisplayName(LOCTEXT("PreviewSceneSettingsTab", "Preview Scene Settings"))
+	//	.SetGroup(WorkspaceMenuCategoryRef)
+	//	.SetIcon(FSlateIcon(FEditorStyle::GetStyleSetName(), "LevelEditor.Tabs.Details"));
+}
+
+void FPoseSearchDatabaseEditorToolkit::UnregisterTabSpawners(const TSharedRef<FTabManager>& InTabManager)
+{
+	FAssetEditorToolkit::UnregisterTabSpawners(InTabManager);
+
+	// InTabManager->UnregisterTabSpawner(FPoseSearchDatabaseEditorTabs::ViewportID);
+	InTabManager->UnregisterTabSpawner(FPoseSearchDatabaseEditorTabs::AssetDetailsID);
+	// InTabManager->UnregisterTabSpawner(FPoseSearchDatabaseEditorTabs::PreviewSettingsID);
+}
+
+FName FPoseSearchDatabaseEditorToolkit::GetToolkitFName() const
+{
+	return FName("PoseSearchDatabaseEditor");
+}
+
+FText FPoseSearchDatabaseEditorToolkit::GetBaseToolkitName() const
+{
+	return LOCTEXT("PoseSearchDatabaseEditorAppLabel", "Pose Search Database Editor");
+}
+
+FText FPoseSearchDatabaseEditorToolkit::GetToolkitName() const
+{
+	FFormatNamedArguments Args;
+	Args.Add(TEXT("AssetName"), FText::FromString(GetPoseSearchDatabase()->GetName()));
+	return FText::Format(LOCTEXT("PoseSearchDatabaseEditorToolkitName", "{AssetName}"), Args);
+}
+
+FLinearColor FPoseSearchDatabaseEditorToolkit::GetWorldCentricTabColorScale() const
+{
+	return FLinearColor::White;
+}
+
+FString FPoseSearchDatabaseEditorToolkit::GetWorldCentricTabPrefix() const
+{
+	return TEXT("PoseSearchDatabaseEditor");
+}
+
+TSharedRef<SDockTab> FPoseSearchDatabaseEditorToolkit::SpawnTab_Viewport(const FSpawnTabArgs& Args)
+{
+	check(Args.GetTabId() == FPoseSearchDatabaseEditorTabs::ViewportID);
+
+	TSharedRef<SDockTab> SpawnedTab = SNew(SDockTab).Label(LOCTEXT("ViewportTab_Title", "Viewport"));
+
+	if (ViewportWidget.IsValid())
+	{
+		SpawnedTab->SetContent(ViewportWidget.ToSharedRef());
+	}
+
+	return SpawnedTab;
+}
+
+TSharedRef<SDockTab> FPoseSearchDatabaseEditorToolkit::SpawnTab_AssetDetails(const FSpawnTabArgs& Args)
+{
+	check(Args.GetTabId() == FPoseSearchDatabaseEditorTabs::AssetDetailsID);
+
+	return SNew(SDockTab)
+		.Label(LOCTEXT("AssetDetails_Title", "Asset Details"))
+		[
+			EditingAssetWidget.ToSharedRef()
+		];
+}
+
+TSharedRef<SDockTab> FPoseSearchDatabaseEditorToolkit::SpawnTab_PreviewSettings(const FSpawnTabArgs& Args)
+{
+	check(Args.GetTabId() == FPoseSearchDatabaseEditorTabs::PreviewSettingsID);
+
+	FAdvancedPreviewSceneModule& AdvancedPreviewSceneModule = FModuleManager::LoadModuleChecked<FAdvancedPreviewSceneModule>("AdvancedPreviewScene");
+	TSharedRef<SWidget> InWidget= AdvancedPreviewSceneModule.CreateAdvancedPreviewSceneSettingsWidget(PreviewScene.ToSharedRef());
+
+	TSharedRef<SDockTab> SpawnedTab = SNew(SDockTab)
+		.Label(LOCTEXT("PreviewSceneSettingsTab", "Preview Scene Settings"))
+		[
+			InWidget
+		];
+
+	return SpawnedTab;
+}
+
+void FPoseSearchDatabaseEditorToolkit::OnFinishedChangingProperties(const FPropertyChangedEvent& PropertyChangedEvent)
+{
+	const FName PropertyName = 
+		(PropertyChangedEvent.Property != nullptr) ? 
+		PropertyChangedEvent.Property->GetFName() : NAME_None;
+	const FName MemberPropertyName = 
+		(PropertyChangedEvent.MemberProperty != nullptr) ? 
+		PropertyChangedEvent.MemberProperty->GetFName() : NAME_None;
+
+	UE_LOG(
+		LogPoseSearchEditor, 
+		Log, 
+		TEXT("FPoseSearchDatabaseEditorToolkit::OnFinishedChangingProperties MemberPropertyName: %s PropertyName: %s"),
+		*MemberPropertyName.ToString(), 
+		*PropertyName.ToString());
+}
+
+#undef LOCTEXT_NAMESPACE
diff --git a/Engine/Plugins/Experimental/Animation/PoseSearch/Source/Editor/Private/PoseSearchDatabasePreviewScene.cpp b/Engine/Plugins/Experimental/Animation/PoseSearch/Source/Editor/Private/PoseSearchDatabasePreviewScene.cpp
new file mode 100644
--- /dev/null
+++ b/Engine/Plugins/Experimental/Animation/PoseSearch/Source/Editor/Private/PoseSearchDatabasePreviewScene.cpp
@@ -0,0 +1,52 @@
+// Copyright Epic Games, Inc. All Rights Reserved.
+
+#include "PoseSearchDatabasePreviewScene.h"
+#include "PoseSearchDatabaseEditorToolkit.h"
+#include "GameFramework/WorldSettings.h"
+#include "EngineUtils.h"
+
+FPoseSearchDatabasePreviewScene::FPoseSearchDatabasePreviewScene(
+	ConstructionValues CVs, 
+	const TSharedRef<FPoseSearchDatabaseEditorToolkit>& EditorToolkit)
+	: FAdvancedPreviewScene(CVs)
+	, EditorToolkitPtr(EditorToolkit)
+{
+	// Disable killing actors outside of the world
+	AWorldSettings* WorldSettings = GetWorld()->GetWorldSettings(true);
+	WorldSettings->bEnableWorldBoundsChecks = false;
+
+	// Spawn an owner for FloorMeshComponent so CharacterMovementComponent can detect it as a valid floor and slide 
+	// along it
+	{
+		AActor* FloorActor = GetWorld()->SpawnActor<AActor>(AActor::StaticClass(), FTransform());
+		check(FloorActor);
+
+		static const FString NewName = FString(TEXT("FloorComponent"));
+		FloorMeshComponent->Rename(*NewName, FloorActor);
+
+		FloorActor->SetRootComponent(FloorMeshComponent);
+	}
+}
+
+void FPoseSearchDatabasePreviewScene::Tick(float InDeltaTime)
+{
+	FAdvancedPreviewScene::Tick(InDeltaTime);
+
+	// Trigger Begin Play in this preview world.
+	// This is needed for the CharacterMovementComponent to be able to switch to falling mode. 
+	// See: UCharacterMovementComponent::StartFalling
+	if (PreviewWorld && !PreviewWorld->bBegunPlay)
+	{
+		for (FActorIterator It(PreviewWorld); It; ++It)
+		{
+			It->DispatchBeginPlay();
+		}
+
+		PreviewWorld->bBegunPlay = true;
+	}
+
+	if (!GIntraFrameDebuggingGameThread)
+	{
+		GetWorld()->Tick(LEVELTICK_All, InDeltaTime);
+	}
+}
diff --git a/Engine/Plugins/Experimental/Animation/PoseSearch/Source/Editor/Private/PoseSearchDatabaseViewModel.cpp b/Engine/Plugins/Experimental/Animation/PoseSearch/Source/Editor/Private/PoseSearchDatabaseViewModel.cpp
new file mode 100644
--- /dev/null
+++ b/Engine/Plugins/Experimental/Animation/PoseSearch/Source/Editor/Private/PoseSearchDatabaseViewModel.cpp
@@ -0,0 +1,101 @@
+// Copyright Epic Games, Inc. All Rights Reserved.
+
+#include "PoseSearchDatabaseViewModel.h"
+#include "PoseSearchEditor.h"
+#include "PoseSearch/PoseSearch.h"
+#include "Modules/ModuleManager.h"
+#include "PoseSearchDatabasePreviewScene.h"
+#include "Animation/AnimInstance.h"
+#include "EngineUtils.h"
+#include "PropertyEditorModule.h"
+#include "Modules/ModuleManager.h"
+#include "GameFramework/Character.h"
+#include "GameFramework/CharacterMovementComponent.h"
+#include "Camera/CameraComponent.h"
+#include "Components/CapsuleComponent.h"
+
+FPoseSearchDatabaseViewModel::FPoseSearchDatabaseViewModel()
+	: PoseSearchDatabase(nullptr)
+{
+}
+
+FPoseSearchDatabaseViewModel::~FPoseSearchDatabaseViewModel()
+{
+}
+
+void FPoseSearchDatabaseViewModel::AddReferencedObjects(FReferenceCollector& Collector)
+{
+	Collector.AddReferencedObject(PoseSearchDatabase);
+}
+
+void FPoseSearchDatabaseViewModel::Initialize(
+	UPoseSearchDatabase* InPoseSearchDatabase,
+	const TSharedRef<FPoseSearchDatabasePreviewScene>& InPreviewScene)
+{
+	PoseSearchDatabase = InPoseSearchDatabase;
+	PreviewScenePtr = InPreviewScene;
+}
+
+void FPoseSearchDatabaseViewModel::RestartAnimations()
+{
+	// todo: implement
+}
+
+void FPoseSearchDatabaseViewModel::BuildSearchIndex()
+{
+	PoseSearchDatabase->BeginCacheDerivedData();
+}
+
+AActor* FPoseSearchDatabaseViewModel::SpawnPreviewActor(const FPoseSearchDatabaseSequence& DatabaseSequence)
+{
+	// todo: implement properly
+
+	UClass* PreviewClass = nullptr;
+	const FTransform SpawnTransform = FTransform::Identity;
+
+	FActorSpawnParameters Params;
+	Params.SpawnCollisionHandlingOverride = ESpawnActorCollisionHandlingMethod::AdjustIfPossibleButAlwaysSpawn;
+	AActor* PreviewActor = GetWorld()->SpawnActor<AActor>(PreviewClass, SpawnTransform, Params);
+
+	if (ACharacter* PreviewCharacter = Cast<ACharacter>(PreviewActor))
+	{
+		PreviewCharacter->bUseControllerRotationYaw = false;
+
+		if (UCharacterMovementComponent* CharacterMovementComp = PreviewCharacter->GetCharacterMovement())
+		{
+			CharacterMovementComp->bOrientRotationToMovement = true;
+			CharacterMovementComp->bUseControllerDesiredRotation = false;
+			CharacterMovementComp->RotationRate = FRotator(0.f, 540.0, 0.f);
+			CharacterMovementComp->bRunPhysicsWithNoController = true;
+
+			CharacterMovementComp->SetMovementMode(EMovementMode::MOVE_Walking);
+		}
+
+		if (UCameraComponent* CameraComp = PreviewCharacter->FindComponentByClass<UCameraComponent>())
+		{
+			CameraComp->DestroyComponent();
+		}
+	}
+
+	UE_LOG(LogPoseSearchEditor, Log, TEXT("Spawned preview Actor: %s at Loc: %s Rot: %s"),
+		*GetNameSafe(PreviewActor), *SpawnTransform.GetLocation().ToString(), *SpawnTransform.Rotator().ToString());
+
+	return PreviewActor;
+}
+
+UWorld* FPoseSearchDatabaseViewModel::GetWorld() const
+{
+	check(PreviewScenePtr.IsValid());
+	return PreviewScenePtr.Pin()->GetWorld();
+}
+
+UObject* FPoseSearchDatabaseViewModel::GetPlaybackContext() const
+{
+	return GetWorld();
+}
+
+
+void FPoseSearchDatabaseViewModel::OnPreviewActorClassChanged()
+{
+	// todo: implement
+}
\ No newline at end of file
diff --git a/Engine/Plugins/Experimental/Animation/PoseSearch/Source/Editor/Private/PoseSearchDatabaseViewportClient.cpp b/Engine/Plugins/Experimental/Animation/PoseSearch/Source/Editor/Private/PoseSearchDatabaseViewportClient.cpp
new file mode 100644
--- /dev/null
+++ b/Engine/Plugins/Experimental/Animation/PoseSearch/Source/Editor/Private/PoseSearchDatabaseViewportClient.cpp
@@ -0,0 +1,46 @@
+// Copyright Epic Games, Inc. All Rights Reserved.
+
+#include "PoseSearchDatabaseViewportClient.h"
+#include "PoseSearchDatabaseEditorToolkit.h"
+#include "PoseSearchDatabasePreviewScene.h"
+#include "SPoseSearchDatabaseViewport.h"
+#include "AssetEditorModeManager.h"
+#include "UnrealWidget.h"
+#include "PoseSearchDatabaseEdMode.h"
+
+FPoseSearchDatabaseViewportClient::FPoseSearchDatabaseViewportClient(
+	const TSharedRef<FPoseSearchDatabasePreviewScene>& InPreviewScene,
+	const TSharedRef<SPoseSearchDatabaseViewport>& InViewport, 
+	const TSharedRef<FPoseSearchDatabaseEditorToolkit>& InAssetEditorToolkit)
+	: FEditorViewportClient(nullptr, &InPreviewScene.Get(), StaticCastSharedRef<SEditorViewport>(InViewport))
+	, PreviewScenePtr(InPreviewScene)
+	, AssetEditorToolkitPtr(InAssetEditorToolkit)
+{
+	Widget->SetUsesEditorModeTools(ModeTools.Get());
+	StaticCastSharedPtr<FAssetEditorModeManager>(ModeTools)->SetPreviewScene(&InPreviewScene.Get());
+	ModeTools->SetDefaultMode(FPoseSearchDatabaseEdMode::EdModeId);
+
+	SetRealtime(true);
+
+	SetWidgetCoordSystemSpace(COORD_Local);
+	ModeTools->SetWidgetMode(UE::Widget::WM_Translate);
+}
+
+void FPoseSearchDatabaseViewportClient::TrackingStarted(
+	const struct FInputEventState& InInputState, 
+	bool bIsDraggingWidget, 
+	bool bNudge)
+{
+	ModeTools->StartTracking(this, Viewport);
+}
+
+void FPoseSearchDatabaseViewportClient::TrackingStopped()
+{
+	ModeTools->EndTracking(this, Viewport);
+	Invalidate();
+}
+
+void FPoseSearchDatabaseViewportClient::Draw(const FSceneView* View, FPrimitiveDrawInterface* PDI)
+{
+	FEditorViewportClient::Draw(View, PDI);
+}
diff --git a/Engine/Plugins/Experimental/Animation/PoseSearch/Source/Editor/Private/PoseSearchEditor.cpp b/Engine/Plugins/Experimental/Animation/PoseSearch/Source/Editor/Private/PoseSearchEditor.cpp
--- a/Engine/Plugins/Experimental/Animation/PoseSearch/Source/Editor/Private/PoseSearchEditor.cpp
+++ b/Engine/Plugins/Experimental/Animation/PoseSearch/Source/Editor/Private/PoseSearchEditor.cpp
@@ -2,18 +2,26 @@
 
 #include "PoseSearchEditor.h"
 #include "PoseSearchCustomization.h"
+#include "PoseSearchDebugger.h"
+#include "PoseSearchTypeActions.h"
+#include "PoseSearchDatabaseEdMode.h"
+#include "PoseSearchDatabaseEditorCommands.h"
+
 #include "Animation/AnimSequence.h"
 #include "Modules/ModuleManager.h"
+#include "AssetToolsModule.h"
 #include "Editor.h"
 #include "PropertyEditorModule.h"
 #include "Subsystems/AssetEditorSubsystem.h"
 #include "IAnimationEditor.h"
 #include "IPersonaToolkit.h"
 #include "IPersonaPreviewScene.h"
-#include "PoseSearchDebugger.h"
 #include "Trace/PoseSearchTraceAnalyzer.h"
 #include "Trace/PoseSearchTraceModule.h"
 
+DEFINE_LOG_CATEGORY(LogPoseSearchEditor);
+
+#define LOCTEXT_NAMESPACE "PoseSearchEditorModule"
 
 //////////////////////////////////////////////////////////////////////////
 // FEditorCommands
@@ -84,10 +92,15 @@ private:
 	/** List of registered class that we must unregister when the module shuts down */
 	TSet<FName> RegisteredClassNames;
 	TSet<FName> RegisteredPropertyTypes;
+
+	TSharedPtr<IAssetTypeActions> PoseSearchDatabaseActions;
 };
 
 void FEditorModule::StartupModule()
 {
+	// Register Asset Editor Commands
+	FPoseSearchDatabaseEditorCommands::Register();
+
 	if (GIsEditor && !IsRunningCommandlet())
 	{
 		FDebugger::Initialize();
@@ -103,13 +116,25 @@ void FEditorModule::StartupModule()
 			FConsoleCommandDelegate::CreateStatic(&FEditorCommands::DrawSearchIndex),
 			ECVF_Default
 		));
-	}
 
-	FPropertyEditorModule& PropertyModule = FModuleManager::GetModuleChecked<FPropertyEditorModule>("PropertyEditor");
+		// Register Ed Mode used by pose search database
+		FEditorModeRegistry::Get().RegisterMode<FPoseSearchDatabaseEdMode>(
+			FPoseSearchDatabaseEdMode::EdModeId,
+			LOCTEXT("PoseSearchDatabaseEdModeName", "PoseSearchDatabase"),
+			FSlateIcon(),
+			false,
+			9000);
+
+		// Register UPoseSearchDatabase Type Actions 
+		IAssetTools& AssetTools = FModuleManager::LoadModuleChecked<FAssetToolsModule>("AssetTools").Get();
+		PoseSearchDatabaseActions = MakeShared<FPoseSearchDatabaseTypeActions>();
+		AssetTools.RegisterAssetTypeActions(PoseSearchDatabaseActions.ToSharedRef());
+	}
 
 	RegisterPropertyTypeCustomizations();
 	RegisterObjectCustomizations();
 
+	FPropertyEditorModule& PropertyModule = FModuleManager::GetModuleChecked<FPropertyEditorModule>("PropertyEditor");
 	PropertyModule.NotifyCustomizationModuleChanged();
 
 }
@@ -121,8 +146,20 @@ void FEditorModule::ShutdownModule()
 		IConsoleManager::Get().UnregisterConsoleObject(ConsoleCmd);
 	}
 	ConsoleCommands.Empty();
+
+	if (FModuleManager::Get().IsModuleLoaded("AssetTools"))
+	{
+		IAssetTools& AssetTools = FModuleManager::GetModuleChecked<FAssetToolsModule>("AssetTools").Get();
+		AssetTools.UnregisterAssetTypeActions(PoseSearchDatabaseActions.ToSharedRef());
+	}
+
+	// Unregister Ed Mode
+	FEditorModeRegistry::Get().UnregisterMode(FPoseSearchDatabaseEdMode::EdModeId);
 	
 	UnregisterCustomizations();
+
+	// Unregister Asset Editor Commands
+	FPoseSearchDatabaseEditorCommands::Unregister();
 	
 	IModularFeatures::Get().UnregisterModularFeature(TraceServices::ModuleFeatureName, TraceModule.Get());
 	FDebugger::Shutdown();
@@ -195,3 +232,5 @@ void FEditorModule::UnregisterCustomizations()
 }} // namespace UE::PoseSearch
 
 IMPLEMENT_MODULE(UE::PoseSearch::FEditorModule, PoseSearchEditor);
+
+#undef LOCTEXT_NAMESPACE
diff --git a/Engine/Plugins/Experimental/Animation/PoseSearch/Source/Editor/Private/PoseSearchTypeActions.cpp b/Engine/Plugins/Experimental/Animation/PoseSearch/Source/Editor/Private/PoseSearchTypeActions.cpp
new file mode 100644
--- /dev/null
+++ b/Engine/Plugins/Experimental/Animation/PoseSearch/Source/Editor/Private/PoseSearchTypeActions.cpp
@@ -0,0 +1,46 @@
+// Copyright Epic Games, Inc. All Rights Reserved.
+
+#include "PoseSearchTypeActions.h"
+#include "PoseSearchDatabaseEditorToolkit.h"
+#include "PoseSearch/PoseSearch.h"
+
+#define LOCTEXT_NAMESPACE "PoseSearchTypeActions"
+
+FText FPoseSearchDatabaseTypeActions::GetName() const
+{
+	return LOCTEXT("PoseSearchDatabaseTypeActionsName", "Pose Search Database");
+}
+
+FColor FPoseSearchDatabaseTypeActions::GetTypeColor() const
+{
+	return FColor(129, 196, 115);
+}
+
+UClass* FPoseSearchDatabaseTypeActions::GetSupportedClass() const
+{
+	return UPoseSearchDatabase::StaticClass();
+}
+
+void FPoseSearchDatabaseTypeActions::OpenAssetEditor(
+	const TArray<UObject*>& InObjects, 
+	TSharedPtr<class IToolkitHost> EditWithinLevelEditor)
+{
+	const EToolkitMode::Type Mode = 
+		EditWithinLevelEditor.IsValid() ? EToolkitMode::WorldCentric : EToolkitMode::Standalone;
+
+	for (auto ObjIt = InObjects.CreateConstIterator(); ObjIt; ++ObjIt)
+	{
+		if (UPoseSearchDatabase* PoseSearchDb = Cast<UPoseSearchDatabase>(*ObjIt))
+		{
+			TSharedRef<FPoseSearchDatabaseEditorToolkit> NewEditor(new FPoseSearchDatabaseEditorToolkit());
+			NewEditor->InitAssetEditor(Mode, EditWithinLevelEditor, PoseSearchDb);
+		}
+	}
+}
+
+uint32 FPoseSearchDatabaseTypeActions::GetCategories()
+{
+	return EAssetTypeCategories::Animation;
+}
+
+#undef LOCTEXT_NAMESPACE
diff --git a/Engine/Plugins/Experimental/Animation/PoseSearch/Source/Editor/Private/SPoseSearchDatabaseViewport.cpp b/Engine/Plugins/Experimental/Animation/PoseSearch/Source/Editor/Private/SPoseSearchDatabaseViewport.cpp
new file mode 100644
--- /dev/null
+++ b/Engine/Plugins/Experimental/Animation/PoseSearch/Source/Editor/Private/SPoseSearchDatabaseViewport.cpp
@@ -0,0 +1,54 @@
+// Copyright Epic Games, Inc. All Rights Reserved.
+
+#include "SPoseSearchDatabaseViewport.h"
+#include "PoseSearchDatabaseViewportClient.h"
+#include "Toolkits/AssetEditorToolkit.h"
+#include "PoseSearchDatabaseEditorCommands.h"
+
+void SPoseSearchDatabaseViewport::Construct(
+	const FArguments& InArgs, 
+	const FPoseSearchDatabaseViewportRequiredArgs& InRequiredArgs)
+{
+	PreviewScenePtr = InRequiredArgs.PreviewScene;
+	AssetEditorToolkitPtr = InRequiredArgs.AssetEditorToolkit;
+
+	SEditorViewport::Construct(
+		SEditorViewport::FArguments()
+		.IsEnabled(FSlateApplication::Get().GetNormalExecutionAttribute())
+		.AddMetaData<FTagMetaData>(TEXT("AnimationTools.Viewport"))
+	);
+}
+
+void SPoseSearchDatabaseViewport::BindCommands()
+{
+	SEditorViewport::BindCommands();
+}
+
+TSharedRef<FEditorViewportClient> SPoseSearchDatabaseViewport::MakeEditorViewportClient()
+{
+	ViewportClient = MakeShared<FPoseSearchDatabaseViewportClient>(
+		PreviewScenePtr.Pin().ToSharedRef(), 
+		SharedThis(this), 
+		AssetEditorToolkitPtr.Pin().ToSharedRef());
+	ViewportClient->ViewportType = LVT_Perspective;
+	ViewportClient->bSetListenerPosition = false;
+	ViewportClient->SetViewLocation(EditorViewportDefs::DefaultPerspectiveViewLocation);
+	ViewportClient->SetViewRotation(EditorViewportDefs::DefaultPerspectiveViewRotation);
+
+	return ViewportClient.ToSharedRef();
+}
+
+TSharedRef<SEditorViewport> SPoseSearchDatabaseViewport::GetViewportWidget()
+{
+	return SharedThis(this);
+}
+
+TSharedPtr<FExtender> SPoseSearchDatabaseViewport::GetExtenders() const
+{
+	TSharedPtr<FExtender> Result(MakeShareable(new FExtender));
+	return Result;
+}
+
+void SPoseSearchDatabaseViewport::OnFloatingButtonClicked()
+{
+}
diff --git a/Engine/Plugins/Experimental/Animation/PoseSearch/Source/Editor/Public/PoseSearchDatabaseEdMode.h b/Engine/Plugins/Experimental/Animation/PoseSearch/Source/Editor/Public/PoseSearchDatabaseEdMode.h
new file mode 100644
--- /dev/null
+++ b/Engine/Plugins/Experimental/Animation/PoseSearch/Source/Editor/Public/PoseSearchDatabaseEdMode.h
@@ -0,0 +1,47 @@
+// Copyright Epic Games, Inc. All Rights Reserved.
+
+#pragma once
+
+#include "CoreMinimal.h"
+#include "EdMode.h"
+
+class FPoseSearchDatabaseViewModel;
+
+class FPoseSearchDatabaseEdMode : public FEdMode
+{
+public:
+
+	const static FEditorModeID EdModeId;
+
+	FPoseSearchDatabaseEdMode();
+	virtual ~FPoseSearchDatabaseEdMode();
+
+	virtual void Tick(FEditorViewportClient* ViewportClient, float DeltaTime) override;
+	virtual void Render(const FSceneView* View, FViewport* Viewport, FPrimitiveDrawInterface* PDI) override;
+	virtual bool HandleClick(
+		FEditorViewportClient* InViewportClient, 
+		HHitProxy* HitProxy, 
+		const FViewportClick& Click) override;
+	virtual bool StartTracking(FEditorViewportClient* InViewportClient, FViewport* InViewport) override;
+	virtual bool EndTracking(FEditorViewportClient* InViewportClient, FViewport* InViewport) override;
+	virtual bool InputDelta(
+		FEditorViewportClient* InViewportClient, 
+		FViewport* InViewport, 
+		FVector& InDrag, 
+		FRotator& InRot, 
+		FVector& InScale) override;
+	virtual bool InputKey(
+		FEditorViewportClient* ViewportClient, 
+		FViewport* Viewport, 
+		FKey Key, 
+		EInputEvent Event) override;
+	virtual bool AllowWidgetMove() override;
+	virtual bool ShouldDrawWidget() const override;
+	virtual FVector GetWidgetLocation() const override;
+	virtual bool GetCustomDrawingCoordinateSystem(FMatrix& InMatrix, void* InData) override;
+	virtual bool GetCustomInputCoordinateSystem(FMatrix& InMatrix, void* InData) override;
+
+private:
+
+	FPoseSearchDatabaseViewModel* ViewModel = nullptr;
+};
diff --git a/Engine/Plugins/Experimental/Animation/PoseSearch/Source/Editor/Public/PoseSearchDatabaseEditorCommands.h b/Engine/Plugins/Experimental/Animation/PoseSearch/Source/Editor/Public/PoseSearchDatabaseEditorCommands.h
new file mode 100644
--- /dev/null
+++ b/Engine/Plugins/Experimental/Animation/PoseSearch/Source/Editor/Public/PoseSearchDatabaseEditorCommands.h
@@ -0,0 +1,28 @@
+// Copyright Epic Games, Inc. All Rights Reserved.
+
+#pragma once
+
+#include "CoreMinimal.h"
+#include "Framework/Commands/Commands.h"
+#include "EditorStyleSet.h"
+
+class FPoseSearchDatabaseEditorCommands : public TCommands<FPoseSearchDatabaseEditorCommands>
+{
+public:
+	FPoseSearchDatabaseEditorCommands()
+		: TCommands<FPoseSearchDatabaseEditorCommands>(
+			TEXT("PoseSearchDatabaseEditor"), 
+			NSLOCTEXT("Contexts", "PoseSearchDatabase", "Pose Search Database"), 
+			NAME_None, 
+			FEditorStyle::GetStyleSetName())
+	{
+	}
+
+	virtual void RegisterCommands() override;
+
+public:
+
+	TSharedPtr<FUICommandInfo> ResetPreviewScene;
+
+	TSharedPtr<FUICommandInfo> BuildSearchIndex;
+};
diff --git a/Engine/Plugins/Experimental/Animation/PoseSearch/Source/Editor/Public/PoseSearchDatabaseEditorToolkit.h b/Engine/Plugins/Experimental/Animation/PoseSearch/Source/Editor/Public/PoseSearchDatabaseEditorToolkit.h
new file mode 100644
--- /dev/null
+++ b/Engine/Plugins/Experimental/Animation/PoseSearch/Source/Editor/Public/PoseSearchDatabaseEditorToolkit.h
@@ -0,0 +1,61 @@
+// Copyright Epic Games, Inc. All Rights Reserved.
+
+#pragma once
+
+#include "CoreMinimal.h"
+#include "Toolkits/AssetEditorToolkit.h"
+#include "Misc/NotifyHook.h"
+
+class IDetailsView;
+class FToolBarBuilder;
+class UPoseSearchDatabase;
+class SPoseSearchDatabaseViewport;
+class FPoseSearchDatabasePreviewScene;
+class FPoseSearchDatabaseViewModel;
+
+class FPoseSearchDatabaseEditorToolkit : public FAssetEditorToolkit, public FNotifyHook
+{
+public:
+
+	FPoseSearchDatabaseEditorToolkit();
+	virtual ~FPoseSearchDatabaseEditorToolkit();
+
+	void InitAssetEditor(
+		const EToolkitMode::Type Mode, 
+		const TSharedPtr<IToolkitHost>& InitToolkitHost, 
+		UPoseSearchDatabase* DatabaseAsset);
+
+	virtual void RegisterTabSpawners(const TSharedRef<FTabManager>& TabManager) override;
+	virtual void UnregisterTabSpawners(const TSharedRef<FTabManager>& TabManager) override;
+	virtual FName GetToolkitFName() const override;
+	virtual FText GetBaseToolkitName() const override;
+	virtual FText GetToolkitName() const override;
+	virtual FLinearColor GetWorldCentricTabColorScale() const override;
+	virtual FString GetWorldCentricTabPrefix() const override;
+
+	UPoseSearchDatabase* GetPoseSearchDatabase() const;
+	FPoseSearchDatabaseViewModel* GetViewModel() const { return ViewModel.Get(); }
+
+	void ResetPreviewScene();
+	void BuildSearchIndex();
+
+private:
+	
+	TSharedRef<SDockTab> SpawnTab_Viewport(const FSpawnTabArgs& Args);
+	TSharedRef<SDockTab> SpawnTab_AssetDetails(const FSpawnTabArgs& Args);
+	TSharedRef<SDockTab> SpawnTab_PreviewSettings(const FSpawnTabArgs& Args);
+
+	void BindCommands();
+	void ExtendToolbar();
+	void FillToolbar(FToolBarBuilder& ToolbarBuilder);
+
+	void OnFinishedChangingProperties(const FPropertyChangedEvent& PropertyChangedEvent);
+
+	TSharedPtr<SPoseSearchDatabaseViewport> ViewportWidget;
+
+	TSharedPtr<IDetailsView> EditingAssetWidget;
+
+	TSharedPtr<FPoseSearchDatabasePreviewScene> PreviewScene;
+
+	TSharedPtr<FPoseSearchDatabaseViewModel> ViewModel;
+};
diff --git a/Engine/Plugins/Experimental/Animation/PoseSearch/Source/Editor/Public/PoseSearchDatabasePreviewScene.h b/Engine/Plugins/Experimental/Animation/PoseSearch/Source/Editor/Public/PoseSearchDatabasePreviewScene.h
new file mode 100644
--- /dev/null
+++ b/Engine/Plugins/Experimental/Animation/PoseSearch/Source/Editor/Public/PoseSearchDatabasePreviewScene.h
@@ -0,0 +1,27 @@
+// Copyright Epic Games, Inc. All Rights Reserved.
+
+#pragma once
+
+#include "CoreMinimal.h"
+#include "AdvancedPreviewScene.h"
+
+class FPoseSearchDatabaseEditorToolkit;
+
+class FPoseSearchDatabasePreviewScene : public FAdvancedPreviewScene
+{
+public:
+
+	FPoseSearchDatabasePreviewScene(
+		ConstructionValues CVs, 
+		const TSharedRef<FPoseSearchDatabaseEditorToolkit>& EditorToolkit);
+	~FPoseSearchDatabasePreviewScene(){}
+
+	virtual void Tick(float InDeltaTime) override;
+
+	TSharedRef<FPoseSearchDatabaseEditorToolkit> GetEditorToolkit() const { return EditorToolkitPtr.Pin().ToSharedRef(); }
+
+private:
+
+	/** The asset editor toolkit we are embedded in */
+	TWeakPtr<FPoseSearchDatabaseEditorToolkit> EditorToolkitPtr;
+};
\ No newline at end of file
diff --git a/Engine/Plugins/Experimental/Animation/PoseSearch/Source/Editor/Public/PoseSearchDatabaseViewModel.h b/Engine/Plugins/Experimental/Animation/PoseSearch/Source/Editor/Public/PoseSearchDatabaseViewModel.h
new file mode 100644
--- /dev/null
+++ b/Engine/Plugins/Experimental/Animation/PoseSearch/Source/Editor/Public/PoseSearchDatabaseViewModel.h
@@ -0,0 +1,50 @@
+// Copyright Epic Games, Inc. All Rights Reserved.
+
+#pragma once
+
+#include "CoreMinimal.h"
+#include "UObject/GCObject.h"
+#include "MovieSceneFwd.h"
+
+class UWorld;
+class FPoseSearchDatabasePreviewScene;
+class UPoseSearchDatabase;
+struct FPoseSearchDatabaseSequence;
+class IDetailsView;
+class IStructureDetailsView;
+
+class FPoseSearchDatabaseViewModel : public TSharedFromThis<FPoseSearchDatabaseViewModel>, public FGCObject
+{
+public:
+
+	FPoseSearchDatabaseViewModel();
+	virtual ~FPoseSearchDatabaseViewModel();
+
+	// ~ FGCObject interface
+	virtual void AddReferencedObjects(FReferenceCollector& Collector) override;
+	virtual FString GetReferencerName() const override { return TEXT("FPoseSearchDatabaseViewModel"); }
+
+	void Initialize(
+		UPoseSearchDatabase* InPoseSearchDatabase, 
+		const TSharedRef<FPoseSearchDatabasePreviewScene>& InPreviewScene);
+
+	void RestartAnimations();
+	void BuildSearchIndex();
+
+	UPoseSearchDatabase* GetPoseSearchDatabase() const { return PoseSearchDatabase; }
+	void OnPreviewActorClassChanged();
+
+private:
+
+	/** Scene asset being viewed and edited by this view model. */
+	TObjectPtr<UPoseSearchDatabase> PoseSearchDatabase;
+
+	/** Weak pointer to the PreviewScene */
+	TWeakPtr<FPoseSearchDatabasePreviewScene> PreviewScenePtr;
+
+	AActor* SpawnPreviewActor(const FPoseSearchDatabaseSequence& DatabaseSequence);
+
+	UWorld* GetWorld() const;
+
+	UObject* GetPlaybackContext() const;
+};
\ No newline at end of file
diff --git a/Engine/Plugins/Experimental/Animation/PoseSearch/Source/Editor/Public/PoseSearchDatabaseViewportClient.h b/Engine/Plugins/Experimental/Animation/PoseSearch/Source/Editor/Public/PoseSearchDatabaseViewportClient.h
new file mode 100644
--- /dev/null
+++ b/Engine/Plugins/Experimental/Animation/PoseSearch/Source/Editor/Public/PoseSearchDatabaseViewportClient.h
@@ -0,0 +1,46 @@
+// Copyright Epic Games, Inc. All Rights Reserved.
+
+#pragma once
+
+#include "CoreMinimal.h"
+#include "EditorViewportClient.h"
+
+class FPoseSearchDatabasePreviewScene;
+class SPoseSearchDatabaseViewport;
+class FPoseSearchDatabaseEditorToolkit;
+
+class FPoseSearchDatabaseViewportClient : public FEditorViewportClient
+{
+public:
+
+	FPoseSearchDatabaseViewportClient(
+		const TSharedRef<FPoseSearchDatabasePreviewScene>& InPreviewScene, 
+		const TSharedRef<SPoseSearchDatabaseViewport>& InViewport, 
+		const TSharedRef<FPoseSearchDatabaseEditorToolkit>& InAssetEditorToolkit);
+	virtual ~FPoseSearchDatabaseViewportClient(){}
+
+	// ~FEditorViewportClient interface
+	virtual void Draw(const FSceneView* View, FPrimitiveDrawInterface* PDI) override;
+	virtual void TrackingStarted(const struct FInputEventState& InInputState, bool bIsDragging, bool bNudge) override;
+	virtual void TrackingStopped() override;
+	// ~End of FEditorViewportClient interface
+
+	/** Get the preview scene we are viewing */
+	TSharedRef<FPoseSearchDatabasePreviewScene> GetPreviewScene() const
+	{
+		return PreviewScenePtr.Pin().ToSharedRef(); 
+	}
+
+	TSharedRef<FPoseSearchDatabaseEditorToolkit> GetAssetEditorToolkit() const
+	{
+		return AssetEditorToolkitPtr.Pin().ToSharedRef(); 
+	}
+
+private:
+
+	/** Preview scene we are viewing */
+	TWeakPtr<FPoseSearchDatabasePreviewScene> PreviewScenePtr;
+
+	/** Asset editor toolkit we are embedded in */
+	TWeakPtr<FPoseSearchDatabaseEditorToolkit> AssetEditorToolkitPtr;
+};
diff --git a/Engine/Plugins/Experimental/Animation/PoseSearch/Source/Editor/Public/PoseSearchEditor.h b/Engine/Plugins/Experimental/Animation/PoseSearch/Source/Editor/Public/PoseSearchEditor.h
--- a/Engine/Plugins/Experimental/Animation/PoseSearch/Source/Editor/Public/PoseSearchEditor.h
+++ b/Engine/Plugins/Experimental/Animation/PoseSearch/Source/Editor/Public/PoseSearchEditor.h
@@ -6,6 +6,9 @@
 #include "Modules/ModuleInterface.h"
 #include "PoseSearch/PoseSearch.h"
 
+
+DECLARE_LOG_CATEGORY_EXTERN(LogPoseSearchEditor, Log, All);
+
 class IPoseSearchEditorModuleInterface : public IModuleInterface
 {
 };
\ No newline at end of file
diff --git a/Engine/Plugins/Experimental/Animation/PoseSearch/Source/Editor/Public/PoseSearchTypeActions.h b/Engine/Plugins/Experimental/Animation/PoseSearch/Source/Editor/Public/PoseSearchTypeActions.h
new file mode 100644
--- /dev/null
+++ b/Engine/Plugins/Experimental/Animation/PoseSearch/Source/Editor/Public/PoseSearchTypeActions.h
@@ -0,0 +1,20 @@
+// Copyright Epic Games, Inc. All Rights Reserved.
+
+#pragma once
+
+#include "CoreMinimal.h"
+#include "AssetTypeActions_Base.h"
+
+class FPoseSearchDatabaseTypeActions : public FAssetTypeActions_Base
+{
+public:
+
+	FPoseSearchDatabaseTypeActions(){}
+	virtual FText GetName() const override;
+	virtual FColor GetTypeColor() const override;
+	virtual UClass* GetSupportedClass() const override;
+	virtual void OpenAssetEditor(
+		const TArray<UObject*>& InObjects, 
+		TSharedPtr<class IToolkitHost> EditWithinLevelEditor = TSharedPtr<IToolkitHost>()) override;
+	virtual uint32 GetCategories() override;
+};
\ No newline at end of file
diff --git a/Engine/Plugins/Experimental/Animation/PoseSearch/Source/Editor/Public/SPoseSearchDatabaseViewport.h b/Engine/Plugins/Experimental/Animation/PoseSearch/Source/Editor/Public/SPoseSearchDatabaseViewport.h
new file mode 100644
--- /dev/null
+++ b/Engine/Plugins/Experimental/Animation/PoseSearch/Source/Editor/Public/SPoseSearchDatabaseViewport.h
@@ -0,0 +1,62 @@
+// Copyright Epic Games, Inc. All Rights Reserved.
+
+#pragma once
+
+#include "CoreMinimal.h"
+#include "SEditorViewport.h"
+#include "SCommonEditorViewportToolbarBase.h"
+
+class FPoseSearchDatabasePreviewScene;
+class FPoseSearchDatabaseViewportClient;
+class FPoseSearchDatabaseEditorToolkit;
+class SPoseSearchDatabaseViewportToolBar;
+
+struct FPoseSearchDatabaseViewportRequiredArgs
+{
+	FPoseSearchDatabaseViewportRequiredArgs(
+		const TSharedRef<FPoseSearchDatabaseEditorToolkit>& InAssetEditorToolkit, 
+		const TSharedRef<FPoseSearchDatabasePreviewScene>& InPreviewScene)
+		: AssetEditorToolkit(InAssetEditorToolkit)
+		, PreviewScene(InPreviewScene)
+	{}
+
+	TSharedRef<FPoseSearchDatabaseEditorToolkit> AssetEditorToolkit;
+
+	TSharedRef<FPoseSearchDatabasePreviewScene> PreviewScene;
+};
+
+class SPoseSearchDatabaseViewport : public SEditorViewport, public ICommonEditorViewportToolbarInfoProvider
+{
+public:
+	
+	SLATE_BEGIN_ARGS(SPoseSearchDatabaseViewport) {}
+	SLATE_END_ARGS();
+
+	void Construct(const FArguments& InArgs, const FPoseSearchDatabaseViewportRequiredArgs& InRequiredArgs);
+	virtual ~SPoseSearchDatabaseViewport(){}
+
+	// ~ICommonEditorViewportToolbarInfoProvider interface
+	virtual TSharedRef<class SEditorViewport> GetViewportWidget() override;
+	virtual TSharedPtr<FExtender> GetExtenders() const override;
+	virtual void OnFloatingButtonClicked() override;
+	// ~End of ICommonEditorViewportToolbarInfoProvider interface
+
+protected:
+
+	// ~SEditorViewport interface
+	virtual void BindCommands() override;
+	virtual TSharedRef<FEditorViewportClient> MakeEditorViewportClient() override;
+	// ~End of SEditorViewport interface
+
+	/** The viewport toolbar */
+	TSharedPtr<SPoseSearchDatabaseViewportToolBar> ViewportToolbar;
+
+	/** Viewport client */
+	TSharedPtr<FPoseSearchDatabaseViewportClient> ViewportClient;
+
+	/** The preview scene that we are viewing */
+	TWeakPtr<FPoseSearchDatabasePreviewScene> PreviewScenePtr;
+
+	/** Asset editor toolkit we are embedded in */
+	TWeakPtr<FPoseSearchDatabaseEditorToolkit> AssetEditorToolkitPtr;
+};
\ No newline at end of file
