exposing looping property into MM debugger

diff --git a/Engine/Plugins/Experimental/Animation/PoseSearch/Source/Editor/Private/PoseSearchDebugger.cpp b/Engine/Plugins/Experimental/Animation/PoseSearch/Source/Editor/Private/PoseSearchDebugger.cpp
--- a/Engine/Plugins/Experimental/Animation/PoseSearch/Source/Editor/Private/PoseSearchDebugger.cpp
+++ b/Engine/Plugins/Experimental/Animation/PoseSearch/Source/Editor/Private/PoseSearchDebugger.cpp
@@ -276,6 +276,7 @@ public:
 	int32 AnimFrame = 0;
 	float AssetTime = 0.0f;
 	bool bMirrored = false;
+	bool bLooping = false;
 	FVector BlendParameters = FVector::Zero();
 	FPoseCostDetails PoseCostDetails;
 };
@@ -544,6 +545,24 @@ namespace DebuggerDatabaseColumns
 	};
 	const FName FMirrored::Name = "Mirrored";
 
+	struct FLooping : ITextColumn
+	{
+		using ITextColumn::ITextColumn;
+		static const FName Name;
+		virtual FName GetName() const override { return Name; }
+
+		virtual FSortPredicate GetSortPredicate() const override
+		{
+			return [](const FRowDataRef& Row0, const FRowDataRef& Row1) -> bool { return Row0->bLooping < Row1->bLooping; };
+		}
+
+		virtual FText GetRowText(const FRowDataRef& Row) const override
+		{
+			return FText::Format(LOCTEXT("Looping", "{0}"), { Row->bLooping });
+		}
+	};
+	const FName FLooping::Name = "Looping";
+
 	struct FBlendParameters : ITextColumn
 	{
 		using ITextColumn::ITextColumn;
@@ -891,6 +910,7 @@ void SDebuggerDatabaseView::CreateRows(const UPoseSearchDatabase& Database)
 				Row->AssetTime = Time;
 				Row->AnimFrame = DbSequence.Sequence->GetFrameAtTime(Time);
 				Row->bMirrored = SearchIndexAsset.bMirrored;
+				Row->bLooping = DbSequence.bLoopAnimation;
 				Row->BlendParameters = FVector::Zero();
 			}
 		}
@@ -911,6 +931,7 @@ void SDebuggerDatabaseView::CreateRows(const UPoseSearchDatabase& Database)
 				Row->AssetTime = Time;
 				Row->AnimFrame = 0; // There is no frame index associated with a blendspace
 				Row->bMirrored = SearchIndexAsset.bMirrored;
+				Row->bLooping = DbBlendSpace.bLoopAnimation;
 				Row->BlendParameters = SearchIndexAsset.BlendParameters;
 			}
 		}
@@ -1006,8 +1027,9 @@ void SDebuggerDatabaseView::Construct(const FArguments& InArgs)
 	AddColumn(MakeShared<FCostModifier>(5));
 	AddColumn(MakeShared<FFrame>(6));
 	AddColumn(MakeShared<FMirrored>(7));
-	AddColumn(MakeShared<FBlendParameters>(8));
-	AddColumn(MakeShared<FPoseIdx>(9));
+	AddColumn(MakeShared<FLooping>(8));
+	AddColumn(MakeShared<FBlendParameters>(9));
+	AddColumn(MakeShared<FPoseIdx>(10));
 
 	// Active Row
 	ActiveView.HeaderRow = SNew(SHeaderRow);
