#PoseSearch - Fix crash with bad notify data.

diff --git a/Engine/Plugins/Experimental/Animation/PoseSearch/Source/Runtime/Private/PoseSearchDerivedData.cpp b/Engine/Plugins/Experimental/Animation/PoseSearch/Source/Runtime/Private/PoseSearchDerivedData.cpp
--- a/Engine/Plugins/Experimental/Animation/PoseSearch/Source/Runtime/Private/PoseSearchDerivedData.cpp
+++ b/Engine/Plugins/Experimental/Animation/PoseSearch/Source/Runtime/Private/PoseSearchDerivedData.cpp
@@ -267,7 +267,7 @@ namespace UE::PoseSearch
 		for (const FAnimNotifyEventReference& EventReference : NotifyContext.ActiveNotifies)
 		{
 			const FAnimNotifyEvent* NotifyEvent = EventReference.GetNotify();
-			if (!NotifyEvent)
+			if (!NotifyEvent || !NotifyEvent->NotifyStateClass)
 			{
 				continue;
 			}
