Pose Search Database Editor - Added Viewport to play database sequences
- viewport toolbar show menu allows selecting:
   - feature display
   - mirrored animation display
- editor toolbar buttons:
   - stop preview
   - reset preview
   - force database rebuild (in case there's an issue with the DDC integration this is the failsafe

diff --git a/Engine/Plugins/Experimental/Animation/PoseSearch/Source/Editor/Private/PoseSearchDatabaseEdMode.cpp b/Engine/Plugins/Experimental/Animation/PoseSearch/Source/Editor/Private/PoseSearchDatabaseEdMode.cpp
--- a/Engine/Plugins/Experimental/Animation/PoseSearch/Source/Editor/Private/PoseSearchDatabaseEdMode.cpp
+++ b/Engine/Plugins/Experimental/Animation/PoseSearch/Source/Editor/Private/PoseSearchDatabaseEdMode.cpp
@@ -28,6 +28,11 @@ void FPoseSearchDatabaseEdMode::Tick(FEditorViewportClient* ViewportClient, floa
 	{
 		ViewModel = PoseSearchDbViewportClient->GetAssetEditorToolkit()->GetViewModel();
 	}
+
+	if (ViewModel)
+	{
+		ViewModel->Tick(DeltaTime);
+	}
 }
 
 void FPoseSearchDatabaseEdMode::Render(const FSceneView* View, FViewport* Viewport, FPrimitiveDrawInterface* PDI)
diff --git a/Engine/Plugins/Experimental/Animation/PoseSearch/Source/Editor/Private/PoseSearchDatabaseEditorCommands.cpp b/Engine/Plugins/Experimental/Animation/PoseSearch/Source/Editor/Private/PoseSearchDatabaseEditorCommands.cpp
--- a/Engine/Plugins/Experimental/Animation/PoseSearch/Source/Editor/Private/PoseSearchDatabaseEditorCommands.cpp
+++ b/Engine/Plugins/Experimental/Animation/PoseSearch/Source/Editor/Private/PoseSearchDatabaseEditorCommands.cpp
@@ -6,8 +6,24 @@
 
 void FPoseSearchDatabaseEditorCommands::RegisterCommands()
 {
+	UI_COMMAND(StopPreviewScene, "Stop Scene", "Stop Scene", EUserInterfaceActionType::Button, FInputChord());
 	UI_COMMAND(ResetPreviewScene, "Reset Scene", "Reset Scene", EUserInterfaceActionType::Button, FInputChord());
 	UI_COMMAND(BuildSearchIndex, "Build Index", "Build Index", EUserInterfaceActionType::Button, FInputChord());
+
+	UI_COMMAND(
+		ShowPoseFeaturesNone, "None", "Hide all features", 
+		EUserInterfaceActionType::RadioButton, FInputChord());
+	UI_COMMAND(
+		ShowPoseFeaturesAll, "All Features", "Show all features", 
+		EUserInterfaceActionType::RadioButton, FInputChord());
+
+	UI_COMMAND(
+		ShowAnimationOriginalOnly, "Original Only", "Show only original animations",
+		EUserInterfaceActionType::RadioButton, FInputChord());
+	UI_COMMAND(
+		ShowAnimationOriginalAndMirrored, "Original and mirrored", "Show original and mirrored animations",
+		EUserInterfaceActionType::RadioButton, FInputChord());
+
 }
 
 #undef LOCTEXT_NAMESPACE
diff --git a/Engine/Plugins/Experimental/Animation/PoseSearch/Source/Editor/Private/PoseSearchDatabaseEditorToolkit.cpp b/Engine/Plugins/Experimental/Animation/PoseSearch/Source/Editor/Private/PoseSearchDatabaseEditorToolkit.cpp
--- a/Engine/Plugins/Experimental/Animation/PoseSearch/Source/Editor/Private/PoseSearchDatabaseEditorToolkit.cpp
+++ b/Engine/Plugins/Experimental/Animation/PoseSearch/Source/Editor/Private/PoseSearchDatabaseEditorToolkit.cpp
@@ -39,14 +39,19 @@ FPoseSearchDatabaseEditorToolkit::~FPoseSearchDatabaseEditorToolkit()
 {
 }
 
-UPoseSearchDatabase* FPoseSearchDatabaseEditorToolkit::GetPoseSearchDatabase() const
+const UPoseSearchDatabase* FPoseSearchDatabaseEditorToolkit::GetPoseSearchDatabase() const
 {
 	return ViewModel.IsValid() ? ViewModel->GetPoseSearchDatabase() : nullptr;
 }
 
+void FPoseSearchDatabaseEditorToolkit::StopPreviewScene()
+{
+	ViewModel->RemovePreviewActors();
+}
+
 void FPoseSearchDatabaseEditorToolkit::ResetPreviewScene()
 {
-	ViewModel->RestartAnimations();
+	ViewModel->ResetPreviewActors();
 }
 
 void FPoseSearchDatabaseEditorToolkit::BuildSearchIndex()
@@ -77,15 +82,16 @@ void FPoseSearchDatabaseEditorToolkit::InitAssetEditor(
 		PreviewScene->GetWorld()->GetWorldSettings()->SetIsTemporarilyHiddenInEditor(false);
 	}
 
+	// Create view model
+	ViewModel = MakeShared<FPoseSearchDatabaseViewModel>();
+	ViewModel->Initialize(DatabaseAsset, PreviewScene.ToSharedRef());
+
 	// Create viewport widget
 	FPoseSearchDatabaseViewportRequiredArgs ViewportArgs(
 		StaticCastSharedRef<FPoseSearchDatabaseEditorToolkit>(AsShared()), 
 		PreviewScene.ToSharedRef());
 	ViewportWidget = SNew(SPoseSearchDatabaseViewport, ViewportArgs);
 
-	ViewModel = MakeShared<FPoseSearchDatabaseViewModel>();
-	ViewModel->Initialize(DatabaseAsset, PreviewScene.ToSharedRef());
-
 	// Create Asset Details widget
 	FPropertyEditorModule& PropertyModule = FModuleManager::LoadModuleChecked<FPropertyEditorModule>("PropertyEditor");
 
@@ -101,7 +107,7 @@ void FPoseSearchDatabaseEditorToolkit::InitAssetEditor(
 
 	// Define Editor Layout
 	const TSharedRef<FTabManager::FLayout> StandaloneDefaultLayout = 
-		FTabManager::NewLayout("Standalone_PoseSearchDatabaseEditor_Layout_v0.01")
+		FTabManager::NewLayout("Standalone_PoseSearchDatabaseEditor_Layout_v0.02")
 		->AddArea
 		(
 			FTabManager::NewPrimaryArea()->SetOrientation(Orient_Vertical)
@@ -110,25 +116,25 @@ void FPoseSearchDatabaseEditorToolkit::InitAssetEditor(
 				FTabManager::NewSplitter()->SetOrientation(Orient_Vertical)->SetSizeCoefficient(0.9f)
 				->Split
 				(
-					//FTabManager::NewSplitter()->SetOrientation(Orient_Horizontal)->SetSizeCoefficient(0.9f)
-					//->Split
-					//(
-					//	FTabManager::NewStack()
-					//	->SetSizeCoefficient(0.65f)
-					//	->AddTab(FPoseSearchDatabaseEditorTabs::ViewportID, ETabState::OpenedTab)
-					//	->SetHideTabWell(true)
-					//)
-					//->Split
-					//(
-					//	FTabManager::NewSplitter()->SetOrientation(Orient_Vertical)
-					//	->Split
-					//	(
+					FTabManager::NewSplitter()->SetOrientation(Orient_Horizontal)->SetSizeCoefficient(0.9f)
+					->Split
+					(
+						FTabManager::NewStack()
+						->SetSizeCoefficient(0.65f)
+						->AddTab(FPoseSearchDatabaseEditorTabs::ViewportID, ETabState::OpenedTab)
+						->SetHideTabWell(true)
+					)
+					->Split
+					(
+						FTabManager::NewSplitter()->SetOrientation(Orient_Vertical)
+						->Split
+						(
 							FTabManager::NewStack()
 							->SetSizeCoefficient(0.3f)
 							->AddTab(FPoseSearchDatabaseEditorTabs::AssetDetailsID, ETabState::OpenedTab)
-					//		->AddTab(FPoseSearchDatabaseEditorTabs::PreviewSettingsID, ETabState::OpenedTab)
-					//	)
-					//)
+							->AddTab(FPoseSearchDatabaseEditorTabs::PreviewSettingsID, ETabState::OpenedTab)
+						)
+					)
 				)
 			)
 		);
@@ -155,6 +161,11 @@ void FPoseSearchDatabaseEditorToolkit::BindCommands()
 {
 	const FPoseSearchDatabaseEditorCommands& Commands = FPoseSearchDatabaseEditorCommands::Get();
 
+	ToolkitCommands->MapAction(
+		Commands.StopPreviewScene,
+		FExecuteAction::CreateSP(this, &FPoseSearchDatabaseEditorToolkit::StopPreviewScene),
+		EUIActionRepeatMode::RepeatDisabled);
+
 	ToolkitCommands->MapAction(
 		Commands.ResetPreviewScene,
 		FExecuteAction::CreateSP(this, &FPoseSearchDatabaseEditorToolkit::ResetPreviewScene),
@@ -181,6 +192,13 @@ void FPoseSearchDatabaseEditorToolkit::ExtendToolbar()
 
 void FPoseSearchDatabaseEditorToolkit::FillToolbar(FToolBarBuilder& ToolbarBuilder)
 {
+	ToolbarBuilder.AddToolBarButton(
+		FPoseSearchDatabaseEditorCommands::Get().StopPreviewScene,
+		NAME_None,
+		TAttribute<FText>(),
+		TAttribute<FText>(),
+		FSlateIcon(FAppStyle::Get().GetStyleSetName(), "Icons.Delete"));
+
 	ToolbarBuilder.AddToolBarButton(
 		FPoseSearchDatabaseEditorCommands::Get().ResetPreviewScene,
 		NAME_None,
@@ -204,12 +222,12 @@ void FPoseSearchDatabaseEditorToolkit::RegisterTabSpawners(const TSharedRef<FTab
 
 	FAssetEditorToolkit::RegisterTabSpawners(InTabManager);
 
-	//InTabManager->RegisterTabSpawner(
-	//	FPoseSearchDatabaseEditorTabs::ViewportID, 
-	//	FOnSpawnTab::CreateSP(this, &FPoseSearchDatabaseEditorToolkit::SpawnTab_Viewport))
-	//	.SetDisplayName(LOCTEXT("ViewportTab", "Viewport"))
-	//	.SetGroup(WorkspaceMenuCategoryRef)
-	//	.SetIcon(FSlateIcon(FEditorStyle::GetStyleSetName(), "GraphEditor.EventGraph_16x"));
+	InTabManager->RegisterTabSpawner(
+		FPoseSearchDatabaseEditorTabs::ViewportID, 
+		FOnSpawnTab::CreateSP(this, &FPoseSearchDatabaseEditorToolkit::SpawnTab_Viewport))
+		.SetDisplayName(LOCTEXT("ViewportTab", "Viewport"))
+		.SetGroup(WorkspaceMenuCategoryRef)
+		.SetIcon(FSlateIcon(FEditorStyle::GetStyleSetName(), "GraphEditor.EventGraph_16x"));
 
 	InTabManager->RegisterTabSpawner(
 		FPoseSearchDatabaseEditorTabs::AssetDetailsID, 
@@ -218,12 +236,12 @@ void FPoseSearchDatabaseEditorToolkit::RegisterTabSpawners(const TSharedRef<FTab
 		.SetGroup(WorkspaceMenuCategoryRef)
 		.SetIcon(FSlateIcon(FEditorStyle::GetStyleSetName(), "LevelEditor.Tabs.Details"));
 
-	//InTabManager->RegisterTabSpawner(
-	//	FPoseSearchDatabaseEditorTabs::PreviewSettingsID, 
-	//	FOnSpawnTab::CreateSP(this, &FPoseSearchDatabaseEditorToolkit::SpawnTab_PreviewSettings))
-	//	.SetDisplayName(LOCTEXT("PreviewSceneSettingsTab", "Preview Scene Settings"))
-	//	.SetGroup(WorkspaceMenuCategoryRef)
-	//	.SetIcon(FSlateIcon(FEditorStyle::GetStyleSetName(), "LevelEditor.Tabs.Details"));
+	InTabManager->RegisterTabSpawner(
+		FPoseSearchDatabaseEditorTabs::PreviewSettingsID, 
+		FOnSpawnTab::CreateSP(this, &FPoseSearchDatabaseEditorToolkit::SpawnTab_PreviewSettings))
+		.SetDisplayName(LOCTEXT("PreviewSceneSettingsTab", "Preview Scene Settings"))
+		.SetGroup(WorkspaceMenuCategoryRef)
+		.SetIcon(FSlateIcon(FEditorStyle::GetStyleSetName(), "LevelEditor.Tabs.Details"));
 }
 
 void FPoseSearchDatabaseEditorToolkit::UnregisterTabSpawners(const TSharedRef<FTabManager>& InTabManager)
diff --git a/Engine/Plugins/Experimental/Animation/PoseSearch/Source/Editor/Private/PoseSearchDatabasePreviewScene.cpp b/Engine/Plugins/Experimental/Animation/PoseSearch/Source/Editor/Private/PoseSearchDatabasePreviewScene.cpp
--- a/Engine/Plugins/Experimental/Animation/PoseSearch/Source/Editor/Private/PoseSearchDatabasePreviewScene.cpp
+++ b/Engine/Plugins/Experimental/Animation/PoseSearch/Source/Editor/Private/PoseSearchDatabasePreviewScene.cpp
@@ -2,6 +2,11 @@
 
 #include "PoseSearchDatabasePreviewScene.h"
 #include "PoseSearchDatabaseEditorToolkit.h"
+#include "PoseSearchDatabaseViewModel.h"
+
+#include "PoseSearch/PoseSearch.h"
+
+#include "Animation/DebugSkelMeshComponent.h"
 #include "GameFramework/WorldSettings.h"
 #include "EngineUtils.h"
 
@@ -49,4 +54,27 @@ void FPoseSearchDatabasePreviewScene::Tick(float InDeltaTime)
 	{
 		GetWorld()->Tick(LEVELTICK_All, InDeltaTime);
 	}
+
+	const FPoseSearchDatabaseViewModel* ViewModel = GetEditorToolkit()->GetViewModel();
+	if (ViewModel->GetPoseSearchDatabase()->IsValidForSearch() &&
+		ViewModel->IsPoseFeaturesDrawMode(EPoseSearchFeaturesDrawMode::All))
+	{
+		for (const FPoseSearchDatabasePreviewActor& PreviewActor : ViewModel->GetPreviewActors())
+		{
+			if (PreviewActor.CurrentPoseIndex != INDEX_NONE)
+			{
+				UE::PoseSearch::FDebugDrawParams DrawParams;
+				DrawParams.RootTransform = PreviewActor.Mesh->GetComponentTransform();
+				DrawParams.Database = ViewModel->GetPoseSearchDatabase();
+				DrawParams.World = GetWorld();
+				DrawParams.DefaultLifeTime = 0.0f;
+				DrawParams.PoseIdx = PreviewActor.CurrentPoseIndex;
+				DrawParams.PointSize = 5.0f;
+				EnumAddFlags(DrawParams.Flags, UE::PoseSearch::EDebugDrawFlags::DrawFast);
+
+				UE::PoseSearch::Draw(DrawParams);
+			}
+		}
+	}
 }
+
diff --git a/Engine/Plugins/Experimental/Animation/PoseSearch/Source/Editor/Private/PoseSearchDatabaseViewModel.cpp b/Engine/Plugins/Experimental/Animation/PoseSearch/Source/Editor/Private/PoseSearchDatabaseViewModel.cpp
--- a/Engine/Plugins/Experimental/Animation/PoseSearch/Source/Editor/Private/PoseSearchDatabaseViewModel.cpp
+++ b/Engine/Plugins/Experimental/Animation/PoseSearch/Source/Editor/Private/PoseSearchDatabaseViewModel.cpp
@@ -5,12 +5,12 @@
 #include "PoseSearch/PoseSearch.h"
 #include "Modules/ModuleManager.h"
 #include "PoseSearchDatabasePreviewScene.h"
-#include "Animation/AnimInstance.h"
+#include "AnimPreviewInstance.h"
+#include "Animation/DebugSkelMeshComponent.h"
+#include "Animation/MirrorDataTable.h"
 #include "EngineUtils.h"
 #include "PropertyEditorModule.h"
 #include "Modules/ModuleManager.h"
-#include "GameFramework/Character.h"
-#include "GameFramework/CharacterMovementComponent.h"
 #include "Camera/CameraComponent.h"
 #include "Components/CapsuleComponent.h"
 
@@ -34,11 +34,40 @@ void FPoseSearchDatabaseViewModel::Initialize(
 {
 	PoseSearchDatabase = InPoseSearchDatabase;
 	PreviewScenePtr = InPreviewScene;
+
+	RemovePreviewActors();
+
+	PoseSearchDatabase->RegisterOnDerivedDataRebuild(
+		UPoseSearchDatabase::FOnDerivedDataRebuild::CreateLambda([this]()
+	{
+		ResetPreviewActors();
+	}));
 }
 
-void FPoseSearchDatabaseViewModel::RestartAnimations()
+void FPoseSearchDatabaseViewModel::ResetPreviewActors()
 {
-	// todo: implement
+	PlayTime = 0.0f;
+	RespawnPreviewActors();
+}
+
+void FPoseSearchDatabaseViewModel::RespawnPreviewActors()
+{
+	RemovePreviewActors();
+	const FPoseSearchIndex* SearchIndex = PoseSearchDatabase->GetSearchIndex();
+	for (const FPoseSearchIndexAsset& IndexAsset : SearchIndex->Assets)
+	{
+		if (AnimationPreviewMode == EAnimationPreviewMode::OriginalAndMirrored || !IndexAsset.bMirrored)
+		{
+			FPoseSearchDatabasePreviewActor PreviewActor = SpawnPreviewActor(IndexAsset);
+			if (PreviewActor.IsValid())
+			{
+				PreviewActors.Add(PreviewActor);
+			}
+		}
+	}
+	UpdatePreviewActors();
+
+	// todo: do blendspaces afterwards
 }
 
 void FPoseSearchDatabaseViewModel::BuildSearchIndex()
@@ -46,39 +75,58 @@ void FPoseSearchDatabaseViewModel::BuildSearchIndex()
 	PoseSearchDatabase->BeginCacheDerivedData();
 }
 
-AActor* FPoseSearchDatabaseViewModel::SpawnPreviewActor(const FPoseSearchDatabaseSequence& DatabaseSequence)
+FPoseSearchDatabasePreviewActor FPoseSearchDatabaseViewModel::SpawnPreviewActor(const FPoseSearchIndexAsset& IndexAsset)
 {
-	// todo: implement properly
+	USkeleton* Skeleton = PoseSearchDatabase->Schema->Skeleton;
+
+	FPoseSearchDatabaseSequence DatabaseSequence = PoseSearchDatabase->Sequences[IndexAsset.SourceAssetIdx];
+
+	FPoseSearchDatabasePreviewActor PreviewActor;
+
+	if (!Skeleton || !DatabaseSequence.Sequence)
+	{
+		return PreviewActor;
+	}
 
-	UClass* PreviewClass = nullptr;
+	PreviewActor.IndexAsset = &IndexAsset;
+	PreviewActor.CurrentPoseIndex = INDEX_NONE;
+
+	// todo: use preview when possible
+	UClass* PreviewClass = AActor::StaticClass();
 	const FTransform SpawnTransform = FTransform::Identity;
 
 	FActorSpawnParameters Params;
 	Params.SpawnCollisionHandlingOverride = ESpawnActorCollisionHandlingMethod::AdjustIfPossibleButAlwaysSpawn;
-	AActor* PreviewActor = GetWorld()->SpawnActor<AActor>(PreviewClass, SpawnTransform, Params);
+	PreviewActor.Actor = GetWorld()->SpawnActor<AActor>(PreviewClass, SpawnTransform, Params);
+	PreviewActor.Actor->SetFlags(RF_Transient);
 
-	if (ACharacter* PreviewCharacter = Cast<ACharacter>(PreviewActor))
-	{
-		PreviewCharacter->bUseControllerRotationYaw = false;
+	PreviewActor.Mesh = NewObject<UDebugSkelMeshComponent>(PreviewActor.Actor.Get());
+	PreviewActor.Mesh->RegisterComponentWithWorld(GetWorld());
 
-		if (UCharacterMovementComponent* CharacterMovementComp = PreviewCharacter->GetCharacterMovement())
-		{
-			CharacterMovementComp->bOrientRotationToMovement = true;
-			CharacterMovementComp->bUseControllerDesiredRotation = false;
-			CharacterMovementComp->RotationRate = FRotator(0.f, 540.0, 0.f);
-			CharacterMovementComp->bRunPhysicsWithNoController = true;
+	PreviewActor.AnimInstance = NewObject<UAnimPreviewInstance>(PreviewActor.Mesh.Get());
+	PreviewActor.Mesh->PreviewInstance = PreviewActor.AnimInstance.Get();
+	PreviewActor.AnimInstance->InitializeAnimation();
 
-			CharacterMovementComp->SetMovementMode(EMovementMode::MOVE_Walking);
-		}
+	PreviewActor.Mesh->SetSkeletalMesh(Skeleton->GetPreviewMesh(true));
+	PreviewActor.Mesh->EnablePreview(true, DatabaseSequence.Sequence);
+	PreviewActor.AnimInstance->SetAnimationAsset(DatabaseSequence.Sequence);
+	if (IndexAsset.bMirrored && PoseSearchDatabase->Schema)
+	{
+		PreviewActor.AnimInstance->SetMirrorDataTable(PoseSearchDatabase->Schema->MirrorDataTable);
+	}
+	
+	PreviewActor.AnimInstance->PlayAnim(false);
 
-		if (UCameraComponent* CameraComp = PreviewCharacter->FindComponentByClass<UCameraComponent>())
-		{
-			CameraComp->DestroyComponent();
-		}
+	if (!PreviewActor.Actor->GetRootComponent())
+	{
+		PreviewActor.Actor->SetRootComponent(PreviewActor.Mesh.Get());
 	}
 
-	UE_LOG(LogPoseSearchEditor, Log, TEXT("Spawned preview Actor: %s at Loc: %s Rot: %s"),
-		*GetNameSafe(PreviewActor), *SpawnTransform.GetLocation().ToString(), *SpawnTransform.Rotator().ToString());
+	UE_LOG(
+		LogPoseSearchEditor, Log, TEXT("Spawned preview Actor: %s at Loc: %s Rot: %s"),
+		*GetNameSafe(PreviewActor.Actor.Get()),
+		*SpawnTransform.GetLocation().ToString(), 
+		*SpawnTransform.Rotator().ToString());
 
 	return PreviewActor;
 }
@@ -98,4 +146,85 @@ UObject* FPoseSearchDatabaseViewModel::GetPlaybackContext() const
 void FPoseSearchDatabaseViewModel::OnPreviewActorClassChanged()
 {
 	// todo: implement
-}
\ No newline at end of file
+}
+
+void FPoseSearchDatabaseViewModel::Tick(float DeltaSeconds)
+{
+	PlayTime += DeltaSeconds;
+	UpdatePreviewActors();
+}
+
+void FPoseSearchDatabaseViewModel::UpdatePreviewActors()
+{
+	TSharedPtr<FPoseSearchDatabasePreviewScene> PreviewScene = PreviewScenePtr.Pin();
+	for (FPoseSearchDatabasePreviewActor& PreviewActor : GetPreviewActors())
+	{
+		if (const UAnimSequence* Sequence = Cast<UAnimSequence>(PreviewActor.AnimInstance->GetAnimationAsset()))
+		{
+			float CurrentTime = 0.0f;
+			FAnimationRuntime::AdvanceTime(false, PlayTime, CurrentTime, Sequence->GetPlayLength());
+
+			FTransform RootMotion = Sequence->ExtractRootMotionFromRange(0.0f, CurrentTime);
+			if (PreviewActor.AnimInstance->GetMirrorDataTable())
+			{
+				RootMotion = MirrorRootMotion(RootMotion, PreviewActor.AnimInstance->GetMirrorDataTable());
+			}
+			PreviewActor.Actor->SetActorTransform(RootMotion);
+			PreviewActor.AnimInstance->SetPosition(CurrentTime);
+
+			PreviewActor.CurrentPoseIndex = PoseSearchDatabase->GetPoseIndexFromTime(CurrentTime, PreviewActor.IndexAsset);
+		}
+	}
+}
+
+void FPoseSearchDatabaseViewModel::RemovePreviewActors()
+{
+	for (auto PreviewActor : PreviewActors)
+	{
+		PreviewActor.Actor->Destroy();
+	}
+	PreviewActors.Reset();
+}
+
+FTransform FPoseSearchDatabaseViewModel::MirrorRootMotion(
+	FTransform RootMotion, 
+	const UMirrorDataTable* MirrorDataTable)
+{
+	const FTransform RootReferenceTransform = 
+		PoseSearchDatabase->Schema->Skeleton->GetReferenceSkeleton().GetRefBonePose()[0];
+	const FQuat RootReferenceRotation = RootReferenceTransform.GetRotation();
+
+	const EAxis::Type MirrorAxis = MirrorDataTable->MirrorAxis;
+	FVector T = RootMotion.GetTranslation();
+	T = FAnimationRuntime::MirrorVector(T, MirrorAxis);
+	FQuat Q = RootMotion.GetRotation();
+	Q = FAnimationRuntime::MirrorQuat(Q, MirrorAxis);
+	Q *= FAnimationRuntime::MirrorQuat(RootReferenceRotation, MirrorAxis).Inverse() * RootReferenceRotation;
+
+	FTransform MirroredRootMotion = FTransform(Q, T, RootMotion.GetScale3D());
+	return MirroredRootMotion;
+}
+
+void FPoseSearchDatabaseViewModel::OnSetPoseFeaturesDrawMode(EPoseSearchFeaturesDrawMode DrawMode)
+{
+	PoseFeaturesDrawMode = DrawMode;
+}
+
+bool FPoseSearchDatabaseViewModel::IsPoseFeaturesDrawMode(EPoseSearchFeaturesDrawMode DrawMode) const
+{
+	return PoseFeaturesDrawMode == DrawMode;
+}
+
+void FPoseSearchDatabaseViewModel::OnSetAnimationPreviewMode(EAnimationPreviewMode PreviewMode)
+{
+	if (PreviewMode != AnimationPreviewMode)
+	{
+		AnimationPreviewMode = PreviewMode;
+		RespawnPreviewActors();
+	}
+}
+
+bool FPoseSearchDatabaseViewModel::IsAnimationPreviewMode(EAnimationPreviewMode PreviewMode) const
+{
+	return AnimationPreviewMode == PreviewMode;
+}
diff --git a/Engine/Plugins/Experimental/Animation/PoseSearch/Source/Editor/Private/PoseSearchEditor.cpp b/Engine/Plugins/Experimental/Animation/PoseSearch/Source/Editor/Private/PoseSearchEditor.cpp
--- a/Engine/Plugins/Experimental/Animation/PoseSearch/Source/Editor/Private/PoseSearchEditor.cpp
+++ b/Engine/Plugins/Experimental/Animation/PoseSearch/Source/Editor/Private/PoseSearchEditor.cpp
@@ -120,10 +120,7 @@ void FEditorModule::StartupModule()
 		// Register Ed Mode used by pose search database
 		FEditorModeRegistry::Get().RegisterMode<FPoseSearchDatabaseEdMode>(
 			FPoseSearchDatabaseEdMode::EdModeId,
-			LOCTEXT("PoseSearchDatabaseEdModeName", "PoseSearchDatabase"),
-			FSlateIcon(),
-			false,
-			9000);
+			LOCTEXT("PoseSearchDatabaseEdModeName", "PoseSearchDatabase"));
 
 		// Register UPoseSearchDatabase Type Actions 
 		IAssetTools& AssetTools = FModuleManager::LoadModuleChecked<FAssetToolsModule>("AssetTools").Get();
diff --git a/Engine/Plugins/Experimental/Animation/PoseSearch/Source/Editor/Private/SPoseSearchDatabaseViewport.cpp b/Engine/Plugins/Experimental/Animation/PoseSearch/Source/Editor/Private/SPoseSearchDatabaseViewport.cpp
--- a/Engine/Plugins/Experimental/Animation/PoseSearch/Source/Editor/Private/SPoseSearchDatabaseViewport.cpp
+++ b/Engine/Plugins/Experimental/Animation/PoseSearch/Source/Editor/Private/SPoseSearchDatabaseViewport.cpp
@@ -2,7 +2,9 @@
 
 #include "SPoseSearchDatabaseViewport.h"
 #include "PoseSearchDatabaseViewportClient.h"
-#include "Toolkits/AssetEditorToolkit.h"
+#include "PoseSearchDatabaseViewModel.h"
+#include "PoseSearchDatabaseEditorToolkit.h"
+#include "SPoseSearchDatabaseViewportToolbar.h"
 #include "PoseSearchDatabaseEditorCommands.h"
 
 void SPoseSearchDatabaseViewport::Construct(
@@ -22,6 +24,59 @@ void SPoseSearchDatabaseViewport::Construct(
 void SPoseSearchDatabaseViewport::BindCommands()
 {
 	SEditorViewport::BindCommands();
+
+	const FPoseSearchDatabaseEditorCommands& Commands = FPoseSearchDatabaseEditorCommands::Get();
+
+	TSharedRef<FPoseSearchDatabaseViewModel> ViewModelRef = 
+		AssetEditorToolkitPtr.Pin()->GetViewModelSharedPtr().ToSharedRef();
+
+	CommandList->MapAction(
+		Commands.ShowPoseFeaturesNone,
+		FExecuteAction::CreateSP(
+			ViewModelRef, 
+			&FPoseSearchDatabaseViewModel::OnSetPoseFeaturesDrawMode, 
+			EPoseSearchFeaturesDrawMode::None),
+		FCanExecuteAction(),
+		FIsActionChecked::CreateSP(
+			ViewModelRef, 
+			&FPoseSearchDatabaseViewModel::IsPoseFeaturesDrawMode, 
+			EPoseSearchFeaturesDrawMode::None));
+
+	CommandList->MapAction(
+		Commands.ShowPoseFeaturesAll,
+		FExecuteAction::CreateSP(
+			ViewModelRef,
+			&FPoseSearchDatabaseViewModel::OnSetPoseFeaturesDrawMode,
+			EPoseSearchFeaturesDrawMode::All),
+		FCanExecuteAction(),
+		FIsActionChecked::CreateSP(
+			ViewModelRef,
+			&FPoseSearchDatabaseViewModel::IsPoseFeaturesDrawMode,
+			EPoseSearchFeaturesDrawMode::All));
+
+	CommandList->MapAction(
+		Commands.ShowAnimationOriginalOnly,
+		FExecuteAction::CreateSP(
+			ViewModelRef,
+			&FPoseSearchDatabaseViewModel::OnSetAnimationPreviewMode,
+			EAnimationPreviewMode::OriginalOnly),
+		FCanExecuteAction(),
+		FIsActionChecked::CreateSP(
+			ViewModelRef,
+			&FPoseSearchDatabaseViewModel::IsAnimationPreviewMode,
+			EAnimationPreviewMode::OriginalOnly));
+
+	CommandList->MapAction(
+		Commands.ShowAnimationOriginalAndMirrored,
+		FExecuteAction::CreateSP(
+			ViewModelRef,
+			&FPoseSearchDatabaseViewModel::OnSetAnimationPreviewMode,
+			EAnimationPreviewMode::OriginalAndMirrored),
+		FCanExecuteAction(),
+		FIsActionChecked::CreateSP(
+			ViewModelRef,
+			&FPoseSearchDatabaseViewModel::IsAnimationPreviewMode,
+			EAnimationPreviewMode::OriginalAndMirrored));
 }
 
 TSharedRef<FEditorViewportClient> SPoseSearchDatabaseViewport::MakeEditorViewportClient()
@@ -38,6 +93,12 @@ TSharedRef<FEditorViewportClient> SPoseSearchDatabaseViewport::MakeEditorViewpor
 	return ViewportClient.ToSharedRef();
 }
 
+TSharedPtr<SWidget> SPoseSearchDatabaseViewport::MakeViewportToolbar()
+{
+	return SAssignNew(ViewportToolbar, SPoseSearchDatabaseViewportToolBar, SharedThis(this));
+}
+
+
 TSharedRef<SEditorViewport> SPoseSearchDatabaseViewport::GetViewportWidget()
 {
 	return SharedThis(this);
diff --git a/Engine/Plugins/Experimental/Animation/PoseSearch/Source/Editor/Private/SPoseSearchDatabaseViewportToolbar.cpp b/Engine/Plugins/Experimental/Animation/PoseSearch/Source/Editor/Private/SPoseSearchDatabaseViewportToolbar.cpp
new file mode 100644
--- /dev/null
+++ b/Engine/Plugins/Experimental/Animation/PoseSearch/Source/Editor/Private/SPoseSearchDatabaseViewportToolbar.cpp
@@ -0,0 +1,67 @@
+// Copyright Epic Games, Inc. All Rights Reserved.
+
+#include "SPoseSearchDatabaseViewportToolbar.h"
+#include "SPoseSearchDatabaseViewport.h"
+#include "PreviewProfileController.h"
+#include "PoseSearchDatabaseEditorCommands.h"
+
+#define LOCTEXT_NAMESPACE "ContextualAnimViewportToolBar"
+
+void SPoseSearchDatabaseViewportToolBar::Construct(
+	const FArguments& InArgs, 
+	TSharedPtr<SPoseSearchDatabaseViewport> InViewport)
+{
+	SCommonEditorViewportToolbarBase::Construct(
+		SCommonEditorViewportToolbarBase::FArguments()
+		.AddRealtimeButton(false)
+		.PreviewProfileController(MakeShared<FPreviewProfileController>()), 
+		InViewport);
+}
+
+TSharedRef<SWidget> SPoseSearchDatabaseViewportToolBar::GenerateShowMenu() const
+{
+	GetInfoProvider().OnFloatingButtonClicked();
+
+	TSharedRef<SEditorViewport> ViewportRef = GetInfoProvider().GetViewportWidget();
+
+	const bool bInShouldCloseWindowAfterMenuSelection = true;
+	FMenuBuilder ShowMenuBuilder(bInShouldCloseWindowAfterMenuSelection, ViewportRef->GetCommandList());
+	{
+		ShowMenuBuilder.AddSubMenu(
+			LOCTEXT("ShowMenu_PoseFeaturesDrawSubMenu", "Pose Features"),
+			LOCTEXT("ShowMenu_PoseFeaturesDrawSubMenuToolTip", "Pose Feature Drawing Options"),
+			FNewMenuDelegate::CreateLambda([](FMenuBuilder& SubMenuBuilder)
+			{
+				const FPoseSearchDatabaseEditorCommands& Commands = FPoseSearchDatabaseEditorCommands::Get();
+
+				SubMenuBuilder.BeginSection("PoseFeatures", LOCTEXT("ShowMenu_PoseFeaturesLabel", "Pose Features"));
+				{
+					SubMenuBuilder.AddMenuEntry(Commands.ShowPoseFeaturesNone);
+					SubMenuBuilder.AddMenuEntry(Commands.ShowPoseFeaturesAll);
+				}
+				SubMenuBuilder.EndSection();
+			})
+		);
+
+		ShowMenuBuilder.AddSubMenu(
+			LOCTEXT("ShowMenu_AnimationsSubMenu", "Animations"),
+			LOCTEXT("ShowMenu_AnimationsSubMenuToolTip", "Animation Preview Options"),
+			FNewMenuDelegate::CreateLambda([](FMenuBuilder& SubMenuBuilder)
+			{
+				const FPoseSearchDatabaseEditorCommands& Commands = FPoseSearchDatabaseEditorCommands::Get();
+
+				SubMenuBuilder.BeginSection("Animations", LOCTEXT("ShowMenu_AnimationsLabel", "Animations"));
+				{
+					SubMenuBuilder.AddMenuEntry(Commands.ShowAnimationOriginalOnly);
+					SubMenuBuilder.AddMenuEntry(Commands.ShowAnimationOriginalAndMirrored);
+				}
+				SubMenuBuilder.EndSection();
+			})
+		);
+
+	}
+
+	return ShowMenuBuilder.MakeWidget();
+}
+
+#undef LOCTEXT_NAMESPACE
diff --git a/Engine/Plugins/Experimental/Animation/PoseSearch/Source/Editor/Public/PoseSearchDatabaseEditorCommands.h b/Engine/Plugins/Experimental/Animation/PoseSearch/Source/Editor/Public/PoseSearchDatabaseEditorCommands.h
--- a/Engine/Plugins/Experimental/Animation/PoseSearch/Source/Editor/Public/PoseSearchDatabaseEditorCommands.h
+++ b/Engine/Plugins/Experimental/Animation/PoseSearch/Source/Editor/Public/PoseSearchDatabaseEditorCommands.h
@@ -21,8 +21,18 @@ public:
 	virtual void RegisterCommands() override;
 
 public:
+	TSharedPtr<FUICommandInfo> StopPreviewScene;
 
 	TSharedPtr<FUICommandInfo> ResetPreviewScene;
 
 	TSharedPtr<FUICommandInfo> BuildSearchIndex;
+
+	TSharedPtr<FUICommandInfo> ShowPoseFeaturesNone;
+
+	TSharedPtr<FUICommandInfo> ShowPoseFeaturesAll;
+
+	TSharedPtr<FUICommandInfo> ShowAnimationOriginalOnly;
+
+	TSharedPtr<FUICommandInfo> ShowAnimationOriginalAndMirrored;
+
 };
diff --git a/Engine/Plugins/Experimental/Animation/PoseSearch/Source/Editor/Public/PoseSearchDatabaseEditorToolkit.h b/Engine/Plugins/Experimental/Animation/PoseSearch/Source/Editor/Public/PoseSearchDatabaseEditorToolkit.h
--- a/Engine/Plugins/Experimental/Animation/PoseSearch/Source/Editor/Public/PoseSearchDatabaseEditorToolkit.h
+++ b/Engine/Plugins/Experimental/Animation/PoseSearch/Source/Editor/Public/PoseSearchDatabaseEditorToolkit.h
@@ -33,9 +33,11 @@ public:
 	virtual FLinearColor GetWorldCentricTabColorScale() const override;
 	virtual FString GetWorldCentricTabPrefix() const override;
 
-	UPoseSearchDatabase* GetPoseSearchDatabase() const;
+	const UPoseSearchDatabase* GetPoseSearchDatabase() const;
 	FPoseSearchDatabaseViewModel* GetViewModel() const { return ViewModel.Get(); }
+	TSharedPtr<FPoseSearchDatabaseViewModel> GetViewModelSharedPtr() const { return ViewModel; }
 
+	void StopPreviewScene();
 	void ResetPreviewScene();
 	void BuildSearchIndex();
 
diff --git a/Engine/Plugins/Experimental/Animation/PoseSearch/Source/Editor/Public/PoseSearchDatabaseViewModel.h b/Engine/Plugins/Experimental/Animation/PoseSearch/Source/Editor/Public/PoseSearchDatabaseViewModel.h
--- a/Engine/Plugins/Experimental/Animation/PoseSearch/Source/Editor/Public/PoseSearchDatabaseViewModel.h
+++ b/Engine/Plugins/Experimental/Animation/PoseSearch/Source/Editor/Public/PoseSearchDatabaseViewModel.h
@@ -6,12 +6,41 @@
 #include "UObject/GCObject.h"
 #include "MovieSceneFwd.h"
 
+#include "PoseSearchDatabasePreviewScene.h"
+
 class UWorld;
-class FPoseSearchDatabasePreviewScene;
 class UPoseSearchDatabase;
-struct FPoseSearchDatabaseSequence;
-class IDetailsView;
-class IStructureDetailsView;
+struct FPoseSearchIndexAsset;
+class UAnimPreviewInstance;
+
+enum class EPoseSearchFeaturesDrawMode : uint8
+{
+	None,
+	All
+};
+
+enum class EAnimationPreviewMode : uint8
+{
+	OriginalOnly,
+	OriginalAndMirrored
+};
+
+
+struct FPoseSearchDatabasePreviewActor
+{
+public:
+	TWeakObjectPtr<AActor> Actor = nullptr;
+	TWeakObjectPtr<UDebugSkelMeshComponent> Mesh = nullptr;
+	TWeakObjectPtr<UAnimPreviewInstance> AnimInstance = nullptr;
+	const FPoseSearchIndexAsset* IndexAsset = nullptr;
+	int32 CurrentPoseIndex = INDEX_NONE;
+
+	bool IsValid()
+	{
+		const bool bIsValid = Actor.IsValid() && Mesh.IsValid() && AnimInstance.IsValid();
+		return  bIsValid;
+	}
+};
 
 class FPoseSearchDatabaseViewModel : public TSharedFromThis<FPoseSearchDatabaseViewModel>, public FGCObject
 {
@@ -28,13 +57,27 @@ public:
 		UPoseSearchDatabase* InPoseSearchDatabase, 
 		const TSharedRef<FPoseSearchDatabasePreviewScene>& InPreviewScene);
 
-	void RestartAnimations();
+	void RemovePreviewActors();
+	void ResetPreviewActors();
+	void RespawnPreviewActors();
 	void BuildSearchIndex();
 
 	UPoseSearchDatabase* GetPoseSearchDatabase() const { return PoseSearchDatabase; }
 	void OnPreviewActorClassChanged();
 
+	void Tick(float DeltaSeconds);
+
+	TArray<FPoseSearchDatabasePreviewActor>& GetPreviewActors() { return PreviewActors; }
+	const TArray<FPoseSearchDatabasePreviewActor>& GetPreviewActors() const { return PreviewActors; }
+
+	void OnSetPoseFeaturesDrawMode(EPoseSearchFeaturesDrawMode DrawMode);
+	bool IsPoseFeaturesDrawMode(EPoseSearchFeaturesDrawMode DrawMode) const;
+
+	void OnSetAnimationPreviewMode(EAnimationPreviewMode PreviewMode);
+	bool IsAnimationPreviewMode(EAnimationPreviewMode PreviewMode) const;
+
 private:
+	float PlayTime = 0.0f;
 
 	/** Scene asset being viewed and edited by this view model. */
 	TObjectPtr<UPoseSearchDatabase> PoseSearchDatabase;
@@ -42,9 +85,22 @@ private:
 	/** Weak pointer to the PreviewScene */
 	TWeakPtr<FPoseSearchDatabasePreviewScene> PreviewScenePtr;
 
-	AActor* SpawnPreviewActor(const FPoseSearchDatabaseSequence& DatabaseSequence);
+	/** Actors to be displayed in the preview viewport */
+	TArray<FPoseSearchDatabasePreviewActor> PreviewActors;
+
+	/** What features to show in the viewport */
+	EPoseSearchFeaturesDrawMode PoseFeaturesDrawMode = EPoseSearchFeaturesDrawMode::None;
+
+	/** What animations to show in the viewport */
+	EAnimationPreviewMode AnimationPreviewMode = EAnimationPreviewMode::OriginalOnly;
 
 	UWorld* GetWorld() const;
 
 	UObject* GetPlaybackContext() const;
+
+	FPoseSearchDatabasePreviewActor SpawnPreviewActor(const FPoseSearchIndexAsset& IndexAsset);
+
+	void UpdatePreviewActors();
+
+	FTransform MirrorRootMotion(FTransform RootMotion, const class UMirrorDataTable* MirrorDataTable);
 };
\ No newline at end of file
diff --git a/Engine/Plugins/Experimental/Animation/PoseSearch/Source/Editor/Public/SPoseSearchDatabaseViewport.h b/Engine/Plugins/Experimental/Animation/PoseSearch/Source/Editor/Public/SPoseSearchDatabaseViewport.h
--- a/Engine/Plugins/Experimental/Animation/PoseSearch/Source/Editor/Public/SPoseSearchDatabaseViewport.h
+++ b/Engine/Plugins/Experimental/Animation/PoseSearch/Source/Editor/Public/SPoseSearchDatabaseViewport.h
@@ -46,6 +46,7 @@ protected:
 	// ~SEditorViewport interface
 	virtual void BindCommands() override;
 	virtual TSharedRef<FEditorViewportClient> MakeEditorViewportClient() override;
+	virtual TSharedPtr<SWidget> MakeViewportToolbar() override;
 	// ~End of SEditorViewport interface
 
 	/** The viewport toolbar */
diff --git a/Engine/Plugins/Experimental/Animation/PoseSearch/Source/Editor/Public/SPoseSearchDatabaseViewportToolbar.h b/Engine/Plugins/Experimental/Animation/PoseSearch/Source/Editor/Public/SPoseSearchDatabaseViewportToolbar.h
new file mode 100644
--- /dev/null
+++ b/Engine/Plugins/Experimental/Animation/PoseSearch/Source/Editor/Public/SPoseSearchDatabaseViewportToolbar.h
@@ -0,0 +1,23 @@
+// Copyright Epic Games, Inc. All Rights Reserved.
+
+#pragma once
+
+#include "CoreMinimal.h"
+#include "SCommonEditorViewportToolbarBase.h"
+
+class SPoseSearchDatabaseViewport;
+
+class SPoseSearchDatabaseViewportToolBar : public SCommonEditorViewportToolbarBase
+{
+public:
+	SLATE_BEGIN_ARGS(SPoseSearchDatabaseViewportToolBar)
+	{}
+
+	SLATE_END_ARGS()
+
+	void Construct(const FArguments& InArgs, TSharedPtr<SPoseSearchDatabaseViewport> InViewport);
+
+	// ~SCommonEditorViewportToolbarBase interface
+	virtual TSharedRef<SWidget> GenerateShowMenu() const override;
+	// ~End of SCommonEditorViewportToolbarBase interface
+};
diff --git a/Engine/Plugins/Experimental/Animation/PoseSearch/Source/Runtime/Private/PoseSearch.cpp b/Engine/Plugins/Experimental/Animation/PoseSearch/Source/Runtime/Private/PoseSearch.cpp
--- a/Engine/Plugins/Experimental/Animation/PoseSearch/Source/Runtime/Private/PoseSearch.cpp
+++ b/Engine/Plugins/Experimental/Animation/PoseSearch/Source/Runtime/Private/PoseSearch.cpp
@@ -1301,7 +1301,16 @@ bool UPoseSearchDatabase::IsValidForIndexing() const
 bool UPoseSearchDatabase::IsValidForSearch() const
 {
 	const FPoseSearchIndex* SearchIndex = GetSearchIndex();
-	return IsValidForIndexing() && SearchIndex && SearchIndex->IsValid();
+	bool bIsValid = IsValidForIndexing() && SearchIndex && SearchIndex->IsValid();
+
+#if WITH_EDITOR
+	const bool bIsCurrentDerivedData = 
+		PrivateDerivedData &&
+		PrivateDerivedData->PendingDerivedDataKey == PrivateDerivedData->DerivedDataKey.Hash;
+	bIsValid = bIsValid && bIsCurrentDerivedData;
+#endif // WITH_EDITOR
+
+	return bIsValid;
 }
 
 void UPoseSearchDatabase::CollectSimpleSequences()
@@ -1663,6 +1672,20 @@ void UPoseSearchDatabase::PostLoad()
 }
 
 #if WITH_EDITOR
+void UPoseSearchDatabase::RegisterOnDerivedDataRebuild(const FOnDerivedDataRebuild& Delegate)
+{
+	OnDerivedDataRebuild.Add(Delegate);
+}
+void UPoseSearchDatabase::UnregisterOnDerivedDataRebuild(void* Unregister)
+{
+	OnDerivedDataRebuild.RemoveAll(Unregister);
+}
+
+void UPoseSearchDatabase::NotifyDerivedDataBuildStarted()
+{
+	OnDerivedDataRebuild.Broadcast();
+}
+
 void UPoseSearchDatabase::BeginCacheDerivedData()
 {
 	bool bPerformCache = true;
@@ -1670,7 +1693,7 @@ void UPoseSearchDatabase::BeginCacheDerivedData()
 	using namespace UE::DerivedData;
 	if (PrivateDerivedData)
 	{
-		const FIoHash ExistingDerivedDataHash = PrivateDerivedData->FetchOrBuildDerivedDataKey;
+		const FIoHash ExistingDerivedDataHash = PrivateDerivedData->PendingDerivedDataKey;
 		if (!ExistingDerivedDataHash.IsZero())
 		{
 			if (ExistingDerivedDataHash == UE::PoseSearch::FPoseSearchDatabaseAsyncCacheTask::CreateKey(*this))
@@ -1754,6 +1777,8 @@ void UPoseSearchDatabase::PostEditChangeProperty(struct FPropertyChangedEvent& P
 			CollectSimpleBlendSpaces();
 		}
 	}
+
+	BeginCacheDerivedData();
 }
 
 void UPoseSearchDatabase::BeginCacheForCookedPlatformData(const ITargetPlatform* TargetPlatform)
@@ -4147,7 +4172,7 @@ static void DrawTrajectoryFeatures(const FDebugDrawParams& DrawParams, const FFe
 			FColor Color = GradientColor.ToFColor(true);
 
 			TrajectoryPos = DrawParams.RootTransform.TransformPosition(TrajectoryPos);
-			if (EnumHasAnyFlags(DrawParams.Flags, EDebugDrawFlags::DrawSearchIndex))
+			if (EnumHasAnyFlags(DrawParams.Flags, EDebugDrawFlags::DrawFast | EDebugDrawFlags::DrawSearchIndex))
 			{
 				DrawDebugPoint(DrawParams.World, TrajectoryPos, DrawParams.PointSize, Color, bPersistent, DrawParams.DefaultLifeTime, DepthPriority);
 			}
@@ -4173,12 +4198,18 @@ static void DrawTrajectoryFeatures(const FDebugDrawParams& DrawParams, const FFe
 			TrajectoryVel *= DrawDebugVelocityScale;
 			TrajectoryVel = DrawParams.RootTransform.TransformVector(TrajectoryVel);
 			FVector TrajectoryVelDirection = TrajectoryVel.GetSafeNormal();
+
+
 			if (EnumHasAnyFlags(DrawParams.Flags, EDebugDrawFlags::DrawSearchIndex))
 			{
 				DrawDebugPoint(DrawParams.World, TrajectoryVel, DrawParams.PointSize, Color, bPersistent, DrawParams.DefaultLifeTime, DepthPriority);
 			}
 			else
 			{
+				const float AdjustedThickness =
+					EnumHasAnyFlags(DrawParams.Flags, EDebugDrawFlags::DrawFast) ?
+					0.0f : DrawDebugLineThickness;
+
 				DrawDebugDirectionalArrow(
 					DrawParams.World,
 					TrajectoryPos + TrajectoryVelDirection * DrawDebugSphereSize,
@@ -4188,7 +4219,7 @@ static void DrawTrajectoryFeatures(const FDebugDrawParams& DrawParams, const FFe
 					bPersistent,
 					LifeTime,
 					DepthPriority,
-					DrawDebugLineThickness
+					AdjustedThickness
 				);
 			}
 		}
@@ -4203,12 +4234,17 @@ static void DrawTrajectoryFeatures(const FDebugDrawParams& DrawParams, const FFe
 			FColor Color = GradientColor.ToFColor(true);
 
 			TrajectoryForward = DrawParams.RootTransform.TransformVector(TrajectoryForward);
+
 			if (EnumHasAnyFlags(DrawParams.Flags, EDebugDrawFlags::DrawSearchIndex))
 			{
 				DrawDebugPoint(DrawParams.World, TrajectoryForward, DrawParams.PointSize, Color, bPersistent, DrawParams.DefaultLifeTime, DepthPriority);
 			}
 			else
 			{
+				const float AdjustedThickness =
+					EnumHasAnyFlags(DrawParams.Flags, EDebugDrawFlags::DrawFast) ?
+					0.0f : DrawDebugLineThickness;
+
 				DrawDebugDirectionalArrow(
 					DrawParams.World,
 					TrajectoryPos + TrajectoryForward * DrawDebugSphereSize,
@@ -4218,7 +4254,7 @@ static void DrawTrajectoryFeatures(const FDebugDrawParams& DrawParams, const FFe
 					bPersistent,
 					LifeTime,
 					DepthPriority,
-					DrawDebugLineThickness
+					AdjustedThickness
 				);
 			}
 		}
@@ -4289,7 +4325,7 @@ static void DrawPoseFeatures(const FDebugDrawParams& DrawParams, const FFeatureV
 				FColor Color = LinearColor.ToFColor(true);
 
 				BonePos = DrawParams.RootTransform.TransformPosition(BonePos);
-				if (EnumHasAnyFlags(DrawParams.Flags, EDebugDrawFlags::DrawSearchIndex))
+				if (EnumHasAnyFlags(DrawParams.Flags, EDebugDrawFlags::DrawFast | EDebugDrawFlags::DrawSearchIndex))
 				{
 					DrawDebugPoint(DrawParams.World, BonePos, DrawParams.PointSize, Color, bPersistent, DrawParams.DefaultLifeTime, DepthPriority);
 				}
@@ -4318,13 +4354,27 @@ static void DrawPoseFeatures(const FDebugDrawParams& DrawParams, const FFeatureV
 				BoneVel *= DrawDebugVelocityScale;
 				BoneVel = DrawParams.RootTransform.TransformVector(BoneVel);
 				FVector BoneVelDirection = BoneVel.GetSafeNormal();
+
 				if (EnumHasAnyFlags(DrawParams.Flags, EDebugDrawFlags::DrawSearchIndex))
 				{
 					DrawDebugPoint(DrawParams.World, BoneVel, DrawParams.PointSize, Color, bPersistent, DrawParams.DefaultLifeTime, DepthPriority);
 				}
 				else
 				{
-					DrawDebugDirectionalArrow(DrawParams.World, BonePos + BoneVelDirection * DrawDebugSphereSize, BonePos + BoneVel, DrawDebugArrowSize, Color, bPersistent, LifeTime, DepthPriority, DrawDebugLineThickness);
+					const float AdjustedThickness =
+						EnumHasAnyFlags(DrawParams.Flags, EDebugDrawFlags::DrawFast) ?
+						0.0f : DrawDebugLineThickness;
+
+					DrawDebugDirectionalArrow(
+						DrawParams.World, 
+						BonePos + BoneVelDirection * DrawDebugSphereSize, 
+						BonePos + BoneVel, 
+						DrawDebugArrowSize, 
+						Color, 
+						bPersistent, 
+						LifeTime, 
+						DepthPriority, 
+						AdjustedThickness);
 				}
 			}
 		}
diff --git a/Engine/Plugins/Experimental/Animation/PoseSearch/Source/Runtime/Private/PoseSearchDerivedData.cpp b/Engine/Plugins/Experimental/Animation/PoseSearch/Source/Runtime/Private/PoseSearchDerivedData.cpp
--- a/Engine/Plugins/Experimental/Animation/PoseSearch/Source/Runtime/Private/PoseSearchDerivedData.cpp
+++ b/Engine/Plugins/Experimental/Animation/PoseSearch/Source/Runtime/Private/PoseSearchDerivedData.cpp
@@ -72,12 +72,13 @@ namespace UE::PoseSearch
 		using namespace UE::DerivedData;
 
 		FIoHash DerivedDataKey = CreateKey(Database);
-		DerivedData.FetchOrBuildDerivedDataKey = DerivedDataKey;
+		DerivedData.PendingDerivedDataKey = DerivedDataKey;
+
+		InDatabase.NotifyDerivedDataBuildStarted();
 
 		if (bForceRebuild)
 		{
-			DerivedData.DerivedDataKey = { Bucket, DerivedDataKey };
-			BuildAndWrite();
+			BuildAndWrite({ Bucket, DerivedDataKey });
 		}
 		else
 		{
@@ -108,7 +109,7 @@ namespace UE::PoseSearch
 
 		TArray<FCacheGetRequest> CacheRequests;
 		ECachePolicy CachePolicy = ECachePolicy::Default;
-		const FCacheKey CacheKey{ Bucket, DerivedData.FetchOrBuildDerivedDataKey };
+		const FCacheKey CacheKey{ Bucket, DerivedData.PendingDerivedDataKey };
 		CacheRequests.Add({ {Database.GetPathName()}, CacheKey, CachePolicy });
 		GetCache().Get(CacheRequests, Owner, [this](FCacheGetResponse&& Response)
 		{
@@ -120,20 +121,20 @@ namespace UE::PoseSearch
 	{
 		using namespace UE::DerivedData;
 
-		DerivedData.DerivedDataKey = Response.Record.GetKey();
 		if (Response.Status == EStatus::Ok)
 		{
 			BuildIndexFromCacheRecord(MoveTemp(Response.Record));
+			DerivedData.DerivedDataKey = Response.Record.GetKey();
 		}
 		else if (Response.Status == EStatus::Error)
 		{
-			BuildAndWrite();
+			BuildAndWrite(Response.Record.GetKey());
 		}
 	}
 
-	void FPoseSearchDatabaseAsyncCacheTask::BuildAndWrite()
+	void FPoseSearchDatabaseAsyncCacheTask::BuildAndWrite(const UE::DerivedData::FCacheKey& NewKey)
 	{
-		GetRequestOwner().LaunchTask(TEXT("PoseSearchDatabaseBuild"), [this]
+		GetRequestOwner().LaunchTask(TEXT("PoseSearchDatabaseBuild"), [this, NewKey]
 		{
 			if (!GetRequestOwner().IsCanceled())
 			{
@@ -142,6 +143,7 @@ namespace UE::PoseSearch
 				const bool bIndexReady = BuildIndex(&Database, DerivedData.SearchIndex);
 
 				WriteIndexToCache();
+				DerivedData.DerivedDataKey = NewKey;
 			}
 		});
 	}
diff --git a/Engine/Plugins/Experimental/Animation/PoseSearch/Source/Runtime/Private/PoseSearchDerivedData.h b/Engine/Plugins/Experimental/Animation/PoseSearch/Source/Runtime/Private/PoseSearchDerivedData.h
--- a/Engine/Plugins/Experimental/Animation/PoseSearch/Source/Runtime/Private/PoseSearchDerivedData.h
+++ b/Engine/Plugins/Experimental/Animation/PoseSearch/Source/Runtime/Private/PoseSearchDerivedData.h
@@ -37,7 +37,7 @@ struct FPoseSearchDatabaseDerivedData
 	*	Note that since this is read on the game thread constantly in CachePlatformData, it
 	*	must be written to on the game thread to avoid false recaches.
 	*/
-	FIoHash FetchOrBuildDerivedDataKey;
+	FIoHash PendingDerivedDataKey;
 
 	/** Async cache task if one is outstanding. */
 	UE::PoseSearch::FPoseSearchDatabaseAsyncCacheTask* AsyncTask;
@@ -76,7 +76,7 @@ namespace UE::PoseSearch
 	private:
 		void BeginCache();
 		void OnGetComplete(UE::DerivedData::FCacheGetResponse&& Response);
-		void BuildAndWrite();
+		void BuildAndWrite(const UE::DerivedData::FCacheKey& NewKey);
 		void WriteIndexToCache();
 		void BuildIndexFromCacheRecord(UE::DerivedData::FCacheRecord&& CacheRecord);
 
diff --git a/Engine/Plugins/Experimental/Animation/PoseSearch/Source/Runtime/Public/PoseSearch/PoseSearch.h b/Engine/Plugins/Experimental/Animation/PoseSearch/Source/Runtime/Public/PoseSearch/PoseSearch.h
--- a/Engine/Plugins/Experimental/Animation/PoseSearch/Source/Runtime/Public/PoseSearch/PoseSearch.h
+++ b/Engine/Plugins/Experimental/Animation/PoseSearch/Source/Runtime/Public/PoseSearch/PoseSearch.h
@@ -824,8 +824,21 @@ public:
 private:
 	FPoseSearchDatabaseDerivedData* PrivateDerivedData;
 
+#if WITH_EDITOR
+	DECLARE_MULTICAST_DELEGATE(FOnDerivedDataRebuildMulticaster);
+	FOnDerivedDataRebuildMulticaster OnDerivedDataRebuild;
+#endif // WITH_EDITOR
+
 public:
 #if WITH_EDITOR
+	typedef FOnDerivedDataRebuildMulticaster::FDelegate FOnDerivedDataRebuild;
+
+	/** Registers a delegate to be called when derived data starts rebuilding */
+	void RegisterOnDerivedDataRebuild(const FOnDerivedDataRebuild& Delegate);
+	void UnregisterOnDerivedDataRebuild(void* Unregister);
+
+	void NotifyDerivedDataBuildStarted();
+
 	void BeginCacheDerivedData();
 #endif // WITH_EDITOR
 
@@ -1016,6 +1029,9 @@ enum class EDebugDrawFlags : uint32
 
 	// Draw Bone Names
 	DrawBoneNames = 1 << 6,
+
+	// Draws simpler shapes to improve performance
+	DrawFast = 1 << 7,
 };
 ENUM_CLASS_FLAGS(EDebugDrawFlags);
 
