exposing ContinuingPose candidate into MM debugger

diff --git a/Engine/Plugins/Experimental/Animation/PoseSearch/Source/Editor/Private/PoseSearchDebugger.cpp b/Engine/Plugins/Experimental/Animation/PoseSearch/Source/Editor/Private/PoseSearchDebugger.cpp
--- a/Engine/Plugins/Experimental/Animation/PoseSearch/Source/Editor/Private/PoseSearchDebugger.cpp
+++ b/Engine/Plugins/Experimental/Animation/PoseSearch/Source/Editor/Private/PoseSearchDebugger.cpp
@@ -719,6 +719,7 @@ void SDebuggerDatabaseView::RefreshColumns()
 	using namespace DebuggerDatabaseColumns;
 	
 	ActiveView.HeaderRow->ClearColumns();
+	ContinuingPoseView.HeaderRow->ClearColumns();
 	FilteredDatabaseView.HeaderRow->ClearColumns();
 	
 	// Sort columns by index
@@ -749,6 +750,8 @@ void SDebuggerDatabaseView::RefreshColumns()
 			
 			// Every time the active column is changed, update the database column
 			ActiveView.HeaderRow->AddColumn(ColumnArgs.OnWidthChanged(this, &SDebuggerDatabaseView::OnColumnWidthChanged, ColumnName));
+			
+			ContinuingPoseView.HeaderRow->AddColumn(ColumnArgs.OnWidthChanged(this, &SDebuggerDatabaseView::OnColumnWidthChanged, ColumnName));
 		}
 	}
 }
@@ -944,6 +947,9 @@ void SDebuggerDatabaseView::CreateRows(const UPoseSearchDatabase& Database)
 
 	ActiveView.Rows.Reset();
 	ActiveView.Rows.Add(MakeShared<FDebuggerDatabaseRowData>());
+
+	ContinuingPoseView.Rows.Reset();
+	ContinuingPoseView.Rows.Add(MakeShared<FDebuggerDatabaseRowData>());
 }
 
 void SDebuggerDatabaseView::UpdateRows(const FTraceMotionMatchingStateMessage& State, const UPoseSearchDatabase& Database)
@@ -955,6 +961,7 @@ void SDebuggerDatabaseView::UpdateRows(const FTraceMotionMatchingStateMessage& S
 		CreateRows(Database);
 	}
 	check(ActiveView.Rows.Num() == 1);
+	check(ContinuingPoseView.Rows.Num() == 1);
 
 	FPoseSearchWeightsContext StateWeights;
 	StateWeights.Update(State.Weights, &Database);
@@ -969,16 +976,38 @@ void SDebuggerDatabaseView::UpdateRows(const FTraceMotionMatchingStateMessage& S
 			EPoseSearchBooleanRequest::TrueValue : EPoseSearchBooleanRequest::FalseValue;
 	}
 	
-	for(const TSharedRef<FDebuggerDatabaseRowData>& Row : UnfilteredDatabaseRows)
+	if (!UnfilteredDatabaseRows.IsEmpty())
 	{
-		const int32 PoseIdx = Row->PoseIdx;
+		bool IsActiveViewRowInitialized = false;
+		bool IsContinuingPoseViewRowInitialized = false;
+		for (const TSharedRef<FDebuggerDatabaseRowData>& Row : UnfilteredDatabaseRows)
+		{
+			const int32 PoseIdx = Row->PoseIdx;
+
+			ComparePoses(PoseIdx, SearchContext, Row->PoseCostDetails);
+
+			// If we are on the active pose for the frame
+			if (PoseIdx == State.DbPoseIdx)
+			{
+				*ActiveView.Rows[0] = Row.Get();
+				IsActiveViewRowInitialized = true;
+			}
+
+			if (PoseIdx == State.ContinuingPoseIdx)
+			{
+				*ContinuingPoseView.Rows[0] = Row.Get();
+				IsContinuingPoseViewRowInitialized = true;
+			}
+		}
 
-		ComparePoses(PoseIdx, SearchContext, Row->PoseCostDetails);
+		if (!IsActiveViewRowInitialized)
+		{
+			*ActiveView.Rows[0] = UnfilteredDatabaseRows[0].Get();
+		}
 
-		// If we are on the active pose for the frame
-		if (PoseIdx == State.DbPoseIdx)
+		if (!IsContinuingPoseViewRowInitialized)
 		{
-			*ActiveView.Rows[0] = Row.Get();
+			*ContinuingPoseView.Rows[0] = *ActiveView.Rows[0];
 		}
 	}
 
@@ -1009,6 +1038,13 @@ TSharedRef<ITableRow> SDebuggerDatabaseView::HandleGenerateActiveRow(TSharedRef<
 		.ColumnMap(this, &SDebuggerDatabaseView::GetColumnMap);
 }
 
+TSharedRef<ITableRow> SDebuggerDatabaseView::HandleGenerateContinuingPoseRow(TSharedRef<FDebuggerDatabaseRowData> Item, const TSharedRef<STableViewBase>& OwnerTable) const
+{
+	return
+		SNew(SDebuggerDatabaseRow, OwnerTable, Item, ContinuingPoseView.RowStyle, &ContinuingPoseView.RowBrush, FMargin(0.0f, 2.0f, 6.0f, 4.0f))
+		.ColumnMap(this, &SDebuggerDatabaseView::GetColumnMap);
+}
+
 void SDebuggerDatabaseView::Construct(const FArguments& InArgs)
 {
 	using namespace DebuggerDatabaseColumns;
@@ -1053,6 +1089,27 @@ void SDebuggerDatabaseView::Construct(const FArguments& InArgs)
 	ActiveView.RowStyle = FEditorStyle::GetWidgetStyle<FTableRowStyle>("TableView.Row");
 	ActiveView.RowBrush = *FEditorStyle::GetBrush("DetailsView.CategoryTop");
 
+	// ContinuingPose Row
+	ContinuingPoseView.HeaderRow = SNew(SHeaderRow).Visibility(EVisibility::Collapsed);
+
+	// Used for spacing
+	ContinuingPoseView.ScrollBar =
+		SNew(SScrollBar)
+		.Orientation(Orient_Vertical)
+		.HideWhenNotInUse(false)
+		.AlwaysShowScrollbar(true)
+		.AlwaysShowScrollbarTrack(true);
+
+	ContinuingPoseView.ListView = SNew(SListView<TSharedRef<FDebuggerDatabaseRowData>>)
+		.ListItemsSource(&ContinuingPoseView.Rows)
+		.HeaderRow(ContinuingPoseView.HeaderRow.ToSharedRef())
+		.OnGenerateRow(this, &SDebuggerDatabaseView::HandleGenerateContinuingPoseRow)
+		.ExternalScrollbar(ContinuingPoseView.ScrollBar)
+		.SelectionMode(ESelectionMode::SingleToggle)
+		.ConsumeMouseWheel(EConsumeMouseWheel::Never);
+
+	ContinuingPoseView.RowStyle = FEditorStyle::GetWidgetStyle<FTableRowStyle>("TableView.Row");
+	ContinuingPoseView.RowBrush = *FEditorStyle::GetBrush("DetailsView.CategoryTop");
 
 	// Filtered Database
 	FilteredDatabaseView.ScrollBar =
@@ -1139,6 +1196,65 @@ void SDebuggerDatabaseView::Construct(const FArguments& InArgs)
 			]	
 		]
 
+		+ SVerticalBox::Slot()
+		// Side and top margins, ignore bottom handled by the color border below
+		.Padding(0.0f, 5.0f, 0.0f, 0.0f)
+		.AutoHeight()
+		[
+			// ContinuingPose Row text tab
+			SNew(SVerticalBox)
+			+ SVerticalBox::Slot()
+			.AutoHeight()
+			.Padding(0.0f)
+			[
+				SNew(SHorizontalBox)
+				+ SHorizontalBox::Slot()
+				.HAlign(HAlign_Center)
+				.VAlign(VAlign_Fill)
+				.Padding(0.0f)
+				.AutoWidth()
+				[
+					SNew(SBorder)
+					.BorderImage(FEditorStyle::GetBrush("DetailsView.CategoryTop"))
+					.Padding(FMargin(30.0f, 3.0f, 30.0f, 0.0f))
+					.HAlign(HAlign_Center)
+					.VAlign(VAlign_Fill)
+					[
+						SNew(STextBlock)
+						.Text(FText::FromString("Continuing Pose"))	
+					]
+				]
+			]
+
+			// ContinuingPose row list view with scroll bar
+			+ SVerticalBox::Slot()
+			
+			.AutoHeight()
+			[
+				SNew(SHorizontalBox)
+				
+				+ SHorizontalBox::Slot()
+				.HAlign(HAlign_Fill)
+				.VAlign(VAlign_Fill)
+				.Padding(0.0f)
+				[
+					SNew(SBorder)
+					
+					.BorderImage(FEditorStyle::GetBrush("NoBorder"))
+					.Padding(0.0f)
+					[
+						ContinuingPoseView.ListView.ToSharedRef()
+					]
+				]
+
+				+ SHorizontalBox::Slot()
+				.AutoWidth()
+				[
+					ContinuingPoseView.ScrollBar.ToSharedRef()
+				]
+			]	
+		]
+		
 		+ SVerticalBox::Slot()
 		.Padding(0.0f, 0.0f, 0.0f, 5.0f)
 		[
@@ -1235,9 +1351,10 @@ void SDebuggerDatabaseView::Construct(const FArguments& InArgs)
 	SortColumn = FCost::Name;
 	SortMode = EColumnSortMode::Ascending;
 
-	// Active view scroll bar only for indenting the columns to align w/ database
+	// Active and Continuing Pose view scroll bars only for indenting the columns to align w/ database
 	ActiveView.ScrollBar->SetVisibility(EVisibility::Hidden);
-	
+	ContinuingPoseView.ScrollBar->SetVisibility(EVisibility::Hidden);
+
 	RefreshColumns();
 }
 
@@ -1656,6 +1773,22 @@ void SDebuggerView::DrawFeatures(
 		Draw(DrawParams);
 	}
 
+	Selected = DatabaseView->GetContinuingPoseRow()->GetSelectedItems();
+
+	// ContinuingPose row should only have 0 or 1
+	check(Selected.Num() < 2);
+
+	if (!Selected.IsEmpty())
+	{
+		// Green for the ContinuingPose view
+		DrawParams.Color = &FLinearColor::Gray;
+
+		// Use the motion-matching state's pose idx, as the ContinuingPose row may be update-throttled at this point
+		DrawParams.PoseIdx = (*Selected.begin())->PoseIdx;
+		DrawParams.LabelPrefix = TEXT("C");
+		Draw(DrawParams);
+	}
+
 	if (Reflection->bDrawActiveSkeleton)
 	{
 		SkeletonDrawParams.Flags |= ESkeletonDrawFlags::ActivePose;
diff --git a/Engine/Plugins/Experimental/Animation/PoseSearch/Source/Editor/Private/PoseSearchDebugger.h b/Engine/Plugins/Experimental/Animation/PoseSearch/Source/Editor/Private/PoseSearchDebugger.h
--- a/Engine/Plugins/Experimental/Animation/PoseSearch/Source/Editor/Private/PoseSearchDebugger.h
+++ b/Engine/Plugins/Experimental/Animation/PoseSearch/Source/Editor/Private/PoseSearchDebugger.h
@@ -268,6 +268,7 @@ class SDebuggerDatabaseView : public SCompoundWidget
 	void Update(const FTraceMotionMatchingStateMessage& State, const UPoseSearchDatabase& Database);
 
 	const TSharedPtr<SListView<TSharedRef<FDebuggerDatabaseRowData>>>& GetActiveRow() const { return ActiveView.ListView; }
+	const TSharedPtr<SListView<TSharedRef<FDebuggerDatabaseRowData>>>& GetContinuingPoseRow() const { return ContinuingPoseView.ListView; }
 	const TSharedPtr<SListView<TSharedRef<FDebuggerDatabaseRowData>>>& GetDatabaseRows() const { return FilteredDatabaseView.ListView; }
 	const TSharedRef<FDebuggerDatabaseRowData>& GetPoseIdxDatabaseRow(int32 PoseIdx) const;
 
@@ -326,6 +327,9 @@ private:
 	/** Generates the active row widget for the given data */
 	TSharedRef<ITableRow> HandleGenerateActiveRow(TSharedRef<FDebuggerDatabaseRowData> Item, const TSharedRef<STableViewBase>& OwnerTable) const;
 
+	/** Generates the continuing pose row widget for the given data */
+	TSharedRef<ITableRow> HandleGenerateContinuingPoseRow(TSharedRef<FDebuggerDatabaseRowData> Item, const TSharedRef<STableViewBase>& OwnerTable) const;
+
 	TWeakPtr<SDebuggerView> ParentDebuggerViewPtr;
 
 	FOnPoseSelectionChanged OnPoseSelectionChanged;
@@ -364,6 +368,9 @@ private:
 	/** Active row at the top of the view */
 	FTable ActiveView;
 
+	/** Continuing pose row below Active row */
+	FTable ContinuingPoseView;
+
 	/** All database poses */
 	TArray<TSharedRef<FDebuggerDatabaseRowData>> UnfilteredDatabaseRows;
 
diff --git a/Engine/Plugins/Experimental/Animation/PoseSearch/Source/Editor/Private/Trace/PoseSearchTraceAnalyzer.cpp b/Engine/Plugins/Experimental/Animation/PoseSearch/Source/Editor/Private/Trace/PoseSearchTraceAnalyzer.cpp
--- a/Engine/Plugins/Experimental/Animation/PoseSearch/Source/Editor/Private/Trace/PoseSearchTraceAnalyzer.cpp
+++ b/Engine/Plugins/Experimental/Animation/PoseSearch/Source/Editor/Private/Trace/PoseSearchTraceAnalyzer.cpp
@@ -44,6 +44,7 @@ bool FTraceAnalyzer::OnEvent(uint16 RouteId, EStyle Style, const FOnEventContext
 				Message.DatabaseId = EventData.GetValue<uint64>("DatabaseId");
 				Message.Flags = static_cast<FTraceMotionMatchingState::EFlags>(EventData.GetValue<uint32>("Flags"));
 				Message.DbPoseIdx = EventData.GetValue<int32>("DbPoseIdx");
+				Message.ContinuingPoseIdx = EventData.GetValue<int32>("ContinuingPoseIdx");
 				Message.ElapsedPoseJumpTime = EventData.GetValue<float>("ElapsedPoseJumpTime");
 				Message.QueryVector = EventData.GetArrayView<float>("QueryVector");
                 Message.QueryVectorNormalized = EventData.GetArrayView<float>("QueryVectorNormalized");
diff --git a/Engine/Plugins/Experimental/Animation/PoseSearch/Source/Editor/Private/Trace/PoseSearchTraceProvider.h b/Engine/Plugins/Experimental/Animation/PoseSearch/Source/Editor/Private/Trace/PoseSearchTraceProvider.h
--- a/Engine/Plugins/Experimental/Animation/PoseSearch/Source/Editor/Private/Trace/PoseSearchTraceProvider.h
+++ b/Engine/Plugins/Experimental/Animation/PoseSearch/Source/Editor/Private/Trace/PoseSearchTraceProvider.h
@@ -30,6 +30,7 @@ struct FTraceMotionMatchingStateMessage : FTraceMessage
 	FPoseSearchDynamicWeightParams Weights;
 	int32 DbPoseIdx = INDEX_NONE;
 	uint64 DatabaseId = INDEX_NONE;
+	int32 ContinuingPoseIdx = INDEX_NONE;
 	float AssetPlayerTime = 0.0f;
 	float DeltaTime = 0.0f;
 	float SimLinearVelocity = 0.0f;
diff --git a/Engine/Plugins/Experimental/Animation/PoseSearch/Source/Runtime/Private/PoseSearchLibrary.cpp b/Engine/Plugins/Experimental/Animation/PoseSearch/Source/Runtime/Private/PoseSearchLibrary.cpp
--- a/Engine/Plugins/Experimental/Animation/PoseSearch/Source/Runtime/Private/PoseSearchLibrary.cpp
+++ b/Engine/Plugins/Experimental/Animation/PoseSearch/Source/Runtime/Private/PoseSearchLibrary.cpp
@@ -496,7 +496,8 @@ void UpdateMotionMatchingState(
 		TraceState.Weights = Settings.Weights;
 		TraceState.DbPoseIdx = InOutMotionMatchingState.DbPoseIdx;
 		TraceState.DatabaseId = FObjectTrace::GetObjectId(Database);
-
+		TraceState.ContinuingPoseIdx = PoseStepper.Result.PoseIdx;
+		
 		TraceState.AssetPlayerTime = InOutMotionMatchingState.AssetPlayerTime;
 		TraceState.DeltaTime = DeltaTime;
 		TraceState.SimLinearVelocity = SimLinearVelocity;
diff --git a/Engine/Plugins/Experimental/Animation/PoseSearch/Source/Runtime/Private/Trace/PoseSearchTraceLogger.cpp b/Engine/Plugins/Experimental/Animation/PoseSearch/Source/Runtime/Private/Trace/PoseSearchTraceLogger.cpp
--- a/Engine/Plugins/Experimental/Animation/PoseSearch/Source/Runtime/Private/Trace/PoseSearchTraceLogger.cpp
+++ b/Engine/Plugins/Experimental/Animation/PoseSearch/Source/Runtime/Private/Trace/PoseSearchTraceLogger.cpp
@@ -18,6 +18,7 @@ UE_TRACE_EVENT_BEGIN(PoseSearch, MotionMatchingState)
 	UE_TRACE_EVENT_FIELD(uint32, Flags)
 	UE_TRACE_EVENT_FIELD(uint64, DatabaseId)
 	UE_TRACE_EVENT_FIELD(int32, DbPoseIdx)
+	UE_TRACE_EVENT_FIELD(int32, ContinuingPoseIdx)
 	UE_TRACE_EVENT_FIELD(float[], QueryVector)
 	UE_TRACE_EVENT_FIELD(float[], QueryVectorNormalized)
 	UE_TRACE_EVENT_FIELD(float[], ChannelWeightScales)
@@ -94,6 +95,7 @@ void FTraceMotionMatchingState::Output(const FAnimationBaseContext& InContext, c
 		<< MotionMatchingState.Flags(static_cast<uint32>(State.Flags))
 		<< MotionMatchingState.DatabaseId(State.DatabaseId)
 		<< MotionMatchingState.DbPoseIdx(State.DbPoseIdx)
+		<< MotionMatchingState.ContinuingPoseIdx(State.ContinuingPoseIdx)
 		<< MotionMatchingState.QueryVector(State.QueryVector.GetData(), State.QueryVector.Num())
 		<< MotionMatchingState.QueryVectorNormalized(State.QueryVectorNormalized.GetData(), State.QueryVectorNormalized.Num())
 		<< MotionMatchingState.ChannelWeightScales(ChannelWeightScales, UE_ARRAY_COUNT(ChannelWeightScales))
diff --git a/Engine/Plugins/Experimental/Animation/PoseSearch/Source/Runtime/Private/Trace/PoseSearchTraceLogger.h b/Engine/Plugins/Experimental/Animation/PoseSearch/Source/Runtime/Private/Trace/PoseSearchTraceLogger.h
--- a/Engine/Plugins/Experimental/Animation/PoseSearch/Source/Runtime/Private/Trace/PoseSearchTraceLogger.h
+++ b/Engine/Plugins/Experimental/Animation/PoseSearch/Source/Runtime/Private/Trace/PoseSearchTraceLogger.h
@@ -59,6 +59,9 @@ struct POSESEARCH_API FTraceMotionMatchingState
 	/** Object Id of the database asset */
 	uint64 DatabaseId = 0;
 
+	/** Index of the continuing pose in our database */
+	int32 ContinuingPoseIdx = 0;
+
 	/** Node Id of the motion matching node associated with this message */
 	int32 NodeId = 0;
 
