[UE-149722] PoseSearchDatabaseViewModel.h(cpp) Non-unity build fix.

diff --git a/Engine/Plugins/Experimental/Animation/PoseSearch/Source/Editor/Private/PoseSearchDatabaseViewModel.cpp b/Engine/Plugins/Experimental/Animation/PoseSearch/Source/Editor/Private/PoseSearchDatabaseViewModel.cpp
--- a/Engine/Plugins/Experimental/Animation/PoseSearch/Source/Editor/Private/PoseSearchDatabaseViewModel.cpp
+++ b/Engine/Plugins/Experimental/Animation/PoseSearch/Source/Editor/Private/PoseSearchDatabaseViewModel.cpp
@@ -8,6 +8,7 @@
 #include "AnimPreviewInstance.h"
 #include "Animation/DebugSkelMeshComponent.h"
 #include "Animation/MirrorDataTable.h"
+#include "Animation/AnimSequence.h"
 #include "EngineUtils.h"
 #include "PropertyEditorModule.h"
 #include "Modules/ModuleManager.h"
diff --git a/Engine/Plugins/Experimental/Animation/PoseSearch/Source/Editor/Public/PoseSearchDatabaseViewModel.h b/Engine/Plugins/Experimental/Animation/PoseSearch/Source/Editor/Public/PoseSearchDatabaseViewModel.h
--- a/Engine/Plugins/Experimental/Animation/PoseSearch/Source/Editor/Public/PoseSearchDatabaseViewModel.h
+++ b/Engine/Plugins/Experimental/Animation/PoseSearch/Source/Editor/Public/PoseSearchDatabaseViewModel.h
@@ -12,6 +12,7 @@ class UWorld;
 class UPoseSearchDatabase;
 struct FPoseSearchIndexAsset;
 class UAnimPreviewInstance;
+class UDebugSkelMeshComponent;
 
 enum class EPoseSearchFeaturesDrawMode : uint8
 {
