PoseSearch: Miscellaneous cleanup
- Added dependency on AnimationWarping for IRootMotionProvider implementation
- Removed overzealous root motion provider ensure that can trigger on first node evaluate in preview scene
- Removed extra state in SDebuggerView::DrawFeatures
- Removed some  header dependencies
- Converted CompareFeatureVectors to equvalent expressions as Eigen array operations (mainly to self document future optimizations)
- Removed overzealous ensure when no pose is found during search due to tag filtering

diff --git a/Engine/Plugins/Experimental/Animation/PoseSearch/PoseSearch.uplugin b/Engine/Plugins/Experimental/Animation/PoseSearch/PoseSearch.uplugin
--- a/Engine/Plugins/Experimental/Animation/PoseSearch/PoseSearch.uplugin
+++ b/Engine/Plugins/Experimental/Animation/PoseSearch/PoseSearch.uplugin
@@ -34,6 +34,10 @@
 			"TargetAllowList": [
 				"Editor"
 			]
+		},
+		{
+			"Name": "AnimationWarping",
+			"Enabled": true
 		}
 	]
 }
\ No newline at end of file
diff --git a/Engine/Plugins/Experimental/Animation/PoseSearch/Source/Editor/Private/PoseSearchDebugger.cpp b/Engine/Plugins/Experimental/Animation/PoseSearch/Source/Editor/Private/PoseSearchDebugger.cpp
--- a/Engine/Plugins/Experimental/Animation/PoseSearch/Source/Editor/Private/PoseSearchDebugger.cpp
+++ b/Engine/Plugins/Experimental/Animation/PoseSearch/Source/Editor/Private/PoseSearchDebugger.cpp
@@ -1743,13 +1743,6 @@ void SDebuggerView::DrawFeatures(
 	DrawParams.LabelPrefix = TEXT("S");
 	SetDrawFlags(DrawParams, Reflection->SelectedPoseDrawOptions);
 
-	FSkeletonDrawParams SkeletonDrawParams;
-	const bool bDrawSelectedPose = Reflection->bDrawSelectedSkeleton;
-	if (bDrawSelectedPose)
-	{
-		SkeletonDrawParams.Flags |= ESkeletonDrawFlags::SelectedPose;
-	}
-
 	// Draw any selected database vectors
 	for (const TSharedRef<FDebuggerDatabaseRowData>& Row : Selected)
 	{
@@ -1789,6 +1782,11 @@ void SDebuggerView::DrawFeatures(
 		Draw(DrawParams);
 	}
 
+	FSkeletonDrawParams SkeletonDrawParams;
+	if (Reflection->bDrawSelectedSkeleton)
+	{
+		SkeletonDrawParams.Flags |= ESkeletonDrawFlags::SelectedPose;
+	}
 	if (Reflection->bDrawActiveSkeleton)
 	{
 		SkeletonDrawParams.Flags |= ESkeletonDrawFlags::ActivePose;
diff --git a/Engine/Plugins/Experimental/Animation/PoseSearch/Source/Runtime/Private/AnimNode_MotionMatching.cpp b/Engine/Plugins/Experimental/Animation/PoseSearch/Source/Runtime/Private/AnimNode_MotionMatching.cpp
--- a/Engine/Plugins/Experimental/Animation/PoseSearch/Source/Runtime/Private/AnimNode_MotionMatching.cpp
+++ b/Engine/Plugins/Experimental/Animation/PoseSearch/Source/Runtime/Private/AnimNode_MotionMatching.cpp
@@ -48,11 +48,13 @@ void FAnimNode_MotionMatching::Evaluate_AnyThread(FPoseContext& Output)
 #endif
 
 #if UE_POSE_SEARCH_TRACE_ENABLED
+	MotionMatchingState.RootMotionTransformDelta = FTransform::Identity;
+
 	const UE::Anim::IAnimRootMotionProvider* RootMotionProvider = UE::Anim::IAnimRootMotionProvider::Get();
 
 	if (ensureMsgf(RootMotionProvider, TEXT("Could not get Root Motion Provider.")))
 	{
-		if (ensureMsgf(RootMotionProvider->HasRootMotion(Output.CustomAttributes), TEXT("Motion Matching Output no Root Motion Attribute.")))
+		if (RootMotionProvider->HasRootMotion(Output.CustomAttributes))
 		{
 			RootMotionProvider->ExtractRootMotion(Output.CustomAttributes, MotionMatchingState.RootMotionTransformDelta);
 		}
diff --git a/Engine/Plugins/Experimental/Animation/PoseSearch/Source/Runtime/Private/PoseSearch.cpp b/Engine/Plugins/Experimental/Animation/PoseSearch/Source/Runtime/Private/PoseSearch.cpp
--- a/Engine/Plugins/Experimental/Animation/PoseSearch/Source/Runtime/Private/PoseSearch.cpp
+++ b/Engine/Plugins/Experimental/Animation/PoseSearch/Source/Runtime/Private/PoseSearch.cpp
@@ -16,9 +16,7 @@
 #include "Animation/AnimCurveTypes.h"
 #include "Animation/AnimInstance.h"
 #include "Animation/AnimInstanceProxy.h"
-#include "Animation/AnimNodeBase.h"
 #include "Animation/AnimMetaData.h"
-#include "Animation/AnimNodeBase.h"
 #include "Animation/AnimSequence.h"
 #include "Animation/AnimSequenceBase.h"
 #include "Animation/AnimationPoseData.h"
@@ -73,28 +71,19 @@ static bool IsSamplingRangeValid(FFloatInterval Range)
 
 static inline float CompareFeatureVectors(int32 NumValues, const float* A, const float* B, const float* Weights)
 {
-	double Dissimilarity = 0.f;
+	Eigen::Map<const Eigen::ArrayXf> VA(A, NumValues);
+	Eigen::Map<const Eigen::ArrayXf> VB(B, NumValues);
+	Eigen::Map<const Eigen::ArrayXf> VW(Weights, NumValues);
 
-	for (int32 ValueIdx = 0; ValueIdx != NumValues; ++ValueIdx)
-	{
-		const float Diff = A[ValueIdx] - B[ValueIdx];
-		Dissimilarity += Weights[ValueIdx] * (Diff * Diff);
-	}
-
-	return (float)Dissimilarity;
+	return ((VA - VB).square() * VW).sum();
 }
 
 static inline float CompareFeatureVectors(int32 NumValues, const float* A, const float* B)
 {
-	double Dissimilarity = 0.f;
+	Eigen::Map<const Eigen::ArrayXf> VA(A, NumValues);
+	Eigen::Map<const Eigen::ArrayXf> VB(B, NumValues);
 
-	for (int32 ValueIdx = 0; ValueIdx != NumValues; ++ValueIdx)
-	{
-		const float Diff = A[ValueIdx] - B[ValueIdx];
-		Dissimilarity += Diff * Diff;
-	}
-
-	return (float)Dissimilarity;
+	return (VA - VB).square().sum();
 }
 
 static FLinearColor GetColorForFeature(FPoseSearchFeatureDesc Feature, const FPoseSearchFeatureVectorLayout* Layout)
@@ -5197,8 +5186,6 @@ FSearchResult Search(FSearchContext& SearchContext)
 		}
 	}
 
-	ensure(BestPoseIdx != INDEX_NONE);
-
 	Result.PoseCost = BestPoseCost;
 	Result.PoseIdx = BestPoseIdx;
 	Result.SearchIndexAsset = SearchIndex->FindAssetForPose(BestPoseIdx);
diff --git a/Engine/Plugins/Experimental/Animation/PoseSearch/Source/Runtime/Private/PoseSearchLibrary.cpp b/Engine/Plugins/Experimental/Animation/PoseSearch/Source/Runtime/Private/PoseSearchLibrary.cpp
--- a/Engine/Plugins/Experimental/Animation/PoseSearch/Source/Runtime/Private/PoseSearchLibrary.cpp
+++ b/Engine/Plugins/Experimental/Animation/PoseSearch/Source/Runtime/Private/PoseSearchLibrary.cpp
@@ -364,8 +364,11 @@ void UpdateMotionMatchingState(
 	// Jump to the candidate pose if the stepper couldn't continue to the next frame
 	if (!PoseStepper.CanContinue())
 	{
-		FSearchResult Result = Search(SearchContext);
-		InOutMotionMatchingState.JumpToPose(Context, Settings, Result);
+		FSearchResult SearchResult = Search(SearchContext);
+		if (SearchResult.IsValid())
+		{
+			InOutMotionMatchingState.JumpToPose(Context, Settings, SearchResult);
+		}
 	}
 
 	// Do a search when enough time has elapsed since the last pose jump
@@ -380,47 +383,49 @@ void UpdateMotionMatchingState(
 		}
 
 		// Search the database for the nearest match to the updated query vector
-		FSearchResult Result = Search(SearchContext);
-
-		// Consider the search result better if it is more similar to the query than the current pose we're playing back from the database
-		check(Result.PoseCost.Dissimilarity >= 0.0f);
-		bool bBetterPose = true;
-		if (CurrentPoseCost.IsValid())
+		FSearchResult SearchResult = Search(SearchContext);
+		if (SearchResult.IsValid())
 		{
-			if ((CurrentPoseCost.TotalCost <= Result.PoseCost.TotalCost) || (CurrentPoseCost.Dissimilarity <= Result.PoseCost.Dissimilarity))
+			// Consider the search result better if it is more similar to the query than the current pose we're playing back from the database
+			check(SearchResult.PoseCost.Dissimilarity >= 0.0f);
+			bool bBetterPose = true;
+			if (CurrentPoseCost.IsValid())
 			{
-				bBetterPose = false;
+				if ((CurrentPoseCost.TotalCost <= SearchResult.PoseCost.TotalCost) || (CurrentPoseCost.Dissimilarity <= SearchResult.PoseCost.Dissimilarity))
+				{
+					bBetterPose = false;
+				}
+				else
+				{
+					checkSlow(CurrentPoseCost.Dissimilarity > 0.0f && CurrentPoseCost.Dissimilarity > SearchResult.PoseCost.Dissimilarity);
+					const float RelativeSimilarityGain = -1.0f * (SearchResult.PoseCost.Dissimilarity - CurrentPoseCost.Dissimilarity) / CurrentPoseCost.Dissimilarity;
+					bBetterPose = RelativeSimilarityGain >= Settings.MinPercentImprovement / 100.0f;
+				}
 			}
-			else
+
+			// Ignore the candidate poses from the same anim when they are too near to the current pose
+			bool bNearbyPose = false;
+			const FPoseSearchIndexAsset* StateSearchIndexAsset = InOutMotionMatchingState.GetCurrentSearchIndexAsset();
+			if (StateSearchIndexAsset == SearchResult.SearchIndexAsset)
 			{
-				checkSlow(CurrentPoseCost.Dissimilarity > 0.0f && CurrentPoseCost.Dissimilarity > Result.PoseCost.Dissimilarity);
-				const float RelativeSimilarityGain = -1.0f * (Result.PoseCost.Dissimilarity - CurrentPoseCost.Dissimilarity) / CurrentPoseCost.Dissimilarity;
-				bBetterPose = RelativeSimilarityGain >= Settings.MinPercentImprovement / 100.0f;
-			}
-		}
+				// We need to check in terms of PoseIdx rather than AssetTime because
+				// for blendspaces, AssetTime is not in seconds, but in the normalized range 
+				// [0, 1] so comparing to `PoseJumpThresholdTime` will not make sense		
+				bNearbyPose = FMath::Abs(InOutMotionMatchingState.DbPoseIdx - SearchResult.PoseIdx) * Database->Schema->SamplingInterval < Settings.PoseJumpThresholdTime;
 
-		// Ignore the candidate poses from the same anim when they are too near to the current pose
-		bool bNearbyPose = false;
-		const FPoseSearchIndexAsset* StateSearchIndexAsset = InOutMotionMatchingState.GetCurrentSearchIndexAsset();
-		if (StateSearchIndexAsset == Result.SearchIndexAsset)
-		{
-			// We need to check in terms of PoseIdx rather than AssetTime because
-			// for blendspaces, AssetTime is not in seconds, but in the normalized range 
-			// [0, 1] so comparing to `PoseJumpThresholdTime` will not make sense		
-			bNearbyPose = FMath::Abs(InOutMotionMatchingState.DbPoseIdx - Result.PoseIdx) * Database->Schema->SamplingInterval < Settings.PoseJumpThresholdTime;
+				// Handle looping anims when checking for the pose being too close
+				if (!bNearbyPose && Database->IsSourceAssetLooping(StateSearchIndexAsset))
+				{
+					bNearbyPose = FMath::Abs(StateSearchIndexAsset->NumPoses - InOutMotionMatchingState.DbPoseIdx - SearchResult.PoseIdx) * Database->Schema->SamplingInterval < Settings.PoseJumpThresholdTime;
+				}
+			}
 
-			// Handle looping anims when checking for the pose being too close
-			if (!bNearbyPose && Database->IsSourceAssetLooping(StateSearchIndexAsset))
+			// Start playback from the candidate pose if we determined it was a better option
+			if (bBetterPose && !bNearbyPose)
 			{
-				bNearbyPose = FMath::Abs(StateSearchIndexAsset->NumPoses - InOutMotionMatchingState.DbPoseIdx - Result.PoseIdx) * Database->Schema->SamplingInterval < Settings.PoseJumpThresholdTime;
+				InOutMotionMatchingState.JumpToPose(Context, Settings, SearchResult);
 			}
 		}
-
-		// Start playback from the candidate pose if we determined it was a better option
-		if (bBetterPose && !bNearbyPose)
-		{
-			InOutMotionMatchingState.JumpToPose(Context, Settings, Result);
-		}
 	}
 
 	// Jump to the pose stepper's next anim if we didn't jump to a pose as a result of the search above
