PoseSearchDebugger coloring of FilteredDatabaseView for Cost, PoseCost, and TrajectoryCost

diff --git a/Engine/Plugins/Experimental/Animation/PoseSearch/Source/Editor/Private/PoseSearchDebugger.cpp b/Engine/Plugins/Experimental/Animation/PoseSearch/Source/Editor/Private/PoseSearchDebugger.cpp
--- a/Engine/Plugins/Experimental/Animation/PoseSearch/Source/Editor/Private/PoseSearchDebugger.cpp
+++ b/Engine/Plugins/Experimental/Animation/PoseSearch/Source/Editor/Private/PoseSearchDebugger.cpp
@@ -29,6 +29,11 @@
 
 #define LOCTEXT_NAMESPACE "PoseSearchDebugger"
 
+static FLinearColor LinearColorBlend(FLinearColor LinearColorA, FLinearColor LinearColorB, float BlendParam)
+{
+	return LinearColorA + (LinearColorB - LinearColorA) * BlendParam;
+}
+
 void FPoseSearchDebuggerPoseVectorChannel::Reset()
 {
 	Positions.Reset();
@@ -268,6 +273,18 @@ public:
 	float GetTrajectoryCost() const { return PoseCostDetails.ChannelCosts.Num() >= 3 ? PoseCostDetails.ChannelCosts[1] + PoseCostDetails.ChannelCosts[2] : 0.0f; }
 	float GetAddendsCost() const { return PoseCostDetails.NotifyCostAddend + PoseCostDetails.MirrorMismatchCostAddend; }
 
+	void CalculateColors(float CostMin, float CostDelta, float PoseCostMin, float PoseCostDelta, float TrajectoryCostMin, float TrajectoryCostDelta)
+	{
+		const float CostColorBlend = (PoseCostDetails.PoseCost.TotalCost - CostMin) / CostDelta;
+		CostColor = LinearColorBlend(FLinearColor::Green, FLinearColor::Red, CostColorBlend);
+
+		const float PoseCostColorBlend = (GetPoseCost() - PoseCostMin) / PoseCostDelta;
+		PoseCostColor = LinearColorBlend(FLinearColor::Green, FLinearColor::Red, PoseCostColorBlend);
+
+		const float TrajectoryCostColorBlend = (GetTrajectoryCost() - TrajectoryCostMin) / TrajectoryCostDelta;
+		TrajectoryCostColor = LinearColorBlend(FLinearColor::Green, FLinearColor::Red, TrajectoryCostColorBlend);
+	}
+
 	ESearchIndexAssetType AssetType = ESearchIndexAssetType::Invalid;
 	int32 PoseIdx = 0;
 	FString AssetName = "";
@@ -279,6 +296,9 @@ public:
 	bool bLooping = false;
 	FVector BlendParameters = FVector::Zero();
 	FPoseCostDetails PoseCostDetails;
+	FLinearColor CostColor = FLinearColor::White;
+	FLinearColor PoseCostColor = FLinearColor::White;
+	FLinearColor TrajectoryCostColor = FLinearColor::White;
 };
 
 namespace DebuggerDatabaseColumns
@@ -328,11 +348,16 @@ namespace DebuggerDatabaseColumns
 			return SNew(STextBlock)
             	.Font(RowFont)
 				.Text_Lambda([this, RowData]() -> FText { return GetRowText(RowData); })
-            	.Justification(ETextJustify::Center);
+            	.Justification(ETextJustify::Center)
+				.ColorAndOpacity_Lambda([this, RowData] { return GetColorAndOpacity(RowData); });
 		}
 
 	protected:
 		virtual FText GetRowText(const FRowDataRef& Row) const = 0;
+		virtual FSlateColor GetColorAndOpacity(const FRowDataRef& Row) const
+		{
+			return FSlateColor(FLinearColor::White);
+		}
 	};
 
 	struct FPoseIdx : ITextColumn
@@ -468,6 +493,11 @@ namespace DebuggerDatabaseColumns
         {
         	return FText::AsNumber(Row->PoseCostDetails.PoseCost.TotalCost);
         }
+
+		virtual FSlateColor GetColorAndOpacity(const FRowDataRef& Row) const override
+		{
+			return Row->CostColor;
+		}
 	};
 	const FName FCost::Name = "Cost";
 
@@ -487,6 +517,11 @@ namespace DebuggerDatabaseColumns
 		{
 			return FText::AsNumber(Row->GetPoseCost());
 		}
+
+		virtual FSlateColor GetColorAndOpacity(const FRowDataRef& Row) const override
+		{
+			return Row->PoseCostColor;
+		}
 	};
 	const FName FPoseCostColumn::Name = "Pose Cost";
 
@@ -506,6 +541,11 @@ namespace DebuggerDatabaseColumns
 		{
 			return FText::AsNumber(Row->GetTrajectoryCost());
         }
+
+		virtual FSlateColor GetColorAndOpacity(const FRowDataRef& Row) const override
+		{
+			return Row->TrajectoryCostColor;
+		}
 	};
 	const FName FTrajectoryCost::Name = "Trajectory Cost";
 
@@ -1021,6 +1061,68 @@ void SDebuggerDatabaseView::UpdateRows(const FTraceMotionMatchingStateMessage& S
 	{
 		FilteredDatabaseView.ListView->ClearSelection();
 	}
+
+	ComputeFilteredDatabaseRowsColors();
+}
+
+void SDebuggerDatabaseView::ComputeFilteredDatabaseRowsColors()
+{
+	float CostMin = MAX_flt;
+	float CostMax = -MAX_flt;
+
+	float PoseCostMin = MAX_flt;
+	float PoseCostMax = -MAX_flt;
+
+	float TrajectoryCostMin = MAX_flt;
+	float TrajectoryCostMax = -MAX_flt;
+
+	for (const TSharedRef<FDebuggerDatabaseRowData>& Row : FilteredDatabaseView.Rows)
+	{
+		const float Cost = Row->PoseCostDetails.PoseCost.TotalCost;
+		CostMin = FMath::Min(CostMin, Cost);
+		CostMax = FMath::Max(CostMax, Cost);
+
+		const float PoseCost = Row->GetPoseCost();
+		PoseCostMin = FMath::Min(PoseCostMin, PoseCost);
+		PoseCostMax = FMath::Max(PoseCostMax, PoseCost);
+
+		const float TrajectoryCost = Row->GetTrajectoryCost();
+		TrajectoryCostMin = FMath::Min(TrajectoryCostMin, TrajectoryCost);
+		TrajectoryCostMax = FMath::Max(TrajectoryCostMax, TrajectoryCost);
+	}
+
+	float CostDelta = CostMax - CostMin;
+	if (FMath::IsNearlyZero(CostDelta))
+	{
+		CostDelta = 1.f;
+	}
+
+	float PoseCostDelta = PoseCostMax - PoseCostMin;
+	if (FMath::IsNearlyZero(PoseCostDelta))
+	{
+		PoseCostDelta = 1.f;
+	}
+
+	float TrajectoryCostDelta = TrajectoryCostMax - TrajectoryCostMin;
+	if (FMath::IsNearlyZero(TrajectoryCostDelta))
+	{
+		TrajectoryCostDelta = 1.f;
+	}
+
+	for (const TSharedRef<FDebuggerDatabaseRowData>& Row : FilteredDatabaseView.Rows)
+	{
+		Row->CalculateColors(CostMin, CostDelta, PoseCostMin, PoseCostDelta, TrajectoryCostMin, TrajectoryCostDelta);
+	}
+
+	if (!ActiveView.Rows.IsEmpty())
+	{
+		ActiveView.Rows[0]->CalculateColors(CostMin, CostDelta, PoseCostMin, PoseCostDelta, TrajectoryCostMin, TrajectoryCostDelta);
+	}
+
+	if (!ContinuingPoseView.Rows.IsEmpty())
+	{
+		ContinuingPoseView.Rows[0]->CalculateColors(CostMin, CostDelta, PoseCostMin, PoseCostDelta, TrajectoryCostMin, TrajectoryCostDelta);
+	}
 }
 
 TSharedRef<ITableRow> SDebuggerDatabaseView::HandleGenerateDatabaseRow(TSharedRef<FDebuggerDatabaseRowData> Item, const TSharedRef<STableViewBase>& OwnerTable) const
diff --git a/Engine/Plugins/Experimental/Animation/PoseSearch/Source/Editor/Private/PoseSearchDebugger.h b/Engine/Plugins/Experimental/Animation/PoseSearch/Source/Editor/Private/PoseSearchDebugger.h
--- a/Engine/Plugins/Experimental/Animation/PoseSearch/Source/Editor/Private/PoseSearchDebugger.h
+++ b/Engine/Plugins/Experimental/Animation/PoseSearch/Source/Editor/Private/PoseSearchDebugger.h
@@ -293,6 +293,8 @@ private:
 
 	void FilterDatabaseRows();
 
+	void ComputeFilteredDatabaseRowsColors();
+
 	/** Sets dynamic data for each row, such as score at the current time */
 	void UpdateRows(const FTraceMotionMatchingStateMessage& State, const UPoseSearchDatabase& Database);
 
