PoseSearch: Implemented pose search database asset list
- asset tree shows sequences organized into groups
- users can press the add button to add sequences, blendspaces or groups
- each group has a + button which allows adding sequences/blendspaces directly to that group
- drag and drop from the content browser allows adding sequences and blendspaces directly to specific groups
- right clicking allows users to remove assets and groups

diff --git a/Engine/Plugins/Experimental/Animation/PoseSearch/Source/Editor/PoseSearchEditor.Build.cs b/Engine/Plugins/Experimental/Animation/PoseSearch/Source/Editor/PoseSearchEditor.Build.cs
--- a/Engine/Plugins/Experimental/Animation/PoseSearch/Source/Editor/PoseSearchEditor.Build.cs
+++ b/Engine/Plugins/Experimental/Animation/PoseSearch/Source/Editor/PoseSearchEditor.Build.cs
@@ -48,7 +48,8 @@ public class PoseSearchEditor : ModuleRules
 				"EditorStyle",
 				"DetailCustomizations",
 				"AdvancedPreviewScene",
-				"EditorFramework"
+				"EditorFramework",
+				"ToolWidgets"
 			}
 		);
 	}
diff --git a/Engine/Plugins/Experimental/Animation/PoseSearch/Source/Editor/Private/PoseSearchDatabaseEditorToolkit.cpp b/Engine/Plugins/Experimental/Animation/PoseSearch/Source/Editor/Private/PoseSearchDatabaseEditorToolkit.cpp
--- a/Engine/Plugins/Experimental/Animation/PoseSearch/Source/Editor/Private/PoseSearchDatabaseEditorToolkit.cpp
+++ b/Engine/Plugins/Experimental/Animation/PoseSearch/Source/Editor/Private/PoseSearchDatabaseEditorToolkit.cpp
@@ -2,6 +2,7 @@
 
 #include "PoseSearchDatabaseEditorToolkit.h"
 #include "SPoseSearchDatabaseViewport.h"
+#include "SPoseSearchDatabaseAssetList.h"
 #include "PoseSearchDatabasePreviewScene.h"
 #include "PoseSearchDatabaseEditorCommands.h"
 #include "PoseSearchDatabaseViewModel.h"
@@ -25,11 +26,13 @@ struct FPoseSearchDatabaseEditorTabs
 	static const FName AssetDetailsID;
 	static const FName ViewportID;
 	static const FName PreviewSettingsID;
+	static const FName AssetTreeViewID;
 };
 
 const FName FPoseSearchDatabaseEditorTabs::AssetDetailsID(TEXT("PoseSearchDatabaseEditorAssetDetailsTabID"));
 const FName FPoseSearchDatabaseEditorTabs::ViewportID(TEXT("PoseSearchDatabaseEditorViewportTabID"));
 const FName FPoseSearchDatabaseEditorTabs::PreviewSettingsID(TEXT("PoseSearchDatabaseEditorPreviewSettingsTabID"));
+const FName FPoseSearchDatabaseEditorTabs::AssetTreeViewID(TEXT("PoseSearchDatabaseEditorAssetTreeViewTabID"));
 
 FPoseSearchDatabaseEditorToolkit::FPoseSearchDatabaseEditorToolkit()
 {
@@ -37,6 +40,12 @@ FPoseSearchDatabaseEditorToolkit::FPoseSearchDatabaseEditorToolkit()
 
 FPoseSearchDatabaseEditorToolkit::~FPoseSearchDatabaseEditorToolkit()
 {
+	UPoseSearchDatabase* DatabaseAsset = ViewModel->GetPoseSearchDatabase();
+	if (IsValid(DatabaseAsset))
+	{
+		DatabaseAsset->UnregisterOnAssetChange(AssetTreeWidget.Get());
+		DatabaseAsset->UnregisterOnGroupChange(AssetTreeWidget.Get());
+	}
 }
 
 const UPoseSearchDatabase* FPoseSearchDatabaseEditorToolkit::GetPoseSearchDatabase() const
@@ -92,6 +101,19 @@ void FPoseSearchDatabaseEditorToolkit::InitAssetEditor(
 		PreviewScene.ToSharedRef());
 	ViewportWidget = SNew(SPoseSearchDatabaseViewport, ViewportArgs);
 
+	AssetTreeWidget = SNew(UE::PoseSearch::SDatabaseAssetTree, ViewModel.ToSharedRef());
+	if (IsValid(DatabaseAsset))
+	{
+		DatabaseAsset->RegisterOnAssetChange(
+			UPoseSearchDatabase::FOnDerivedDataRebuild::CreateSP(
+				AssetTreeWidget.Get(),
+				&UE::PoseSearch::SDatabaseAssetTree::RefreshTreeView, false));
+		DatabaseAsset->RegisterOnGroupChange(
+			UPoseSearchDatabase::FOnDerivedDataRebuild::CreateSP(
+				AssetTreeWidget.Get(),
+				&UE::PoseSearch::SDatabaseAssetTree::RefreshTreeView, false));
+	}
+
 	// Create Asset Details widget
 	FPropertyEditorModule& PropertyModule = FModuleManager::LoadModuleChecked<FPropertyEditorModule>("PropertyEditor");
 
@@ -107,7 +129,7 @@ void FPoseSearchDatabaseEditorToolkit::InitAssetEditor(
 
 	// Define Editor Layout
 	const TSharedRef<FTabManager::FLayout> StandaloneDefaultLayout = 
-		FTabManager::NewLayout("Standalone_PoseSearchDatabaseEditor_Layout_v0.02")
+		FTabManager::NewLayout("Standalone_PoseSearchDatabaseEditor_Layout_v0.04")
 		->AddArea
 		(
 			FTabManager::NewPrimaryArea()->SetOrientation(Orient_Vertical)
@@ -120,7 +142,14 @@ void FPoseSearchDatabaseEditorToolkit::InitAssetEditor(
 					->Split
 					(
 						FTabManager::NewStack()
-						->SetSizeCoefficient(0.65f)
+						->SetSizeCoefficient(0.25f)
+						->AddTab(FPoseSearchDatabaseEditorTabs::AssetTreeViewID, ETabState::OpenedTab)
+						->SetHideTabWell(true)
+					)
+					->Split
+					(
+						FTabManager::NewStack()
+						->SetSizeCoefficient(0.4f)
 						->AddTab(FPoseSearchDatabaseEditorTabs::ViewportID, ETabState::OpenedTab)
 						->SetHideTabWell(true)
 					)
@@ -242,15 +271,23 @@ void FPoseSearchDatabaseEditorToolkit::RegisterTabSpawners(const TSharedRef<FTab
 		.SetDisplayName(LOCTEXT("PreviewSceneSettingsTab", "Preview Scene Settings"))
 		.SetGroup(WorkspaceMenuCategoryRef)
 		.SetIcon(FSlateIcon(FEditorStyle::GetStyleSetName(), "LevelEditor.Tabs.Details"));
+
+	InTabManager->RegisterTabSpawner(
+		FPoseSearchDatabaseEditorTabs::AssetTreeViewID,
+		FOnSpawnTab::CreateSP(this, &FPoseSearchDatabaseEditorToolkit::SpawnTab_AssetTreeView))
+		.SetDisplayName(LOCTEXT("TreeViewTab", "Tree View"))
+		.SetGroup(WorkspaceMenuCategoryRef)
+		.SetIcon(FSlateIcon(FEditorStyle::GetStyleSetName(), "GraphEditor.EventGraph_16x"));
 }
 
 void FPoseSearchDatabaseEditorToolkit::UnregisterTabSpawners(const TSharedRef<FTabManager>& InTabManager)
 {
 	FAssetEditorToolkit::UnregisterTabSpawners(InTabManager);
 
-	// InTabManager->UnregisterTabSpawner(FPoseSearchDatabaseEditorTabs::ViewportID);
+	InTabManager->UnregisterTabSpawner(FPoseSearchDatabaseEditorTabs::ViewportID);
 	InTabManager->UnregisterTabSpawner(FPoseSearchDatabaseEditorTabs::AssetDetailsID);
-	// InTabManager->UnregisterTabSpawner(FPoseSearchDatabaseEditorTabs::PreviewSettingsID);
+	InTabManager->UnregisterTabSpawner(FPoseSearchDatabaseEditorTabs::PreviewSettingsID);
+	InTabManager->UnregisterTabSpawner(FPoseSearchDatabaseEditorTabs::AssetTreeViewID);
 }
 
 FName FPoseSearchDatabaseEditorToolkit::GetToolkitFName() const
@@ -309,8 +346,10 @@ TSharedRef<SDockTab> FPoseSearchDatabaseEditorToolkit::SpawnTab_PreviewSettings(
 {
 	check(Args.GetTabId() == FPoseSearchDatabaseEditorTabs::PreviewSettingsID);
 
-	FAdvancedPreviewSceneModule& AdvancedPreviewSceneModule = FModuleManager::LoadModuleChecked<FAdvancedPreviewSceneModule>("AdvancedPreviewScene");
-	TSharedRef<SWidget> InWidget= AdvancedPreviewSceneModule.CreateAdvancedPreviewSceneSettingsWidget(PreviewScene.ToSharedRef());
+	FAdvancedPreviewSceneModule& AdvancedPreviewSceneModule = 
+		FModuleManager::LoadModuleChecked<FAdvancedPreviewSceneModule>("AdvancedPreviewScene");
+	TSharedRef<SWidget> InWidget = 
+		AdvancedPreviewSceneModule.CreateAdvancedPreviewSceneSettingsWidget(PreviewScene.ToSharedRef());
 
 	TSharedRef<SDockTab> SpawnedTab = SNew(SDockTab)
 		.Label(LOCTEXT("PreviewSceneSettingsTab", "Preview Scene Settings"))
@@ -321,6 +360,17 @@ TSharedRef<SDockTab> FPoseSearchDatabaseEditorToolkit::SpawnTab_PreviewSettings(
 	return SpawnedTab;
 }
 
+TSharedRef<SDockTab> FPoseSearchDatabaseEditorToolkit::SpawnTab_AssetTreeView(const FSpawnTabArgs& Args)
+{
+	check(Args.GetTabId() == FPoseSearchDatabaseEditorTabs::AssetTreeViewID);
+
+	return SNew(SDockTab)
+		.Label(LOCTEXT("AssetTreeView_Title", "Asset Tree"))
+		[
+			AssetTreeWidget.ToSharedRef()
+		];
+}
+
 void FPoseSearchDatabaseEditorToolkit::OnFinishedChangingProperties(const FPropertyChangedEvent& PropertyChangedEvent)
 {
 	const FName PropertyName = 
diff --git a/Engine/Plugins/Experimental/Animation/PoseSearch/Source/Editor/Private/PoseSearchDatabaseViewModel.cpp b/Engine/Plugins/Experimental/Animation/PoseSearch/Source/Editor/Private/PoseSearchDatabaseViewModel.cpp
--- a/Engine/Plugins/Experimental/Animation/PoseSearch/Source/Editor/Private/PoseSearchDatabaseViewModel.cpp
+++ b/Engine/Plugins/Experimental/Animation/PoseSearch/Source/Editor/Private/PoseSearchDatabaseViewModel.cpp
@@ -9,6 +9,7 @@
 #include "Animation/DebugSkelMeshComponent.h"
 #include "Animation/MirrorDataTable.h"
 #include "Animation/AnimSequence.h"
+#include "Animation/BlendSpace.h"
 #include "EngineUtils.h"
 #include "PropertyEditorModule.h"
 #include "Modules/ModuleManager.h"
@@ -22,6 +23,7 @@ FPoseSearchDatabaseViewModel::FPoseSearchDatabaseViewModel()
 
 FPoseSearchDatabaseViewModel::~FPoseSearchDatabaseViewModel()
 {
+	PoseSearchDatabase->UnregisterOnDerivedDataRebuild(this);
 }
 
 void FPoseSearchDatabaseViewModel::AddReferencedObjects(FReferenceCollector& Collector)
@@ -39,10 +41,9 @@ void FPoseSearchDatabaseViewModel::Initialize(
 	RemovePreviewActors();
 
 	PoseSearchDatabase->RegisterOnDerivedDataRebuild(
-		UPoseSearchDatabase::FOnDerivedDataRebuild::CreateLambda([this]()
-	{
-		ResetPreviewActors();
-	}));
+		UPoseSearchDatabase::FOnDerivedDataRebuild::CreateSP(
+			this, 
+			&FPoseSearchDatabaseViewModel::ResetPreviewActors));
 }
 
 void FPoseSearchDatabaseViewModel::ResetPreviewActors()
@@ -229,3 +230,83 @@ bool FPoseSearchDatabaseViewModel::IsAnimationPreviewMode(EAnimationPreviewMode
 {
 	return AnimationPreviewMode == PreviewMode;
 }
+
+void FPoseSearchDatabaseViewModel::AddSequenceToDatabase(UAnimSequence* AnimSequence, int InitialGroupIdx)
+{
+	FPoseSearchDatabaseSequence& NewDbSequence = PoseSearchDatabase->Sequences.AddDefaulted_GetRef();
+	NewDbSequence.Sequence = AnimSequence;
+	
+	if (InitialGroupIdx >= 0)
+	{
+		const FPoseSearchDatabaseGroup& InitialGroup = PoseSearchDatabase->Groups[InitialGroupIdx];
+		NewDbSequence.GroupTags.AddTag(InitialGroup.Tag);
+	}
+}
+
+void FPoseSearchDatabaseViewModel::AddBlendSpaceToDatabase(UBlendSpace* BlendSpace, int InitialGroupIdx)
+{
+	FPoseSearchDatabaseBlendSpace& NewDbBlendSpace = PoseSearchDatabase->BlendSpaces.AddDefaulted_GetRef();
+	NewDbBlendSpace.BlendSpace = BlendSpace;
+
+	if (InitialGroupIdx >= 0)
+	{
+		const FPoseSearchDatabaseGroup& InitialGroup = PoseSearchDatabase->Groups[InitialGroupIdx];
+		NewDbBlendSpace.GroupTags.AddTag(InitialGroup.Tag);
+	}
+}
+
+void FPoseSearchDatabaseViewModel::AddGroupToDatabase()
+{
+	PoseSearchDatabase->Groups.AddDefaulted();
+}
+
+void FPoseSearchDatabaseViewModel::DeleteSequenceFromDatabase(int32 SequenceIdx)
+{
+	PoseSearchDatabase->Sequences.RemoveAt(SequenceIdx);
+}
+
+void FPoseSearchDatabaseViewModel::RemoveSequenceFromGroup(int32 SequenceIdx, int32 GroupIdx)
+{
+	FPoseSearchDatabaseSequence& DbSequence = PoseSearchDatabase->Sequences[SequenceIdx];
+	FPoseSearchDatabaseGroup& DbGroup = PoseSearchDatabase->Groups[GroupIdx];
+	if (DbGroup.Tag.IsValid())
+	{
+		DbSequence.GroupTags.RemoveTag(DbGroup.Tag);
+	}
+}
+
+void FPoseSearchDatabaseViewModel::DeleteBlendSpaceFromDatabase(int32 BlendSpaceIdx)
+{
+	PoseSearchDatabase->BlendSpaces.RemoveAt(BlendSpaceIdx);
+}
+
+void FPoseSearchDatabaseViewModel::RemoveBlendSpaceFromGroup(int32 BlendSpaceIdx, int32 GroupIdx)
+{
+	FPoseSearchDatabaseBlendSpace& DbBlendSpace = PoseSearchDatabase->BlendSpaces[BlendSpaceIdx];
+	FPoseSearchDatabaseGroup& DbGroup = PoseSearchDatabase->Groups[GroupIdx];
+	if (DbGroup.Tag.IsValid())
+	{
+		DbBlendSpace.GroupTags.RemoveTag(DbGroup.Tag);
+	}
+}
+
+void FPoseSearchDatabaseViewModel::DeleteGroup(int32 GroupIdx)
+{
+	FPoseSearchDatabaseGroup& DbGroup = PoseSearchDatabase->Groups[GroupIdx];
+	
+	if (DbGroup.Tag.IsValid())
+	{
+		for (FPoseSearchDatabaseSequence& Sequence : PoseSearchDatabase->Sequences)
+		{
+			Sequence.GroupTags.RemoveTag(DbGroup.Tag);
+		}
+
+		for (FPoseSearchDatabaseBlendSpace& BlendSpace : PoseSearchDatabase->BlendSpaces)
+		{
+			BlendSpace.GroupTags.RemoveTag(DbGroup.Tag);
+		}
+	}
+
+	PoseSearchDatabase->Groups.RemoveAt(GroupIdx);
+}
+
diff --git a/Engine/Plugins/Experimental/Animation/PoseSearch/Source/Editor/Private/SPoseSearchDatabaseAssetList.cpp b/Engine/Plugins/Experimental/Animation/PoseSearch/Source/Editor/Private/SPoseSearchDatabaseAssetList.cpp
new file mode 100644
--- /dev/null
+++ b/Engine/Plugins/Experimental/Animation/PoseSearch/Source/Editor/Private/SPoseSearchDatabaseAssetList.cpp
@@ -0,0 +1,781 @@
+// Copyright Epic Games, Inc. All Rights Reserved.
+
+#include "SPoseSearchDatabaseAssetList.h"
+#include "PoseSearchDatabaseViewModel.h"
+
+#include "PoseSearch/PoseSearch.h"
+
+#include "Animation/AnimSequence.h"
+#include "Animation/BlendSpace.h"
+
+#include "DragAndDrop/AssetDragDropOp.h"
+#include "Misc/FeedbackContext.h"
+#include "AssetSelection.h"
+
+#include "Widgets/Text/SRichTextBlock.h"
+#include "Framework/MultiBox/MultiBoxBuilder.h"
+#include "SPositiveActionButton.h"
+
+#define LOCTEXT_NAMESPACE "PoseSearchDatabaseAssetList"
+
+namespace UE::PoseSearch
+{
+	FDatabaseAssetTreeNode::FDatabaseAssetTreeNode(
+		int32 InSourceAssetIdx,
+		ESearchIndexAssetType InSourceAssetType,
+		const TSharedRef<FPoseSearchDatabaseViewModel>& InEditorViewModel) 
+		: SourceAssetIdx(InSourceAssetIdx)
+		, SourceAssetType(InSourceAssetType)
+		, EditorViewModel(InEditorViewModel)
+	{ }
+
+	TSharedRef<ITableRow> FDatabaseAssetTreeNode::MakeTreeRowWidget(
+		const TSharedRef<STableViewBase>& InOwnerTable,
+		TSharedRef<FDatabaseAssetTreeNode> InDatabaseAssetNode,
+		TSharedRef<FUICommandList> InCommandList,
+		TSharedPtr<SDatabaseAssetTree> InHierarchy)
+	{
+		return SNew(
+			SDatabaseAssetListItem, 
+			EditorViewModel.Pin().ToSharedRef(), 
+			InOwnerTable, 
+			InDatabaseAssetNode, 
+			InCommandList, 
+			InHierarchy);
+	}
+
+	void SDatabaseAssetListItem::Construct(
+		const FArguments& InArgs,
+		const TSharedRef<FPoseSearchDatabaseViewModel>& InEditorViewModel,
+		const TSharedRef<STableViewBase>& OwnerTable,
+		TSharedRef<FDatabaseAssetTreeNode> InAssetTreeNode,
+		TSharedRef<FUICommandList> InCommandList,
+		TSharedPtr<SDatabaseAssetTree> InHierarchy)
+	{
+		WeakAssetTreeNode = InAssetTreeNode;
+		EditorViewModel = InEditorViewModel;
+		SkeletonView = InHierarchy;
+
+		if (InAssetTreeNode->SourceAssetType == ESearchIndexAssetType::Invalid)
+		{
+			ConstructGroupItem(OwnerTable);
+		}
+		else
+		{
+			ConstructAssetItem(OwnerTable);
+		}
+	}
+
+	void SDatabaseAssetListItem::ConstructGroupItem(const TSharedRef<STableViewBase>& OwnerTable)
+	{
+		STableRow<TSharedPtr<FDatabaseAssetTreeNode>>::ChildSlot
+		.Padding(0.0f, 2.0f, 0.0f, 0.0f)
+		[
+			GenerateItemWidget()
+		];
+
+		STableRow<TSharedPtr<FDatabaseAssetTreeNode>>::ConstructInternal(
+			STableRow<TSharedPtr<FDatabaseAssetTreeNode>>::FArguments()
+			.Style(FEditorStyle::Get(), "DetailsView.TreeView.TableRow")
+			.OnCanAcceptDrop(SkeletonView.Pin().Get(), &SDatabaseAssetTree::OnCanAcceptDrop)
+			.OnAcceptDrop(SkeletonView.Pin().Get(), &SDatabaseAssetTree::OnAcceptDrop)
+			.ShowSelection(true),
+			OwnerTable);
+	}
+
+	void SDatabaseAssetListItem::ConstructAssetItem(const TSharedRef<STableViewBase>& OwnerTable)
+	{
+		STableRow<TSharedPtr<FDatabaseAssetTreeNode>>::Construct(
+			STableRow<TSharedPtr<FDatabaseAssetTreeNode>>::FArguments()
+			.Style(&FEditorStyle::Get().GetWidgetStyle<FTableRowStyle>("TableView.Row"))
+			.OnCanAcceptDrop(SkeletonView.Pin().Get(), &SDatabaseAssetTree::OnCanAcceptDrop)
+			.OnAcceptDrop(SkeletonView.Pin().Get(), &SDatabaseAssetTree::OnAcceptDrop)
+			.ShowWires(false)
+			.Content()
+			[
+				GenerateItemWidget()
+			], OwnerTable);
+	}
+
+	void SDatabaseAssetListItem::OnAddSequence()
+	{
+		TSharedPtr<FDatabaseAssetTreeNode> NodePtr = WeakAssetTreeNode.Pin();
+		const int32 GroupIdx = NodePtr->Parent.IsValid() ? NodePtr->Parent->SourceAssetIdx : NodePtr->SourceAssetIdx;
+
+		EditorViewModel.Pin()->AddSequenceToDatabase(nullptr, GroupIdx);
+
+		SkeletonView.Pin()->RefreshTreeView(false);
+	}
+
+	void SDatabaseAssetListItem::OnAddBlendSpace()
+	{
+		TSharedPtr<FDatabaseAssetTreeNode> NodePtr = WeakAssetTreeNode.Pin();
+		const int32 GroupIdx = NodePtr->Parent.IsValid() ? NodePtr->Parent->SourceAssetIdx : NodePtr->SourceAssetIdx;
+
+		EditorViewModel.Pin()->AddBlendSpaceToDatabase(nullptr, GroupIdx);
+
+		SkeletonView.Pin()->RefreshTreeView(false);
+	}
+
+	FText SDatabaseAssetListItem::GetName() const
+	{
+		FText Name = LOCTEXT("None", "None");
+
+		TSharedPtr<FDatabaseAssetTreeNode> Node = WeakAssetTreeNode.Pin();
+		TSharedPtr<FPoseSearchDatabaseViewModel> ViewModel = EditorViewModel.Pin();
+		UPoseSearchDatabase* Database = ViewModel->GetPoseSearchDatabase();
+
+		if (Database)
+		{
+			switch (Node->SourceAssetType)
+			{
+			case ESearchIndexAssetType::Sequence:
+			{
+				const UAnimSequence* Sequence = Database->Sequences[Node->SourceAssetIdx].Sequence;
+				if (Sequence)
+				{
+					Name = FText::FromString(Sequence->GetName());
+				}
+				break;
+			}
+			case ESearchIndexAssetType::BlendSpace:
+			{
+				const UBlendSpace* BlendSpace = Database->BlendSpaces[Node->SourceAssetIdx].BlendSpace;
+				if (BlendSpace)
+				{
+					Name = FText::FromString(BlendSpace->GetName());
+				}
+				break;
+			}
+			default:
+			{
+				if (Node->SourceAssetIdx == INDEX_NONE)
+				{
+					Name = LOCTEXT("Default", "Default");
+				}
+				else
+				{
+					FPoseSearchDatabaseGroup& Group = Database->Groups[Node->SourceAssetIdx];
+					if (Group.Tag.IsValid())
+					{
+						Name = FText::FromName(Group.Tag.GetTagName());
+					}
+				}
+			}
+			}
+		}
+
+		return Name;
+	}
+
+	TSharedRef<SWidget> SDatabaseAssetListItem::GenerateItemWidget()
+	{
+		TSharedPtr<FDatabaseAssetTreeNode> Node = WeakAssetTreeNode.Pin();
+
+		TSharedPtr<SWidget> ItemWidget;
+
+		if (Node->SourceAssetType == ESearchIndexAssetType::Invalid)
+		{
+			// it's a group
+			SAssignNew(ItemWidget, SBorder)
+			.BorderImage(this, &SDatabaseAssetListItem::GetGroupBackgroundImage)
+			.Padding(FMargin(3.0f, 5.0f))
+			[
+				SNew(SHorizontalBox)
+				+ SHorizontalBox::Slot()
+				.VAlign(VAlign_Center)
+				.Padding(5.0f)
+				.AutoWidth()
+				[
+					SNew(SExpanderArrow, STableRow<TSharedPtr<FDatabaseAssetTreeNode>>::SharedThis(this))
+				]
+				+ SHorizontalBox::Slot()
+				.FillWidth(1.0f)
+				.VAlign(VAlign_Center)
+				[
+					SNew(SRichTextBlock)
+					.Text(this, &SDatabaseAssetListItem::GetName)
+					.TransformPolicy(ETextTransformPolicy::ToUpper)
+					.DecoratorStyleSet(&FEditorStyle::Get())
+					.TextStyle(FAppStyle::Get(), "DetailsView.CategoryTextStyle")
+				]
+				+SHorizontalBox::Slot()
+				.AutoWidth()
+				.VAlign(VAlign_Center)
+				.HAlign(HAlign_Right)
+				.Padding(2, 0, 0, 0)
+				[
+					GenerateAddButtonWidget()
+				]
+			];
+		}
+		else
+		{
+			// it's an asset (sequence or blendspace)
+			SAssignNew(ItemWidget, SHorizontalBox)
+			+ SHorizontalBox::Slot()
+			.MaxWidth(18)
+			.AutoWidth()
+			.HAlign(HAlign_Left)
+			.VAlign(VAlign_Center)
+			[
+				SNew(SImage)
+				.Image(FAppStyle::Get().GetBrush("Icons.Star"))
+			]
+			+ SHorizontalBox::Slot()
+			.FillWidth(1.0f)
+			.VAlign(VAlign_Center)
+			[
+				SNew(STextBlock)
+				.Text(this, &SDatabaseAssetListItem::GetName)
+			];
+		}
+
+		return ItemWidget.ToSharedRef();
+	}
+
+	TSharedRef<SWidget> SDatabaseAssetListItem::GenerateAddButtonWidget()
+	{
+		FMenuBuilder AddOptions(true, nullptr);
+
+		AddOptions.AddMenuEntry(
+			LOCTEXT("AddSequence", "Add Sequence"),
+			LOCTEXT("AddSequenceTooltip", "Add new sequence to this group"),
+			FSlateIcon(),
+			FUIAction(FExecuteAction::CreateSP(this, &SDatabaseAssetListItem::OnAddSequence)),
+			NAME_None,
+			EUserInterfaceActionType::Button);
+
+		AddOptions.AddMenuEntry(
+			LOCTEXT("AddBlendSpace", "Add Blend Space"),
+			LOCTEXT("AddBlendSpaceTooltip", "Add new blend space to this group"),
+			FSlateIcon(),
+			FUIAction(FExecuteAction::CreateSP(this, &SDatabaseAssetListItem::OnAddBlendSpace)),
+			NAME_None,
+			EUserInterfaceActionType::Button);
+
+		TSharedPtr<SComboButton> AddButton;
+		SAssignNew(AddButton, SComboButton)
+		.ContentPadding(0)
+		.ComboButtonStyle(&FAppStyle::Get().GetWidgetStyle<FComboButtonStyle>("SimpleComboButton"))
+		.HasDownArrow(false)
+		.ButtonContent()
+		[
+			SNew(SImage)
+			.ColorAndOpacity(FSlateColor::UseForeground())
+			.Image(FAppStyle::Get().GetBrush("Icons.PlusCircle"))
+		]
+		.MenuContent()
+		[
+			AddOptions.MakeWidget()
+		];
+
+		return AddButton.ToSharedRef();
+	}
+
+
+	const FSlateBrush* SDatabaseAssetListItem::GetGroupBackgroundImage() const
+	{
+		if (STableRow<TSharedPtr<FDatabaseAssetTreeNode>>::IsHovered())
+		{
+			return FAppStyle::Get().GetBrush("Brushes.Secondary");
+		}
+		else
+		{
+			return FAppStyle::Get().GetBrush("Brushes.Header");
+		}
+	}
+
+	SDatabaseAssetTree::~SDatabaseAssetTree()
+	{
+	}
+
+	void SDatabaseAssetTree::Construct(
+		const FArguments& InArgs, 
+		TSharedRef<FPoseSearchDatabaseViewModel> InEditorViewModel)
+	{
+		EditorViewModel = InEditorViewModel;
+		CommandList = MakeShared<FUICommandList>();
+
+		ChildSlot
+		[
+			SNew(SVerticalBox)
+			+SVerticalBox::Slot()
+			.AutoHeight()
+			[
+				SNew(SHorizontalBox)
+				+SHorizontalBox::Slot()
+				.AutoWidth()
+				.VAlign(VAlign_Center)
+				.Padding(0, 0, 4, 0)
+				[
+					SNew(SPositiveActionButton)
+					.Icon(FAppStyle::Get().GetBrush("Icons.Plus"))
+					.Text(LOCTEXT("AddNew", "Add"))
+					.ToolTipText(LOCTEXT("AddNewToolTip", "Add a new Sequence, Blend Space or Group"))
+					.OnGetMenuContent(this, &SDatabaseAssetTree::CreateAddNewMenuWidget)
+				]
+			]
+			+SVerticalBox::Slot()
+			.Padding(0.0f, 0.0f)
+			[
+				SNew(SBorder)
+				.Padding(2.0f)
+				.BorderImage(FEditorStyle::GetBrush("SCSEditor.TreePanel"))
+				[
+					SAssignNew(TreeView, STreeView<TSharedPtr<FDatabaseAssetTreeNode>>)
+					.TreeItemsSource(&RootNodes)
+					.SelectionMode(ESelectionMode::Multi)
+					.OnGenerateRow(this, &SDatabaseAssetTree::MakeTableRowWidget)
+					.OnGetChildren(this, &SDatabaseAssetTree::HandleGetChildrenForTree)
+					.OnContextMenuOpening(this, &SDatabaseAssetTree::CreateContextMenu)
+					.HighlightParentNodesForSelection(false)
+					.ItemHeight(24)
+				]
+			]
+		];
+
+		const bool IsInitialSetup = true;
+		RefreshTreeView(IsInitialSetup);
+	}
+
+	FReply SDatabaseAssetTree::OnDragOver(const FGeometry& MyGeometry, const FDragDropEvent& DragDropEvent)
+	{
+		FReply Reply = FReply::Unhandled();
+
+		TSharedPtr<FDragDropOperation> Operation = DragDropEvent.GetOperation();
+
+		const bool bValidOperation =
+			Operation.IsValid() &&
+			(Operation->IsOfType<FExternalDragOperation>() || Operation->IsOfType<FAssetDragDropOp>());
+		if (bValidOperation)
+		{
+			Reply = AssetUtil::CanHandleAssetDrag(DragDropEvent);
+
+			if (!Reply.IsEventHandled())
+			{
+				if (Operation->IsOfType<FAssetDragDropOp>())
+				{
+					const TSharedPtr<FAssetDragDropOp> AssetDragDropOp = StaticCastSharedPtr<FAssetDragDropOp>(Operation);
+
+					for (const FAssetData& AssetData : AssetDragDropOp->GetAssets())
+					{
+						if (UClass* AssetClass = AssetData.GetClass())
+						{
+							if (AssetClass->IsChildOf(UAnimSequence::StaticClass()) ||
+								AssetClass->IsChildOf(UBlendSpace::StaticClass()))
+							{
+								Reply = FReply::Handled();
+								break;
+							}
+						}
+					}
+				}
+			}
+		}
+
+		return Reply;
+	}
+
+	FReply SDatabaseAssetTree::OnDrop(const FGeometry& MyGeometry, const FDragDropEvent& DragDropEvent)
+	{
+		return OnAcceptDrop(DragDropEvent, EItemDropZone::OntoItem, nullptr);
+	}
+
+	void SDatabaseAssetTree::RefreshTreeView(bool IsInitialSetup)
+	{
+		const TSharedPtr<FPoseSearchDatabaseViewModel> ViewModel = EditorViewModel.Pin();
+		if (!ViewModel.IsValid())
+		{
+			return;
+		}
+
+		const TSharedRef<FPoseSearchDatabaseViewModel> ViewModelRef = ViewModel.ToSharedRef();
+
+		RootNodes.Reset();
+		AllNodes.Reset();
+
+		const UPoseSearchDatabase* Database = ViewModel->GetPoseSearchDatabase();
+		if (!IsValid(Database))
+		{
+			TreeView->RequestTreeRefresh();
+			return;
+		}
+
+		// create all group nodes
+		for (int32 GroupIdx = 0; GroupIdx < Database->Groups.Num(); ++GroupIdx)
+		{
+			TSharedPtr<FDatabaseAssetTreeNode> GroupNode = MakeShared<FDatabaseAssetTreeNode>(
+				GroupIdx,
+				ESearchIndexAssetType::Invalid,
+				ViewModelRef);
+			AllNodes.Add(GroupNode);
+			RootNodes.Add(GroupNode);
+		}
+		TSharedPtr<FDatabaseAssetTreeNode> DefaultGroupNode = MakeShared<FDatabaseAssetTreeNode>(
+			INDEX_NONE,
+			ESearchIndexAssetType::Invalid,
+			ViewModelRef);
+		AllNodes.Add(DefaultGroupNode);
+		RootNodes.Add(DefaultGroupNode);
+
+		const int32 DefaultGroupIdx = RootNodes.Num() - 1;
+
+		auto CreateAssetNode = [this, ViewModelRef](int32 AssetIdx, ESearchIndexAssetType AssetType, int32 GroupIdx)
+		{
+			TSharedPtr<FDatabaseAssetTreeNode> SequenceGroupNode = MakeShared<FDatabaseAssetTreeNode>(
+				AssetIdx,
+				AssetType,
+				ViewModelRef);
+			TSharedPtr<FDatabaseAssetTreeNode>& ParentGroupNode = RootNodes[GroupIdx];
+			SequenceGroupNode->Parent = ParentGroupNode;
+			ParentGroupNode->Children.Add(SequenceGroupNode);
+			AllNodes.Add(SequenceGroupNode);
+		};
+
+		// create all sequence nodes
+		for (int32 SequenceIdx = 0; SequenceIdx < Database->Sequences.Num(); ++SequenceIdx)
+		{
+			const FPoseSearchDatabaseSequence& DbSequence = Database->Sequences[SequenceIdx];
+
+			int32 NumGroups = 0;
+			for (const FGameplayTag& GroupTag : DbSequence.GroupTags)
+			{
+				const int32 GroupIdx = Database->Groups.IndexOfByPredicate(
+					[GroupTag](const FPoseSearchDatabaseGroup& Group)
+				{
+					return Group.Tag == GroupTag;
+				});
+
+				if (GroupIdx != INDEX_NONE)
+				{
+					CreateAssetNode(SequenceIdx, ESearchIndexAssetType::Sequence, GroupIdx);
+					++NumGroups;
+				}
+			}
+
+			if (NumGroups == 0)
+			{
+				CreateAssetNode(SequenceIdx, ESearchIndexAssetType::Sequence, DefaultGroupIdx);
+			}
+		}
+
+		// create all blendspace nodes
+		for (int32 BlendSpaceIdx = 0; BlendSpaceIdx < Database->BlendSpaces.Num(); ++BlendSpaceIdx)
+		{
+			const FPoseSearchDatabaseBlendSpace& DbBlendSpace = Database->BlendSpaces[BlendSpaceIdx];
+
+			int32 NumGroups = 0;
+			for (const FGameplayTag& GroupTag : DbBlendSpace.GroupTags)
+			{
+				const int32 GroupIdx = Database->Groups.IndexOfByPredicate(
+					[GroupTag](const FPoseSearchDatabaseGroup& Group)
+				{
+					return Group.Tag == GroupTag;
+				});
+
+				if (GroupIdx != INDEX_NONE)
+				{
+					CreateAssetNode(BlendSpaceIdx, ESearchIndexAssetType::BlendSpace, GroupIdx);
+					++NumGroups;
+				}
+			}
+
+			if (NumGroups == 0)
+			{
+				CreateAssetNode(BlendSpaceIdx, ESearchIndexAssetType::BlendSpace, DefaultGroupIdx);
+			}
+		}
+
+		TreeView->RequestTreeRefresh();
+
+		for (TSharedPtr<FDatabaseAssetTreeNode>& RootNode : RootNodes)
+		{
+			TreeView->SetItemExpansion(RootNode, true);
+		}
+	}
+
+	TSharedRef<ITableRow> SDatabaseAssetTree::MakeTableRowWidget(
+		TSharedPtr<FDatabaseAssetTreeNode> InItem,
+		const TSharedRef<STableViewBase>& OwnerTable)
+	{
+		return InItem->MakeTreeRowWidget(OwnerTable, InItem.ToSharedRef(), CommandList.ToSharedRef(), SharedThis(this));
+	}
+
+	void SDatabaseAssetTree::HandleGetChildrenForTree(
+		TSharedPtr<FDatabaseAssetTreeNode> InNode,
+		TArray<TSharedPtr<FDatabaseAssetTreeNode>>& OutChildren)
+	{
+		OutChildren = InNode.Get()->Children;
+	}
+
+	TOptional<EItemDropZone> SDatabaseAssetTree::OnCanAcceptDrop(
+		const FDragDropEvent& DragDropEvent,
+		EItemDropZone DropZone,
+		TSharedPtr<FDatabaseAssetTreeNode> TargetItem)
+	{
+		TOptional<EItemDropZone> ReturnedDropZone;
+
+		TSharedPtr<FDragDropOperation> Operation = DragDropEvent.GetOperation();
+
+		const bool bValidOperation = Operation.IsValid() && Operation->IsOfType<FAssetDragDropOp>();
+		if (bValidOperation)
+		{
+			const TSharedPtr<FAssetDragDropOp> AssetDragDropOp = StaticCastSharedPtr<FAssetDragDropOp>(Operation);
+
+			for (const FAssetData& AssetData : AssetDragDropOp->GetAssets())
+			{
+				if (UClass* AssetClass = AssetData.GetClass())
+				{
+					if (AssetClass->IsChildOf(UAnimSequence::StaticClass()) ||
+						AssetClass->IsChildOf(UBlendSpace::StaticClass()))
+					{
+						ReturnedDropZone = EItemDropZone::OntoItem;
+						break;
+					}
+				}
+			}
+		}
+
+		return ReturnedDropZone;
+	}
+
+	FReply SDatabaseAssetTree::OnAcceptDrop(
+		const FDragDropEvent& DragDropEvent,
+		EItemDropZone DropZone,
+		TSharedPtr<FDatabaseAssetTreeNode> TargetItem)
+	{
+		TSharedPtr<FDragDropOperation> Operation = DragDropEvent.GetOperation();
+
+		const bool bValidOperation = Operation.IsValid() && Operation->IsOfType<FAssetDragDropOp>();
+		if (!bValidOperation)
+		{
+			return FReply::Unhandled();
+		}
+
+		TSharedPtr<FPoseSearchDatabaseViewModel> ViewModel = EditorViewModel.Pin();
+		if (!ViewModel.IsValid())
+		{
+			return FReply::Unhandled();
+		}
+
+		TArray<FAssetData> DroppedAssetData = AssetUtil::ExtractAssetDataFromDrag(Operation);
+		const int32 NumAssets = DroppedAssetData.Num();
+
+		int32 AddedAssets = 0;
+		if (NumAssets > 0)
+		{
+			GWarn->BeginSlowTask(LOCTEXT("LoadingAssets", "Loading Asset(s)"), true);
+			for (int32 DroppedAssetIdx = 0; DroppedAssetIdx < NumAssets; ++DroppedAssetIdx)
+			{
+				const FAssetData& AssetData = DroppedAssetData[DroppedAssetIdx];
+
+				if (!AssetData.IsAssetLoaded())
+				{
+					GWarn->StatusUpdate(
+						DroppedAssetIdx,
+						NumAssets,
+						FText::Format(
+							LOCTEXT("LoadingAsset", "Loading Asset {0}"),
+							FText::FromName(AssetData.AssetName)));
+				}
+
+				UClass* AssetClass = AssetData.GetClass();
+				UObject* Asset = AssetData.GetAsset();
+
+				if (AssetClass->IsChildOf(UAnimSequence::StaticClass()))
+				{
+					const int32 GroupIdx = FindGroupIndex(TargetItem);
+					ViewModel->AddSequenceToDatabase(Cast<UAnimSequence>(Asset), GroupIdx);
+					++AddedAssets;
+				}
+				else if (AssetClass->IsChildOf(UBlendSpace::StaticClass()))
+				{
+					const int32 GroupIdx = FindGroupIndex(TargetItem);
+					ViewModel->AddBlendSpaceToDatabase(Cast<UBlendSpace>(Asset), GroupIdx);
+					++AddedAssets;
+				}
+			}
+
+			GWarn->EndSlowTask();
+		}
+
+		if (AddedAssets == 0)
+		{
+			return FReply::Unhandled();
+		}
+
+		RefreshTreeView(false);
+		return FReply::Handled();
+	}
+
+	int32 SDatabaseAssetTree::FindGroupIndex(TSharedPtr<FDatabaseAssetTreeNode> TargetItem)
+	{
+		if (!TargetItem.IsValid())
+		{
+			return INDEX_NONE;
+		}
+
+		if (TargetItem->SourceAssetType == ESearchIndexAssetType::Invalid)
+		{
+			return TargetItem->SourceAssetIdx;
+		}
+
+		check(TargetItem->Parent.IsValid() && TargetItem->Parent->SourceAssetType == ESearchIndexAssetType::Invalid);
+		return TargetItem->Parent->SourceAssetIdx;
+	}
+
+	TSharedRef<SWidget> SDatabaseAssetTree::CreateAddNewMenuWidget()
+	{
+		FMenuBuilder AddOptions(true, nullptr);
+
+		AddOptions.BeginSection("AddOptions", LOCTEXT("AssetAddOptions", "Assets"));
+		AddOptions.AddMenuEntry(
+			LOCTEXT("AddSequence", "Sequence"),
+			LOCTEXT("AddSequenceTooltip", "Add new sequence to the default group"),
+			FSlateIcon(),
+			FUIAction(FExecuteAction::CreateSP(this, &SDatabaseAssetTree::OnAddSequence)),
+			NAME_None,
+			EUserInterfaceActionType::Button);
+
+		AddOptions.AddMenuEntry(
+			LOCTEXT("AddBlendSpace", "Blend Space"),
+			LOCTEXT("AddBlendSpaceTooltip", "Add new blend space to the default group"),
+			FSlateIcon(),
+			FUIAction(FExecuteAction::CreateSP(this, &SDatabaseAssetTree::OnAddBlendSpace)),
+			NAME_None,
+			EUserInterfaceActionType::Button);
+		AddOptions.EndSection();
+
+		AddOptions.BeginSection("AddOptions", LOCTEXT("GroupAddOptions", "Groups"));
+		AddOptions.AddMenuEntry(
+			LOCTEXT("AddBlendSpace", "Group"),
+			LOCTEXT("AddBlendSpaceTooltip", "Add new group"),
+			FSlateIcon(),
+			FUIAction(FExecuteAction::CreateSP(this, &SDatabaseAssetTree::OnAddGroup)),
+			NAME_None,
+			EUserInterfaceActionType::Button);
+		AddOptions.EndSection();
+
+		return AddOptions.MakeWidget();
+	}
+
+	TSharedPtr<SWidget> SDatabaseAssetTree::CreateContextMenu()
+	{
+		const bool CloseAfterSelection = true;
+		FMenuBuilder MenuBuilder(CloseAfterSelection, CommandList);
+
+		const TArray<TSharedPtr<FDatabaseAssetTreeNode>> SelectedNodes = TreeView->GetSelectedItems();
+		if (!SelectedNodes.IsEmpty())
+		{
+			TSharedPtr<FDatabaseAssetTreeNode> SelectedNode = SelectedNodes[0];
+
+			if (SelectedNode->SourceAssetType != ESearchIndexAssetType::Invalid)
+			{
+				const int32 GroupIdx = SelectedNode->Parent->SourceAssetIdx;
+				if (GroupIdx == INDEX_NONE)
+				{
+					MenuBuilder.AddMenuEntry(
+						LOCTEXT("DeleteAsset", "Delete Asset"),
+						LOCTEXT("DeleteAssetTooltip", "Delete asset from database"),
+						FSlateIcon(),
+						FUIAction(FExecuteAction::CreateSP(this, &SDatabaseAssetTree::OnDeleteAsset, SelectedNode)),
+						NAME_None,
+						EUserInterfaceActionType::Button);
+				}
+				else
+				{
+					MenuBuilder.AddMenuEntry(
+						LOCTEXT("RemoveFromGroup", "Remove From Group"),
+						LOCTEXT("RemoveFromGroupTooltip", "Remove asset from group"),
+						FSlateIcon(),
+						FUIAction(FExecuteAction::CreateSP(this, &SDatabaseAssetTree::OnRemoveFromGroup, SelectedNode)),
+						NAME_None,
+						EUserInterfaceActionType::Button);
+				}
+			}
+			else if (SelectedNode->SourceAssetIdx != INDEX_NONE)
+			{
+				// a group that's not default
+				MenuBuilder.AddMenuEntry(
+					LOCTEXT("DeleteGroup", "Delete Group"),
+					LOCTEXT("DeleteGroupTooltip", "Delete group and remove it from the sequences in it"),
+					FSlateIcon(),
+					FUIAction(FExecuteAction::CreateSP(this, &SDatabaseAssetTree::OnDeleteGroup, SelectedNode)),
+					NAME_None,
+					EUserInterfaceActionType::Button);
+			}
+
+		}
+
+		return MenuBuilder.MakeWidget();
+	}
+
+	void SDatabaseAssetTree::OnAddGroup()
+	{
+		EditorViewModel.Pin()->AddGroupToDatabase();
+		RefreshTreeView(false);
+	}
+
+	void SDatabaseAssetTree::OnAddSequence()
+	{
+		EditorViewModel.Pin()->AddSequenceToDatabase(nullptr, INDEX_NONE);
+		RefreshTreeView(false);
+	}
+
+	void SDatabaseAssetTree::OnAddBlendSpace()
+	{
+		EditorViewModel.Pin()->AddBlendSpaceToDatabase(nullptr, INDEX_NONE);
+		RefreshTreeView(false);
+	}
+
+	void SDatabaseAssetTree::OnDeleteAsset(TSharedPtr<FDatabaseAssetTreeNode> Node)
+	{
+		if (Node->SourceAssetType == ESearchIndexAssetType::Sequence)
+		{
+			EditorViewModel.Pin()->DeleteSequenceFromDatabase(Node->SourceAssetIdx);
+		}
+		else if (Node->SourceAssetType == ESearchIndexAssetType::BlendSpace)
+		{
+			EditorViewModel.Pin()->DeleteBlendSpaceFromDatabase(Node->SourceAssetIdx);
+		}
+		else
+		{
+			checkNoEntry();
+		}
+
+		RefreshTreeView(false);
+	}
+
+	void SDatabaseAssetTree::OnRemoveFromGroup(TSharedPtr<FDatabaseAssetTreeNode> Node)
+	{
+		check(Node->Parent->SourceAssetType == ESearchIndexAssetType::Invalid);
+		const int32 GroupIdx = Node->Parent->SourceAssetIdx;
+
+		if (Node->SourceAssetType == ESearchIndexAssetType::Sequence)
+		{
+			EditorViewModel.Pin()->RemoveSequenceFromGroup(Node->SourceAssetIdx, GroupIdx);
+		}
+		else if (Node->SourceAssetType == ESearchIndexAssetType::BlendSpace)
+		{
+			EditorViewModel.Pin()->RemoveBlendSpaceFromGroup(Node->SourceAssetIdx, GroupIdx);
+		}
+		else
+		{
+			checkNoEntry();
+		}
+
+		RefreshTreeView(false);
+	}
+
+	void SDatabaseAssetTree::OnDeleteGroup(TSharedPtr<FDatabaseAssetTreeNode> Node)
+	{
+		check(Node->SourceAssetType == ESearchIndexAssetType::Invalid);
+		const int32 GroupIdx = Node->SourceAssetIdx;
+		EditorViewModel.Pin()->DeleteGroup(GroupIdx);
+
+		RefreshTreeView(false);
+	}
+}
+
+#undef LOCTEXT_NAMESPACE
diff --git a/Engine/Plugins/Experimental/Animation/PoseSearch/Source/Editor/Private/SPoseSearchDatabaseViewportToolbar.cpp b/Engine/Plugins/Experimental/Animation/PoseSearch/Source/Editor/Private/SPoseSearchDatabaseViewportToolbar.cpp
--- a/Engine/Plugins/Experimental/Animation/PoseSearch/Source/Editor/Private/SPoseSearchDatabaseViewportToolbar.cpp
+++ b/Engine/Plugins/Experimental/Animation/PoseSearch/Source/Editor/Private/SPoseSearchDatabaseViewportToolbar.cpp
@@ -5,7 +5,7 @@
 #include "PreviewProfileController.h"
 #include "PoseSearchDatabaseEditorCommands.h"
 
-#define LOCTEXT_NAMESPACE "ContextualAnimViewportToolBar"
+#define LOCTEXT_NAMESPACE "PoseSearchDatabaseViewportToolBar"
 
 void SPoseSearchDatabaseViewportToolBar::Construct(
 	const FArguments& InArgs, 
diff --git a/Engine/Plugins/Experimental/Animation/PoseSearch/Source/Editor/Public/PoseSearchDatabaseEditorToolkit.h b/Engine/Plugins/Experimental/Animation/PoseSearch/Source/Editor/Public/PoseSearchDatabaseEditorToolkit.h
--- a/Engine/Plugins/Experimental/Animation/PoseSearch/Source/Editor/Public/PoseSearchDatabaseEditorToolkit.h
+++ b/Engine/Plugins/Experimental/Animation/PoseSearch/Source/Editor/Public/PoseSearchDatabaseEditorToolkit.h
@@ -13,6 +13,11 @@ class SPoseSearchDatabaseViewport;
 class FPoseSearchDatabasePreviewScene;
 class FPoseSearchDatabaseViewModel;
 
+namespace UE::PoseSearch
+{
+	class SDatabaseAssetTree;
+}
+
 class FPoseSearchDatabaseEditorToolkit : public FAssetEditorToolkit, public FNotifyHook
 {
 public:
@@ -46,6 +51,7 @@ private:
 	TSharedRef<SDockTab> SpawnTab_Viewport(const FSpawnTabArgs& Args);
 	TSharedRef<SDockTab> SpawnTab_AssetDetails(const FSpawnTabArgs& Args);
 	TSharedRef<SDockTab> SpawnTab_PreviewSettings(const FSpawnTabArgs& Args);
+	TSharedRef<SDockTab> SpawnTab_AssetTreeView(const FSpawnTabArgs& Args);
 
 	void BindCommands();
 	void ExtendToolbar();
@@ -55,6 +61,8 @@ private:
 
 	TSharedPtr<SPoseSearchDatabaseViewport> ViewportWidget;
 
+	TSharedPtr<UE::PoseSearch::SDatabaseAssetTree> AssetTreeWidget;
+
 	TSharedPtr<IDetailsView> EditingAssetWidget;
 
 	TSharedPtr<FPoseSearchDatabasePreviewScene> PreviewScene;
diff --git a/Engine/Plugins/Experimental/Animation/PoseSearch/Source/Editor/Public/PoseSearchDatabaseViewModel.h b/Engine/Plugins/Experimental/Animation/PoseSearch/Source/Editor/Public/PoseSearchDatabaseViewModel.h
--- a/Engine/Plugins/Experimental/Animation/PoseSearch/Source/Editor/Public/PoseSearchDatabaseViewModel.h
+++ b/Engine/Plugins/Experimental/Animation/PoseSearch/Source/Editor/Public/PoseSearchDatabaseViewModel.h
@@ -13,6 +13,8 @@ class UPoseSearchDatabase;
 struct FPoseSearchIndexAsset;
 class UAnimPreviewInstance;
 class UDebugSkelMeshComponent;
+class UAnimSequence;
+class UBlendSpace;
 
 enum class EPoseSearchFeaturesDrawMode : uint8
 {
@@ -77,6 +79,17 @@ public:
 	void OnSetAnimationPreviewMode(EAnimationPreviewMode PreviewMode);
 	bool IsAnimationPreviewMode(EAnimationPreviewMode PreviewMode) const;
 
+	void AddSequenceToDatabase(UAnimSequence* AnimSequence, int InitialGroupIdx = -1);
+	void AddBlendSpaceToDatabase(UBlendSpace* BlendSpace, int InitialGroupIdx = -1);
+	void AddGroupToDatabase();
+
+	void DeleteSequenceFromDatabase(int32 SequenceIdx);
+	void RemoveSequenceFromGroup(int32 SequenceIdx, int32 GroupIdx); 
+	void DeleteBlendSpaceFromDatabase(int32 BlendSpaceIdx);
+	void RemoveBlendSpaceFromGroup(int32 BlendSpaceIdx, int32 GroupIdx);
+	void DeleteGroup(int32 GroupIdx);
+
+
 private:
 	float PlayTime = 0.0f;
 
diff --git a/Engine/Plugins/Experimental/Animation/PoseSearch/Source/Editor/Public/SPoseSearchDatabaseAssetList.h b/Engine/Plugins/Experimental/Animation/PoseSearch/Source/Editor/Public/SPoseSearchDatabaseAssetList.h
new file mode 100644
--- /dev/null
+++ b/Engine/Plugins/Experimental/Animation/PoseSearch/Source/Editor/Public/SPoseSearchDatabaseAssetList.h
@@ -0,0 +1,132 @@
+// Copyright Epic Games, Inc. All Rights Reserved.
+
+#pragma once
+
+#include "CoreMinimal.h"
+#include "PoseSearch/PoseSearch.h"
+#include "Widgets/Views/STreeView.h"
+#include "EditorUndoClient.h"
+
+class FPoseSearchDatabaseViewModel;
+
+namespace UE::PoseSearch
+{
+	class SDatabaseAssetTree;
+
+	class FDatabaseAssetTreeNode : public TSharedFromThis<FDatabaseAssetTreeNode>
+	{
+
+	public:
+		FDatabaseAssetTreeNode(
+			int32 InSourceAssetIdx,
+			ESearchIndexAssetType InSourceAssetType,
+			const TSharedRef<FPoseSearchDatabaseViewModel>& InEditorViewModel);
+
+		TSharedRef<ITableRow> MakeTreeRowWidget(
+			const TSharedRef<STableViewBase>& InOwnerTable,
+			TSharedRef<FDatabaseAssetTreeNode> InDatabaseAssetNode,
+			TSharedRef<FUICommandList> InCommandList,
+			TSharedPtr<SDatabaseAssetTree> InHierarchy);
+
+		int32 SourceAssetIdx;
+		ESearchIndexAssetType SourceAssetType;
+		TSharedPtr<FDatabaseAssetTreeNode> Parent;
+		TArray<TSharedPtr<FDatabaseAssetTreeNode>> Children;
+
+	protected:
+
+		TWeakPtr<FPoseSearchDatabaseViewModel> EditorViewModel;
+
+	};
+
+	class SDatabaseAssetListItem : public STableRow<TSharedPtr<FDatabaseAssetTreeNode>>
+	{
+	public:
+		void Construct(
+			const FArguments& InArgs,
+			const TSharedRef<FPoseSearchDatabaseViewModel>& InEditorViewModel,
+			const TSharedRef<STableViewBase>& OwnerTable,
+			TSharedRef<FDatabaseAssetTreeNode> InAssetTreeNode,
+			TSharedRef<FUICommandList> InCommandList,
+			TSharedPtr<SDatabaseAssetTree> InHierarchy);
+
+
+	protected:
+		FText GetName() const;
+		TSharedRef<SWidget> GenerateItemWidget();
+		TSharedRef<SWidget> GenerateAddButtonWidget();
+
+		const FSlateBrush* GetGroupBackgroundImage() const;
+		void ConstructGroupItem(const TSharedRef<STableViewBase>& OwnerTable);
+		void ConstructAssetItem(const TSharedRef<STableViewBase>& OwnerTable);
+
+		void OnAddSequence();
+		void OnAddBlendSpace();
+
+		TWeakPtr<FDatabaseAssetTreeNode> WeakAssetTreeNode;
+		TWeakPtr<FPoseSearchDatabaseViewModel> EditorViewModel;
+		TWeakPtr<SDatabaseAssetTree> SkeletonView;
+	};
+
+	class SDatabaseAssetTree : public SCompoundWidget, public FEditorUndoClient
+	{
+	public:
+		SLATE_BEGIN_ARGS(SDatabaseAssetTree) {}
+		SLATE_END_ARGS()
+
+		virtual ~SDatabaseAssetTree();
+
+		void Construct(const FArguments& InArgs, TSharedRef<FPoseSearchDatabaseViewModel> InEditorViewModel);
+
+		// SWidget interface
+		virtual FReply OnDragOver(const FGeometry& MyGeometry, const FDragDropEvent& DragDropEvent) override;
+		virtual FReply OnDrop(const FGeometry& MyGeometry, const FDragDropEvent& DragDropEvent) override;
+		// End SWidget interface
+	
+		void RefreshTreeView(bool IsInitialSetup = false);
+
+	protected:
+		TWeakPtr<FPoseSearchDatabaseViewModel> EditorViewModel;
+
+		/** command list we bind to */
+		TSharedPtr<FUICommandList> CommandList;
+
+		/** tree view widget */
+		TSharedPtr<STreeView<TSharedPtr<FDatabaseAssetTreeNode>>> TreeView;
+		TArray<TSharedPtr<FDatabaseAssetTreeNode>> RootNodes;
+		TArray<TSharedPtr<FDatabaseAssetTreeNode>> AllNodes;
+
+		TSharedRef<ITableRow> MakeTableRowWidget(
+			TSharedPtr<FDatabaseAssetTreeNode> InItem,
+			const TSharedRef<STableViewBase>& OwnerTable);
+		void HandleGetChildrenForTree(
+			TSharedPtr<FDatabaseAssetTreeNode> InNode, 
+			TArray<TSharedPtr<FDatabaseAssetTreeNode>>& OutChildren);
+
+		TOptional<EItemDropZone> OnCanAcceptDrop(
+			const FDragDropEvent& DragDropEvent, 
+			EItemDropZone DropZone, 
+			TSharedPtr<FDatabaseAssetTreeNode> TargetItem);
+
+		FReply OnAcceptDrop(
+			const FDragDropEvent& DragDropEvent,
+			EItemDropZone DropZone,
+			TSharedPtr<FDatabaseAssetTreeNode> TargetItem);
+
+		int32 FindGroupIndex(TSharedPtr<FDatabaseAssetTreeNode> TargetItem);
+
+		TSharedRef<SWidget> CreateAddNewMenuWidget();
+		TSharedPtr<SWidget> CreateContextMenu();
+		
+		void OnAddGroup();
+		void OnAddSequence();
+		void OnAddBlendSpace();
+
+		void OnDeleteAsset(TSharedPtr<FDatabaseAssetTreeNode> Node);
+		void OnRemoveFromGroup(TSharedPtr<FDatabaseAssetTreeNode> Node);
+		void OnDeleteGroup(TSharedPtr<FDatabaseAssetTreeNode> Node);
+
+		friend SDatabaseAssetListItem;
+	};
+}
+
diff --git a/Engine/Plugins/Experimental/Animation/PoseSearch/Source/Runtime/Private/PoseSearch.cpp b/Engine/Plugins/Experimental/Animation/PoseSearch/Source/Runtime/Private/PoseSearch.cpp
--- a/Engine/Plugins/Experimental/Animation/PoseSearch/Source/Runtime/Private/PoseSearch.cpp
+++ b/Engine/Plugins/Experimental/Animation/PoseSearch/Source/Runtime/Private/PoseSearch.cpp
@@ -1675,6 +1675,36 @@ void UPoseSearchDatabase::NotifyDerivedDataBuildStarted()
 	OnDerivedDataRebuild.Broadcast();
 }
 
+void UPoseSearchDatabase::RegisterOnAssetChange(const FOnAssetChange& Delegate)
+{
+	OnAssetChange.Add(Delegate);
+}
+
+void UPoseSearchDatabase::UnregisterOnAssetChange(void* Unregister)
+{
+	OnAssetChange.RemoveAll(Unregister);
+}
+
+void UPoseSearchDatabase::NotifyAssetChange()
+{
+	OnAssetChange.Broadcast();
+}
+
+void UPoseSearchDatabase::RegisterOnGroupChange(const FOnGroupChange& Delegate)
+{
+	OnGroupChange.Add(Delegate);
+}
+
+void UPoseSearchDatabase::UnregisterOnGroupChange(void* Unregister)
+{
+	OnGroupChange.RemoveAll(Unregister);
+}
+
+void UPoseSearchDatabase::NotifyGroupChange()
+{
+	OnGroupChange.Broadcast();
+}
+
 void UPoseSearchDatabase::BeginCacheDerivedData()
 {
 	bool bPerformCache = true;
@@ -1751,11 +1781,14 @@ void UPoseSearchDatabase::PostEditChangeProperty(struct FPropertyChangedEvent& P
 {
 	Super::PostEditChangeProperty(PropertyChangedEvent);
 
+	bool bAssetChange = false;
+
 	if (PropertyChangedEvent.Property->GetFName() == GET_MEMBER_NAME_CHECKED(UPoseSearchDatabase, SimpleSequences))
 	{
 		if (!SimpleSequences.IsEmpty())
 		{
 			CollectSimpleSequences();
+			bAssetChange = true;
 		}
 	}
 
@@ -1764,9 +1797,26 @@ void UPoseSearchDatabase::PostEditChangeProperty(struct FPropertyChangedEvent& P
 		if (!SimpleBlendSpaces.IsEmpty())
 		{
 			CollectSimpleBlendSpaces();
+			bAssetChange = true;
 		}
 	}
 
+	if (PropertyChangedEvent.MemberProperty->GetFName() == GET_MEMBER_NAME_CHECKED(UPoseSearchDatabase, Sequences) ||
+		PropertyChangedEvent.MemberProperty->GetFName() == GET_MEMBER_NAME_CHECKED(UPoseSearchDatabase, BlendSpaces))
+	{
+		bAssetChange = true;
+	}
+
+	if (bAssetChange)
+	{
+		NotifyAssetChange();
+	}
+
+	if (PropertyChangedEvent.MemberProperty->GetFName() == GET_MEMBER_NAME_CHECKED(UPoseSearchDatabase, Groups))
+	{
+		NotifyGroupChange();
+	}
+
 	BeginCacheDerivedData();
 }
 
diff --git a/Engine/Plugins/Experimental/Animation/PoseSearch/Source/Runtime/Public/PoseSearch/PoseSearch.h b/Engine/Plugins/Experimental/Animation/PoseSearch/Source/Runtime/Public/PoseSearch/PoseSearch.h
--- a/Engine/Plugins/Experimental/Animation/PoseSearch/Source/Runtime/Public/PoseSearch/PoseSearch.h
+++ b/Engine/Plugins/Experimental/Animation/PoseSearch/Source/Runtime/Public/PoseSearch/PoseSearch.h
@@ -828,18 +828,32 @@ private:
 #if WITH_EDITOR
 	DECLARE_MULTICAST_DELEGATE(FOnDerivedDataRebuildMulticaster);
 	FOnDerivedDataRebuildMulticaster OnDerivedDataRebuild;
+
+	DECLARE_MULTICAST_DELEGATE(FOnAssetChangeMulticaster);
+	FOnDerivedDataRebuildMulticaster OnAssetChange;
+
+	DECLARE_MULTICAST_DELEGATE(FOnGroupChangeMulticaster);
+	FOnDerivedDataRebuildMulticaster OnGroupChange;
 #endif // WITH_EDITOR
 
 public:
 #if WITH_EDITOR
-	typedef FOnDerivedDataRebuildMulticaster::FDelegate FOnDerivedDataRebuild;
 
-	/** Registers a delegate to be called when derived data starts rebuilding */
+	typedef FOnDerivedDataRebuildMulticaster::FDelegate FOnDerivedDataRebuild;
 	void RegisterOnDerivedDataRebuild(const FOnDerivedDataRebuild& Delegate);
 	void UnregisterOnDerivedDataRebuild(void* Unregister);
-
 	void NotifyDerivedDataBuildStarted();
 
+	typedef FOnAssetChangeMulticaster::FDelegate FOnAssetChange;
+	void RegisterOnAssetChange(const FOnAssetChange& Delegate);
+	void UnregisterOnAssetChange(void* Unregister);
+	void NotifyAssetChange();
+
+	typedef FOnGroupChangeMulticaster::FDelegate FOnGroupChange;
+	void RegisterOnGroupChange(const FOnGroupChange& Delegate);
+	void UnregisterOnGroupChange(void* Unregister);
+	void NotifyGroupChange();
+
 	void BeginCacheDerivedData();
 #endif // WITH_EDITOR
 
