PoseSearch - moved editor classes into UE::PoseSearch namespace

diff --git a/Engine/Plugins/Experimental/Animation/PoseSearch/Source/Editor/Private/PoseSearchDatabaseEdMode.cpp b/Engine/Plugins/Experimental/Animation/PoseSearch/Source/Editor/Private/PoseSearchDatabaseEdMode.cpp
--- a/Engine/Plugins/Experimental/Animation/PoseSearch/Source/Editor/Private/PoseSearchDatabaseEdMode.cpp
+++ b/Engine/Plugins/Experimental/Animation/PoseSearch/Source/Editor/Private/PoseSearchDatabaseEdMode.cpp
@@ -8,88 +8,100 @@
 
 #include "EngineUtils.h"
 
-const FEditorModeID FPoseSearchDatabaseEdMode::EdModeId = TEXT("PoseSearchDatabaseEdMode");
-
-FPoseSearchDatabaseEdMode::FPoseSearchDatabaseEdMode()
+namespace UE::PoseSearch
 {
-}
+	const FEditorModeID FDatabaseEdMode::EdModeId = TEXT("PoseSearchDatabaseEdMode");
 
-FPoseSearchDatabaseEdMode::~FPoseSearchDatabaseEdMode()
-{
-}
-
-void FPoseSearchDatabaseEdMode::Tick(FEditorViewportClient* ViewportClient, float DeltaTime)
-{
-	FEdMode::Tick(ViewportClient, DeltaTime);
+	FDatabaseEdMode::FDatabaseEdMode()
+	{
+	}
 
-	FPoseSearchDatabaseViewportClient* PoseSearchDbViewportClient = 
-		static_cast<FPoseSearchDatabaseViewportClient*>(ViewportClient);
-	if (!ViewModel && PoseSearchDbViewportClient)
+	FDatabaseEdMode::~FDatabaseEdMode()
 	{
-		ViewModel = PoseSearchDbViewportClient->GetAssetEditorToolkit()->GetViewModel();
 	}
 
-	if (ViewModel)
+	void FDatabaseEdMode::Tick(FEditorViewportClient* ViewportClient, float DeltaTime)
 	{
-		ViewModel->Tick(DeltaTime);
+		FEdMode::Tick(ViewportClient, DeltaTime);
+
+		FDatabaseViewportClient* PoseSearchDbViewportClient =
+			static_cast<FDatabaseViewportClient*>(ViewportClient);
+		if (!ViewModel && PoseSearchDbViewportClient)
+		{
+			ViewModel = PoseSearchDbViewportClient->GetAssetEditorToolkit()->GetViewModel();
+		}
+
+		if (ViewModel)
+		{
+			ViewModel->Tick(DeltaTime);
+		}
 	}
-}
 
-void FPoseSearchDatabaseEdMode::Render(const FSceneView* View, FViewport* Viewport, FPrimitiveDrawInterface* PDI)
-{
-	FEdMode::Render(View, Viewport, PDI);
-}
+	void FDatabaseEdMode::Render(const FSceneView* View, FViewport* Viewport, FPrimitiveDrawInterface* PDI)
+	{
+		FEdMode::Render(View, Viewport, PDI);
+	}
 
-bool FPoseSearchDatabaseEdMode::HandleClick(
-	FEditorViewportClient* InViewportClient, 
-	HHitProxy* HitProxy, 
-	const FViewportClick& Click)
-{
-	return FEdMode::HandleClick(InViewportClient, HitProxy, Click);
-}
+	bool FDatabaseEdMode::HandleClick(
+		FEditorViewportClient* InViewportClient,
+		HHitProxy* HitProxy,
+		const FViewportClick& Click)
+	{
+		return FEdMode::HandleClick(InViewportClient, HitProxy, Click);
+	}
 
 
-bool FPoseSearchDatabaseEdMode::StartTracking(FEditorViewportClient* InViewportClient, FViewport* InViewport)
-{
-	return FEdMode::StartTracking(InViewportClient, InViewport);
-}
+	bool FDatabaseEdMode::StartTracking(FEditorViewportClient* InViewportClient, FViewport* InViewport)
+	{
+		return FEdMode::StartTracking(InViewportClient, InViewport);
+	}
 
-bool FPoseSearchDatabaseEdMode::EndTracking(FEditorViewportClient* InViewportClient, FViewport* InViewport)
-{
-	return FEdMode::EndTracking(InViewportClient, InViewport);
-}
+	bool FDatabaseEdMode::EndTracking(FEditorViewportClient* InViewportClient, FViewport* InViewport)
+	{
+		return FEdMode::EndTracking(InViewportClient, InViewport);
+	}
 
-bool FPoseSearchDatabaseEdMode::InputDelta(FEditorViewportClient* InViewportClient, FViewport* InViewport, FVector& InDrag, FRotator& InRot, FVector& InScale)
-{
-	return FEdMode::InputDelta(InViewportClient, InViewport, InDrag, InRot, InScale);
-}
+	bool FDatabaseEdMode::InputDelta(
+		FEditorViewportClient* InViewportClient, 
+		FViewport* InViewport, 
+		FVector& InDrag, 
+		FRotator& InRot, 
+		FVector& InScale)
+	{
+		return FEdMode::InputDelta(InViewportClient, InViewport, InDrag, InRot, InScale);
+	}
 
-bool FPoseSearchDatabaseEdMode::InputKey(FEditorViewportClient* ViewportClient, FViewport* Viewport, FKey Key, EInputEvent Event)
-{
-	return FEdMode::InputKey(ViewportClient, Viewport, Key, Event);
-}
+	bool FDatabaseEdMode::InputKey(
+		FEditorViewportClient* ViewportClient, 
+		FViewport* Viewport, 
+		FKey Key, 
+		EInputEvent Event)
+	{
+		return FEdMode::InputKey(ViewportClient, Viewport, Key, Event);
+	}
 
-bool FPoseSearchDatabaseEdMode::AllowWidgetMove()
-{
-	return FEdMode::ShouldDrawWidget();
-}
+	bool FDatabaseEdMode::AllowWidgetMove()
+	{
+		return FEdMode::ShouldDrawWidget();
+	}
 
-bool FPoseSearchDatabaseEdMode::ShouldDrawWidget() const
-{
-	return FEdMode::ShouldDrawWidget();
-}
+	bool FDatabaseEdMode::ShouldDrawWidget() const
+	{
+		return FEdMode::ShouldDrawWidget();
+	}
 
-bool FPoseSearchDatabaseEdMode::GetCustomDrawingCoordinateSystem(FMatrix& InMatrix, void* InData)
-{
-	return FEdMode::GetCustomDrawingCoordinateSystem(InMatrix, InData);
-}
+	bool FDatabaseEdMode::GetCustomDrawingCoordinateSystem(FMatrix& InMatrix, void* InData)
+	{
+		return FEdMode::GetCustomDrawingCoordinateSystem(InMatrix, InData);
+	}
 
-bool FPoseSearchDatabaseEdMode::GetCustomInputCoordinateSystem(FMatrix& InMatrix, void* InData)
-{
-	return FEdMode::GetCustomDrawingCoordinateSystem(InMatrix, InData);
-}
+	bool FDatabaseEdMode::GetCustomInputCoordinateSystem(FMatrix& InMatrix, void* InData)
+	{
+		return FEdMode::GetCustomDrawingCoordinateSystem(InMatrix, InData);
+	}
 
-FVector FPoseSearchDatabaseEdMode::GetWidgetLocation() const
-{
-	return FEdMode::GetWidgetLocation();
+	FVector FDatabaseEdMode::GetWidgetLocation() const
+	{
+		return FEdMode::GetWidgetLocation();
+	}
 }
diff --git a/Engine/Plugins/Experimental/Animation/PoseSearch/Source/Editor/Private/PoseSearchDatabaseEditorCommands.cpp b/Engine/Plugins/Experimental/Animation/PoseSearch/Source/Editor/Private/PoseSearchDatabaseEditorCommands.cpp
--- a/Engine/Plugins/Experimental/Animation/PoseSearch/Source/Editor/Private/PoseSearchDatabaseEditorCommands.cpp
+++ b/Engine/Plugins/Experimental/Animation/PoseSearch/Source/Editor/Private/PoseSearchDatabaseEditorCommands.cpp
@@ -4,26 +4,29 @@
 
 #define LOCTEXT_NAMESPACE "PoseSearchDatabaseEditorCommands"
 
-void FPoseSearchDatabaseEditorCommands::RegisterCommands()
+namespace UE::PoseSearch
 {
-	UI_COMMAND(StopPreviewScene, "Stop Scene", "Stop Scene", EUserInterfaceActionType::Button, FInputChord());
-	UI_COMMAND(ResetPreviewScene, "Reset Scene", "Reset Scene", EUserInterfaceActionType::Button, FInputChord());
-	UI_COMMAND(BuildSearchIndex, "Build Index", "Build Index", EUserInterfaceActionType::Button, FInputChord());
+	void FDatabaseEditorCommands::RegisterCommands()
+	{
+		UI_COMMAND(StopPreviewScene, "Stop Scene", "Stop Scene", EUserInterfaceActionType::Button, FInputChord());
+		UI_COMMAND(ResetPreviewScene, "Reset Scene", "Reset Scene", EUserInterfaceActionType::Button, FInputChord());
+		UI_COMMAND(BuildSearchIndex, "Build Index", "Build Index", EUserInterfaceActionType::Button, FInputChord());
 
-	UI_COMMAND(
-		ShowPoseFeaturesNone, "None", "Hide all features", 
-		EUserInterfaceActionType::RadioButton, FInputChord());
-	UI_COMMAND(
-		ShowPoseFeaturesAll, "All Features", "Show all features", 
-		EUserInterfaceActionType::RadioButton, FInputChord());
+		UI_COMMAND(
+			ShowPoseFeaturesNone, "None", "Hide all features",
+			EUserInterfaceActionType::RadioButton, FInputChord());
+		UI_COMMAND(
+			ShowPoseFeaturesAll, "All Features", "Show all features",
+			EUserInterfaceActionType::RadioButton, FInputChord());
 
-	UI_COMMAND(
-		ShowAnimationOriginalOnly, "Original Only", "Show only original animations",
-		EUserInterfaceActionType::RadioButton, FInputChord());
-	UI_COMMAND(
-		ShowAnimationOriginalAndMirrored, "Original and mirrored", "Show original and mirrored animations",
-		EUserInterfaceActionType::RadioButton, FInputChord());
+		UI_COMMAND(
+			ShowAnimationOriginalOnly, "Original Only", "Show only original animations",
+			EUserInterfaceActionType::RadioButton, FInputChord());
+		UI_COMMAND(
+			ShowAnimationOriginalAndMirrored, "Original and mirrored", "Show original and mirrored animations",
+			EUserInterfaceActionType::RadioButton, FInputChord());
 
+	}
 }
 
 #undef LOCTEXT_NAMESPACE
diff --git a/Engine/Plugins/Experimental/Animation/PoseSearch/Source/Editor/Private/PoseSearchDatabaseEditorToolkit.cpp b/Engine/Plugins/Experimental/Animation/PoseSearch/Source/Editor/Private/PoseSearchDatabaseEditorToolkit.cpp
--- a/Engine/Plugins/Experimental/Animation/PoseSearch/Source/Editor/Private/PoseSearchDatabaseEditorToolkit.cpp
+++ b/Engine/Plugins/Experimental/Animation/PoseSearch/Source/Editor/Private/PoseSearchDatabaseEditorToolkit.cpp
@@ -18,374 +18,377 @@
 
 #define LOCTEXT_NAMESPACE "PoseSearchDatabaseEditorToolkit"
 
-const FName PoseSearchDatabaseEditorAppName = FName(TEXT("PoseSearchDatabaseEditorApp"));
-
-// Tab identifiers
-struct FPoseSearchDatabaseEditorTabs
-{
-	static const FName AssetDetailsID;
-	static const FName ViewportID;
-	static const FName PreviewSettingsID;
-	static const FName AssetTreeViewID;
-};
-
-const FName FPoseSearchDatabaseEditorTabs::AssetDetailsID(TEXT("PoseSearchDatabaseEditorAssetDetailsTabID"));
-const FName FPoseSearchDatabaseEditorTabs::ViewportID(TEXT("PoseSearchDatabaseEditorViewportTabID"));
-const FName FPoseSearchDatabaseEditorTabs::PreviewSettingsID(TEXT("PoseSearchDatabaseEditorPreviewSettingsTabID"));
-const FName FPoseSearchDatabaseEditorTabs::AssetTreeViewID(TEXT("PoseSearchDatabaseEditorAssetTreeViewTabID"));
-
-FPoseSearchDatabaseEditorToolkit::FPoseSearchDatabaseEditorToolkit()
+namespace UE::PoseSearch
 {
-}
+	const FName PoseSearchDatabaseEditorAppName = FName(TEXT("PoseSearchDatabaseEditorApp"));
 
-FPoseSearchDatabaseEditorToolkit::~FPoseSearchDatabaseEditorToolkit()
-{
-	UPoseSearchDatabase* DatabaseAsset = ViewModel->GetPoseSearchDatabase();
-	if (IsValid(DatabaseAsset))
+	// Tab identifiers
+	struct FDatabaseEditorTabs
+	{
+		static const FName AssetDetailsID;
+		static const FName ViewportID;
+		static const FName PreviewSettingsID;
+		static const FName AssetTreeViewID;
+	};
+
+	const FName FDatabaseEditorTabs::AssetDetailsID(TEXT("PoseSearchDatabaseEditorAssetDetailsTabID"));
+	const FName FDatabaseEditorTabs::ViewportID(TEXT("PoseSearchDatabaseEditorViewportTabID"));
+	const FName FDatabaseEditorTabs::PreviewSettingsID(TEXT("PoseSearchDatabaseEditorPreviewSettingsTabID"));
+	const FName FDatabaseEditorTabs::AssetTreeViewID(TEXT("PoseSearchDatabaseEditorAssetTreeViewTabID"));
+
+	FDatabaseEditorToolkit::FDatabaseEditorToolkit()
 	{
-		DatabaseAsset->UnregisterOnAssetChange(AssetTreeWidget.Get());
-		DatabaseAsset->UnregisterOnGroupChange(AssetTreeWidget.Get());
 	}
-}
-
-const UPoseSearchDatabase* FPoseSearchDatabaseEditorToolkit::GetPoseSearchDatabase() const
-{
-	return ViewModel.IsValid() ? ViewModel->GetPoseSearchDatabase() : nullptr;
-}
 
-void FPoseSearchDatabaseEditorToolkit::StopPreviewScene()
-{
-	ViewModel->RemovePreviewActors();
-}
-
-void FPoseSearchDatabaseEditorToolkit::ResetPreviewScene()
-{
-	ViewModel->ResetPreviewActors();
-}
-
-void FPoseSearchDatabaseEditorToolkit::BuildSearchIndex()
-{
-	ViewModel->BuildSearchIndex();
-}
-
-void FPoseSearchDatabaseEditorToolkit::InitAssetEditor(
-	const EToolkitMode::Type Mode,
-	const TSharedPtr<IToolkitHost>& InitToolkitHost,
-	UPoseSearchDatabase* DatabaseAsset)
-{
-	// Bind Commands
-	BindCommands();
-
-	// Create Preview Scene
-	if (!PreviewScene.IsValid())
+	FDatabaseEditorToolkit::~FDatabaseEditorToolkit()
 	{
-		PreviewScene = MakeShareable(
-			new FPoseSearchDatabasePreviewScene(
-				FPreviewScene::ConstructionValues()
-				.AllowAudioPlayback(true)
-				.ShouldSimulatePhysics(true)
-				.ForceUseMovementComponentInNonGameWorld(true),
-				StaticCastSharedRef<FPoseSearchDatabaseEditorToolkit>(AsShared())));
-
-		//Temporary fix for missing attached assets - MDW (Copied from FPersonaToolkit::CreatePreviewScene)
-		PreviewScene->GetWorld()->GetWorldSettings()->SetIsTemporarilyHiddenInEditor(false);
+		UPoseSearchDatabase* DatabaseAsset = ViewModel->GetPoseSearchDatabase();
+		if (IsValid(DatabaseAsset))
+		{
+			DatabaseAsset->UnregisterOnAssetChange(AssetTreeWidget.Get());
+			DatabaseAsset->UnregisterOnGroupChange(AssetTreeWidget.Get());
+		}
 	}
 
-	// Create view model
-	ViewModel = MakeShared<FPoseSearchDatabaseViewModel>();
-	ViewModel->Initialize(DatabaseAsset, PreviewScene.ToSharedRef());
-
-	// Create viewport widget
-	FPoseSearchDatabaseViewportRequiredArgs ViewportArgs(
-		StaticCastSharedRef<FPoseSearchDatabaseEditorToolkit>(AsShared()), 
-		PreviewScene.ToSharedRef());
-	ViewportWidget = SNew(SPoseSearchDatabaseViewport, ViewportArgs);
-
-	AssetTreeWidget = SNew(UE::PoseSearch::SDatabaseAssetTree, ViewModel.ToSharedRef());
-	if (IsValid(DatabaseAsset))
+	const UPoseSearchDatabase* FDatabaseEditorToolkit::GetPoseSearchDatabase() const
 	{
-		DatabaseAsset->RegisterOnAssetChange(
-			UPoseSearchDatabase::FOnDerivedDataRebuild::CreateSP(
-				AssetTreeWidget.Get(),
-				&UE::PoseSearch::SDatabaseAssetTree::RefreshTreeView, false));
-		DatabaseAsset->RegisterOnGroupChange(
-			UPoseSearchDatabase::FOnDerivedDataRebuild::CreateSP(
-				AssetTreeWidget.Get(),
-				&UE::PoseSearch::SDatabaseAssetTree::RefreshTreeView, false));
+		return ViewModel.IsValid() ? ViewModel->GetPoseSearchDatabase() : nullptr;
 	}
 
-	// Create Asset Details widget
-	FPropertyEditorModule& PropertyModule = FModuleManager::LoadModuleChecked<FPropertyEditorModule>("PropertyEditor");
+	void FDatabaseEditorToolkit::StopPreviewScene()
+	{
+		ViewModel->RemovePreviewActors();
+	}
 
-	FDetailsViewArgs Args;
-	Args.bHideSelectionTip = true;
-	Args.NotifyHook = this;
+	void FDatabaseEditorToolkit::ResetPreviewScene()
+	{
+		ViewModel->ResetPreviewActors();
+	}
 
-	EditingAssetWidget = PropertyModule.CreateDetailView(Args);
-	EditingAssetWidget->SetObject(DatabaseAsset);
-	EditingAssetWidget->OnFinishedChangingProperties().AddSP(
-		this, 
-		&FPoseSearchDatabaseEditorToolkit::OnFinishedChangingProperties);
+	void FDatabaseEditorToolkit::BuildSearchIndex()
+	{
+		ViewModel->BuildSearchIndex();
+	}
 
-	// Define Editor Layout
-	const TSharedRef<FTabManager::FLayout> StandaloneDefaultLayout = 
+	void FDatabaseEditorToolkit::InitAssetEditor(
+		const EToolkitMode::Type Mode,
+		const TSharedPtr<IToolkitHost>& InitToolkitHost,
+		UPoseSearchDatabase* DatabaseAsset)
+	{
+		// Bind Commands
+		BindCommands();
+
+		// Create Preview Scene
+		if (!PreviewScene.IsValid())
+		{
+			PreviewScene = MakeShareable(
+				new FDatabasePreviewScene(
+					FPreviewScene::ConstructionValues()
+					.AllowAudioPlayback(true)
+					.ShouldSimulatePhysics(true)
+					.ForceUseMovementComponentInNonGameWorld(true),
+					StaticCastSharedRef<FDatabaseEditorToolkit>(AsShared())));
+
+			//Temporary fix for missing attached assets - MDW (Copied from FPersonaToolkit::CreatePreviewScene)
+			PreviewScene->GetWorld()->GetWorldSettings()->SetIsTemporarilyHiddenInEditor(false);
+		}
+
+		// Create view model
+		ViewModel = MakeShared<FDatabaseViewModel>();
+		ViewModel->Initialize(DatabaseAsset, PreviewScene.ToSharedRef());
+
+		// Create viewport widget
+		FDatabaseViewportRequiredArgs ViewportArgs(
+			StaticCastSharedRef<FDatabaseEditorToolkit>(AsShared()),
+			PreviewScene.ToSharedRef());
+		ViewportWidget = SNew(SDatabaseViewport, ViewportArgs);
+
+		AssetTreeWidget = SNew(UE::PoseSearch::SDatabaseAssetTree, ViewModel.ToSharedRef());
+		if (IsValid(DatabaseAsset))
+		{
+			DatabaseAsset->RegisterOnAssetChange(
+				UPoseSearchDatabase::FOnDerivedDataRebuild::CreateSP(
+					AssetTreeWidget.Get(),
+					&UE::PoseSearch::SDatabaseAssetTree::RefreshTreeView, false));
+			DatabaseAsset->RegisterOnGroupChange(
+				UPoseSearchDatabase::FOnDerivedDataRebuild::CreateSP(
+					AssetTreeWidget.Get(),
+					&UE::PoseSearch::SDatabaseAssetTree::RefreshTreeView, false));
+		}
+
+		// Create Asset Details widget
+		FPropertyEditorModule& PropertyModule = FModuleManager::LoadModuleChecked<FPropertyEditorModule>("PropertyEditor");
+
+		FDetailsViewArgs Args;
+		Args.bHideSelectionTip = true;
+		Args.NotifyHook = this;
+
+		EditingAssetWidget = PropertyModule.CreateDetailView(Args);
+		EditingAssetWidget->SetObject(DatabaseAsset);
+		EditingAssetWidget->OnFinishedChangingProperties().AddSP(
+			this,
+			&FDatabaseEditorToolkit::OnFinishedChangingProperties);
+
+		// Define Editor Layout
+		const TSharedRef<FTabManager::FLayout> StandaloneDefaultLayout =
 		FTabManager::NewLayout("Standalone_PoseSearchDatabaseEditor_Layout_v0.04")
-		->AddArea
-		(
-			FTabManager::NewPrimaryArea()->SetOrientation(Orient_Vertical)
-			->Split
+			->AddArea
 			(
-				FTabManager::NewSplitter()->SetOrientation(Orient_Vertical)->SetSizeCoefficient(0.9f)
+				FTabManager::NewPrimaryArea()->SetOrientation(Orient_Vertical)
 				->Split
 				(
-					FTabManager::NewSplitter()->SetOrientation(Orient_Horizontal)->SetSizeCoefficient(0.9f)
+					FTabManager::NewSplitter()->SetOrientation(Orient_Vertical)->SetSizeCoefficient(0.9f)
 					->Split
 					(
-						FTabManager::NewStack()
-						->SetSizeCoefficient(0.25f)
-						->AddTab(FPoseSearchDatabaseEditorTabs::AssetTreeViewID, ETabState::OpenedTab)
-						->SetHideTabWell(true)
-					)
-					->Split
-					(
-						FTabManager::NewStack()
-						->SetSizeCoefficient(0.4f)
-						->AddTab(FPoseSearchDatabaseEditorTabs::ViewportID, ETabState::OpenedTab)
-						->SetHideTabWell(true)
-					)
-					->Split
-					(
-						FTabManager::NewSplitter()->SetOrientation(Orient_Vertical)
+						FTabManager::NewSplitter()->SetOrientation(Orient_Horizontal)->SetSizeCoefficient(0.9f)
+						->Split
+						(
+							FTabManager::NewStack()
+							->SetSizeCoefficient(0.25f)
+							->AddTab(FDatabaseEditorTabs::AssetTreeViewID, ETabState::OpenedTab)
+							->SetHideTabWell(true)
+						)
 						->Split
 						(
 							FTabManager::NewStack()
-							->SetSizeCoefficient(0.3f)
-							->AddTab(FPoseSearchDatabaseEditorTabs::AssetDetailsID, ETabState::OpenedTab)
-							->AddTab(FPoseSearchDatabaseEditorTabs::PreviewSettingsID, ETabState::OpenedTab)
+							->SetSizeCoefficient(0.4f)
+							->AddTab(FDatabaseEditorTabs::ViewportID, ETabState::OpenedTab)
+							->SetHideTabWell(true)
+						)
+						->Split
+						(
+							FTabManager::NewSplitter()->SetOrientation(Orient_Vertical)
+							->Split
+							(
+								FTabManager::NewStack()
+								->SetSizeCoefficient(0.3f)
+								->AddTab(FDatabaseEditorTabs::AssetDetailsID, ETabState::OpenedTab)
+								->AddTab(FDatabaseEditorTabs::PreviewSettingsID, ETabState::OpenedTab)
+							)
 						)
 					)
 				)
-			)
-		);
-
-	const bool bCreateDefaultStandaloneMenuParam = true;
-	const bool bCreateDefaultToolbarParam = true;
-	const bool bIsToolbarFocusableParam = false;
-	FAssetEditorToolkit::InitAssetEditor(
-		Mode, 
-		InitToolkitHost, 
-		PoseSearchDatabaseEditorAppName, 
-		StandaloneDefaultLayout, 
-		bCreateDefaultStandaloneMenuParam, 
-		bCreateDefaultToolbarParam, 
-		DatabaseAsset, 
-		bIsToolbarFocusableParam);
-
-	ExtendToolbar();
-
-	RegenerateMenusAndToolbars();
-}
+			);
+
+		const bool bCreateDefaultStandaloneMenuParam = true;
+		const bool bCreateDefaultToolbarParam = true;
+		const bool bIsToolbarFocusableParam = false;
+		FAssetEditorToolkit::InitAssetEditor(
+			Mode,
+			InitToolkitHost,
+			PoseSearchDatabaseEditorAppName,
+			StandaloneDefaultLayout,
+			bCreateDefaultStandaloneMenuParam,
+			bCreateDefaultToolbarParam,
+			DatabaseAsset,
+			bIsToolbarFocusableParam);
+
+		ExtendToolbar();
+
+		RegenerateMenusAndToolbars();
+	}
 
-void FPoseSearchDatabaseEditorToolkit::BindCommands()
-{
-	const FPoseSearchDatabaseEditorCommands& Commands = FPoseSearchDatabaseEditorCommands::Get();
-
-	ToolkitCommands->MapAction(
-		Commands.StopPreviewScene,
-		FExecuteAction::CreateSP(this, &FPoseSearchDatabaseEditorToolkit::StopPreviewScene),
-		EUIActionRepeatMode::RepeatDisabled);
-
-	ToolkitCommands->MapAction(
-		Commands.ResetPreviewScene,
-		FExecuteAction::CreateSP(this, &FPoseSearchDatabaseEditorToolkit::ResetPreviewScene),
-		EUIActionRepeatMode::RepeatDisabled);
-
-	ToolkitCommands->MapAction(
-		Commands.BuildSearchIndex,
-		FExecuteAction::CreateSP(this, &FPoseSearchDatabaseEditorToolkit::BuildSearchIndex),
-		EUIActionRepeatMode::RepeatDisabled);
-}
+	void FDatabaseEditorToolkit::BindCommands()
+	{
+		const FDatabaseEditorCommands& Commands = FDatabaseEditorCommands::Get();
+
+		ToolkitCommands->MapAction(
+			Commands.StopPreviewScene,
+			FExecuteAction::CreateSP(this, &FDatabaseEditorToolkit::StopPreviewScene),
+			EUIActionRepeatMode::RepeatDisabled);
+
+		ToolkitCommands->MapAction(
+			Commands.ResetPreviewScene,
+			FExecuteAction::CreateSP(this, &FDatabaseEditorToolkit::ResetPreviewScene),
+			EUIActionRepeatMode::RepeatDisabled);
+
+		ToolkitCommands->MapAction(
+			Commands.BuildSearchIndex,
+			FExecuteAction::CreateSP(this, &FDatabaseEditorToolkit::BuildSearchIndex),
+			EUIActionRepeatMode::RepeatDisabled);
+	}
 
-void FPoseSearchDatabaseEditorToolkit::ExtendToolbar()
-{
-	TSharedPtr<FExtender> ToolbarExtender = MakeShareable(new FExtender);
+	void FDatabaseEditorToolkit::ExtendToolbar()
+	{
+		TSharedPtr<FExtender> ToolbarExtender = MakeShareable(new FExtender);
 
-	AddToolbarExtender(ToolbarExtender);
+		AddToolbarExtender(ToolbarExtender);
 
-	ToolbarExtender->AddToolBarExtension(
-		"Asset",
-		EExtensionHook::After,
-		GetToolkitCommands(),
-		FToolBarExtensionDelegate::CreateSP(this, &FPoseSearchDatabaseEditorToolkit::FillToolbar));
-}
+		ToolbarExtender->AddToolBarExtension(
+			"Asset",
+			EExtensionHook::After,
+			GetToolkitCommands(),
+			FToolBarExtensionDelegate::CreateSP(this, &FDatabaseEditorToolkit::FillToolbar));
+	}
 
-void FPoseSearchDatabaseEditorToolkit::FillToolbar(FToolBarBuilder& ToolbarBuilder)
-{
-	ToolbarBuilder.AddToolBarButton(
-		FPoseSearchDatabaseEditorCommands::Get().StopPreviewScene,
-		NAME_None,
-		TAttribute<FText>(),
-		TAttribute<FText>(),
-		FSlateIcon(FAppStyle::Get().GetStyleSetName(), "Icons.Delete"));
-
-	ToolbarBuilder.AddToolBarButton(
-		FPoseSearchDatabaseEditorCommands::Get().ResetPreviewScene,
-		NAME_None,
-		TAttribute<FText>(),
-		TAttribute<FText>(),
-		FSlateIcon(FAppStyle::Get().GetStyleSetName(), "Icons.Refresh"));
-
-	ToolbarBuilder.AddToolBarButton(
-		FPoseSearchDatabaseEditorCommands::Get().BuildSearchIndex,
-		NAME_None,
-		TAttribute<FText>(),
-		TAttribute<FText>(),
-		FSlateIcon());
-}
+	void FDatabaseEditorToolkit::FillToolbar(FToolBarBuilder& ToolbarBuilder)
+	{
+		ToolbarBuilder.AddToolBarButton(
+			FDatabaseEditorCommands::Get().StopPreviewScene,
+			NAME_None,
+			TAttribute<FText>(),
+			TAttribute<FText>(),
+			FSlateIcon(FAppStyle::Get().GetStyleSetName(), "Icons.Delete"));
+
+		ToolbarBuilder.AddToolBarButton(
+			FDatabaseEditorCommands::Get().ResetPreviewScene,
+			NAME_None,
+			TAttribute<FText>(),
+			TAttribute<FText>(),
+			FSlateIcon(FAppStyle::Get().GetStyleSetName(), "Icons.Refresh"));
+
+		ToolbarBuilder.AddToolBarButton(
+			FDatabaseEditorCommands::Get().BuildSearchIndex,
+			NAME_None,
+			TAttribute<FText>(),
+			TAttribute<FText>(),
+			FSlateIcon());
+	}
 
-void FPoseSearchDatabaseEditorToolkit::RegisterTabSpawners(const TSharedRef<FTabManager>& InTabManager)
-{
-	WorkspaceMenuCategory = InTabManager->AddLocalWorkspaceMenuCategory(
-		LOCTEXT("WorkspaceMenu_PoseSearchDbEditor", "Pose Search Database Editor"));
-	auto WorkspaceMenuCategoryRef = WorkspaceMenuCategory.ToSharedRef();
-
-	FAssetEditorToolkit::RegisterTabSpawners(InTabManager);
-
-	InTabManager->RegisterTabSpawner(
-		FPoseSearchDatabaseEditorTabs::ViewportID, 
-		FOnSpawnTab::CreateSP(this, &FPoseSearchDatabaseEditorToolkit::SpawnTab_Viewport))
-		.SetDisplayName(LOCTEXT("ViewportTab", "Viewport"))
-		.SetGroup(WorkspaceMenuCategoryRef)
-		.SetIcon(FSlateIcon(FEditorStyle::GetStyleSetName(), "GraphEditor.EventGraph_16x"));
-
-	InTabManager->RegisterTabSpawner(
-		FPoseSearchDatabaseEditorTabs::AssetDetailsID, 
-		FOnSpawnTab::CreateSP(this, &FPoseSearchDatabaseEditorToolkit::SpawnTab_AssetDetails))
-		.SetDisplayName(LOCTEXT("AssetDetailsTab", "AssetDetails"))
-		.SetGroup(WorkspaceMenuCategoryRef)
-		.SetIcon(FSlateIcon(FEditorStyle::GetStyleSetName(), "LevelEditor.Tabs.Details"));
-
-	InTabManager->RegisterTabSpawner(
-		FPoseSearchDatabaseEditorTabs::PreviewSettingsID, 
-		FOnSpawnTab::CreateSP(this, &FPoseSearchDatabaseEditorToolkit::SpawnTab_PreviewSettings))
-		.SetDisplayName(LOCTEXT("PreviewSceneSettingsTab", "Preview Scene Settings"))
-		.SetGroup(WorkspaceMenuCategoryRef)
-		.SetIcon(FSlateIcon(FEditorStyle::GetStyleSetName(), "LevelEditor.Tabs.Details"));
-
-	InTabManager->RegisterTabSpawner(
-		FPoseSearchDatabaseEditorTabs::AssetTreeViewID,
-		FOnSpawnTab::CreateSP(this, &FPoseSearchDatabaseEditorToolkit::SpawnTab_AssetTreeView))
-		.SetDisplayName(LOCTEXT("TreeViewTab", "Tree View"))
-		.SetGroup(WorkspaceMenuCategoryRef)
-		.SetIcon(FSlateIcon(FEditorStyle::GetStyleSetName(), "GraphEditor.EventGraph_16x"));
-}
+	void FDatabaseEditorToolkit::RegisterTabSpawners(const TSharedRef<FTabManager>& InTabManager)
+	{
+		WorkspaceMenuCategory = InTabManager->AddLocalWorkspaceMenuCategory(
+			LOCTEXT("WorkspaceMenu_PoseSearchDbEditor", "Pose Search Database Editor"));
+		auto WorkspaceMenuCategoryRef = WorkspaceMenuCategory.ToSharedRef();
+
+		FAssetEditorToolkit::RegisterTabSpawners(InTabManager);
+
+		InTabManager->RegisterTabSpawner(
+			FDatabaseEditorTabs::ViewportID,
+			FOnSpawnTab::CreateSP(this, &FDatabaseEditorToolkit::SpawnTab_Viewport))
+			.SetDisplayName(LOCTEXT("ViewportTab", "Viewport"))
+			.SetGroup(WorkspaceMenuCategoryRef)
+			.SetIcon(FSlateIcon(FEditorStyle::GetStyleSetName(), "GraphEditor.EventGraph_16x"));
+
+		InTabManager->RegisterTabSpawner(
+			FDatabaseEditorTabs::AssetDetailsID,
+			FOnSpawnTab::CreateSP(this, &FDatabaseEditorToolkit::SpawnTab_AssetDetails))
+			.SetDisplayName(LOCTEXT("AssetDetailsTab", "AssetDetails"))
+			.SetGroup(WorkspaceMenuCategoryRef)
+			.SetIcon(FSlateIcon(FEditorStyle::GetStyleSetName(), "LevelEditor.Tabs.Details"));
+
+		InTabManager->RegisterTabSpawner(
+			FDatabaseEditorTabs::PreviewSettingsID,
+			FOnSpawnTab::CreateSP(this, &FDatabaseEditorToolkit::SpawnTab_PreviewSettings))
+			.SetDisplayName(LOCTEXT("PreviewSceneSettingsTab", "Preview Scene Settings"))
+			.SetGroup(WorkspaceMenuCategoryRef)
+			.SetIcon(FSlateIcon(FEditorStyle::GetStyleSetName(), "LevelEditor.Tabs.Details"));
+
+		InTabManager->RegisterTabSpawner(
+			FDatabaseEditorTabs::AssetTreeViewID,
+			FOnSpawnTab::CreateSP(this, &FDatabaseEditorToolkit::SpawnTab_AssetTreeView))
+			.SetDisplayName(LOCTEXT("TreeViewTab", "Tree View"))
+			.SetGroup(WorkspaceMenuCategoryRef)
+			.SetIcon(FSlateIcon(FEditorStyle::GetStyleSetName(), "GraphEditor.EventGraph_16x"));
+	}
 
-void FPoseSearchDatabaseEditorToolkit::UnregisterTabSpawners(const TSharedRef<FTabManager>& InTabManager)
-{
-	FAssetEditorToolkit::UnregisterTabSpawners(InTabManager);
+	void FDatabaseEditorToolkit::UnregisterTabSpawners(const TSharedRef<FTabManager>& InTabManager)
+	{
+		FAssetEditorToolkit::UnregisterTabSpawners(InTabManager);
 
-	InTabManager->UnregisterTabSpawner(FPoseSearchDatabaseEditorTabs::ViewportID);
-	InTabManager->UnregisterTabSpawner(FPoseSearchDatabaseEditorTabs::AssetDetailsID);
-	InTabManager->UnregisterTabSpawner(FPoseSearchDatabaseEditorTabs::PreviewSettingsID);
-	InTabManager->UnregisterTabSpawner(FPoseSearchDatabaseEditorTabs::AssetTreeViewID);
-}
+		InTabManager->UnregisterTabSpawner(FDatabaseEditorTabs::ViewportID);
+		InTabManager->UnregisterTabSpawner(FDatabaseEditorTabs::AssetDetailsID);
+		InTabManager->UnregisterTabSpawner(FDatabaseEditorTabs::PreviewSettingsID);
+		InTabManager->UnregisterTabSpawner(FDatabaseEditorTabs::AssetTreeViewID);
+	}
 
-FName FPoseSearchDatabaseEditorToolkit::GetToolkitFName() const
-{
-	return FName("PoseSearchDatabaseEditor");
-}
+	FName FDatabaseEditorToolkit::GetToolkitFName() const
+	{
+		return FName("PoseSearchDatabaseEditor");
+	}
 
-FText FPoseSearchDatabaseEditorToolkit::GetBaseToolkitName() const
-{
-	return LOCTEXT("PoseSearchDatabaseEditorAppLabel", "Pose Search Database Editor");
-}
+	FText FDatabaseEditorToolkit::GetBaseToolkitName() const
+	{
+		return LOCTEXT("PoseSearchDatabaseEditorAppLabel", "Pose Search Database Editor");
+	}
 
-FText FPoseSearchDatabaseEditorToolkit::GetToolkitName() const
-{
-	FFormatNamedArguments Args;
-	Args.Add(TEXT("AssetName"), FText::FromString(GetPoseSearchDatabase()->GetName()));
-	return FText::Format(LOCTEXT("PoseSearchDatabaseEditorToolkitName", "{AssetName}"), Args);
-}
+	FText FDatabaseEditorToolkit::GetToolkitName() const
+	{
+		FFormatNamedArguments Args;
+		Args.Add(TEXT("AssetName"), FText::FromString(GetPoseSearchDatabase()->GetName()));
+		return FText::Format(LOCTEXT("PoseSearchDatabaseEditorToolkitName", "{AssetName}"), Args);
+	}
 
-FLinearColor FPoseSearchDatabaseEditorToolkit::GetWorldCentricTabColorScale() const
-{
-	return FLinearColor::White;
-}
+	FLinearColor FDatabaseEditorToolkit::GetWorldCentricTabColorScale() const
+	{
+		return FLinearColor::White;
+	}
 
-FString FPoseSearchDatabaseEditorToolkit::GetWorldCentricTabPrefix() const
-{
-	return TEXT("PoseSearchDatabaseEditor");
-}
+	FString FDatabaseEditorToolkit::GetWorldCentricTabPrefix() const
+	{
+		return TEXT("PoseSearchDatabaseEditor");
+	}
 
-TSharedRef<SDockTab> FPoseSearchDatabaseEditorToolkit::SpawnTab_Viewport(const FSpawnTabArgs& Args)
-{
-	check(Args.GetTabId() == FPoseSearchDatabaseEditorTabs::ViewportID);
+	TSharedRef<SDockTab> FDatabaseEditorToolkit::SpawnTab_Viewport(const FSpawnTabArgs& Args)
+	{
+		check(Args.GetTabId() == FDatabaseEditorTabs::ViewportID);
 
-	TSharedRef<SDockTab> SpawnedTab = SNew(SDockTab).Label(LOCTEXT("ViewportTab_Title", "Viewport"));
+		TSharedRef<SDockTab> SpawnedTab = SNew(SDockTab).Label(LOCTEXT("ViewportTab_Title", "Viewport"));
 
-	if (ViewportWidget.IsValid())
-	{
-		SpawnedTab->SetContent(ViewportWidget.ToSharedRef());
-	}
+		if (ViewportWidget.IsValid())
+		{
+			SpawnedTab->SetContent(ViewportWidget.ToSharedRef());
+		}
 
-	return SpawnedTab;
-}
+		return SpawnedTab;
+	}
 
-TSharedRef<SDockTab> FPoseSearchDatabaseEditorToolkit::SpawnTab_AssetDetails(const FSpawnTabArgs& Args)
-{
-	check(Args.GetTabId() == FPoseSearchDatabaseEditorTabs::AssetDetailsID);
+	TSharedRef<SDockTab> FDatabaseEditorToolkit::SpawnTab_AssetDetails(const FSpawnTabArgs& Args)
+	{
+		check(Args.GetTabId() == FDatabaseEditorTabs::AssetDetailsID);
 
-	return SNew(SDockTab)
-		.Label(LOCTEXT("AssetDetails_Title", "Asset Details"))
-		[
-			EditingAssetWidget.ToSharedRef()
-		];
-}
+		return SNew(SDockTab)
+			.Label(LOCTEXT("AssetDetails_Title", "Asset Details"))
+			[
+				EditingAssetWidget.ToSharedRef()
+			];
+	}
 
-TSharedRef<SDockTab> FPoseSearchDatabaseEditorToolkit::SpawnTab_PreviewSettings(const FSpawnTabArgs& Args)
-{
-	check(Args.GetTabId() == FPoseSearchDatabaseEditorTabs::PreviewSettingsID);
+	TSharedRef<SDockTab> FDatabaseEditorToolkit::SpawnTab_PreviewSettings(const FSpawnTabArgs& Args)
+	{
+		check(Args.GetTabId() == FDatabaseEditorTabs::PreviewSettingsID);
 
 	FAdvancedPreviewSceneModule& AdvancedPreviewSceneModule = 
 		FModuleManager::LoadModuleChecked<FAdvancedPreviewSceneModule>("AdvancedPreviewScene");
 	TSharedRef<SWidget> InWidget = 
 		AdvancedPreviewSceneModule.CreateAdvancedPreviewSceneSettingsWidget(PreviewScene.ToSharedRef());
 
-	TSharedRef<SDockTab> SpawnedTab = SNew(SDockTab)
-		.Label(LOCTEXT("PreviewSceneSettingsTab", "Preview Scene Settings"))
-		[
-			InWidget
-		];
+		TSharedRef<SDockTab> SpawnedTab = SNew(SDockTab)
+			.Label(LOCTEXT("PreviewSceneSettingsTab", "Preview Scene Settings"))
+			[
+				InWidget
+			];
 
-	return SpawnedTab;
-}
+		return SpawnedTab;
+	}
 
-TSharedRef<SDockTab> FPoseSearchDatabaseEditorToolkit::SpawnTab_AssetTreeView(const FSpawnTabArgs& Args)
-{
-	check(Args.GetTabId() == FPoseSearchDatabaseEditorTabs::AssetTreeViewID);
+	TSharedRef<SDockTab> FDatabaseEditorToolkit::SpawnTab_AssetTreeView(const FSpawnTabArgs& Args)
+	{
+		check(Args.GetTabId() == FDatabaseEditorTabs::AssetTreeViewID);
 
-	return SNew(SDockTab)
-		.Label(LOCTEXT("AssetTreeView_Title", "Asset Tree"))
-		[
-			AssetTreeWidget.ToSharedRef()
-		];
-}
+		return SNew(SDockTab)
+			.Label(LOCTEXT("AssetTreeView_Title", "Asset Tree"))
+			[
+				AssetTreeWidget.ToSharedRef()
+			];
+	}
 
-void FPoseSearchDatabaseEditorToolkit::OnFinishedChangingProperties(const FPropertyChangedEvent& PropertyChangedEvent)
-{
-	const FName PropertyName = 
-		(PropertyChangedEvent.Property != nullptr) ? 
-		PropertyChangedEvent.Property->GetFName() : NAME_None;
-	const FName MemberPropertyName = 
-		(PropertyChangedEvent.MemberProperty != nullptr) ? 
-		PropertyChangedEvent.MemberProperty->GetFName() : NAME_None;
-
-	UE_LOG(
-		LogPoseSearchEditor, 
-		Log, 
-		TEXT("FPoseSearchDatabaseEditorToolkit::OnFinishedChangingProperties MemberPropertyName: %s PropertyName: %s"),
-		*MemberPropertyName.ToString(), 
-		*PropertyName.ToString());
+	void FDatabaseEditorToolkit::OnFinishedChangingProperties(const FPropertyChangedEvent& PropertyChangedEvent)
+	{
+		const FName PropertyName = 
+			(PropertyChangedEvent.Property != nullptr) ? 
+			PropertyChangedEvent.Property->GetFName() : NAME_None;
+		const FName MemberPropertyName = 
+			(PropertyChangedEvent.MemberProperty != nullptr) ? 
+			PropertyChangedEvent.MemberProperty->GetFName() : NAME_None;
+
+		UE_LOG(
+			LogPoseSearchEditor,
+			Log,
+			TEXT("FPoseSearchDatabaseEditorToolkit::OnFinishedChangingProperties MemberPropertyName: %s PropertyName: %s"),
+			*MemberPropertyName.ToString(),
+			*PropertyName.ToString());
+	}
 }
 
 #undef LOCTEXT_NAMESPACE
diff --git a/Engine/Plugins/Experimental/Animation/PoseSearch/Source/Editor/Private/PoseSearchDatabasePreviewScene.cpp b/Engine/Plugins/Experimental/Animation/PoseSearch/Source/Editor/Private/PoseSearchDatabasePreviewScene.cpp
--- a/Engine/Plugins/Experimental/Animation/PoseSearch/Source/Editor/Private/PoseSearchDatabasePreviewScene.cpp
+++ b/Engine/Plugins/Experimental/Animation/PoseSearch/Source/Editor/Private/PoseSearchDatabasePreviewScene.cpp
@@ -10,71 +10,73 @@
 #include "GameFramework/WorldSettings.h"
 #include "EngineUtils.h"
 
-FPoseSearchDatabasePreviewScene::FPoseSearchDatabasePreviewScene(
-	ConstructionValues CVs, 
-	const TSharedRef<FPoseSearchDatabaseEditorToolkit>& EditorToolkit)
-	: FAdvancedPreviewScene(CVs)
-	, EditorToolkitPtr(EditorToolkit)
+namespace UE::PoseSearch
 {
-	// Disable killing actors outside of the world
-	AWorldSettings* WorldSettings = GetWorld()->GetWorldSettings(true);
-	WorldSettings->bEnableWorldBoundsChecks = false;
-
-	// Spawn an owner for FloorMeshComponent so CharacterMovementComponent can detect it as a valid floor and slide 
-	// along it
+	FDatabasePreviewScene::FDatabasePreviewScene(
+		ConstructionValues CVs,
+		const TSharedRef<FDatabaseEditorToolkit>& EditorToolkit)
+		: FAdvancedPreviewScene(CVs)
+		, EditorToolkitPtr(EditorToolkit)
 	{
-		AActor* FloorActor = GetWorld()->SpawnActor<AActor>(AActor::StaticClass(), FTransform());
-		check(FloorActor);
+		// Disable killing actors outside of the world
+		AWorldSettings* WorldSettings = GetWorld()->GetWorldSettings(true);
+		WorldSettings->bEnableWorldBoundsChecks = false;
 
-		static const FString NewName = FString(TEXT("FloorComponent"));
-		FloorMeshComponent->Rename(*NewName, FloorActor);
+		// Spawn an owner for FloorMeshComponent so CharacterMovementComponent can detect it as a valid floor and slide 
+		// along it
+		{
+			AActor* FloorActor = GetWorld()->SpawnActor<AActor>(AActor::StaticClass(), FTransform());
+			check(FloorActor);
 
-		FloorActor->SetRootComponent(FloorMeshComponent);
-	}
-}
+			static const FString NewName = FString(TEXT("FloorComponent"));
+			FloorMeshComponent->Rename(*NewName, FloorActor);
 
-void FPoseSearchDatabasePreviewScene::Tick(float InDeltaTime)
-{
-	FAdvancedPreviewScene::Tick(InDeltaTime);
+			FloorActor->SetRootComponent(FloorMeshComponent);
+		}
+	}
 
-	// Trigger Begin Play in this preview world.
-	// This is needed for the CharacterMovementComponent to be able to switch to falling mode. 
-	// See: UCharacterMovementComponent::StartFalling
-	if (PreviewWorld && !PreviewWorld->bBegunPlay)
+	void FDatabasePreviewScene::Tick(float InDeltaTime)
 	{
-		for (FActorIterator It(PreviewWorld); It; ++It)
+		FAdvancedPreviewScene::Tick(InDeltaTime);
+
+		// Trigger Begin Play in this preview world.
+		// This is needed for the CharacterMovementComponent to be able to switch to falling mode. 
+		// See: UCharacterMovementComponent::StartFalling
+		if (PreviewWorld && !PreviewWorld->bBegunPlay)
 		{
-			It->DispatchBeginPlay();
-		}
+			for (FActorIterator It(PreviewWorld); It; ++It)
+			{
+				It->DispatchBeginPlay();
+			}
 
-		PreviewWorld->bBegunPlay = true;
-	}
+			PreviewWorld->bBegunPlay = true;
+		}
 
-	if (!GIntraFrameDebuggingGameThread)
-	{
-		GetWorld()->Tick(LEVELTICK_All, InDeltaTime);
-	}
+		if (!GIntraFrameDebuggingGameThread)
+		{
+			GetWorld()->Tick(LEVELTICK_All, InDeltaTime);
+		}
 
-	const FPoseSearchDatabaseViewModel* ViewModel = GetEditorToolkit()->GetViewModel();
-	if (ViewModel->GetPoseSearchDatabase()->IsValidForSearch() &&
-		ViewModel->IsPoseFeaturesDrawMode(EPoseSearchFeaturesDrawMode::All))
-	{
-		for (const FPoseSearchDatabasePreviewActor& PreviewActor : ViewModel->GetPreviewActors())
+		const FDatabaseViewModel* ViewModel = GetEditorToolkit()->GetViewModel();
+		if (ViewModel->GetPoseSearchDatabase()->IsValidForSearch() &&
+			ViewModel->IsPoseFeaturesDrawMode(EFeaturesDrawMode::All))
 		{
-			if (PreviewActor.CurrentPoseIndex != INDEX_NONE)
+			for (const FDatabasePreviewActor& PreviewActor : ViewModel->GetPreviewActors())
 			{
-				UE::PoseSearch::FDebugDrawParams DrawParams;
-				DrawParams.RootTransform = PreviewActor.Mesh->GetComponentTransform();
-				DrawParams.Database = ViewModel->GetPoseSearchDatabase();
-				DrawParams.World = GetWorld();
-				DrawParams.DefaultLifeTime = 0.0f;
-				DrawParams.PoseIdx = PreviewActor.CurrentPoseIndex;
-				DrawParams.PointSize = 5.0f;
-				EnumAddFlags(DrawParams.Flags, UE::PoseSearch::EDebugDrawFlags::DrawFast);
+				if (PreviewActor.CurrentPoseIndex != INDEX_NONE)
+				{
+					UE::PoseSearch::FDebugDrawParams DrawParams;
+					DrawParams.RootTransform = PreviewActor.Mesh->GetComponentTransform();
+					DrawParams.Database = ViewModel->GetPoseSearchDatabase();
+					DrawParams.World = GetWorld();
+					DrawParams.DefaultLifeTime = 0.0f;
+					DrawParams.PoseIdx = PreviewActor.CurrentPoseIndex;
+					DrawParams.PointSize = 5.0f;
+					EnumAddFlags(DrawParams.Flags, UE::PoseSearch::EDebugDrawFlags::DrawFast);
 
-				UE::PoseSearch::Draw(DrawParams);
+					UE::PoseSearch::Draw(DrawParams);
+				}
 			}
 		}
 	}
 }
-
diff --git a/Engine/Plugins/Experimental/Animation/PoseSearch/Source/Editor/Private/PoseSearchDatabaseViewModel.cpp b/Engine/Plugins/Experimental/Animation/PoseSearch/Source/Editor/Private/PoseSearchDatabaseViewModel.cpp
--- a/Engine/Plugins/Experimental/Animation/PoseSearch/Source/Editor/Private/PoseSearchDatabaseViewModel.cpp
+++ b/Engine/Plugins/Experimental/Animation/PoseSearch/Source/Editor/Private/PoseSearchDatabaseViewModel.cpp
@@ -16,297 +16,301 @@
 #include "Camera/CameraComponent.h"
 #include "Components/CapsuleComponent.h"
 
-FPoseSearchDatabaseViewModel::FPoseSearchDatabaseViewModel()
-	: PoseSearchDatabase(nullptr)
-{
-}
 
-FPoseSearchDatabaseViewModel::~FPoseSearchDatabaseViewModel()
+namespace UE::PoseSearch
 {
-	PoseSearchDatabase->UnregisterOnDerivedDataRebuild(this);
-}
+	FDatabaseViewModel::FDatabaseViewModel()
+		: PoseSearchDatabase(nullptr)
+	{
+	}
 
-void FPoseSearchDatabaseViewModel::AddReferencedObjects(FReferenceCollector& Collector)
-{
-	Collector.AddReferencedObject(PoseSearchDatabase);
-}
+	FDatabaseViewModel::~FDatabaseViewModel()
+	{
+		PoseSearchDatabase->UnregisterOnDerivedDataRebuild(this);
+	}
 
-void FPoseSearchDatabaseViewModel::Initialize(
-	UPoseSearchDatabase* InPoseSearchDatabase,
-	const TSharedRef<FPoseSearchDatabasePreviewScene>& InPreviewScene)
-{
-	PoseSearchDatabase = InPoseSearchDatabase;
-	PreviewScenePtr = InPreviewScene;
+	void FDatabaseViewModel::AddReferencedObjects(FReferenceCollector& Collector)
+	{
+		Collector.AddReferencedObject(PoseSearchDatabase);
+	}
 
-	RemovePreviewActors();
+	void FDatabaseViewModel::Initialize(
+		UPoseSearchDatabase* InPoseSearchDatabase,
+		const TSharedRef<FDatabasePreviewScene>& InPreviewScene)
+	{
+		PoseSearchDatabase = InPoseSearchDatabase;
+		PreviewScenePtr = InPreviewScene;
 
-	PoseSearchDatabase->RegisterOnDerivedDataRebuild(
-		UPoseSearchDatabase::FOnDerivedDataRebuild::CreateSP(
-			this, 
-			&FPoseSearchDatabaseViewModel::ResetPreviewActors));
-}
+		RemovePreviewActors();
 
-void FPoseSearchDatabaseViewModel::ResetPreviewActors()
-{
-	PlayTime = 0.0f;
-	RespawnPreviewActors();
-}
+		PoseSearchDatabase->RegisterOnDerivedDataRebuild(
+			UPoseSearchDatabase::FOnDerivedDataRebuild::CreateSP(
+				this,
+				&FDatabaseViewModel::ResetPreviewActors));
+	}
 
-void FPoseSearchDatabaseViewModel::RespawnPreviewActors()
-{
-	RemovePreviewActors();
-	const FPoseSearchIndex* SearchIndex = PoseSearchDatabase->GetSearchIndex();
-	for (const FPoseSearchIndexAsset& IndexAsset : SearchIndex->Assets)
+	void FDatabaseViewModel::ResetPreviewActors()
+	{
+		PlayTime = 0.0f;
+		RespawnPreviewActors();
+	}
+
+	void FDatabaseViewModel::RespawnPreviewActors()
 	{
-		if (AnimationPreviewMode == EAnimationPreviewMode::OriginalAndMirrored || !IndexAsset.bMirrored)
+		RemovePreviewActors();
+		const FPoseSearchIndex* SearchIndex = PoseSearchDatabase->GetSearchIndex();
+		for (const FPoseSearchIndexAsset& IndexAsset : SearchIndex->Assets)
 		{
-			FPoseSearchDatabasePreviewActor PreviewActor = SpawnPreviewActor(IndexAsset);
-			if (PreviewActor.IsValid())
+			if (AnimationPreviewMode == EAnimationPreviewMode::OriginalAndMirrored || !IndexAsset.bMirrored)
 			{
-				PreviewActors.Add(PreviewActor);
+				FDatabasePreviewActor PreviewActor = SpawnPreviewActor(IndexAsset);
+				if (PreviewActor.IsValid())
+				{
+					PreviewActors.Add(PreviewActor);
+				}
 			}
 		}
+		UpdatePreviewActors();
+
+		// todo: do blendspaces afterwards
 	}
-	UpdatePreviewActors();
 
-	// todo: do blendspaces afterwards
-}
+	void FDatabaseViewModel::BuildSearchIndex()
+	{
+		PoseSearchDatabase->BeginCacheDerivedData();
+	}
 
-void FPoseSearchDatabaseViewModel::BuildSearchIndex()
-{
-	PoseSearchDatabase->BeginCacheDerivedData();
-}
+	FDatabasePreviewActor FDatabaseViewModel::SpawnPreviewActor(const FPoseSearchIndexAsset& IndexAsset)
+	{
+		USkeleton* Skeleton = PoseSearchDatabase->Schema->Skeleton;
 
-FPoseSearchDatabasePreviewActor FPoseSearchDatabaseViewModel::SpawnPreviewActor(const FPoseSearchIndexAsset& IndexAsset)
-{
-	USkeleton* Skeleton = PoseSearchDatabase->Schema->Skeleton;
+		FPoseSearchDatabaseSequence DatabaseSequence = PoseSearchDatabase->Sequences[IndexAsset.SourceAssetIdx];
 
-	FPoseSearchDatabaseSequence DatabaseSequence = PoseSearchDatabase->Sequences[IndexAsset.SourceAssetIdx];
+		FDatabasePreviewActor PreviewActor;
 
-	FPoseSearchDatabasePreviewActor PreviewActor;
+		if (!Skeleton || !DatabaseSequence.Sequence)
+		{
+			return PreviewActor;
+		}
 
-	if (!Skeleton || !DatabaseSequence.Sequence)
-	{
-		return PreviewActor;
-	}
+		PreviewActor.IndexAsset = &IndexAsset;
+		PreviewActor.CurrentPoseIndex = INDEX_NONE;
 
-	PreviewActor.IndexAsset = &IndexAsset;
-	PreviewActor.CurrentPoseIndex = INDEX_NONE;
+		// todo: use preview when possible
+		UClass* PreviewClass = AActor::StaticClass();
+		const FTransform SpawnTransform = FTransform::Identity;
 
-	// todo: use preview when possible
-	UClass* PreviewClass = AActor::StaticClass();
-	const FTransform SpawnTransform = FTransform::Identity;
+		FActorSpawnParameters Params;
+		Params.SpawnCollisionHandlingOverride = ESpawnActorCollisionHandlingMethod::AdjustIfPossibleButAlwaysSpawn;
+		PreviewActor.Actor = GetWorld()->SpawnActor<AActor>(PreviewClass, SpawnTransform, Params);
+		PreviewActor.Actor->SetFlags(RF_Transient);
 
-	FActorSpawnParameters Params;
-	Params.SpawnCollisionHandlingOverride = ESpawnActorCollisionHandlingMethod::AdjustIfPossibleButAlwaysSpawn;
-	PreviewActor.Actor = GetWorld()->SpawnActor<AActor>(PreviewClass, SpawnTransform, Params);
-	PreviewActor.Actor->SetFlags(RF_Transient);
+		PreviewActor.Mesh = NewObject<UDebugSkelMeshComponent>(PreviewActor.Actor.Get());
+		PreviewActor.Mesh->RegisterComponentWithWorld(GetWorld());
 
-	PreviewActor.Mesh = NewObject<UDebugSkelMeshComponent>(PreviewActor.Actor.Get());
-	PreviewActor.Mesh->RegisterComponentWithWorld(GetWorld());
+		PreviewActor.AnimInstance = NewObject<UAnimPreviewInstance>(PreviewActor.Mesh.Get());
+		PreviewActor.Mesh->PreviewInstance = PreviewActor.AnimInstance.Get();
+		PreviewActor.AnimInstance->InitializeAnimation();
 
-	PreviewActor.AnimInstance = NewObject<UAnimPreviewInstance>(PreviewActor.Mesh.Get());
-	PreviewActor.Mesh->PreviewInstance = PreviewActor.AnimInstance.Get();
-	PreviewActor.AnimInstance->InitializeAnimation();
+		PreviewActor.Mesh->SetSkeletalMesh(Skeleton->GetPreviewMesh(true));
+		PreviewActor.Mesh->EnablePreview(true, DatabaseSequence.Sequence);
+		PreviewActor.AnimInstance->SetAnimationAsset(DatabaseSequence.Sequence);
+		if (IndexAsset.bMirrored && PoseSearchDatabase->Schema)
+		{
+			PreviewActor.AnimInstance->SetMirrorDataTable(PoseSearchDatabase->Schema->MirrorDataTable);
+		}
 
-	PreviewActor.Mesh->SetSkeletalMesh(Skeleton->GetPreviewMesh(true));
-	PreviewActor.Mesh->EnablePreview(true, DatabaseSequence.Sequence);
-	PreviewActor.AnimInstance->SetAnimationAsset(DatabaseSequence.Sequence);
-	if (IndexAsset.bMirrored && PoseSearchDatabase->Schema)
-	{
-		PreviewActor.AnimInstance->SetMirrorDataTable(PoseSearchDatabase->Schema->MirrorDataTable);
-	}
-	
-	PreviewActor.AnimInstance->PlayAnim(false);
+		PreviewActor.AnimInstance->PlayAnim(false);
 
-	if (!PreviewActor.Actor->GetRootComponent())
-	{
-		PreviewActor.Actor->SetRootComponent(PreviewActor.Mesh.Get());
-	}
+		if (!PreviewActor.Actor->GetRootComponent())
+		{
+			PreviewActor.Actor->SetRootComponent(PreviewActor.Mesh.Get());
+		}
 
-	UE_LOG(
-		LogPoseSearchEditor, Log, TEXT("Spawned preview Actor: %s at Loc: %s Rot: %s"),
-		*GetNameSafe(PreviewActor.Actor.Get()),
-		*SpawnTransform.GetLocation().ToString(), 
-		*SpawnTransform.Rotator().ToString());
+		UE_LOG(
+			LogPoseSearchEditor, Log, TEXT("Spawned preview Actor: %s at Loc: %s Rot: %s"),
+			*GetNameSafe(PreviewActor.Actor.Get()),
+			*SpawnTransform.GetLocation().ToString(),
+			*SpawnTransform.Rotator().ToString());
 
-	return PreviewActor;
-}
+		return PreviewActor;
+	}
 
-UWorld* FPoseSearchDatabaseViewModel::GetWorld() const
-{
-	check(PreviewScenePtr.IsValid());
-	return PreviewScenePtr.Pin()->GetWorld();
-}
+	UWorld* FDatabaseViewModel::GetWorld() const
+	{
+		check(PreviewScenePtr.IsValid());
+		return PreviewScenePtr.Pin()->GetWorld();
+	}
 
-UObject* FPoseSearchDatabaseViewModel::GetPlaybackContext() const
-{
-	return GetWorld();
-}
+	UObject* FDatabaseViewModel::GetPlaybackContext() const
+	{
+		return GetWorld();
+	}
 
 
-void FPoseSearchDatabaseViewModel::OnPreviewActorClassChanged()
-{
-	// todo: implement
-}
+	void FDatabaseViewModel::OnPreviewActorClassChanged()
+	{
+		// todo: implement
+	}
 
-void FPoseSearchDatabaseViewModel::Tick(float DeltaSeconds)
-{
-	PlayTime += DeltaSeconds;
-	UpdatePreviewActors();
-}
+	void FDatabaseViewModel::Tick(float DeltaSeconds)
+	{
+		PlayTime += DeltaSeconds;
+		UpdatePreviewActors();
+	}
 
-void FPoseSearchDatabaseViewModel::UpdatePreviewActors()
-{
-	TSharedPtr<FPoseSearchDatabasePreviewScene> PreviewScene = PreviewScenePtr.Pin();
-	for (FPoseSearchDatabasePreviewActor& PreviewActor : GetPreviewActors())
+	void FDatabaseViewModel::UpdatePreviewActors()
 	{
-		if (const UAnimSequence* Sequence = Cast<UAnimSequence>(PreviewActor.AnimInstance->GetAnimationAsset()))
+		TSharedPtr<FDatabasePreviewScene> PreviewScene = PreviewScenePtr.Pin();
+		for (FDatabasePreviewActor& PreviewActor : GetPreviewActors())
 		{
-			float CurrentTime = 0.0f;
-			FAnimationRuntime::AdvanceTime(false, PlayTime, CurrentTime, Sequence->GetPlayLength());
-
-			FTransform RootMotion = Sequence->ExtractRootMotionFromRange(0.0f, CurrentTime);
-			if (PreviewActor.AnimInstance->GetMirrorDataTable())
+			if (const UAnimSequence* Sequence = Cast<UAnimSequence>(PreviewActor.AnimInstance->GetAnimationAsset()))
 			{
-				RootMotion = MirrorRootMotion(RootMotion, PreviewActor.AnimInstance->GetMirrorDataTable());
+				float CurrentTime = 0.0f;
+				FAnimationRuntime::AdvanceTime(false, PlayTime, CurrentTime, Sequence->GetPlayLength());
+
+				FTransform RootMotion = Sequence->ExtractRootMotionFromRange(0.0f, CurrentTime);
+				if (PreviewActor.AnimInstance->GetMirrorDataTable())
+				{
+					RootMotion = MirrorRootMotion(RootMotion, PreviewActor.AnimInstance->GetMirrorDataTable());
+				}
+				PreviewActor.Actor->SetActorTransform(RootMotion);
+				PreviewActor.AnimInstance->SetPosition(CurrentTime);
+
+				PreviewActor.CurrentPoseIndex =
+					PoseSearchDatabase->GetPoseIndexFromTime(CurrentTime, PreviewActor.IndexAsset);
 			}
-			PreviewActor.Actor->SetActorTransform(RootMotion);
-			PreviewActor.AnimInstance->SetPosition(CurrentTime);
-
-			PreviewActor.CurrentPoseIndex = PoseSearchDatabase->GetPoseIndexFromTime(CurrentTime, PreviewActor.IndexAsset);
 		}
 	}
-}
 
-void FPoseSearchDatabaseViewModel::RemovePreviewActors()
-{
-	for (auto PreviewActor : PreviewActors)
+	void FDatabaseViewModel::RemovePreviewActors()
 	{
-		PreviewActor.Actor->Destroy();
+		for (auto PreviewActor : PreviewActors)
+		{
+			PreviewActor.Actor->Destroy();
+		}
+		PreviewActors.Reset();
 	}
-	PreviewActors.Reset();
-}
 
-FTransform FPoseSearchDatabaseViewModel::MirrorRootMotion(
-	FTransform RootMotion, 
-	const UMirrorDataTable* MirrorDataTable)
-{
-	const FTransform RootReferenceTransform = 
-		PoseSearchDatabase->Schema->Skeleton->GetReferenceSkeleton().GetRefBonePose()[0];
-	const FQuat RootReferenceRotation = RootReferenceTransform.GetRotation();
-
-	const EAxis::Type MirrorAxis = MirrorDataTable->MirrorAxis;
-	FVector T = RootMotion.GetTranslation();
-	T = FAnimationRuntime::MirrorVector(T, MirrorAxis);
-	FQuat Q = RootMotion.GetRotation();
-	Q = FAnimationRuntime::MirrorQuat(Q, MirrorAxis);
-	Q *= FAnimationRuntime::MirrorQuat(RootReferenceRotation, MirrorAxis).Inverse() * RootReferenceRotation;
-
-	FTransform MirroredRootMotion = FTransform(Q, T, RootMotion.GetScale3D());
-	return MirroredRootMotion;
-}
-
-void FPoseSearchDatabaseViewModel::OnSetPoseFeaturesDrawMode(EPoseSearchFeaturesDrawMode DrawMode)
-{
-	PoseFeaturesDrawMode = DrawMode;
-}
-
-bool FPoseSearchDatabaseViewModel::IsPoseFeaturesDrawMode(EPoseSearchFeaturesDrawMode DrawMode) const
-{
-	return PoseFeaturesDrawMode == DrawMode;
-}
-
-void FPoseSearchDatabaseViewModel::OnSetAnimationPreviewMode(EAnimationPreviewMode PreviewMode)
-{
-	if (PreviewMode != AnimationPreviewMode)
+	FTransform FDatabaseViewModel::MirrorRootMotion(
+		FTransform RootMotion,
+		const UMirrorDataTable* MirrorDataTable)
 	{
-		AnimationPreviewMode = PreviewMode;
-		RespawnPreviewActors();
+		const FTransform RootReferenceTransform =
+			PoseSearchDatabase->Schema->Skeleton->GetReferenceSkeleton().GetRefBonePose()[0];
+		const FQuat RootReferenceRotation = RootReferenceTransform.GetRotation();
+
+		const EAxis::Type MirrorAxis = MirrorDataTable->MirrorAxis;
+		FVector T = RootMotion.GetTranslation();
+		T = FAnimationRuntime::MirrorVector(T, MirrorAxis);
+		FQuat Q = RootMotion.GetRotation();
+		Q = FAnimationRuntime::MirrorQuat(Q, MirrorAxis);
+		Q *= FAnimationRuntime::MirrorQuat(RootReferenceRotation, MirrorAxis).Inverse() * RootReferenceRotation;
+
+		FTransform MirroredRootMotion = FTransform(Q, T, RootMotion.GetScale3D());
+		return MirroredRootMotion;
 	}
-}
 
-bool FPoseSearchDatabaseViewModel::IsAnimationPreviewMode(EAnimationPreviewMode PreviewMode) const
-{
-	return AnimationPreviewMode == PreviewMode;
-}
+	void FDatabaseViewModel::OnSetPoseFeaturesDrawMode(EFeaturesDrawMode DrawMode)
+	{
+		PoseFeaturesDrawMode = DrawMode;
+	}
 
-void FPoseSearchDatabaseViewModel::AddSequenceToDatabase(UAnimSequence* AnimSequence, int InitialGroupIdx)
-{
-	FPoseSearchDatabaseSequence& NewDbSequence = PoseSearchDatabase->Sequences.AddDefaulted_GetRef();
-	NewDbSequence.Sequence = AnimSequence;
-	
-	if (InitialGroupIdx >= 0)
+	bool FDatabaseViewModel::IsPoseFeaturesDrawMode(EFeaturesDrawMode DrawMode) const
 	{
-		const FPoseSearchDatabaseGroup& InitialGroup = PoseSearchDatabase->Groups[InitialGroupIdx];
-		NewDbSequence.GroupTags.AddTag(InitialGroup.Tag);
+		return PoseFeaturesDrawMode == DrawMode;
 	}
-}
 
-void FPoseSearchDatabaseViewModel::AddBlendSpaceToDatabase(UBlendSpace* BlendSpace, int InitialGroupIdx)
-{
-	FPoseSearchDatabaseBlendSpace& NewDbBlendSpace = PoseSearchDatabase->BlendSpaces.AddDefaulted_GetRef();
-	NewDbBlendSpace.BlendSpace = BlendSpace;
+	void FDatabaseViewModel::OnSetAnimationPreviewMode(EAnimationPreviewMode PreviewMode)
+	{
+		if (PreviewMode != AnimationPreviewMode)
+		{
+			AnimationPreviewMode = PreviewMode;
+			RespawnPreviewActors();
+		}
+	}
 
-	if (InitialGroupIdx >= 0)
+	bool FDatabaseViewModel::IsAnimationPreviewMode(EAnimationPreviewMode PreviewMode) const
 	{
-		const FPoseSearchDatabaseGroup& InitialGroup = PoseSearchDatabase->Groups[InitialGroupIdx];
-		NewDbBlendSpace.GroupTags.AddTag(InitialGroup.Tag);
+		return AnimationPreviewMode == PreviewMode;
 	}
-}
 
-void FPoseSearchDatabaseViewModel::AddGroupToDatabase()
-{
-	PoseSearchDatabase->Groups.AddDefaulted();
-}
+	void FDatabaseViewModel::AddSequenceToDatabase(UAnimSequence* AnimSequence, int InitialGroupIdx)
+	{
+		FPoseSearchDatabaseSequence& NewDbSequence = PoseSearchDatabase->Sequences.AddDefaulted_GetRef();
+		NewDbSequence.Sequence = AnimSequence;
 
-void FPoseSearchDatabaseViewModel::DeleteSequenceFromDatabase(int32 SequenceIdx)
-{
-	PoseSearchDatabase->Sequences.RemoveAt(SequenceIdx);
-}
+		if (InitialGroupIdx >= 0)
+		{
+			const FPoseSearchDatabaseGroup& InitialGroup = PoseSearchDatabase->Groups[InitialGroupIdx];
+			NewDbSequence.GroupTags.AddTag(InitialGroup.Tag);
+		}
+	}
 
-void FPoseSearchDatabaseViewModel::RemoveSequenceFromGroup(int32 SequenceIdx, int32 GroupIdx)
-{
-	FPoseSearchDatabaseSequence& DbSequence = PoseSearchDatabase->Sequences[SequenceIdx];
-	FPoseSearchDatabaseGroup& DbGroup = PoseSearchDatabase->Groups[GroupIdx];
-	if (DbGroup.Tag.IsValid())
+	void FDatabaseViewModel::AddBlendSpaceToDatabase(UBlendSpace* BlendSpace, int InitialGroupIdx)
 	{
-		DbSequence.GroupTags.RemoveTag(DbGroup.Tag);
+		FPoseSearchDatabaseBlendSpace& NewDbBlendSpace = PoseSearchDatabase->BlendSpaces.AddDefaulted_GetRef();
+		NewDbBlendSpace.BlendSpace = BlendSpace;
+
+		if (InitialGroupIdx >= 0)
+		{
+			const FPoseSearchDatabaseGroup& InitialGroup = PoseSearchDatabase->Groups[InitialGroupIdx];
+			NewDbBlendSpace.GroupTags.AddTag(InitialGroup.Tag);
+		}
 	}
-}
 
-void FPoseSearchDatabaseViewModel::DeleteBlendSpaceFromDatabase(int32 BlendSpaceIdx)
-{
-	PoseSearchDatabase->BlendSpaces.RemoveAt(BlendSpaceIdx);
-}
+	void FDatabaseViewModel::AddGroupToDatabase()
+	{
+		PoseSearchDatabase->Groups.AddDefaulted();
+	}
 
-void FPoseSearchDatabaseViewModel::RemoveBlendSpaceFromGroup(int32 BlendSpaceIdx, int32 GroupIdx)
-{
-	FPoseSearchDatabaseBlendSpace& DbBlendSpace = PoseSearchDatabase->BlendSpaces[BlendSpaceIdx];
-	FPoseSearchDatabaseGroup& DbGroup = PoseSearchDatabase->Groups[GroupIdx];
-	if (DbGroup.Tag.IsValid())
+	void FDatabaseViewModel::DeleteSequenceFromDatabase(int32 SequenceIdx)
 	{
-		DbBlendSpace.GroupTags.RemoveTag(DbGroup.Tag);
+		PoseSearchDatabase->Sequences.RemoveAt(SequenceIdx);
 	}
-}
 
-void FPoseSearchDatabaseViewModel::DeleteGroup(int32 GroupIdx)
-{
-	FPoseSearchDatabaseGroup& DbGroup = PoseSearchDatabase->Groups[GroupIdx];
-	
-	if (DbGroup.Tag.IsValid())
+	void FDatabaseViewModel::RemoveSequenceFromGroup(int32 SequenceIdx, int32 GroupIdx)
 	{
-		for (FPoseSearchDatabaseSequence& Sequence : PoseSearchDatabase->Sequences)
+		FPoseSearchDatabaseSequence& DbSequence = PoseSearchDatabase->Sequences[SequenceIdx];
+		FPoseSearchDatabaseGroup& DbGroup = PoseSearchDatabase->Groups[GroupIdx];
+		if (DbGroup.Tag.IsValid())
 		{
-			Sequence.GroupTags.RemoveTag(DbGroup.Tag);
+			DbSequence.GroupTags.RemoveTag(DbGroup.Tag);
 		}
+	}
 
-		for (FPoseSearchDatabaseBlendSpace& BlendSpace : PoseSearchDatabase->BlendSpaces)
+	void FDatabaseViewModel::DeleteBlendSpaceFromDatabase(int32 BlendSpaceIdx)
+	{
+		PoseSearchDatabase->BlendSpaces.RemoveAt(BlendSpaceIdx);
+	}
+
+	void FDatabaseViewModel::RemoveBlendSpaceFromGroup(int32 BlendSpaceIdx, int32 GroupIdx)
+	{
+		FPoseSearchDatabaseBlendSpace& DbBlendSpace = PoseSearchDatabase->BlendSpaces[BlendSpaceIdx];
+		FPoseSearchDatabaseGroup& DbGroup = PoseSearchDatabase->Groups[GroupIdx];
+		if (DbGroup.Tag.IsValid())
 		{
-			BlendSpace.GroupTags.RemoveTag(DbGroup.Tag);
+			DbBlendSpace.GroupTags.RemoveTag(DbGroup.Tag);
 		}
 	}
 
-	PoseSearchDatabase->Groups.RemoveAt(GroupIdx);
-}
+	void FDatabaseViewModel::DeleteGroup(int32 GroupIdx)
+	{
+		FPoseSearchDatabaseGroup& DbGroup = PoseSearchDatabase->Groups[GroupIdx];
+
+		if (DbGroup.Tag.IsValid())
+		{
+			for (FPoseSearchDatabaseSequence& Sequence : PoseSearchDatabase->Sequences)
+			{
+				Sequence.GroupTags.RemoveTag(DbGroup.Tag);
+			}
+
+			for (FPoseSearchDatabaseBlendSpace& BlendSpace : PoseSearchDatabase->BlendSpaces)
+			{
+				BlendSpace.GroupTags.RemoveTag(DbGroup.Tag);
+			}
+		}
 
+		PoseSearchDatabase->Groups.RemoveAt(GroupIdx);
+	}
+}
diff --git a/Engine/Plugins/Experimental/Animation/PoseSearch/Source/Editor/Private/PoseSearchDatabaseViewportClient.cpp b/Engine/Plugins/Experimental/Animation/PoseSearch/Source/Editor/Private/PoseSearchDatabaseViewportClient.cpp
--- a/Engine/Plugins/Experimental/Animation/PoseSearch/Source/Editor/Private/PoseSearchDatabaseViewportClient.cpp
+++ b/Engine/Plugins/Experimental/Animation/PoseSearch/Source/Editor/Private/PoseSearchDatabaseViewportClient.cpp
@@ -8,39 +8,42 @@
 #include "UnrealWidget.h"
 #include "PoseSearchDatabaseEdMode.h"
 
-FPoseSearchDatabaseViewportClient::FPoseSearchDatabaseViewportClient(
-	const TSharedRef<FPoseSearchDatabasePreviewScene>& InPreviewScene,
-	const TSharedRef<SPoseSearchDatabaseViewport>& InViewport, 
-	const TSharedRef<FPoseSearchDatabaseEditorToolkit>& InAssetEditorToolkit)
-	: FEditorViewportClient(nullptr, &InPreviewScene.Get(), StaticCastSharedRef<SEditorViewport>(InViewport))
-	, PreviewScenePtr(InPreviewScene)
-	, AssetEditorToolkitPtr(InAssetEditorToolkit)
+namespace UE::PoseSearch
 {
-	Widget->SetUsesEditorModeTools(ModeTools.Get());
-	StaticCastSharedPtr<FAssetEditorModeManager>(ModeTools)->SetPreviewScene(&InPreviewScene.Get());
-	ModeTools->SetDefaultMode(FPoseSearchDatabaseEdMode::EdModeId);
+	FDatabaseViewportClient::FDatabaseViewportClient(
+		const TSharedRef<FDatabasePreviewScene>& InPreviewScene,
+		const TSharedRef<SDatabaseViewport>& InViewport,
+		const TSharedRef<FDatabaseEditorToolkit>& InAssetEditorToolkit)
+		: FEditorViewportClient(nullptr, &InPreviewScene.Get(), StaticCastSharedRef<SEditorViewport>(InViewport))
+		, PreviewScenePtr(InPreviewScene)
+		, AssetEditorToolkitPtr(InAssetEditorToolkit)
+	{
+		Widget->SetUsesEditorModeTools(ModeTools.Get());
+		StaticCastSharedPtr<FAssetEditorModeManager>(ModeTools)->SetPreviewScene(&InPreviewScene.Get());
+		ModeTools->SetDefaultMode(FDatabaseEdMode::EdModeId);
 
-	SetRealtime(true);
+		SetRealtime(true);
 
-	SetWidgetCoordSystemSpace(COORD_Local);
-	ModeTools->SetWidgetMode(UE::Widget::WM_Translate);
-}
+		SetWidgetCoordSystemSpace(COORD_Local);
+		ModeTools->SetWidgetMode(UE::Widget::WM_Translate);
+	}
 
-void FPoseSearchDatabaseViewportClient::TrackingStarted(
-	const struct FInputEventState& InInputState, 
-	bool bIsDraggingWidget, 
-	bool bNudge)
-{
-	ModeTools->StartTracking(this, Viewport);
-}
+	void FDatabaseViewportClient::TrackingStarted(
+		const struct FInputEventState& InInputState,
+		bool bIsDraggingWidget,
+		bool bNudge)
+	{
+		ModeTools->StartTracking(this, Viewport);
+	}
 
-void FPoseSearchDatabaseViewportClient::TrackingStopped()
-{
-	ModeTools->EndTracking(this, Viewport);
-	Invalidate();
-}
+	void FDatabaseViewportClient::TrackingStopped()
+	{
+		ModeTools->EndTracking(this, Viewport);
+		Invalidate();
+	}
 
-void FPoseSearchDatabaseViewportClient::Draw(const FSceneView* View, FPrimitiveDrawInterface* PDI)
-{
-	FEditorViewportClient::Draw(View, PDI);
+	void FDatabaseViewportClient::Draw(const FSceneView* View, FPrimitiveDrawInterface* PDI)
+	{
+		FEditorViewportClient::Draw(View, PDI);
+	}
 }
diff --git a/Engine/Plugins/Experimental/Animation/PoseSearch/Source/Editor/Private/PoseSearchEditor.cpp b/Engine/Plugins/Experimental/Animation/PoseSearch/Source/Editor/Private/PoseSearchEditor.cpp
--- a/Engine/Plugins/Experimental/Animation/PoseSearch/Source/Editor/Private/PoseSearchEditor.cpp
+++ b/Engine/Plugins/Experimental/Animation/PoseSearch/Source/Editor/Private/PoseSearchEditor.cpp
@@ -26,7 +26,7 @@ DEFINE_LOG_CATEGORY(LogPoseSearchEditor);
 //////////////////////////////////////////////////////////////////////////
 // FEditorCommands
 
-namespace UE { namespace PoseSearch {
+namespace UE::PoseSearch {
 
 struct FEditorCommands
 {
@@ -99,7 +99,7 @@ private:
 void FEditorModule::StartupModule()
 {
 	// Register Asset Editor Commands
-	FPoseSearchDatabaseEditorCommands::Register();
+	FDatabaseEditorCommands::Register();
 
 	if (GIsEditor && !IsRunningCommandlet())
 	{
@@ -118,13 +118,13 @@ void FEditorModule::StartupModule()
 		));
 
 		// Register Ed Mode used by pose search database
-		FEditorModeRegistry::Get().RegisterMode<FPoseSearchDatabaseEdMode>(
-			FPoseSearchDatabaseEdMode::EdModeId,
+		FEditorModeRegistry::Get().RegisterMode<FDatabaseEdMode>(
+			FDatabaseEdMode::EdModeId,
 			LOCTEXT("PoseSearchDatabaseEdModeName", "PoseSearchDatabase"));
 
 		// Register UPoseSearchDatabase Type Actions 
 		IAssetTools& AssetTools = FModuleManager::LoadModuleChecked<FAssetToolsModule>("AssetTools").Get();
-		PoseSearchDatabaseActions = MakeShared<FPoseSearchDatabaseTypeActions>();
+		PoseSearchDatabaseActions = MakeShared<FDatabaseTypeActions>();
 		AssetTools.RegisterAssetTypeActions(PoseSearchDatabaseActions.ToSharedRef());
 	}
 
@@ -151,12 +151,12 @@ void FEditorModule::ShutdownModule()
 	}
 
 	// Unregister Ed Mode
-	FEditorModeRegistry::Get().UnregisterMode(FPoseSearchDatabaseEdMode::EdModeId);
+	FEditorModeRegistry::Get().UnregisterMode(FDatabaseEdMode::EdModeId);
 	
 	UnregisterCustomizations();
 
 	// Unregister Asset Editor Commands
-	FPoseSearchDatabaseEditorCommands::Unregister();
+	FDatabaseEditorCommands::Unregister();
 	
 	IModularFeatures::Get().UnregisterModularFeature(TraceServices::ModuleFeatureName, TraceModule.Get());
 	FDebugger::Shutdown();
@@ -226,7 +226,7 @@ void FEditorModule::UnregisterCustomizations()
 }
 
 
-}} // namespace UE::PoseSearch
+} // namespace UE::PoseSearch
 
 IMPLEMENT_MODULE(UE::PoseSearch::FEditorModule, PoseSearchEditor);
 
diff --git a/Engine/Plugins/Experimental/Animation/PoseSearch/Source/Editor/Private/PoseSearchTypeActions.cpp b/Engine/Plugins/Experimental/Animation/PoseSearch/Source/Editor/Private/PoseSearchTypeActions.cpp
--- a/Engine/Plugins/Experimental/Animation/PoseSearch/Source/Editor/Private/PoseSearchTypeActions.cpp
+++ b/Engine/Plugins/Experimental/Animation/PoseSearch/Source/Editor/Private/PoseSearchTypeActions.cpp
@@ -6,41 +6,44 @@
 
 #define LOCTEXT_NAMESPACE "PoseSearchTypeActions"
 
-FText FPoseSearchDatabaseTypeActions::GetName() const
+namespace UE::PoseSearch
 {
-	return LOCTEXT("PoseSearchDatabaseTypeActionsName", "Pose Search Database");
-}
-
-FColor FPoseSearchDatabaseTypeActions::GetTypeColor() const
-{
-	return FColor(129, 196, 115);
-}
+	FText FDatabaseTypeActions::GetName() const
+	{
+		return LOCTEXT("PoseSearchDatabaseTypeActionsName", "Pose Search Database");
+	}
 
-UClass* FPoseSearchDatabaseTypeActions::GetSupportedClass() const
-{
-	return UPoseSearchDatabase::StaticClass();
-}
+	FColor FDatabaseTypeActions::GetTypeColor() const
+	{
+		return FColor(129, 196, 115);
+	}
 
-void FPoseSearchDatabaseTypeActions::OpenAssetEditor(
-	const TArray<UObject*>& InObjects, 
-	TSharedPtr<class IToolkitHost> EditWithinLevelEditor)
-{
-	const EToolkitMode::Type Mode = 
-		EditWithinLevelEditor.IsValid() ? EToolkitMode::WorldCentric : EToolkitMode::Standalone;
+	UClass* FDatabaseTypeActions::GetSupportedClass() const
+	{
+		return UPoseSearchDatabase::StaticClass();
+	}
 
-	for (auto ObjIt = InObjects.CreateConstIterator(); ObjIt; ++ObjIt)
+	void FDatabaseTypeActions::OpenAssetEditor(
+		const TArray<UObject*>& InObjects,
+		TSharedPtr<class IToolkitHost> EditWithinLevelEditor)
 	{
-		if (UPoseSearchDatabase* PoseSearchDb = Cast<UPoseSearchDatabase>(*ObjIt))
+		const EToolkitMode::Type Mode =
+			EditWithinLevelEditor.IsValid() ? EToolkitMode::WorldCentric : EToolkitMode::Standalone;
+
+		for (auto ObjIt = InObjects.CreateConstIterator(); ObjIt; ++ObjIt)
 		{
-			TSharedRef<FPoseSearchDatabaseEditorToolkit> NewEditor(new FPoseSearchDatabaseEditorToolkit());
-			NewEditor->InitAssetEditor(Mode, EditWithinLevelEditor, PoseSearchDb);
+			if (UPoseSearchDatabase* PoseSearchDb = Cast<UPoseSearchDatabase>(*ObjIt))
+			{
+				TSharedRef<FDatabaseEditorToolkit> NewEditor(new FDatabaseEditorToolkit());
+				NewEditor->InitAssetEditor(Mode, EditWithinLevelEditor, PoseSearchDb);
+			}
 		}
 	}
-}
 
-uint32 FPoseSearchDatabaseTypeActions::GetCategories()
-{
-	return EAssetTypeCategories::Animation;
+	uint32 FDatabaseTypeActions::GetCategories()
+	{
+		return EAssetTypeCategories::Animation;
+	}
 }
 
 #undef LOCTEXT_NAMESPACE
diff --git a/Engine/Plugins/Experimental/Animation/PoseSearch/Source/Editor/Private/SPoseSearchDatabaseAssetList.cpp b/Engine/Plugins/Experimental/Animation/PoseSearch/Source/Editor/Private/SPoseSearchDatabaseAssetList.cpp
--- a/Engine/Plugins/Experimental/Animation/PoseSearch/Source/Editor/Private/SPoseSearchDatabaseAssetList.cpp
+++ b/Engine/Plugins/Experimental/Animation/PoseSearch/Source/Editor/Private/SPoseSearchDatabaseAssetList.cpp
@@ -23,7 +23,7 @@ namespace UE::PoseSearch
 	FDatabaseAssetTreeNode::FDatabaseAssetTreeNode(
 		int32 InSourceAssetIdx,
 		ESearchIndexAssetType InSourceAssetType,
-		const TSharedRef<FPoseSearchDatabaseViewModel>& InEditorViewModel) 
+		const TSharedRef<FDatabaseViewModel>& InEditorViewModel) 
 		: SourceAssetIdx(InSourceAssetIdx)
 		, SourceAssetType(InSourceAssetType)
 		, EditorViewModel(InEditorViewModel)
@@ -46,7 +46,7 @@ namespace UE::PoseSearch
 
 	void SDatabaseAssetListItem::Construct(
 		const FArguments& InArgs,
-		const TSharedRef<FPoseSearchDatabaseViewModel>& InEditorViewModel,
+		const TSharedRef<FDatabaseViewModel>& InEditorViewModel,
 		const TSharedRef<STableViewBase>& OwnerTable,
 		TSharedRef<FDatabaseAssetTreeNode> InAssetTreeNode,
 		TSharedRef<FUICommandList> InCommandList,
@@ -122,7 +122,7 @@ namespace UE::PoseSearch
 		FText Name = LOCTEXT("None", "None");
 
 		TSharedPtr<FDatabaseAssetTreeNode> Node = WeakAssetTreeNode.Pin();
-		TSharedPtr<FPoseSearchDatabaseViewModel> ViewModel = EditorViewModel.Pin();
+		TSharedPtr<FDatabaseViewModel> ViewModel = EditorViewModel.Pin();
 		UPoseSearchDatabase* Database = ViewModel->GetPoseSearchDatabase();
 
 		if (Database)
@@ -292,7 +292,7 @@ namespace UE::PoseSearch
 
 	void SDatabaseAssetTree::Construct(
 		const FArguments& InArgs, 
-		TSharedRef<FPoseSearchDatabaseViewModel> InEditorViewModel)
+		TSharedRef<FDatabaseViewModel> InEditorViewModel)
 	{
 		EditorViewModel = InEditorViewModel;
 		CommandList = MakeShared<FUICommandList>();
@@ -384,13 +384,13 @@ namespace UE::PoseSearch
 
 	void SDatabaseAssetTree::RefreshTreeView(bool IsInitialSetup)
 	{
-		const TSharedPtr<FPoseSearchDatabaseViewModel> ViewModel = EditorViewModel.Pin();
+		const TSharedPtr<FDatabaseViewModel> ViewModel = EditorViewModel.Pin();
 		if (!ViewModel.IsValid())
 		{
 			return;
 		}
 
-		const TSharedRef<FPoseSearchDatabaseViewModel> ViewModelRef = ViewModel.ToSharedRef();
+		const TSharedRef<FDatabaseViewModel> ViewModelRef = ViewModel.ToSharedRef();
 
 		RootNodes.Reset();
 		AllNodes.Reset();
@@ -553,7 +553,7 @@ namespace UE::PoseSearch
 			return FReply::Unhandled();
 		}
 
-		TSharedPtr<FPoseSearchDatabaseViewModel> ViewModel = EditorViewModel.Pin();
+		TSharedPtr<FDatabaseViewModel> ViewModel = EditorViewModel.Pin();
 		if (!ViewModel.IsValid())
 		{
 			return FReply::Unhandled();
diff --git a/Engine/Plugins/Experimental/Animation/PoseSearch/Source/Editor/Private/SPoseSearchDatabaseViewport.cpp b/Engine/Plugins/Experimental/Animation/PoseSearch/Source/Editor/Private/SPoseSearchDatabaseViewport.cpp
--- a/Engine/Plugins/Experimental/Animation/PoseSearch/Source/Editor/Private/SPoseSearchDatabaseViewport.cpp
+++ b/Engine/Plugins/Experimental/Animation/PoseSearch/Source/Editor/Private/SPoseSearchDatabaseViewport.cpp
@@ -7,109 +7,113 @@
 #include "SPoseSearchDatabaseViewportToolbar.h"
 #include "PoseSearchDatabaseEditorCommands.h"
 
-void SPoseSearchDatabaseViewport::Construct(
-	const FArguments& InArgs, 
-	const FPoseSearchDatabaseViewportRequiredArgs& InRequiredArgs)
+namespace UE::PoseSearch
 {
-	PreviewScenePtr = InRequiredArgs.PreviewScene;
-	AssetEditorToolkitPtr = InRequiredArgs.AssetEditorToolkit;
-
-	SEditorViewport::Construct(
-		SEditorViewport::FArguments()
-		.IsEnabled(FSlateApplication::Get().GetNormalExecutionAttribute())
-		.AddMetaData<FTagMetaData>(TEXT("AnimationTools.Viewport"))
-	);
+	void SDatabaseViewport::Construct(
+		const FArguments& InArgs,
+		const FDatabaseViewportRequiredArgs& InRequiredArgs)
+	{
+		PreviewScenePtr = InRequiredArgs.PreviewScene;
+		AssetEditorToolkitPtr = InRequiredArgs.AssetEditorToolkit;
+
+		SEditorViewport::Construct(
+			SEditorViewport::FArguments()
+			.IsEnabled(FSlateApplication::Get().GetNormalExecutionAttribute())
+			.AddMetaData<FTagMetaData>(TEXT("AnimationTools.Viewport"))
+		);
+	}
+
+	void SDatabaseViewport::BindCommands()
+	{
+		SEditorViewport::BindCommands();
+
+		const FDatabaseEditorCommands& Commands = FDatabaseEditorCommands::Get();
+
+		TSharedRef<FDatabaseViewModel> ViewModelRef =
+			AssetEditorToolkitPtr.Pin()->GetViewModelSharedPtr().ToSharedRef();
+
+		CommandList->MapAction(
+			Commands.ShowPoseFeaturesNone,
+			FExecuteAction::CreateSP(
+				ViewModelRef,
+				&FDatabaseViewModel::OnSetPoseFeaturesDrawMode,
+				EFeaturesDrawMode::None),
+			FCanExecuteAction(),
+			FIsActionChecked::CreateSP(
+				ViewModelRef,
+				&FDatabaseViewModel::IsPoseFeaturesDrawMode,
+				EFeaturesDrawMode::None));
+
+		CommandList->MapAction(
+			Commands.ShowPoseFeaturesAll,
+			FExecuteAction::CreateSP(
+				ViewModelRef,
+				&FDatabaseViewModel::OnSetPoseFeaturesDrawMode,
+				EFeaturesDrawMode::All),
+			FCanExecuteAction(),
+			FIsActionChecked::CreateSP(
+				ViewModelRef,
+				&FDatabaseViewModel::IsPoseFeaturesDrawMode,
+				EFeaturesDrawMode::All));
+
+		CommandList->MapAction(
+			Commands.ShowAnimationOriginalOnly,
+			FExecuteAction::CreateSP(
+				ViewModelRef,
+				&FDatabaseViewModel::OnSetAnimationPreviewMode,
+				EAnimationPreviewMode::OriginalOnly),
+			FCanExecuteAction(),
+			FIsActionChecked::CreateSP(
+				ViewModelRef,
+				&FDatabaseViewModel::IsAnimationPreviewMode,
+				EAnimationPreviewMode::OriginalOnly));
+
+		CommandList->MapAction(
+			Commands.ShowAnimationOriginalAndMirrored,
+			FExecuteAction::CreateSP(
+				ViewModelRef,
+				&FDatabaseViewModel::OnSetAnimationPreviewMode,
+				EAnimationPreviewMode::OriginalAndMirrored),
+			FCanExecuteAction(),
+			FIsActionChecked::CreateSP(
+				ViewModelRef,
+				&FDatabaseViewModel::IsAnimationPreviewMode,
+				EAnimationPreviewMode::OriginalAndMirrored));
+	}
+
+	TSharedRef<FEditorViewportClient> SDatabaseViewport::MakeEditorViewportClient()
+	{
+		ViewportClient = MakeShared<FDatabaseViewportClient>(
+			PreviewScenePtr.Pin().ToSharedRef(),
+			SharedThis(this),
+			AssetEditorToolkitPtr.Pin().ToSharedRef());
+		ViewportClient->ViewportType = LVT_Perspective;
+		ViewportClient->bSetListenerPosition = false;
+		ViewportClient->SetViewLocation(EditorViewportDefs::DefaultPerspectiveViewLocation);
+		ViewportClient->SetViewRotation(EditorViewportDefs::DefaultPerspectiveViewRotation);
+
+		return ViewportClient.ToSharedRef();
+	}
+
+	TSharedPtr<SWidget> SDatabaseViewport::MakeViewportToolbar()
+	{
+		return SAssignNew(ViewportToolbar, SPoseSearchDatabaseViewportToolBar, SharedThis(this));
+	}
+
+
+	TSharedRef<SEditorViewport> SDatabaseViewport::GetViewportWidget()
+	{
+		return SharedThis(this);
+	}
+
+	TSharedPtr<FExtender> SDatabaseViewport::GetExtenders() const
+	{
+		TSharedPtr<FExtender> Result(MakeShareable(new FExtender));
+		return Result;
+	}
+
+	void SDatabaseViewport::OnFloatingButtonClicked()
+	{
+	}
 }
 
-void SPoseSearchDatabaseViewport::BindCommands()
-{
-	SEditorViewport::BindCommands();
-
-	const FPoseSearchDatabaseEditorCommands& Commands = FPoseSearchDatabaseEditorCommands::Get();
-
-	TSharedRef<FPoseSearchDatabaseViewModel> ViewModelRef = 
-		AssetEditorToolkitPtr.Pin()->GetViewModelSharedPtr().ToSharedRef();
-
-	CommandList->MapAction(
-		Commands.ShowPoseFeaturesNone,
-		FExecuteAction::CreateSP(
-			ViewModelRef, 
-			&FPoseSearchDatabaseViewModel::OnSetPoseFeaturesDrawMode, 
-			EPoseSearchFeaturesDrawMode::None),
-		FCanExecuteAction(),
-		FIsActionChecked::CreateSP(
-			ViewModelRef, 
-			&FPoseSearchDatabaseViewModel::IsPoseFeaturesDrawMode, 
-			EPoseSearchFeaturesDrawMode::None));
-
-	CommandList->MapAction(
-		Commands.ShowPoseFeaturesAll,
-		FExecuteAction::CreateSP(
-			ViewModelRef,
-			&FPoseSearchDatabaseViewModel::OnSetPoseFeaturesDrawMode,
-			EPoseSearchFeaturesDrawMode::All),
-		FCanExecuteAction(),
-		FIsActionChecked::CreateSP(
-			ViewModelRef,
-			&FPoseSearchDatabaseViewModel::IsPoseFeaturesDrawMode,
-			EPoseSearchFeaturesDrawMode::All));
-
-	CommandList->MapAction(
-		Commands.ShowAnimationOriginalOnly,
-		FExecuteAction::CreateSP(
-			ViewModelRef,
-			&FPoseSearchDatabaseViewModel::OnSetAnimationPreviewMode,
-			EAnimationPreviewMode::OriginalOnly),
-		FCanExecuteAction(),
-		FIsActionChecked::CreateSP(
-			ViewModelRef,
-			&FPoseSearchDatabaseViewModel::IsAnimationPreviewMode,
-			EAnimationPreviewMode::OriginalOnly));
-
-	CommandList->MapAction(
-		Commands.ShowAnimationOriginalAndMirrored,
-		FExecuteAction::CreateSP(
-			ViewModelRef,
-			&FPoseSearchDatabaseViewModel::OnSetAnimationPreviewMode,
-			EAnimationPreviewMode::OriginalAndMirrored),
-		FCanExecuteAction(),
-		FIsActionChecked::CreateSP(
-			ViewModelRef,
-			&FPoseSearchDatabaseViewModel::IsAnimationPreviewMode,
-			EAnimationPreviewMode::OriginalAndMirrored));
-}
-
-TSharedRef<FEditorViewportClient> SPoseSearchDatabaseViewport::MakeEditorViewportClient()
-{
-	ViewportClient = MakeShared<FPoseSearchDatabaseViewportClient>(
-		PreviewScenePtr.Pin().ToSharedRef(), 
-		SharedThis(this), 
-		AssetEditorToolkitPtr.Pin().ToSharedRef());
-	ViewportClient->ViewportType = LVT_Perspective;
-	ViewportClient->bSetListenerPosition = false;
-	ViewportClient->SetViewLocation(EditorViewportDefs::DefaultPerspectiveViewLocation);
-	ViewportClient->SetViewRotation(EditorViewportDefs::DefaultPerspectiveViewRotation);
-
-	return ViewportClient.ToSharedRef();
-}
-
-TSharedPtr<SWidget> SPoseSearchDatabaseViewport::MakeViewportToolbar()
-{
-	return SAssignNew(ViewportToolbar, SPoseSearchDatabaseViewportToolBar, SharedThis(this));
-}
-
-
-TSharedRef<SEditorViewport> SPoseSearchDatabaseViewport::GetViewportWidget()
-{
-	return SharedThis(this);
-}
-
-TSharedPtr<FExtender> SPoseSearchDatabaseViewport::GetExtenders() const
-{
-	TSharedPtr<FExtender> Result(MakeShareable(new FExtender));
-	return Result;
-}
-
-void SPoseSearchDatabaseViewport::OnFloatingButtonClicked()
-{
-}
diff --git a/Engine/Plugins/Experimental/Animation/PoseSearch/Source/Editor/Private/SPoseSearchDatabaseViewportToolbar.cpp b/Engine/Plugins/Experimental/Animation/PoseSearch/Source/Editor/Private/SPoseSearchDatabaseViewportToolbar.cpp
--- a/Engine/Plugins/Experimental/Animation/PoseSearch/Source/Editor/Private/SPoseSearchDatabaseViewportToolbar.cpp
+++ b/Engine/Plugins/Experimental/Animation/PoseSearch/Source/Editor/Private/SPoseSearchDatabaseViewportToolbar.cpp
@@ -7,32 +7,34 @@
 
 #define LOCTEXT_NAMESPACE "PoseSearchDatabaseViewportToolBar"
 
-void SPoseSearchDatabaseViewportToolBar::Construct(
-	const FArguments& InArgs, 
-	TSharedPtr<SPoseSearchDatabaseViewport> InViewport)
+namespace UE::PoseSearch
 {
-	SCommonEditorViewportToolbarBase::Construct(
-		SCommonEditorViewportToolbarBase::FArguments()
-		.AddRealtimeButton(false)
-		.PreviewProfileController(MakeShared<FPreviewProfileController>()), 
-		InViewport);
-}
+	void SPoseSearchDatabaseViewportToolBar::Construct(
+		const FArguments& InArgs,
+		TSharedPtr<SDatabaseViewport> InViewport)
+	{
+		SCommonEditorViewportToolbarBase::Construct(
+			SCommonEditorViewportToolbarBase::FArguments()
+			.AddRealtimeButton(false)
+			.PreviewProfileController(MakeShared<FPreviewProfileController>()),
+			InViewport);
+	}
 
-TSharedRef<SWidget> SPoseSearchDatabaseViewportToolBar::GenerateShowMenu() const
-{
-	GetInfoProvider().OnFloatingButtonClicked();
+	TSharedRef<SWidget> SPoseSearchDatabaseViewportToolBar::GenerateShowMenu() const
+	{
+		GetInfoProvider().OnFloatingButtonClicked();
 
-	TSharedRef<SEditorViewport> ViewportRef = GetInfoProvider().GetViewportWidget();
+		TSharedRef<SEditorViewport> ViewportRef = GetInfoProvider().GetViewportWidget();
 
-	const bool bInShouldCloseWindowAfterMenuSelection = true;
-	FMenuBuilder ShowMenuBuilder(bInShouldCloseWindowAfterMenuSelection, ViewportRef->GetCommandList());
-	{
-		ShowMenuBuilder.AddSubMenu(
-			LOCTEXT("ShowMenu_PoseFeaturesDrawSubMenu", "Pose Features"),
-			LOCTEXT("ShowMenu_PoseFeaturesDrawSubMenuToolTip", "Pose Feature Drawing Options"),
-			FNewMenuDelegate::CreateLambda([](FMenuBuilder& SubMenuBuilder)
+		const bool bInShouldCloseWindowAfterMenuSelection = true;
+		FMenuBuilder ShowMenuBuilder(bInShouldCloseWindowAfterMenuSelection, ViewportRef->GetCommandList());
+		{
+			ShowMenuBuilder.AddSubMenu(
+				LOCTEXT("ShowMenu_PoseFeaturesDrawSubMenu", "Pose Features"),
+				LOCTEXT("ShowMenu_PoseFeaturesDrawSubMenuToolTip", "Pose Feature Drawing Options"),
+				FNewMenuDelegate::CreateLambda([](FMenuBuilder& SubMenuBuilder)
 			{
-				const FPoseSearchDatabaseEditorCommands& Commands = FPoseSearchDatabaseEditorCommands::Get();
+				const FDatabaseEditorCommands& Commands = FDatabaseEditorCommands::Get();
 
 				SubMenuBuilder.BeginSection("PoseFeatures", LOCTEXT("ShowMenu_PoseFeaturesLabel", "Pose Features"));
 				{
@@ -41,14 +43,14 @@ TSharedRef<SWidget> SPoseSearchDatabaseViewportToolBar::GenerateShowMenu() const
 				}
 				SubMenuBuilder.EndSection();
 			})
-		);
+			);
 
-		ShowMenuBuilder.AddSubMenu(
-			LOCTEXT("ShowMenu_AnimationsSubMenu", "Animations"),
-			LOCTEXT("ShowMenu_AnimationsSubMenuToolTip", "Animation Preview Options"),
-			FNewMenuDelegate::CreateLambda([](FMenuBuilder& SubMenuBuilder)
+			ShowMenuBuilder.AddSubMenu(
+				LOCTEXT("ShowMenu_AnimationsSubMenu", "Animations"),
+				LOCTEXT("ShowMenu_AnimationsSubMenuToolTip", "Animation Preview Options"),
+				FNewMenuDelegate::CreateLambda([](FMenuBuilder& SubMenuBuilder)
 			{
-				const FPoseSearchDatabaseEditorCommands& Commands = FPoseSearchDatabaseEditorCommands::Get();
+				const FDatabaseEditorCommands& Commands = FDatabaseEditorCommands::Get();
 
 				SubMenuBuilder.BeginSection("Animations", LOCTEXT("ShowMenu_AnimationsLabel", "Animations"));
 				{
@@ -57,11 +59,12 @@ TSharedRef<SWidget> SPoseSearchDatabaseViewportToolBar::GenerateShowMenu() const
 				}
 				SubMenuBuilder.EndSection();
 			})
-		);
+			);
 
-	}
+		}
 
-	return ShowMenuBuilder.MakeWidget();
+		return ShowMenuBuilder.MakeWidget();
+	}
 }
 
 #undef LOCTEXT_NAMESPACE
diff --git a/Engine/Plugins/Experimental/Animation/PoseSearch/Source/Editor/Public/PoseSearchDatabaseEdMode.h b/Engine/Plugins/Experimental/Animation/PoseSearch/Source/Editor/Public/PoseSearchDatabaseEdMode.h
--- a/Engine/Plugins/Experimental/Animation/PoseSearch/Source/Editor/Public/PoseSearchDatabaseEdMode.h
+++ b/Engine/Plugins/Experimental/Animation/PoseSearch/Source/Editor/Public/PoseSearchDatabaseEdMode.h
@@ -5,43 +5,47 @@
 #include "CoreMinimal.h"
 #include "EdMode.h"
 
-class FPoseSearchDatabaseViewModel;
 
-class FPoseSearchDatabaseEdMode : public FEdMode
+namespace UE::PoseSearch
 {
-public:
-
-	const static FEditorModeID EdModeId;
-
-	FPoseSearchDatabaseEdMode();
-	virtual ~FPoseSearchDatabaseEdMode();
-
-	virtual void Tick(FEditorViewportClient* ViewportClient, float DeltaTime) override;
-	virtual void Render(const FSceneView* View, FViewport* Viewport, FPrimitiveDrawInterface* PDI) override;
-	virtual bool HandleClick(
-		FEditorViewportClient* InViewportClient, 
-		HHitProxy* HitProxy, 
-		const FViewportClick& Click) override;
-	virtual bool StartTracking(FEditorViewportClient* InViewportClient, FViewport* InViewport) override;
-	virtual bool EndTracking(FEditorViewportClient* InViewportClient, FViewport* InViewport) override;
-	virtual bool InputDelta(
-		FEditorViewportClient* InViewportClient, 
-		FViewport* InViewport, 
-		FVector& InDrag, 
-		FRotator& InRot, 
-		FVector& InScale) override;
-	virtual bool InputKey(
-		FEditorViewportClient* ViewportClient, 
-		FViewport* Viewport, 
-		FKey Key, 
-		EInputEvent Event) override;
-	virtual bool AllowWidgetMove() override;
-	virtual bool ShouldDrawWidget() const override;
-	virtual FVector GetWidgetLocation() const override;
-	virtual bool GetCustomDrawingCoordinateSystem(FMatrix& InMatrix, void* InData) override;
-	virtual bool GetCustomInputCoordinateSystem(FMatrix& InMatrix, void* InData) override;
-
-private:
-
-	FPoseSearchDatabaseViewModel* ViewModel = nullptr;
-};
+	class FDatabaseViewModel;
+
+	class FDatabaseEdMode : public FEdMode
+	{
+	public:
+
+		const static FEditorModeID EdModeId;
+
+		FDatabaseEdMode();
+		virtual ~FDatabaseEdMode();
+
+		virtual void Tick(FEditorViewportClient* ViewportClient, float DeltaTime) override;
+		virtual void Render(const FSceneView* View, FViewport* Viewport, FPrimitiveDrawInterface* PDI) override;
+		virtual bool HandleClick(
+			FEditorViewportClient* InViewportClient,
+			HHitProxy* HitProxy,
+			const FViewportClick& Click) override;
+		virtual bool StartTracking(FEditorViewportClient* InViewportClient, FViewport* InViewport) override;
+		virtual bool EndTracking(FEditorViewportClient* InViewportClient, FViewport* InViewport) override;
+		virtual bool InputDelta(
+			FEditorViewportClient* InViewportClient,
+			FViewport* InViewport,
+			FVector& InDrag,
+			FRotator& InRot,
+			FVector& InScale) override;
+		virtual bool InputKey(
+			FEditorViewportClient* ViewportClient,
+			FViewport* Viewport,
+			FKey Key,
+			EInputEvent Event) override;
+		virtual bool AllowWidgetMove() override;
+		virtual bool ShouldDrawWidget() const override;
+		virtual FVector GetWidgetLocation() const override;
+		virtual bool GetCustomDrawingCoordinateSystem(FMatrix& InMatrix, void* InData) override;
+		virtual bool GetCustomInputCoordinateSystem(FMatrix& InMatrix, void* InData) override;
+
+	private:
+
+		FDatabaseViewModel* ViewModel = nullptr;
+	};
+}
diff --git a/Engine/Plugins/Experimental/Animation/PoseSearch/Source/Editor/Public/PoseSearchDatabaseEditorCommands.h b/Engine/Plugins/Experimental/Animation/PoseSearch/Source/Editor/Public/PoseSearchDatabaseEditorCommands.h
--- a/Engine/Plugins/Experimental/Animation/PoseSearch/Source/Editor/Public/PoseSearchDatabaseEditorCommands.h
+++ b/Engine/Plugins/Experimental/Animation/PoseSearch/Source/Editor/Public/PoseSearchDatabaseEditorCommands.h
@@ -6,33 +6,38 @@
 #include "Framework/Commands/Commands.h"
 #include "EditorStyleSet.h"
 
-class FPoseSearchDatabaseEditorCommands : public TCommands<FPoseSearchDatabaseEditorCommands>
+
+namespace UE::PoseSearch
 {
-public:
-	FPoseSearchDatabaseEditorCommands()
-		: TCommands<FPoseSearchDatabaseEditorCommands>(
-			TEXT("PoseSearchDatabaseEditor"), 
-			NSLOCTEXT("Contexts", "PoseSearchDatabase", "Pose Search Database"), 
-			NAME_None, 
-			FEditorStyle::GetStyleSetName())
+	class FDatabaseEditorCommands : public TCommands<FDatabaseEditorCommands>
 	{
-	}
+	public:
+		FDatabaseEditorCommands()
+			: TCommands<FDatabaseEditorCommands>(
+				TEXT("PoseSearchDatabaseEditor"),
+				NSLOCTEXT("Contexts", "PoseSearchDatabase", "Pose Search Database"),
+				NAME_None,
+				FEditorStyle::GetStyleSetName())
+		{
+		}
+
+		virtual void RegisterCommands() override;
 
-	virtual void RegisterCommands() override;
+	public:
+		TSharedPtr<FUICommandInfo> StopPreviewScene;
 
-public:
-	TSharedPtr<FUICommandInfo> StopPreviewScene;
+		TSharedPtr<FUICommandInfo> ResetPreviewScene;
 
-	TSharedPtr<FUICommandInfo> ResetPreviewScene;
+		TSharedPtr<FUICommandInfo> BuildSearchIndex;
 
-	TSharedPtr<FUICommandInfo> BuildSearchIndex;
+		TSharedPtr<FUICommandInfo> ShowPoseFeaturesNone;
 
-	TSharedPtr<FUICommandInfo> ShowPoseFeaturesNone;
+		TSharedPtr<FUICommandInfo> ShowPoseFeaturesAll;
 
-	TSharedPtr<FUICommandInfo> ShowPoseFeaturesAll;
+		TSharedPtr<FUICommandInfo> ShowAnimationOriginalOnly;
 
-	TSharedPtr<FUICommandInfo> ShowAnimationOriginalOnly;
+		TSharedPtr<FUICommandInfo> ShowAnimationOriginalAndMirrored;
 
-	TSharedPtr<FUICommandInfo> ShowAnimationOriginalAndMirrored;
+	};
+}
 
-};
diff --git a/Engine/Plugins/Experimental/Animation/PoseSearch/Source/Editor/Public/PoseSearchDatabaseEditorToolkit.h b/Engine/Plugins/Experimental/Animation/PoseSearch/Source/Editor/Public/PoseSearchDatabaseEditorToolkit.h
--- a/Engine/Plugins/Experimental/Animation/PoseSearch/Source/Editor/Public/PoseSearchDatabaseEditorToolkit.h
+++ b/Engine/Plugins/Experimental/Animation/PoseSearch/Source/Editor/Public/PoseSearchDatabaseEditorToolkit.h
@@ -9,63 +9,64 @@
 class IDetailsView;
 class FToolBarBuilder;
 class UPoseSearchDatabase;
-class SPoseSearchDatabaseViewport;
-class FPoseSearchDatabasePreviewScene;
-class FPoseSearchDatabaseViewModel;
 
 namespace UE::PoseSearch
 {
+	class SDatabaseViewport;
+	class FDatabasePreviewScene;
+	class FDatabaseViewModel;
 	class SDatabaseAssetTree;
-}
 
-class FPoseSearchDatabaseEditorToolkit : public FAssetEditorToolkit, public FNotifyHook
-{
-public:
-
-	FPoseSearchDatabaseEditorToolkit();
-	virtual ~FPoseSearchDatabaseEditorToolkit();
-
-	void InitAssetEditor(
-		const EToolkitMode::Type Mode, 
-		const TSharedPtr<IToolkitHost>& InitToolkitHost, 
-		UPoseSearchDatabase* DatabaseAsset);
-
-	virtual void RegisterTabSpawners(const TSharedRef<FTabManager>& TabManager) override;
-	virtual void UnregisterTabSpawners(const TSharedRef<FTabManager>& TabManager) override;
-	virtual FName GetToolkitFName() const override;
-	virtual FText GetBaseToolkitName() const override;
-	virtual FText GetToolkitName() const override;
-	virtual FLinearColor GetWorldCentricTabColorScale() const override;
-	virtual FString GetWorldCentricTabPrefix() const override;
-
-	const UPoseSearchDatabase* GetPoseSearchDatabase() const;
-	FPoseSearchDatabaseViewModel* GetViewModel() const { return ViewModel.Get(); }
-	TSharedPtr<FPoseSearchDatabaseViewModel> GetViewModelSharedPtr() const { return ViewModel; }
-
-	void StopPreviewScene();
-	void ResetPreviewScene();
-	void BuildSearchIndex();
-
-private:
-	
-	TSharedRef<SDockTab> SpawnTab_Viewport(const FSpawnTabArgs& Args);
-	TSharedRef<SDockTab> SpawnTab_AssetDetails(const FSpawnTabArgs& Args);
-	TSharedRef<SDockTab> SpawnTab_PreviewSettings(const FSpawnTabArgs& Args);
+	class FDatabaseEditorToolkit : public FAssetEditorToolkit, public FNotifyHook
+	{
+	public:
+
+		FDatabaseEditorToolkit();
+		virtual ~FDatabaseEditorToolkit();
+
+		void InitAssetEditor(
+			const EToolkitMode::Type Mode,
+			const TSharedPtr<IToolkitHost>& InitToolkitHost,
+			UPoseSearchDatabase* DatabaseAsset);
+
+		virtual void RegisterTabSpawners(const TSharedRef<FTabManager>& TabManager) override;
+		virtual void UnregisterTabSpawners(const TSharedRef<FTabManager>& TabManager) override;
+		virtual FName GetToolkitFName() const override;
+		virtual FText GetBaseToolkitName() const override;
+		virtual FText GetToolkitName() const override;
+		virtual FLinearColor GetWorldCentricTabColorScale() const override;
+		virtual FString GetWorldCentricTabPrefix() const override;
+
+		const UPoseSearchDatabase* GetPoseSearchDatabase() const;
+		FDatabaseViewModel* GetViewModel() const { return ViewModel.Get(); }
+		TSharedPtr<FDatabaseViewModel> GetViewModelSharedPtr() const { return ViewModel; }
+
+		void StopPreviewScene();
+		void ResetPreviewScene();
+		void BuildSearchIndex();
+
+	private:
+
+		TSharedRef<SDockTab> SpawnTab_Viewport(const FSpawnTabArgs& Args);
+		TSharedRef<SDockTab> SpawnTab_AssetDetails(const FSpawnTabArgs& Args);
+		TSharedRef<SDockTab> SpawnTab_PreviewSettings(const FSpawnTabArgs& Args);
 	TSharedRef<SDockTab> SpawnTab_AssetTreeView(const FSpawnTabArgs& Args);
 
-	void BindCommands();
-	void ExtendToolbar();
-	void FillToolbar(FToolBarBuilder& ToolbarBuilder);
+		void BindCommands();
+		void ExtendToolbar();
+		void FillToolbar(FToolBarBuilder& ToolbarBuilder);
 
-	void OnFinishedChangingProperties(const FPropertyChangedEvent& PropertyChangedEvent);
+		void OnFinishedChangingProperties(const FPropertyChangedEvent& PropertyChangedEvent);
 
-	TSharedPtr<SPoseSearchDatabaseViewport> ViewportWidget;
+		TSharedPtr<SDatabaseViewport> ViewportWidget;
 
 	TSharedPtr<UE::PoseSearch::SDatabaseAssetTree> AssetTreeWidget;
 
-	TSharedPtr<IDetailsView> EditingAssetWidget;
+		TSharedPtr<IDetailsView> EditingAssetWidget;
 
-	TSharedPtr<FPoseSearchDatabasePreviewScene> PreviewScene;
+		TSharedPtr<FDatabasePreviewScene> PreviewScene;
+
+		TSharedPtr<FDatabaseViewModel> ViewModel;
+	};
+}
 
-	TSharedPtr<FPoseSearchDatabaseViewModel> ViewModel;
-};
diff --git a/Engine/Plugins/Experimental/Animation/PoseSearch/Source/Editor/Public/PoseSearchDatabasePreviewScene.h b/Engine/Plugins/Experimental/Animation/PoseSearch/Source/Editor/Public/PoseSearchDatabasePreviewScene.h
--- a/Engine/Plugins/Experimental/Animation/PoseSearch/Source/Editor/Public/PoseSearchDatabasePreviewScene.h
+++ b/Engine/Plugins/Experimental/Animation/PoseSearch/Source/Editor/Public/PoseSearchDatabasePreviewScene.h
@@ -5,26 +5,30 @@
 #include "CoreMinimal.h"
 #include "AdvancedPreviewScene.h"
 
-class FPoseSearchDatabaseEditorToolkit;
-
-class FPoseSearchDatabasePreviewScene : public FAdvancedPreviewScene
+namespace UE::PoseSearch
 {
-public:
+	class FDatabaseEditorToolkit;
 
-	FPoseSearchDatabasePreviewScene(
-		ConstructionValues CVs, 
-		const TSharedRef<FPoseSearchDatabaseEditorToolkit>& EditorToolkit);
-	~FPoseSearchDatabasePreviewScene(){}
+	class FDatabasePreviewScene : public FAdvancedPreviewScene
+	{
+	public:
 
-	virtual void Tick(float InDeltaTime) override;
+		FDatabasePreviewScene(
+			ConstructionValues CVs,
+			const TSharedRef<FDatabaseEditorToolkit>& EditorToolkit);
+		~FDatabasePreviewScene() {}
 
-	TSharedRef<FPoseSearchDatabaseEditorToolkit> GetEditorToolkit() const
-	{
-		return EditorToolkitPtr.Pin().ToSharedRef(); 
-	}
+		virtual void Tick(float InDeltaTime) override;
+
+		TSharedRef<FDatabaseEditorToolkit> GetEditorToolkit() const
+		{
+			return EditorToolkitPtr.Pin().ToSharedRef();
+		}
+
+	private:
 
-private:
+		/** The asset editor toolkit we are embedded in */
+		TWeakPtr<FDatabaseEditorToolkit> EditorToolkitPtr;
+	};
+}
 
-	/** The asset editor toolkit we are embedded in */
-	TWeakPtr<FPoseSearchDatabaseEditorToolkit> EditorToolkitPtr;
-};
\ No newline at end of file
diff --git a/Engine/Plugins/Experimental/Animation/PoseSearch/Source/Editor/Public/PoseSearchDatabaseViewModel.h b/Engine/Plugins/Experimental/Animation/PoseSearch/Source/Editor/Public/PoseSearchDatabaseViewModel.h
--- a/Engine/Plugins/Experimental/Animation/PoseSearch/Source/Editor/Public/PoseSearchDatabaseViewModel.h
+++ b/Engine/Plugins/Experimental/Animation/PoseSearch/Source/Editor/Public/PoseSearchDatabaseViewModel.h
@@ -16,68 +16,70 @@ class UDebugSkelMeshComponent;
 class UAnimSequence;
 class UBlendSpace;
 
-enum class EPoseSearchFeaturesDrawMode : uint8
+namespace UE::PoseSearch
 {
-	None,
-	All
-};
+	enum class EFeaturesDrawMode : uint8
+	{
+		None,
+		All
+	};
 
-enum class EAnimationPreviewMode : uint8
-{
-	OriginalOnly,
-	OriginalAndMirrored
-};
+	enum class EAnimationPreviewMode : uint8
+	{
+		OriginalOnly,
+		OriginalAndMirrored
+	};
 
 
-struct FPoseSearchDatabasePreviewActor
-{
-public:
-	TWeakObjectPtr<AActor> Actor = nullptr;
-	TWeakObjectPtr<UDebugSkelMeshComponent> Mesh = nullptr;
-	TWeakObjectPtr<UAnimPreviewInstance> AnimInstance = nullptr;
-	const FPoseSearchIndexAsset* IndexAsset = nullptr;
-	int32 CurrentPoseIndex = INDEX_NONE;
-
-	bool IsValid()
+	struct FDatabasePreviewActor
 	{
-		const bool bIsValid = Actor.IsValid() && Mesh.IsValid() && AnimInstance.IsValid();
-		return  bIsValid;
-	}
-};
-
-class FPoseSearchDatabaseViewModel : public TSharedFromThis<FPoseSearchDatabaseViewModel>, public FGCObject
-{
-public:
+	public:
+		TWeakObjectPtr<AActor> Actor = nullptr;
+		TWeakObjectPtr<UDebugSkelMeshComponent> Mesh = nullptr;
+		TWeakObjectPtr<UAnimPreviewInstance> AnimInstance = nullptr;
+		const FPoseSearchIndexAsset* IndexAsset = nullptr;
+		int32 CurrentPoseIndex = INDEX_NONE;
+
+		bool IsValid()
+		{
+			const bool bIsValid = Actor.IsValid() && Mesh.IsValid() && AnimInstance.IsValid();
+			return  bIsValid;
+		}
+	};
+
+	class FDatabaseViewModel : public TSharedFromThis<FDatabaseViewModel>, public FGCObject
+	{
+	public:
 
-	FPoseSearchDatabaseViewModel();
-	virtual ~FPoseSearchDatabaseViewModel();
+		FDatabaseViewModel();
+		virtual ~FDatabaseViewModel();
 
-	// ~ FGCObject interface
-	virtual void AddReferencedObjects(FReferenceCollector& Collector) override;
-	virtual FString GetReferencerName() const override { return TEXT("FPoseSearchDatabaseViewModel"); }
+		// ~ FGCObject interface
+		virtual void AddReferencedObjects(FReferenceCollector& Collector) override;
+		virtual FString GetReferencerName() const override { return TEXT("FPoseSearchDatabaseViewModel"); }
 
-	void Initialize(
-		UPoseSearchDatabase* InPoseSearchDatabase, 
-		const TSharedRef<FPoseSearchDatabasePreviewScene>& InPreviewScene);
+		void Initialize(
+			UPoseSearchDatabase* InPoseSearchDatabase,
+			const TSharedRef<FDatabasePreviewScene>& InPreviewScene);
 
-	void RemovePreviewActors();
-	void ResetPreviewActors();
-	void RespawnPreviewActors();
-	void BuildSearchIndex();
+		void RemovePreviewActors();
+		void ResetPreviewActors();
+		void RespawnPreviewActors();
+		void BuildSearchIndex();
 
-	UPoseSearchDatabase* GetPoseSearchDatabase() const { return PoseSearchDatabase; }
-	void OnPreviewActorClassChanged();
+		UPoseSearchDatabase* GetPoseSearchDatabase() const { return PoseSearchDatabase; }
+		void OnPreviewActorClassChanged();
 
-	void Tick(float DeltaSeconds);
+		void Tick(float DeltaSeconds);
 
-	TArray<FPoseSearchDatabasePreviewActor>& GetPreviewActors() { return PreviewActors; }
-	const TArray<FPoseSearchDatabasePreviewActor>& GetPreviewActors() const { return PreviewActors; }
+		TArray<FDatabasePreviewActor>& GetPreviewActors() { return PreviewActors; }
+		const TArray<FDatabasePreviewActor>& GetPreviewActors() const { return PreviewActors; }
 
-	void OnSetPoseFeaturesDrawMode(EPoseSearchFeaturesDrawMode DrawMode);
-	bool IsPoseFeaturesDrawMode(EPoseSearchFeaturesDrawMode DrawMode) const;
+		void OnSetPoseFeaturesDrawMode(EFeaturesDrawMode DrawMode);
+		bool IsPoseFeaturesDrawMode(EFeaturesDrawMode DrawMode) const;
 
-	void OnSetAnimationPreviewMode(EAnimationPreviewMode PreviewMode);
-	bool IsAnimationPreviewMode(EAnimationPreviewMode PreviewMode) const;
+		void OnSetAnimationPreviewMode(EAnimationPreviewMode PreviewMode);
+		bool IsAnimationPreviewMode(EAnimationPreviewMode PreviewMode) const;
 
 	void AddSequenceToDatabase(UAnimSequence* AnimSequence, int InitialGroupIdx = -1);
 	void AddBlendSpaceToDatabase(UBlendSpace* BlendSpace, int InitialGroupIdx = -1);
@@ -90,31 +92,32 @@ public:
 	void DeleteGroup(int32 GroupIdx);
 
 
-private:
-	float PlayTime = 0.0f;
+	private:
+		float PlayTime = 0.0f;
 
-	/** Scene asset being viewed and edited by this view model. */
-	TObjectPtr<UPoseSearchDatabase> PoseSearchDatabase;
+		/** Scene asset being viewed and edited by this view model. */
+		TObjectPtr<UPoseSearchDatabase> PoseSearchDatabase;
 
-	/** Weak pointer to the PreviewScene */
-	TWeakPtr<FPoseSearchDatabasePreviewScene> PreviewScenePtr;
+		/** Weak pointer to the PreviewScene */
+		TWeakPtr<FDatabasePreviewScene> PreviewScenePtr;
 
-	/** Actors to be displayed in the preview viewport */
-	TArray<FPoseSearchDatabasePreviewActor> PreviewActors;
+		/** Actors to be displayed in the preview viewport */
+		TArray<FDatabasePreviewActor> PreviewActors;
 
-	/** What features to show in the viewport */
-	EPoseSearchFeaturesDrawMode PoseFeaturesDrawMode = EPoseSearchFeaturesDrawMode::None;
+		/** What features to show in the viewport */
+		EFeaturesDrawMode PoseFeaturesDrawMode = EFeaturesDrawMode::None;
 
-	/** What animations to show in the viewport */
-	EAnimationPreviewMode AnimationPreviewMode = EAnimationPreviewMode::OriginalOnly;
+		/** What animations to show in the viewport */
+		EAnimationPreviewMode AnimationPreviewMode = EAnimationPreviewMode::OriginalOnly;
 
-	UWorld* GetWorld() const;
+		UWorld* GetWorld() const;
 
-	UObject* GetPlaybackContext() const;
+		UObject* GetPlaybackContext() const;
 
-	FPoseSearchDatabasePreviewActor SpawnPreviewActor(const FPoseSearchIndexAsset& IndexAsset);
+		FDatabasePreviewActor SpawnPreviewActor(const FPoseSearchIndexAsset& IndexAsset);
 
-	void UpdatePreviewActors();
+		void UpdatePreviewActors();
 
-	FTransform MirrorRootMotion(FTransform RootMotion, const class UMirrorDataTable* MirrorDataTable);
-};
\ No newline at end of file
+		FTransform MirrorRootMotion(FTransform RootMotion, const class UMirrorDataTable* MirrorDataTable);
+	};
+}
diff --git a/Engine/Plugins/Experimental/Animation/PoseSearch/Source/Editor/Public/PoseSearchDatabaseViewportClient.h b/Engine/Plugins/Experimental/Animation/PoseSearch/Source/Editor/Public/PoseSearchDatabaseViewportClient.h
--- a/Engine/Plugins/Experimental/Animation/PoseSearch/Source/Editor/Public/PoseSearchDatabaseViewportClient.h
+++ b/Engine/Plugins/Experimental/Animation/PoseSearch/Source/Editor/Public/PoseSearchDatabaseViewportClient.h
@@ -5,42 +5,46 @@
 #include "CoreMinimal.h"
 #include "EditorViewportClient.h"
 
-class FPoseSearchDatabasePreviewScene;
-class SPoseSearchDatabaseViewport;
-class FPoseSearchDatabaseEditorToolkit;
 
-class FPoseSearchDatabaseViewportClient : public FEditorViewportClient
+namespace UE::PoseSearch
 {
-public:
-
-	FPoseSearchDatabaseViewportClient(
-		const TSharedRef<FPoseSearchDatabasePreviewScene>& InPreviewScene, 
-		const TSharedRef<SPoseSearchDatabaseViewport>& InViewport, 
-		const TSharedRef<FPoseSearchDatabaseEditorToolkit>& InAssetEditorToolkit);
-	virtual ~FPoseSearchDatabaseViewportClient(){}
-
-	// ~FEditorViewportClient interface
-	virtual void Draw(const FSceneView* View, FPrimitiveDrawInterface* PDI) override;
-	virtual void TrackingStarted(const struct FInputEventState& InInputState, bool bIsDragging, bool bNudge) override;
-	virtual void TrackingStopped() override;
-	// ~End of FEditorViewportClient interface
-
-	/** Get the preview scene we are viewing */
-	TSharedRef<FPoseSearchDatabasePreviewScene> GetPreviewScene() const
-	{
-		return PreviewScenePtr.Pin().ToSharedRef(); 
-	}
+	class FDatabasePreviewScene;
+	class SDatabaseViewport;
+	class FDatabaseEditorToolkit;
 
-	TSharedRef<FPoseSearchDatabaseEditorToolkit> GetAssetEditorToolkit() const
+	class FDatabaseViewportClient : public FEditorViewportClient
 	{
-		return AssetEditorToolkitPtr.Pin().ToSharedRef(); 
-	}
-
-private:
-
-	/** Preview scene we are viewing */
-	TWeakPtr<FPoseSearchDatabasePreviewScene> PreviewScenePtr;
-
-	/** Asset editor toolkit we are embedded in */
-	TWeakPtr<FPoseSearchDatabaseEditorToolkit> AssetEditorToolkitPtr;
-};
+	public:
+
+		FDatabaseViewportClient(
+			const TSharedRef<FDatabasePreviewScene>& InPreviewScene,
+			const TSharedRef<SDatabaseViewport>& InViewport,
+			const TSharedRef<FDatabaseEditorToolkit>& InAssetEditorToolkit);
+		virtual ~FDatabaseViewportClient() {}
+
+		// ~FEditorViewportClient interface
+		virtual void Draw(const FSceneView* View, FPrimitiveDrawInterface* PDI) override;
+		virtual void TrackingStarted(const struct FInputEventState& InInputState, bool bIsDragging, bool bNudge) override;
+		virtual void TrackingStopped() override;
+		// ~End of FEditorViewportClient interface
+
+		/** Get the preview scene we are viewing */
+		TSharedRef<FDatabasePreviewScene> GetPreviewScene() const
+		{
+			return PreviewScenePtr.Pin().ToSharedRef();
+		}
+
+		TSharedRef<FDatabaseEditorToolkit> GetAssetEditorToolkit() const
+		{
+			return AssetEditorToolkitPtr.Pin().ToSharedRef();
+		}
+
+	private:
+
+		/** Preview scene we are viewing */
+		TWeakPtr<FDatabasePreviewScene> PreviewScenePtr;
+
+		/** Asset editor toolkit we are embedded in */
+		TWeakPtr<FDatabaseEditorToolkit> AssetEditorToolkitPtr;
+	};
+}
diff --git a/Engine/Plugins/Experimental/Animation/PoseSearch/Source/Editor/Public/PoseSearchTypeActions.h b/Engine/Plugins/Experimental/Animation/PoseSearch/Source/Editor/Public/PoseSearchTypeActions.h
--- a/Engine/Plugins/Experimental/Animation/PoseSearch/Source/Editor/Public/PoseSearchTypeActions.h
+++ b/Engine/Plugins/Experimental/Animation/PoseSearch/Source/Editor/Public/PoseSearchTypeActions.h
@@ -5,16 +5,19 @@
 #include "CoreMinimal.h"
 #include "AssetTypeActions_Base.h"
 
-class FPoseSearchDatabaseTypeActions : public FAssetTypeActions_Base
+namespace UE::PoseSearch
 {
-public:
+	class FDatabaseTypeActions : public FAssetTypeActions_Base
+	{
+	public:
 
-	FPoseSearchDatabaseTypeActions(){}
-	virtual FText GetName() const override;
-	virtual FColor GetTypeColor() const override;
-	virtual UClass* GetSupportedClass() const override;
-	virtual void OpenAssetEditor(
-		const TArray<UObject*>& InObjects, 
-		TSharedPtr<class IToolkitHost> EditWithinLevelEditor = TSharedPtr<IToolkitHost>()) override;
-	virtual uint32 GetCategories() override;
-};
\ No newline at end of file
+		FDatabaseTypeActions() {}
+		virtual FText GetName() const override;
+		virtual FColor GetTypeColor() const override;
+		virtual UClass* GetSupportedClass() const override;
+		virtual void OpenAssetEditor(
+			const TArray<UObject*>& InObjects,
+			TSharedPtr<class IToolkitHost> EditWithinLevelEditor = TSharedPtr<IToolkitHost>()) override;
+		virtual uint32 GetCategories() override;
+	};
+}
diff --git a/Engine/Plugins/Experimental/Animation/PoseSearch/Source/Editor/Public/SPoseSearchDatabaseAssetList.h b/Engine/Plugins/Experimental/Animation/PoseSearch/Source/Editor/Public/SPoseSearchDatabaseAssetList.h
--- a/Engine/Plugins/Experimental/Animation/PoseSearch/Source/Editor/Public/SPoseSearchDatabaseAssetList.h
+++ b/Engine/Plugins/Experimental/Animation/PoseSearch/Source/Editor/Public/SPoseSearchDatabaseAssetList.h
@@ -7,10 +7,10 @@
 #include "Widgets/Views/STreeView.h"
 #include "EditorUndoClient.h"
 
-class FPoseSearchDatabaseViewModel;
 
 namespace UE::PoseSearch
 {
+	class FDatabaseViewModel;
 	class SDatabaseAssetTree;
 
 	class FDatabaseAssetTreeNode : public TSharedFromThis<FDatabaseAssetTreeNode>
@@ -20,7 +20,7 @@ namespace UE::PoseSearch
 		FDatabaseAssetTreeNode(
 			int32 InSourceAssetIdx,
 			ESearchIndexAssetType InSourceAssetType,
-			const TSharedRef<FPoseSearchDatabaseViewModel>& InEditorViewModel);
+			const TSharedRef<FDatabaseViewModel>& InEditorViewModel);
 
 		TSharedRef<ITableRow> MakeTreeRowWidget(
 			const TSharedRef<STableViewBase>& InOwnerTable,
@@ -35,7 +35,7 @@ namespace UE::PoseSearch
 
 	protected:
 
-		TWeakPtr<FPoseSearchDatabaseViewModel> EditorViewModel;
+		TWeakPtr<FDatabaseViewModel> EditorViewModel;
 
 	};
 
@@ -44,7 +44,7 @@ namespace UE::PoseSearch
 	public:
 		void Construct(
 			const FArguments& InArgs,
-			const TSharedRef<FPoseSearchDatabaseViewModel>& InEditorViewModel,
+			const TSharedRef<FDatabaseViewModel>& InEditorViewModel,
 			const TSharedRef<STableViewBase>& OwnerTable,
 			TSharedRef<FDatabaseAssetTreeNode> InAssetTreeNode,
 			TSharedRef<FUICommandList> InCommandList,
@@ -64,7 +64,7 @@ namespace UE::PoseSearch
 		void OnAddBlendSpace();
 
 		TWeakPtr<FDatabaseAssetTreeNode> WeakAssetTreeNode;
-		TWeakPtr<FPoseSearchDatabaseViewModel> EditorViewModel;
+		TWeakPtr<FDatabaseViewModel> EditorViewModel;
 		TWeakPtr<SDatabaseAssetTree> SkeletonView;
 	};
 
@@ -76,7 +76,7 @@ namespace UE::PoseSearch
 
 		virtual ~SDatabaseAssetTree();
 
-		void Construct(const FArguments& InArgs, TSharedRef<FPoseSearchDatabaseViewModel> InEditorViewModel);
+		void Construct(const FArguments& InArgs, TSharedRef<FDatabaseViewModel> InEditorViewModel);
 
 		// SWidget interface
 		virtual FReply OnDragOver(const FGeometry& MyGeometry, const FDragDropEvent& DragDropEvent) override;
@@ -86,7 +86,7 @@ namespace UE::PoseSearch
 		void RefreshTreeView(bool IsInitialSetup = false);
 
 	protected:
-		TWeakPtr<FPoseSearchDatabaseViewModel> EditorViewModel;
+		TWeakPtr<FDatabaseViewModel> EditorViewModel;
 
 		/** command list we bind to */
 		TSharedPtr<FUICommandList> CommandList;
diff --git a/Engine/Plugins/Experimental/Animation/PoseSearch/Source/Editor/Public/SPoseSearchDatabaseViewport.h b/Engine/Plugins/Experimental/Animation/PoseSearch/Source/Editor/Public/SPoseSearchDatabaseViewport.h
--- a/Engine/Plugins/Experimental/Animation/PoseSearch/Source/Editor/Public/SPoseSearchDatabaseViewport.h
+++ b/Engine/Plugins/Experimental/Animation/PoseSearch/Source/Editor/Public/SPoseSearchDatabaseViewport.h
@@ -6,58 +6,63 @@
 #include "SEditorViewport.h"
 #include "SCommonEditorViewportToolbarBase.h"
 
-class FPoseSearchDatabasePreviewScene;
-class FPoseSearchDatabaseViewportClient;
-class FPoseSearchDatabaseEditorToolkit;
-class SPoseSearchDatabaseViewportToolBar;
 
-struct FPoseSearchDatabaseViewportRequiredArgs
+namespace UE::PoseSearch
 {
-	FPoseSearchDatabaseViewportRequiredArgs(
-		const TSharedRef<FPoseSearchDatabaseEditorToolkit>& InAssetEditorToolkit, 
-		const TSharedRef<FPoseSearchDatabasePreviewScene>& InPreviewScene)
-		: AssetEditorToolkit(InAssetEditorToolkit)
-		, PreviewScene(InPreviewScene)
-	{}
+	class FDatabasePreviewScene;
+	class FDatabaseViewportClient;
+	class FDatabaseEditorToolkit;
+	class SPoseSearchDatabaseViewportToolBar;
 
-	TSharedRef<FPoseSearchDatabaseEditorToolkit> AssetEditorToolkit;
+	struct FDatabaseViewportRequiredArgs
+	{
+		FDatabaseViewportRequiredArgs(
+			const TSharedRef<FDatabaseEditorToolkit>& InAssetEditorToolkit,
+			const TSharedRef<FDatabasePreviewScene>& InPreviewScene)
+			: AssetEditorToolkit(InAssetEditorToolkit)
+			, PreviewScene(InPreviewScene)
+		{
+		}
 
-	TSharedRef<FPoseSearchDatabasePreviewScene> PreviewScene;
-};
+		TSharedRef<FDatabaseEditorToolkit> AssetEditorToolkit;
 
-class SPoseSearchDatabaseViewport : public SEditorViewport, public ICommonEditorViewportToolbarInfoProvider
-{
-public:
-	
-	SLATE_BEGIN_ARGS(SPoseSearchDatabaseViewport) {}
-	SLATE_END_ARGS();
+		TSharedRef<FDatabasePreviewScene> PreviewScene;
+	};
+
+	class SDatabaseViewport : public SEditorViewport, public ICommonEditorViewportToolbarInfoProvider
+	{
+	public:
+
+		SLATE_BEGIN_ARGS(SDatabaseViewport) {}
+		SLATE_END_ARGS();
 
-	void Construct(const FArguments& InArgs, const FPoseSearchDatabaseViewportRequiredArgs& InRequiredArgs);
-	virtual ~SPoseSearchDatabaseViewport(){}
+		void Construct(const FArguments& InArgs, const FDatabaseViewportRequiredArgs& InRequiredArgs);
+		virtual ~SDatabaseViewport() {}
 
-	// ~ICommonEditorViewportToolbarInfoProvider interface
-	virtual TSharedRef<class SEditorViewport> GetViewportWidget() override;
-	virtual TSharedPtr<FExtender> GetExtenders() const override;
-	virtual void OnFloatingButtonClicked() override;
-	// ~End of ICommonEditorViewportToolbarInfoProvider interface
+		// ~ICommonEditorViewportToolbarInfoProvider interface
+		virtual TSharedRef<class SEditorViewport> GetViewportWidget() override;
+		virtual TSharedPtr<FExtender> GetExtenders() const override;
+		virtual void OnFloatingButtonClicked() override;
+		// ~End of ICommonEditorViewportToolbarInfoProvider interface
 
-protected:
+	protected:
 
-	// ~SEditorViewport interface
-	virtual void BindCommands() override;
-	virtual TSharedRef<FEditorViewportClient> MakeEditorViewportClient() override;
-	virtual TSharedPtr<SWidget> MakeViewportToolbar() override;
-	// ~End of SEditorViewport interface
+		// ~SEditorViewport interface
+		virtual void BindCommands() override;
+		virtual TSharedRef<FEditorViewportClient> MakeEditorViewportClient() override;
+		virtual TSharedPtr<SWidget> MakeViewportToolbar() override;
+		// ~End of SEditorViewport interface
 
-	/** The viewport toolbar */
-	TSharedPtr<SPoseSearchDatabaseViewportToolBar> ViewportToolbar;
+		/** The viewport toolbar */
+		TSharedPtr<SPoseSearchDatabaseViewportToolBar> ViewportToolbar;
 
-	/** Viewport client */
-	TSharedPtr<FPoseSearchDatabaseViewportClient> ViewportClient;
+		/** Viewport client */
+		TSharedPtr<FDatabaseViewportClient> ViewportClient;
 
-	/** The preview scene that we are viewing */
-	TWeakPtr<FPoseSearchDatabasePreviewScene> PreviewScenePtr;
+		/** The preview scene that we are viewing */
+		TWeakPtr<FDatabasePreviewScene> PreviewScenePtr;
 
-	/** Asset editor toolkit we are embedded in */
-	TWeakPtr<FPoseSearchDatabaseEditorToolkit> AssetEditorToolkitPtr;
-};
\ No newline at end of file
+		/** Asset editor toolkit we are embedded in */
+		TWeakPtr<FDatabaseEditorToolkit> AssetEditorToolkitPtr;
+	};
+}
diff --git a/Engine/Plugins/Experimental/Animation/PoseSearch/Source/Editor/Public/SPoseSearchDatabaseViewportToolbar.h b/Engine/Plugins/Experimental/Animation/PoseSearch/Source/Editor/Public/SPoseSearchDatabaseViewportToolbar.h
--- a/Engine/Plugins/Experimental/Animation/PoseSearch/Source/Editor/Public/SPoseSearchDatabaseViewportToolbar.h
+++ b/Engine/Plugins/Experimental/Animation/PoseSearch/Source/Editor/Public/SPoseSearchDatabaseViewportToolbar.h
@@ -5,19 +5,23 @@
 #include "CoreMinimal.h"
 #include "SCommonEditorViewportToolbarBase.h"
 
-class SPoseSearchDatabaseViewport;
 
-class SPoseSearchDatabaseViewportToolBar : public SCommonEditorViewportToolbarBase
+namespace UE::PoseSearch
 {
-public:
-	SLATE_BEGIN_ARGS(SPoseSearchDatabaseViewportToolBar)
-	{}
+	class SDatabaseViewport;
 
-	SLATE_END_ARGS()
+	class SPoseSearchDatabaseViewportToolBar : public SCommonEditorViewportToolbarBase
+	{
+	public:
+		SLATE_BEGIN_ARGS(SPoseSearchDatabaseViewportToolBar)
+		{}
 
-	void Construct(const FArguments& InArgs, TSharedPtr<SPoseSearchDatabaseViewport> InViewport);
+		SLATE_END_ARGS()
 
-	// ~SCommonEditorViewportToolbarBase interface
-	virtual TSharedRef<SWidget> GenerateShowMenu() const override;
-	// ~End of SCommonEditorViewportToolbarBase interface
-};
+		void Construct(const FArguments& InArgs, TSharedPtr<SDatabaseViewport> InViewport);
+
+		// ~SCommonEditorViewportToolbarBase interface
+		virtual TSharedRef<SWidget> GenerateShowMenu() const override;
+		// ~End of SCommonEditorViewportToolbarBase interface
+	};
+}
