PoseSearch: Fixed assert when motion matching node becomes relevant and no valid database has been set

diff --git a/Engine/Plugins/Experimental/Animation/PoseSearch/Source/Runtime/Private/PoseSearchLibrary.cpp b/Engine/Plugins/Experimental/Animation/PoseSearch/Source/Runtime/Private/PoseSearchLibrary.cpp
--- a/Engine/Plugins/Experimental/Animation/PoseSearch/Source/Runtime/Private/PoseSearchLibrary.cpp
+++ b/Engine/Plugins/Experimental/Animation/PoseSearch/Source/Runtime/Private/PoseSearchLibrary.cpp
@@ -215,19 +215,20 @@ bool FMotionMatchingState::InitNewDatabaseSearch(
 
 void FMotionMatchingState::Reset()
 {
-	check(CurrentDatabase.IsValid() && CurrentDatabase->IsValidForSearch());
-
 	DbPoseIdx = INDEX_NONE;
 	SearchIndexAssetIdx = INDEX_NONE;
 	AssetPlayerTime = 0.0f;
 	// Set the elapsed time to INFINITY to trigger a search right away
 	ElapsedPoseJumpTime = INFINITY;
 
-	if (!ComposedQuery.IsInitializedForSchema(CurrentDatabase->Schema))
+	if (CurrentDatabase.IsValid() && CurrentDatabase->IsValidForSearch())
 	{
 		ComposedQuery.Init(CurrentDatabase->Schema);
 	}
-
+	else
+	{
+		ComposedQuery.Reset();
+	}
 }
 
 void FMotionMatchingState::ComposeQuery(const UPoseSearchDatabase* Database, const FTrajectorySampleRange& Trajectory)
