[UE-150523] Fixed non-unity build

diff --git a/Engine/Plugins/Experimental/Animation/PoseSearch/Source/Editor/Public/PoseSearchDatabaseViewModel.h b/Engine/Plugins/Experimental/Animation/PoseSearch/Source/Editor/Public/PoseSearchDatabaseViewModel.h
--- a/Engine/Plugins/Experimental/Animation/PoseSearch/Source/Editor/Public/PoseSearchDatabaseViewModel.h
+++ b/Engine/Plugins/Experimental/Animation/PoseSearch/Source/Editor/Public/PoseSearchDatabaseViewModel.h
@@ -15,6 +15,7 @@ class UAnimPreviewInstance;
 class UDebugSkelMeshComponent;
 class UAnimSequence;
 class UBlendSpace;
+class UMirrorDataTable;
 
 namespace UE::PoseSearch
 {
@@ -81,15 +82,15 @@ namespace UE::PoseSearch
 		void OnSetAnimationPreviewMode(EAnimationPreviewMode PreviewMode);
 		bool IsAnimationPreviewMode(EAnimationPreviewMode PreviewMode) const;
 
-	void AddSequenceToDatabase(UAnimSequence* AnimSequence, int InitialGroupIdx = -1);
-	void AddBlendSpaceToDatabase(UBlendSpace* BlendSpace, int InitialGroupIdx = -1);
-	void AddGroupToDatabase();
+		void AddSequenceToDatabase(UAnimSequence* AnimSequence, int InitialGroupIdx = -1);
+		void AddBlendSpaceToDatabase(UBlendSpace* BlendSpace, int InitialGroupIdx = -1);
+		void AddGroupToDatabase();
 
-	void DeleteSequenceFromDatabase(int32 SequenceIdx);
-	void RemoveSequenceFromGroup(int32 SequenceIdx, int32 GroupIdx); 
-	void DeleteBlendSpaceFromDatabase(int32 BlendSpaceIdx);
-	void RemoveBlendSpaceFromGroup(int32 BlendSpaceIdx, int32 GroupIdx);
-	void DeleteGroup(int32 GroupIdx);
+		void DeleteSequenceFromDatabase(int32 SequenceIdx);
+		void RemoveSequenceFromGroup(int32 SequenceIdx, int32 GroupIdx); 
+		void DeleteBlendSpaceFromDatabase(int32 BlendSpaceIdx);
+		void RemoveBlendSpaceFromGroup(int32 BlendSpaceIdx, int32 GroupIdx);
+		void DeleteGroup(int32 GroupIdx);
 
 
 	private:
