PoseSearch: linked selected assets on asset tree to new asset details tab.

diff --git a/Engine/Plugins/Experimental/Animation/PoseSearch/Source/Editor/Private/PoseSearchDatabaseEditorReflection.cpp b/Engine/Plugins/Experimental/Animation/PoseSearch/Source/Editor/Private/PoseSearchDatabaseEditorReflection.cpp
new file mode 100644
--- /dev/null
+++ b/Engine/Plugins/Experimental/Animation/PoseSearch/Source/Editor/Private/PoseSearchDatabaseEditorReflection.cpp
@@ -0,0 +1,62 @@
+// Copyright Epic Games, Inc. All Rights Reserved.
+
+#include "PoseSearchDatabaseEditorReflection.h"
+#include "SPoseSearchDatabaseAssetList.h"
+#include "PoseSearchDatabaseViewModel.h"
+
+#if WITH_EDITOR
+
+void UPoseSearchDatabaseReflectionBase::SetSourceLink(
+	const TWeakPtr<UE::PoseSearch::FDatabaseAssetTreeNode>& InWeakAssetTreeNode,
+	const TSharedPtr<UE::PoseSearch::SDatabaseAssetTree>& InAssetTreeWidget)
+{
+	WeakAssetTreeNode = InWeakAssetTreeNode;
+	AssetTreeWidget = InAssetTreeWidget;
+}
+
+void UPoseSearchDatabaseSequenceReflection::PostEditChangeProperty(
+	struct FPropertyChangedEvent& PropertyChangedEvent)
+{
+	Super::PostEditChangeProperty(PropertyChangedEvent);
+
+	check(WeakAssetTreeNode.Pin()->SourceAssetType == ESearchIndexAssetType::Sequence);
+
+	UPoseSearchDatabase* Database = WeakAssetTreeNode.Pin()->EditorViewModel.Pin()->GetPoseSearchDatabase();
+	if (IsValid(Database))
+	{
+		Database->Sequences[WeakAssetTreeNode.Pin()->SourceAssetIdx] = Sequence;
+		AssetTreeWidget->RefreshTreeView(false, true);
+	}
+}
+
+void UPoseSearchDatabaseBlendSpaceReflection::PostEditChangeProperty(
+	struct FPropertyChangedEvent& PropertyChangedEvent)
+{
+	Super::PostEditChangeProperty(PropertyChangedEvent);
+
+	check(WeakAssetTreeNode.Pin()->SourceAssetType == ESearchIndexAssetType::BlendSpace);
+
+	UPoseSearchDatabase* Database = WeakAssetTreeNode.Pin()->EditorViewModel.Pin()->GetPoseSearchDatabase();
+	if (IsValid(Database))
+	{
+		Database->BlendSpaces[WeakAssetTreeNode.Pin()->SourceAssetIdx] = BlendSpace;
+		AssetTreeWidget->RefreshTreeView(false, true);
+	}
+}
+
+void UPoseSearchDatabaseGroupReflection::PostEditChangeProperty(
+	struct FPropertyChangedEvent& PropertyChangedEvent)
+{
+	Super::PostEditChangeProperty(PropertyChangedEvent);
+
+	check(WeakAssetTreeNode.Pin()->SourceAssetType == ESearchIndexAssetType::Invalid);
+
+	UPoseSearchDatabase* Database = WeakAssetTreeNode.Pin()->EditorViewModel.Pin()->GetPoseSearchDatabase();
+	if (IsValid(Database))
+	{
+		Database->Groups[WeakAssetTreeNode.Pin()->SourceAssetIdx] = Group;
+		AssetTreeWidget->RefreshTreeView(false, true);
+	}
+}
+
+#endif // WITH_EDITOR
diff --git a/Engine/Plugins/Experimental/Animation/PoseSearch/Source/Editor/Private/PoseSearchDatabaseEditorToolkit.cpp b/Engine/Plugins/Experimental/Animation/PoseSearch/Source/Editor/Private/PoseSearchDatabaseEditorToolkit.cpp
--- a/Engine/Plugins/Experimental/Animation/PoseSearch/Source/Editor/Private/PoseSearchDatabaseEditorToolkit.cpp
+++ b/Engine/Plugins/Experimental/Animation/PoseSearch/Source/Editor/Private/PoseSearchDatabaseEditorToolkit.cpp
@@ -6,6 +6,7 @@
 #include "PoseSearchDatabasePreviewScene.h"
 #include "PoseSearchDatabaseEditorCommands.h"
 #include "PoseSearchDatabaseViewModel.h"
+#include "PoseSearchDatabaseEditorReflection.h"
 #include "PoseSearchEditor.h"
 #include "PoseSearch/PoseSearch.h"
 #include "GameFramework/WorldSettings.h"
@@ -29,12 +30,14 @@ namespace UE::PoseSearch
 		static const FName ViewportID;
 		static const FName PreviewSettingsID;
 		static const FName AssetTreeViewID;
+		static const FName SelectionDetailsID;
 	};
 
 	const FName FDatabaseEditorTabs::AssetDetailsID(TEXT("PoseSearchDatabaseEditorAssetDetailsTabID"));
 	const FName FDatabaseEditorTabs::ViewportID(TEXT("PoseSearchDatabaseEditorViewportTabID"));
 	const FName FDatabaseEditorTabs::PreviewSettingsID(TEXT("PoseSearchDatabaseEditorPreviewSettingsTabID"));
 	const FName FDatabaseEditorTabs::AssetTreeViewID(TEXT("PoseSearchDatabaseEditorAssetTreeViewTabID"));
+	const FName FDatabaseEditorTabs::SelectionDetailsID(TEXT("PoseSearchDatabaseEditorSelectionDetailsID"));
 
 	FDatabaseEditorToolkit::FDatabaseEditorToolkit()
 	{
@@ -55,6 +58,11 @@ namespace UE::PoseSearch
 		return ViewModel.IsValid() ? ViewModel->GetPoseSearchDatabase() : nullptr;
 	}
 
+	UPoseSearchDatabase* FDatabaseEditorToolkit::GetPoseSearchDatabase()
+	{
+		return ViewModel.IsValid() ? ViewModel->GetPoseSearchDatabase() : nullptr;
+	}
+
 	void FDatabaseEditorToolkit::StopPreviewScene()
 	{
 		ViewModel->RemovePreviewActors();
@@ -103,17 +111,21 @@ namespace UE::PoseSearch
 			PreviewScene.ToSharedRef());
 		ViewportWidget = SNew(SDatabaseViewport, ViewportArgs);
 
-		AssetTreeWidget = SNew(UE::PoseSearch::SDatabaseAssetTree, ViewModel.ToSharedRef());
+		AssetTreeWidget = SNew(SDatabaseAssetTree, ViewModel.ToSharedRef());
+		AssetTreeWidget->RegisterOnSelectionChanged(
+			SDatabaseAssetTree::FOnSelectionChanged::CreateSP(
+				this,
+				&FDatabaseEditorToolkit::OnAssetTreeSelectionChanged));
 		if (IsValid(DatabaseAsset))
 		{
 			DatabaseAsset->RegisterOnAssetChange(
 				UPoseSearchDatabase::FOnDerivedDataRebuild::CreateSP(
 					AssetTreeWidget.Get(),
-					&UE::PoseSearch::SDatabaseAssetTree::RefreshTreeView, false));
+					&SDatabaseAssetTree::RefreshTreeView, false, false));
 			DatabaseAsset->RegisterOnGroupChange(
 				UPoseSearchDatabase::FOnDerivedDataRebuild::CreateSP(
 					AssetTreeWidget.Get(),
-					&UE::PoseSearch::SDatabaseAssetTree::RefreshTreeView, false));
+					&SDatabaseAssetTree::RefreshTreeView, false, false));
 		}
 
 		// Create Asset Details widget
@@ -129,9 +141,12 @@ namespace UE::PoseSearch
 			this,
 			&FDatabaseEditorToolkit::OnFinishedChangingProperties);
 
+		SelectionWidget = PropertyModule.CreateDetailView(Args);
+		SelectionWidget->SetObject(nullptr);
+
 		// Define Editor Layout
 		const TSharedRef<FTabManager::FLayout> StandaloneDefaultLayout =
-		FTabManager::NewLayout("Standalone_PoseSearchDatabaseEditor_Layout_v0.04")
+		FTabManager::NewLayout("Standalone_PoseSearchDatabaseEditor_Layout_v0.05")
 			->AddArea
 			(
 				FTabManager::NewPrimaryArea()->SetOrientation(Orient_Vertical)
@@ -162,6 +177,7 @@ namespace UE::PoseSearch
 							(
 								FTabManager::NewStack()
 								->SetSizeCoefficient(0.3f)
+								->AddTab(FDatabaseEditorTabs::SelectionDetailsID, ETabState::OpenedTab)
 								->AddTab(FDatabaseEditorTabs::AssetDetailsID, ETabState::OpenedTab)
 								->AddTab(FDatabaseEditorTabs::PreviewSettingsID, ETabState::OpenedTab)
 							)
@@ -280,6 +296,13 @@ namespace UE::PoseSearch
 			.SetDisplayName(LOCTEXT("TreeViewTab", "Tree View"))
 			.SetGroup(WorkspaceMenuCategoryRef)
 			.SetIcon(FSlateIcon(FEditorStyle::GetStyleSetName(), "GraphEditor.EventGraph_16x"));
+
+		InTabManager->RegisterTabSpawner(
+			FDatabaseEditorTabs::SelectionDetailsID,
+			FOnSpawnTab::CreateSP(this, &FDatabaseEditorToolkit::SpawnTab_SelectionDetails))
+			.SetDisplayName(LOCTEXT("SelectionDetailsTab", "Selection Details"))
+			.SetGroup(WorkspaceMenuCategoryRef)
+			.SetIcon(FSlateIcon(FEditorStyle::GetStyleSetName(), "GraphEditor.EventGraph_16x"));
 	}
 
 	void FDatabaseEditorToolkit::UnregisterTabSpawners(const TSharedRef<FTabManager>& InTabManager)
@@ -290,6 +313,7 @@ namespace UE::PoseSearch
 		InTabManager->UnregisterTabSpawner(FDatabaseEditorTabs::AssetDetailsID);
 		InTabManager->UnregisterTabSpawner(FDatabaseEditorTabs::PreviewSettingsID);
 		InTabManager->UnregisterTabSpawner(FDatabaseEditorTabs::AssetTreeViewID);
+		InTabManager->UnregisterTabSpawner(FDatabaseEditorTabs::SelectionDetailsID);
 	}
 
 	FName FDatabaseEditorToolkit::GetToolkitFName() const
@@ -373,6 +397,17 @@ namespace UE::PoseSearch
 			];
 	}
 
+	TSharedRef<SDockTab> FDatabaseEditorToolkit::SpawnTab_SelectionDetails(const FSpawnTabArgs& Args)
+	{
+		check(Args.GetTabId() == FDatabaseEditorTabs::SelectionDetailsID);
+
+		return SNew(SDockTab)
+			.Label(LOCTEXT("AssetDetails_Title", "Asset Details"))
+			[
+				SelectionWidget.ToSharedRef()
+			];
+	}
+
 	void FDatabaseEditorToolkit::OnFinishedChangingProperties(const FPropertyChangedEvent& PropertyChangedEvent)
 	{
 		const FName PropertyName = 
@@ -385,10 +420,80 @@ namespace UE::PoseSearch
 		UE_LOG(
 			LogPoseSearchEditor,
 			Log,
-			TEXT("FPoseSearchDatabaseEditorToolkit::OnFinishedChangingProperties MemberPropertyName: %s PropertyName: %s"),
+			TEXT("FDatabaseEditorToolkit::OnFinishedChangingProperties MemberPropertyName: %s PropertyName: %s"),
 			*MemberPropertyName.ToString(),
 			*PropertyName.ToString());
 	}
+
+	void FDatabaseEditorToolkit::OnAssetTreeSelectionChanged(
+		const TArrayView<TSharedPtr<FDatabaseAssetTreeNode>>& SelectedItems,
+		ESelectInfo::Type SelectionType)
+	{
+		for (TWeakObjectPtr<UObject>& SingleSelection : SelectionReflection)
+		{
+			SingleSelection->RemoveFromRoot();
+		}
+		SelectionReflection.Reset();
+
+		if (SelectedItems.Num() > 0)
+		{
+			ESearchIndexAssetType AssetType = SelectedItems[0]->SourceAssetType;
+			bool bConsistentAssetType = true;
+			for (int32 SelectionIdx = 1; SelectionIdx < SelectedItems.Num(); ++SelectionIdx)
+			{
+				if (SelectedItems[SelectionIdx]->SourceAssetType != AssetType)
+				{
+					bConsistentAssetType = false;
+					break;
+				}
+			}
+			
+			if (bConsistentAssetType)
+			{
+				if (AssetType == ESearchIndexAssetType::Sequence)
+				{
+					for (TSharedPtr<FDatabaseAssetTreeNode>& SelectedItem : SelectedItems)
+					{
+						UPoseSearchDatabaseSequenceReflection* NewSelectionReflection =
+							NewObject<UPoseSearchDatabaseSequenceReflection>();
+						NewSelectionReflection->AddToRoot();
+						NewSelectionReflection->Sequence = 
+							GetPoseSearchDatabase()->Sequences[SelectedItem->SourceAssetIdx];
+						NewSelectionReflection->SetSourceLink(SelectedItem, AssetTreeWidget);
+						SelectionReflection.Add(NewSelectionReflection);
+					}
+				}
+				else if (AssetType == ESearchIndexAssetType::BlendSpace)
+				{
+					for (TSharedPtr<FDatabaseAssetTreeNode>& SelectedItem : SelectedItems)
+					{
+						UPoseSearchDatabaseBlendSpaceReflection* NewSelectionReflection =
+							NewObject<UPoseSearchDatabaseBlendSpaceReflection>();
+						NewSelectionReflection->AddToRoot();
+						NewSelectionReflection->BlendSpace =
+							GetPoseSearchDatabase()->BlendSpaces[SelectedItem->SourceAssetIdx];
+						NewSelectionReflection->SetSourceLink(SelectedItem, AssetTreeWidget);
+						SelectionReflection.Add(NewSelectionReflection);
+					}
+				}
+				else
+				{
+					for (TSharedPtr<FDatabaseAssetTreeNode>& SelectedItem : SelectedItems)
+					{
+						UPoseSearchDatabaseGroupReflection* NewSelectionReflection =
+							NewObject<UPoseSearchDatabaseGroupReflection>();
+						NewSelectionReflection->AddToRoot();
+						NewSelectionReflection->Group =
+							GetPoseSearchDatabase()->Groups[SelectedItem->SourceAssetIdx];
+						NewSelectionReflection->SetSourceLink(SelectedItem, AssetTreeWidget);
+						SelectionReflection.Add(NewSelectionReflection);
+					}
+				}
+			}
+		}
+
+		SelectionWidget->SetObjects(SelectionReflection, true);
+	}
 }
 
 #undef LOCTEXT_NAMESPACE
diff --git a/Engine/Plugins/Experimental/Animation/PoseSearch/Source/Editor/Private/SPoseSearchDatabaseAssetList.cpp b/Engine/Plugins/Experimental/Animation/PoseSearch/Source/Editor/Private/SPoseSearchDatabaseAssetList.cpp
--- a/Engine/Plugins/Experimental/Animation/PoseSearch/Source/Editor/Private/SPoseSearchDatabaseAssetList.cpp
+++ b/Engine/Plugins/Experimental/Animation/PoseSearch/Source/Editor/Private/SPoseSearchDatabaseAssetList.cpp
@@ -330,13 +330,17 @@ namespace UE::PoseSearch
 					.OnGetChildren(this, &SDatabaseAssetTree::HandleGetChildrenForTree)
 					.OnContextMenuOpening(this, &SDatabaseAssetTree::CreateContextMenu)
 					.HighlightParentNodesForSelection(false)
+					.OnSelectionChanged_Lambda([this](TSharedPtr<FDatabaseAssetTreeNode> Item, ESelectInfo::Type Type)
+						{
+							TArray<TSharedPtr<FDatabaseAssetTreeNode>> SelectedItems = TreeView->GetSelectedItems();
+							OnSelectionChanged.Broadcast(SelectedItems, Type);
+						})
 					.ItemHeight(24)
 				]
 			]
 		];
 
-		const bool IsInitialSetup = true;
-		RefreshTreeView(IsInitialSetup);
+		RefreshTreeView(true);
 	}
 
 	FReply SDatabaseAssetTree::OnDragOver(const FGeometry& MyGeometry, const FDragDropEvent& DragDropEvent)
@@ -382,7 +386,7 @@ namespace UE::PoseSearch
 		return OnAcceptDrop(DragDropEvent, EItemDropZone::OntoItem, nullptr);
 	}
 
-	void SDatabaseAssetTree::RefreshTreeView(bool IsInitialSetup)
+	void SDatabaseAssetTree::RefreshTreeView(bool bIsInitialSetup, bool bRecoverSelection)
 	{
 		const TSharedPtr<FDatabaseViewModel> ViewModel = EditorViewModel.Pin();
 		if (!ViewModel.IsValid())
@@ -402,6 +406,9 @@ namespace UE::PoseSearch
 			return;
 		}
 
+		// store selection so we can recover it afterwards (if possible)
+		TArray<TSharedPtr<FDatabaseAssetTreeNode>> PreviouslySelectedNodes = TreeView->GetSelectedItems();
+
 		// create all group nodes
 		for (int32 GroupIdx = 0; GroupIdx < Database->Groups.Num(); ++GroupIdx)
 		{
@@ -493,6 +500,11 @@ namespace UE::PoseSearch
 		{
 			TreeView->SetItemExpansion(RootNode, true);
 		}
+
+		if (bRecoverSelection)
+		{
+			RecoverSelection(PreviouslySelectedNodes);
+		}
 	}
 
 	TSharedRef<ITableRow> SDatabaseAssetTree::MakeTableRowWidget(
@@ -776,6 +788,39 @@ namespace UE::PoseSearch
 
 		RefreshTreeView(false);
 	}
+
+	void SDatabaseAssetTree::RegisterOnSelectionChanged(const FOnSelectionChanged& Delegate)
+	{
+		OnSelectionChanged.Add(Delegate);
+	}
+
+	void SDatabaseAssetTree::UnregisterOnSelectionChanged(void* Unregister)
+	{
+		OnSelectionChanged.RemoveAll(Unregister);
+	}
+
+	void SDatabaseAssetTree::RecoverSelection(const TArray<TSharedPtr<FDatabaseAssetTreeNode>>& PreviouslySelectedNodes)
+	{
+		TArray<TSharedPtr<FDatabaseAssetTreeNode>> NewSelectedNodes;
+
+		for (const TSharedPtr<FDatabaseAssetTreeNode>& Node : AllNodes)
+		{
+			bool bFoundNode = PreviouslySelectedNodes.ContainsByPredicate(
+				[Node](const TSharedPtr<FDatabaseAssetTreeNode>& PrevSelectedNode)
+			{
+				return
+					PrevSelectedNode->SourceAssetType == Node->SourceAssetType &&
+					PrevSelectedNode->SourceAssetIdx == Node->SourceAssetIdx;
+			});
+
+			if (bFoundNode)
+			{
+				NewSelectedNodes.Add(Node);
+			}
+		}
+
+		TreeView->SetItemSelection(NewSelectedNodes, true, ESelectInfo::Direct);
+	}
 }
 
 #undef LOCTEXT_NAMESPACE
diff --git a/Engine/Plugins/Experimental/Animation/PoseSearch/Source/Editor/Public/PoseSearchDatabaseEditorReflection.h b/Engine/Plugins/Experimental/Animation/PoseSearch/Source/Editor/Public/PoseSearchDatabaseEditorReflection.h
new file mode 100644
--- /dev/null
+++ b/Engine/Plugins/Experimental/Animation/PoseSearch/Source/Editor/Public/PoseSearchDatabaseEditorReflection.h
@@ -0,0 +1,74 @@
+// Copyright Epic Games, Inc. All Rights Reserved.
+
+#pragma once
+
+#include "CoreMinimal.h"
+
+#include "PoseSearch/PoseSearch.h"
+
+#include "PoseSearchDatabaseEditorReflection.generated.h"
+
+
+namespace UE::PoseSearch
+{
+	class FDatabaseAssetTreeNode;
+	class SDatabaseAssetTree;
+}
+
+
+UCLASS()
+class UPoseSearchDatabaseReflectionBase : public UObject
+{
+	GENERATED_BODY()
+
+public:
+	void SetSourceLink(
+		const TWeakPtr<UE::PoseSearch::FDatabaseAssetTreeNode>& InWeakAssetTreeNode,
+		const TSharedPtr<UE::PoseSearch::SDatabaseAssetTree>& InAssetTreeWidget);
+
+protected:
+	TWeakPtr<UE::PoseSearch::FDatabaseAssetTreeNode> WeakAssetTreeNode;
+	TSharedPtr<UE::PoseSearch::SDatabaseAssetTree> AssetTreeWidget;
+};
+
+UCLASS()
+class UPoseSearchDatabaseSequenceReflection : public UPoseSearchDatabaseReflectionBase
+{
+	GENERATED_BODY()
+
+public:
+	UPROPERTY(EditAnywhere, Category = "Selected Sequence")
+	FPoseSearchDatabaseSequence Sequence;
+
+#if WITH_EDITOR
+	virtual void PostEditChangeProperty(struct FPropertyChangedEvent& PropertyChangedEvent) override;
+#endif
+};
+
+UCLASS()
+class UPoseSearchDatabaseBlendSpaceReflection : public UPoseSearchDatabaseReflectionBase
+{
+	GENERATED_BODY()
+
+public:
+	UPROPERTY(EditAnywhere, Category = "Selected Sequence")
+	FPoseSearchDatabaseBlendSpace BlendSpace;
+
+#if WITH_EDITOR
+	virtual void PostEditChangeProperty(struct FPropertyChangedEvent& PropertyChangedEvent) override;
+#endif
+};
+
+UCLASS()
+class UPoseSearchDatabaseGroupReflection : public UPoseSearchDatabaseReflectionBase
+{
+	GENERATED_BODY()
+
+public:
+	UPROPERTY(EditAnywhere, Category = "Selected Sequence")
+	FPoseSearchDatabaseGroup Group;
+
+#if WITH_EDITOR
+	virtual void PostEditChangeProperty(struct FPropertyChangedEvent& PropertyChangedEvent) override;
+#endif
+};
diff --git a/Engine/Plugins/Experimental/Animation/PoseSearch/Source/Editor/Public/PoseSearchDatabaseEditorToolkit.h b/Engine/Plugins/Experimental/Animation/PoseSearch/Source/Editor/Public/PoseSearchDatabaseEditorToolkit.h
--- a/Engine/Plugins/Experimental/Animation/PoseSearch/Source/Editor/Public/PoseSearchDatabaseEditorToolkit.h
+++ b/Engine/Plugins/Experimental/Animation/PoseSearch/Source/Editor/Public/PoseSearchDatabaseEditorToolkit.h
@@ -5,17 +5,17 @@
 #include "CoreMinimal.h"
 #include "Toolkits/AssetEditorToolkit.h"
 #include "Misc/NotifyHook.h"
+#include "SPoseSearchDatabaseAssetList.h"
 
 class IDetailsView;
 class FToolBarBuilder;
-class UPoseSearchDatabase;
+class UPoseSearchDatabaseSelectionReflection;
 
 namespace UE::PoseSearch
 {
 	class SDatabaseViewport;
 	class FDatabasePreviewScene;
 	class FDatabaseViewModel;
-	class SDatabaseAssetTree;
 
 	class FDatabaseEditorToolkit : public FAssetEditorToolkit, public FNotifyHook
 	{
@@ -38,6 +38,7 @@ namespace UE::PoseSearch
 		virtual FString GetWorldCentricTabPrefix() const override;
 
 		const UPoseSearchDatabase* GetPoseSearchDatabase() const;
+		UPoseSearchDatabase* GetPoseSearchDatabase();
 		FDatabaseViewModel* GetViewModel() const { return ViewModel.Get(); }
 		TSharedPtr<FDatabaseViewModel> GetViewModelSharedPtr() const { return ViewModel; }
 
@@ -50,23 +51,31 @@ namespace UE::PoseSearch
 		TSharedRef<SDockTab> SpawnTab_Viewport(const FSpawnTabArgs& Args);
 		TSharedRef<SDockTab> SpawnTab_AssetDetails(const FSpawnTabArgs& Args);
 		TSharedRef<SDockTab> SpawnTab_PreviewSettings(const FSpawnTabArgs& Args);
-	TSharedRef<SDockTab> SpawnTab_AssetTreeView(const FSpawnTabArgs& Args);
+		TSharedRef<SDockTab> SpawnTab_AssetTreeView(const FSpawnTabArgs& Args);
+		TSharedRef<SDockTab> SpawnTab_SelectionDetails(const FSpawnTabArgs& Args);
 
 		void BindCommands();
 		void ExtendToolbar();
 		void FillToolbar(FToolBarBuilder& ToolbarBuilder);
 
 		void OnFinishedChangingProperties(const FPropertyChangedEvent& PropertyChangedEvent);
+		void OnAssetTreeSelectionChanged(
+			const TArrayView<TSharedPtr<FDatabaseAssetTreeNode>>& SelectedItems,
+			ESelectInfo::Type SelectionType);
 
 		TSharedPtr<SDatabaseViewport> ViewportWidget;
 
-	TSharedPtr<UE::PoseSearch::SDatabaseAssetTree> AssetTreeWidget;
+		TSharedPtr<SDatabaseAssetTree> AssetTreeWidget;
 
 		TSharedPtr<IDetailsView> EditingAssetWidget;
 
+		TSharedPtr<IDetailsView> SelectionWidget;
+
 		TSharedPtr<FDatabasePreviewScene> PreviewScene;
 
 		TSharedPtr<FDatabaseViewModel> ViewModel;
+
+		TArray<TWeakObjectPtr<UObject>> SelectionReflection;
 	};
 }
 
diff --git a/Engine/Plugins/Experimental/Animation/PoseSearch/Source/Editor/Public/SPoseSearchDatabaseAssetList.h b/Engine/Plugins/Experimental/Animation/PoseSearch/Source/Editor/Public/SPoseSearchDatabaseAssetList.h
--- a/Engine/Plugins/Experimental/Animation/PoseSearch/Source/Editor/Public/SPoseSearchDatabaseAssetList.h
+++ b/Engine/Plugins/Experimental/Animation/PoseSearch/Source/Editor/Public/SPoseSearchDatabaseAssetList.h
@@ -32,11 +32,7 @@ namespace UE::PoseSearch
 		ESearchIndexAssetType SourceAssetType;
 		TSharedPtr<FDatabaseAssetTreeNode> Parent;
 		TArray<TSharedPtr<FDatabaseAssetTreeNode>> Children;
-
-	protected:
-
 		TWeakPtr<FDatabaseViewModel> EditorViewModel;
-
 	};
 
 	class SDatabaseAssetListItem : public STableRow<TSharedPtr<FDatabaseAssetTreeNode>>
@@ -83,7 +79,7 @@ namespace UE::PoseSearch
 		virtual FReply OnDrop(const FGeometry& MyGeometry, const FDragDropEvent& DragDropEvent) override;
 		// End SWidget interface
 	
-		void RefreshTreeView(bool IsInitialSetup = false);
+		void RefreshTreeView(bool bIsInitialSetup = false, bool bRecoverSelection = false);
 
 	protected:
 		TWeakPtr<FDatabaseViewModel> EditorViewModel;
@@ -127,6 +123,22 @@ namespace UE::PoseSearch
 		void OnDeleteGroup(TSharedPtr<FDatabaseAssetTreeNode> Node);
 
 		friend SDatabaseAssetListItem;
+
+	protected:
+		// Called when an item is selected/deselected
+		DECLARE_MULTICAST_DELEGATE_TwoParams(
+			FOnSelectionChangedMulticaster, 
+			const TArrayView<TSharedPtr<FDatabaseAssetTreeNode>>& /* InSelectedItems */,
+			ESelectInfo::Type /* SelectInfo */);
+
+		FOnSelectionChangedMulticaster OnSelectionChanged;
+
+	public:
+		typedef FOnSelectionChangedMulticaster::FDelegate FOnSelectionChanged;
+		void RegisterOnSelectionChanged(const FOnSelectionChanged& Delegate);
+		void UnregisterOnSelectionChanged(void* Unregister);
+
+		void RecoverSelection(const TArray<TSharedPtr<FDatabaseAssetTreeNode>>& PreviouslySelectedNodes);
 	};
 }
 
