PoseSearch - removed CanEvenRender() check before building derived data. It was causing problems in server builds.

diff --git a/Engine/Plugins/Experimental/Animation/PoseSearch/Source/Runtime/Private/PoseSearch.cpp b/Engine/Plugins/Experimental/Animation/PoseSearch/Source/Runtime/Private/PoseSearch.cpp
--- a/Engine/Plugins/Experimental/Animation/PoseSearch/Source/Runtime/Private/PoseSearch.cpp
+++ b/Engine/Plugins/Experimental/Animation/PoseSearch/Source/Runtime/Private/PoseSearch.cpp
@@ -1647,13 +1647,9 @@ bool UPoseSearchDatabase::TryInitSearchIndexAssets(FPoseSearchIndex& OutSearchIn
 void UPoseSearchDatabase::PostLoad()
 {
 #if WITH_EDITOR
-	// todo: might want to have a different check to see if we need the data or not
-	if (FApp::CanEverRender())
+	if (!PrivateDerivedData)
 	{
-		if (!PrivateDerivedData)
-		{
-			BeginCacheDerivedData();
-		}
+		BeginCacheDerivedData();
 	}
 #endif
 
