Refactoring of pose search order of operations.

diff --git a/Engine/Plugins/Experimental/Animation/PoseSearch/Source/Runtime/Private/AnimNode_MotionMatching.cpp b/Engine/Plugins/Experimental/Animation/PoseSearch/Source/Runtime/Private/AnimNode_MotionMatching.cpp
--- a/Engine/Plugins/Experimental/Animation/PoseSearch/Source/Runtime/Private/AnimNode_MotionMatching.cpp
+++ b/Engine/Plugins/Experimental/Animation/PoseSearch/Source/Runtime/Private/AnimNode_MotionMatching.cpp
@@ -23,7 +23,7 @@ void FAnimNode_MotionMatching::Initialize_AnyThread(const FAnimationInitializeCo
 
 	GetEvaluateGraphExposedInputs().Execute(Context);
 
-	MotionMatchingState.InitNewDatabaseSearch(Database, Settings.SearchThrottleTime, nullptr /*OutError*/);
+	MotionMatchingState.InitNewDatabaseSearch(Database, nullptr /*OutError*/);
 
 	CurrentAssetPlayerNode = &SequencePlayerNode;
 
@@ -70,18 +70,21 @@ void FAnimNode_MotionMatching::UpdateAssetPlayer(const FAnimationUpdateContext&
 
 	GetEvaluateGraphExposedInputs().Execute(Context);
 
-	// If we just became relevant and haven't been initialized yet, then reset motion matching state, otherwise update the current time using the player.
+	// If we just became relevant and haven't been initialized yet, then reset motion matching state, otherwise update the asset time using the player node.
 	if (bResetOnBecomingRelevant && UpdateCounter.HasEverBeenUpdated() && !UpdateCounter.WasSynchronizedCounter(Context.AnimInstanceProxy->GetUpdateCounter()))
 	{
 		MotionMatchingState.Reset();
 	}
 	else
 	{
-		MotionMatchingState.AssetPlayerTime = CurrentAssetPlayerNode->GetAccumulatedTime();
+		// We adjust the motion matching state asset time to the current player node's asset time. This is done 
+		// because the player node may have ticked more or less time than we expected due to variable dt or the 
+		// dynamic playback rate adjustment and as such the motion matching state does not update by itself
+		MotionMatchingState.AdjustAssetTime(CurrentAssetPlayerNode->GetAccumulatedTime());
 	}
 	UpdateCounter.SynchronizeWith(Context.AnimInstanceProxy->GetUpdateCounter());
 
-	// Execute core motion matching algorithm and retain across frame state
+	// Execute core motion matching algorithm
 	UpdateMotionMatchingState(
 		Context,
 		Database,
diff --git a/Engine/Plugins/Experimental/Animation/PoseSearch/Source/Runtime/Private/PoseSearchLibrary.cpp b/Engine/Plugins/Experimental/Animation/PoseSearch/Source/Runtime/Private/PoseSearchLibrary.cpp
--- a/Engine/Plugins/Experimental/Animation/PoseSearch/Source/Runtime/Private/PoseSearchLibrary.cpp
+++ b/Engine/Plugins/Experimental/Animation/PoseSearch/Source/Runtime/Private/PoseSearchLibrary.cpp
@@ -60,29 +60,134 @@ static void ComputeDatabaseBlendSpaceFilter(
 	}
 }
 
+} // namespace UE::PoseSearch
+
+
 //////////////////////////////////////////////////////////////////////////
-// FMotionMatchingPoseStepper
+// FMotionMatchingState
+
+bool FMotionMatchingState::InitNewDatabaseSearch(
+	const UPoseSearchDatabase* Database,
+	FText* OutError
+)
+{
+	bool bValidDatabase = Database && Database->IsValidForSearch();
+
+	if (bValidDatabase)
+	{
+		CurrentDatabase = Database;
+		Reset();
+	}
+
+	if (!bValidDatabase && OutError)
+	{
+		if (Database)
+		{
+			*OutError = FText::Format(LOCTEXT("InvalidDatabase", "Invalid database for motion matching. Try re-saving {0}."), FText::FromString(Database->GetPathName()));
+		}
+		else
+		{
+			*OutError = LOCTEXT("NoDatabase", "No database provided for motion matching.");
+		}
+	}
+
+	return bValidDatabase;
+}
+
+void FMotionMatchingState::Reset()
+{
+	DbPoseIdx = INDEX_NONE;
+	SearchIndexAssetIdx = INDEX_NONE;
+	AssetPlayerTime = 0.0f;
+	// Set the elapsed time to INFINITY to trigger a search right away
+	ElapsedPoseJumpTime = INFINITY;
+
+	if (CurrentDatabase.IsValid() && CurrentDatabase->IsValidForSearch())
+	{
+		ComposedQuery.Init(CurrentDatabase->Schema);
+	}
+	else
+	{
+		ComposedQuery.Reset();
+	}
 
-void FMotionMatchingPoseStepper::Update(const FAnimationUpdateContext& UpdateContext, const FMotionMatchingState& State)
+}
+
+void FMotionMatchingState::AdjustAssetTime(float AssetTime)
 {
-	if (State.DbPoseIdx == INDEX_NONE || !State.CurrentDatabase.IsValid() || !(State.CurrentDatabase->GetSearchIndex()))
+	if (SearchIndexAssetIdx == INDEX_NONE || !CurrentDatabase.IsValid() || !(CurrentDatabase->GetSearchIndex()))
 	{
 		return;
 	}
 
-	check(State.SearchIndexAssetIdx != INDEX_NONE);
+	const FPoseSearchIndexAsset* SearchIndexAsset = GetCurrentSearchIndexAsset();
+
+	if (SearchIndexAsset->Type == ESearchIndexAssetType::Sequence)
+	{
+		const FPoseSearchDatabaseSequence& DbSequence = CurrentDatabase->GetSequenceSourceAsset(SearchIndexAsset);
+
+		if (SearchIndexAsset->SamplingInterval.Contains(AssetTime))
+		{
+			DbPoseIdx = CurrentDatabase->GetPoseIndexFromTime(AssetTime, SearchIndexAsset);
+			AssetPlayerTime = AssetTime;
+		}
+		else
+		{
+			DbPoseIdx = INDEX_NONE;
+			AssetPlayerTime = 0.0f;
+			SearchIndexAssetIdx = INDEX_NONE;
+		}
+	}
+	else if (SearchIndexAsset->Type == ESearchIndexAssetType::BlendSpace)
+	{
+		const FPoseSearchDatabaseBlendSpace& DbBlendSpace = CurrentDatabase->GetBlendSpaceSourceAsset(SearchIndexAsset);
+
+		TArray<FBlendSampleData> BlendSamples;
+		int32 TriangulationIndex = 0;
+		DbBlendSpace.BlendSpace->GetSamplesFromBlendInput(SearchIndexAsset->BlendParameters, BlendSamples, TriangulationIndex, true);
+
+		float PlayLength = DbBlendSpace.BlendSpace->GetAnimationLengthFromSampleData(BlendSamples);
+
+		// Asset player time for blendspaces is normalized [0, 1] so we need to convert 
+		// to a real time before we advance it
+		float RealTime = AssetTime * PlayLength;
+
+		if (SearchIndexAsset->SamplingInterval.Contains(RealTime))
+		{
+			DbPoseIdx = CurrentDatabase->GetPoseIndexFromTime(RealTime, SearchIndexAsset);
+			AssetPlayerTime = AssetTime;
+		}
+		else
+		{
+			DbPoseIdx = INDEX_NONE;
+			AssetPlayerTime = 0.0f;
+			SearchIndexAssetIdx = INDEX_NONE;
+		}
+	}
+	else
+	{
+		checkNoEntry();
+	}
+}
 
-	Result.PoseIdx = INDEX_NONE;
+bool FMotionMatchingState::CanAdvance(float DeltaTime, bool& bOutAdvanceToFollowUpAsset, UE::PoseSearch::FSearchResult& OutFollowUpAsset) const
+{
+	bOutAdvanceToFollowUpAsset = false;
+	OutFollowUpAsset = UE::PoseSearch::FSearchResult();
 
-	const FPoseSearchIndexAsset* SearchIndexAsset = State.GetCurrentSearchIndexAsset();
+	if (SearchIndexAssetIdx == INDEX_NONE || !CurrentDatabase.IsValid() || !(CurrentDatabase->GetSearchIndex()))
+	{
+		return false;
+	}
+
+	const FPoseSearchIndexAsset* SearchIndexAsset = GetCurrentSearchIndexAsset();
 
 	if (SearchIndexAsset->Type == ESearchIndexAssetType::Sequence)
 	{
-		const FPoseSearchDatabaseSequence& DbSequence = State.CurrentDatabase->GetSequenceSourceAsset(SearchIndexAsset);
+		const FPoseSearchDatabaseSequence& DbSequence = CurrentDatabase->GetSequenceSourceAsset(SearchIndexAsset);
 		const float AssetLength = DbSequence.Sequence->GetPlayLength();
-		const float DeltaTime = UpdateContext.GetDeltaTime();
 
-		float SteppedTime = State.AssetPlayerTime;
+		float SteppedTime = AssetPlayerTime;
 		ETypeAdvanceAnim AdvanceType = FAnimationRuntime::AdvanceTime(
 			DbSequence.bLoopAnimation,
 			DeltaTime,
@@ -91,25 +196,18 @@ void FMotionMatchingPoseStepper::Update(const FAnimationUpdateContext& UpdateCon
 
 		if (AdvanceType != ETAA_Finished)
 		{
-			// if the sampling interval doesn't contain the new time, there's no pose to return.
-			if (SearchIndexAsset->SamplingInterval.Contains(SteppedTime))
-			{
-				// we can continue ticking the same sequence forward
-				Result.PoseIdx = State.CurrentDatabase->GetPoseIndexFromTime(State.AssetPlayerTime, SearchIndexAsset);
-				Result.SearchIndexAsset = SearchIndexAsset;
-				Result.AssetTime = State.AssetPlayerTime;
-			}
+			return SearchIndexAsset->SamplingInterval.Contains(SteppedTime);
 		}
 		else
 		{
 			// check if there's a follow-up that can be used
-			int32 FollowUpDbSequenceIdx = State.CurrentDatabase->Sequences.IndexOfByPredicate(
+			int32 FollowUpDbSequenceIdx = CurrentDatabase->Sequences.IndexOfByPredicate(
 				[&](const FPoseSearchDatabaseSequence& Entry)
 				{
 					return Entry.Sequence == DbSequence.FollowUpSequence;
 				});
 
-			int32 FollowUpSearchIndexAssetIdx = State.CurrentDatabase->GetSearchIndex()->Assets.IndexOfByPredicate(
+			int32 FollowUpSearchIndexAssetIdx = CurrentDatabase->GetSearchIndex()->Assets.IndexOfByPredicate(
 				[&](const FPoseSearchIndexAsset& Entry)
 				{
 					const bool bIsMatch =
@@ -121,28 +219,29 @@ void FMotionMatchingPoseStepper::Update(const FAnimationUpdateContext& UpdateCon
 
 			if (FollowUpSearchIndexAssetIdx != INDEX_NONE)
 			{
+				bOutAdvanceToFollowUpAsset = true;
+
 				const FPoseSearchIndexAsset* FollowUpSearchIndexAsset =
-				&State.CurrentDatabase->GetSearchIndex()->Assets[FollowUpSearchIndexAssetIdx];
-				const float FollowUpAssetTime = State.AssetPlayerTime + DeltaTime - AssetLength;
-				const int32 FollowUpPoseIdx = State.CurrentDatabase->GetPoseIndexFromTime(
-					FollowUpAssetTime,
-					FollowUpSearchIndexAsset);
-				const FFloatInterval SamplingRange = FollowUpSearchIndexAsset->SamplingInterval;
-
-				Result.PoseIdx = FollowUpPoseIdx;
-				Result.SearchIndexAsset = FollowUpSearchIndexAsset;
-				Result.AssetTime =
-					SamplingRange.Min +
-					(State.CurrentDatabase->Schema->SamplingInterval *
-						(Result.PoseIdx - FollowUpSearchIndexAsset->FirstPoseIdx));
-				bJumpRequired = true;
+					&CurrentDatabase->GetSearchIndex()->Assets[FollowUpSearchIndexAssetIdx];
+
+				// Follow up asset time will start slightly before the beginning of the sequence as 
+				// this is essentially what the matching time in the corresponding main sequence is.
+				// Here we are assuming that the tick will advance the asset player timer into the 
+				// valid region
+				const float FollowUpAssetTime = AssetPlayerTime - AssetLength;
+
+				// There is no correspoding pose index when we switch due to what is mentioned above
+				// so for now we just take whatever pose index is associated with the first frame.
+				OutFollowUpAsset.PoseIdx = CurrentDatabase->GetPoseIndexFromTime(FollowUpSearchIndexAsset->SamplingInterval.Min, FollowUpSearchIndexAsset);
+				OutFollowUpAsset.SearchIndexAsset = FollowUpSearchIndexAsset;
+				OutFollowUpAsset.AssetTime = FollowUpAssetTime;
+				return true;
 			}
 		}
 	}
 	else if (SearchIndexAsset->Type == ESearchIndexAssetType::BlendSpace)
 	{
-		const FPoseSearchDatabaseBlendSpace& DbBlendSpace = State.CurrentDatabase->GetBlendSpaceSourceAsset(SearchIndexAsset);
-		const float DeltaTime = UpdateContext.GetDeltaTime();
+		const FPoseSearchDatabaseBlendSpace& DbBlendSpace = CurrentDatabase->GetBlendSpaceSourceAsset(SearchIndexAsset);
 
 		TArray<FBlendSampleData> BlendSamples;
 		int32 TriangulationIndex = 0;
@@ -152,7 +251,7 @@ void FMotionMatchingPoseStepper::Update(const FAnimationUpdateContext& UpdateCon
 
 		// Asset player time for blendspaces is normalized [0, 1] so we need to convert 
 		// to a real time before we advance it
-		float RealTime = State.AssetPlayerTime * PlayLength;		
+		float RealTime = AssetPlayerTime * PlayLength;
 		float SteppedTime = RealTime;
 		ETypeAdvanceAnim AdvanceType = FAnimationRuntime::AdvanceTime(
 			DbBlendSpace.bLoopAnimation,
@@ -162,73 +261,15 @@ void FMotionMatchingPoseStepper::Update(const FAnimationUpdateContext& UpdateCon
 
 		if (AdvanceType != ETAA_Finished)
 		{
-			// if the sampling interval doesn't contain the new time, there's no pose to return.
-			if (SearchIndexAsset->SamplingInterval.Contains(SteppedTime))
-			{
-				// we can continue ticking the same sequence forward
-				Result.PoseIdx = State.CurrentDatabase->GetPoseIndexFromTime(RealTime, SearchIndexAsset);
-				Result.SearchIndexAsset = SearchIndexAsset;
-				Result.AssetTime = State.AssetPlayerTime;
-			}
+			return SearchIndexAsset->SamplingInterval.Contains(SteppedTime);
 		}
 	}
 	else
 	{
 		checkNoEntry();
 	}
-}
-
-} // namespace UE::PoseSearch
-
-
-//////////////////////////////////////////////////////////////////////////
-// FMotionMatchingState
-
-bool FMotionMatchingState::InitNewDatabaseSearch(
-	const UPoseSearchDatabase* Database,
-	float SearchThrottleTime,
-	FText* OutError
-)
-{
-	bool bValidDatabase = Database && Database->IsValidForSearch();
-
-	if (bValidDatabase)
-	{
-		CurrentDatabase = Database;
-		Reset();
-	}
-
-	if (!bValidDatabase && OutError)
-	{
-		if (Database)
-		{
-			*OutError = FText::Format(LOCTEXT("InvalidDatabase", "Invalid database for motion matching. Try re-saving {0}."), FText::FromString(Database->GetPathName()));
-		}
-		else
-		{
-			*OutError = LOCTEXT("NoDatabase", "No database provided for motion matching.");
-		}
-	}
 
-	return bValidDatabase;
-}
-
-void FMotionMatchingState::Reset()
-{
-	DbPoseIdx = INDEX_NONE;
-	SearchIndexAssetIdx = INDEX_NONE;
-	AssetPlayerTime = 0.0f;
-	// Set the elapsed time to INFINITY to trigger a search right away
-	ElapsedPoseJumpTime = INFINITY;
-
-	if (CurrentDatabase.IsValid() && CurrentDatabase->IsValidForSearch())
-	{
-		ComposedQuery.Init(CurrentDatabase->Schema);
-	}
-	else
-	{
-		ComposedQuery.Reset();
-	}
+	return false;
 }
 
 void FMotionMatchingState::ComposeQuery(const UPoseSearchDatabase* Database, const FTrajectorySampleRange& Trajectory)
@@ -290,12 +331,12 @@ void UpdateMotionMatchingState(
 		Context.LogMessage(EMessageSeverity::Error, LOCTEXT("NoDatabase", "No database provided for motion matching."));
 		return;
 	}
-
-	InOutMotionMatchingState.Flags = EMotionMatchingFlags::None;
+	
+	// Check for a switch in the database
 	if (InOutMotionMatchingState.CurrentDatabase != Database)
 	{
 		FText InitError;
-		if (!InOutMotionMatchingState.InitNewDatabaseSearch(Database, Settings.SearchThrottleTime, &InitError))
+		if (!InOutMotionMatchingState.InitNewDatabaseSearch(Database, &InitError))
 		{
 			Context.LogMessage(EMessageSeverity::Error, InitError);
 			return;
@@ -304,29 +345,27 @@ void UpdateMotionMatchingState(
 
 	const float DeltaTime = Context.GetDeltaTime();
 
-	// Step the pose forward
-	FMotionMatchingPoseStepper PoseStepper;
-	PoseStepper.Update(Context, InOutMotionMatchingState);
-	if (PoseStepper.CanContinue())
-	{
-		InOutMotionMatchingState.DbPoseIdx = PoseStepper.Result.PoseIdx;
-		InOutMotionMatchingState.SearchIndexAssetIdx = 
-			InOutMotionMatchingState.CurrentDatabase->GetSearchIndex()->FindAssetIndex(PoseStepper.Result.SearchIndexAsset);
-	}
+	// Reset State Flags
+	InOutMotionMatchingState.Flags = EMotionMatchingFlags::None;
+
+	// Record Current Pose Index for Debugger
+	int32 CurrentPoseIdx = InOutMotionMatchingState.DbPoseIdx;
 
-	// Build the search query
+	// Build the search query. This is done even when we don't search 
+	// since we still want to record it for debugging purposes.
 	if (InOutMotionMatchingState.DbPoseIdx != INDEX_NONE)
 	{
 		// Copy search query directly from the database if we have an active pose
 		InOutMotionMatchingState.ComposedQuery.CopyFromSearchIndex(
-			*Database->GetSearchIndex(), 
+			*Database->GetSearchIndex(),
 			InOutMotionMatchingState.DbPoseIdx);
 	}
 	else
 	{
+		// When we don't have an active pose in the database, 
+		// initialize the search query from the pose history provider
 		bool bPoseHistoryFound = false;
 
-		// When we don't have an active pose in the database, initialize the search query from the pose history provider
 		IPoseHistoryProvider* PoseHistoryProvider = Context.GetMessage<IPoseHistoryProvider>();
 		if (PoseHistoryProvider)
 		{
@@ -344,106 +383,116 @@ void UpdateMotionMatchingState(
 		}
 	}
 
-	// Update features in the query with the latest inputs
+	// Update features in the query with the latest input trajectory
 	InOutMotionMatchingState.ComposeQuery(Database, Trajectory);
 
-	FSearchContext SearchContext;
-	SearchContext.SetSource(InOutMotionMatchingState.CurrentDatabase.Get());
-	SearchContext.QueryValues = InOutMotionMatchingState.ComposedQuery.GetNormalizedValues();
-	SearchContext.WeightsContext = &InOutMotionMatchingState.WeightsContext;
-	SearchContext.DatabaseTagQuery = DatabaseTagQuery;
-	if (const FPoseSearchIndexAsset* CurrentIndexAsset = InOutMotionMatchingState.GetCurrentSearchIndexAsset())
-	{
-		SearchContext.QueryMirrorRequest =
-			CurrentIndexAsset->bMirrored ?
-			EPoseSearchBooleanRequest::TrueValue :
-			EPoseSearchBooleanRequest::FalseValue;
-	}
+	// Check if we can advance. Includes the case where we can advance but only by switching to a follow up asset.
+	bool bAdvanceToFollowUpAsset = false;
+	FSearchResult FollowUpAsset;
+	bool bCanAdvance = InOutMotionMatchingState.CanAdvance(Context.GetDeltaTime(), bAdvanceToFollowUpAsset, FollowUpAsset);
 
-	// Update weight groups
-	InOutMotionMatchingState.WeightsContext.Update(Settings.Weights, Database);
-
-	// Jump to the candidate pose if the stepper couldn't continue to the next frame
-	if (!PoseStepper.CanContinue())
+	// If we can't advance or enough time has elapsed since the last pose jump then search
+	if (!bCanAdvance || (InOutMotionMatchingState.ElapsedPoseJumpTime >= Settings.SearchThrottleTime))
 	{
-		FSearchResult SearchResult = Search(SearchContext);
-		if (SearchResult.IsValid())
+		// Build the search context
+		FSearchContext SearchContext;
+		SearchContext.SetSource(InOutMotionMatchingState.CurrentDatabase.Get());
+		SearchContext.QueryValues = InOutMotionMatchingState.ComposedQuery.GetNormalizedValues();
+		SearchContext.WeightsContext = &InOutMotionMatchingState.WeightsContext;
+		SearchContext.DatabaseTagQuery = DatabaseTagQuery;
+		if (const FPoseSearchIndexAsset* CurrentIndexAsset = InOutMotionMatchingState.GetCurrentSearchIndexAsset())
 		{
-			InOutMotionMatchingState.JumpToPose(Context, Settings, SearchResult);
+			SearchContext.QueryMirrorRequest =
+				CurrentIndexAsset->bMirrored ?
+				EPoseSearchBooleanRequest::TrueValue :
+				EPoseSearchBooleanRequest::FalseValue;
 		}
-	}
 
-	// Do a search when enough time has elapsed since the last pose jump
-	else if ((InOutMotionMatchingState.ElapsedPoseJumpTime >= Settings.SearchThrottleTime))
-	{
-		// Determine how much the updated query vector deviates from the current pose vector
-		FPoseCost CurrentPoseCost;
-		if (InOutMotionMatchingState.DbPoseIdx != INDEX_NONE)
-		{
-			const FPoseSearchIndexAsset* SearchIndexAsset = &Database->GetSearchIndex()->Assets[InOutMotionMatchingState.SearchIndexAssetIdx];
-			CurrentPoseCost = ComparePoses(InOutMotionMatchingState.DbPoseIdx, SearchContext, SearchIndexAsset->SourceGroupIdx);
-		}
+		// Update weight groups
+		InOutMotionMatchingState.WeightsContext.Update(Settings.Weights, Database);
 
 		// Search the database for the nearest match to the updated query vector
 		FSearchResult SearchResult = Search(SearchContext);
+
 		if (SearchResult.IsValid())
 		{
-			// Consider the search result better if it is more similar to the query than the current pose we're playing back from the database
-			check(SearchResult.PoseCost.Dissimilarity >= 0.0f);
-			bool bBetterPose = true;
-			if (CurrentPoseCost.IsValid())
+			// If the result is valid and we couldn't advance we should always jump to the search result
+			if (!bCanAdvance)
+			{
+				InOutMotionMatchingState.JumpToPose(Context, Settings, SearchResult);
+			}
+			// Otherwise we need to check if the result is a good improvement over the current pose
+			else
 			{
-				if ((CurrentPoseCost.TotalCost <= SearchResult.PoseCost.TotalCost) || (CurrentPoseCost.Dissimilarity <= SearchResult.PoseCost.Dissimilarity))
+				// Determine how much the updated query vector deviates from the current pose vector
+				FPoseCost CurrentPoseCost;
+				if (InOutMotionMatchingState.DbPoseIdx != INDEX_NONE)
 				{
-					bBetterPose = false;
+					const FPoseSearchIndexAsset* SearchIndexAsset = &Database->GetSearchIndex()->Assets[InOutMotionMatchingState.SearchIndexAssetIdx];
+					CurrentPoseCost = ComparePoses(InOutMotionMatchingState.DbPoseIdx, SearchContext, SearchIndexAsset->SourceGroupIdx);
 				}
-				else
+
+				// Consider the search result better if it is more similar to the query than the current pose we're playing back from the database
+				check(SearchResult.PoseCost.Dissimilarity >= 0.0f);
+				bool bBetterPose = true;
+				if (CurrentPoseCost.IsValid())
 				{
-					checkSlow(CurrentPoseCost.Dissimilarity > 0.0f && CurrentPoseCost.Dissimilarity > SearchResult.PoseCost.Dissimilarity);
-					const float RelativeSimilarityGain = -1.0f * (SearchResult.PoseCost.Dissimilarity - CurrentPoseCost.Dissimilarity) / CurrentPoseCost.Dissimilarity;
-					bBetterPose = RelativeSimilarityGain >= Settings.MinPercentImprovement / 100.0f;
+					if ((CurrentPoseCost.TotalCost <= SearchResult.PoseCost.TotalCost) || (CurrentPoseCost.Dissimilarity <= SearchResult.PoseCost.Dissimilarity))
+					{
+						bBetterPose = false;
+					}
+					else
+					{
+						checkSlow(CurrentPoseCost.Dissimilarity > 0.0f && CurrentPoseCost.Dissimilarity > SearchResult.PoseCost.Dissimilarity);
+						const float RelativeSimilarityGain = -1.0f * (SearchResult.PoseCost.Dissimilarity - CurrentPoseCost.Dissimilarity) / CurrentPoseCost.Dissimilarity;
+						bBetterPose = RelativeSimilarityGain >= Settings.MinPercentImprovement / 100.0f;
+					}
 				}
-			}
 
-			// Ignore the candidate poses from the same anim when they are too near to the current pose
-			bool bNearbyPose = false;
-			const FPoseSearchIndexAsset* StateSearchIndexAsset = InOutMotionMatchingState.GetCurrentSearchIndexAsset();
-			if (StateSearchIndexAsset == SearchResult.SearchIndexAsset)
-			{
-				// We need to check in terms of PoseIdx rather than AssetTime because
-				// for blendspaces, AssetTime is not in seconds, but in the normalized range 
-				// [0, 1] so comparing to `PoseJumpThresholdTime` will not make sense		
-				bNearbyPose = FMath::Abs(InOutMotionMatchingState.DbPoseIdx - SearchResult.PoseIdx) * Database->Schema->SamplingInterval < Settings.PoseJumpThresholdTime;
-
-				// Handle looping anims when checking for the pose being too close
-				if (!bNearbyPose && Database->IsSourceAssetLooping(StateSearchIndexAsset))
+				// Ignore the candidate poses from the same anim when they are too near to the current pose
+				bool bNearbyPose = false;
+				const FPoseSearchIndexAsset* StateSearchIndexAsset = InOutMotionMatchingState.GetCurrentSearchIndexAsset();
+				if (StateSearchIndexAsset == SearchResult.SearchIndexAsset)
 				{
-					bNearbyPose = FMath::Abs(StateSearchIndexAsset->NumPoses - InOutMotionMatchingState.DbPoseIdx - SearchResult.PoseIdx) * Database->Schema->SamplingInterval < Settings.PoseJumpThresholdTime;
+					// We need to check in terms of PoseIdx rather than AssetTime because
+					// for blendspaces, AssetTime is not in seconds, but in the normalized range 
+					// [0, 1] so comparing to `PoseJumpThresholdTime` will not make sense		
+					bNearbyPose = FMath::Abs(InOutMotionMatchingState.DbPoseIdx - SearchResult.PoseIdx) * Database->Schema->SamplingInterval < Settings.PoseJumpThresholdTime;
+
+					// Handle looping anims when checking for the pose being too close
+					if (!bNearbyPose && Database->IsSourceAssetLooping(StateSearchIndexAsset))
+					{
+						bNearbyPose = FMath::Abs(StateSearchIndexAsset->NumPoses - InOutMotionMatchingState.DbPoseIdx - SearchResult.PoseIdx) * Database->Schema->SamplingInterval < Settings.PoseJumpThresholdTime;
+					}
 				}
-			}
 
-			// Start playback from the candidate pose if we determined it was a better option
-			if (bBetterPose && !bNearbyPose)
-			{
-				InOutMotionMatchingState.JumpToPose(Context, Settings, SearchResult);
+				// Jump to candidate pose if there was a better option
+				if (bBetterPose && !bNearbyPose)
+				{
+					InOutMotionMatchingState.JumpToPose(Context, Settings, SearchResult);
+				}
 			}
 		}
 	}
 
-	// Jump to the pose stepper's next anim if we didn't jump to a pose as a result of the search above
+	// If we didn't search or it didn't find a pose to jump to, and we can 
+	// advance but only with the follow up asset, jump to that. Otherwise we 
+	// are advancing as normal, and nothing needs to be done.
 	if (!(InOutMotionMatchingState.Flags & EMotionMatchingFlags::JumpedToPose)
-		&& PoseStepper.CanContinue()
-		&& PoseStepper.bJumpRequired)
+		&& bCanAdvance
+		&& bAdvanceToFollowUpAsset)
 	{
-		InOutMotionMatchingState.JumpToPose(Context, Settings, PoseStepper.Result);
+		InOutMotionMatchingState.JumpToPose(Context, Settings, FollowUpAsset);
 		InOutMotionMatchingState.Flags |= EMotionMatchingFlags::JumpedToFollowUp;
 	}
 
+	// Tick elapsed pose jump timer
 	if (!(InOutMotionMatchingState.Flags & EMotionMatchingFlags::JumpedToPose))
 	{
 		InOutMotionMatchingState.ElapsedPoseJumpTime += DeltaTime;
 	}
 
+	// Record debugger details
 #if UE_POSE_SEARCH_TRACE_ENABLED
 	if (InOutMotionMatchingState.DbPoseIdx != INDEX_NONE)
 	{
@@ -503,8 +552,8 @@ void UpdateMotionMatchingState(
 		TraceState.Weights = Settings.Weights;
 		TraceState.DbPoseIdx = InOutMotionMatchingState.DbPoseIdx;
 		TraceState.DatabaseId = FObjectTrace::GetObjectId(Database);
-		TraceState.ContinuingPoseIdx = PoseStepper.Result.PoseIdx;
-		
+		TraceState.ContinuingPoseIdx = CurrentPoseIdx;
+
 		TraceState.AssetPlayerTime = InOutMotionMatchingState.AssetPlayerTime;
 		TraceState.DeltaTime = DeltaTime;
 		TraceState.SimLinearVelocity = SimLinearVelocity;
diff --git a/Engine/Plugins/Experimental/Animation/PoseSearch/Source/Runtime/Public/PoseSearch/PoseSearchLibrary.h b/Engine/Plugins/Experimental/Animation/PoseSearch/Source/Runtime/Public/PoseSearch/PoseSearchLibrary.h
--- a/Engine/Plugins/Experimental/Animation/PoseSearch/Source/Runtime/Public/PoseSearch/PoseSearchLibrary.h
+++ b/Engine/Plugins/Experimental/Animation/PoseSearch/Source/Runtime/Public/PoseSearch/PoseSearchLibrary.h
@@ -13,28 +13,6 @@
 
 #include "PoseSearchLibrary.generated.h"
 
-namespace UE::PoseSearch
-{
-	struct FMotionMatchingPoseStepper
-	{
-		FSearchResult Result;
-		bool bJumpRequired = false;
-
-		bool CanContinue() const
-		{
-			return Result.IsValid();
-		}
-
-		void Reset()
-		{
-			Result = UE::PoseSearch::FSearchResult();
-			bJumpRequired = false;
-		}
-
-		void Update(const FAnimationUpdateContext& UpdateContext, const struct FMotionMatchingState& State);
-	};
-}
-
 UENUM(BlueprintType, Category="Motion Trajectory", meta=(Bitflags, UseEnumValuesAsMaskValuesInEditor="true"))
 enum class EMotionMatchingFlags : uint8
 {
@@ -80,11 +58,21 @@ struct POSESEARCH_API FMotionMatchingState
 	GENERATED_BODY()
 
 	// Initializes the minimum required motion matching state
-	bool InitNewDatabaseSearch(const UPoseSearchDatabase* Database, float SearchThrottleTime, FText* OutError);
+	bool InitNewDatabaseSearch(const UPoseSearchDatabase* Database, FText* OutError);
 
 	// Reset the state to a default state using the current Database
 	void Reset();
 
+	// Checks if the currently playing asset can advance and stay in bounds under the provided DeltaTime. If the 
+	// asset cannot advance but a follow up asset is provided by the database then this function will return true and 
+	// also set bAdvanceToFollowUpAsset to true, while providing the pose to jump to in the follow up asset in OutFollowUpAsset
+	bool CanAdvance(float DeltaTime, bool& bOutAdvanceToFollowUpAsset, UE::PoseSearch::FSearchResult& OutFollowUpAsset) const;
+
+	// Attempts to set the internal state to match the provided asset time including updating the internal DbPoseIdx. 
+	// If the provided asset time is out of bounds for the currently playing asset then this function will reset the 
+	// state back to the default state.
+	void AdjustAssetTime(float AssetTime);
+
 	// Adds trajectory prediction and history information to ComposedQuery
 	void ComposeQuery(const UPoseSearchDatabase* Database, const FTrajectorySampleRange& Trajectory);
 
