PoseSearchDatabase Editor
- improved layout
- implemented Del key (to remove assets from groups/delete)
- updated selection details to show details of all selected items, even if they are of different types
- fixed a bunch of small problems causing crashes

diff --git a/Engine/Plugins/Experimental/Animation/PoseSearch/Source/Editor/Private/PoseSearchCustomization.cpp b/Engine/Plugins/Experimental/Animation/PoseSearch/Source/Editor/Private/PoseSearchCustomization.cpp
--- a/Engine/Plugins/Experimental/Animation/PoseSearch/Source/Editor/Private/PoseSearchCustomization.cpp
+++ b/Engine/Plugins/Experimental/Animation/PoseSearch/Source/Editor/Private/PoseSearchCustomization.cpp
@@ -167,25 +167,25 @@ void FPoseSearchDatabaseSequenceCustomization::CustomizeChildren(TSharedRef<IPro
 ///////////////////////////////////////////////////////////////////////////////////////////////////////////////////////
 /// FPoseSearchDatabaseDetails
 
-//TSharedRef<IDetailCustomization> FPoseSearchDatabaseDetails::MakeInstance()
-//{
-//	return MakeShareable(new FPoseSearchDatabaseDetails);
-//}
-//
-//void FPoseSearchDatabaseDetails::CustomizeDetails(IDetailLayoutBuilder& DetailBuilder)
-//{
-//	//TSharedPtr<IPropertyHandle> PropertyHandle = DetailBuilder.GetProperty("Sequences", UPoseSearchDatabase::StaticClass());
-//	//if (PropertyHandle.IsValid() && PropertyHandle->IsValidHandle())
-//	//{
-//	//	auto Refresh = FSimpleDelegate::CreateLambda([&DetailBuilder]()
-//	//	{
-//	//		DetailBuilder.ForceRefreshDetails();
-//	//	});
-//
-//	//	PropertyHandle->SetOnPropertyValueChanged(Refresh);
-//	//	PropertyHandle->SetOnChildPropertyValueChanged(Refresh);
-//	//}
-//}
+TSharedRef<IDetailCustomization> FPoseSearchDatabaseDetails::MakeInstance()
+{
+	return MakeShareable(new FPoseSearchDatabaseDetails);
+}
+
+void FPoseSearchDatabaseDetails::CustomizeDetails(IDetailLayoutBuilder& DetailBuilder)
+{
+	TArray<TSharedPtr<IPropertyHandle>> HiddenHandles;
+	HiddenHandles.Add(DetailBuilder.GetProperty(GET_MEMBER_NAME_CHECKED(UPoseSearchDatabase, Sequences)));
+	HiddenHandles.Add(DetailBuilder.GetProperty(GET_MEMBER_NAME_CHECKED(UPoseSearchDatabase, BlendSpaces)));
+	HiddenHandles.Add(DetailBuilder.GetProperty(GET_MEMBER_NAME_CHECKED(UPoseSearchDatabase, Groups)));
+	HiddenHandles.Add(DetailBuilder.GetProperty(GET_MEMBER_NAME_CHECKED(UPoseSearchDatabase, SimpleSequences)));
+	HiddenHandles.Add(DetailBuilder.GetProperty(GET_MEMBER_NAME_CHECKED(UPoseSearchDatabase, SimpleBlendSpaces)));
+
+	for (TSharedPtr<IPropertyHandle> PropertyHandle : HiddenHandles)
+	{
+		DetailBuilder.HideProperty(PropertyHandle);
+	}
+}
 
 #undef LOCTEXT_NAMESPACE
 
diff --git a/Engine/Plugins/Experimental/Animation/PoseSearch/Source/Editor/Private/PoseSearchDatabaseEditorCommands.cpp b/Engine/Plugins/Experimental/Animation/PoseSearch/Source/Editor/Private/PoseSearchDatabaseEditorCommands.cpp
--- a/Engine/Plugins/Experimental/Animation/PoseSearch/Source/Editor/Private/PoseSearchDatabaseEditorCommands.cpp
+++ b/Engine/Plugins/Experimental/Animation/PoseSearch/Source/Editor/Private/PoseSearchDatabaseEditorCommands.cpp
@@ -8,7 +8,6 @@ namespace UE::PoseSearch
 {
 	void FDatabaseEditorCommands::RegisterCommands()
 	{
-		UI_COMMAND(StopPreviewScene, "Stop Scene", "Stop Scene", EUserInterfaceActionType::Button, FInputChord());
 		UI_COMMAND(ResetPreviewScene, "Reset Scene", "Reset Scene", EUserInterfaceActionType::Button, FInputChord());
 		UI_COMMAND(BuildSearchIndex, "Build Index", "Build Index", EUserInterfaceActionType::Button, FInputChord());
 
@@ -19,6 +18,9 @@ namespace UE::PoseSearch
 			ShowPoseFeaturesAll, "All Features", "Show all features",
 			EUserInterfaceActionType::RadioButton, FInputChord());
 
+		UI_COMMAND(
+			ShowAnimationNone, "None", "Don't show animations",
+			EUserInterfaceActionType::RadioButton, FInputChord());
 		UI_COMMAND(
 			ShowAnimationOriginalOnly, "Original Only", "Show only original animations",
 			EUserInterfaceActionType::RadioButton, FInputChord());
diff --git a/Engine/Plugins/Experimental/Animation/PoseSearch/Source/Editor/Private/PoseSearchDatabaseEditorReflection.cpp b/Engine/Plugins/Experimental/Animation/PoseSearch/Source/Editor/Private/PoseSearchDatabaseEditorReflection.cpp
--- a/Engine/Plugins/Experimental/Animation/PoseSearch/Source/Editor/Private/PoseSearchDatabaseEditorReflection.cpp
+++ b/Engine/Plugins/Experimental/Animation/PoseSearch/Source/Editor/Private/PoseSearchDatabaseEditorReflection.cpp
@@ -25,7 +25,7 @@ void UPoseSearchDatabaseSequenceReflection::PostEditChangeProperty(
 	if (IsValid(Database))
 	{
 		Database->Sequences[WeakAssetTreeNode.Pin()->SourceAssetIdx] = Sequence;
-		AssetTreeWidget->RefreshTreeView(false, true);
+		AssetTreeWidget->FinalizeTreeChanges(true);
 	}
 }
 
@@ -40,7 +40,7 @@ void UPoseSearchDatabaseBlendSpaceReflection::PostEditChangeProperty(
 	if (IsValid(Database))
 	{
 		Database->BlendSpaces[WeakAssetTreeNode.Pin()->SourceAssetIdx] = BlendSpace;
-		AssetTreeWidget->RefreshTreeView(false, true);
+		AssetTreeWidget->FinalizeTreeChanges(true);
 	}
 }
 
@@ -55,7 +55,7 @@ void UPoseSearchDatabaseGroupReflection::PostEditChangeProperty(
 	if (IsValid(Database))
 	{
 		Database->Groups[WeakAssetTreeNode.Pin()->SourceAssetIdx] = Group;
-		AssetTreeWidget->RefreshTreeView(false, true);
+		AssetTreeWidget->FinalizeTreeChanges(true);
 	}
 }
 
diff --git a/Engine/Plugins/Experimental/Animation/PoseSearch/Source/Editor/Private/PoseSearchDatabaseEditorToolkit.cpp b/Engine/Plugins/Experimental/Animation/PoseSearch/Source/Editor/Private/PoseSearchDatabaseEditorToolkit.cpp
--- a/Engine/Plugins/Experimental/Animation/PoseSearch/Source/Editor/Private/PoseSearchDatabaseEditorToolkit.cpp
+++ b/Engine/Plugins/Experimental/Animation/PoseSearch/Source/Editor/Private/PoseSearchDatabaseEditorToolkit.cpp
@@ -16,12 +16,14 @@
 #include "Modules/ModuleManager.h"
 #include "Widgets/Input/SButton.h"
 #include "Widgets/Docking/SDockTab.h"
+#include "Widgets/Layout/SScrollBox.h"
 
 #define LOCTEXT_NAMESPACE "PoseSearchDatabaseEditorToolkit"
 
 namespace UE::PoseSearch
 {
 	const FName PoseSearchDatabaseEditorAppName = FName(TEXT("PoseSearchDatabaseEditorApp"));
+	constexpr int32 SelectionDetailsCount = 3; // sequence, blendspace, group
 
 	// Tab identifiers
 	struct FDatabaseEditorTabs
@@ -129,24 +131,39 @@ namespace UE::PoseSearch
 		}
 
 		// Create Asset Details widget
-		FPropertyEditorModule& PropertyModule = FModuleManager::LoadModuleChecked<FPropertyEditorModule>("PropertyEditor");
+		FPropertyEditorModule& PropertyModule = 
+			FModuleManager::LoadModuleChecked<FPropertyEditorModule>("PropertyEditor");
 
-		FDetailsViewArgs Args;
-		Args.bHideSelectionTip = true;
-		Args.NotifyHook = this;
+		FDetailsViewArgs DatabaseDetailsArgs;
+		{
+			DatabaseDetailsArgs.bHideSelectionTip = true;
+			DatabaseDetailsArgs.NotifyHook = this;
+		}
 
-		EditingAssetWidget = PropertyModule.CreateDetailView(Args);
+		EditingAssetWidget = PropertyModule.CreateDetailView(DatabaseDetailsArgs);
 		EditingAssetWidget->SetObject(DatabaseAsset);
-		EditingAssetWidget->OnFinishedChangingProperties().AddSP(
-			this,
-			&FDatabaseEditorToolkit::OnFinishedChangingProperties);
 
-		SelectionWidget = PropertyModule.CreateDetailView(Args);
-		SelectionWidget->SetObject(nullptr);
+		FDetailsViewArgs SelectionDetailsArgs;
+		{
+			SelectionDetailsArgs.bHideSelectionTip = true;
+			SelectionDetailsArgs.NotifyHook = this;
+			SelectionDetailsArgs.bShowScrollBar = false;
+		}
+
+		SelectionWidgets.Reset(SelectionDetailsCount);
+		for (int32 i = 0; i < SelectionDetailsCount; ++i)
+		{
+			TSharedPtr<IDetailsView>& SelectionWidget = 
+				SelectionWidgets.Add_GetRef(PropertyModule.CreateDetailView(SelectionDetailsArgs));
+			SelectionWidget->SetObject(nullptr);
+			SelectionWidget->OnFinishedChangingProperties().AddSP(
+				this,
+				&FDatabaseEditorToolkit::OnFinishedChangingSelectionProperties);
+		}
 
 		// Define Editor Layout
 		const TSharedRef<FTabManager::FLayout> StandaloneDefaultLayout =
-		FTabManager::NewLayout("Standalone_PoseSearchDatabaseEditor_Layout_v0.05")
+		FTabManager::NewLayout("Standalone_PoseSearchDatabaseEditor_Layout_v0.07")
 			->AddArea
 			(
 				FTabManager::NewPrimaryArea()->SetOrientation(Orient_Vertical)
@@ -158,10 +175,21 @@ namespace UE::PoseSearch
 						FTabManager::NewSplitter()->SetOrientation(Orient_Horizontal)->SetSizeCoefficient(0.9f)
 						->Split
 						(
-							FTabManager::NewStack()
-							->SetSizeCoefficient(0.25f)
-							->AddTab(FDatabaseEditorTabs::AssetTreeViewID, ETabState::OpenedTab)
-							->SetHideTabWell(true)
+							FTabManager::NewSplitter()->SetOrientation(Orient_Vertical)->SetSizeCoefficient(0.9f)
+							->Split
+							(
+								FTabManager::NewStack()
+								->SetSizeCoefficient(0.5f)
+								->AddTab(FDatabaseEditorTabs::AssetTreeViewID, ETabState::OpenedTab)
+								->SetHideTabWell(false)
+							)
+							->Split
+							(
+								FTabManager::NewStack()
+								->SetSizeCoefficient(0.5f)
+								->AddTab(FDatabaseEditorTabs::AssetDetailsID, ETabState::OpenedTab)
+								->SetHideTabWell(false)
+							)
 						)
 						->Split
 						(
@@ -178,7 +206,6 @@ namespace UE::PoseSearch
 								FTabManager::NewStack()
 								->SetSizeCoefficient(0.3f)
 								->AddTab(FDatabaseEditorTabs::SelectionDetailsID, ETabState::OpenedTab)
-								->AddTab(FDatabaseEditorTabs::AssetDetailsID, ETabState::OpenedTab)
 								->AddTab(FDatabaseEditorTabs::PreviewSettingsID, ETabState::OpenedTab)
 							)
 						)
@@ -208,11 +235,6 @@ namespace UE::PoseSearch
 	{
 		const FDatabaseEditorCommands& Commands = FDatabaseEditorCommands::Get();
 
-		ToolkitCommands->MapAction(
-			Commands.StopPreviewScene,
-			FExecuteAction::CreateSP(this, &FDatabaseEditorToolkit::StopPreviewScene),
-			EUIActionRepeatMode::RepeatDisabled);
-
 		ToolkitCommands->MapAction(
 			Commands.ResetPreviewScene,
 			FExecuteAction::CreateSP(this, &FDatabaseEditorToolkit::ResetPreviewScene),
@@ -239,13 +261,6 @@ namespace UE::PoseSearch
 
 	void FDatabaseEditorToolkit::FillToolbar(FToolBarBuilder& ToolbarBuilder)
 	{
-		ToolbarBuilder.AddToolBarButton(
-			FDatabaseEditorCommands::Get().StopPreviewScene,
-			NAME_None,
-			TAttribute<FText>(),
-			TAttribute<FText>(),
-			FSlateIcon(FAppStyle::Get().GetStyleSetName(), "Icons.Delete"));
-
 		ToolbarBuilder.AddToolBarButton(
 			FDatabaseEditorCommands::Get().ResetPreviewScene,
 			NAME_None,
@@ -279,7 +294,7 @@ namespace UE::PoseSearch
 		InTabManager->RegisterTabSpawner(
 			FDatabaseEditorTabs::AssetDetailsID,
 			FOnSpawnTab::CreateSP(this, &FDatabaseEditorToolkit::SpawnTab_AssetDetails))
-			.SetDisplayName(LOCTEXT("AssetDetailsTab", "AssetDetails"))
+			.SetDisplayName(LOCTEXT("DatabaseDetailsTab", "Database Details"))
 			.SetGroup(WorkspaceMenuCategoryRef)
 			.SetIcon(FSlateIcon(FEditorStyle::GetStyleSetName(), "LevelEditor.Tabs.Details"));
 
@@ -362,7 +377,7 @@ namespace UE::PoseSearch
 		check(Args.GetTabId() == FDatabaseEditorTabs::AssetDetailsID);
 
 		return SNew(SDockTab)
-			.Label(LOCTEXT("AssetDetails_Title", "Asset Details"))
+			.Label(LOCTEXT("AssetDetails_Title", "Database Details"))
 			[
 				EditingAssetWidget.ToSharedRef()
 			];
@@ -372,10 +387,10 @@ namespace UE::PoseSearch
 	{
 		check(Args.GetTabId() == FDatabaseEditorTabs::PreviewSettingsID);
 
-	FAdvancedPreviewSceneModule& AdvancedPreviewSceneModule = 
-		FModuleManager::LoadModuleChecked<FAdvancedPreviewSceneModule>("AdvancedPreviewScene");
-	TSharedRef<SWidget> InWidget = 
-		AdvancedPreviewSceneModule.CreateAdvancedPreviewSceneSettingsWidget(PreviewScene.ToSharedRef());
+		FAdvancedPreviewSceneModule& AdvancedPreviewSceneModule = 
+			FModuleManager::LoadModuleChecked<FAdvancedPreviewSceneModule>("AdvancedPreviewScene");
+		TSharedRef<SWidget> InWidget = 
+			AdvancedPreviewSceneModule.CreateAdvancedPreviewSceneSettingsWidget(PreviewScene.ToSharedRef());
 
 		TSharedRef<SDockTab> SpawnedTab = SNew(SDockTab)
 			.Label(LOCTEXT("PreviewSceneSettingsTab", "Preview Scene Settings"))
@@ -401,84 +416,76 @@ namespace UE::PoseSearch
 	{
 		check(Args.GetTabId() == FDatabaseEditorTabs::SelectionDetailsID);
 
+		TSharedPtr<SVerticalBox> DetailsContainer;
+		SAssignNew(DetailsContainer, SVerticalBox);
+
+		for (TSharedPtr<IDetailsView> SelectionWidget : SelectionWidgets)
+		{
+			DetailsContainer->AddSlot()
+				.AutoHeight()
+				[
+					SelectionWidget.ToSharedRef()
+				];
+		}
+
 		return SNew(SDockTab)
-			.Label(LOCTEXT("AssetDetails_Title", "Asset Details"))
+			.Label(LOCTEXT("AssetDetails_Title", "Selection Details"))
 			[
-				SelectionWidget.ToSharedRef()
+				SNew(SScrollBox)
+				+SScrollBox::Slot()
+				[
+					DetailsContainer.ToSharedRef()
+				]
 			];
 	}
 
-	void FDatabaseEditorToolkit::OnFinishedChangingProperties(const FPropertyChangedEvent& PropertyChangedEvent)
+	void FDatabaseEditorToolkit::OnFinishedChangingSelectionProperties(const FPropertyChangedEvent& PropertyChangedEvent)
 	{
-		const FName PropertyName = 
-			(PropertyChangedEvent.Property != nullptr) ? 
-			PropertyChangedEvent.Property->GetFName() : NAME_None;
-		const FName MemberPropertyName = 
-			(PropertyChangedEvent.MemberProperty != nullptr) ? 
-			PropertyChangedEvent.MemberProperty->GetFName() : NAME_None;
-
-		UE_LOG(
-			LogPoseSearchEditor,
-			Log,
-			TEXT("FDatabaseEditorToolkit::OnFinishedChangingProperties MemberPropertyName: %s PropertyName: %s"),
-			*MemberPropertyName.ToString(),
-			*PropertyName.ToString());
+		ViewModel->BuildSearchIndex();
 	}
 
 	void FDatabaseEditorToolkit::OnAssetTreeSelectionChanged(
 		const TArrayView<TSharedPtr<FDatabaseAssetTreeNode>>& SelectedItems,
 		ESelectInfo::Type SelectionType)
 	{
-		for (TWeakObjectPtr<UObject>& SingleSelection : SelectionReflection)
+		enum ESelectionAssetTypeIndex
 		{
-			SingleSelection->RemoveFromRoot();
-		}
-		SelectionReflection.Reset();
+			SequenceSelectionIndex = 0,
+			BlendSpaceSelectionIndex = 1,
+			GroupSelectionIndex = 2,
+			SelectionTypeCount = 3
+		};
+
+		TArray<TArray<TWeakObjectPtr<UObject>>> SelectionReflections;
+		SelectionReflections.SetNum(SelectionDetailsCount);
 
 		if (SelectedItems.Num() > 0)
 		{
-			ESearchIndexAssetType AssetType = SelectedItems[0]->SourceAssetType;
-			bool bConsistentAssetType = true;
-			for (int32 SelectionIdx = 1; SelectionIdx < SelectedItems.Num(); ++SelectionIdx)
+			for (TSharedPtr<FDatabaseAssetTreeNode>& SelectedItem : SelectedItems)
 			{
-				if (SelectedItems[SelectionIdx]->SourceAssetType != AssetType)
+				if (SelectedItem->SourceAssetType == ESearchIndexAssetType::Sequence)
 				{
-					bConsistentAssetType = false;
-					break;
+					UPoseSearchDatabaseSequenceReflection* NewSelectionReflection =
+						NewObject<UPoseSearchDatabaseSequenceReflection>();
+					NewSelectionReflection->AddToRoot();
+					NewSelectionReflection->Sequence =
+						GetPoseSearchDatabase()->Sequences[SelectedItem->SourceAssetIdx];
+					NewSelectionReflection->SetSourceLink(SelectedItem, AssetTreeWidget);
+					SelectionReflections[SequenceSelectionIndex].Add(NewSelectionReflection);
 				}
-			}
-			
-			if (bConsistentAssetType)
-			{
-				if (AssetType == ESearchIndexAssetType::Sequence)
+				else if (SelectedItem->SourceAssetType == ESearchIndexAssetType::BlendSpace)
 				{
-					for (TSharedPtr<FDatabaseAssetTreeNode>& SelectedItem : SelectedItems)
-					{
-						UPoseSearchDatabaseSequenceReflection* NewSelectionReflection =
-							NewObject<UPoseSearchDatabaseSequenceReflection>();
-						NewSelectionReflection->AddToRoot();
-						NewSelectionReflection->Sequence = 
-							GetPoseSearchDatabase()->Sequences[SelectedItem->SourceAssetIdx];
-						NewSelectionReflection->SetSourceLink(SelectedItem, AssetTreeWidget);
-						SelectionReflection.Add(NewSelectionReflection);
-					}
-				}
-				else if (AssetType == ESearchIndexAssetType::BlendSpace)
-				{
-					for (TSharedPtr<FDatabaseAssetTreeNode>& SelectedItem : SelectedItems)
-					{
-						UPoseSearchDatabaseBlendSpaceReflection* NewSelectionReflection =
-							NewObject<UPoseSearchDatabaseBlendSpaceReflection>();
-						NewSelectionReflection->AddToRoot();
-						NewSelectionReflection->BlendSpace =
-							GetPoseSearchDatabase()->BlendSpaces[SelectedItem->SourceAssetIdx];
-						NewSelectionReflection->SetSourceLink(SelectedItem, AssetTreeWidget);
-						SelectionReflection.Add(NewSelectionReflection);
-					}
+					UPoseSearchDatabaseBlendSpaceReflection* NewSelectionReflection =
+						NewObject<UPoseSearchDatabaseBlendSpaceReflection>();
+					NewSelectionReflection->AddToRoot();
+					NewSelectionReflection->BlendSpace =
+						GetPoseSearchDatabase()->BlendSpaces[SelectedItem->SourceAssetIdx];
+					NewSelectionReflection->SetSourceLink(SelectedItem, AssetTreeWidget);
+					SelectionReflections[BlendSpaceSelectionIndex].Add(NewSelectionReflection);
 				}
 				else
 				{
-					for (TSharedPtr<FDatabaseAssetTreeNode>& SelectedItem : SelectedItems)
+					if (SelectedItem->SourceAssetIdx != INDEX_NONE)
 					{
 						UPoseSearchDatabaseGroupReflection* NewSelectionReflection =
 							NewObject<UPoseSearchDatabaseGroupReflection>();
@@ -486,13 +493,27 @@ namespace UE::PoseSearch
 						NewSelectionReflection->Group =
 							GetPoseSearchDatabase()->Groups[SelectedItem->SourceAssetIdx];
 						NewSelectionReflection->SetSourceLink(SelectedItem, AssetTreeWidget);
-						SelectionReflection.Add(NewSelectionReflection);
+						SelectionReflections[GroupSelectionIndex].Add(NewSelectionReflection);
 					}
 				}
 			}
 		}
 
-		SelectionWidget->SetObjects(SelectionReflection, true);
+		int32 SelectionWidgetIdx = 0;
+		for (int32 SelectionTypeIdx = 0; SelectionTypeIdx < SelectionTypeCount; ++SelectionTypeIdx)
+		{
+			if (SelectionReflections[SelectionTypeIdx].Num() > 0)
+			{
+				SelectionWidgets[SelectionWidgetIdx++]->SetObjects(SelectionReflections[SelectionTypeIdx], true);
+			}
+		}
+
+		while (SelectionWidgetIdx < SelectionDetailsCount)
+		{
+			SelectionWidgets[SelectionWidgetIdx++]->SetObject(nullptr, true);
+		}
+
+		ViewModel->SetSelectedNodes(SelectedItems);
 	}
 }
 
diff --git a/Engine/Plugins/Experimental/Animation/PoseSearch/Source/Editor/Private/PoseSearchDatabaseViewModel.cpp b/Engine/Plugins/Experimental/Animation/PoseSearch/Source/Editor/Private/PoseSearchDatabaseViewModel.cpp
--- a/Engine/Plugins/Experimental/Animation/PoseSearch/Source/Editor/Private/PoseSearchDatabaseViewModel.cpp
+++ b/Engine/Plugins/Experimental/Animation/PoseSearch/Source/Editor/Private/PoseSearchDatabaseViewModel.cpp
@@ -1,10 +1,11 @@
 // Copyright Epic Games, Inc. All Rights Reserved.
 
 #include "PoseSearchDatabaseViewModel.h"
+#include "PoseSearchDatabasePreviewScene.h"
+#include "SPoseSearchDatabaseAssetList.h"
 #include "PoseSearchEditor.h"
 #include "PoseSearch/PoseSearch.h"
 #include "Modules/ModuleManager.h"
-#include "PoseSearchDatabasePreviewScene.h"
 #include "AnimPreviewInstance.h"
 #include "Animation/DebugSkelMeshComponent.h"
 #include "Animation/MirrorDataTable.h"
@@ -58,19 +59,46 @@ namespace UE::PoseSearch
 	void FDatabaseViewModel::RespawnPreviewActors()
 	{
 		RemovePreviewActors();
-		const FPoseSearchIndex* SearchIndex = PoseSearchDatabase->GetSearchIndex();
-		for (const FPoseSearchIndexAsset& IndexAsset : SearchIndex->Assets)
+
+		if (AnimationPreviewMode == EAnimationPreviewMode::None)
+		{
+			return;
+		}
+
+		if (SelectedNodes.Num() > 0)
 		{
-			if (AnimationPreviewMode == EAnimationPreviewMode::OriginalAndMirrored || !IndexAsset.bMirrored)
+			TSet<int32> AssociatedAssetIndices;
+			for (const TSharedPtr<FDatabaseAssetTreeNode>& SelectedNode : SelectedNodes)
 			{
-				FDatabasePreviewActor PreviewActor = SpawnPreviewActor(IndexAsset);
-				if (PreviewActor.IsValid())
+				if (SelectedNode->SourceAssetType == ESearchIndexAssetType::Sequence)
 				{
-					PreviewActors.Add(PreviewActor);
+					AssociatedAssetIndices.Add(SelectedNode->SourceAssetIdx);
 				}
 			}
+
+			const FPoseSearchIndex* SearchIndex = PoseSearchDatabase->GetSearchIndex();
+			for (const FPoseSearchIndexAsset& IndexAsset : SearchIndex->Assets)
+			{
+				const bool bIsAssociatedToSelection =
+					IndexAsset.Type == ESearchIndexAssetType::Sequence &&
+					AssociatedAssetIndices.Contains(IndexAsset.SourceAssetIdx);
+
+				const bool bSpawn =
+					bIsAssociatedToSelection &&
+					(AnimationPreviewMode == EAnimationPreviewMode::OriginalAndMirrored ||
+					 !IndexAsset.bMirrored);
+				
+				if (bSpawn)
+				{
+					FDatabasePreviewActor PreviewActor = SpawnPreviewActor(IndexAsset);
+					if (PreviewActor.IsValid())
+					{
+						PreviewActors.Add(PreviewActor);
+					}
+				}
+			}
+			UpdatePreviewActors();
 		}
-		UpdatePreviewActors();
 
 		// todo: do blendspaces afterwards
 	}
@@ -313,4 +341,10 @@ namespace UE::PoseSearch
 
 		PoseSearchDatabase->Groups.RemoveAt(GroupIdx);
 	}
+
+	void FDatabaseViewModel::SetSelectedNodes(const TArrayView<TSharedPtr<FDatabaseAssetTreeNode>>& InSelectedNodes)
+	{
+		SelectedNodes = InSelectedNodes;
+		ResetPreviewActors();
+	}
 }
diff --git a/Engine/Plugins/Experimental/Animation/PoseSearch/Source/Editor/Private/PoseSearchEditor.cpp b/Engine/Plugins/Experimental/Animation/PoseSearch/Source/Editor/Private/PoseSearchEditor.cpp
--- a/Engine/Plugins/Experimental/Animation/PoseSearch/Source/Editor/Private/PoseSearchEditor.cpp
+++ b/Engine/Plugins/Experimental/Animation/PoseSearch/Source/Editor/Private/PoseSearchEditor.cpp
@@ -194,7 +194,7 @@ void FEditorModule::RegisterPropertyTypeCustomizations()
 
 void FEditorModule::RegisterObjectCustomizations()
 {
-//	RegisterCustomClassLayout("PoseSearchDatabase", FOnGetDetailCustomizationInstance::CreateStatic(&FPoseSearchDatabaseDetails::MakeInstance));
+	RegisterCustomClassLayout("PoseSearchDatabase", FOnGetDetailCustomizationInstance::CreateStatic(&FPoseSearchDatabaseDetails::MakeInstance));
 }
 
 void FEditorModule::UnregisterCustomizations()
diff --git a/Engine/Plugins/Experimental/Animation/PoseSearch/Source/Editor/Private/SPoseSearchDatabaseAssetList.cpp b/Engine/Plugins/Experimental/Animation/PoseSearch/Source/Editor/Private/SPoseSearchDatabaseAssetList.cpp
--- a/Engine/Plugins/Experimental/Animation/PoseSearch/Source/Editor/Private/SPoseSearchDatabaseAssetList.cpp
+++ b/Engine/Plugins/Experimental/Animation/PoseSearch/Source/Editor/Private/SPoseSearchDatabaseAssetList.cpp
@@ -13,9 +13,11 @@
 #include "AssetSelection.h"
 
 #include "Widgets/Text/SRichTextBlock.h"
+#include "Framework/Commands/GenericCommands.h"
 #include "Framework/MultiBox/MultiBoxBuilder.h"
 #include "SPositiveActionButton.h"
 
+
 #define LOCTEXT_NAMESPACE "PoseSearchDatabaseAssetList"
 
 namespace UE::PoseSearch
@@ -211,6 +213,18 @@ namespace UE::PoseSearch
 		}
 		else
 		{
+			TSharedPtr<SImage> ItemIconWidget;
+			if (Node->SourceAssetType == ESearchIndexAssetType::Sequence)
+			{
+				SAssignNew(ItemIconWidget, SImage)
+				.Image(FAppStyle::Get().GetBrush("Icons.Minus"));
+			}
+			else
+			{
+				SAssignNew(ItemIconWidget, SImage)
+				.Image(FAppStyle::Get().GetBrush("Icons.Plus"));
+			}
+
 			// it's an asset (sequence or blendspace)
 			SAssignNew(ItemWidget, SHorizontalBox)
 			+ SHorizontalBox::Slot()
@@ -219,8 +233,7 @@ namespace UE::PoseSearch
 			.HAlign(HAlign_Left)
 			.VAlign(VAlign_Center)
 			[
-				SNew(SImage)
-				.Image(FAppStyle::Get().GetBrush("Icons.Star"))
+				ItemIconWidget.ToSharedRef()
 			]
 			+ SHorizontalBox::Slot()
 			.FillWidth(1.0f)
@@ -295,7 +308,8 @@ namespace UE::PoseSearch
 		TSharedRef<FDatabaseViewModel> InEditorViewModel)
 	{
 		EditorViewModel = InEditorViewModel;
-		CommandList = MakeShared<FUICommandList>();
+
+		CreateCommandList();
 
 		ChildSlot
 		[
@@ -386,6 +400,15 @@ namespace UE::PoseSearch
 		return OnAcceptDrop(DragDropEvent, EItemDropZone::OntoItem, nullptr);
 	}
 
+	FReply SDatabaseAssetTree::OnKeyDown(const FGeometry& MyGeometry, const FKeyEvent& InKeyEvent)
+	{
+		if (CommandList->ProcessCommandBindings(InKeyEvent))
+		{
+			return FReply::Handled();
+		}
+		return FReply::Unhandled();
+	}
+
 	void SDatabaseAssetTree::RefreshTreeView(bool bIsInitialSetup, bool bRecoverSelection)
 	{
 		const TSharedPtr<FDatabaseViewModel> ViewModel = EditorViewModel.Pin();
@@ -505,6 +528,10 @@ namespace UE::PoseSearch
 		{
 			RecoverSelection(PreviouslySelectedNodes);
 		}
+		else
+		{
+			TreeView->SetItemSelection(PreviouslySelectedNodes, false, ESelectInfo::Direct);
+		}
 	}
 
 	TSharedRef<ITableRow> SDatabaseAssetTree::MakeTableRowWidget(
@@ -617,7 +644,7 @@ namespace UE::PoseSearch
 			return FReply::Unhandled();
 		}
 
-		RefreshTreeView(false);
+		FinalizeTreeChanges(false);
 		return FReply::Handled();
 	}
 
@@ -646,7 +673,7 @@ namespace UE::PoseSearch
 			LOCTEXT("AddSequence", "Sequence"),
 			LOCTEXT("AddSequenceTooltip", "Add new sequence to the default group"),
 			FSlateIcon(),
-			FUIAction(FExecuteAction::CreateSP(this, &SDatabaseAssetTree::OnAddSequence)),
+			FUIAction(FExecuteAction::CreateSP(this, &SDatabaseAssetTree::OnAddSequence, true)),
 			NAME_None,
 			EUserInterfaceActionType::Button);
 
@@ -654,7 +681,7 @@ namespace UE::PoseSearch
 			LOCTEXT("AddBlendSpace", "Blend Space"),
 			LOCTEXT("AddBlendSpaceTooltip", "Add new blend space to the default group"),
 			FSlateIcon(),
-			FUIAction(FExecuteAction::CreateSP(this, &SDatabaseAssetTree::OnAddBlendSpace)),
+			FUIAction(FExecuteAction::CreateSP(this, &SDatabaseAssetTree::OnAddBlendSpace, true)),
 			NAME_None,
 			EUserInterfaceActionType::Button);
 		AddOptions.EndSection();
@@ -664,7 +691,7 @@ namespace UE::PoseSearch
 			LOCTEXT("AddBlendSpace", "Group"),
 			LOCTEXT("AddBlendSpaceTooltip", "Add new group"),
 			FSlateIcon(),
-			FUIAction(FExecuteAction::CreateSP(this, &SDatabaseAssetTree::OnAddGroup)),
+			FUIAction(FExecuteAction::CreateSP(this, &SDatabaseAssetTree::OnAddGroup, true)),
 			NAME_None,
 			EUserInterfaceActionType::Button);
 		AddOptions.EndSection();
@@ -691,7 +718,8 @@ namespace UE::PoseSearch
 						LOCTEXT("DeleteAsset", "Delete Asset"),
 						LOCTEXT("DeleteAssetTooltip", "Delete asset from database"),
 						FSlateIcon(),
-						FUIAction(FExecuteAction::CreateSP(this, &SDatabaseAssetTree::OnDeleteAsset, SelectedNode)),
+						FUIAction(
+							FExecuteAction::CreateSP(this, &SDatabaseAssetTree::OnDeleteAsset, SelectedNode, true)),
 						NAME_None,
 						EUserInterfaceActionType::Button);
 				}
@@ -701,7 +729,8 @@ namespace UE::PoseSearch
 						LOCTEXT("RemoveFromGroup", "Remove From Group"),
 						LOCTEXT("RemoveFromGroupTooltip", "Remove asset from group"),
 						FSlateIcon(),
-						FUIAction(FExecuteAction::CreateSP(this, &SDatabaseAssetTree::OnRemoveFromGroup, SelectedNode)),
+						FUIAction(
+							FExecuteAction::CreateSP(this, &SDatabaseAssetTree::OnRemoveFromGroup, SelectedNode, true)),
 						NAME_None,
 						EUserInterfaceActionType::Button);
 				}
@@ -713,7 +742,8 @@ namespace UE::PoseSearch
 					LOCTEXT("DeleteGroup", "Delete Group"),
 					LOCTEXT("DeleteGroupTooltip", "Delete group and remove it from the sequences in it"),
 					FSlateIcon(),
-					FUIAction(FExecuteAction::CreateSP(this, &SDatabaseAssetTree::OnDeleteGroup, SelectedNode)),
+					FUIAction(
+						FExecuteAction::CreateSP(this, &SDatabaseAssetTree::OnDeleteGroup, SelectedNode, true)),
 					NAME_None,
 					EUserInterfaceActionType::Button);
 			}
@@ -723,26 +753,46 @@ namespace UE::PoseSearch
 		return MenuBuilder.MakeWidget();
 	}
 
-	void SDatabaseAssetTree::OnAddGroup()
+	void SDatabaseAssetTree::OnAddGroup(bool bFinalizeChanges)
 	{
+		FScopedTransaction Transaction(LOCTEXT("AddGroup", "Add Group"));
+
 		EditorViewModel.Pin()->AddGroupToDatabase();
-		RefreshTreeView(false);
+
+		if (bFinalizeChanges)
+		{
+			FinalizeTreeChanges();
+		}
 	}
 
-	void SDatabaseAssetTree::OnAddSequence()
+	void SDatabaseAssetTree::OnAddSequence(bool bFinalizeChanges)
 	{
+		FScopedTransaction Transaction(LOCTEXT("AddSequence", "Add Sequence"));
+
 		EditorViewModel.Pin()->AddSequenceToDatabase(nullptr, INDEX_NONE);
-		RefreshTreeView(false);
+
+		if (bFinalizeChanges)
+		{
+			FinalizeTreeChanges();
+		}
 	}
 
-	void SDatabaseAssetTree::OnAddBlendSpace()
+	void SDatabaseAssetTree::OnAddBlendSpace(bool bFinalizeChanges)
 	{
+		FScopedTransaction Transaction(LOCTEXT("AddBlendSpace", "Add Blend Space"));
+
 		EditorViewModel.Pin()->AddBlendSpaceToDatabase(nullptr, INDEX_NONE);
-		RefreshTreeView(false);
+
+		if (bFinalizeChanges)
+		{
+			FinalizeTreeChanges();
+		}
 	}
 
-	void SDatabaseAssetTree::OnDeleteAsset(TSharedPtr<FDatabaseAssetTreeNode> Node)
+	void SDatabaseAssetTree::OnDeleteAsset(TSharedPtr<FDatabaseAssetTreeNode> Node, bool bFinalizeChanges)
 	{
+		FScopedTransaction Transaction(LOCTEXT("DeleteAsset", "Delete Asset"));
+
 		if (Node->SourceAssetType == ESearchIndexAssetType::Sequence)
 		{
 			EditorViewModel.Pin()->DeleteSequenceFromDatabase(Node->SourceAssetIdx);
@@ -756,14 +806,20 @@ namespace UE::PoseSearch
 			checkNoEntry();
 		}
 
-		RefreshTreeView(false);
+		if (bFinalizeChanges)
+		{
+			FinalizeTreeChanges();
+		}
 	}
 
-	void SDatabaseAssetTree::OnRemoveFromGroup(TSharedPtr<FDatabaseAssetTreeNode> Node)
+	void SDatabaseAssetTree::OnRemoveFromGroup(TSharedPtr<FDatabaseAssetTreeNode> Node, bool bFinalizeChanges)
 	{
 		check(Node->Parent->SourceAssetType == ESearchIndexAssetType::Invalid);
+		
 		const int32 GroupIdx = Node->Parent->SourceAssetIdx;
 
+		FScopedTransaction Transaction(LOCTEXT("RemoveFromGroup", "Remove Asset From Group"));
+
 		if (Node->SourceAssetType == ESearchIndexAssetType::Sequence)
 		{
 			EditorViewModel.Pin()->RemoveSequenceFromGroup(Node->SourceAssetIdx, GroupIdx);
@@ -777,16 +833,25 @@ namespace UE::PoseSearch
 			checkNoEntry();
 		}
 
-		RefreshTreeView(false);
+		if (bFinalizeChanges)
+		{
+			FinalizeTreeChanges();
+		}
 	}
 
-	void SDatabaseAssetTree::OnDeleteGroup(TSharedPtr<FDatabaseAssetTreeNode> Node)
+	void SDatabaseAssetTree::OnDeleteGroup(TSharedPtr<FDatabaseAssetTreeNode> Node, bool bFinalizeChanges)
 	{
 		check(Node->SourceAssetType == ESearchIndexAssetType::Invalid);
+
+		FScopedTransaction Transaction(LOCTEXT("DeleteGroup", "Delete Group"));
+
 		const int32 GroupIdx = Node->SourceAssetIdx;
 		EditorViewModel.Pin()->DeleteGroup(GroupIdx);
 
-		RefreshTreeView(false);
+		if (bFinalizeChanges)
+		{
+			FinalizeTreeChanges();
+		}
 	}
 
 	void SDatabaseAssetTree::RegisterOnSelectionChanged(const FOnSelectionChanged& Delegate)
@@ -821,6 +886,83 @@ namespace UE::PoseSearch
 
 		TreeView->SetItemSelection(NewSelectedNodes, true, ESelectInfo::Direct);
 	}
+
+	void SDatabaseAssetTree::CreateCommandList()
+	{
+		CommandList = MakeShared<FUICommandList>();
+
+		CommandList->MapAction(
+			FGenericCommands::Get().Delete,
+			FUIAction(
+				FExecuteAction::CreateSP(this, &SDatabaseAssetTree::OnDeleteNodes),
+				FCanExecuteAction::CreateSP(this, &SDatabaseAssetTree::CanDeleteNodes)));
+	}
+
+	bool SDatabaseAssetTree::CanDeleteNodes() const
+	{
+		TArray<TSharedPtr<FDatabaseAssetTreeNode>> SelectedNodes = TreeView->GetSelectedItems();
+		for (TSharedPtr<FDatabaseAssetTreeNode> SelectedNode : SelectedNodes)
+		{
+			if (SelectedNode->SourceAssetType != ESearchIndexAssetType::Invalid ||
+				SelectedNode->SourceAssetIdx != INDEX_NONE)
+			{
+				return true;
+			}
+		}
+
+		return false;
+	}
+
+	void SDatabaseAssetTree::OnDeleteNodes()
+	{
+		TArray<TSharedPtr<FDatabaseAssetTreeNode>> SelectedNodes = TreeView->GetSelectedItems();
+		if (!SelectedNodes.IsEmpty())
+		{
+			SelectedNodes.Sort(
+				[](const TSharedPtr<FDatabaseAssetTreeNode>& A, const TSharedPtr<FDatabaseAssetTreeNode>& B)
+			{
+				if (A->SourceAssetType != ESearchIndexAssetType::Invalid &&
+					B->SourceAssetType == ESearchIndexAssetType::Invalid)
+				{
+					return true;
+				}
+				if (B->SourceAssetType != ESearchIndexAssetType::Invalid &&
+					A->SourceAssetType == ESearchIndexAssetType::Invalid)
+				{
+					return false;
+				}
+				return B->SourceAssetIdx < A->SourceAssetIdx;
+			});
+
+			for (TSharedPtr<FDatabaseAssetTreeNode> SelectedNode : SelectedNodes)
+			{
+				if (SelectedNode->SourceAssetType != ESearchIndexAssetType::Invalid)
+				{
+					const int32 GroupIdx = SelectedNode->Parent->SourceAssetIdx;
+					if (GroupIdx == INDEX_NONE)
+					{
+						OnDeleteAsset(SelectedNode, false);
+					}
+					else
+					{
+						OnRemoveFromGroup(SelectedNode, false);
+					}
+				}
+				else if (SelectedNode->SourceAssetIdx != INDEX_NONE)
+				{
+					OnDeleteGroup(SelectedNode, false);
+				}
+			}
+
+			FinalizeTreeChanges();
+		}
+	}
+
+	void SDatabaseAssetTree::FinalizeTreeChanges(bool bRecoverSelection)
+	{
+		RefreshTreeView(false, bRecoverSelection);
+		EditorViewModel.Pin()->BuildSearchIndex();
+	}
 }
 
 #undef LOCTEXT_NAMESPACE
diff --git a/Engine/Plugins/Experimental/Animation/PoseSearch/Source/Editor/Private/SPoseSearchDatabaseViewport.cpp b/Engine/Plugins/Experimental/Animation/PoseSearch/Source/Editor/Private/SPoseSearchDatabaseViewport.cpp
--- a/Engine/Plugins/Experimental/Animation/PoseSearch/Source/Editor/Private/SPoseSearchDatabaseViewport.cpp
+++ b/Engine/Plugins/Experimental/Animation/PoseSearch/Source/Editor/Private/SPoseSearchDatabaseViewport.cpp
@@ -56,6 +56,18 @@ namespace UE::PoseSearch
 				&FDatabaseViewModel::IsPoseFeaturesDrawMode,
 				EFeaturesDrawMode::All));
 
+		CommandList->MapAction(
+			Commands.ShowAnimationNone,
+			FExecuteAction::CreateSP(
+				ViewModelRef,
+				&FDatabaseViewModel::OnSetAnimationPreviewMode,
+				EAnimationPreviewMode::None),
+			FCanExecuteAction(),
+			FIsActionChecked::CreateSP(
+				ViewModelRef,
+				&FDatabaseViewModel::IsAnimationPreviewMode,
+				EAnimationPreviewMode::None));
+
 		CommandList->MapAction(
 			Commands.ShowAnimationOriginalOnly,
 			FExecuteAction::CreateSP(
diff --git a/Engine/Plugins/Experimental/Animation/PoseSearch/Source/Editor/Private/SPoseSearchDatabaseViewportToolbar.cpp b/Engine/Plugins/Experimental/Animation/PoseSearch/Source/Editor/Private/SPoseSearchDatabaseViewportToolbar.cpp
--- a/Engine/Plugins/Experimental/Animation/PoseSearch/Source/Editor/Private/SPoseSearchDatabaseViewportToolbar.cpp
+++ b/Engine/Plugins/Experimental/Animation/PoseSearch/Source/Editor/Private/SPoseSearchDatabaseViewportToolbar.cpp
@@ -54,6 +54,7 @@ namespace UE::PoseSearch
 
 				SubMenuBuilder.BeginSection("Animations", LOCTEXT("ShowMenu_AnimationsLabel", "Animations"));
 				{
+					SubMenuBuilder.AddMenuEntry(Commands.ShowAnimationNone);
 					SubMenuBuilder.AddMenuEntry(Commands.ShowAnimationOriginalOnly);
 					SubMenuBuilder.AddMenuEntry(Commands.ShowAnimationOriginalAndMirrored);
 				}
diff --git a/Engine/Plugins/Experimental/Animation/PoseSearch/Source/Editor/Public/PoseSearchCustomization.h b/Engine/Plugins/Experimental/Animation/PoseSearch/Source/Editor/Public/PoseSearchCustomization.h
--- a/Engine/Plugins/Experimental/Animation/PoseSearch/Source/Editor/Public/PoseSearchCustomization.h
+++ b/Engine/Plugins/Experimental/Animation/PoseSearch/Source/Editor/Public/PoseSearchCustomization.h
@@ -50,13 +50,13 @@ public:
 /**
  * Customizes UPoseSearchDatabase to improve workflow
  */
-//class FPoseSearchDatabaseDetails : public IDetailCustomization
-//{
-//public:
-//	// Creates an instance of FObjectDetails
-//	static TSharedRef<IDetailCustomization> MakeInstance();
-//
-//	// IDetailCustomization interface
-//	virtual void CustomizeDetails(IDetailLayoutBuilder& DetailLayout) override;
-//	// End of IDetailCustomization interface
-//};
+class FPoseSearchDatabaseDetails : public IDetailCustomization
+{
+public:
+	// Creates an instance of FObjectDetails
+	static TSharedRef<IDetailCustomization> MakeInstance();
+
+	// IDetailCustomization interface
+	virtual void CustomizeDetails(IDetailLayoutBuilder& DetailLayout) override;
+	// End of IDetailCustomization interface
+};
diff --git a/Engine/Plugins/Experimental/Animation/PoseSearch/Source/Editor/Public/PoseSearchDatabaseEditorCommands.h b/Engine/Plugins/Experimental/Animation/PoseSearch/Source/Editor/Public/PoseSearchDatabaseEditorCommands.h
--- a/Engine/Plugins/Experimental/Animation/PoseSearch/Source/Editor/Public/PoseSearchDatabaseEditorCommands.h
+++ b/Engine/Plugins/Experimental/Animation/PoseSearch/Source/Editor/Public/PoseSearchDatabaseEditorCommands.h
@@ -24,8 +24,6 @@ namespace UE::PoseSearch
 		virtual void RegisterCommands() override;
 
 	public:
-		TSharedPtr<FUICommandInfo> StopPreviewScene;
-
 		TSharedPtr<FUICommandInfo> ResetPreviewScene;
 
 		TSharedPtr<FUICommandInfo> BuildSearchIndex;
@@ -34,6 +32,8 @@ namespace UE::PoseSearch
 
 		TSharedPtr<FUICommandInfo> ShowPoseFeaturesAll;
 
+		TSharedPtr<FUICommandInfo> ShowAnimationNone;
+
 		TSharedPtr<FUICommandInfo> ShowAnimationOriginalOnly;
 
 		TSharedPtr<FUICommandInfo> ShowAnimationOriginalAndMirrored;
diff --git a/Engine/Plugins/Experimental/Animation/PoseSearch/Source/Editor/Public/PoseSearchDatabaseEditorReflection.h b/Engine/Plugins/Experimental/Animation/PoseSearch/Source/Editor/Public/PoseSearchDatabaseEditorReflection.h
--- a/Engine/Plugins/Experimental/Animation/PoseSearch/Source/Editor/Public/PoseSearchDatabaseEditorReflection.h
+++ b/Engine/Plugins/Experimental/Animation/PoseSearch/Source/Editor/Public/PoseSearchDatabaseEditorReflection.h
@@ -51,7 +51,7 @@ class UPoseSearchDatabaseBlendSpaceReflection : public UPoseSearchDatabaseReflec
 	GENERATED_BODY()
 
 public:
-	UPROPERTY(EditAnywhere, Category = "Selected Sequence")
+	UPROPERTY(EditAnywhere, Category = "Selected Blend Space")
 	FPoseSearchDatabaseBlendSpace BlendSpace;
 
 #if WITH_EDITOR
@@ -65,7 +65,7 @@ class UPoseSearchDatabaseGroupReflection : public UPoseSearchDatabaseReflectionB
 	GENERATED_BODY()
 
 public:
-	UPROPERTY(EditAnywhere, Category = "Selected Sequence")
+	UPROPERTY(EditAnywhere, Category = "Selected Group")
 	FPoseSearchDatabaseGroup Group;
 
 #if WITH_EDITOR
diff --git a/Engine/Plugins/Experimental/Animation/PoseSearch/Source/Editor/Public/PoseSearchDatabaseEditorToolkit.h b/Engine/Plugins/Experimental/Animation/PoseSearch/Source/Editor/Public/PoseSearchDatabaseEditorToolkit.h
--- a/Engine/Plugins/Experimental/Animation/PoseSearch/Source/Editor/Public/PoseSearchDatabaseEditorToolkit.h
+++ b/Engine/Plugins/Experimental/Animation/PoseSearch/Source/Editor/Public/PoseSearchDatabaseEditorToolkit.h
@@ -58,7 +58,7 @@ namespace UE::PoseSearch
 		void ExtendToolbar();
 		void FillToolbar(FToolBarBuilder& ToolbarBuilder);
 
-		void OnFinishedChangingProperties(const FPropertyChangedEvent& PropertyChangedEvent);
+		void OnFinishedChangingSelectionProperties(const FPropertyChangedEvent& PropertyChangedEvent);
 		void OnAssetTreeSelectionChanged(
 			const TArrayView<TSharedPtr<FDatabaseAssetTreeNode>>& SelectedItems,
 			ESelectInfo::Type SelectionType);
@@ -69,13 +69,11 @@ namespace UE::PoseSearch
 
 		TSharedPtr<IDetailsView> EditingAssetWidget;
 
-		TSharedPtr<IDetailsView> SelectionWidget;
+		TArray<TSharedPtr<IDetailsView>> SelectionWidgets;
 
 		TSharedPtr<FDatabasePreviewScene> PreviewScene;
 
 		TSharedPtr<FDatabaseViewModel> ViewModel;
-
-		TArray<TWeakObjectPtr<UObject>> SelectionReflection;
 	};
 }
 
diff --git a/Engine/Plugins/Experimental/Animation/PoseSearch/Source/Editor/Public/PoseSearchDatabaseViewModel.h b/Engine/Plugins/Experimental/Animation/PoseSearch/Source/Editor/Public/PoseSearchDatabaseViewModel.h
--- a/Engine/Plugins/Experimental/Animation/PoseSearch/Source/Editor/Public/PoseSearchDatabaseViewModel.h
+++ b/Engine/Plugins/Experimental/Animation/PoseSearch/Source/Editor/Public/PoseSearchDatabaseViewModel.h
@@ -19,6 +19,8 @@ class UMirrorDataTable;
 
 namespace UE::PoseSearch
 {
+	class FDatabaseAssetTreeNode;
+
 	enum class EFeaturesDrawMode : uint8
 	{
 		None,
@@ -27,6 +29,7 @@ namespace UE::PoseSearch
 
 	enum class EAnimationPreviewMode : uint8
 	{
+		None,
 		OriginalOnly,
 		OriginalAndMirrored
 	};
@@ -92,6 +95,7 @@ namespace UE::PoseSearch
 		void RemoveBlendSpaceFromGroup(int32 BlendSpaceIdx, int32 GroupIdx);
 		void DeleteGroup(int32 GroupIdx);
 
+		void SetSelectedNodes(const TArrayView<TSharedPtr<FDatabaseAssetTreeNode>>& InSelectedNodes);
 
 	private:
 		float PlayTime = 0.0f;
@@ -111,6 +115,8 @@ namespace UE::PoseSearch
 		/** What animations to show in the viewport */
 		EAnimationPreviewMode AnimationPreviewMode = EAnimationPreviewMode::OriginalOnly;
 
+		TArray<TSharedPtr<FDatabaseAssetTreeNode>> SelectedNodes;
+
 		UWorld* GetWorld() const;
 
 		UObject* GetPlaybackContext() const;
diff --git a/Engine/Plugins/Experimental/Animation/PoseSearch/Source/Editor/Public/SPoseSearchDatabaseAssetList.h b/Engine/Plugins/Experimental/Animation/PoseSearch/Source/Editor/Public/SPoseSearchDatabaseAssetList.h
--- a/Engine/Plugins/Experimental/Animation/PoseSearch/Source/Editor/Public/SPoseSearchDatabaseAssetList.h
+++ b/Engine/Plugins/Experimental/Animation/PoseSearch/Source/Editor/Public/SPoseSearchDatabaseAssetList.h
@@ -77,9 +77,11 @@ namespace UE::PoseSearch
 		// SWidget interface
 		virtual FReply OnDragOver(const FGeometry& MyGeometry, const FDragDropEvent& DragDropEvent) override;
 		virtual FReply OnDrop(const FGeometry& MyGeometry, const FDragDropEvent& DragDropEvent) override;
+		virtual FReply OnKeyDown(const FGeometry& MyGeometry, const FKeyEvent& InKeyEvent) override;
 		// End SWidget interface
 	
 		void RefreshTreeView(bool bIsInitialSetup = false, bool bRecoverSelection = false);
+		void FinalizeTreeChanges(bool bRecoverSelection = false);
 
 	protected:
 		TWeakPtr<FDatabaseViewModel> EditorViewModel;
@@ -114,13 +116,19 @@ namespace UE::PoseSearch
 		TSharedRef<SWidget> CreateAddNewMenuWidget();
 		TSharedPtr<SWidget> CreateContextMenu();
 		
-		void OnAddGroup();
-		void OnAddSequence();
-		void OnAddBlendSpace();
+		void OnAddGroup(bool bFinalizeChanges = true);
+		void OnAddSequence(bool bFinalizeChanges = true);
+		void OnAddBlendSpace(bool bFinalizeChanges = true);
+
+		void OnDeleteAsset(TSharedPtr<FDatabaseAssetTreeNode> Node, bool bFinalizeChanges = true);
+		void OnRemoveFromGroup(TSharedPtr<FDatabaseAssetTreeNode> Node, bool bFinalizeChanges = true);
+		void OnDeleteGroup(TSharedPtr<FDatabaseAssetTreeNode> Node, bool bFinalizeChanges = true);
+
+		void CreateCommandList();
 
-		void OnDeleteAsset(TSharedPtr<FDatabaseAssetTreeNode> Node);
-		void OnRemoveFromGroup(TSharedPtr<FDatabaseAssetTreeNode> Node);
-		void OnDeleteGroup(TSharedPtr<FDatabaseAssetTreeNode> Node);
+		/** Removes existing selected component nodes from the tree*/
+		bool CanDeleteNodes() const;
+		void OnDeleteNodes();
 
 		friend SDatabaseAssetListItem;
 
diff --git a/Engine/Plugins/Experimental/Animation/PoseSearch/Source/Runtime/Private/PoseSearch.cpp b/Engine/Plugins/Experimental/Animation/PoseSearch/Source/Runtime/Private/PoseSearch.cpp
--- a/Engine/Plugins/Experimental/Animation/PoseSearch/Source/Runtime/Private/PoseSearch.cpp
+++ b/Engine/Plugins/Experimental/Animation/PoseSearch/Source/Runtime/Private/PoseSearch.cpp
@@ -1035,6 +1035,7 @@ TArrayView<const float> FPoseSearchIndex::GetPoseValues(int32 PoseIdx) const
 void FPoseSearchIndex::Reset()
 {
 	NumPoses = 0;
+	Assets.Reset();
 	Values.Reset();
 	Schema = nullptr;
 }
@@ -1728,6 +1729,16 @@ void UPoseSearchDatabase::BeginCacheDerivedData()
 		PrivateDerivedData->Cache(*this, false);
 	}
 }
+
+FIoHash UPoseSearchDatabase::GetSearchIndexHash() const
+{
+	if (!PrivateDerivedData)
+	{
+		return FIoHash::Zero;
+	}
+
+	return PrivateDerivedData->DerivedDataKey.Hash;
+}
 #endif // WITH_EDITOR
 
 
diff --git a/Engine/Plugins/Experimental/Animation/PoseSearch/Source/Runtime/Private/PoseSearchDerivedData.cpp b/Engine/Plugins/Experimental/Animation/PoseSearch/Source/Runtime/Private/PoseSearchDerivedData.cpp
--- a/Engine/Plugins/Experimental/Animation/PoseSearch/Source/Runtime/Private/PoseSearchDerivedData.cpp
+++ b/Engine/Plugins/Experimental/Animation/PoseSearch/Source/Runtime/Private/PoseSearchDerivedData.cpp
@@ -27,6 +27,12 @@ void FPoseSearchDatabaseDerivedData::Cache(UPoseSearchDatabase& Database, bool b
 	{
 		CreateDatabaseBuildTask(Database, bForceRebuild);
 	}
+	else
+	{
+		SearchIndex.Reset();
+		SearchIndex.Schema = Database.Schema;
+		DerivedDataKey = { UE::DerivedData::FCacheBucket(), FIoHash::Zero };
+	}
 }
 
 void FPoseSearchDatabaseDerivedData::CancelCache()
diff --git a/Engine/Plugins/Experimental/Animation/PoseSearch/Source/Runtime/Private/PoseSearchLibrary.cpp b/Engine/Plugins/Experimental/Animation/PoseSearch/Source/Runtime/Private/PoseSearchLibrary.cpp
--- a/Engine/Plugins/Experimental/Animation/PoseSearch/Source/Runtime/Private/PoseSearchLibrary.cpp
+++ b/Engine/Plugins/Experimental/Animation/PoseSearch/Source/Runtime/Private/PoseSearchLibrary.cpp
@@ -76,7 +76,16 @@ bool FMotionMatchingState::InitNewDatabaseSearch(
 	if (bValidDatabase)
 	{
 		CurrentDatabase = Database;
-		Reset();
+#if WITH_EDITOR
+		CurrentSearchIndexHash = Database->GetSearchIndexHash();
+#endif
+	}
+	else
+	{
+		CurrentDatabase = nullptr;
+#if WITH_EDITOR
+		CurrentSearchIndexHash = FIoHash::Zero;
+#endif	
 	}
 
 	if (!bValidDatabase && OutError)
@@ -91,6 +100,8 @@ bool FMotionMatchingState::InitNewDatabaseSearch(
 		}
 	}
 
+	Reset();
+
 	return bValidDatabase;
 }
 
@@ -333,7 +344,7 @@ void UpdateMotionMatchingState(
 	}
 	
 	// Check for a switch in the database
-	if (InOutMotionMatchingState.CurrentDatabase != Database)
+	if (!InOutMotionMatchingState.IsCompatibleDatabase(Database))
 	{
 		FText InitError;
 		if (!InOutMotionMatchingState.InitNewDatabaseSearch(Database, &InitError))
@@ -599,4 +610,21 @@ float FMotionMatchingState::ComputeJumpBlendTime(
 	return JumpBlendTime;
 }
 
+bool FMotionMatchingState::IsCompatibleDatabase(const UPoseSearchDatabase* Database) const
+{
+	if (Database != CurrentDatabase)
+	{
+		return false;
+	}
+
+#if WITH_EDITOR
+	if (Database->GetSearchIndexHash() != CurrentSearchIndexHash)
+	{
+		return false;
+	}
+#endif
+
+	return true;
+}
+
 #undef LOCTEXT_NAMESPACE
\ No newline at end of file
diff --git a/Engine/Plugins/Experimental/Animation/PoseSearch/Source/Runtime/Public/PoseSearch/PoseSearch.h b/Engine/Plugins/Experimental/Animation/PoseSearch/Source/Runtime/Public/PoseSearch/PoseSearch.h
--- a/Engine/Plugins/Experimental/Animation/PoseSearch/Source/Runtime/Public/PoseSearch/PoseSearch.h
+++ b/Engine/Plugins/Experimental/Animation/PoseSearch/Source/Runtime/Public/PoseSearch/PoseSearch.h
@@ -855,6 +855,8 @@ public:
 	void NotifyGroupChange();
 
 	void BeginCacheDerivedData();
+
+	FIoHash GetSearchIndexHash() const;
 #endif // WITH_EDITOR
 
 	bool IsDerivedDataValid();
diff --git a/Engine/Plugins/Experimental/Animation/PoseSearch/Source/Runtime/Public/PoseSearch/PoseSearchLibrary.h b/Engine/Plugins/Experimental/Animation/PoseSearch/Source/Runtime/Public/PoseSearch/PoseSearchLibrary.h
--- a/Engine/Plugins/Experimental/Animation/PoseSearch/Source/Runtime/Public/PoseSearch/PoseSearchLibrary.h
+++ b/Engine/Plugins/Experimental/Animation/PoseSearch/Source/Runtime/Public/PoseSearch/PoseSearchLibrary.h
@@ -83,6 +83,8 @@ struct POSESEARCH_API FMotionMatchingState
 
 	float ComputeJumpBlendTime(const UE::PoseSearch::FSearchResult& Result, const FMotionMatchingSettings& Settings) const;
 
+	bool IsCompatibleDatabase(const UPoseSearchDatabase* Database) const;
+
 	// The current pose we're playing from the database
 	UPROPERTY(Transient)
 	int32 DbPoseIdx = INDEX_NONE;
@@ -103,6 +105,10 @@ struct POSESEARCH_API FMotionMatchingState
 	UPROPERTY(Transient)
 	TWeakObjectPtr<const UPoseSearchDatabase> CurrentDatabase = nullptr;
 
+#if WITH_EDITOR
+	FIoHash CurrentSearchIndexHash = FIoHash::Zero;
+#endif // _DEBUG
+
 	// Time since the last pose jump
 	UPROPERTY(Transient)
 	float ElapsedPoseJumpTime = 0.f;
