Fixed non-unity build

diff --git a/Engine/Plugins/Experimental/Animation/PoseSearch/Source/Editor/Private/SPoseSearchDatabaseAssetList.cpp b/Engine/Plugins/Experimental/Animation/PoseSearch/Source/Editor/Private/SPoseSearchDatabaseAssetList.cpp
--- a/Engine/Plugins/Experimental/Animation/PoseSearch/Source/Editor/Private/SPoseSearchDatabaseAssetList.cpp
+++ b/Engine/Plugins/Experimental/Animation/PoseSearch/Source/Editor/Private/SPoseSearchDatabaseAssetList.cpp
@@ -17,6 +17,8 @@
 #include "Framework/MultiBox/MultiBoxBuilder.h"
 #include "SPositiveActionButton.h"
 
+#include "ScopedTransaction.h"
+
 
 #define LOCTEXT_NAMESPACE "PoseSearchDatabaseAssetList"
 
