PoseSearch - added several checks to make sure the editor doesn't crash when the pose search database is edited while it's being searched.

diff --git a/Engine/Plugins/Experimental/Animation/PoseSearch/Source/Runtime/Private/AnimNode_MotionMatching.cpp b/Engine/Plugins/Experimental/Animation/PoseSearch/Source/Runtime/Private/AnimNode_MotionMatching.cpp
--- a/Engine/Plugins/Experimental/Animation/PoseSearch/Source/Runtime/Private/AnimNode_MotionMatching.cpp
+++ b/Engine/Plugins/Experimental/Animation/PoseSearch/Source/Runtime/Private/AnimNode_MotionMatching.cpp
@@ -70,8 +70,17 @@ void FAnimNode_MotionMatching::UpdateAssetPlayer(const FAnimationUpdateContext&
 
 	GetEvaluateGraphExposedInputs().Execute(Context);
 
+	bool bNeedsReset =
+		bResetOnBecomingRelevant &&
+		UpdateCounter.HasEverBeenUpdated() &&
+		!UpdateCounter.WasSynchronizedCounter(Context.AnimInstanceProxy->GetUpdateCounter());
+
+#if WITH_EDITOR
+	bNeedsReset = bNeedsReset || MotionMatchingState.HasSearchIndexChanged();
+#endif
+
 	// If we just became relevant and haven't been initialized yet, then reset motion matching state, otherwise update the asset time using the player node.
-	if (bResetOnBecomingRelevant && UpdateCounter.HasEverBeenUpdated() && !UpdateCounter.WasSynchronizedCounter(Context.AnimInstanceProxy->GetUpdateCounter()))
+	if (bNeedsReset)
 	{
 		MotionMatchingState.Reset();
 	}
diff --git a/Engine/Plugins/Experimental/Animation/PoseSearch/Source/Runtime/Private/PoseSearch.cpp b/Engine/Plugins/Experimental/Animation/PoseSearch/Source/Runtime/Private/PoseSearch.cpp
--- a/Engine/Plugins/Experimental/Animation/PoseSearch/Source/Runtime/Private/PoseSearch.cpp
+++ b/Engine/Plugins/Experimental/Animation/PoseSearch/Source/Runtime/Private/PoseSearch.cpp
@@ -1739,6 +1739,16 @@ FIoHash UPoseSearchDatabase::GetSearchIndexHash() const
 
 	return PrivateDerivedData->DerivedDataKey.Hash;
 }
+
+bool UPoseSearchDatabase::IsDerivedDataBuildPending() const
+{
+	if (!PrivateDerivedData)
+	{
+		return true;
+	}
+
+	return PrivateDerivedData->DerivedDataKey.Hash != PrivateDerivedData->PendingDerivedDataKey;
+}
 #endif // WITH_EDITOR
 
 
diff --git a/Engine/Plugins/Experimental/Animation/PoseSearch/Source/Runtime/Private/PoseSearchLibrary.cpp b/Engine/Plugins/Experimental/Animation/PoseSearch/Source/Runtime/Private/PoseSearchLibrary.cpp
--- a/Engine/Plugins/Experimental/Animation/PoseSearch/Source/Runtime/Private/PoseSearchLibrary.cpp
+++ b/Engine/Plugins/Experimental/Animation/PoseSearch/Source/Runtime/Private/PoseSearchLibrary.cpp
@@ -354,6 +354,17 @@ void UpdateMotionMatchingState(
 		}
 	}
 
+#if WITH_EDITOR
+	if (InOutMotionMatchingState.CurrentDatabase->IsDerivedDataBuildPending())
+	{
+		InOutMotionMatchingState.Reset();
+		Context.LogMessage(
+			EMessageSeverity::Error, 
+			LOCTEXT("PendingDerivedData", "Derived data build pending."));
+		return;
+	}
+#endif
+
 	const float DeltaTime = Context.GetDeltaTime();
 
 	// Reset State Flags
@@ -627,4 +638,15 @@ bool FMotionMatchingState::IsCompatibleDatabase(const UPoseSearchDatabase* Datab
 	return true;
 }
 
+#if WITH_EDITOR
+bool FMotionMatchingState::HasSearchIndexChanged() const
+{
+	const bool bIsConsistent = 
+		!CurrentDatabase.IsValid() ||
+		CurrentDatabase->IsDerivedDataBuildPending() ||
+		CurrentDatabase->GetSearchIndexHash() != CurrentSearchIndexHash;
+	return bIsConsistent;
+}
+#endif
+
 #undef LOCTEXT_NAMESPACE
\ No newline at end of file
diff --git a/Engine/Plugins/Experimental/Animation/PoseSearch/Source/Runtime/Public/PoseSearch/PoseSearch.h b/Engine/Plugins/Experimental/Animation/PoseSearch/Source/Runtime/Public/PoseSearch/PoseSearch.h
--- a/Engine/Plugins/Experimental/Animation/PoseSearch/Source/Runtime/Public/PoseSearch/PoseSearch.h
+++ b/Engine/Plugins/Experimental/Animation/PoseSearch/Source/Runtime/Public/PoseSearch/PoseSearch.h
@@ -857,6 +857,8 @@ public:
 	void BeginCacheDerivedData();
 
 	FIoHash GetSearchIndexHash() const;
+
+	bool IsDerivedDataBuildPending() const;
 #endif // WITH_EDITOR
 
 	bool IsDerivedDataValid();
diff --git a/Engine/Plugins/Experimental/Animation/PoseSearch/Source/Runtime/Public/PoseSearch/PoseSearchLibrary.h b/Engine/Plugins/Experimental/Animation/PoseSearch/Source/Runtime/Public/PoseSearch/PoseSearchLibrary.h
--- a/Engine/Plugins/Experimental/Animation/PoseSearch/Source/Runtime/Public/PoseSearch/PoseSearchLibrary.h
+++ b/Engine/Plugins/Experimental/Animation/PoseSearch/Source/Runtime/Public/PoseSearch/PoseSearchLibrary.h
@@ -85,6 +85,10 @@ struct POSESEARCH_API FMotionMatchingState
 
 	bool IsCompatibleDatabase(const UPoseSearchDatabase* Database) const;
 
+#if WITH_EDITOR
+	bool HasSearchIndexChanged() const;
+#endif 
+
 	// The current pose we're playing from the database
 	UPROPERTY(Transient)
 	int32 DbPoseIdx = INDEX_NONE;
