PoseSearch - Database editor update
- added scrubbing timeline to preview
- clicked actors are highlighted in the asset tree

diff --git a/Engine/Plugins/Experimental/Animation/PoseSearch/Source/Editor/Private/PoseSearchDatabaseEdMode.cpp b/Engine/Plugins/Experimental/Animation/PoseSearch/Source/Editor/Private/PoseSearchDatabaseEdMode.cpp
--- a/Engine/Plugins/Experimental/Animation/PoseSearch/Source/Editor/Private/PoseSearchDatabaseEdMode.cpp
+++ b/Engine/Plugins/Experimental/Animation/PoseSearch/Source/Editor/Private/PoseSearchDatabaseEdMode.cpp
@@ -47,7 +47,18 @@ namespace UE::PoseSearch
 		HHitProxy* HitProxy,
 		const FViewportClick& Click)
 	{
-		return FEdMode::HandleClick(InViewportClient, HitProxy, Click);
+		if (HActor* ActorHitProxy = HitProxyCast<HActor>(HitProxy))
+		{
+			if (ViewModel && IsValid(ActorHitProxy->Actor))
+			{
+				ViewModel->ProcessSelectedActor(ActorHitProxy->Actor);
+				return true;
+			}
+		}
+
+		ViewModel->ProcessSelectedActor(nullptr);
+
+		return false; // unhandled
 	}
 
 
diff --git a/Engine/Plugins/Experimental/Animation/PoseSearch/Source/Editor/Private/PoseSearchDatabaseEditorCommands.cpp b/Engine/Plugins/Experimental/Animation/PoseSearch/Source/Editor/Private/PoseSearchDatabaseEditorCommands.cpp
--- a/Engine/Plugins/Experimental/Animation/PoseSearch/Source/Editor/Private/PoseSearchDatabaseEditorCommands.cpp
+++ b/Engine/Plugins/Experimental/Animation/PoseSearch/Source/Editor/Private/PoseSearchDatabaseEditorCommands.cpp
@@ -8,7 +8,6 @@ namespace UE::PoseSearch
 {
 	void FDatabaseEditorCommands::RegisterCommands()
 	{
-		UI_COMMAND(ResetPreviewScene, "Reset Scene", "Reset Scene", EUserInterfaceActionType::Button, FInputChord());
 		UI_COMMAND(BuildSearchIndex, "Build Index", "Build Index", EUserInterfaceActionType::Button, FInputChord());
 
 		UI_COMMAND(
diff --git a/Engine/Plugins/Experimental/Animation/PoseSearch/Source/Editor/Private/PoseSearchDatabaseEditorToolkit.cpp b/Engine/Plugins/Experimental/Animation/PoseSearch/Source/Editor/Private/PoseSearchDatabaseEditorToolkit.cpp
--- a/Engine/Plugins/Experimental/Animation/PoseSearch/Source/Editor/Private/PoseSearchDatabaseEditorToolkit.cpp
+++ b/Engine/Plugins/Experimental/Animation/PoseSearch/Source/Editor/Private/PoseSearchDatabaseEditorToolkit.cpp
@@ -25,6 +25,7 @@ namespace UE::PoseSearch
 {
 	const FName PoseSearchDatabaseEditorAppName = FName(TEXT("PoseSearchDatabaseEditorApp"));
 	constexpr int32 SelectionDetailsCount = 3; // sequence, blendspace, group
+	constexpr double ViewRangeSlack = 0.25;
 
 	// Tab identifiers
 	struct FDatabaseEditorTabs
@@ -66,19 +67,44 @@ namespace UE::PoseSearch
 		return ViewModel.IsValid() ? ViewModel->GetPoseSearchDatabase() : nullptr;
 	}
 
-	void FDatabaseEditorToolkit::StopPreviewScene()
+	void FDatabaseEditorToolkit::BuildSearchIndex()
 	{
-		ViewModel->RemovePreviewActors();
+		ViewModel->BuildSearchIndex();
 	}
 
-	void FDatabaseEditorToolkit::ResetPreviewScene()
+	void FDatabaseEditorToolkit::PreviewBackwardEnd()
 	{
-		ViewModel->ResetPreviewActors();
+		ViewModel->PreviewBackwardEnd();
 	}
 
-	void FDatabaseEditorToolkit::BuildSearchIndex()
+	void FDatabaseEditorToolkit::PreviewBackwardStep()
 	{
-		ViewModel->BuildSearchIndex();
+		ViewModel->PreviewBackwardStep();
+	}
+
+	void FDatabaseEditorToolkit::PreviewBackward()
+	{
+		ViewModel->PreviewBackward();
+	}
+
+	void FDatabaseEditorToolkit::PreviewPause()
+	{
+		ViewModel->PreviewPause();
+	}
+
+	void FDatabaseEditorToolkit::PreviewForward()
+	{
+		ViewModel->PreviewForward();
+	}
+
+	void FDatabaseEditorToolkit::PreviewForwardStep()
+	{
+		ViewModel->PreviewForwardStep();
+	}
+
+	void FDatabaseEditorToolkit::PreviewForwardEnd()
+	{
+		ViewModel->PreviewForwardEnd();
 	}
 
 	void FDatabaseEditorToolkit::InitAssetEditor(
@@ -109,10 +135,26 @@ namespace UE::PoseSearch
 		ViewModel->Initialize(DatabaseAsset, PreviewScene.ToSharedRef());
 
 		// Create viewport widget
-		FDatabaseViewportRequiredArgs ViewportArgs(
+		FDatabasePreviewRequiredArgs PreviewArgs(
 			StaticCastSharedRef<FDatabaseEditorToolkit>(AsShared()),
 			PreviewScene.ToSharedRef());
-		ViewportWidget = SNew(SDatabaseViewport, ViewportArgs);
+		PreviewWidget = SNew(SDatabasePreview, PreviewArgs)
+			.SliderScrubTime_Lambda([this]() { return ViewModel->GetPlayTime(); })
+			.SliderViewRange_Lambda([this]() 
+			{ 
+				return TRange<double>(-ViewRangeSlack, ViewModel->GetMaxPreviewPlayLength() + ViewRangeSlack);
+			})
+			.OnSliderScrubPositionChanged_Lambda([this](float NewScrubPosition, bool bScrubbing)
+			{
+				ViewModel->SetPlayTime(NewScrubPosition, !bScrubbing);
+			})
+			.OnBackwardEnd_Raw(this, &FDatabaseEditorToolkit::PreviewBackwardEnd)
+			.OnBackwardStep_Raw(this, &FDatabaseEditorToolkit::PreviewBackwardStep)
+			.OnBackward_Raw(this, &FDatabaseEditorToolkit::PreviewBackward)
+			.OnPause_Raw(this, &FDatabaseEditorToolkit::PreviewPause)
+			.OnForward_Raw(this, &FDatabaseEditorToolkit::PreviewForward)
+			.OnForwardStep_Raw(this, &FDatabaseEditorToolkit::PreviewForwardStep)
+			.OnForwardEnd_Raw(this, &FDatabaseEditorToolkit::PreviewForwardEnd);
 
 		AssetTreeWidget = SNew(SDatabaseAssetTree, ViewModel.ToSharedRef());
 		AssetTreeWidget->RegisterOnSelectionChanged(
@@ -236,11 +278,6 @@ namespace UE::PoseSearch
 	{
 		const FDatabaseEditorCommands& Commands = FDatabaseEditorCommands::Get();
 
-		ToolkitCommands->MapAction(
-			Commands.ResetPreviewScene,
-			FExecuteAction::CreateSP(this, &FDatabaseEditorToolkit::ResetPreviewScene),
-			EUIActionRepeatMode::RepeatDisabled);
-
 		ToolkitCommands->MapAction(
 			Commands.BuildSearchIndex,
 			FExecuteAction::CreateSP(this, &FDatabaseEditorToolkit::BuildSearchIndex),
@@ -262,13 +299,6 @@ namespace UE::PoseSearch
 
 	void FDatabaseEditorToolkit::FillToolbar(FToolBarBuilder& ToolbarBuilder)
 	{
-		ToolbarBuilder.AddToolBarButton(
-			FDatabaseEditorCommands::Get().ResetPreviewScene,
-			NAME_None,
-			TAttribute<FText>(),
-			TAttribute<FText>(),
-			FSlateIcon(FAppStyle::GetAppStyleSetName(), "Icons.Refresh"));
-
 		ToolbarBuilder.AddToolBarButton(
 			FDatabaseEditorCommands::Get().BuildSearchIndex,
 			NAME_None,
@@ -365,9 +395,9 @@ namespace UE::PoseSearch
 
 		TSharedRef<SDockTab> SpawnedTab = SNew(SDockTab).Label(LOCTEXT("ViewportTab_Title", "Viewport"));
 
-		if (ViewportWidget.IsValid())
+		if (PreviewWidget.IsValid())
 		{
-			SpawnedTab->SetContent(ViewportWidget.ToSharedRef());
+			SpawnedTab->SetContent(PreviewWidget.ToSharedRef());
 		}
 
 		return SpawnedTab;
diff --git a/Engine/Plugins/Experimental/Animation/PoseSearch/Source/Editor/Private/PoseSearchDatabaseViewModel.cpp b/Engine/Plugins/Experimental/Animation/PoseSearch/Source/Editor/Private/PoseSearchDatabaseViewModel.cpp
--- a/Engine/Plugins/Experimental/Animation/PoseSearch/Source/Editor/Private/PoseSearchDatabaseViewModel.cpp
+++ b/Engine/Plugins/Experimental/Animation/PoseSearch/Source/Editor/Private/PoseSearchDatabaseViewModel.cpp
@@ -20,6 +20,8 @@
 
 namespace UE::PoseSearch
 {
+	constexpr float StepDeltaTime = 1.0f / 30.0f;
+
 	FDatabaseViewModel::FDatabaseViewModel()
 		: PoseSearchDatabase(nullptr)
 	{
@@ -53,6 +55,7 @@ namespace UE::PoseSearch
 	void FDatabaseViewModel::ResetPreviewActors()
 	{
 		PlayTime = 0.0f;
+		DeltaTimeMultiplier = 1.0f;
 		RespawnPreviewActors();
 	}
 
@@ -75,6 +78,8 @@ namespace UE::PoseSearch
 					AssociatedAssetIndices.Add(SelectedNode->SourceAssetIdx);
 				}
 			}
+			
+			MaxPreviewPlayLength = 0.0f;
 
 			const FPoseSearchIndex* SearchIndex = PoseSearchDatabase->GetSearchIndex();
 			for (const FPoseSearchIndexAsset& IndexAsset : SearchIndex->Assets)
@@ -93,10 +98,18 @@ namespace UE::PoseSearch
 					FDatabasePreviewActor PreviewActor = SpawnPreviewActor(IndexAsset);
 					if (PreviewActor.IsValid())
 					{
+						const UAnimSequence* Sequence =
+							Cast<UAnimSequence>(PreviewActor.AnimInstance->GetAnimationAsset());
+						if (Sequence)
+						{
+							MaxPreviewPlayLength = FMath::Max(MaxPreviewPlayLength, Sequence->GetPlayLength());
+						}
+
 						PreviewActors.Add(PreviewActor);
 					}
 				}
 			}
+
 			UpdatePreviewActors();
 		}
 
@@ -108,6 +121,43 @@ namespace UE::PoseSearch
 		PoseSearchDatabase->BeginCacheDerivedData();
 	}
 
+	void FDatabaseViewModel::PreviewBackwardEnd()
+	{
+		PlayTime = 0.0f;
+	}
+
+	void FDatabaseViewModel::PreviewBackwardStep()
+	{
+		PlayTime = FMath::Clamp(PlayTime - StepDeltaTime, 0.0f, MaxPreviewPlayLength);
+		DeltaTimeMultiplier = 0.0f;
+	}
+
+	void FDatabaseViewModel::PreviewBackward()
+	{
+		DeltaTimeMultiplier = -1.0f;
+	}
+
+	void FDatabaseViewModel::PreviewPause()
+	{
+		DeltaTimeMultiplier = 0.0f;
+	}
+
+	void FDatabaseViewModel::PreviewForward()
+	{
+		DeltaTimeMultiplier = 1.0f;
+	}
+
+	void FDatabaseViewModel::PreviewForwardStep()
+	{
+		PlayTime = FMath::Clamp(PlayTime + StepDeltaTime, 0.0f, MaxPreviewPlayLength);
+		DeltaTimeMultiplier = 0.0f;
+	}
+
+	void FDatabaseViewModel::PreviewForwardEnd()
+	{
+		PlayTime = MaxPreviewPlayLength;
+	}
+
 	FDatabasePreviewActor FDatabaseViewModel::SpawnPreviewActor(const FPoseSearchIndexAsset& IndexAsset)
 	{
 		USkeleton* Skeleton = PoseSearchDatabase->Schema->Skeleton;
@@ -142,13 +192,13 @@ namespace UE::PoseSearch
 
 		PreviewActor.Mesh->SetSkeletalMesh(Skeleton->GetPreviewMesh(true));
 		PreviewActor.Mesh->EnablePreview(true, DatabaseSequence.Sequence);
-		PreviewActor.AnimInstance->SetAnimationAsset(DatabaseSequence.Sequence);
+		PreviewActor.AnimInstance->SetAnimationAsset(DatabaseSequence.Sequence, false, 0.0f);
 		if (IndexAsset.bMirrored && PoseSearchDatabase->Schema)
 		{
 			PreviewActor.AnimInstance->SetMirrorDataTable(PoseSearchDatabase->Schema->MirrorDataTable);
 		}
 
-		PreviewActor.AnimInstance->PlayAnim(false);
+		PreviewActor.AnimInstance->PlayAnim(false, 0.0f);
 
 		if (!PreviewActor.Actor->GetRootComponent())
 		{
@@ -183,7 +233,8 @@ namespace UE::PoseSearch
 
 	void FDatabaseViewModel::Tick(float DeltaSeconds)
 	{
-		PlayTime += DeltaSeconds;
+		PlayTime += DeltaSeconds * DeltaTimeMultiplier;
+		PlayTime = FMath::Clamp(PlayTime, 0.0f, MaxPreviewPlayLength);
 		UpdatePreviewActors();
 	}
 
@@ -218,6 +269,8 @@ namespace UE::PoseSearch
 			PreviewActor.Actor->Destroy();
 		}
 		PreviewActors.Reset();
+
+		MaxPreviewPlayLength = 0.0f;
 	}
 
 	FTransform FDatabaseViewModel::MirrorRootMotion(
@@ -346,5 +399,39 @@ namespace UE::PoseSearch
 	{
 		SelectedNodes = InSelectedNodes;
 		ResetPreviewActors();
+		ProcessSelectedActor(nullptr);
+	}
+
+	void FDatabaseViewModel::ProcessSelectedActor(AActor* Actor)
+	{
+		SelectedActorIndexAsset = nullptr;
+		
+		const FDatabasePreviewActor* SelectedPreviewActor = PreviewActors.FindByPredicate(
+			[Actor](const FDatabasePreviewActor& PreviewActor)
+		{
+			return PreviewActor.Actor == Actor;
+		});
+
+		if (SelectedPreviewActor)
+		{
+			SelectedActorIndexAsset = SelectedPreviewActor->IndexAsset;
+		}
+	}
+
+	float FDatabaseViewModel::GetMaxPreviewPlayLength() const
+	{
+		return MaxPreviewPlayLength;
+	}
+
+	float FDatabaseViewModel::GetPlayTime() const
+	{
+		const float ClampedPlayTime = FMath::Clamp(PlayTime, 0.0f, MaxPreviewPlayLength);
+		return ClampedPlayTime;
+	}
+
+	void FDatabaseViewModel::SetPlayTime(float NewPlayTime, bool bInTickPlayTime)
+	{
+		PlayTime = NewPlayTime;
+		DeltaTimeMultiplier = bInTickPlayTime ? DeltaTimeMultiplier : 0.0f;
 	}
 }
diff --git a/Engine/Plugins/Experimental/Animation/PoseSearch/Source/Editor/Private/SPoseSearchDatabaseAssetList.cpp b/Engine/Plugins/Experimental/Animation/PoseSearch/Source/Editor/Private/SPoseSearchDatabaseAssetList.cpp
--- a/Engine/Plugins/Experimental/Animation/PoseSearch/Source/Editor/Private/SPoseSearchDatabaseAssetList.cpp
+++ b/Engine/Plugins/Experimental/Animation/PoseSearch/Source/Editor/Private/SPoseSearchDatabaseAssetList.cpp
@@ -244,6 +244,16 @@ namespace UE::PoseSearch
 			[
 				SNew(STextBlock)
 				.Text(this, &SDatabaseAssetListItem::GetName)
+			]
+			+ SHorizontalBox::Slot()
+			.MaxWidth(18)
+			.AutoWidth()
+			.HAlign(HAlign_Right)
+			.VAlign(VAlign_Center)
+			[
+				SNew(SImage)
+				.Image(FAppStyle::Get().GetBrush("Icons.EyeDropper"))
+				.Visibility_Raw(this, &SDatabaseAssetListItem::GetSelectedActorIconVisbility)
 			];
 		}
 
@@ -302,6 +312,22 @@ namespace UE::PoseSearch
 		}
 	}
 
+	EVisibility SDatabaseAssetListItem::GetSelectedActorIconVisbility() const
+	{
+		TSharedPtr<FDatabaseViewModel> ViewModelPtr = EditorViewModel.Pin();
+		TSharedPtr<FDatabaseAssetTreeNode> TreeNodePtr = WeakAssetTreeNode.Pin();
+		if (const FPoseSearchIndexAsset* SelectedIndexAsset = ViewModelPtr->GetSelectedActorIndexAsset())
+		{
+			if (TreeNodePtr->SourceAssetType == ESearchIndexAssetType::Sequence &&
+				TreeNodePtr->SourceAssetIdx == SelectedIndexAsset->SourceAssetIdx)
+			{
+				return EVisibility::Visible;
+			}
+		}
+
+		return EVisibility::Hidden;
+	}
+
 	SDatabaseAssetTree::~SDatabaseAssetTree()
 	{
 	}
diff --git a/Engine/Plugins/Experimental/Animation/PoseSearch/Source/Editor/Private/SPoseSearchDatabaseViewport.cpp b/Engine/Plugins/Experimental/Animation/PoseSearch/Source/Editor/Private/SPoseSearchDatabaseViewport.cpp
--- a/Engine/Plugins/Experimental/Animation/PoseSearch/Source/Editor/Private/SPoseSearchDatabaseViewport.cpp
+++ b/Engine/Plugins/Experimental/Animation/PoseSearch/Source/Editor/Private/SPoseSearchDatabaseViewport.cpp
@@ -7,11 +7,14 @@
 #include "SPoseSearchDatabaseViewportToolbar.h"
 #include "PoseSearchDatabaseEditorCommands.h"
 
+#include "SSimpleTimeSlider.h"
+#include "Widgets/Input/SButton.h"
+
 namespace UE::PoseSearch
 {
 	void SDatabaseViewport::Construct(
 		const FArguments& InArgs,
-		const FDatabaseViewportRequiredArgs& InRequiredArgs)
+		const FDatabasePreviewRequiredArgs& InRequiredArgs)
 	{
 		PreviewScenePtr = InRequiredArgs.PreviewScene;
 		AssetEditorToolkitPtr = InRequiredArgs.AssetEditorToolkit;
@@ -127,5 +130,91 @@ namespace UE::PoseSearch
 	void SDatabaseViewport::OnFloatingButtonClicked()
 	{
 	}
+
+	void SDatabasePreview::Construct(const FArguments& InArgs, const FDatabasePreviewRequiredArgs& InRequiredArgs)
+	{
+		SliderScrubTimeAttribute = InArgs._SliderScrubTime;
+		SliderViewRange = InArgs._SliderViewRange;
+		OnSliderScrubPositionChanged = InArgs._OnSliderScrubPositionChanged;
+
+		OnBackwardEnd = InArgs._OnBackwardEnd;
+		OnBackwardStep = InArgs._OnBackwardStep;
+		OnBackward = InArgs._OnBackward;
+		OnPause = InArgs._OnPause;
+		OnForward = InArgs._OnForward;
+		OnForwardStep = InArgs._OnForwardStep;
+		OnForwardEnd = InArgs._OnForwardEnd;
+
+		FSlimHorizontalToolBarBuilder ToolBarBuilder(
+			TSharedPtr<const FUICommandList>(), 
+			FMultiBoxCustomization::None, 
+			nullptr, true);
+
+		auto AddToolBarButton = [&ToolBarBuilder](FName ButtonStyleName, FOnButtonClickedEvent& OnClicked)
+		{
+			ToolBarBuilder.AddToolBarWidget(
+				SNew(SButton)
+				.ButtonStyle(FAppStyle::Get(), ButtonStyleName)
+				.OnClicked_Lambda([&OnClicked]() 
+				{
+					if (OnClicked.IsBound())
+					{
+						OnClicked.Execute();
+						return FReply::Handled();
+					}
+					return FReply::Unhandled();
+				}));
+		};
+
+		//ToolBarBuilder.SetStyle(&FAppStyle::Get(), "PaletteToolBar");
+		ToolBarBuilder.BeginSection("Preview");
+		{
+			AddToolBarButton("Animation.Backward_End", OnBackwardEnd);
+			AddToolBarButton("Animation.Backward_Step", OnBackwardStep);
+			AddToolBarButton("Animation.Backward", OnBackward);
+			AddToolBarButton("Animation.Pause", OnPause);
+			AddToolBarButton("Animation.Forward", OnForward);
+			AddToolBarButton("Animation.Forward_Step", OnForwardStep);
+			AddToolBarButton("Animation.Forward_End", OnForwardEnd);
+		}
+
+		ChildSlot
+		[
+			SNew(SVerticalBox)
+			+SVerticalBox::Slot()
+			.FillHeight(1.0f)
+			[
+				SNew(SDatabaseViewport, InRequiredArgs)
+			]
+			+SVerticalBox::Slot()
+			.AutoHeight()
+			[
+				SNew(SHorizontalBox)
+				+SHorizontalBox::Slot()
+				.AutoWidth()
+				[
+					ToolBarBuilder.MakeWidget()
+				]
+				+SHorizontalBox::Slot()
+				.FillWidth(1.0f)
+				[
+					SNew(SSimpleTimeSlider)
+					.ClampRangeHighlightSize(0.15f)
+					.ClampRangeHighlightColor(FLinearColor::Red.CopyWithNewOpacity(0.5f))
+					.ScrubPosition_Lambda([this]() { return SliderScrubTimeAttribute.Get(); })
+					.ViewRange_Lambda([this]() { return SliderViewRange.Get(); })
+					.ClampRange_Lambda([this]() { return SliderViewRange.Get(); })
+					.OnScrubPositionChanged_Lambda(
+						[this](double NewScrubTime, bool bIsScrubbing)
+					{
+						if (bIsScrubbing)
+						{
+							OnSliderScrubPositionChanged.ExecuteIfBound(NewScrubTime, bIsScrubbing);
+						}
+					})
+				]
+			]
+		];
+	}
 }
 
diff --git a/Engine/Plugins/Experimental/Animation/PoseSearch/Source/Editor/Public/PoseSearchDatabaseEdMode.h b/Engine/Plugins/Experimental/Animation/PoseSearch/Source/Editor/Public/PoseSearchDatabaseEdMode.h
--- a/Engine/Plugins/Experimental/Animation/PoseSearch/Source/Editor/Public/PoseSearchDatabaseEdMode.h
+++ b/Engine/Plugins/Experimental/Animation/PoseSearch/Source/Editor/Public/PoseSearchDatabaseEdMode.h
@@ -47,5 +47,7 @@ namespace UE::PoseSearch
 	private:
 
 		FDatabaseViewModel* ViewModel = nullptr;
+
+		TWeakObjectPtr<AActor> SelectedActor = nullptr;
 	};
 }
diff --git a/Engine/Plugins/Experimental/Animation/PoseSearch/Source/Editor/Public/PoseSearchDatabaseEditorCommands.h b/Engine/Plugins/Experimental/Animation/PoseSearch/Source/Editor/Public/PoseSearchDatabaseEditorCommands.h
--- a/Engine/Plugins/Experimental/Animation/PoseSearch/Source/Editor/Public/PoseSearchDatabaseEditorCommands.h
+++ b/Engine/Plugins/Experimental/Animation/PoseSearch/Source/Editor/Public/PoseSearchDatabaseEditorCommands.h
@@ -24,18 +24,13 @@ namespace UE::PoseSearch
 		virtual void RegisterCommands() override;
 
 	public:
-		TSharedPtr<FUICommandInfo> ResetPreviewScene;
-
 		TSharedPtr<FUICommandInfo> BuildSearchIndex;
 
 		TSharedPtr<FUICommandInfo> ShowPoseFeaturesNone;
-
 		TSharedPtr<FUICommandInfo> ShowPoseFeaturesAll;
 
 		TSharedPtr<FUICommandInfo> ShowAnimationNone;
-
 		TSharedPtr<FUICommandInfo> ShowAnimationOriginalOnly;
-
 		TSharedPtr<FUICommandInfo> ShowAnimationOriginalAndMirrored;
 
 	};
diff --git a/Engine/Plugins/Experimental/Animation/PoseSearch/Source/Editor/Public/PoseSearchDatabaseEditorToolkit.h b/Engine/Plugins/Experimental/Animation/PoseSearch/Source/Editor/Public/PoseSearchDatabaseEditorToolkit.h
--- a/Engine/Plugins/Experimental/Animation/PoseSearch/Source/Editor/Public/PoseSearchDatabaseEditorToolkit.h
+++ b/Engine/Plugins/Experimental/Animation/PoseSearch/Source/Editor/Public/PoseSearchDatabaseEditorToolkit.h
@@ -13,7 +13,7 @@ class UPoseSearchDatabaseSelectionReflection;
 
 namespace UE::PoseSearch
 {
-	class SDatabaseViewport;
+	class SDatabasePreview;
 	class FDatabasePreviewScene;
 	class FDatabaseViewModel;
 
@@ -42,10 +42,16 @@ namespace UE::PoseSearch
 		FDatabaseViewModel* GetViewModel() const { return ViewModel.Get(); }
 		TSharedPtr<FDatabaseViewModel> GetViewModelSharedPtr() const { return ViewModel; }
 
-		void StopPreviewScene();
-		void ResetPreviewScene();
 		void BuildSearchIndex();
 
+		void PreviewBackwardEnd();
+		void PreviewBackwardStep();
+		void PreviewBackward();
+		void PreviewPause();
+		void PreviewForward();
+		void PreviewForwardStep();
+		void PreviewForwardEnd();
+
 	private:
 
 		TSharedRef<SDockTab> SpawnTab_Viewport(const FSpawnTabArgs& Args);
@@ -63,7 +69,7 @@ namespace UE::PoseSearch
 			const TArrayView<TSharedPtr<FDatabaseAssetTreeNode>>& SelectedItems,
 			ESelectInfo::Type SelectionType);
 
-		TSharedPtr<SDatabaseViewport> ViewportWidget;
+		TSharedPtr<SDatabasePreview> PreviewWidget;
 
 		TSharedPtr<SDatabaseAssetTree> AssetTreeWidget;
 
diff --git a/Engine/Plugins/Experimental/Animation/PoseSearch/Source/Editor/Public/PoseSearchDatabaseViewModel.h b/Engine/Plugins/Experimental/Animation/PoseSearch/Source/Editor/Public/PoseSearchDatabaseViewModel.h
--- a/Engine/Plugins/Experimental/Animation/PoseSearch/Source/Editor/Public/PoseSearchDatabaseViewModel.h
+++ b/Engine/Plugins/Experimental/Animation/PoseSearch/Source/Editor/Public/PoseSearchDatabaseViewModel.h
@@ -71,6 +71,14 @@ namespace UE::PoseSearch
 		void RespawnPreviewActors();
 		void BuildSearchIndex();
 
+		void PreviewBackwardEnd();
+		void PreviewBackwardStep();
+		void PreviewBackward();
+		void PreviewPause();
+		void PreviewForward();
+		void PreviewForwardStep();
+		void PreviewForwardEnd();
+
 		UPoseSearchDatabase* GetPoseSearchDatabase() const { return PoseSearchDatabase; }
 		void OnPreviewActorClassChanged();
 
@@ -96,9 +104,17 @@ namespace UE::PoseSearch
 		void DeleteGroup(int32 GroupIdx);
 
 		void SetSelectedNodes(const TArrayView<TSharedPtr<FDatabaseAssetTreeNode>>& InSelectedNodes);
+		void ProcessSelectedActor(AActor* Actor);
+		
+		const FPoseSearchIndexAsset* GetSelectedActorIndexAsset() const { return SelectedActorIndexAsset; }
+
+		float GetMaxPreviewPlayLength() const;
+		float GetPlayTime() const;
+		void SetPlayTime(float NewPlayTime, bool bInTickPlayTime);
 
 	private:
 		float PlayTime = 0.0f;
+		float DeltaTimeMultiplier = 1.0f;
 
 		/** Scene asset being viewed and edited by this view model. */
 		TObjectPtr<UPoseSearchDatabase> PoseSearchDatabase;
@@ -109,6 +125,9 @@ namespace UE::PoseSearch
 		/** Actors to be displayed in the preview viewport */
 		TArray<FDatabasePreviewActor> PreviewActors;
 
+		/** From zero to the play length of the longest preview */
+		float MaxPreviewPlayLength = 0.0f;
+
 		/** What features to show in the viewport */
 		EFeaturesDrawMode PoseFeaturesDrawMode = EFeaturesDrawMode::None;
 
@@ -117,6 +136,8 @@ namespace UE::PoseSearch
 
 		TArray<TSharedPtr<FDatabaseAssetTreeNode>> SelectedNodes;
 
+		const FPoseSearchIndexAsset* SelectedActorIndexAsset = nullptr;
+
 		UWorld* GetWorld() const;
 
 		UObject* GetPlaybackContext() const;
diff --git a/Engine/Plugins/Experimental/Animation/PoseSearch/Source/Editor/Public/SPoseSearchDatabaseAssetList.h b/Engine/Plugins/Experimental/Animation/PoseSearch/Source/Editor/Public/SPoseSearchDatabaseAssetList.h
--- a/Engine/Plugins/Experimental/Animation/PoseSearch/Source/Editor/Public/SPoseSearchDatabaseAssetList.h
+++ b/Engine/Plugins/Experimental/Animation/PoseSearch/Source/Editor/Public/SPoseSearchDatabaseAssetList.h
@@ -59,6 +59,8 @@ namespace UE::PoseSearch
 		void OnAddSequence();
 		void OnAddBlendSpace();
 
+		EVisibility GetSelectedActorIconVisbility() const;
+
 		TWeakPtr<FDatabaseAssetTreeNode> WeakAssetTreeNode;
 		TWeakPtr<FDatabaseViewModel> EditorViewModel;
 		TWeakPtr<SDatabaseAssetTree> SkeletonView;
diff --git a/Engine/Plugins/Experimental/Animation/PoseSearch/Source/Editor/Public/SPoseSearchDatabaseViewport.h b/Engine/Plugins/Experimental/Animation/PoseSearch/Source/Editor/Public/SPoseSearchDatabaseViewport.h
--- a/Engine/Plugins/Experimental/Animation/PoseSearch/Source/Editor/Public/SPoseSearchDatabaseViewport.h
+++ b/Engine/Plugins/Experimental/Animation/PoseSearch/Source/Editor/Public/SPoseSearchDatabaseViewport.h
@@ -14,9 +14,9 @@ namespace UE::PoseSearch
 	class FDatabaseEditorToolkit;
 	class SPoseSearchDatabaseViewportToolBar;
 
-	struct FDatabaseViewportRequiredArgs
+	struct FDatabasePreviewRequiredArgs
 	{
-		FDatabaseViewportRequiredArgs(
+		FDatabasePreviewRequiredArgs(
 			const TSharedRef<FDatabaseEditorToolkit>& InAssetEditorToolkit,
 			const TSharedRef<FDatabasePreviewScene>& InPreviewScene)
 			: AssetEditorToolkit(InAssetEditorToolkit)
@@ -32,11 +32,10 @@ namespace UE::PoseSearch
 	class SDatabaseViewport : public SEditorViewport, public ICommonEditorViewportToolbarInfoProvider
 	{
 	public:
-
 		SLATE_BEGIN_ARGS(SDatabaseViewport) {}
 		SLATE_END_ARGS();
 
-		void Construct(const FArguments& InArgs, const FDatabaseViewportRequiredArgs& InRequiredArgs);
+		void Construct(const FArguments& InArgs, const FDatabasePreviewRequiredArgs& InRequiredArgs);
 		virtual ~SDatabaseViewport() {}
 
 		// ~ICommonEditorViewportToolbarInfoProvider interface
@@ -65,4 +64,38 @@ namespace UE::PoseSearch
 		/** Asset editor toolkit we are embedded in */
 		TWeakPtr<FDatabaseEditorToolkit> AssetEditorToolkitPtr;
 	};
+
+	class SDatabasePreview : public SCompoundWidget
+	{
+	public:
+		DECLARE_DELEGATE_TwoParams(FOnScrubPositionChanged, double, bool)
+		DECLARE_DELEGATE(FOnButtonClickedEvent)
+
+		SLATE_BEGIN_ARGS(SDatabasePreview) {}
+			SLATE_ATTRIBUTE(double, SliderScrubTime);
+			SLATE_ATTRIBUTE(TRange<double>, SliderViewRange);
+			SLATE_EVENT(FOnScrubPositionChanged, OnSliderScrubPositionChanged);
+			SLATE_EVENT(FOnButtonClickedEvent, OnBackwardEnd);
+			SLATE_EVENT(FOnButtonClickedEvent, OnBackwardStep);
+			SLATE_EVENT(FOnButtonClickedEvent, OnBackward);
+			SLATE_EVENT(FOnButtonClickedEvent, OnPause);
+			SLATE_EVENT(FOnButtonClickedEvent, OnForward);
+			SLATE_EVENT(FOnButtonClickedEvent, OnForwardStep);
+			SLATE_EVENT(FOnButtonClickedEvent, OnForwardEnd);
+		SLATE_END_ARGS();
+
+		void Construct(const FArguments& InArgs, const FDatabasePreviewRequiredArgs& InRequiredArgs);
+
+	protected:
+		TAttribute<double> SliderScrubTimeAttribute;
+		TAttribute<TRange<double>> SliderViewRange = TRange<double>(0.0, 1.0);
+		FOnScrubPositionChanged OnSliderScrubPositionChanged;
+		FOnButtonClickedEvent OnBackwardEnd;
+		FOnButtonClickedEvent OnBackwardStep;
+		FOnButtonClickedEvent OnBackward;
+		FOnButtonClickedEvent OnPause;
+		FOnButtonClickedEvent OnForward;
+		FOnButtonClickedEvent OnForwardStep;
+		FOnButtonClickedEvent OnForwardEnd;
+	};
 }
