undo changes to PoseSearchDebugger for Rewind Debugger Timelines

diff --git a/Engine/Plugins/Experimental/Animation/PoseSearch/Source/Editor/Private/PoseSearchDebugger.cpp b/Engine/Plugins/Experimental/Animation/PoseSearch/Source/Editor/Private/PoseSearchDebugger.cpp
--- a/Engine/Plugins/Experimental/Animation/PoseSearch/Source/Editor/Private/PoseSearchDebugger.cpp
+++ b/Engine/Plugins/Experimental/Animation/PoseSearch/Source/Editor/Private/PoseSearchDebugger.cpp
@@ -2827,29 +2827,6 @@ TSharedPtr<IRewindDebuggerView> FDebuggerViewCreator::CreateDebugView(uint64 Obj
 	return FDebugger::Get()->GenerateInstance(ObjectId);
 }
 
-bool FDebuggerViewCreator::HasDebugInfo(uint64 AnimInstanceId, const TraceServices::IAnalysisSession& Session) const
-{
-	// Get provider and validate
-	TraceServices::FAnalysisSessionReadScope SessionReadScope(Session);
-
-	const FTraceProvider* PoseSearchProvider = Session.ReadProvider<FTraceProvider>(FTraceProvider::ProviderName);
-	const IAnimationProvider* AnimationProvider = Session.ReadProvider<IAnimationProvider>("AnimationProvider");
-	const IGameplayProvider* GameplayProvider = Session.ReadProvider<IGameplayProvider>("GameplayProvider");
-	if (!(PoseSearchProvider && AnimationProvider && GameplayProvider))
-	{
-		return false;
-	}
-	
-	bool bHasData = false;
-	
-	PoseSearchProvider->EnumerateMotionMatchingStateTimelines(AnimInstanceId, [&bHasData](const FTraceProvider::FMotionMatchingStateTimeline& InTimeline)
-	{
-		bHasData = true;
-	});
-	
-	return bHasData;
-}
-
 FName FDebuggerViewCreator::GetName() const
 {
 	static const FName Name("PoseSearchDebugger");
diff --git a/Engine/Plugins/Experimental/Animation/PoseSearch/Source/Editor/Private/PoseSearchDebugger.h b/Engine/Plugins/Experimental/Animation/PoseSearch/Source/Editor/Private/PoseSearchDebugger.h
--- a/Engine/Plugins/Experimental/Animation/PoseSearch/Source/Editor/Private/PoseSearchDebugger.h
+++ b/Engine/Plugins/Experimental/Animation/PoseSearch/Source/Editor/Private/PoseSearchDebugger.h
@@ -758,7 +758,6 @@ public:
 	
 	/** Creates the PoseSearch Slate view for the provided AnimInstance */
 	virtual TSharedPtr<IRewindDebuggerView> CreateDebugView(uint64 ObjectId, double CurrentTime, const TraceServices::IAnalysisSession& InAnalysisSession) const override;
-	virtual bool HasDebugInfo(uint64 ObjectId, const TraceServices::IAnalysisSession& InAnalysisSession) const override;
 };
 
 
