Pose search: improvements to DDC and runtime logic
- fixed ddc issue with cooking and running packaged builds
- fixed crash when the schema changes while the database is being searched

diff --git a/Engine/Plugins/Experimental/Animation/PoseSearch/Source/Runtime/Private/PoseSearch.cpp b/Engine/Plugins/Experimental/Animation/PoseSearch/Source/Runtime/Private/PoseSearch.cpp
--- a/Engine/Plugins/Experimental/Animation/PoseSearch/Source/Runtime/Private/PoseSearch.cpp
+++ b/Engine/Plugins/Experimental/Animation/PoseSearch/Source/Runtime/Private/PoseSearch.cpp
@@ -895,6 +895,9 @@ void FPoseSearchWeightsContext::Update(
 	if (Database != ActiveDatabase)
 	{
 		Database = ActiveDatabase;
+#if WITH_EDITOR
+		SearchIndexHash = Database->GetSearchIndexHash();
+#endif
 		bRecomputeWeights = true;
 	}
 
@@ -904,6 +907,14 @@ void FPoseSearchWeightsContext::Update(
 		bRecomputeWeights = true;
 	}
 
+#if WITH_EDITOR
+	if (Database.IsValid() && Database->GetSearchIndexHash() != SearchIndexHash)
+	{
+		SearchIndexHash = Database->GetSearchIndexHash();
+		bRecomputeWeights = true;
+	}
+#endif
+
 	if (bRecomputeWeights)
 	{
 		ComputedDefaultGroupWeights.Init(Database->DefaultWeights, Database->Schema, ActiveWeights);
@@ -1018,13 +1029,18 @@ float FPoseSearchIndex::GetAssetTime(int32 PoseIdx, const FPoseSearchIndexAsset*
 
 bool FPoseSearchIndex::IsValid() const
 {
-	bool bSchemaValid = Schema && Schema->IsValid();
-	bool bAssetsValid = Assets.Num() > 0;
-	bool bSearchIndexValid = bSchemaValid && bAssetsValid && (NumPoses * Schema->Layout.NumFloats == Values.Num());
+	const bool bSchemaValid = Schema && Schema->IsValid();
+	const bool bSearchIndexValid = bSchemaValid && (NumPoses * Schema->Layout.NumFloats == Values.Num());
 
 	return bSearchIndexValid;
 }
 
+bool FPoseSearchIndex::IsEmpty() const
+{
+	const bool bEmpty = Assets.Num() == 0 || NumPoses == 0;
+	return bEmpty;
+}
+
 TArrayView<const float> FPoseSearchIndex::GetPoseValues(int32 PoseIdx) const
 {
 	check(PoseIdx < NumPoses);
@@ -1126,7 +1142,7 @@ bool UPoseSearchSequenceMetaData::IsValidForIndexing() const
 
 bool UPoseSearchSequenceMetaData::IsValidForSearch() const
 {
-	return IsValidForIndexing() && SearchIndex.IsValid();
+	return IsValidForIndexing() && SearchIndex.IsValid() && !SearchIndex.IsEmpty();
 }
 
 
@@ -1291,7 +1307,7 @@ bool UPoseSearchDatabase::IsValidForIndexing() const
 bool UPoseSearchDatabase::IsValidForSearch() const
 {
 	const FPoseSearchIndex* SearchIndex = GetSearchIndex();
-	bool bIsValid = IsValidForIndexing() && SearchIndex && SearchIndex->IsValid();
+	bool bIsValid = IsValidForIndexing() && SearchIndex && SearchIndex->IsValid() && !SearchIndex->IsEmpty();
 
 #if WITH_EDITOR
 	const bool bIsCurrentDerivedData = 
@@ -1780,8 +1796,16 @@ void UPoseSearchDatabase::Serialize(FArchive& Ar)
 
 	if (Ar.IsFilterEditorOnly())
 	{
+		if (Ar.IsLoading())
+		{
+			if (!PrivateDerivedData)
+			{
+				PrivateDerivedData = new FPoseSearchDatabaseDerivedData();
+				PrivateDerivedData->SearchIndex.Schema = Schema;
+			}
+		}
 		check(Ar.IsLoading() || (Ar.IsCooking() && IsDerivedDataValid()));
-		const FPoseSearchIndex* SearchIndex = GetSearchIndex();
+		FPoseSearchIndex* SearchIndex = GetSearchIndex();
 		Ar << *SearchIndex;
 	}
 }
@@ -2518,7 +2542,7 @@ bool FDebugDrawParams::CanDraw() const
 		return false;
 	}
 
-	return SearchIndex->IsValid();
+	return SearchIndex->IsValid() && !SearchIndex->IsEmpty();
 }
 
 const FPoseSearchIndex* FDebugDrawParams::GetSearchIndex() const
@@ -4522,7 +4546,7 @@ static void PreprocessSearchIndexNone(FPoseSearchIndex* SearchIndex)
 
 	using namespace Eigen;
 
-	check(SearchIndex->IsValid());
+	check(SearchIndex->IsValid() && !SearchIndex->IsEmpty());
 
 	FPoseSearchIndexPreprocessInfo& Info = SearchIndex->PreprocessInfo;
 	Info.Reset();
@@ -4633,7 +4657,7 @@ static void PreprocessSearchIndexNormalize(FPoseSearchIndex* SearchIndex)
 
 	using namespace Eigen;
 
-	check(SearchIndex->IsValid());
+	check(SearchIndex->IsValid() && !SearchIndex->IsEmpty());
 
 	FPoseSearchIndexPreprocessInfo& Info = SearchIndex->PreprocessInfo;
 	Info.Reset();
@@ -4791,7 +4815,7 @@ static void PreprocessSearchIndexSphere(FPoseSearchIndex* SearchIndex)
 
 	using namespace Eigen;
 
-	check(SearchIndex->IsValid());
+	check(SearchIndex->IsValid() && !SearchIndex->IsEmpty());
 
 	FPoseSearchIndexPreprocessInfo& Info = SearchIndex->PreprocessInfo;
 	Info.Reset();
@@ -5208,7 +5232,7 @@ FSearchResult Search(FSearchContext& SearchContext)
 		return Result;
 	}
 
-	if (!ensure(SearchIndex->IsValid()))
+	if (!ensure(SearchIndex->IsValid() && !SearchIndex->IsEmpty()))
 	{
 		return Result;
 	}
diff --git a/Engine/Plugins/Experimental/Animation/PoseSearch/Source/Runtime/Private/PoseSearchDerivedData.cpp b/Engine/Plugins/Experimental/Animation/PoseSearch/Source/Runtime/Private/PoseSearchDerivedData.cpp
--- a/Engine/Plugins/Experimental/Animation/PoseSearch/Source/Runtime/Private/PoseSearchDerivedData.cpp
+++ b/Engine/Plugins/Experimental/Animation/PoseSearch/Source/Runtime/Private/PoseSearchDerivedData.cpp
@@ -148,13 +148,12 @@ namespace UE::PoseSearch
 				DerivedData.SearchIndex.Schema = Database.Schema;
 				const bool bIndexReady = BuildIndex(&Database, DerivedData.SearchIndex);
 
-				WriteIndexToCache();
-				DerivedData.DerivedDataKey = NewKey;
+				WriteIndexToCache(NewKey);
 			}
 		});
 	}
 
-	void FPoseSearchDatabaseAsyncCacheTask::WriteIndexToCache()
+	void FPoseSearchDatabaseAsyncCacheTask::WriteIndexToCache(const UE::DerivedData::FCacheKey& NewKey)
 	{
 		using namespace UE::DerivedData;
 
@@ -163,11 +162,12 @@ namespace UE::PoseSearch
 		Writer << DerivedData.SearchIndex;
 		FSharedBuffer RawData = MakeSharedBufferFromArray(MoveTemp(RawBytes));
 
-		FCacheRecordBuilder Builder(DerivedData.DerivedDataKey);
+		FCacheRecordBuilder Builder(NewKey);
 		Builder.AddValue(Id, RawData);
 
 		Owner.KeepAlive();
 		GetCache().Put({ { { Database.GetPathName() }, Builder.Build() } }, Owner);
+		DerivedData.DerivedDataKey = NewKey;
 	}
 
 	void FPoseSearchDatabaseAsyncCacheTask::BuildIndexFromCacheRecord(UE::DerivedData::FCacheRecord&& CacheRecord)
@@ -362,13 +362,6 @@ namespace UE::PoseSearch
 		return Ar;
 	}
 
-	FArchive& operator<<(FArchive& Ar, const FPoseSearchIndexPreprocessInfo& Info)
-	{
-		check(Ar.IsSaving());
-		return (Ar << const_cast<FPoseSearchIndexPreprocessInfo&>(Info));
-	}
-
-
 	FArchive& operator<<(FArchive& Ar, FPoseSearchIndex& Index)
 	{
 		int32 NumValues = 0;
@@ -411,9 +404,4 @@ namespace UE::PoseSearch
 		return Ar;
 	}
 
-	FArchive& operator<<(FArchive& Ar, const FPoseSearchIndex& Index)
-	{
-		check(Ar.IsSaving());
-		return (Ar << const_cast<FPoseSearchIndex&>(Index));
-	}
 }
\ No newline at end of file
diff --git a/Engine/Plugins/Experimental/Animation/PoseSearch/Source/Runtime/Private/PoseSearchDerivedData.h b/Engine/Plugins/Experimental/Animation/PoseSearch/Source/Runtime/Private/PoseSearchDerivedData.h
--- a/Engine/Plugins/Experimental/Animation/PoseSearch/Source/Runtime/Private/PoseSearchDerivedData.h
+++ b/Engine/Plugins/Experimental/Animation/PoseSearch/Source/Runtime/Private/PoseSearchDerivedData.h
@@ -77,7 +77,7 @@ namespace UE::PoseSearch
 		void BeginCache();
 		void OnGetComplete(UE::DerivedData::FCacheGetResponse&& Response);
 		void BuildAndWrite(const UE::DerivedData::FCacheKey& NewKey);
-		void WriteIndexToCache();
+		void WriteIndexToCache(const UE::DerivedData::FCacheKey& NewKey);
 		void BuildIndexFromCacheRecord(UE::DerivedData::FCacheRecord&& CacheRecord);
 
 		UE::DerivedData::IRequestOwner& GetRequestOwner() { return Owner; }
@@ -103,12 +103,6 @@ namespace UE::PoseSearch
 	// Serialization for FPoseSearchIndexPreprocessInfo.
 	FArchive& operator<<(FArchive& Ar, FPoseSearchIndexPreprocessInfo& Info);
 
-	// Serialization for const FPoseSearchIndexPreprocessInfo.
-	FArchive& operator<<(FArchive& Ar, const FPoseSearchIndexPreprocessInfo& Info);
-
 	// Serialization for FPoseSearchIndex.
 	FArchive& operator<<(FArchive& Ar, FPoseSearchIndex& Index);
-
-	// Serialization for const FPoseSearchIndex.
-	FArchive& operator<<(FArchive& Ar, const FPoseSearchIndex& Index);
 }
\ No newline at end of file
diff --git a/Engine/Plugins/Experimental/Animation/PoseSearch/Source/Runtime/Public/PoseSearch/PoseSearch.h b/Engine/Plugins/Experimental/Animation/PoseSearch/Source/Runtime/Public/PoseSearch/PoseSearch.h
--- a/Engine/Plugins/Experimental/Animation/PoseSearch/Source/Runtime/Public/PoseSearch/PoseSearch.h
+++ b/Engine/Plugins/Experimental/Animation/PoseSearch/Source/Runtime/Public/PoseSearch/PoseSearch.h
@@ -417,6 +417,7 @@ struct POSESEARCH_API FPoseSearchIndex
 	TArray<FPoseSearchIndexAsset> Assets;
 
 	bool IsValid() const;
+	bool IsEmpty() const;
 
 	TArrayView<const float> GetPoseValues(int32 PoseIdx) const;
 
@@ -634,6 +635,11 @@ private:
 
 	UPROPERTY(Transient)
 	TArray<FPoseSearchWeights> ComputedGroupWeights;
+
+#if WITH_EDITOR
+	// used to check if the data has changed, which requires the weights to be recomputed
+	FIoHash SearchIndexHash = FIoHash::Zero;
+#endif
 };
 
 UENUM()
