Filtering for Pose Search entries in Rewind Debugger timeline

diff --git a/Engine/Plugins/Experimental/Animation/PoseSearch/Source/Editor/Private/PoseSearchDebugger.cpp b/Engine/Plugins/Experimental/Animation/PoseSearch/Source/Editor/Private/PoseSearchDebugger.cpp
--- a/Engine/Plugins/Experimental/Animation/PoseSearch/Source/Editor/Private/PoseSearchDebugger.cpp
+++ b/Engine/Plugins/Experimental/Animation/PoseSearch/Source/Editor/Private/PoseSearchDebugger.cpp
@@ -2827,6 +2827,30 @@ TSharedPtr<IRewindDebuggerView> FDebuggerViewCreator::CreateDebugView(uint64 Obj
 	return FDebugger::Get()->GenerateInstance(ObjectId);
 }
 
+bool FDebuggerViewCreator::HasDebugInfo(uint64 AnimInstanceId) const
+{
+	// Get provider and validate
+	const TraceServices::IAnalysisSession* Session = IRewindDebugger::Instance()->GetAnalysisSession();
+	TraceServices::FAnalysisSessionReadScope SessionReadScope(*Session);
+
+	const FTraceProvider* PoseSearchProvider = Session->ReadProvider<FTraceProvider>(FTraceProvider::ProviderName);
+	const IAnimationProvider* AnimationProvider = Session->ReadProvider<IAnimationProvider>("AnimationProvider");
+	const IGameplayProvider* GameplayProvider = Session->ReadProvider<IGameplayProvider>("GameplayProvider");
+	if (!(PoseSearchProvider && AnimationProvider && GameplayProvider))
+	{
+		return false;
+	}
+	
+	bool bHasData = false;
+	
+	PoseSearchProvider->EnumerateMotionMatchingStateTimelines(AnimInstanceId, [&bHasData](const FTraceProvider::FMotionMatchingStateTimeline& InTimeline)
+	{
+		bHasData = true;
+	});
+	
+	return bHasData;
+}
+
 FName FDebuggerViewCreator::GetName() const
 {
 	static const FName Name("PoseSearchDebugger");
diff --git a/Engine/Plugins/Experimental/Animation/PoseSearch/Source/Editor/Private/PoseSearchDebugger.h b/Engine/Plugins/Experimental/Animation/PoseSearch/Source/Editor/Private/PoseSearchDebugger.h
--- a/Engine/Plugins/Experimental/Animation/PoseSearch/Source/Editor/Private/PoseSearchDebugger.h
+++ b/Engine/Plugins/Experimental/Animation/PoseSearch/Source/Editor/Private/PoseSearchDebugger.h
@@ -758,6 +758,7 @@ public:
 	
 	/** Creates the PoseSearch Slate view for the provided AnimInstance */
 	virtual TSharedPtr<IRewindDebuggerView> CreateDebugView(uint64 ObjectId, double CurrentTime, const TraceServices::IAnalysisSession& InAnalysisSession) const override;
+	virtual bool HasDebugInfo(uint64 ObjectId) const override;
 };
 
 
