Expose Pose Search assets under "Animation\Motion Matching" in the Content Browser's context menu. We previously had to go through "Miscellaneous\Data Asset" to create them.

diff --git a/Engine/Plugins/Experimental/Animation/PoseSearch/Source/Editor/Private/PoseSearchDatabaseFactory.cpp b/Engine/Plugins/Experimental/Animation/PoseSearch/Source/Editor/Private/PoseSearchDatabaseFactory.cpp
new file mode 100644
--- /dev/null
+++ b/Engine/Plugins/Experimental/Animation/PoseSearch/Source/Editor/Private/PoseSearchDatabaseFactory.cpp
@@ -0,0 +1,21 @@
+// Copyright Epic Games, Inc. All Rights Reserved.
+
+#include "PoseSearchDatabaseFactory.h"
+#include "PoseSearch/PoseSearch.h"
+
+#define LOCTEXT_NAMESPACE "PoseSearchEditor"
+
+UPoseSearchDatabaseFactory::UPoseSearchDatabaseFactory(const FObjectInitializer& ObjectInitializer)
+	: Super(ObjectInitializer)
+{
+	bCreateNew = true;
+	bEditAfterNew = true;
+	SupportedClass = UPoseSearchDatabase::StaticClass();
+}
+
+UObject* UPoseSearchDatabaseFactory::FactoryCreateNew(UClass* Class, UObject* InParent, FName Name, EObjectFlags Flags, UObject* Context, FFeedbackContext* Warn)
+{
+	return NewObject<UPoseSearchDatabase>(InParent, Class, Name, Flags);
+}
+
+#undef LOCTEXT_NAMESPACE
\ No newline at end of file
diff --git a/Engine/Plugins/Experimental/Animation/PoseSearch/Source/Editor/Private/PoseSearchEditor.cpp b/Engine/Plugins/Experimental/Animation/PoseSearch/Source/Editor/Private/PoseSearchEditor.cpp
--- a/Engine/Plugins/Experimental/Animation/PoseSearch/Source/Editor/Private/PoseSearchEditor.cpp
+++ b/Engine/Plugins/Experimental/Animation/PoseSearch/Source/Editor/Private/PoseSearchEditor.cpp
@@ -94,6 +94,7 @@ private:
 	TSet<FName> RegisteredPropertyTypes;
 
 	TSharedPtr<IAssetTypeActions> PoseSearchDatabaseActions;
+	TSharedPtr<IAssetTypeActions> PoseSearchSchemaActions;
 };
 
 void FEditorModule::StartupModule()
@@ -126,6 +127,10 @@ void FEditorModule::StartupModule()
 		IAssetTools& AssetTools = FModuleManager::LoadModuleChecked<FAssetToolsModule>("AssetTools").Get();
 		PoseSearchDatabaseActions = MakeShared<FDatabaseTypeActions>();
 		AssetTools.RegisterAssetTypeActions(PoseSearchDatabaseActions.ToSharedRef());
+
+		// Register UPoseSearchSchema Type Actions 
+		PoseSearchSchemaActions = MakeShared<FSchemaTypeActions>();
+		AssetTools.RegisterAssetTypeActions(PoseSearchSchemaActions.ToSharedRef());
 	}
 
 	RegisterPropertyTypeCustomizations();
@@ -148,6 +153,7 @@ void FEditorModule::ShutdownModule()
 	{
 		IAssetTools& AssetTools = FModuleManager::GetModuleChecked<FAssetToolsModule>("AssetTools").Get();
 		AssetTools.UnregisterAssetTypeActions(PoseSearchDatabaseActions.ToSharedRef());
+		AssetTools.UnregisterAssetTypeActions(PoseSearchSchemaActions.ToSharedRef());
 	}
 
 	// Unregister Ed Mode
diff --git a/Engine/Plugins/Experimental/Animation/PoseSearch/Source/Editor/Private/PoseSearchSchemaFactory.cpp b/Engine/Plugins/Experimental/Animation/PoseSearch/Source/Editor/Private/PoseSearchSchemaFactory.cpp
new file mode 100644
--- /dev/null
+++ b/Engine/Plugins/Experimental/Animation/PoseSearch/Source/Editor/Private/PoseSearchSchemaFactory.cpp
@@ -0,0 +1,21 @@
+// Copyright Epic Games, Inc. All Rights Reserved.
+
+#include "PoseSearchSchemaFactory.h"
+#include "PoseSearch/PoseSearch.h"
+
+#define LOCTEXT_NAMESPACE "PoseSearchEditor"
+
+UPoseSearchSchemaFactory::UPoseSearchSchemaFactory(const FObjectInitializer& ObjectInitializer)
+	: Super(ObjectInitializer)
+{
+	bCreateNew = true;
+	bEditAfterNew = true;
+	SupportedClass = UPoseSearchSchema::StaticClass();
+}
+
+UObject* UPoseSearchSchemaFactory::FactoryCreateNew(UClass* Class, UObject* InParent, FName Name, EObjectFlags Flags, UObject* Context, FFeedbackContext* Warn)
+{
+	return NewObject<UPoseSearchSchema>(InParent, Class, Name, Flags);
+}
+
+#undef LOCTEXT_NAMESPACE
\ No newline at end of file
diff --git a/Engine/Plugins/Experimental/Animation/PoseSearch/Source/Editor/Private/PoseSearchTypeActions.cpp b/Engine/Plugins/Experimental/Animation/PoseSearch/Source/Editor/Private/PoseSearchTypeActions.cpp
--- a/Engine/Plugins/Experimental/Animation/PoseSearch/Source/Editor/Private/PoseSearchTypeActions.cpp
+++ b/Engine/Plugins/Experimental/Animation/PoseSearch/Source/Editor/Private/PoseSearchTypeActions.cpp
@@ -8,14 +8,20 @@
 
 namespace UE::PoseSearch
 {
+	static const FText PoseSearchSubMenuName = LOCTEXT("PoseSearchSubMenuName", "Motion Matching");
+	static const FColor PoseSearchAssetColor(29, 96, 125);
+
+	//////////////////////////////////////////////////////////////////////////
+	// FDatabaseTypeActions
+
 	FText FDatabaseTypeActions::GetName() const
 	{
-		return LOCTEXT("PoseSearchDatabaseTypeActionsName", "Pose Search Database");
+		return LOCTEXT("PoseSearchDatabaseTypeActionsName", "Motion Database");
 	}
 
 	FColor FDatabaseTypeActions::GetTypeColor() const
 	{
-		return FColor(129, 196, 115);
+		return PoseSearchAssetColor;
 	}
 
 	UClass* FDatabaseTypeActions::GetSupportedClass() const
@@ -44,6 +50,47 @@ namespace UE::PoseSearch
 	{
 		return EAssetTypeCategories::Animation;
 	}
+
+	const TArray<FText>& FDatabaseTypeActions::GetSubMenus() const
+	{
+		static const TArray<FText> SubMenus
+		{
+			PoseSearchSubMenuName
+		};
+		return SubMenus;
+	}
+
+	//////////////////////////////////////////////////////////////////////////
+	// FSchemaTypeActions
+
+	FText FSchemaTypeActions::GetName() const
+	{
+		return LOCTEXT("PoseSearchSchemaTypeActionsName", "Motion Database Config");
+	}
+
+	FColor FSchemaTypeActions::GetTypeColor() const
+	{
+		return PoseSearchAssetColor;
+	}
+
+	UClass* FSchemaTypeActions::GetSupportedClass() const
+	{
+		return UPoseSearchSchema::StaticClass();
+	}
+
+	uint32 FSchemaTypeActions::GetCategories()
+	{
+		return EAssetTypeCategories::Animation;
+	}
+
+	const TArray<FText>& FSchemaTypeActions::GetSubMenus() const
+	{
+		static const TArray<FText> SubMenus
+		{
+			PoseSearchSubMenuName
+		};
+		return SubMenus;
+	}
 }
 
 #undef LOCTEXT_NAMESPACE
diff --git a/Engine/Plugins/Experimental/Animation/PoseSearch/Source/Editor/Public/PoseSearchDatabaseFactory.h b/Engine/Plugins/Experimental/Animation/PoseSearch/Source/Editor/Public/PoseSearchDatabaseFactory.h
new file mode 100644
--- /dev/null
+++ b/Engine/Plugins/Experimental/Animation/PoseSearch/Source/Editor/Public/PoseSearchDatabaseFactory.h
@@ -0,0 +1,18 @@
+// Copyright Epic Games, Inc. All Rights Reserved.
+
+#pragma once
+
+#include "CoreMinimal.h"
+#include "UObject/ObjectMacros.h"
+#include "Factories/Factory.h"
+#include "PoseSearchDatabaseFactory.generated.h"
+
+UCLASS()
+class POSESEARCHEDITOR_API UPoseSearchDatabaseFactory : public UFactory
+{
+	GENERATED_UCLASS_BODY()
+
+	// UFactory interface
+	virtual UObject* FactoryCreateNew(UClass* Class, UObject* InParent, FName Name, EObjectFlags Flags, UObject* Context, FFeedbackContext* Warn) override;
+	// End of UFactory interface
+};
diff --git a/Engine/Plugins/Experimental/Animation/PoseSearch/Source/Editor/Public/PoseSearchSchemaFactory.h b/Engine/Plugins/Experimental/Animation/PoseSearch/Source/Editor/Public/PoseSearchSchemaFactory.h
new file mode 100644
--- /dev/null
+++ b/Engine/Plugins/Experimental/Animation/PoseSearch/Source/Editor/Public/PoseSearchSchemaFactory.h
@@ -0,0 +1,18 @@
+// Copyright Epic Games, Inc. All Rights Reserved.
+
+#pragma once
+
+#include "CoreMinimal.h"
+#include "UObject/ObjectMacros.h"
+#include "Factories/Factory.h"
+#include "PoseSearchSchemaFactory.generated.h"
+
+UCLASS()
+class POSESEARCHEDITOR_API UPoseSearchSchemaFactory : public UFactory
+{
+	GENERATED_UCLASS_BODY()
+
+	// UFactory interface
+	virtual UObject* FactoryCreateNew(UClass* Class, UObject* InParent, FName Name, EObjectFlags Flags, UObject* Context, FFeedbackContext* Warn) override;
+	// End of UFactory interface
+};
diff --git a/Engine/Plugins/Experimental/Animation/PoseSearch/Source/Editor/Public/PoseSearchTypeActions.h b/Engine/Plugins/Experimental/Animation/PoseSearch/Source/Editor/Public/PoseSearchTypeActions.h
--- a/Engine/Plugins/Experimental/Animation/PoseSearch/Source/Editor/Public/PoseSearchTypeActions.h
+++ b/Engine/Plugins/Experimental/Animation/PoseSearch/Source/Editor/Public/PoseSearchTypeActions.h
@@ -19,5 +19,18 @@ namespace UE::PoseSearch
 			const TArray<UObject*>& InObjects,
 			TSharedPtr<class IToolkitHost> EditWithinLevelEditor = TSharedPtr<IToolkitHost>()) override;
 		virtual uint32 GetCategories() override;
+		virtual const TArray<FText>& GetSubMenus() const override;
+	};
+
+	class FSchemaTypeActions : public FAssetTypeActions_Base
+	{
+	public:
+
+		FSchemaTypeActions() {}
+		virtual FText GetName() const override;
+		virtual FColor GetTypeColor() const override;
+		virtual UClass* GetSupportedClass() const override;
+		virtual uint32 GetCategories() override;
+		virtual const TArray<FText>& GetSubMenus() const override;
 	};
 }
