PoseSearch - Database Editor - fixed editor crash when using the editor with an empty pose search database

diff --git a/Engine/Plugins/Experimental/Animation/PoseSearch/Source/Editor/Private/PoseSearchDatabaseViewModel.cpp b/Engine/Plugins/Experimental/Animation/PoseSearch/Source/Editor/Private/PoseSearchDatabaseViewModel.cpp
--- a/Engine/Plugins/Experimental/Animation/PoseSearch/Source/Editor/Private/PoseSearchDatabaseViewModel.cpp
+++ b/Engine/Plugins/Experimental/Animation/PoseSearch/Source/Editor/Private/PoseSearchDatabaseViewModel.cpp
@@ -82,35 +82,38 @@ namespace UE::PoseSearch
 			MaxPreviewPlayLength = 0.0f;
 
 			const FPoseSearchIndex* SearchIndex = PoseSearchDatabase->GetSearchIndex();
-			for (const FPoseSearchIndexAsset& IndexAsset : SearchIndex->Assets)
+			if (SearchIndex)
 			{
-				const bool bIsAssociatedToSelection =
-					IndexAsset.Type == ESearchIndexAssetType::Sequence &&
-					AssociatedAssetIndices.Contains(IndexAsset.SourceAssetIdx);
-
-				const bool bSpawn =
-					bIsAssociatedToSelection &&
-					(AnimationPreviewMode == EAnimationPreviewMode::OriginalAndMirrored ||
-					 !IndexAsset.bMirrored);
-				
-				if (bSpawn)
+				for (const FPoseSearchIndexAsset& IndexAsset : SearchIndex->Assets)
 				{
-					FDatabasePreviewActor PreviewActor = SpawnPreviewActor(IndexAsset);
-					if (PreviewActor.IsValid())
+					const bool bIsAssociatedToSelection =
+						IndexAsset.Type == ESearchIndexAssetType::Sequence &&
+						AssociatedAssetIndices.Contains(IndexAsset.SourceAssetIdx);
+
+					const bool bSpawn =
+						bIsAssociatedToSelection &&
+						(AnimationPreviewMode == EAnimationPreviewMode::OriginalAndMirrored ||
+						 !IndexAsset.bMirrored);
+
+					if (bSpawn)
 					{
-						const UAnimSequence* Sequence =
-							Cast<UAnimSequence>(PreviewActor.AnimInstance->GetAnimationAsset());
-						if (Sequence)
+						FDatabasePreviewActor PreviewActor = SpawnPreviewActor(IndexAsset);
+						if (PreviewActor.IsValid())
 						{
-							MaxPreviewPlayLength = FMath::Max(MaxPreviewPlayLength, Sequence->GetPlayLength());
+							const UAnimSequence* Sequence =
+								Cast<UAnimSequence>(PreviewActor.AnimInstance->GetAnimationAsset());
+							if (Sequence)
+							{
+								MaxPreviewPlayLength = FMath::Max(MaxPreviewPlayLength, Sequence->GetPlayLength());
+							}
+
+							PreviewActors.Add(PreviewActor);
 						}
-
-						PreviewActors.Add(PreviewActor);
 					}
 				}
-			}
 
-			UpdatePreviewActors();
+				UpdatePreviewActors();
+			}
 		}
 
 		// todo: do blendspaces afterwards
diff --git a/Engine/Plugins/Experimental/Animation/PoseSearch/Source/Editor/Private/SPoseSearchDatabaseAssetList.cpp b/Engine/Plugins/Experimental/Animation/PoseSearch/Source/Editor/Private/SPoseSearchDatabaseAssetList.cpp
--- a/Engine/Plugins/Experimental/Animation/PoseSearch/Source/Editor/Private/SPoseSearchDatabaseAssetList.cpp
+++ b/Engine/Plugins/Experimental/Animation/PoseSearch/Source/Editor/Private/SPoseSearchDatabaseAssetList.cpp
@@ -736,47 +736,15 @@ namespace UE::PoseSearch
 		const TArray<TSharedPtr<FDatabaseAssetTreeNode>> SelectedNodes = TreeView->GetSelectedItems();
 		if (!SelectedNodes.IsEmpty())
 		{
-			TSharedPtr<FDatabaseAssetTreeNode> SelectedNode = SelectedNodes[0];
-
-			if (SelectedNode->SourceAssetType != ESearchIndexAssetType::Invalid)
-			{
-				const int32 GroupIdx = SelectedNode->Parent->SourceAssetIdx;
-				if (GroupIdx == INDEX_NONE)
-				{
-					MenuBuilder.AddMenuEntry(
-						LOCTEXT("DeleteAsset", "Delete Asset"),
-						LOCTEXT("DeleteAssetTooltip", "Delete asset from database"),
-						FSlateIcon(),
-						FUIAction(
-							FExecuteAction::CreateSP(this, &SDatabaseAssetTree::OnDeleteAsset, SelectedNode, true)),
-						NAME_None,
-						EUserInterfaceActionType::Button);
-				}
-				else
-				{
-					MenuBuilder.AddMenuEntry(
-						LOCTEXT("RemoveFromGroup", "Remove From Group"),
-						LOCTEXT("RemoveFromGroupTooltip", "Remove asset from group"),
-						FSlateIcon(),
-						FUIAction(
-							FExecuteAction::CreateSP(this, &SDatabaseAssetTree::OnRemoveFromGroup, SelectedNode, true)),
-						NAME_None,
-						EUserInterfaceActionType::Button);
-				}
-			}
-			else if (SelectedNode->SourceAssetIdx != INDEX_NONE)
-			{
-				// a group that's not default
-				MenuBuilder.AddMenuEntry(
-					LOCTEXT("DeleteGroup", "Delete Group"),
-					LOCTEXT("DeleteGroupTooltip", "Delete group and remove it from the sequences in it"),
-					FSlateIcon(),
-					FUIAction(
-						FExecuteAction::CreateSP(this, &SDatabaseAssetTree::OnDeleteGroup, SelectedNode, true)),
-					NAME_None,
-					EUserInterfaceActionType::Button);
-			}
-
+			MenuBuilder.AddMenuEntry(
+				LOCTEXT("DeleteUngroup", "Delete / Remove"),
+				LOCTEXT(
+					"DeleteUngroupTooltip", 
+					"Deletes groups and ungrouped assets; removes grouped assets from group."),
+				FSlateIcon(),
+				FUIAction(FExecuteAction::CreateSP(this, &SDatabaseAssetTree::OnDeleteNodes)),
+				NAME_None,
+				EUserInterfaceActionType::Button);
 		}
 
 		return MenuBuilder.MakeWidget();
diff --git a/Engine/Plugins/Experimental/Animation/PoseSearch/Source/Runtime/Private/PoseSearch.cpp b/Engine/Plugins/Experimental/Animation/PoseSearch/Source/Runtime/Private/PoseSearch.cpp
--- a/Engine/Plugins/Experimental/Animation/PoseSearch/Source/Runtime/Private/PoseSearch.cpp
+++ b/Engine/Plugins/Experimental/Animation/PoseSearch/Source/Runtime/Private/PoseSearch.cpp
@@ -1180,7 +1180,12 @@ int32 UPoseSearchDatabase::GetPoseIndexFromTime(float Time, const FPoseSearchInd
 {
 	FFloatInterval Range = SearchIndexAsset->SamplingInterval;
 
-	if (Range.Contains(Time))
+	const bool bHasPoseIndex =
+		SearchIndexAsset->FirstPoseIdx != INDEX_NONE &&
+		SearchIndexAsset->NumPoses > 0 &&
+		Range.Contains(Time);
+
+	if (bHasPoseIndex)
 	{
 		int32 PoseOffset = FMath::RoundToInt(Schema->SampleRate * (Time - Range.Min));
 		
@@ -1728,7 +1733,8 @@ void UPoseSearchDatabase::BeginCacheDerivedData()
 		const FIoHash ExistingDerivedDataHash = PrivateDerivedData->PendingDerivedDataKey;
 		if (!ExistingDerivedDataHash.IsZero())
 		{
-			if (ExistingDerivedDataHash == UE::PoseSearch::FPoseSearchDatabaseAsyncCacheTask::CreateKey(*this))
+			const FIoHash CurrentHash = UE::PoseSearch::FPoseSearchDatabaseAsyncCacheTask::CreateKey(*this);
+			if (ExistingDerivedDataHash == CurrentHash)
 			{
 				bPerformCache = false;
 			}
diff --git a/Engine/Plugins/Experimental/Animation/PoseSearch/Source/Runtime/Private/PoseSearchDerivedData.cpp b/Engine/Plugins/Experimental/Animation/PoseSearch/Source/Runtime/Private/PoseSearchDerivedData.cpp
--- a/Engine/Plugins/Experimental/Animation/PoseSearch/Source/Runtime/Private/PoseSearchDerivedData.cpp
+++ b/Engine/Plugins/Experimental/Animation/PoseSearch/Source/Runtime/Private/PoseSearchDerivedData.cpp
@@ -84,6 +84,9 @@ namespace UE::PoseSearch
 
 		if (bForceRebuild)
 		{
+			// when the build is forced, the derived data key is zeroed so the comparison with the pending key fails, 
+			// informing other systems that data is being rebuilt
+			DerivedData.DerivedDataKey.Hash = FIoHash::Zero;
 			BuildAndWrite({ Bucket, DerivedDataKey });
 		}
 		else
