Follow up to CL#20250517 to update more locations to the new Pose Search asset names. I'm probably still missing a couple spots.

diff --git a/Engine/Plugins/Experimental/Animation/PoseSearch/Source/Editor/Private/PoseSearchDatabaseFactory.cpp b/Engine/Plugins/Experimental/Animation/PoseSearch/Source/Editor/Private/PoseSearchDatabaseFactory.cpp
--- a/Engine/Plugins/Experimental/Animation/PoseSearch/Source/Editor/Private/PoseSearchDatabaseFactory.cpp
+++ b/Engine/Plugins/Experimental/Animation/PoseSearch/Source/Editor/Private/PoseSearchDatabaseFactory.cpp
@@ -18,4 +18,9 @@ UObject* UPoseSearchDatabaseFactory::FactoryCreateNew(UClass* Class, UObject* In
 	return NewObject<UPoseSearchDatabase>(InParent, Class, Name, Flags);
 }
 
+FString UPoseSearchDatabaseFactory::GetDefaultNewAssetName() const
+{
+	return FString(TEXT("NewMotionDatabase"));
+}
+
 #undef LOCTEXT_NAMESPACE
\ No newline at end of file
diff --git a/Engine/Plugins/Experimental/Animation/PoseSearch/Source/Editor/Private/PoseSearchSchemaFactory.cpp b/Engine/Plugins/Experimental/Animation/PoseSearch/Source/Editor/Private/PoseSearchSchemaFactory.cpp
--- a/Engine/Plugins/Experimental/Animation/PoseSearch/Source/Editor/Private/PoseSearchSchemaFactory.cpp
+++ b/Engine/Plugins/Experimental/Animation/PoseSearch/Source/Editor/Private/PoseSearchSchemaFactory.cpp
@@ -18,4 +18,9 @@ UObject* UPoseSearchSchemaFactory::FactoryCreateNew(UClass* Class, UObject* InPa
 	return NewObject<UPoseSearchSchema>(InParent, Class, Name, Flags);
 }
 
+FString UPoseSearchSchemaFactory::GetDefaultNewAssetName() const
+{
+	return FString(TEXT("NewMotionDatabaseConfig"));
+}
+
 #undef LOCTEXT_NAMESPACE
\ No newline at end of file
diff --git a/Engine/Plugins/Experimental/Animation/PoseSearch/Source/Editor/Public/PoseSearchDatabaseFactory.h b/Engine/Plugins/Experimental/Animation/PoseSearch/Source/Editor/Public/PoseSearchDatabaseFactory.h
--- a/Engine/Plugins/Experimental/Animation/PoseSearch/Source/Editor/Public/PoseSearchDatabaseFactory.h
+++ b/Engine/Plugins/Experimental/Animation/PoseSearch/Source/Editor/Public/PoseSearchDatabaseFactory.h
@@ -14,5 +14,6 @@ class POSESEARCHEDITOR_API UPoseSearchDatabaseFactory : public UFactory
 
 	// UFactory interface
 	virtual UObject* FactoryCreateNew(UClass* Class, UObject* InParent, FName Name, EObjectFlags Flags, UObject* Context, FFeedbackContext* Warn) override;
+	virtual FString GetDefaultNewAssetName() const override;
 	// End of UFactory interface
 };
diff --git a/Engine/Plugins/Experimental/Animation/PoseSearch/Source/Editor/Public/PoseSearchSchemaFactory.h b/Engine/Plugins/Experimental/Animation/PoseSearch/Source/Editor/Public/PoseSearchSchemaFactory.h
--- a/Engine/Plugins/Experimental/Animation/PoseSearch/Source/Editor/Public/PoseSearchSchemaFactory.h
+++ b/Engine/Plugins/Experimental/Animation/PoseSearch/Source/Editor/Public/PoseSearchSchemaFactory.h
@@ -14,5 +14,6 @@ class POSESEARCHEDITOR_API UPoseSearchSchemaFactory : public UFactory
 
 	// UFactory interface
 	virtual UObject* FactoryCreateNew(UClass* Class, UObject* InParent, FName Name, EObjectFlags Flags, UObject* Context, FFeedbackContext* Warn) override;
+	virtual FString GetDefaultNewAssetName() const override;
 	// End of UFactory interface
 };
diff --git a/Engine/Plugins/Experimental/Animation/PoseSearch/Source/Runtime/Public/PoseSearch/PoseSearch.h b/Engine/Plugins/Experimental/Animation/PoseSearch/Source/Runtime/Public/PoseSearch/PoseSearch.h
--- a/Engine/Plugins/Experimental/Animation/PoseSearch/Source/Runtime/Public/PoseSearch/PoseSearch.h
+++ b/Engine/Plugins/Experimental/Animation/PoseSearch/Source/Runtime/Public/PoseSearch/PoseSearch.h
@@ -168,7 +168,7 @@ struct POSESEARCH_API FPoseSearchBone
 /**
 * Specifies the format of a pose search index. At runtime, queries are built according to the schema for searching.
 */
-UCLASS(BlueprintType, Category = "Animation|Pose Search", Experimental)
+UCLASS(BlueprintType, Category = "Animation|Pose Search", Experimental, meta = (DisplayName = "Motion Database Config"))
 class POSESEARCH_API UPoseSearchSchema : public UDataAsset, public IBoneReferenceSkeletonProvider
 {
 	GENERATED_BODY()
@@ -754,13 +754,13 @@ public:
 };
 
 /** A data asset for indexing a collection of animation sequences. */
-UCLASS(BlueprintType, Category = "Animation|Pose Search", Experimental)
+UCLASS(BlueprintType, Category = "Animation|Pose Search", Experimental, meta = (DisplayName = "Motion Database"))
 class POSESEARCH_API UPoseSearchDatabase : public UDataAsset
 {
 	GENERATED_BODY()
 public:
-
-	UPROPERTY(EditAnywhere, BlueprintReadOnly, Category="Database")
+	// Motion Database Config asset to use with this database.
+	UPROPERTY(EditAnywhere, BlueprintReadOnly, Category="Database", DisplayName="Config")
 	const UPoseSearchSchema* Schema;
 
 	UPROPERTY(EditAnywhere, Category = "Database")
