PoseSearch: Extensible channels
Schemas are now composed of user defined channels. Channels are responsible for describing their features to the schema, indexing assets, query construction, and debug drawing. Existing schemas are migrated to the default pose or trajectory channel implementations

This change will cause some minor data loss. Schemas will automatically migrate, but database default weights and per-group weights will have to be setup again. The weights structure is now composed of an array of channel weight params rather than hard-coded pose and trajectory weights. The number of entries in the weights array is expected to mirror the number of channels setup on a Schema, but there is no UI feedback about that yet and the weights system will fallback to defaults when there aren't sufficient entries in the array.

diff --git a/Engine/Plugins/Experimental/Animation/PoseSearch/Source/Editor/Private/PoseSearchDebugger.cpp b/Engine/Plugins/Experimental/Animation/PoseSearch/Source/Editor/Private/PoseSearchDebugger.cpp
--- a/Engine/Plugins/Experimental/Animation/PoseSearch/Source/Editor/Private/PoseSearchDebugger.cpp
+++ b/Engine/Plugins/Experimental/Animation/PoseSearch/Source/Editor/Private/PoseSearchDebugger.cpp
@@ -265,25 +265,44 @@ void UPoseSearchMeshComponent::UpdatePose(const FUpdateContext& UpdateContext)
 
 namespace UE::PoseSearch {
 
+struct FChannelCostRange
+{
+	float Min = MAX_flt;
+	float Max = -MAX_flt;
+	float Delta = 1.0f;
+};
+
 class FDebuggerDatabaseRowData : public TSharedFromThis<FDebuggerDatabaseRowData>
 {
 public:
 	FDebuggerDatabaseRowData() = default;
 	
-	float GetPoseCost() const { return PoseCostDetails.ChannelCosts.Num() ? PoseCostDetails.ChannelCosts[0] : 0.0f; }
-	float GetTrajectoryCost() const { return PoseCostDetails.ChannelCosts.Num() >= 3 ? PoseCostDetails.ChannelCosts[1] + PoseCostDetails.ChannelCosts[2] : 0.0f; }
+	float GetChannelCost(int32 ChannelIdx) const
+	{
+		return PoseCostDetails.ChannelCosts.IsValidIndex(ChannelIdx)
+			? PoseCostDetails.ChannelCosts[ChannelIdx]
+			: 0.0f;
+	}
+	
 	float GetAddendsCost() const { return PoseCostDetails.NotifyCostAddend + PoseCostDetails.MirrorMismatchCostAddend; }
 
-	void CalculateColors(float CostMin, float CostDelta, float PoseCostMin, float PoseCostDelta, float TrajectoryCostMin, float TrajectoryCostDelta)
+	void CalculateColors(FChannelCostRange TotalCostRange, TArrayView<const FChannelCostRange> ChannelCostRanges)
 	{
-		const float CostColorBlend = (PoseCostDetails.PoseCost.TotalCost - CostMin) / CostDelta;
+		const float CostColorBlend = (PoseCostDetails.PoseCost.TotalCost - TotalCostRange.Min) / TotalCostRange.Delta;
 		CostColor = LinearColorBlend(FLinearColor::Green, FLinearColor::Red, CostColorBlend);
 
-		const float PoseCostColorBlend = (GetPoseCost() - PoseCostMin) / PoseCostDelta;
-		PoseCostColor = LinearColorBlend(FLinearColor::Green, FLinearColor::Red, PoseCostColorBlend);
+		ChannelCostColors.SetNum(ChannelCostRanges.Num());
+		for (int32 ChannelIdx = 0; ChannelIdx != ChannelCostRanges.Num(); ++ChannelIdx)
+		{
+			const FChannelCostRange& CostRange = ChannelCostRanges[ChannelIdx];
+			const float ColorBlend = (GetChannelCost(ChannelIdx) - CostRange.Min) / CostRange.Delta;
+			ChannelCostColors[ChannelIdx] = LinearColorBlend(FLinearColor::Green, FLinearColor::Red, ColorBlend);
+		}
+	}
 
-		const float TrajectoryCostColorBlend = (GetTrajectoryCost() - TrajectoryCostMin) / TrajectoryCostDelta;
-		TrajectoryCostColor = LinearColorBlend(FLinearColor::Green, FLinearColor::Red, TrajectoryCostColorBlend);
+	FLinearColor GetChannelCostColor (int32 ChannelIdx) const
+	{
+		return ChannelCostColors.IsValidIndex(ChannelIdx) ? ChannelCostColors[ChannelIdx] : FLinearColor::White;
 	}
 
 	ESearchIndexAssetType AssetType = ESearchIndexAssetType::Invalid;
@@ -298,8 +317,7 @@ public:
 	FVector BlendParameters = FVector::Zero();
 	FPoseCostDetails PoseCostDetails;
 	FLinearColor CostColor = FLinearColor::White;
-	FLinearColor PoseCostColor = FLinearColor::White;
-	FLinearColor TrajectoryCostColor = FLinearColor::White;
+	TArray<FLinearColor> ChannelCostColors;
 };
 
 namespace DebuggerDatabaseColumns
@@ -311,10 +329,13 @@ namespace DebuggerDatabaseColumns
 			: SortIndex(InSortIndex)
 			, bEnabled(InEnabled)
 		{
+			ColumnId = FName(FString::Printf(TEXT("Column %d"), SortIndex));
 		}
 
 		virtual ~IColumn() = default;
 
+		FName ColumnId;
+
 		/** Sorted left to right based on this index */
 		int32 SortIndex = 0;
 		/** Current width, starts at 1 to be evenly spaced between all columns */
@@ -322,7 +343,7 @@ namespace DebuggerDatabaseColumns
 		/** Disabled selectively with view options */
 		bool bEnabled = false;
 
-		virtual FName GetName() const = 0;
+		virtual FText GetLabel() const = 0;
 
 		using FRowDataRef = TSharedRef<FDebuggerDatabaseRowData>;
 		using FSortPredicate = TFunctionRef<bool(const FRowDataRef&, const FRowDataRef&)>;
@@ -364,8 +385,8 @@ namespace DebuggerDatabaseColumns
 	struct FPoseIdx : ITextColumn
 	{
 		using ITextColumn::ITextColumn;
-		static const FName Name;
-		virtual FName GetName() const override { return Name; }
+
+		virtual FText GetLabel() const override { return LOCTEXT("ColumnLabelPoseIndex", "Pose Index"); }
 
 		virtual FSortPredicate GetSortPredicate() const override
 		{
@@ -377,13 +398,12 @@ namespace DebuggerDatabaseColumns
 			return FText::AsNumber(Row->PoseIdx, &FNumberFormattingOptions::DefaultNoGrouping());
 		}
 	};
-	const FName FPoseIdx::Name = "Pose Index";
 
 	struct FAssetName : IColumn
 	{
 		using IColumn::IColumn;
-		static const FName Name;
-		virtual FName GetName() const override { return Name; }
+
+		virtual FText GetLabel() const override { return LOCTEXT("ColumnLabelAssetName", "Asset"); }
 		
 		virtual FSortPredicate GetSortPredicate() const override
 		{
@@ -407,13 +427,12 @@ namespace DebuggerDatabaseColumns
 					});
 		}
 	};
-	const FName FAssetName::Name = "Asset";
 
 	struct FAssetType : ITextColumn
 	{
 		using ITextColumn::ITextColumn;
-		static const FName Name;
-		virtual FName GetName() const override { return Name; }
+
+		virtual FText GetLabel() const override { return LOCTEXT("ColumnLabelAssetType", "Type"); }
 
 		virtual FSortPredicate GetSortPredicate() const override
 		{
@@ -437,13 +456,12 @@ namespace DebuggerDatabaseColumns
 			}
 		}
 	};
-	const FName FAssetType::Name = "Type";
 
 	struct FFrame : ITextColumn
 	{
 		using ITextColumn::ITextColumn;
-		static const FName Name;
-		virtual FName GetName() const override { return Name; }
+
+		virtual FText GetLabel() const override { return LOCTEXT("ColumnLabelFrame", "Frame"); }
 		
 		virtual FSortPredicate GetSortPredicate() const override
 		{
@@ -477,13 +495,12 @@ namespace DebuggerDatabaseColumns
 			}
 		}
 	};
-	const FName FFrame::Name = "Frame";
 
 	struct FCost : ITextColumn
 	{
 		using ITextColumn::ITextColumn;
-		static const FName Name;
-		virtual FName GetName() const override { return Name; }
+
+		virtual FText GetLabel() const override { return LOCTEXT("ColumnLabelCost", "Cost"); }
 		
 		virtual FSortPredicate GetSortPredicate() const override
 		{
@@ -500,61 +517,46 @@ namespace DebuggerDatabaseColumns
 			return Row->CostColor;
 		}
 	};
-	const FName FCost::Name = "Cost";
-
 
-	struct FPoseCostColumn : ITextColumn
+	struct FChannelCostColumn : ITextColumn
 	{
-		using ITextColumn::ITextColumn;
-		static const FName Name;
-		virtual FName GetName() const override { return Name; }
-		
-		virtual FSortPredicate GetSortPredicate() const override
+		struct FParams
 		{
-			return [](const FRowDataRef& Row0, const FRowDataRef& Row1) -> bool { return Row0->GetPoseCost() < Row1->GetPoseCost(); };
-		}
+			int32 SortIndex;
+			bool bEnabled = true;
+			int32 ChannelIdx;
+		};
 
-		virtual FText GetRowText(const FRowDataRef& Row) const override
-		{
-			return FText::AsNumber(Row->GetPoseCost());
-		}
-
-		virtual FSlateColor GetColorAndOpacity(const FRowDataRef& Row) const override
-		{
-			return Row->PoseCostColor;
-		}
-	};
-	const FName FPoseCostColumn::Name = "Pose Cost";
+		FChannelCostColumn(const FParams& Params)
+			: ITextColumn(Params.SortIndex, Params.bEnabled)
+			, ChannelIdx(Params.ChannelIdx)
+		{}
 
+		virtual FText GetLabel() const override { return FText::Format(LOCTEXT("ColumnLabelChannelCost", "Cost[{0}]"), ChannelIdx); }
 
-	struct FTrajectoryCost : ITextColumn
-	{
-		using ITextColumn::ITextColumn;
-		static const FName Name;
-		virtual FName GetName() const override { return Name; }
-		
 		virtual FSortPredicate GetSortPredicate() const override
 		{
-			return [](const FRowDataRef& Row0, const FRowDataRef& Row1) -> bool { return Row0->GetTrajectoryCost() < Row1->GetTrajectoryCost(); };
+			return [this](const FRowDataRef& Row0, const FRowDataRef& Row1) -> bool { return Row0->GetChannelCost(ChannelIdx) < Row1->GetChannelCost(ChannelIdx); };
 		}
-		
+
 		virtual FText GetRowText(const FRowDataRef& Row) const override
 		{
-			return FText::AsNumber(Row->GetTrajectoryCost());
-        }
+			return FText::AsNumber(Row->GetChannelCost(ChannelIdx));
+		}
 
 		virtual FSlateColor GetColorAndOpacity(const FRowDataRef& Row) const override
 		{
-			return Row->TrajectoryCostColor;
+			return Row->GetChannelCostColor(ChannelIdx);
 		}
+
+		int32 ChannelIdx = -1;
 	};
-	const FName FTrajectoryCost::Name = "Trajectory Cost";
 
 	struct FCostModifier : ITextColumn
 	{
 		using ITextColumn::ITextColumn;
-		static const FName Name;
-		virtual FName GetName() const override { return Name; }
+
+		virtual FText GetLabel() const override { return LOCTEXT("ColumnLabelCostModifier", "Cost Modifier"); }
 
 		virtual FSortPredicate GetSortPredicate() const override
 		{
@@ -566,13 +568,12 @@ namespace DebuggerDatabaseColumns
 			return FText::AsNumber(Row->GetAddendsCost());
 		}
 	};
-	const FName FCostModifier::Name = "Cost Modifier";
 
 	struct FMirrored : ITextColumn
 	{
 		using ITextColumn::ITextColumn;
-		static const FName Name;
-		virtual FName GetName() const override { return Name; }
+
+		virtual FText GetLabel() const override { return LOCTEXT("ColumnLabelMirrored", "Mirrored"); }
 
 		virtual FSortPredicate GetSortPredicate() const override
 		{
@@ -584,13 +585,12 @@ namespace DebuggerDatabaseColumns
 			return FText::Format(LOCTEXT("Mirrored", "{0}"), { Row->bMirrored } );
 		}
 	};
-	const FName FMirrored::Name = "Mirrored";
 
 	struct FLooping : ITextColumn
 	{
 		using ITextColumn::ITextColumn;
-		static const FName Name;
-		virtual FName GetName() const override { return Name; }
+
+		virtual FText GetLabel() const override { return LOCTEXT("ColumnLabelLooping", "Looping"); }
 
 		virtual FSortPredicate GetSortPredicate() const override
 		{
@@ -602,13 +602,12 @@ namespace DebuggerDatabaseColumns
 			return FText::Format(LOCTEXT("Looping", "{0}"), { Row->bLooping });
 		}
 	};
-	const FName FLooping::Name = "Looping";
 
 	struct FBlendParameters : ITextColumn
 	{
 		using ITextColumn::ITextColumn;
-		static const FName Name;
-		virtual FName GetName() const override { return Name; }
+
+		virtual FText GetLabel() const override { return LOCTEXT("ColumnLabelBlendParams", "Blend Parameters"); }
 
 		virtual FSortPredicate GetSortPredicate() const override
 		{
@@ -640,8 +639,6 @@ namespace DebuggerDatabaseColumns
 			}
 		}
 	};
-	const FName FBlendParameters::Name = "Blend Parameters";
-
 }
 
 
@@ -773,15 +770,14 @@ void SDebuggerDatabaseView::RefreshColumns()
 	for(TPair<FName, TSharedRef<IColumn>>& ColumnPair : Columns)
 	{
 		IColumn& Column = ColumnPair.Value.Get();
-		const FName& ColumnName = Column.GetName();
 		if(ColumnPair.Value->bEnabled)
 		{
 			SHeaderRow::FColumn::FArguments ColumnArgs = SHeaderRow::FColumn::FArguments()
-				.ColumnId(ColumnName)
-				.DefaultLabel(FText::FromName(ColumnName))
-				.SortMode(this, &SDebuggerDatabaseView::GetColumnSortMode, ColumnName)
+				.ColumnId(Column.ColumnId)
+				.DefaultLabel(Column.GetLabel())
+				.SortMode(this, &SDebuggerDatabaseView::GetColumnSortMode, Column.ColumnId)
 				.OnSort(this, &SDebuggerDatabaseView::OnColumnSortModeChanged)
-				.FillWidth(this, &SDebuggerDatabaseView::GetColumnWidth, ColumnName)
+				.FillWidth(this, &SDebuggerDatabaseView::GetColumnWidth, Column.ColumnId)
 				.VAlignCell(VAlign_Center)
 				.VAlignHeader(VAlign_Center)
 				.HAlignHeader(HAlign_Center)
@@ -790,16 +786,16 @@ void SDebuggerDatabaseView::RefreshColumns()
 			FilteredDatabaseView.HeaderRow->AddColumn(ColumnArgs);
 			
 			// Every time the active column is changed, update the database column
-			ActiveView.HeaderRow->AddColumn(ColumnArgs.OnWidthChanged(this, &SDebuggerDatabaseView::OnColumnWidthChanged, ColumnName));
+			ActiveView.HeaderRow->AddColumn(ColumnArgs.OnWidthChanged(this, &SDebuggerDatabaseView::OnColumnWidthChanged, Column.ColumnId));
 			
-			ContinuingPoseView.HeaderRow->AddColumn(ColumnArgs.OnWidthChanged(this, &SDebuggerDatabaseView::OnColumnWidthChanged, ColumnName));
+			ContinuingPoseView.HeaderRow->AddColumn(ColumnArgs.OnWidthChanged(this, &SDebuggerDatabaseView::OnColumnWidthChanged, Column.ColumnId));
 		}
 	}
 }
 
 void SDebuggerDatabaseView::AddColumn(TSharedRef<DebuggerDatabaseColumns::IColumn>&& Column)
 {
-	Columns.Add(Column->GetName(), Column);
+	Columns.Add(Column->ColumnId, Column);
 }
 
 EColumnSortMode::Type SDebuggerDatabaseView::GetColumnSortMode(const FName ColumnId) const
@@ -1005,7 +1001,7 @@ void SDebuggerDatabaseView::UpdateRows(const FTraceMotionMatchingStateMessage& S
 	check(ContinuingPoseView.Rows.Num() == 1);
 
 	FPoseSearchWeightsContext StateWeights;
-	StateWeights.Update(State.Weights, &Database);
+	StateWeights.Update(&Database);
 
 	UE::PoseSearch::FSearchContext SearchContext;
 	SearchContext.SetSource(&Database);
@@ -1068,61 +1064,56 @@ void SDebuggerDatabaseView::UpdateRows(const FTraceMotionMatchingStateMessage& S
 
 void SDebuggerDatabaseView::ComputeFilteredDatabaseRowsColors()
 {
-	float CostMin = MAX_flt;
-	float CostMax = -MAX_flt;
-
-	float PoseCostMin = MAX_flt;
-	float PoseCostMax = -MAX_flt;
-
-	float TrajectoryCostMin = MAX_flt;
-	float TrajectoryCostMax = -MAX_flt;
-
+	FChannelCostRange CostRange;
+	TArray<FChannelCostRange> ChannelCostRanges;
 	for (const TSharedRef<FDebuggerDatabaseRowData>& Row : FilteredDatabaseView.Rows)
 	{
 		const float Cost = Row->PoseCostDetails.PoseCost.TotalCost;
-		CostMin = FMath::Min(CostMin, Cost);
-		CostMax = FMath::Max(CostMax, Cost);
-
-		const float PoseCost = Row->GetPoseCost();
-		PoseCostMin = FMath::Min(PoseCostMin, PoseCost);
-		PoseCostMax = FMath::Max(PoseCostMax, PoseCost);
+		CostRange.Min = FMath::Min(CostRange.Min, Cost);
+		CostRange.Max = FMath::Max(CostRange.Max, Cost);
 
-		const float TrajectoryCost = Row->GetTrajectoryCost();
-		TrajectoryCostMin = FMath::Min(TrajectoryCostMin, TrajectoryCost);
-		TrajectoryCostMax = FMath::Max(TrajectoryCostMax, TrajectoryCost);
-	}
+		const int32 NumChannels = Row->PoseCostDetails.ChannelCosts.Num();
+		if (ChannelCostRanges.Num() < NumChannels)
+		{
+			ChannelCostRanges.SetNum(NumChannels);
+		}
 
-	float CostDelta = CostMax - CostMin;
-	if (FMath::IsNearlyZero(CostDelta))
-	{
-		CostDelta = 1.f;
+		for (int32 ChannelIdx = 0; ChannelIdx != NumChannels; ++ChannelIdx)
+		{
+			ChannelCostRanges[ChannelIdx].Min = FMath::Min(ChannelCostRanges[ChannelIdx].Min, Row->PoseCostDetails.ChannelCosts[ChannelIdx]);
+			ChannelCostRanges[ChannelIdx].Max = FMath::Max(ChannelCostRanges[ChannelIdx].Max, Row->PoseCostDetails.ChannelCosts[ChannelIdx]);
+		}
 	}
 
-	float PoseCostDelta = PoseCostMax - PoseCostMin;
-	if (FMath::IsNearlyZero(PoseCostDelta))
+	CostRange.Delta = CostRange.Max - CostRange.Min;
+	if (FMath::IsNearlyZero(CostRange.Delta))
 	{
-		PoseCostDelta = 1.f;
+		CostRange.Delta = 1.f;
 	}
 
-	float TrajectoryCostDelta = TrajectoryCostMax - TrajectoryCostMin;
-	if (FMath::IsNearlyZero(TrajectoryCostDelta))
+	for (int32 ChannelIdx = 0; ChannelIdx != ChannelCostRanges.Num(); ++ChannelIdx)
 	{
-		TrajectoryCostDelta = 1.f;
+		FChannelCostRange ChannelCostRange = ChannelCostRanges[ChannelIdx];
+		ChannelCostRange.Delta = ChannelCostRange.Max - ChannelCostRange.Min;
+		if (FMath::IsNearlyZero(ChannelCostRange.Delta))
+		{
+			ChannelCostRange.Delta = 1.f;
+		}
 	}
 
 	for (const TSharedRef<FDebuggerDatabaseRowData>& Row : FilteredDatabaseView.Rows)
 	{
-		Row->CalculateColors(CostMin, CostDelta, PoseCostMin, PoseCostDelta, TrajectoryCostMin, TrajectoryCostDelta);
+		Row->CalculateColors(CostRange, ChannelCostRanges);
 	}
 
 	if (!ActiveView.Rows.IsEmpty())
 	{
-		ActiveView.Rows[0]->CalculateColors(CostMin, CostDelta, PoseCostMin, PoseCostDelta, TrajectoryCostMin, TrajectoryCostDelta);
+		ActiveView.Rows[0]->CalculateColors(CostRange, ChannelCostRanges);
 	}
 
 	if (!ContinuingPoseView.Rows.IsEmpty())
 	{
-		ContinuingPoseView.Rows[0]->CalculateColors(CostMin, CostDelta, PoseCostMin, PoseCostDelta, TrajectoryCostMin, TrajectoryCostDelta);
+		ContinuingPoseView.Rows[0]->CalculateColors(CostRange, ChannelCostRanges);
 	}
 }
 
@@ -1158,17 +1149,38 @@ void SDebuggerDatabaseView::Construct(const FArguments& InArgs)
 
 	// @TODO: Support runtime reordering of these indices
 	// Construct all column types
-	AddColumn(MakeShared<FAssetName>(0));
-	AddColumn(MakeShared<FAssetType>(1));
-	AddColumn(MakeShared<FCost>(2));
-	AddColumn(MakeShared<FPoseCostColumn>(3));
-	AddColumn(MakeShared<FTrajectoryCost>(4));
-	AddColumn(MakeShared<FCostModifier>(5));
-	AddColumn(MakeShared<FFrame>(6));
-	AddColumn(MakeShared<FMirrored>(7));
-	AddColumn(MakeShared<FLooping>(8));
-	AddColumn(MakeShared<FBlendParameters>(9));
-	AddColumn(MakeShared<FPoseIdx>(10));
+	int32 ColumnIdx = 0;
+	AddColumn(MakeShared<FAssetName>(ColumnIdx++));
+	AddColumn(MakeShared<FAssetType>(ColumnIdx++));
+
+	auto CostColumn = MakeShared<FCost>(ColumnIdx++);
+	AddColumn(CostColumn);
+
+	int32 ChannelIdx = 0;
+	FChannelCostColumn::FParams ChannelCostParams;
+	ChannelCostParams.ChannelIdx = ChannelIdx++;
+	ChannelCostParams.SortIndex = ColumnIdx++;
+	AddColumn(MakeShared<FChannelCostColumn>(ChannelCostParams));
+
+	ChannelCostParams.ChannelIdx = ChannelIdx++;
+	ChannelCostParams.SortIndex = ColumnIdx++;
+	AddColumn(MakeShared<FChannelCostColumn>(ChannelCostParams));
+
+	ChannelCostParams.ChannelIdx = ChannelIdx++;
+	ChannelCostParams.SortIndex = ColumnIdx++;
+	AddColumn(MakeShared<FChannelCostColumn>(ChannelCostParams));
+
+	ChannelCostParams.ChannelIdx = ChannelIdx++;
+	ChannelCostParams.SortIndex = ColumnIdx++;
+	AddColumn(MakeShared<FChannelCostColumn>(ChannelCostParams));
+
+
+	AddColumn(MakeShared<FCostModifier>(ColumnIdx++));
+	AddColumn(MakeShared<FFrame>(ColumnIdx++));
+	AddColumn(MakeShared<FMirrored>(ColumnIdx++));
+	AddColumn(MakeShared<FLooping>(ColumnIdx++));
+	AddColumn(MakeShared<FBlendParameters>(ColumnIdx++));
+	AddColumn(MakeShared<FPoseIdx>(ColumnIdx++));
 
 	// Active Row
 	ActiveView.HeaderRow = SNew(SHeaderRow);
@@ -1451,7 +1463,7 @@ void SDebuggerDatabaseView::Construct(const FArguments& InArgs)
 		]
 	];
 	
-	SortColumn = FCost::Name;
+	SortColumn = CostColumn->ColumnId;
 	SortMode = EColumnSortMode::Ascending;
 
 	// Active and Continuing Pose view scroll bars only for indenting the columns to align w/ database
@@ -1805,16 +1817,6 @@ void SDebuggerView::DrawFeatures(
 			return;
 		}
 
-		if (Options.bDrawPoseFeatures)
-		{
-			InDrawParams.Flags |= EDebugDrawFlags::IncludePose;
-		}
-
-		if (Options.bDrawTrajectoryFeatures)
-		{
-			InDrawParams.Flags |= EDebugDrawFlags::IncludeTrajectory;
-		}
-
 		if (Options.bDrawBoneNames)
 		{
 			InDrawParams.Flags |= EDebugDrawFlags::DrawBoneNames;
@@ -2594,10 +2596,10 @@ void FDebuggerViewModel::UpdateAsset()
 		const float DT = static_cast<float>(FApp::GetDeltaTime()) * AssetPlayRate;
 		const float PlayLength = AnimAsset->GetPlayLength();
 
-		const float DistanceHorizon = Database->Schema->GetTrajectoryFutureDistanceHorizon();
-		const bool bExceededDistanceHorizon = Component->LastRootMotionDelta.GetTranslation().Size() > DistanceHorizon;
-		const float TimeHorizon = Database->Schema->GetTrajectoryFutureTimeHorizon();
-		const bool bExceededTimeHorizon = (AssetSkeleton.Time - AssetData.StartTime) > TimeHorizon;
+		const FFloatRange DistanceRange = Database->Schema->GetHorizonRange(EPoseSearchFeatureDomain::Distance);
+		const bool bExceededDistanceHorizon = !DistanceRange.IsEmpty() && (Component->LastRootMotionDelta.GetTranslation().Size() > DistanceRange.GetUpperBoundValue());
+		const FFloatRange TimeRange = Database->Schema->GetHorizonRange(EPoseSearchFeatureDomain::Time);
+		const bool bExceededTimeHorizon = !TimeRange.IsEmpty() && ((AssetSkeleton.Time - AssetData.StartTime) > TimeRange.GetUpperBoundValue());
 		const bool bExceededHorizon = bExceededDistanceHorizon && bExceededTimeHorizon;
 		if (bAssetLooping)
 		{
diff --git a/Engine/Plugins/Experimental/Animation/PoseSearch/Source/Editor/Private/PoseSearchDebugger.h b/Engine/Plugins/Experimental/Animation/PoseSearch/Source/Editor/Private/PoseSearchDebugger.h
--- a/Engine/Plugins/Experimental/Animation/PoseSearch/Source/Editor/Private/PoseSearchDebugger.h
+++ b/Engine/Plugins/Experimental/Animation/PoseSearch/Source/Editor/Private/PoseSearchDebugger.h
@@ -143,12 +143,6 @@ struct FPoseSearchDebuggerFeatureDrawOptions
 	UPROPERTY(EditAnywhere, Category="Draw Options")
     bool bDisable = false;
 
-	UPROPERTY(EditAnywhere, Category="Draw Options", Meta=(EditCondition="!bDisable"))
-	bool bDrawPoseFeatures = true;
-
-	UPROPERTY(EditAnywhere, Category="Draw Options", Meta=(EditCondition="!bDisable"))
-    bool bDrawTrajectoryFeatures = true;
-
 	UPROPERTY(EditAnywhere, Category = "Draw Options", Meta = (EditCondition = "!bDisable"))
 	bool bDrawBoneNames = true;
 
diff --git a/Engine/Plugins/Experimental/Animation/PoseSearch/Source/Editor/Private/Trace/PoseSearchTraceAnalyzer.cpp b/Engine/Plugins/Experimental/Animation/PoseSearch/Source/Editor/Private/Trace/PoseSearchTraceAnalyzer.cpp
--- a/Engine/Plugins/Experimental/Animation/PoseSearch/Source/Editor/Private/Trace/PoseSearchTraceAnalyzer.cpp
+++ b/Engine/Plugins/Experimental/Animation/PoseSearch/Source/Editor/Private/Trace/PoseSearchTraceAnalyzer.cpp
@@ -62,21 +62,6 @@ bool FTraceAnalyzer::OnEvent(uint16 RouteId, EStyle Style, const FOnEventContext
 				Message.DatabaseSequenceFilter = EventData.GetArrayView<bool>("DatabaseSequenceFilter");
 				Message.DatabaseBlendSpaceFilter = EventData.GetArrayView<bool>("DatabaseBlendSpaceFilter");
 
-				TArrayView<const float> ChannelWeightScales = EventData.GetArrayView<float>("ChannelWeightScales");
-				TArrayView<const float> HistoryWeightScales = EventData.GetArrayView<float>("HistoryWeightScales");
-				TArrayView<const float> PredictionWeightScales = EventData.GetArrayView<float>("PredictionWeightScales");
-
-				if ((ChannelWeightScales.Num() == 2) && (HistoryWeightScales.Num() == 2) && (PredictionWeightScales.Num() == 2))
-				{
-					Message.Weights.PoseDynamicWeights.ChannelWeightScale = ChannelWeightScales[0];
-					Message.Weights.TrajectoryDynamicWeights.ChannelWeightScale = ChannelWeightScales[1];
-					Message.Weights.PoseDynamicWeights.HistoryWeightScale = HistoryWeightScales[0];
-					Message.Weights.TrajectoryDynamicWeights.HistoryWeightScale = HistoryWeightScales[1];
-					Message.Weights.PoseDynamicWeights.PredictionWeightScale = PredictionWeightScales[0];
-					Message.Weights.TrajectoryDynamicWeights.PredictionWeightScale = PredictionWeightScales[1];
-					Message.Weights.bDebugDisableWeights = EventData.GetValue<bool>("DebugDisableWeights");
-				}
-
 				// Common data
 				Message.NodeId = NodeId;
 				Message.AnimInstanceId = AnimInstanceId;
diff --git a/Engine/Plugins/Experimental/Animation/PoseSearch/Source/Editor/Private/Trace/PoseSearchTraceProvider.h b/Engine/Plugins/Experimental/Animation/PoseSearch/Source/Editor/Private/Trace/PoseSearchTraceProvider.h
--- a/Engine/Plugins/Experimental/Animation/PoseSearch/Source/Editor/Private/Trace/PoseSearchTraceProvider.h
+++ b/Engine/Plugins/Experimental/Animation/PoseSearch/Source/Editor/Private/Trace/PoseSearchTraceProvider.h
@@ -27,7 +27,6 @@ struct FTraceMotionMatchingStateMessage : FTraceMessage
 	float ElapsedPoseJumpTime = -1.0f;
 	TArray<float> QueryVector;
 	TArray<float> QueryVectorNormalized;
-	FPoseSearchDynamicWeightParams Weights;
 	int32 DbPoseIdx = INDEX_NONE;
 	uint64 DatabaseId = INDEX_NONE;
 	int32 ContinuingPoseIdx = INDEX_NONE;
diff --git a/Engine/Plugins/Experimental/Animation/PoseSearch/Source/Runtime/Private/AnimNode_PoseSearchHistoryCollector.cpp b/Engine/Plugins/Experimental/Animation/PoseSearch/Source/Runtime/Private/AnimNode_PoseSearchHistoryCollector.cpp
--- a/Engine/Plugins/Experimental/Animation/PoseSearch/Source/Runtime/Private/AnimNode_PoseSearchHistoryCollector.cpp
+++ b/Engine/Plugins/Experimental/Animation/PoseSearch/Source/Runtime/Private/AnimNode_PoseSearchHistoryCollector.cpp
@@ -32,6 +32,7 @@ public:
 
 }}} // namespace UE::PoseSearch::Private
 
+
 /////////////////////////////////////////////////////
 // FAnimNode_PoseSearchHistoryCollector
 
@@ -60,6 +61,8 @@ void FAnimNode_PoseSearchHistoryCollector::Evaluate_AnyThread(FPoseContext& Outp
 {
 	DECLARE_SCOPE_HIERARCHICAL_COUNTER_ANIMNODE(Evaluate_AnyThread);
 
+	using namespace UE::PoseSearch;
+
 	Source.Evaluate(Output);
 
 	FText ErrorText;
@@ -70,8 +73,8 @@ void FAnimNode_PoseSearchHistoryCollector::Evaluate_AnyThread(FPoseContext& Outp
 		Output.AnimInstanceProxy->GetComponentTransform(),
 		&ErrorText,
 		bUseRootMotion ?
-		UE::PoseSearch::EPoseHistoryRootUpdateMode::RootMotionDelta :
-		UE::PoseSearch::EPoseHistoryRootUpdateMode::ComponentTransformDelta))
+		FPoseHistory::ERootUpdateMode::RootMotionDelta :
+		FPoseHistory::ERootUpdateMode::ComponentTransformDelta))
 	{
 		Output.LogMessage(EMessageSeverity::Error, ErrorText);
 	}
diff --git a/Engine/Plugins/Experimental/Animation/PoseSearch/Source/Runtime/Private/PoseSearch.cpp b/Engine/Plugins/Experimental/Animation/PoseSearch/Source/Runtime/Private/PoseSearch.cpp
--- a/Engine/Plugins/Experimental/Animation/PoseSearch/Source/Runtime/Private/PoseSearch.cpp
+++ b/Engine/Plugins/Experimental/Animation/PoseSearch/Source/Runtime/Private/PoseSearch.cpp
@@ -2,6 +2,9 @@
 
 #include "PoseSearch/PoseSearch.h"
 #include "PoseSearch/PoseSearchAnimNotifies.h"
+#if WITH_EDITORONLY_DATA
+#include "PoseSearch/PoseSearchFeatureChannels.h"
+#endif // WITH_EDITORONLY_DATA
 #include "PoseSearchDerivedData.h"
 #include "PoseSearchEigenHelper.h"
 
@@ -47,22 +50,20 @@ namespace UE { namespace PoseSearch {
 //////////////////////////////////////////////////////////////////////////
 // Constants and utilities
 
-// Debug drawing options
-constexpr float DrawDebugLineThickness = 2.0f;
-constexpr float DrawDebugPointSize = 3.0f;
-constexpr float DrawDebugVelocityScale = 0.08f;
-constexpr float DrawDebugArrowSize = 30.0f;
-constexpr float DrawDebugSphereSize = 3.0f;
-constexpr int32 DrawDebugSphereSegments = 10;
-constexpr float DrawDebugGradientStrength = 0.8f;
-constexpr float DrawDebugSampleLabelFontScale = 1.0f;
-static const FVector DrawDebugSampleLabelOffset = FVector(0.0f, 0.0f, -10.0f);
-
-// Stopgap channel indices since schemas don't yet support explicit channels of data
-constexpr int32 ChannelIdxPose = 0;
-constexpr int32 ChannelIdxTrajectoryTime = 1;
-constexpr int32 ChannelIdxTrajectoryDistance = 2;
-constexpr int32 MaxChannels = 3;
+// Temporarily disable parallel indexing due to a bug
+// A single BoneContainer is being used while indexing assets, but some
+// const BoneContainer methods change mutable lookup tables.
+#define UE_POSE_SEARCH_FORCE_SINGLE_THREAD 1
+
+#ifndef UE_POSE_SEARCH_FORCE_SINGLE_THREAD
+#define UE_POSE_SEARCH_FORCE_SINGLE_THREAD 0
+#endif
+
+#if UE_POSE_SEARCH_FORCE_SINGLE_THREAD
+	constexpr EParallelForFlags ParallelForFlags = EParallelForFlags::ForceSingleThread;
+#else
+	constexpr EParallelForFlags ParallelForFlags = EParallelForFlags::None;
+#endif // UE_POSE_SEARCH_FORCE_SINGLE_THREAD
 
 static bool IsSamplingRangeValid(FFloatInterval Range)
 {
@@ -86,21 +87,6 @@ static inline float CompareFeatureVectors(int32 NumValues, const float* A, const
 	return (VA - VB).square().sum();
 }
 
-static FLinearColor GetColorForFeature(FPoseSearchFeatureDesc Feature, const FPoseSearchFeatureVectorLayout* Layout)
-{
-	const float FeatureIdx = Layout->Features.IndexOfByKey(Feature);
-	const float FeatureCountIdx = Layout->Features.Num() - 1;
-	const float FeatureCountIdxHalf = FeatureCountIdx / 2.f;
-	check(FeatureIdx != INDEX_NONE);
-
-	const float Hue = FeatureIdx < FeatureCountIdxHalf
-		? FMath::GetMappedRangeValueUnclamped({ 0.f, FeatureCountIdxHalf }, FVector2f(60.f, 0.f), FeatureIdx)
-		: FMath::GetMappedRangeValueUnclamped({ FeatureCountIdxHalf, FeatureCountIdx }, FVector2f(280.f, 220.f), FeatureIdx);
-
-	const FLinearColor ColorHSV(Hue, 1.f, 1.f);
-	return ColorHSV.HSVToLinearRGB();
-}
-
 static FFloatInterval GetEffectiveSamplingRange(const UAnimSequenceBase* Sequence, FFloatInterval RequestedSamplingRange)
 {
 	const bool bSampleAll = (RequestedSamplingRange.Min == 0.0f) && (RequestedSamplingRange.Max == 0.0f);
@@ -164,17 +150,6 @@ FORCEINLINE auto LowerBound(IteratorType First, IteratorType Last, const ValueTy
 	return LowerBound(First, Last, Value, FIdentityFunctor(), TLess<>());
 }
 
-static EPoseSearchFeatureDomain FeatureDomainFromChannel (int32 ChannelIdx)
-{
-	switch (ChannelIdx)
-	{
-		case ChannelIdxPose: return EPoseSearchFeatureDomain::Time;
-		case ChannelIdxTrajectoryTime: return EPoseSearchFeatureDomain::Time;
-		case ChannelIdxTrajectoryDistance: return EPoseSearchFeatureDomain::Distance;
-		default: return EPoseSearchFeatureDomain::Invalid;
-	}
-}
-
 
 //////////////////////////////////////////////////////////////////////////
 // FFeatureTypeTraits
@@ -213,8 +188,8 @@ static FFeatureTypeTraits GetFeatureTypeTraits(EPoseSearchFeatureType Type)
 static void CalcChannelCosts(const UPoseSearchSchema* Schema, TArrayView<const float> CostVector, TArray<float>& OutChannelCosts)
 {
 	OutChannelCosts.Reset();
-	OutChannelCosts.AddZeroed(MaxChannels);
-	for (int ChannelIdx = 0; ChannelIdx != MaxChannels; ++ChannelIdx)
+	OutChannelCosts.AddZeroed(Schema->Channels.Num());
+	for (int ChannelIdx = 0; ChannelIdx != Schema->Channels.Num(); ++ChannelIdx)
 	{
 		for (int FeatureIdx = INDEX_NONE; Schema->Layout.EnumerateBy(ChannelIdx, EPoseSearchFeatureType::Invalid, FeatureIdx); /*empty*/)
 		{
@@ -237,17 +212,17 @@ static void CalcChannelCosts(const UPoseSearchSchema* Schema, TArrayView<const f
 bool FPoseSearchFeatureDesc::operator==(const FPoseSearchFeatureDesc& Other) const
 {
 	return
-		(SchemaBoneIdx == Other.SchemaBoneIdx) &&
+		(ChannelIdx == Other.ChannelIdx) &&
+		(ChannelFeatureId == Other.ChannelFeatureId) &&
 		(SubsampleIdx == Other.SubsampleIdx) &&
-		(Type == Other.Type) &&
-		(Domain == Other.Domain);
+		(Type == Other.Type);
 }
 
 
 //////////////////////////////////////////////////////////////////////////
 // FPoseSearchFeatureVectorLayout
 
-void FPoseSearchFeatureVectorLayout::Init()
+void FPoseSearchFeatureVectorLayout::Finalize()
 {
 	uint32 FloatCount = 0;
 
@@ -260,31 +235,7 @@ void FPoseSearchFeatureVectorLayout::Init()
 		FloatCount += FeatureNumFloats;
 	}
 
-	// Determine channel count
-	TBitArray ChannelsFound;
-	ChannelsFound.Init(false, UE::PoseSearch::MaxChannels);
-	for (FPoseSearchFeatureDesc& Feature : Features)
-	{
-		if (Feature.SchemaBoneIdx != FPoseSearchFeatureDesc::TrajectoryBoneIndex)
-		{
-			ChannelsFound[UE::PoseSearch::ChannelIdxPose] = true;
-		}
-		else if (Feature.Domain == EPoseSearchFeatureDomain::Time)
-		{
-			ChannelsFound[UE::PoseSearch::ChannelIdxTrajectoryTime] = true;
-		}
-		else if (Feature.Domain == EPoseSearchFeatureDomain::Distance)
-		{
-			ChannelsFound[UE::PoseSearch::ChannelIdxTrajectoryDistance] = true;
-		}
-		else
-		{
-			check(false);
-		}
-	}
-
 	NumFloats = FloatCount;
-	NumChannels = ChannelsFound.CountSetBits();
 }
 
 void FPoseSearchFeatureVectorLayout::Reset()
@@ -293,7 +244,7 @@ void FPoseSearchFeatureVectorLayout::Reset()
 	NumFloats = 0;
 }
 
-bool FPoseSearchFeatureVectorLayout::IsValid(int32 MaxNumBones) const
+bool FPoseSearchFeatureVectorLayout::IsValid(int32 ChannelCount) const
 {
 	if (NumFloats == 0)
 	{
@@ -302,7 +253,7 @@ bool FPoseSearchFeatureVectorLayout::IsValid(int32 MaxNumBones) const
 
 	for (const FPoseSearchFeatureDesc& Feature : Features)
 	{
-		if (Feature.SchemaBoneIdx >= MaxNumBones)
+		if (Feature.ChannelIdx >= ChannelCount)
 		{
 			return false;
 		}
@@ -377,17 +328,51 @@ uint32 FPoseSearchBone::GetTypeMask() const
 //////////////////////////////////////////////////////////////////////////
 // UPoseSearchSchema
 
-void UPoseSearchSchema::PreSave(FObjectPreSaveContext ObjectSaveContext)
+namespace UE::PoseSearch {
+
+int32 FSchemaInitializer::AddBoneReference(const FBoneReference& BoneReference)
+{
+	check(BoneReferences.Num() + 1 <= UPoseSearchSchema::MaxBoneReferences);
+	return BoneReferences.AddUnique(BoneReference);
+}
+
+int32 FSchemaInitializer::AddFeatureDesc(const FPoseSearchFeatureDesc& FeatureDesc)
+{
+	check(FeatureDesc.ChannelIdx == CurrentChannelIdx);
+	check(Features.Num() + 1 <= UPoseSearchSchema::MaxFeatures);
+	check(!Features.Contains(FeatureDesc));
+	return Features.Add(FeatureDesc);
+}
+
+} // namespace UE::PoseSearch
+
+void UPoseSearchSchema::Finalize()
 {
+	using namespace UE::PoseSearch;
+
 	SampleRate = FMath::Clamp(SampleRate, 1, 60);
 	SamplingInterval = 1.0f / SampleRate;
 
-	PoseSampleTimes.Sort(TLess<>());
-	TrajectorySampleTimes.Sort(TLess<>());
-	TrajectorySampleDistances.Sort(TLess<>());
+	// Discard null channels
+	Channels.RemoveAll([](TObjectPtr<UPoseSearchFeatureChannel>& Channel) { return Channel.IsNull(); });
 
-	ResolveBoneReferences();
-	GenerateLayout();
+	Layout.Reset();
+	BoneReferences.Reset();
+
+	FSchemaInitializer Initializer;
+	for (int32 ChannelIdx = 0; ChannelIdx != Channels.Num(); ++ChannelIdx)
+	{
+		TObjectPtr<UPoseSearchFeatureChannel>& Channel = Channels[ChannelIdx];
+		Channel->ChannelIdx = ChannelIdx;
+
+		Initializer.CurrentChannelIdx = ChannelIdx;
+		Channel->InitializeSchema(Initializer);
+	}
+
+	Layout.Features = MoveTemp(Initializer.Features);
+	BoneReferences = MoveTemp(Initializer.BoneReferences);
+
+	Layout.Finalize();
 
 	EffectiveDataPreprocessor = DataPreprocessor;
 	if (EffectiveDataPreprocessor == EPoseSearchDataPreprocessor::Automatic)
@@ -395,6 +380,13 @@ void UPoseSearchSchema::PreSave(FObjectPreSaveContext ObjectSaveContext)
 		EffectiveDataPreprocessor = EPoseSearchDataPreprocessor::Normalize;
 	}
 
+	ResolveBoneReferences();
+}
+
+void UPoseSearchSchema::PreSave(FObjectPreSaveContext ObjectSaveContext)
+{
+	Finalize();
+
 	Super::PreSave(ObjectSaveContext);
 }
 
@@ -402,177 +394,127 @@ void UPoseSearchSchema::PostLoad()
 {
 	Super::PostLoad();
 
-	ResolveBoneReferences();
+	bool bNeedFinalize = false;
 
-	// todo: remove this and deprecated variables
-	if (SampledBones.Num() == 0 && Bones_DEPRECATED.Num() > 0)
+	// Migrate deprecated schema properties into channels
+#if WITH_EDITORONLY_DATA
+	if (!TrajectorySampleTimes_DEPRECATED.IsEmpty())
 	{
-		SampledBones.Empty();
-		for (const FBoneReference& Ref : Bones_DEPRECATED)
-		{
-			FPoseSearchBone Bone;
-			Bone.Reference = Ref;
-			Bone.bUsePosition = bUseBonePositions_DEPRECATED;
-			Bone.bUseVelocity = bUseBoneVelocities_DEPRECATED;
-			Bone.bUseRotation = false;
-			SampledBones.Add(Bone);
-		}
-		Bones_DEPRECATED.Empty();
-	}
-}
+		bNeedFinalize = true;
 
-bool UPoseSearchSchema::IsValid() const
-{
-	bool bValid = Skeleton != nullptr;
+		UPoseSearchFeatureChannel_Trajectory* Channel = NewObject<UPoseSearchFeatureChannel_Trajectory>(this);
+		Channel->Domain = EPoseSearchFeatureDomain::Time;
+		Channel->SampleOffsets = TrajectorySampleTimes_DEPRECATED;
+		Channel->bUseFacingDirections = bUseTrajectoryForwardVectors_DEPRECATED;
+		Channel->bUseLinearVelocities = bUseTrajectoryVelocities_DEPRECATED;
+		Channel->bUsePositions = bUseTrajectoryPositions_DEPRECATED;
+		Channels.Add(Channel);
 
-	for (const FPoseSearchBone& Bone : SampledBones)
-	{
-		bValid &= Bone.Reference.HasValidSetup();
+		TrajectorySampleTimes_DEPRECATED.Empty();
 	}
 
-	bValid &= (SampledBones.Num() == BoneIndices.Num());
-
-	bValid &= Layout.IsValid(BoneIndices.Num());
-	
-	return bValid;
-}
+	if (!TrajectorySampleDistances_DEPRECATED.IsEmpty())
+	{
+		bNeedFinalize = true;
 
-float UPoseSearchSchema::GetTrajectoryFutureTimeHorizon () const
-{
-	return TrajectorySampleTimes.Num() ? TrajectorySampleTimes.Last() : -1.0f;
-}
+		UPoseSearchFeatureChannel_Trajectory* Channel = NewObject<UPoseSearchFeatureChannel_Trajectory>(this);
+		Channel->Domain = EPoseSearchFeatureDomain::Distance;
+		Channel->SampleOffsets = TrajectorySampleDistances_DEPRECATED;
+		Channel->bUseFacingDirections = bUseTrajectoryForwardVectors_DEPRECATED;
+		Channel->bUseLinearVelocities = bUseTrajectoryVelocities_DEPRECATED;
+		Channel->bUsePositions = bUseTrajectoryPositions_DEPRECATED;
+		Channels.Add(Channel);
 
-float UPoseSearchSchema::GetTrajectoryPastTimeHorizon () const
-{
-	return TrajectorySampleTimes.Num() ? TrajectorySampleTimes[0] : 1.0f;
-}
+		TrajectorySampleDistances_DEPRECATED.Empty();
+	}
 
-float UPoseSearchSchema::GetTrajectoryFutureDistanceHorizon () const
-{
-	return TrajectorySampleDistances.Num() ? TrajectorySampleDistances.Last() : -1.0f;
-}
+	if (!PoseSampleTimes_DEPRECATED.IsEmpty())
+	{
+		bNeedFinalize = true;
 
-float UPoseSearchSchema::GetTrajectoryPastDistanceHorizon () const
-{
-	return TrajectorySampleDistances.Num() ? TrajectorySampleDistances[0] : 1.0f;
-}
+		UPoseSearchFeatureChannel_Pose* Channel = NewObject<UPoseSearchFeatureChannel_Pose>(this);
+		Channel->SampledBones = SampledBones_DEPRECATED;
+		Channel->SampleTimes = PoseSampleTimes_DEPRECATED;
+		Channels.Add(Channel);
 
-TArrayView<const float> UPoseSearchSchema::GetChannelSampleOffsets (int32 ChannelIdx) const
-{
-	if (ChannelIdx == UE::PoseSearch::ChannelIdxPose)
-	{
-		return PoseSampleTimes;
+		SampledBones_DEPRECATED.Empty();
+		PoseSampleTimes_DEPRECATED.Empty();
 	}
-	else if (ChannelIdx == UE::PoseSearch::ChannelIdxTrajectoryTime)
+#endif // WITH_EDITORONLY_DATA
+
+	if (bNeedFinalize)
 	{
-		return TrajectorySampleTimes;
+		Finalize();
 	}
-	else if (ChannelIdx == UE::PoseSearch::ChannelIdxTrajectoryDistance)
+	else
 	{
-		return TrajectorySampleDistances;
+		ResolveBoneReferences();
 	}
-
-	return {};
 }
 
-void UPoseSearchSchema::GenerateLayout()
+bool UPoseSearchSchema::IsValid() const
 {
-	Layout.Reset();
-
-	auto AddChannelFeatures = [this](int32 ChannelIdx, EPoseSearchFeatureType Type)
-	{
-		const int32 NumSamples = GetChannelSampleOffsets(ChannelIdx).Num();
-		EPoseSearchFeatureDomain Domain = UE::PoseSearch::FeatureDomainFromChannel(ChannelIdx);
-
-		FPoseSearchFeatureDesc Feature;
-		Feature.Domain = Domain;
-		Feature.Type = Type;
-		Feature.ChannelIdx = ChannelIdx;
-
-		if (ChannelIdx == UE::PoseSearch::ChannelIdxPose)
-		{
-			const int32 NumBones = SampledBones.Num();
-			for (Feature.SubsampleIdx = 0; Feature.SubsampleIdx != NumSamples; ++Feature.SubsampleIdx)
-			{
-				for (Feature.SchemaBoneIdx = 0; Feature.SchemaBoneIdx != NumBones; ++Feature.SchemaBoneIdx)
-				{
-					const FPoseSearchBone& SampledBone = SampledBones[Feature.SchemaBoneIdx];
-					bool bAddFeature = (SampledBone.GetTypeMask() & (1 << static_cast<int>(Type))) != 0;
-					if (bAddFeature)
-					{
-						Layout.Features.Add(Feature);
-					}
-				}
-			}
-		}
-		else
-		{
-			Feature.SchemaBoneIdx = FPoseSearchFeatureDesc::TrajectoryBoneIndex;
-			for (Feature.SubsampleIdx = 0; Feature.SubsampleIdx != NumSamples; ++Feature.SubsampleIdx)
-			{
-				Layout.Features.Add(Feature);
-			}
-		}
-	};
+	bool bValid = Skeleton != nullptr;
 
-	// Time domain trajectory positions
-	if (bUseTrajectoryPositions)
+	for (const FBoneReference& BoneRef : BoneReferences)
 	{
-		AddChannelFeatures(UE::PoseSearch::ChannelIdxTrajectoryTime, EPoseSearchFeatureType::Position);
+		bValid &= BoneRef.HasValidSetup();
 	}
 
-	// Time domain trajectory linear velocities
-	if (bUseTrajectoryVelocities)
+	for (const TObjectPtr<UPoseSearchFeatureChannel>& Channel: Channels)
 	{
-		AddChannelFeatures(UE::PoseSearch::ChannelIdxTrajectoryTime, EPoseSearchFeatureType::LinearVelocity);
+		bValid &= !Channel.IsNull();
 	}
 
-	// Time domain trajectory forward vectors
-	if (bUseTrajectoryForwardVectors)
-	{
-		AddChannelFeatures(UE::PoseSearch::ChannelIdxTrajectoryTime, EPoseSearchFeatureType::ForwardVector);
-	}
+	bValid &= (BoneReferences.Num() == BoneIndices.Num());
 
-	// Distance domain trajectory positions
-	if (bUseTrajectoryPositions)
-	{
-		AddChannelFeatures(UE::PoseSearch::ChannelIdxTrajectoryDistance, EPoseSearchFeatureType::Position);
-	}
+	bValid &= Layout.IsValid(Channels.Num());
+	
+	return bValid;
+}
 
-	// Distance domain trajectory linear velocities
-	if (bUseTrajectoryVelocities)
+FFloatRange UPoseSearchSchema::GetHorizonRange(EPoseSearchFeatureDomain Domain) const
+{
+	FFloatRange GlobalRange = FFloatRange::Empty();
+	for (const auto & Channel : Channels)
 	{
-		AddChannelFeatures(UE::PoseSearch::ChannelIdxTrajectoryDistance, EPoseSearchFeatureType::LinearVelocity);
+		if (Channel)
+		{
+			FFloatRange ChannelRange = Channel->GetHorizonRange(Domain);
+			GlobalRange = FFloatRange::Hull(GlobalRange, ChannelRange);
+		}
 	}
 
-	// Distance domain trajectory forward vectors
-	if (bUseTrajectoryForwardVectors)
+	return GlobalRange;
+}
+
+TArrayView<const float> UPoseSearchSchema::GetChannelSampleOffsets(int32 ChannelIdx) const
+{
+	if (Channels.IsValidIndex(ChannelIdx))
 	{
-		AddChannelFeatures(UE::PoseSearch::ChannelIdxTrajectoryDistance, EPoseSearchFeatureType::ForwardVector);
+		const auto & Channel = Channels[ChannelIdx];
+		if (Channel)
+		{
+			return Channel->GetSampleOffsets();
+		}
 	}
 
-
-	// Time domain bone features
-	AddChannelFeatures(UE::PoseSearch::ChannelIdxPose, EPoseSearchFeatureType::Position);
-	AddChannelFeatures(UE::PoseSearch::ChannelIdxPose, EPoseSearchFeatureType::Rotation);
-	AddChannelFeatures(UE::PoseSearch::ChannelIdxPose, EPoseSearchFeatureType::LinearVelocity);
-
-	Layout.Init();
+	return {};
 }
 
 void UPoseSearchSchema::ResolveBoneReferences()
 {
 	// Initialize references to obtain bone indices
-	for (FPoseSearchBone& Bone : SampledBones)
+	for (FBoneReference& BoneRef : BoneReferences)
 	{
-		Bone.Reference.Initialize(Skeleton);
+		BoneRef.Initialize(Skeleton);
 	}
 
 	// Fill out bone index array
-	BoneIndices.SetNum(SampledBones.Num());
-	for (int32 Index = 0; Index != SampledBones.Num(); ++Index)
+	BoneIndices.SetNum(BoneReferences.Num());
+	for (int32 BoneRefIdx = 0; BoneRefIdx != BoneReferences.Num(); ++BoneRefIdx)
 	{
-		BoneIndices[Index] = SampledBones[Index].Reference.BoneIndex;
+		BoneIndices[BoneRefIdx] = BoneReferences[BoneRefIdx].BoneIndex;
 	}
 
 	// Build separate index array with parent indices guaranteed to be present. Sort for EnsureParentsPresent.
@@ -592,19 +534,13 @@ void UPoseSearchSchema::ResolveBoneReferences()
 }
 
 
-//////////////////////////////////////////////////////////////////////////
-// FPoseSearchChannelWeightParams
-FPoseSearchWeightParams::FPoseSearchWeightParams()
-{
-	TrajectoryWeight.ChannelWeight = 3.0f;
-}
-
 
 //////////////////////////////////////////////////////////////////////////
 // FPoseSearchChannelWeightParams
 
 FPoseSearchChannelWeightParams::FPoseSearchChannelWeightParams()
 {
+	TypeWeights.Empty((int32)EPoseSearchFeatureType::Num);
 	for (int32 Type = 0; Type != (int32)EPoseSearchFeatureType::Num; ++Type)
 	{
 		TypeWeights.Add((EPoseSearchFeatureType)Type, 1.0f);
@@ -612,10 +548,20 @@ FPoseSearchChannelWeightParams::FPoseSearchChannelWeightParams()
 }
 
 
+
+//////////////////////////////////////////////////////////////////////////
+// FPoseSearchWeightParams
+const FPoseSearchChannelWeightParams* FPoseSearchWeightParams::GetChannelWeights(int32 ChannelIdx) const
+{
+	return ChannelWeights.IsValidIndex(ChannelIdx) ? &ChannelWeights[ChannelIdx] : nullptr;
+}
+
+
+
 //////////////////////////////////////////////////////////////////////////
 // FPoseSearchWeights
 
-void FPoseSearchWeights::Init(const FPoseSearchWeightParams& WeightParams, const UPoseSearchSchema* Schema, const FPoseSearchDynamicWeightParams& DynamicWeightParams)
+void FPoseSearchWeights::Init(const FPoseSearchWeightParams& WeightParams, const UPoseSearchSchema* Schema)
 {
 	using namespace UE::PoseSearch;
 	using namespace Eigen;
@@ -632,34 +578,15 @@ void FPoseSearchWeights::Init(const FPoseSearchWeightParams& WeightParams, const
 	Weights.Reset();
 	Weights.SetNumZeroed(Schema->Layout.NumFloats);
 
-	// Completely disable weights if requested
-	if (DynamicWeightParams.bDebugDisableWeights)
-	{
-		for (float& Weight: Weights)
-		{
-			Weight = 1.0f;
-		}
-		return;
-	}
-
 	// Setup channel indexable weight params
-	constexpr int ChannelNum = MaxChannels;
-	const FPoseSearchChannelWeightParams* ChannelWeightParams[ChannelNum];
-	const FPoseSearchChannelDynamicWeightParams* ChannelDynamicWeightParams[ChannelNum];
-
-	ChannelWeightParams[ChannelIdxPose] = &WeightParams.PoseWeight;
-	ChannelWeightParams[ChannelIdxTrajectoryTime] = &WeightParams.TrajectoryWeight;
-	ChannelWeightParams[ChannelIdxTrajectoryDistance] = &WeightParams.TrajectoryWeight;
-
-	ChannelDynamicWeightParams[ChannelIdxPose] = &DynamicWeightParams.PoseDynamicWeights;
-	ChannelDynamicWeightParams[ChannelIdxTrajectoryTime] = &DynamicWeightParams.TrajectoryDynamicWeights;
-	ChannelDynamicWeightParams[ChannelIdxTrajectoryDistance] = &DynamicWeightParams.TrajectoryDynamicWeights;
+	const int32 ChannelNum = Schema->Channels.Num();
 
 	// Normalize channel weights
-	Eigen::Array<float, ChannelNum, 1> NormalizedChannelWeights;
+	Eigen::Array<float, Dynamic, 1> NormalizedChannelWeights(ChannelNum);
 	for (int ChannelIdx = 0; ChannelIdx != ChannelNum; ++ChannelIdx)
 	{
-		NormalizedChannelWeights[ChannelIdx] = ChannelWeightParams[ChannelIdx]->ChannelWeight * ChannelDynamicWeightParams[ChannelIdx]->ChannelWeightScale;
+		const FPoseSearchChannelWeightParams* ChannelWeights = WeightParams.GetChannelWeights(ChannelIdx);
+		NormalizedChannelWeights[ChannelIdx] = ChannelWeights ? ChannelWeights->ChannelWeight : 1.0f;
 
 		// Zero the channel weight if there are no features in this channel
 		int32 FeatureIdx = INDEX_NONE;
@@ -708,8 +635,7 @@ void FPoseSearchWeights::Init(const FPoseSearchWeightParams& WeightParams, const
 		}
 
 		// Get channel info
-		const FPoseSearchChannelWeightParams& ChannelWeights = *ChannelWeightParams[ChannelIdx];
-		const FPoseSearchChannelDynamicWeightParams& ChannelDynamicWeights = *ChannelDynamicWeightParams[ChannelIdx];
+		const FPoseSearchChannelWeightParams* ChannelWeights = WeightParams.GetChannelWeights(ChannelIdx);
 		TArrayView<const float> ChannelSampleOffsets = Schema->GetChannelSampleOffsets(ChannelIdx);
 
 		// Reset scratch weights
@@ -721,7 +647,7 @@ void FPoseSearchWeights::Init(const FPoseSearchWeightParams& WeightParams, const
 		// Initialize weights by type lookup
 		for (int Type = 0; Type != (int)EPoseSearchFeatureType::Num; ++Type)
 		{
-			WeightsByType[Type] = ChannelWeights.TypeWeights.FindRef((EPoseSearchFeatureType)Type);
+			WeightsByType[Type] = ChannelWeights ? ChannelWeights->TypeWeights.FindRef((EPoseSearchFeatureType)Type) : 1.0f;
 
 			// Zero the weight if this channel doesn't have any features using this type
 			int32 FeatureIdx = INDEX_NONE;
@@ -763,11 +689,11 @@ void FPoseSearchWeights::Init(const FPoseSearchWeightParams& WeightParams, const
 
 		if (!HorizonSampleIdxRanges[EHorizon::History].IsEmpty())
 		{
-			NormalizedHorizonWeights[EHorizon::History] = ChannelWeights.HistoryParams.Weight * ChannelDynamicWeights.HistoryWeightScale;
+			NormalizedHorizonWeights[EHorizon::History] = ChannelWeights ? ChannelWeights->HistoryParams.Weight : 1.0f;
 		}
 		if (!HorizonSampleIdxRanges[EHorizon::Prediction].IsEmpty())
 		{
-			NormalizedHorizonWeights[EHorizon::Prediction] = ChannelWeights.PredictionParams.Weight * ChannelDynamicWeights.PredictionWeightScale;
+			NormalizedHorizonWeights[EHorizon::Prediction] = ChannelWeights ? ChannelWeights->PredictionParams.Weight : 1.0f;
 		}
 		
 		// Normalize horizon weights
@@ -783,7 +709,7 @@ void FPoseSearchWeights::Init(const FPoseSearchWeightParams& WeightParams, const
 		}
 
 
-		auto SetHorizonSampleWeights = [&HorizonWeightsBySample, &ChannelSampleOffsets](FInt32Range SampleIdxRange, const FPoseSearchChannelHorizonParams& HorizonParams)
+		auto SetHorizonSampleWeights = [&HorizonWeightsBySample, &ChannelSampleOffsets](FInt32Range SampleIdxRange, const FPoseSearchChannelHorizonParams* HorizonParams)
 		{
 			// Segment length represents the number of sample offsets in the span that make up this horizon
 			int32 SegmentLength = SampleIdxRange.Size<int32>();
@@ -791,20 +717,20 @@ void FPoseSearchWeights::Init(const FPoseSearchWeightParams& WeightParams, const
 			if (SegmentLength > 0)
 			{
 				int32 SegmentBegin = SampleIdxRange.GetLowerBoundValue();
-				if (HorizonParams.bInterpolate && SegmentLength > 1)
+				if (HorizonParams && HorizonParams->bInterpolate && SegmentLength > 1)
 				{
 					// We'll map the range spanned by the horizon's sample offsets to the interpolation range
 					// The interpolation range is 0 to 1 unless an initial value was set
 					// The initial value allows the user to set a minimum weight or reverse the lerp direction
 					// We'll normalize these weights in the next step
 					FVector2f InputRange(ChannelSampleOffsets[SegmentBegin], ChannelSampleOffsets[SegmentBegin + SegmentLength - 1]);
-					FVector2f OutputRange(HorizonParams.InitialValue, 1.0f - HorizonParams.InitialValue);
+					FVector2f OutputRange(HorizonParams->InitialValue, 1.0f - HorizonParams->InitialValue);
 
 					for (int32 OffsetIdx = SegmentBegin; OffsetIdx != SegmentBegin + SegmentLength; ++OffsetIdx)
 					{
 						float SampleOffset = ChannelSampleOffsets[OffsetIdx];
 						float Alpha = FMath::GetMappedRangeValueUnclamped(InputRange, OutputRange, SampleOffset);
-						float Weight = FAlphaBlend::AlphaToBlendOption(Alpha, HorizonParams.InterpolationMethod);
+						float Weight = FAlphaBlend::AlphaToBlendOption(Alpha, HorizonParams->InterpolationMethod);
 						HorizonWeightsBySample[OffsetIdx] = Weight;
 					}
 				}
@@ -823,9 +749,14 @@ void FPoseSearchWeights::Init(const FPoseSearchWeightParams& WeightParams, const
 			}
 		};
 
-		SetHorizonSampleWeights(HorizonSampleIdxRanges[EHorizon::History], ChannelWeights.HistoryParams);
-		SetHorizonSampleWeights(HorizonSampleIdxRanges[EHorizon::Prediction], ChannelWeights.PredictionParams);
-
+		SetHorizonSampleWeights(
+			HorizonSampleIdxRanges[EHorizon::History],
+			ChannelWeights ? &ChannelWeights->HistoryParams : nullptr
+		);
+		SetHorizonSampleWeights(
+			HorizonSampleIdxRanges[EHorizon::Prediction],
+			ChannelWeights ? &ChannelWeights->PredictionParams : nullptr
+		);
 
 		// Now set this channel's weights for every feature in each horizon
 		Eigen::Array<float, 1, EHorizon::Num> HorizonSums;
@@ -839,6 +770,8 @@ void FPoseSearchWeights::Init(const FPoseSearchWeightParams& WeightParams, const
 				if (HorizonSampleIdxRanges[HorizonIdx].Contains(Feature.SubsampleIdx))
 				{
 					int HorizonSize = HorizonSampleIdxRanges[HorizonIdx].Size<int>();
+					
+					// Bug: HorizonSize * WeightsByType[(int)Feature.Type] assumes horizons are composed of homogenous features, not true any more for pose channel
 					WeightsByFeature[FeatureIdx] = HorizonWeightsBySample[Feature.SubsampleIdx] * (HorizonSize * WeightsByType[(int)Feature.Type]);
 					HorizonSums[HorizonIdx] += WeightsByFeature[FeatureIdx];
 					break;
@@ -887,9 +820,7 @@ void FPoseSearchWeights::Init(const FPoseSearchWeightParams& WeightParams, const
 //////////////////////////////////////////////////////////////////////////
 // FPoseSearchWeightsContext
 
-void FPoseSearchWeightsContext::Update(
-	const FPoseSearchDynamicWeightParams& ActiveWeights, 
-	const UPoseSearchDatabase* ActiveDatabase)
+void FPoseSearchWeightsContext::Update(const UPoseSearchDatabase* ActiveDatabase)
 {
 	bool bRecomputeWeights = false;
 	if (Database != ActiveDatabase)
@@ -901,12 +832,6 @@ void FPoseSearchWeightsContext::Update(
 		bRecomputeWeights = true;
 	}
 
-	if (DynamicWeights != ActiveWeights)
-	{
-		DynamicWeights = ActiveWeights;
-		bRecomputeWeights = true;
-	}
-
 #if WITH_EDITOR
 	if (Database.IsValid() && Database->GetSearchIndexHash() != SearchIndexHash)
 	{
@@ -917,7 +842,7 @@ void FPoseSearchWeightsContext::Update(
 
 	if (bRecomputeWeights)
 	{
-		ComputedDefaultGroupWeights.Init(Database->DefaultWeights, Database->Schema, ActiveWeights);
+		ComputedDefaultGroupWeights.Init(Database->DefaultWeights, Database->Schema);
 
 		int32 NumGroups = ActiveDatabase ? ActiveDatabase->Groups.Num() : 0;
 		ComputedGroupWeights.SetNum(NumGroups);
@@ -926,8 +851,7 @@ void FPoseSearchWeightsContext::Update(
 		{
 			ComputedGroupWeights[GroupIdx].Init(
 				Database->Groups[GroupIdx].Weights, 
-				Database->Schema, 
-				ActiveWeights);
+				Database->Schema);
 		}
 	}
 }
@@ -2051,95 +1975,6 @@ void FPoseSearchFeatureVectorBuilder::SetVector(FPoseSearchFeatureDesc Feature,
 	}
 }
 
-bool FPoseSearchFeatureVectorBuilder::TrySetPoseFeatures(UE::PoseSearch::FPoseHistory* History, const FBoneContainer& BoneContainer)
-{
-	check(Schema.IsValid() && Schema->IsValid());
-	check(History);
-
-	FPoseSearchFeatureDesc Feature;
-	Feature.Domain = EPoseSearchFeatureDomain::Time;
-
-	for (int32 SchemaSubsampleIdx = 0; SchemaSubsampleIdx != Schema->PoseSampleTimes.Num(); ++SchemaSubsampleIdx)
-	{
-		Feature.SubsampleIdx = SchemaSubsampleIdx;
-
-		// Stop when we've reached future samples
-		float SampleTime = Schema->PoseSampleTimes[SchemaSubsampleIdx];
-		if (SampleTime > 0.0f)
-		{
-			break;
-		}
-
-		float SecondsAgo = -SampleTime;
-		if (!History->TrySamplePose(SecondsAgo, Schema->Skeleton->GetReferenceSkeleton(), Schema->BoneIndicesWithParents))
-		{
-			return false;
-		}
-
-		TArrayView<const FTransform> ComponentPose = History->GetComponentPoseSample();
-		TArrayView<const FTransform> ComponentPrevPose = History->GetPrevComponentPoseSample();
-		FTransform RootTransform = History->GetRootTransformSample();
-		FTransform RootTransformPrev = History->GetPrevRootTransformSample();
-		for (int32 SchemaBoneIdx = 0; SchemaBoneIdx != Schema->BoneIndices.Num(); ++SchemaBoneIdx)
-		{
-			Feature.SchemaBoneIdx = SchemaBoneIdx;
-
-			int32 SkeletonBoneIndex = Schema->BoneIndices[SchemaBoneIdx];
-
-			const FTransform& Transform = ComponentPose[SkeletonBoneIndex];
-			const FTransform& PrevTransform = ComponentPrevPose[SkeletonBoneIndex] * (RootTransformPrev * RootTransform.Inverse());
-			SetTransform(Feature, Transform);
-			SetTransformVelocity(Feature, Transform, PrevTransform, History->GetSampleTimeInterval());
-		}
-	}
-
-	return true;
-}
-
-void FPoseSearchFeatureVectorBuilder::BuildFromTrajectoryTimeBased(const FTrajectorySampleRange& Trajectory)
-{
-	check(Schema.IsValid() && Schema->IsValid());
-
-	FPoseSearchFeatureDesc Feature;
-	Feature.Domain = EPoseSearchFeatureDomain::Time;
-	Feature.SchemaBoneIdx = FPoseSearchFeatureDesc::TrajectoryBoneIndex;
-
-	for (int32 Idx = 0, NextIterStartIdx = 0, Num = Schema->TrajectorySampleTimes.Num(); Idx < Num; ++Idx)
-	{
-		const float SampleTime = Schema->TrajectorySampleTimes[Idx];
-		const FTrajectorySample Sample = FTrajectorySampleRange::IterSampleTrajectory(Trajectory.Samples, ETrajectorySampleDomain::Time, SampleTime, NextIterStartIdx);
-
-		Feature.SubsampleIdx = Idx;
-
-		Feature.Type = EPoseSearchFeatureType::LinearVelocity;
-		SetVector(Feature, Sample.LinearVelocity);
-
-		SetTransform(Feature, Sample.Transform);
-	}
-}
-
-void FPoseSearchFeatureVectorBuilder::BuildFromTrajectoryDistanceBased(const FTrajectorySampleRange& Trajectory)
-{
-	check(Schema.IsValid() && Schema->IsValid());
-
-	FPoseSearchFeatureDesc Feature;
-	Feature.Domain = EPoseSearchFeatureDomain::Distance;
-	Feature.SchemaBoneIdx = FPoseSearchFeatureDesc::TrajectoryBoneIndex;
-
-	for (int32 Idx = 0, NextIterStartIdx = 0, Num = Schema->TrajectorySampleDistances.Num(); Idx < Num; ++Idx)
-	{
-		const float SampleDistance = Schema->TrajectorySampleDistances[Idx];
-		const FTrajectorySample Sample = FTrajectorySampleRange::IterSampleTrajectory(Trajectory.Samples, ETrajectorySampleDomain::Distance, SampleDistance, NextIterStartIdx);
-
-		Feature.SubsampleIdx = Idx;
-
-		Feature.Type = EPoseSearchFeatureType::LinearVelocity;
-		SetVector(Feature, Sample.LinearVelocity);
-
-		SetTransform(Feature, Sample.Transform);
-	} 
-}
-
 void FPoseSearchFeatureVectorBuilder::CopyFromSearchIndex(const FPoseSearchIndex& SearchIndex, int32 PoseIdx)
 {
 	check(Schema == SearchIndex.Schema);
@@ -2205,16 +2040,15 @@ bool FPoseSearchFeatureVectorBuilder::IsCompatible(const FPoseSearchFeatureVecto
 	return IsInitialized() && (Schema == OtherBuilder.Schema);
 }
 
-void FPoseSearchFeatureVectorBuilder::Normalize(const FPoseSearchIndex& ForSearchIndex)
+const TBitArray<>& FPoseSearchFeatureVectorBuilder::GetFeaturesAdded() const
 {
-	ValuesNormalized = Values;
-	ForSearchIndex.Normalize(ValuesNormalized);
+	return FeaturesAdded;
 }
 
-void FPoseSearchFeatureVectorBuilder::BuildFromTrajectory(const FTrajectorySampleRange& Trajectory)
+void FPoseSearchFeatureVectorBuilder::Normalize(const FPoseSearchIndex& ForSearchIndex)
 {
-	BuildFromTrajectoryTimeBased(Trajectory);
-	BuildFromTrajectoryDistanceBased(Trajectory);
+	ValuesNormalized = Values;
+	ForSearchIndex.Normalize(ValuesNormalized);
 }
 
 namespace UE { namespace PoseSearch {
@@ -2320,7 +2154,7 @@ bool FPoseHistory::TrySamplePose(float SecondsAgo, const FReferenceSkeleton& Ref
 	return bSampled;
 }
 
-bool FPoseHistory::Update(float SecondsElapsed, const FPoseContext& PoseContext, FTransform ComponentTransform, FText* OutError, EPoseHistoryRootUpdateMode UpdateMode)
+bool FPoseHistory::Update(float SecondsElapsed, const FPoseContext& PoseContext, FTransform ComponentTransform, FText* OutError, ERootUpdateMode UpdateMode)
 {
 	// Age our elapsed times
 	for (float& Knot : Knots)
@@ -2366,7 +2200,7 @@ bool FPoseHistory::Update(float SecondsElapsed, const FPoseContext& PoseContext,
 	CurrentPose.RootTransform = Poses.Num() > 1 ? Poses[Poses.Num() - 2].RootTransform : FTransform::Identity;
 	
 	// Update using either AniumRootMotionProvider or Component Transform
-	if (UpdateMode == EPoseHistoryRootUpdateMode::RootMotionDelta)
+	if (UpdateMode == ERootUpdateMode::RootMotionDelta)
 	{
 		const UE::Anim::IAnimRootMotionProvider* RootMotionProvider = UE::Anim::IAnimRootMotionProvider::Get();
 
@@ -2403,7 +2237,7 @@ bool FPoseHistory::Update(float SecondsElapsed, const FPoseContext& PoseContext,
 		}
 #endif
 	}
-	else if (UpdateMode == EPoseHistoryRootUpdateMode::ComponentTransformDelta)
+	else if (UpdateMode == ERootUpdateMode::ComponentTransformDelta)
 	{
 		CurrentPose.RootTransform = ComponentTransform;
 	}
@@ -2581,6 +2415,7 @@ const UPoseSearchSchema* FDebugDrawParams::GetSchema() const
 	return nullptr;
 }
 
+
 //////////////////////////////////////////////////////////////////////////
 // FSearchContext
 
@@ -2628,42 +2463,35 @@ float FSearchContext::GetMirrorMismatchCost() const
 	return MirrorMismatchCost;
 }
 
+
 //////////////////////////////////////////////////////////////////////////
-// FAnimSamplingContext
+// FAssetSamplerContext
 
-struct FAnimSamplingContext
+static void FillCompactPoseAndComponentRefRotations(FAssetSamplingContext& Context)
 {
-	// Time delta used for computing pose derivatives
-	static constexpr float FiniteDelta = 1 / 60.0f;
-
-	FBoneContainer BoneContainer;
-	
-	// Mirror data table pointer copied from Schema for convenience
-	TObjectPtr<UMirrorDataTable> MirrorDataTable = nullptr;
-	
-	// Compact pose format of Mirror Bone Map
-	TCustomBoneIndexArray<FCompactPoseBoneIndex, FCompactPoseBoneIndex> CompactPoseMirrorBones;
-	
-	// Pre-calculated component space rotations of reference pose, which allows mirror to work with any joint orientation
-	// Only initialized and used when a mirroring table is specified
-	TCustomBoneIndexArray<FQuat, FCompactPoseBoneIndex> ComponentSpaceRefRotations;
-	
-	void Init(const UPoseSearchSchema* Schema);
-	
-	FTransform MirrorTransform(const FTransform& Transform) const;
-
-private:
-	void FillCompactPoseAndComponentRefRotations();
-};
+	if (!Context.MirrorDataTable.IsNull())
+	{
+		const UMirrorDataTable* MirrorDataTablePtr = Context.MirrorDataTable.Get();
+		MirrorDataTablePtr->FillCompactPoseAndComponentRefRotations(
+			Context.BoneContainer,
+			Context.CompactPoseMirrorBones,
+			Context.ComponentSpaceRefRotations);
+	}
+	else
+	{
+		Context.CompactPoseMirrorBones.Reset();
+		Context.ComponentSpaceRefRotations.Reset();
+	}
+}
 
-void FAnimSamplingContext::Init(const UPoseSearchSchema* Schema)
+void FAssetSamplingContext::Init(const UPoseSearchSchema* Schema)
 {
 	MirrorDataTable = Schema->MirrorDataTable;
 	BoneContainer.InitializeTo(Schema->BoneIndicesWithParents, FCurveEvaluationOption(false), *Schema->Skeleton);
-	FillCompactPoseAndComponentRefRotations();
+	FillCompactPoseAndComponentRefRotations(*this);
 }
 
-FTransform FAnimSamplingContext::MirrorTransform(const FTransform& InTransform) const
+FTransform FAssetSamplingContext::MirrorTransform(const FTransform& InTransform) const
 {
 	const EAxis::Type MirrorAxis = MirrorDataTable->MirrorAxis;
 	FVector T = InTransform.GetTranslation();
@@ -2676,58 +2504,9 @@ FTransform FAnimSamplingContext::MirrorTransform(const FTransform& InTransform)
 	return Result;
 }
 
-void FAnimSamplingContext::FillCompactPoseAndComponentRefRotations()
-{
-	if (!MirrorDataTable.IsNull())
-	{
-		const UMirrorDataTable* MirrorDataTablePtr = MirrorDataTable.Get();
-		MirrorDataTablePtr->FillCompactPoseAndComponentRefRotations(
-			BoneContainer, 
-			CompactPoseMirrorBones, 
-			ComponentSpaceRefRotations);
-	}
-	else
-	{
-		CompactPoseMirrorBones.Reset();
-		ComponentSpaceRefRotations.Reset();
-	}
-}
 
 //////////////////////////////////////////////////////////////////////////
-// IAssetSampler
-
-struct IAssetSampler
-{
-public:
-	virtual ~IAssetSampler() {};
-
-	virtual float GetPlayLength() const = 0;
-	virtual bool IsLoopable() const = 0;
-
-	// Gets the time associated with a particular root distance travelled
-	virtual float GetTimeFromRootDistance(float Distance) const = 0;
-
-	// Gets the total root distance travelled 
-	virtual float GetTotalRootDistance() const = 0;
-
-	// Gets the final root transformation at the end of the asset's playback time
-	virtual FTransform GetTotalRootTransform() const = 0;
-
-	// Extracts pose for this asset for a given context
-	virtual void ExtractPose(const FAnimExtractContext& ExtractionCtx, FAnimationPoseData& OutAnimPoseData) const = 0;
-
-	// Extracts the accumulated root distance at the given time, using the extremities of the sequence to extrapolate 
-	// beyond the sequence limits when Time is less than zero or greater than the sequence length
-	virtual float ExtractRootDistance(float Time) const = 0;
-
-	// Extracts root transform at the given time, using the extremities of the sequence to extrapolate beyond the 
-	// sequence limits when Time is less than zero or greater than the sequence length.
-	virtual FTransform ExtractRootTransform(float Time) const = 0;
-
-	// Extracts notify states inheriting from UAnimNotifyState_PoseSearchBase present in the sequence at Time.
-	// The function does not empty NotifyStates before adding new notifies!
-	virtual void ExtractPoseSearchNotifyStates(float Time, TArray<class UAnimNotifyState_PoseSearchBase*>& NotifyStates) const = 0;
-};
+// Root motion extrapolation
 
 // Uses distance delta between NextRootDistanceIndex and NextRootDistanceIndex - 1 and extrapolates it to ExtrapolationTime
 static float ExtrapolateAccumulatedRootDistance(
@@ -3048,7 +2827,7 @@ void FSequenceSampler::ProcessRootDistance()
 	// Build a distance lookup table by sampling root motion at a fixed rate and accumulating
 	// absolute translation deltas. During indexing we'll bsearch this table and interpolate
 	// between samples in order to convert distance offsets to time offsets.
-	// See also FIndexer::AddTrajectoryDistanceFeatures().
+	// See also FAssetIndexer::AddTrajectoryDistanceFeatures().
 
 	double TotalAccumulatedRootDistance = 0.0;
 	FTransform LastRootTransform = InitialRootTransform;
@@ -3069,7 +2848,7 @@ void FSequenceSampler::ProcessRootDistance()
 	check(SampleTime == Input.Sequence->GetPlayLength());
 
 	// Also emit root motion summary info to help with sample wrapping in 
-	// FIndexer::GetSampleTimeFromDistance() and FIndexer::GetSampleInfo()
+	// FAssetIndexer::GetSampleTimeFromDistance() and FAssetIndexer::GetSampleInfo()
 	TotalRootTransform = LastRootTransform.GetRelativeTransform(InitialRootTransform);
 	TotalRootDistance = AccumulatedRootDistance.Last();
 }
@@ -3082,7 +2861,7 @@ struct FBlendSpaceSampler : public IAssetSampler
 public:
 	struct FInput
 	{
-		const FAnimSamplingContext* SamplingContext = nullptr;
+		const FAssetSamplingContext* SamplingContext = nullptr;
 		const UBlendSpace* BlendSpace = nullptr;
 		bool bLoopable = false;
 		int32 RootDistanceSamplingRate = 60;
@@ -3481,7 +3260,7 @@ void FBlendSpaceSampler::ProcessRootDistance()
 	// Build a distance lookup table by sampling root motion at a fixed rate and accumulating
 	// absolute translation deltas. During indexing we'll bsearch this table and interpolate
 	// between samples in order to convert distance offsets to time offsets.
-	// See also FIndexer::AddTrajectoryDistanceFeatures().
+	// See also FAssetIndexer::AddTrajectoryDistanceFeatures().
 	double TotalAccumulatedRootDistance = 0.0;
 	FTransform LastRootTransform = InitialRootTransform;
 	float SampleTime = 0.0f;
@@ -3501,13 +3280,13 @@ void FBlendSpaceSampler::ProcessRootDistance()
 	check(SampleTime == PlayLength);
 
 	// Also emit root motion summary info to help with sample wrapping in 
-	// FIndexer::GetSampleTimeFromDistance() and FIndexer::GetSampleInfo()
+	// FAssetIndexer::GetSampleTimeFromDistance() and FAssetIndexer::GetSampleInfo()
 	TotalRootTransform = LastRootTransform.GetRelativeTransform(InitialRootTransform);
 	TotalRootDistance = AccumulatedRootDistance.Last();
 }
 
 //////////////////////////////////////////////////////////////////////////
-// FIndexer helpers
+// FAssetIndexer helpers
 
 struct FSamplingParam
 {
@@ -3565,24 +3344,12 @@ static FSamplingParam WrapOrClampSamplingParam(bool bCanWrap, float SamplingPara
 
 
 //////////////////////////////////////////////////////////////////////////
-// FIndexer
+// FAssetIndexer
 
-class FIndexer
+class FAssetIndexer : public IAssetIndexer
 {
 public:
 
-	struct FInput
-	{
-		const FAnimSamplingContext* SamplingContext = nullptr;
-		const UPoseSearchSchema* Schema = nullptr;
-		const IAssetSampler* MainSampler = nullptr;
-		const IAssetSampler* LeadInSampler = nullptr;
-		const IAssetSampler* FollowUpSampler = nullptr;
-		bool bMirrored = false;
-		FFloatInterval RequestedSamplingRange = FFloatInterval(0.0f, 0.0f);
-		FPoseSearchBlockTransitionParameters BlockTransitionParameters;
-	} Input;
-
 	struct FOutput
 	{
 		int32 FirstIndexedSample = 0;
@@ -3590,42 +3357,31 @@ public:
 		int32 NumIndexedPoses = 0;
 		TArray<float> FeatureVectorTable;
 		TArray<FPoseSearchPoseMetadata> PoseMetadata;
+		TSet<int32> InvalidChannels;
+		TBitArray<> AllFeaturesNotAdded;
 	} Output;
 
 	void Reset();
-	void Init(const FInput& Input);
-	void Process();
-
-private:
-	FPoseSearchFeatureVectorBuilder FeatureVector;
-	FPoseSearchPoseMetadata Metadata;
+	void Init(const FAssetIndexingContext& IndexingContext);
+	bool Process();
 
-	struct FSampleInfo
-	{
-		const IAssetSampler* Clip = nullptr;
-		FTransform RootTransform;
-		float ClipTime = 0.0f;
-		float RootDistance = 0.0f;
-		bool bClamped = false;
+public: // IAssetIndexer
 
-		bool IsValid() const { return Clip != nullptr; }
-	};
+	const FAssetIndexingContext& GetIndexingContext() const override { return IndexingContext; }
+	FSampleInfo GetSampleInfo(float SampleTime) const override;
+	FSampleInfo GetSampleInfoRelative(float SampleTime, const FSampleInfo& Origin) const override;
+	const float GetSampleTimeFromDistance(float Distance) const override;
+	FTransform MirrorTransform(const FTransform& Transform) const override;
 
-	FSampleInfo GetSampleInfo(float SampleTime) const;
-	FSampleInfo GetSampleInfoRelative(float SampleTime, const FSampleInfo& Origin) const;
-	const float GetSampleTimeFromDistance(float Distance) const;
-	FTransform MirrorTransform(const FTransform& Transform) const;
-
-	void Reserve();
-	void SampleBegin(int32 SampleIdx);
-	void SampleEnd(int32 SampleIdx);
-	void AddPoseFeatures(int32 SampleIdx);
-	void AddTrajectoryTimeFeatures(int32 SampleIdx);
-	void AddTrajectoryDistanceFeatures(int32 SampleIdx);
+private:
+	FAssetIndexingContext IndexingContext;
+	TArray<FPoseSearchFeatureVectorBuilder> FeatureVectorBuilders;
+	FPoseSearchPoseMetadata Metadata;
+	
 	void AddMetadata(int32 SampleIdx);
 };
 
-void FIndexer::Reset()
+void FAssetIndexer::Reset()
 {
 	Output.FirstIndexedSample = 0;
 	Output.LastIndexedSample = 0;
@@ -3633,70 +3389,101 @@ void FIndexer::Reset()
 
 	Output.FeatureVectorTable.Reset(0);
 	Output.PoseMetadata.Reset(0);
+	Output.InvalidChannels.Reset();
+	Output.AllFeaturesNotAdded.Reset();
 }
 
-void FIndexer::Reserve()
-{
-	Output.FeatureVectorTable.SetNumZeroed(Input.Schema->Layout.NumFloats * Output.NumIndexedPoses);
-	Output.PoseMetadata.SetNum(Output.NumIndexedPoses);
-}
-
-void FIndexer::Init(const FInput& InSettings)
+void FAssetIndexer::Init(const FAssetIndexingContext& InIndexingContext)
 {
-	check(InSettings.Schema);
-	check(InSettings.Schema->IsValid());
-	check(InSettings.MainSampler);
+	check(InIndexingContext.Schema);
+	check(InIndexingContext.Schema->IsValid());
+	check(InIndexingContext.MainSampler);
 
-	Input = InSettings;
+	IndexingContext = InIndexingContext;
 
 	Reset();
-	Output.FirstIndexedSample = FMath::FloorToInt(Input.RequestedSamplingRange.Min * Input.Schema->SampleRate);
+
+	Output.FirstIndexedSample = FMath::FloorToInt(IndexingContext.RequestedSamplingRange.Min * IndexingContext.Schema->SampleRate);
 	Output.LastIndexedSample = 
-		FMath::Max(0, FMath::CeilToInt(Input.RequestedSamplingRange.Max * Input.Schema->SampleRate));
+		FMath::Max(0, FMath::CeilToInt(IndexingContext.RequestedSamplingRange.Max * IndexingContext.Schema->SampleRate));
 	Output.NumIndexedPoses = Output.LastIndexedSample - Output.FirstIndexedSample + 1;
-	Reserve();
+	
+	Output.FeatureVectorTable.SetNumZeroed(IndexingContext.Schema->Layout.NumFloats * Output.NumIndexedPoses);
+
+	Output.PoseMetadata.SetNum(Output.NumIndexedPoses);
+
+	FeatureVectorBuilders.SetNum(Output.NumIndexedPoses);
+	for (FPoseSearchFeatureVectorBuilder& Builder : FeatureVectorBuilders)
+	{
+		Builder.Init(IndexingContext.Schema);
+	}
 }
 
-void FIndexer::Process()
+bool FAssetIndexer::Process()
 {
-	for (int32 SampleIdx = Output.FirstIndexedSample; SampleIdx <= Output.LastIndexedSample; ++SampleIdx)
-	{
-		FMemMark Mark(FMemStack::Get());
+	FMemMark Mark(FMemStack::Get());
 
-		SampleBegin(SampleIdx);
+	IndexingContext.BeginSampleIdx = Output.FirstIndexedSample;
+	IndexingContext.EndSampleIdx = Output.LastIndexedSample + 1;
+	const int32 NumSamplesInRange = IndexingContext.EndSampleIdx - IndexingContext.BeginSampleIdx;
 
-		AddPoseFeatures(SampleIdx);
-		AddTrajectoryTimeFeatures(SampleIdx);
-		AddTrajectoryDistanceFeatures(SampleIdx);
-		AddMetadata(SampleIdx);
+	FAssetIndexingOutput AssetIndexingOutput{FeatureVectorBuilders};
 
-		SampleEnd(SampleIdx);
+	// Index each channel
+	for (int32 ChannelIdx = 0; ChannelIdx != IndexingContext.Schema->Channels.Num(); ++ChannelIdx)
+	{
+		const TObjectPtr<UPoseSearchFeatureChannel>& Channel = IndexingContext.Schema->Channels[ChannelIdx];
+		Channel->IndexAsset(*this, AssetIndexingOutput);
 	}
-}
 
-void FIndexer::SampleBegin(int32 SampleIdx)
-{
-	FeatureVector.Init(Input.Schema);
-}
+	// Verify all channels provided a complete feature vector for all poses
+	TBitArray<> PoseFeaturesNotAdded;
+	for (int32 VectorIdx = 0; VectorIdx != NumSamplesInRange; ++VectorIdx)
+	{
+		if (!FeatureVectorBuilders[VectorIdx].IsComplete())
+		{
+			PoseFeaturesNotAdded = FeatureVectorBuilders[VectorIdx].GetFeaturesAdded();
+			PoseFeaturesNotAdded.BitwiseNOT();
+			TBitArray<>::BitwiseOR(Output.AllFeaturesNotAdded, PoseFeaturesNotAdded, EBitwiseOperatorFlags::MaxSize);
 
-void FIndexer::SampleEnd(int32 SampleIdx)
-{
-	check(FeatureVector.IsComplete());
+			for (TConstSetBitIterator<> FeatureIter(PoseFeaturesNotAdded); FeatureIter; ++FeatureIter)
+			{
+				int32 FeatureIdx = FeatureIter.GetIndex();
+				const FPoseSearchFeatureDesc& FeatureDesc = IndexingContext.Schema->Layout.Features[FeatureIdx];
+				Output.InvalidChannels.Add(FeatureDesc.ChannelIdx);
+			}
+		}
+	}
 
-	const int32 PoseIdx = SampleIdx - Output.FirstIndexedSample;
+	Output.InvalidChannels.Sort(TLess<>());
 
-	const int32 FirstValueIdx = PoseIdx * Input.Schema->Layout.NumFloats;
-	TArrayView<float> WriteValues = MakeArrayView(&Output.FeatureVectorTable[FirstValueIdx], Input.Schema->Layout.NumFloats);
+	// Merge spans of feature vectors into contiguous buffer
+	for (int32 VectorIdx = 0; VectorIdx != NumSamplesInRange; ++VectorIdx)
+	{
+		if (FeatureVectorBuilders[VectorIdx].IsComplete())
+		{
+			const int32 SampleIdx = VectorIdx + IndexingContext.BeginSampleIdx;
+			const int32 PoseIdx = SampleIdx - Output.FirstIndexedSample;
+			const int32 FirstValueIdx = PoseIdx * IndexingContext.Schema->Layout.NumFloats;
+			TArrayView<float> WriteValues = MakeArrayView(&Output.FeatureVectorTable[FirstValueIdx], IndexingContext.Schema->Layout.NumFloats);
+			TArrayView<const float> ReadValues = FeatureVectorBuilders[VectorIdx].GetValues();
+			check(WriteValues.Num() == ReadValues.Num());
+			FMemory::Memcpy(WriteValues.GetData(), ReadValues.GetData(), WriteValues.Num() * WriteValues.GetTypeSize());
+		}
+	}
 
-	TArrayView<const float> ReadValues = FeatureVector.GetValues();
-	
-	check(WriteValues.Num() == ReadValues.Num());
-	FMemory::Memcpy(WriteValues.GetData(), ReadValues.GetData(), WriteValues.Num() * WriteValues.GetTypeSize());
+	// Generate pose metadata
+	for (int32 SampleIdx = IndexingContext.BeginSampleIdx; SampleIdx != IndexingContext.EndSampleIdx; ++SampleIdx)
+	{
+		AddMetadata(SampleIdx);
+		const int32 PoseIdx = SampleIdx - Output.FirstIndexedSample;
+		Output.PoseMetadata[PoseIdx] = Metadata;
+	}
 
-	Output.PoseMetadata[PoseIdx] = Metadata;
+	return Output.InvalidChannels.IsEmpty();
 }
 
-const float FIndexer::GetSampleTimeFromDistance(float SampleDistance) const
+const float FAssetIndexer::GetSampleTimeFromDistance(float SampleDistance) const
 {
 	auto CanWrapDistanceSamples = [](const IAssetSampler* Sampler) -> bool
 	{
@@ -3704,19 +3491,19 @@ const float FIndexer::GetSampleTimeFromDistance(float SampleDistance) const
 		return Sampler->IsLoopable() && Sampler->GetTotalRootDistance() > SMALL_ROOT_DISTANCE;
 	};
 
-	float MainTotalDistance = Input.MainSampler->GetTotalRootDistance();
-	bool bMainCanWrap = CanWrapDistanceSamples(Input.MainSampler);
+	float MainTotalDistance = IndexingContext.MainSampler->GetTotalRootDistance();
+	bool bMainCanWrap = CanWrapDistanceSamples(IndexingContext.MainSampler);
 
 	float SampleTime = MAX_flt;
 
 	if (!bMainCanWrap)
 	{
 		// Use the lead in anim if we would have to clamp to the beginning of the main anim
-		if (Input.LeadInSampler && (SampleDistance < 0.0f))
+		if (IndexingContext.LeadInSampler && (SampleDistance < 0.0f))
 		{
-			const IAssetSampler* ClipSampler = Input.LeadInSampler;
+			const IAssetSampler* ClipSampler = IndexingContext.LeadInSampler;
 
-			bool bLeadInCanWrap = CanWrapDistanceSamples(Input.LeadInSampler);
+			bool bLeadInCanWrap = CanWrapDistanceSamples(IndexingContext.LeadInSampler);
 			float LeadRelativeDistance = SampleDistance + ClipSampler->GetTotalRootDistance();
 			FSamplingParam SamplingParam = WrapOrClampSamplingParam(bLeadInCanWrap, ClipSampler->GetTotalRootDistance(), LeadRelativeDistance);
 
@@ -3728,11 +3515,11 @@ const float FIndexer::GetSampleTimeFromDistance(float SampleDistance) const
 		}
 
 		// Use the follow up anim if we would have clamp to the end of the main anim
-		else if (Input.FollowUpSampler && (SampleDistance > MainTotalDistance))
+		else if (IndexingContext.FollowUpSampler && (SampleDistance > MainTotalDistance))
 		{
-			const IAssetSampler* ClipSampler = Input.FollowUpSampler;
+			const IAssetSampler* ClipSampler = IndexingContext.FollowUpSampler;
 
-			bool bFollowUpCanWrap = CanWrapDistanceSamples(Input.FollowUpSampler);
+			bool bFollowUpCanWrap = CanWrapDistanceSamples(IndexingContext.FollowUpSampler);
 			float FollowRelativeDistance = SampleDistance - MainTotalDistance;
 			FSamplingParam SamplingParam = WrapOrClampSamplingParam(bFollowUpCanWrap, ClipSampler->GetTotalRootDistance(), FollowRelativeDistance);
 
@@ -3740,7 +3527,7 @@ const float FIndexer::GetSampleTimeFromDistance(float SampleDistance) const
 				SamplingParam.WrappedParam + SamplingParam.Extrapolation);
 
 			// Make the follow up clip time relative to the main sequence again and unwrap
-			SampleTime = Input.MainSampler->GetPlayLength() + SamplingParam.NumCycles * ClipSampler->GetPlayLength() + ClipTime;
+			SampleTime = IndexingContext.MainSampler->GetPlayLength() + SamplingParam.NumCycles * ClipSampler->GetPlayLength() + ClipTime;
 		}
 	}
 
@@ -3753,11 +3540,11 @@ const float FIndexer::GetSampleTimeFromDistance(float SampleDistance) const
 		{
 			// In this case we're sampling a loop backwards, so MainRelativeDistance must adjust so the number of cycles 
 			// is counted correctly.
-			MainRelativeDistance += Input.MainSampler->GetTotalRootDistance();
+			MainRelativeDistance += IndexingContext.MainSampler->GetTotalRootDistance();
 		}
 
 		FSamplingParam SamplingParam = WrapOrClampSamplingParam(bMainCanWrap, MainTotalDistance, MainRelativeDistance);
-		float ClipTime = Input.MainSampler->GetTimeFromRootDistance(
+		float ClipTime = IndexingContext.MainSampler->GetTimeFromRootDistance(
 			SamplingParam.WrappedParam + SamplingParam.Extrapolation);
 
 		// Unwrap the main clip time
@@ -3765,11 +3552,11 @@ const float FIndexer::GetSampleTimeFromDistance(float SampleDistance) const
 		{
 			if (SampleDistance < 0.0f)
 			{
-				SampleTime = -((SamplingParam.NumCycles * Input.MainSampler->GetPlayLength()) + (Input.MainSampler->GetPlayLength() - ClipTime));
+				SampleTime = -((SamplingParam.NumCycles * IndexingContext.MainSampler->GetPlayLength()) + (IndexingContext.MainSampler->GetPlayLength() - ClipTime));
 			}
 			else
 			{
-				SampleTime = SamplingParam.NumCycles * Input.MainSampler->GetPlayLength() + ClipTime;
+				SampleTime = SamplingParam.NumCycles * IndexingContext.MainSampler->GetPlayLength() + ClipTime;
 			}
 		}
 		else
@@ -3781,7 +3568,7 @@ const float FIndexer::GetSampleTimeFromDistance(float SampleDistance) const
 	return SampleTime;
 }
 
-FIndexer::FSampleInfo FIndexer::GetSampleInfo(float SampleTime) const
+FAssetIndexer::FSampleInfo FAssetIndexer::GetSampleInfo(float SampleTime) const
 {
 	FSampleInfo Sample;
 
@@ -3796,28 +3583,28 @@ FIndexer::FSampleInfo FIndexer::GetSampleInfo(float SampleTime) const
 		return Sampler->IsLoopable();
 	};
 
-	float MainPlayLength = Input.MainSampler->GetPlayLength();
-	bool bMainCanWrap = CanWrapTimeSamples(Input.MainSampler);
+	float MainPlayLength = IndexingContext.MainSampler->GetPlayLength();
+	bool bMainCanWrap = CanWrapTimeSamples(IndexingContext.MainSampler);
 
 	FSamplingParam SamplingParam;
 	if (!bMainCanWrap)
 	{
 		// Use the lead in anim if we would have to clamp to the beginning of the main anim
-		if (Input.LeadInSampler && (SampleTime < 0.0f))
+		if (IndexingContext.LeadInSampler && (SampleTime < 0.0f))
 		{
-			const IAssetSampler* ClipSampler = Input.LeadInSampler;
+			const IAssetSampler* ClipSampler = IndexingContext.LeadInSampler;
 
-			bool bLeadInCanWrap = CanWrapTimeSamples(Input.LeadInSampler);
+			bool bLeadInCanWrap = CanWrapTimeSamples(IndexingContext.LeadInSampler);
 			float LeadRelativeTime = SampleTime + ClipSampler->GetPlayLength();
 			SamplingParam = WrapOrClampSamplingParam(bLeadInCanWrap, ClipSampler->GetPlayLength(), LeadRelativeTime);
 
-			Sample.Clip = Input.LeadInSampler;
+			Sample.Clip = IndexingContext.LeadInSampler;
 
 			check(SamplingParam.Extrapolation <= 0.0f);
 			if (SamplingParam.Extrapolation < 0.0f)
 			{
-				RootMotionInitial = Input.LeadInSampler->GetTotalRootTransform().Inverse();
-				RootDistanceInitial = -Input.LeadInSampler->GetTotalRootDistance();
+				RootMotionInitial = IndexingContext.LeadInSampler->GetTotalRootTransform().Inverse();
+				RootDistanceInitial = -IndexingContext.LeadInSampler->GetTotalRootDistance();
 			}
 			else
 			{
@@ -3825,26 +3612,26 @@ FIndexer::FSampleInfo FIndexer::GetSampleInfo(float SampleTime) const
 				RootDistanceInitial = 0.0f;
 			}
 
-			RootMotionLast = Input.LeadInSampler->GetTotalRootTransform();
-			RootDistanceLast = Input.LeadInSampler->GetTotalRootDistance();
+			RootMotionLast = IndexingContext.LeadInSampler->GetTotalRootTransform();
+			RootDistanceLast = IndexingContext.LeadInSampler->GetTotalRootDistance();
 		}
 
 		// Use the follow up anim if we would have clamp to the end of the main anim
-		else if (Input.FollowUpSampler && (SampleTime > MainPlayLength))
+		else if (IndexingContext.FollowUpSampler && (SampleTime > MainPlayLength))
 		{
-			const IAssetSampler* ClipSampler = Input.FollowUpSampler;
+			const IAssetSampler* ClipSampler = IndexingContext.FollowUpSampler;
 
-			bool bFollowUpCanWrap = CanWrapTimeSamples(Input.FollowUpSampler);
+			bool bFollowUpCanWrap = CanWrapTimeSamples(IndexingContext.FollowUpSampler);
 			float FollowRelativeTime = SampleTime - MainPlayLength;
 			SamplingParam = WrapOrClampSamplingParam(bFollowUpCanWrap, ClipSampler->GetPlayLength(), FollowRelativeTime);
 
-			Sample.Clip = Input.FollowUpSampler;
+			Sample.Clip = IndexingContext.FollowUpSampler;
 
-			RootMotionInitial = Input.MainSampler->GetTotalRootTransform();
-			RootDistanceInitial = Input.MainSampler->GetTotalRootDistance();
+			RootMotionInitial = IndexingContext.MainSampler->GetTotalRootTransform();
+			RootDistanceInitial = IndexingContext.MainSampler->GetTotalRootDistance();
 
-			RootMotionLast = Input.FollowUpSampler->GetTotalRootTransform();
-			RootDistanceLast = Input.FollowUpSampler->GetTotalRootDistance();
+			RootMotionLast = IndexingContext.FollowUpSampler->GetTotalRootTransform();
+			RootDistanceLast = IndexingContext.FollowUpSampler->GetTotalRootDistance();
 		}
 	}
 
@@ -3862,13 +3649,13 @@ FIndexer::FSampleInfo FIndexer::GetSampleInfo(float SampleTime) const
 
 		SamplingParam = WrapOrClampSamplingParam(bMainCanWrap, MainPlayLength, MainRelativeTime);
 
-		Sample.Clip = Input.MainSampler;
+		Sample.Clip = IndexingContext.MainSampler;
 
 		RootMotionInitial = FTransform::Identity;
 		RootDistanceInitial = 0.0f;
 
-		RootMotionLast = Input.MainSampler->GetTotalRootTransform();
-		RootDistanceLast = Input.MainSampler->GetTotalRootDistance();
+		RootMotionLast = IndexingContext.MainSampler->GetTotalRootTransform();
+		RootDistanceLast = IndexingContext.MainSampler->GetTotalRootDistance();
 	}
 
 
@@ -3926,7 +3713,7 @@ FIndexer::FSampleInfo FIndexer::GetSampleInfo(float SampleTime) const
 	return Sample;
 }
 
-FIndexer::FSampleInfo FIndexer::GetSampleInfoRelative(float SampleTime, const FSampleInfo& Origin) const
+FAssetIndexer::FSampleInfo FAssetIndexer::GetSampleInfoRelative(float SampleTime, const FSampleInfo& Origin) const
 {
 	FSampleInfo Sample = GetSampleInfo(SampleTime);
 	Sample.RootTransform.SetToRelativeTransform(Origin.RootTransform);
@@ -3934,260 +3721,22 @@ FIndexer::FSampleInfo FIndexer::GetSampleInfoRelative(float SampleTime, const FS
 	return Sample;
 }
 
-FTransform FIndexer::MirrorTransform(const FTransform& Transform) const
+FTransform FAssetIndexer::MirrorTransform(const FTransform& Transform) const
 {
-	return Input.bMirrored ? Input.SamplingContext->MirrorTransform(Transform) : Transform;
+	return IndexingContext.bMirrored ? IndexingContext.SamplingContext->MirrorTransform(Transform) : Transform;
 }
 
-void FIndexer::AddPoseFeatures(int32 SampleIdx)
+void FAssetIndexer::AddMetadata(int32 SampleIdx)
 {
-	// This function samples the instantaneous pose at time t as well as the pose's velocity and acceleration at time t.
-	// Symmetric finite differences are used to approximate derivatives:
-	//	First symmetric derivative:   f'(t) ~ (f(t+h) - f(t-h)) / 2h
-	//	Second symmetric derivative: f''(t) ~ (f(t+h) - 2f(t) + f(t-h)) / h^2
-	// Where h is a constant time delta
-	// So this means three pose extractions are taken at time t-h, t, and t+h
-	constexpr int32 NumFiniteDiffTerms = 3;
-
-	if (Input.Schema->SampledBones.IsEmpty() || Input.Schema->PoseSampleTimes.IsEmpty())
-	{
-		return;
-	}
-
-	FCompactPose Poses[NumFiniteDiffTerms];
-	FCSPose<FCompactPose> ComponentSpacePoses[NumFiniteDiffTerms];
-	FBlendedCurve UnusedCurves[NumFiniteDiffTerms];
-	UE::Anim::FStackAttributeContainer UnusedAtrributes[NumFiniteDiffTerms];
-
-	for(FCompactPose& Pose: Poses)
-	{
-		Pose.SetBoneContainer(&Input.SamplingContext->BoneContainer);
-	}
-
-	for (FBlendedCurve& Curve : UnusedCurves)
-	{
-		Curve.InitFrom(Input.SamplingContext->BoneContainer);
-	}
-
-	FAnimationPoseData AnimPoseData[NumFiniteDiffTerms] = 
-	{
-		{Poses[0], UnusedCurves[0], UnusedAtrributes[0]},
-		{Poses[1], UnusedCurves[1], UnusedAtrributes[1]},
-		{Poses[2], UnusedCurves[2], UnusedAtrributes[2]},
-	};
-
-	FPoseSearchFeatureDesc Feature;
-	Feature.Domain = EPoseSearchFeatureDomain::Time;
-
-	float SampleTime = FMath::Min(SampleIdx * Input.Schema->SamplingInterval, Input.MainSampler->GetPlayLength());
-	FSampleInfo Origin = GetSampleInfo(SampleTime);
-	
-	for (int32 SchemaSubsampleIdx = 0; SchemaSubsampleIdx != Input.Schema->PoseSampleTimes.Num(); ++SchemaSubsampleIdx)
-	{
-		Feature.SubsampleIdx = SchemaSubsampleIdx;
-
-		float SubsampleTime = SampleTime + Input.Schema->PoseSampleTimes[SchemaSubsampleIdx];
-
-		// For each pose subsample term, get the corresponding clip, accumulated root motion,
-		// and wrap the time parameter based on the clip's length.
-		FSampleInfo Samples[NumFiniteDiffTerms];
-		Samples[0] = GetSampleInfoRelative(SubsampleTime - Input.SamplingContext->FiniteDelta, Origin);
-		Samples[1] = GetSampleInfoRelative(SubsampleTime, Origin);
-		Samples[2] = GetSampleInfoRelative(SubsampleTime + Input.SamplingContext->FiniteDelta, Origin);
-
-		// Get pose samples
-		for (int32 Term = 0; Term != NumFiniteDiffTerms; ++Term)
-		{
-			float CurrentTime = Samples[Term].ClipTime;
-			float PreviousTime = CurrentTime - Input.SamplingContext->FiniteDelta;
-
-			FDeltaTimeRecord DeltaTimeRecord;
-			DeltaTimeRecord.Set(PreviousTime, CurrentTime - PreviousTime);
-			FAnimExtractContext ExtractionCtx(CurrentTime, true, DeltaTimeRecord, Samples[Term].Clip->IsLoopable());
-
-			Samples[Term].Clip->ExtractPose(ExtractionCtx, AnimPoseData[Term]);
-
-			if (Input.bMirrored)
-			{
-				FAnimationRuntime::MirrorPose(
-					AnimPoseData[Term].GetPose(),
-					Input.Schema->MirrorDataTable->MirrorAxis,
-					Input.SamplingContext->CompactPoseMirrorBones,
-					Input.SamplingContext->ComponentSpaceRefRotations
-				);
-				// Note curves and attributes are not used during the indexing process and therefore don't need to be mirrored
-			}
-
-			ComponentSpacePoses[Term].InitPose(Poses[Term]);
-		}
-
-		// Get each bone's component transform, velocity, and acceleration and add accumulated root motion at this time offset
-		// Think of this process as freezing the character in place (at SampleTime) and then tracing the paths of their joints
-		// as they move through space from past to present to future (at times indicated by PoseSampleTimes).
-		for (int32 SchemaBoneIndex = 0; SchemaBoneIndex != Input.Schema->GetNumBones(); ++SchemaBoneIndex)
-		{
-			Feature.SchemaBoneIdx = SchemaBoneIndex;
-
-			FCompactPoseBoneIndex CompactBoneIndex = 
-				Input.SamplingContext->BoneContainer.MakeCompactPoseIndex(
-					FMeshPoseBoneIndex(Input.Schema->BoneIndices[SchemaBoneIndex]));
-
-			FTransform BoneTransforms[NumFiniteDiffTerms];
-			for (int32 Term = 0; Term != NumFiniteDiffTerms; ++Term)
-			{
-				BoneTransforms[Term] = ComponentSpacePoses[Term].GetComponentSpaceTransform(CompactBoneIndex);
-				BoneTransforms[Term] = BoneTransforms[Term] * MirrorTransform(Samples[Term].RootTransform);
-			}
-
-			// Add properties to the feature vector for the pose at SampleIdx
-			FeatureVector.SetTransform(Feature, BoneTransforms[1]);
-
-			// We can get a better finite difference if we ignore samples that have
-			// been clamped at either side of the clip. However, if the central sample 
-			// itself is clamped, or there are no samples that are clamped, we can just 
-			// use the central difference as normal.
-			if (Samples[0].bClamped && !Samples[1].bClamped && !Samples[2].bClamped)
-			{
-				FeatureVector.SetTransformVelocity(Feature, BoneTransforms[2], BoneTransforms[1], Input.SamplingContext->FiniteDelta);
-			}
-			else if (Samples[2].bClamped && !Samples[1].bClamped && !Samples[0].bClamped)
-			{
-				FeatureVector.SetTransformVelocity(Feature, BoneTransforms[1], BoneTransforms[0], Input.SamplingContext->FiniteDelta);
-			}
-			else
-			{
-				FeatureVector.SetTransformVelocity(Feature, BoneTransforms[2], BoneTransforms[1], BoneTransforms[0], Input.SamplingContext->FiniteDelta);
-			}
-		}
-	}
-}
-
-void FIndexer::AddTrajectoryTimeFeatures(int32 SampleIdx)
-{
-	// This function samples the instantaneous trajectory at time t as well as the trajectory's velocity and acceleration at time t.
-	// Symmetric finite differences are used to approximate derivatives:
-	//	First symmetric derivative:   f'(t) ~ (f(t+h) - f(t-h)) / 2h
-	//	Second symmetric derivative: f''(t) ~ (f(t+h) - 2f(t) + f(t-h)) / h^2
-	// Where h is a constant time delta
-	// So this means three root motion extractions are taken at time t-h, t, and t+h
-
-	FPoseSearchFeatureDesc Feature;
-	Feature.Domain = EPoseSearchFeatureDomain::Time;
-	Feature.SchemaBoneIdx = FPoseSearchFeatureDesc::TrajectoryBoneIndex;
-
-	float SampleTime = FMath::Min(SampleIdx * Input.Schema->SamplingInterval, Input.MainSampler->GetPlayLength());
-	FSampleInfo Origin = GetSampleInfo(SampleTime);
-	
-	for (int32 SchemaSubsampleIdx = 0; SchemaSubsampleIdx != Input.Schema->TrajectorySampleTimes.Num(); ++SchemaSubsampleIdx)
-	{
-		Feature.SubsampleIdx = SchemaSubsampleIdx;
-
-		float SubsampleTime = SampleTime + Input.Schema->TrajectorySampleTimes[SchemaSubsampleIdx];
-
-		// For each pose subsample term, get the corresponding clip, accumulated root motion,
-		// and wrap the time parameter based on the clip's length.
-		FSampleInfo Samples[3];
-		Samples[0] = GetSampleInfoRelative(SubsampleTime - Input.SamplingContext->FiniteDelta, Origin);
-		Samples[1] = GetSampleInfoRelative(SubsampleTime, Origin);
-		Samples[2] = GetSampleInfoRelative(SubsampleTime + Input.SamplingContext->FiniteDelta, Origin);
-
-		// Mirror transforms if requested
-		FTransform MirroredRoots[3];
-		MirroredRoots[0] = MirrorTransform(Samples[0].RootTransform);
-		MirroredRoots[1] = MirrorTransform(Samples[1].RootTransform);
-		MirroredRoots[2] = MirrorTransform(Samples[2].RootTransform);
-
-		// Add properties to the feature vector for the pose at SampleIdx
-		FeatureVector.SetTransform(Feature, MirroredRoots[1]);
-
-		// We can get a better finite difference if we ignore samples that have
-		// been clamped at either side of the clip. However, if the central sample 
-		// itself is clamped, or there are no samples that are clamped, we can just 
-		// use the central difference as normal.
-		if (Samples[0].bClamped && !Samples[1].bClamped && !Samples[2].bClamped)
-		{
-			FeatureVector.SetTransformVelocity(Feature, MirroredRoots[2], MirroredRoots[1], Input.SamplingContext->FiniteDelta);
-		}
-		else if (Samples[2].bClamped && !Samples[1].bClamped && !Samples[0].bClamped)
-		{
-			FeatureVector.SetTransformVelocity(Feature, MirroredRoots[1], MirroredRoots[0], Input.SamplingContext->FiniteDelta);
-		}
-		else
-		{
-			FeatureVector.SetTransformVelocity(Feature, MirroredRoots[2], MirroredRoots[1], MirroredRoots[0], Input.SamplingContext->FiniteDelta);
-		}
-	}
-}
-
-void FIndexer::AddTrajectoryDistanceFeatures(int32 SampleIdx)
-{
-	// This function is very similar to AddTrajectoryTimeFeatures, but samples are taken in the distance domain
-	// instead of time domain.
-
-	FPoseSearchFeatureDesc Feature;
-	Feature.Domain = EPoseSearchFeatureDomain::Distance;
-	Feature.SchemaBoneIdx = FPoseSearchFeatureDesc::TrajectoryBoneIndex;
-
-	float SampleTime = FMath::Min(SampleIdx * Input.Schema->SamplingInterval, Input.MainSampler->GetPlayLength());
-	FSampleInfo Origin = GetSampleInfo(SampleTime);
-	
-	FQuat RootReferenceRot = Input.SamplingContext->BoneContainer.GetRefPoseTransform(FCompactPoseBoneIndex(0)).GetRotation();
-	for (int32 SchemaSubsampleIdx = 0; SchemaSubsampleIdx != Input.Schema->TrajectorySampleDistances.Num(); ++SchemaSubsampleIdx)
-	{
-		Feature.SubsampleIdx = SchemaSubsampleIdx;
-
-		// For distance based sampling of the trajectory, we first have to look up the time value
-		// we're sampling given the desired travel distance of the root for this distance offset.
-		// Once we know the time, we can then carry on just like time-based sampling.
-		const float SubsampleDistance = Origin.RootDistance + Input.Schema->TrajectorySampleDistances[SchemaSubsampleIdx];
-		float SubsampleTime = GetSampleTimeFromDistance(SubsampleDistance);
-
-		// For each pose subsample term, get the corresponding clip, accumulated root motion,
-		// and wrap the time parameter based on the clip's length.
-		FSampleInfo Samples[3];
-		Samples[0] = GetSampleInfoRelative(SubsampleTime - Input.SamplingContext->FiniteDelta, Origin);
-		Samples[1] = GetSampleInfoRelative(SubsampleTime, Origin);
-		Samples[2] = GetSampleInfoRelative(SubsampleTime + Input.SamplingContext->FiniteDelta, Origin);
-
-		// Mirror transforms if requested
-		FTransform MirroredRoots[3];
-		MirroredRoots[0] = MirrorTransform(Samples[0].RootTransform);
-		MirroredRoots[1] = MirrorTransform(Samples[1].RootTransform);
-		MirroredRoots[2] = MirrorTransform(Samples[2].RootTransform);
-
-		// Add properties to the feature vector for the pose at SampleIdx
-		FeatureVector.SetTransform(Feature, MirroredRoots[1]);
-
-		// We can get a better finite difference if we ignore samples that have
-		// been clamped at either side of the clip. However, if the central sample 
-		// itself is clamped, or there are no samples that are clamped, we can just 
-		// use the central difference as normal.
-		if (Samples[0].bClamped && !Samples[1].bClamped && !Samples[2].bClamped)
-		{
-			FeatureVector.SetTransformVelocity(Feature, MirroredRoots[2], MirroredRoots[1], Input.SamplingContext->FiniteDelta);
-		}
-		else if (Samples[2].bClamped && !Samples[1].bClamped && !Samples[0].bClamped)
-		{
-			FeatureVector.SetTransformVelocity(Feature, MirroredRoots[1], MirroredRoots[0], Input.SamplingContext->FiniteDelta);
-		}
-		else
-		{
-			FeatureVector.SetTransformVelocity(Feature, MirroredRoots[2], MirroredRoots[1], MirroredRoots[0], Input.SamplingContext->FiniteDelta);
-		}
-	}
-}
-
-void FIndexer::AddMetadata(int32 SampleIdx)
-{
-	const float SequenceLength = Input.MainSampler->GetPlayLength();
-	const float SampleTime = FMath::Min(SampleIdx * Input.Schema->SamplingInterval, SequenceLength);
+	const float SequenceLength = IndexingContext.MainSampler->GetPlayLength();
+	const float SampleTime = FMath::Min(SampleIdx * IndexingContext.Schema->SamplingInterval, SequenceLength);
 
 	Metadata = FPoseSearchPoseMetadata();
 
 	const bool bBlockTransition =
-		!Input.MainSampler->IsLoopable() &&
-		(SampleTime < Input.RequestedSamplingRange.Min + Input.BlockTransitionParameters.SequenceStartInterval ||
-		 SampleTime > Input.RequestedSamplingRange.Max - Input.BlockTransitionParameters.SequenceEndInterval);
+		!IndexingContext.MainSampler->IsLoopable() &&
+		(SampleTime < IndexingContext.RequestedSamplingRange.Min + IndexingContext.BlockTransitionParameters.SequenceStartInterval ||
+		 SampleTime > IndexingContext.RequestedSamplingRange.Max - IndexingContext.BlockTransitionParameters.SequenceEndInterval);
 
 	if (bBlockTransition)
 	{
@@ -4195,7 +3744,7 @@ void FIndexer::AddMetadata(int32 SampleIdx)
 	}
 
 	TArray<UAnimNotifyState_PoseSearchBase*> NotifyStates;
-	Input.MainSampler->ExtractPoseSearchNotifyStates(SampleTime, NotifyStates);
+	IndexingContext.MainSampler->ExtractPoseSearchNotifyStates(SampleTime, NotifyStates);
 	for (const UAnimNotifyState_PoseSearchBase* PoseSearchNotify : NotifyStates)
 	{
 		if (PoseSearchNotify->GetClass()->IsChildOf<UAnimNotifyState_PoseSearchBlockTransition>())
@@ -4214,273 +3763,19 @@ void FIndexer::AddMetadata(int32 SampleIdx)
 //////////////////////////////////////////////////////////////////////////
 // PoseSearch API
 
-static void DrawTrajectoryFeatures(const FDebugDrawParams& DrawParams, const FFeatureVectorReader& Reader, EPoseSearchFeatureDomain Domain)
-{
-	const float LifeTime = DrawParams.DefaultLifeTime;
-	const uint8 DepthPriority = ESceneDepthPriorityGroup::SDPG_Foreground + 2;
-	const bool bPersistent = EnumHasAnyFlags(DrawParams.Flags, EDebugDrawFlags::Persistent);
-
-	FPoseSearchFeatureDesc Feature;
-	Feature.Domain = Domain;
-	Feature.SchemaBoneIdx = FPoseSearchFeatureDesc::TrajectoryBoneIndex;
-
-	const int32 NumSubsamples = Domain == EPoseSearchFeatureDomain::Time ?
-		DrawParams.GetSchema()->TrajectorySampleTimes.Num() :
-		DrawParams.GetSchema()->TrajectorySampleDistances.Num();
-
-	if (NumSubsamples == 0)
-	{
-		return;
-	}
-
-	auto GetGradientColor = [](const FLinearColor& OriginalColor, int SampleIdx, int NumSamples, EDebugDrawFlags Flags)
-	{
-		int Denominator = NumSamples - 1;
-		if (Denominator <= 0 || !EnumHasAnyFlags(Flags, EDebugDrawFlags::DrawSamplesWithColorGradient))
-		{
-			return OriginalColor;
-		}
-
-		return OriginalColor * (1.0f - DrawDebugGradientStrength * (SampleIdx / (float)Denominator));
-	};
-
-	for (int32 SchemaSubsampleIdx = 0; SchemaSubsampleIdx != NumSubsamples; ++SchemaSubsampleIdx)
-	{
-		Feature.SubsampleIdx = SchemaSubsampleIdx;
-
-		FVector TrajectoryPos;
-		if (Reader.GetPosition(Feature, &TrajectoryPos))
-		{	
-			Feature.Type = EPoseSearchFeatureType::Position;
-
-			FLinearColor LinearColor = DrawParams.Color ? *DrawParams.Color : GetColorForFeature(Feature, Reader.GetLayout());
-			FLinearColor GradientColor = GetGradientColor(LinearColor, SchemaSubsampleIdx, NumSubsamples, DrawParams.Flags);
-			FColor Color = GradientColor.ToFColor(true);
-
-			TrajectoryPos = DrawParams.RootTransform.TransformPosition(TrajectoryPos);
-			if (EnumHasAnyFlags(DrawParams.Flags, EDebugDrawFlags::DrawFast | EDebugDrawFlags::DrawSearchIndex))
-			{
-				DrawDebugPoint(DrawParams.World, TrajectoryPos, DrawParams.PointSize, Color, bPersistent, DrawParams.DefaultLifeTime, DepthPriority);
-			}
-			else
-			{
-				DrawDebugSphere(DrawParams.World, TrajectoryPos, DrawDebugSphereSize, DrawDebugSphereSegments, Color, bPersistent, LifeTime, DepthPriority);
-			}
-		}
-		else
-		{
-			TrajectoryPos = DrawParams.RootTransform.GetTranslation();
-		}
-
-		FVector TrajectoryVel;
-		if (Reader.GetLinearVelocity(Feature, &TrajectoryVel))
-		{
-			Feature.Type = EPoseSearchFeatureType::LinearVelocity;
-			
-			FLinearColor LinearColor = DrawParams.Color ? *DrawParams.Color : GetColorForFeature(Feature, Reader.GetLayout());
-			FLinearColor GradientColor = GetGradientColor(LinearColor, SchemaSubsampleIdx, NumSubsamples, DrawParams.Flags);
-			FColor Color = GradientColor.ToFColor(true);
-
-			TrajectoryVel *= DrawDebugVelocityScale;
-			TrajectoryVel = DrawParams.RootTransform.TransformVector(TrajectoryVel);
-			FVector TrajectoryVelDirection = TrajectoryVel.GetSafeNormal();
-
-
-			if (EnumHasAnyFlags(DrawParams.Flags, EDebugDrawFlags::DrawSearchIndex))
-			{
-				DrawDebugPoint(DrawParams.World, TrajectoryVel, DrawParams.PointSize, Color, bPersistent, DrawParams.DefaultLifeTime, DepthPriority);
-			}
-			else
-			{
-				const float AdjustedThickness =
-					EnumHasAnyFlags(DrawParams.Flags, EDebugDrawFlags::DrawFast) ?
-					0.0f : DrawDebugLineThickness;
-
-				DrawDebugDirectionalArrow(
-					DrawParams.World,
-					TrajectoryPos + TrajectoryVelDirection * DrawDebugSphereSize,
-					TrajectoryPos + TrajectoryVel,
-					DrawDebugArrowSize,
-					Color,
-					bPersistent,
-					LifeTime,
-					DepthPriority,
-					AdjustedThickness
-				);
-			}
-		}
-
-		FVector TrajectoryForward;
-		if (Reader.GetForwardVector(Feature, &TrajectoryForward))
-		{
-			Feature.Type = EPoseSearchFeatureType::ForwardVector;
-
-			FLinearColor LinearColor = DrawParams.Color ? *DrawParams.Color : GetColorForFeature(Feature, Reader.GetLayout());
-			FLinearColor GradientColor = GetGradientColor(LinearColor, SchemaSubsampleIdx, NumSubsamples, DrawParams.Flags);
-			FColor Color = GradientColor.ToFColor(true);
-
-			TrajectoryForward = DrawParams.RootTransform.TransformVector(TrajectoryForward);
-
-			if (EnumHasAnyFlags(DrawParams.Flags, EDebugDrawFlags::DrawSearchIndex))
-			{
-				DrawDebugPoint(DrawParams.World, TrajectoryForward, DrawParams.PointSize, Color, bPersistent, DrawParams.DefaultLifeTime, DepthPriority);
-			}
-			else
-			{
-				const float AdjustedThickness =
-					EnumHasAnyFlags(DrawParams.Flags, EDebugDrawFlags::DrawFast) ?
-					0.0f : DrawDebugLineThickness;
-
-				DrawDebugDirectionalArrow(
-					DrawParams.World,
-					TrajectoryPos + TrajectoryForward * DrawDebugSphereSize,
-					TrajectoryPos + TrajectoryForward * DrawDebugSphereSize * 2.0f,
-					DrawDebugArrowSize,
-					Color,
-					bPersistent,
-					LifeTime,
-					DepthPriority,
-					AdjustedThickness
-				);
-			}
-		}
-
-		if (EnumHasAnyFlags(DrawParams.Flags, EDebugDrawFlags::DrawSampleLabels))
-		{
-			FLinearColor LinearColor = DrawParams.Color ? *DrawParams.Color : GetColorForFeature(Feature, Reader.GetLayout());
-			FLinearColor GradientColor = GetGradientColor(LinearColor, SchemaSubsampleIdx, NumSubsamples, DrawParams.Flags);
-			FColor Color = GradientColor.ToFColor(true);
-
-			FString SampleLabel;
-			if (DrawParams.LabelPrefix.IsEmpty())
-			{
-				SampleLabel = FString::Format(TEXT("{0}"), { SchemaSubsampleIdx });
-			}
-			else
-			{
-				SampleLabel = FString::Format(TEXT("{1}[{0}]"), { SchemaSubsampleIdx, DrawParams.LabelPrefix.GetData() });
-			}
-			DrawDebugString(
-				DrawParams.World,
-				TrajectoryPos + DrawDebugSampleLabelOffset,
-				SampleLabel,
-				nullptr,
-				Color,
-				LifeTime,
-				false,
-				DrawDebugSampleLabelFontScale);
-		}
-	}
-}
-
-static void DrawPoseFeatures(const FDebugDrawParams& DrawParams, const FFeatureVectorReader& Reader)
+static void DrawFeatureVector(const FDebugDrawParams& DrawParams, const FFeatureVectorReader& Reader)
 {
 	const UPoseSearchSchema* Schema = DrawParams.GetSchema();
-	check(Schema && Schema->IsValid());
-
-	const float LifeTime = DrawParams.DefaultLifeTime;
-	const uint8 DepthPriority = ESceneDepthPriorityGroup::SDPG_Foreground + 2;
-	const bool bPersistent = EnumHasAnyFlags(DrawParams.Flags, EDebugDrawFlags::Persistent);
-
-	FPoseSearchFeatureDesc Feature;
-	Feature.Domain = EPoseSearchFeatureDomain::Time;
-
-	const int32 NumSubsamples = Schema->PoseSampleTimes.Num();
-	const int32 NumBones = Schema->SampledBones.Num();
-
-	if ((NumSubsamples * NumBones) == 0)
-	{
-		return;
-	}
 
-	for (int32 SchemaSubsampleIdx = 0; SchemaSubsampleIdx != NumSubsamples; ++SchemaSubsampleIdx)
+	for (int32 ChannelIdx = 0; ChannelIdx != Schema->Channels.Num(); ++ChannelIdx)
 	{
-		Feature.SubsampleIdx = SchemaSubsampleIdx;
-		
-		for (int32 SchemaBoneIdx = 0; SchemaBoneIdx != NumBones; ++SchemaBoneIdx)
+		if (DrawParams.ChannelMask & (1 << ChannelIdx))
 		{
-			Feature.SchemaBoneIdx = SchemaBoneIdx;
-
-			FVector BonePos;
-			const bool bHaveBonePos = Reader.GetPosition(Feature, &BonePos);
-			if (bHaveBonePos)
-			{
-				Feature.Type = EPoseSearchFeatureType::Position;
-				
-				FLinearColor LinearColor = DrawParams.Color ? *DrawParams.Color : GetColorForFeature(Feature, Reader.GetLayout());
-				FColor Color = LinearColor.ToFColor(true);
-
-				BonePos = DrawParams.RootTransform.TransformPosition(BonePos);
-				if (EnumHasAnyFlags(DrawParams.Flags, EDebugDrawFlags::DrawFast | EDebugDrawFlags::DrawSearchIndex))
-				{
-					DrawDebugPoint(DrawParams.World, BonePos, DrawParams.PointSize, Color, bPersistent, DrawParams.DefaultLifeTime, DepthPriority);
-				}
-				else
-				{
-					DrawDebugSphere(DrawParams.World, BonePos, DrawDebugSphereSize, DrawDebugSphereSegments, Color, bPersistent, LifeTime, DepthPriority);
-				}
-
-				if (EnumHasAnyFlags(DrawParams.Flags, EDebugDrawFlags::DrawBoneNames))
-				{
-					DrawDebugString(
-						DrawParams.World, BonePos + FVector(0.0, 0.0, 10.0),
-						Schema->SampledBones[SchemaBoneIdx].Reference.BoneName.ToString(),
-						nullptr, Color, LifeTime, false, 1.0f);
-				}
-			}
-
-			FVector BoneVel;
-			if (bHaveBonePos && Reader.GetLinearVelocity(Feature, &BoneVel))
-			{
-				Feature.Type = EPoseSearchFeatureType::LinearVelocity;
-				
-				FLinearColor LinearColor = DrawParams.Color ? *DrawParams.Color : GetColorForFeature(Feature, Reader.GetLayout());
-				FColor Color = LinearColor.ToFColor(true);
-
-				BoneVel *= DrawDebugVelocityScale;
-				BoneVel = DrawParams.RootTransform.TransformVector(BoneVel);
-				FVector BoneVelDirection = BoneVel.GetSafeNormal();
-
-				if (EnumHasAnyFlags(DrawParams.Flags, EDebugDrawFlags::DrawSearchIndex))
-				{
-					DrawDebugPoint(DrawParams.World, BoneVel, DrawParams.PointSize, Color, bPersistent, DrawParams.DefaultLifeTime, DepthPriority);
-				}
-				else
-				{
-					const float AdjustedThickness =
-						EnumHasAnyFlags(DrawParams.Flags, EDebugDrawFlags::DrawFast) ?
-						0.0f : DrawDebugLineThickness;
-
-					DrawDebugDirectionalArrow(
-						DrawParams.World, 
-						BonePos + BoneVelDirection * DrawDebugSphereSize, 
-						BonePos + BoneVel, 
-						DrawDebugArrowSize, 
-						Color, 
-						bPersistent, 
-						LifeTime, 
-						DepthPriority, 
-						AdjustedThickness);
-				}
-			}
+			Schema->Channels[ChannelIdx]->DebugDraw(DrawParams, Reader);
 		}
 	}
 }
 
-static void DrawFeatureVector(const FDebugDrawParams& DrawParams, const FFeatureVectorReader& Reader)
-{
-	if (EnumHasAnyFlags(DrawParams.Flags, EDebugDrawFlags::IncludePose))
-	{
-		DrawPoseFeatures(DrawParams, Reader);
-	}
-
-	if (EnumHasAnyFlags(DrawParams.Flags, EDebugDrawFlags::IncludeTrajectory))
-	{
-		DrawTrajectoryFeatures(DrawParams, Reader, EPoseSearchFeatureDomain::Time);
-		DrawTrajectoryFeatures(DrawParams, Reader, EPoseSearchFeatureDomain::Distance);
-	}
-}
-
 static void DrawFeatureVector(const FDebugDrawParams& DrawParams, TArrayView<const float> PoseVector)
 {
 	const UPoseSearchSchema* Schema = DrawParams.GetSchema();
@@ -4877,7 +4172,7 @@ static void PreprocessSearchIndexSphere(FPoseSearchIndex* SearchIndex)
 		for (Eigen::Index i = 0; i < n-1; ++i)
 		{
 			Index k;
-			EigenValues.segment(i,n-i).maxCoeff(&k);
+			EigenValues.segment(i,n-i).cwiseAbs().maxCoeff(&k);
 			if (k > 0)
 			{
 				std::swap(EigenValues[i], EigenValues[k+i]);
@@ -4988,7 +4283,7 @@ bool BuildIndex(const UAnimSequence* Sequence, UPoseSearchSequenceMetaData* Sequ
 		return false;
 	}
 
-	FAnimSamplingContext SamplingContext;
+	FAssetSamplingContext SamplingContext;
 	SamplingContext.Init(SequenceMetaData->Schema);
 
 	FSequenceSampler Sampler;
@@ -4999,21 +4294,24 @@ bool BuildIndex(const UAnimSequence* Sequence, UPoseSearchSequenceMetaData* Sequ
 	Sampler.Init(SamplerInput);
 	Sampler.Process();
 
-	FIndexer Indexer;
-	FIndexer::FInput IndexerInput;
-	IndexerInput.SamplingContext = &SamplingContext;
-	IndexerInput.MainSampler = &Sampler;
-	IndexerInput.Schema = SequenceMetaData->Schema;
-	IndexerInput.RequestedSamplingRange = GetEffectiveSamplingRange(Sequence, SequenceMetaData->SamplingRange);
-	Indexer.Init(IndexerInput);
-	Indexer.Process();
+	FAssetIndexer Indexer;
+	FAssetIndexingContext IndexerContext;
+	IndexerContext.SamplingContext = &SamplingContext;
+	IndexerContext.MainSampler = &Sampler;
+	IndexerContext.Schema = SequenceMetaData->Schema;
+	IndexerContext.RequestedSamplingRange = GetEffectiveSamplingRange(Sequence, SequenceMetaData->SamplingRange);
+	Indexer.Init(IndexerContext);
+	if (!Indexer.Process())
+	{
+		return false;
+	}
 
 	SequenceMetaData->SearchIndex.Assets.Empty();
 	FPoseSearchIndexAsset SearchIndexAsset;
 	SearchIndexAsset.SourceAssetIdx = 0;
 	SearchIndexAsset.FirstPoseIdx = 0;
 	SearchIndexAsset.NumPoses = Indexer.Output.NumIndexedPoses;
-	SearchIndexAsset.SamplingInterval = IndexerInput.RequestedSamplingRange;
+	SearchIndexAsset.SamplingInterval = IndexerContext.RequestedSamplingRange;
 
 	SequenceMetaData->SearchIndex.Values = Indexer.Output.FeatureVectorTable;
 	SequenceMetaData->SearchIndex.NumPoses = Indexer.Output.NumIndexedPoses;
@@ -5026,30 +4324,31 @@ bool BuildIndex(const UAnimSequence* Sequence, UPoseSearchSequenceMetaData* Sequ
 	return true;
 }
 
-bool BuildIndex(UPoseSearchDatabase* Database, FPoseSearchIndex& OutSearchIndex)
+struct FDatabaseIndexingContext
 {
-	check(Database);
+	UPoseSearchDatabase* Database = nullptr;
+	FPoseSearchIndex* SearchIndex = nullptr;
 
-	if (!Database->IsValidForIndexing())
-	{
-		return false;
-	}
+	FAssetSamplingContext SamplingContext;
+	TArray<FSequenceSampler> SequenceSamplers;
+	TMap<const UAnimSequence*, int32> SequenceSamplerMap;
+	TArray<FBlendSpaceSampler> BlendSpaceSamplers;
+	TMap<TPair<const UBlendSpace*, FVector>, int32> BlendSpaceSamplerMap;
 
-	OutSearchIndex.Schema = Database->Schema;
+	TArray<FAssetIndexer> Indexers;
 
-	if (!Database->TryInitSearchIndexAssets(OutSearchIndex))
-	{
-		return false;
-	}
+	void PrepareSamplers();
+	void PrepareIndexers();
+	bool IndexAssets();
+	void JoinIndex();
+};
 
-	FAnimSamplingContext SamplingContext;
+void FDatabaseIndexingContext::PrepareSamplers()
+{
 	SamplingContext.Init(Database->Schema);
 
 	// Prepare samplers for all sequences
 
-	TArray<FSequenceSampler> SequenceSamplers;
-	TMap<const UAnimSequence*, int32> SequenceSamplerMap;
-
 	auto AddSequenceSampler = [&](const UAnimSequence* Sequence, bool bLoopable)
 	{
 		if (!SequenceSamplerMap.Contains(Sequence))
@@ -5083,13 +4382,14 @@ bool BuildIndex(UPoseSearchDatabase* Database, FPoseSearchIndex& OutSearchIndex)
 		}
 	}
 
-	ParallelFor(SequenceSamplers.Num(), [&SequenceSamplers](int32 SamplerIdx){ SequenceSamplers[SamplerIdx].Process(); });
+	ParallelFor(
+		SequenceSamplers.Num(),
+		[this](int32 SamplerIdx) { SequenceSamplers[SamplerIdx].Process(); },
+		ParallelForFlags
+	);
 
 	// Prepare samplers for all blend spaces
 
-	TArray<FBlendSpaceSampler> BlendSpaceSamplers;
-	TMap<TPair<const UBlendSpace*,FVector>, int32> BlendSpaceSamplerMap;
-
 	for (const FPoseSearchDatabaseBlendSpace& DbBlendSpace : Database->BlendSpaces)
 	{
 		if (DbBlendSpace.BlendSpace)
@@ -5138,11 +4438,15 @@ bool BuildIndex(UPoseSearchDatabase* Database, FPoseSearchIndex& OutSearchIndex)
 		}
 	}
 
-	ParallelFor(BlendSpaceSamplers.Num(), [&BlendSpaceSamplers](int32 SamplerIdx) { BlendSpaceSamplers[SamplerIdx].Process(); });
-
-	// Prepare sequences for indexing
+	ParallelFor(
+		BlendSpaceSamplers.Num(),
+		[this](int32 SamplerIdx) { BlendSpaceSamplers[SamplerIdx].Process(); },
+		ParallelForFlags
+	);
+}
 
-	TArray<FIndexer> Indexers;
+void FDatabaseIndexingContext::PrepareIndexers()
+{
 	Indexers.Reserve(Database->GetSearchIndex()->Assets.Num());
 
 	auto GetSequenceSampler = [&](const UAnimSequence* Sequence) -> const FSequenceSampler*
@@ -5155,52 +4459,89 @@ bool BuildIndex(UPoseSearchDatabase* Database, FPoseSearchIndex& OutSearchIndex)
 		return BlendSpace ? &BlendSpaceSamplers[BlendSpaceSamplerMap[{BlendSpace, BlendParameters}]] : nullptr;
 	};
 	Indexers.Reserve(Database->GetSearchIndex()->Assets.Num());
-	for (int32 AssetIdx = 0; AssetIdx != OutSearchIndex.Assets.Num(); ++AssetIdx)
+	for (int32 AssetIdx = 0; AssetIdx != SearchIndex->Assets.Num(); ++AssetIdx)
 	{
-		const FPoseSearchIndexAsset& SearchIndexAsset = OutSearchIndex.Assets[AssetIdx];
+		const FPoseSearchIndexAsset& SearchIndexAsset = SearchIndex->Assets[AssetIdx];
 
-		FIndexer::FInput Input;
-		Input.SamplingContext = &SamplingContext;
-
-		Input.Schema = Database->Schema;
-		Input.BlockTransitionParameters = Database->BlockTransitionParameters;
-		Input.RequestedSamplingRange = SearchIndexAsset.SamplingInterval;
-		Input.bMirrored = SearchIndexAsset.bMirrored;
+		FAssetIndexingContext IndexerContext;
+		IndexerContext.SamplingContext = &SamplingContext;
+		IndexerContext.Schema = Database->Schema;
+		IndexerContext.BlockTransitionParameters = Database->BlockTransitionParameters;
+		IndexerContext.RequestedSamplingRange = SearchIndexAsset.SamplingInterval;
+		IndexerContext.bMirrored = SearchIndexAsset.bMirrored;
 
 		if (SearchIndexAsset.Type == ESearchIndexAssetType::Sequence)
 		{
 			const FPoseSearchDatabaseSequence& DbSequence = Database->GetSequenceSourceAsset(&SearchIndexAsset);
 			const float SequenceLength = DbSequence.Sequence->GetPlayLength();
-			Input.MainSampler = GetSequenceSampler(DbSequence.Sequence);
-			Input.LeadInSampler =
+			IndexerContext.MainSampler = GetSequenceSampler(DbSequence.Sequence);
+			IndexerContext.LeadInSampler =
 				SearchIndexAsset.SamplingInterval.Min == 0.0f ? GetSequenceSampler(DbSequence.LeadInSequence) : nullptr;
-			Input.FollowUpSampler =
+			IndexerContext.FollowUpSampler =
 				SearchIndexAsset.SamplingInterval.Max == SequenceLength ? GetSequenceSampler(DbSequence.FollowUpSequence) : nullptr;
 		}
 		else if (SearchIndexAsset.Type == ESearchIndexAssetType::BlendSpace)
 		{
 			const FPoseSearchDatabaseBlendSpace& DbBlendSpace = Database->GetBlendSpaceSourceAsset(&SearchIndexAsset);
-			Input.MainSampler = GetBlendSpaceSampler(DbBlendSpace.BlendSpace, SearchIndexAsset.BlendParameters);
+			IndexerContext.MainSampler = GetBlendSpaceSampler(DbBlendSpace.BlendSpace, SearchIndexAsset.BlendParameters);
 		}
 		else
 		{
 			checkNoEntry();
 		}
 
-		FIndexer& Indexer = Indexers.AddDefaulted_GetRef();
-		Indexer.Init(Input);
+		FAssetIndexer& Indexer = Indexers.AddDefaulted_GetRef();
+		Indexer.Init(IndexerContext);
+	}
+}
+
+bool FDatabaseIndexingContext::IndexAssets()
+{
+	bool bSuccess = true;
+
+	// Index sequence data
+	ParallelFor(
+		Indexers.Num(),
+		[this](int32 SequenceIdx) { Indexers[SequenceIdx].Process(); },
+		ParallelForFlags
+	);
+
+	// Validate indexing process
+	TSet<int32> AllInvalidChannels;
+	TBitArray<> AllInvalidFeatures;
+	for (FAssetIndexer& Indexer : Indexers)
+	{
+		AllInvalidChannels.Union(Indexer.Output.InvalidChannels);
+		TBitArray<>::BitwiseOR(AllInvalidFeatures, Indexer.Output.AllFeaturesNotAdded, EBitwiseOperatorFlags::MaxSize);
+	}
+
+	if (!AllInvalidChannels.IsEmpty())
+	{
+		bSuccess = false;
 	}
 
-	ParallelFor(Indexers.Num(), [&Indexers](int32 SequenceIdx){ Indexers[SequenceIdx].Process(); });
+	for (int32 InvalidChannelIdx : AllInvalidChannels)
+	{
+		const TObjectPtr<UPoseSearchFeatureChannel> ChannelPtr = Database->Schema->Channels[InvalidChannelIdx];
+		if (ChannelPtr)
+		{
+			UE_LOG(LogPoseSearch, Error, TEXT("Feature channel '%f' failed to index all features from all assets"), *ChannelPtr->GetName());
+		}
+	}
+
+	return bSuccess;
+}
 
+void FDatabaseIndexingContext::JoinIndex()
+{
 	// Write index info to sequence and count up total poses and storage required
 	int32 TotalPoses = 0;
 	int32 TotalFloats = 0;
-	for (int32 AssetIdx = 0; AssetIdx != OutSearchIndex.Assets.Num(); ++AssetIdx)
+	for (int32 AssetIdx = 0; AssetIdx != SearchIndex->Assets.Num(); ++AssetIdx)
 	{
-		const FIndexer::FOutput& Output = Indexers[AssetIdx].Output;
+		const FAssetIndexer::FOutput& Output = Indexers[AssetIdx].Output;
 
-		FPoseSearchIndexAsset& SearchIndexAsset = OutSearchIndex.Assets[AssetIdx];
+		FPoseSearchIndexAsset& SearchIndexAsset = SearchIndex->Assets[AssetIdx];
 		SearchIndexAsset.NumPoses = Output.NumIndexedPoses;
 		SearchIndexAsset.FirstPoseIdx = TotalPoses;
 
@@ -5209,24 +4550,74 @@ bool BuildIndex(UPoseSearchDatabase* Database, FPoseSearchIndex& OutSearchIndex)
 	}
 
 	// Join animation data into a single search index
-	OutSearchIndex.Values.Reset(TotalFloats);
-	OutSearchIndex.PoseMetadata.Reset(TotalPoses);
+	SearchIndex->Values.Reset(TotalFloats);
+	SearchIndex->PoseMetadata.Reset(TotalPoses);
+
+	for (const FAssetIndexer& Indexer : Indexers)
+	{
+		const FAssetIndexer::FOutput& Output = Indexer.Output;
+		SearchIndex->Values.Append(Output.FeatureVectorTable.GetData(), Output.FeatureVectorTable.Num());
+		SearchIndex->PoseMetadata.Append(Output.PoseMetadata);
+	}
+
+	SearchIndex->NumPoses = TotalPoses;
+	SearchIndex->Schema = Database->Schema;
+}
+
+bool BuildIndex(UPoseSearchDatabase* Database, FPoseSearchIndex& OutSearchIndex)
+{
+	check(Database);
+
+	bool bSuccess = true;
 
-	for (const FIndexer& Indexer : Indexers)
+	if (!Database->IsValidForIndexing())
 	{
-		const FIndexer::FOutput& Output = Indexer.Output;
-		OutSearchIndex.Values.Append(Output.FeatureVectorTable.GetData(), Output.FeatureVectorTable.Num());
-		OutSearchIndex.PoseMetadata.Append(Output.PoseMetadata);
+		UE_LOG(LogPoseSearch, Error, TEXT("Database '%f' is invalid for indexing"), *Database->GetName());
+		return false;
 	}
 
-	OutSearchIndex.NumPoses = TotalPoses;
 	OutSearchIndex.Schema = Database->Schema;
 
+	if (!Database->TryInitSearchIndexAssets(OutSearchIndex))
+	{
+		return false;
+	}
+
+	FDatabaseIndexingContext DbIndexingContext;
+	DbIndexingContext.Database = Database;
+	DbIndexingContext.SearchIndex = &OutSearchIndex;
+
+	DbIndexingContext.PrepareSamplers();
+	DbIndexingContext.PrepareIndexers();
+	bSuccess &= DbIndexingContext.IndexAssets();
+	DbIndexingContext.JoinIndex();
+
 	PreprocessSearchIndex(&OutSearchIndex);
 
-	return true;
+	return bSuccess;
 }
 
+bool FQueryBuildingContext::IsInitialized() const
+{
+	return Schema && Schema->IsValid() && Query.IsInitializedForSchema(Schema);
+}
+
+bool BuildQuery(FQueryBuildingContext& QueryBuildingContext)
+{
+	if (!QueryBuildingContext.IsInitialized())
+	{
+		return false;
+	}
+
+	bool bSuccess = true;
+	for (const TObjectPtr<UPoseSearchFeatureChannel>& Channel : QueryBuildingContext.Schema->Channels)
+	{
+		bool bChannelSuccess = Channel->BuildQuery(QueryBuildingContext);
+		bSuccess &= bChannelSuccess;
+	}
+
+	return bSuccess;
+}
 
 FSearchResult Search(FSearchContext& SearchContext)
 {
@@ -5497,9 +4888,14 @@ UE::Anim::IPoseSearchProvider::FSearchResult FModule::Search(const FAnimationBas
 
 	FPoseHistory& PoseHistory = PoseHistoryProvider->GetPoseHistory();
 	FPoseSearchFeatureVectorBuilder& QueryBuilder = PoseHistory.GetQueryBuilder();
-
 	QueryBuilder.Init(MetaData->Schema);
-	if (!QueryBuilder.TrySetPoseFeatures(&PoseHistory, GraphContext.AnimInstanceProxy->GetRequiredBones()))
+
+	UE::PoseSearch::FQueryBuildingContext QueryBuildingContext(QueryBuilder);
+	QueryBuildingContext.Schema = MetaData->Schema;
+	QueryBuildingContext.History = &PoseHistory;
+	QueryBuildingContext.Trajectory = nullptr;
+
+	if (!UE::PoseSearch::BuildQuery(QueryBuildingContext))
 	{
 		return ProviderResult;
 	}
diff --git a/Engine/Plugins/Experimental/Animation/PoseSearch/Source/Runtime/Private/PoseSearchDerivedData.cpp b/Engine/Plugins/Experimental/Animation/PoseSearch/Source/Runtime/Private/PoseSearchDerivedData.cpp
--- a/Engine/Plugins/Experimental/Animation/PoseSearch/Source/Runtime/Private/PoseSearchDerivedData.cpp
+++ b/Engine/Plugins/Experimental/Animation/PoseSearch/Source/Runtime/Private/PoseSearchDerivedData.cpp
@@ -6,6 +6,8 @@
 #include "Animation/AnimSequence.h"
 #include "Animation/BlendSpace.h"
 
+#include "UObject/Class.h"
+
 #if WITH_EDITOR
 #include "Serialization/BulkDataRegistry.h"
 #endif // WITH_EDITOR
@@ -213,13 +215,16 @@ namespace UE::PoseSearch
 	{
 		if (IsValid(Schema))
 		{
-			InOutWriter.Update(&Schema->bUseTrajectoryVelocities, sizeof(Schema->bUseTrajectoryVelocities));
-			InOutWriter.Update(&Schema->bUseTrajectoryPositions, sizeof(Schema->bUseTrajectoryPositions));
-			InOutWriter.Update(&Schema->bUseTrajectoryForwardVectors, sizeof(Schema->bUseTrajectoryForwardVectors));
-			InOutWriter.Update(MakeMemoryView(Schema->SampledBones));
-			InOutWriter.Update(MakeMemoryView(Schema->PoseSampleTimes));
-			InOutWriter.Update(MakeMemoryView(Schema->TrajectorySampleTimes));
-			InOutWriter.Update(MakeMemoryView(Schema->TrajectorySampleDistances));
+			for (const TObjectPtr<UPoseSearchFeatureChannel>& Channel : Schema->Channels)
+			{
+				if (Channel)
+				{
+					const FBlake3Hash& ChannelClassHash = Channel->GetClass()->GetSchemaHash(false);
+					InOutWriter.Update(MakeMemoryView(ChannelClassHash.GetBytes()));
+
+					Channel->GenerateDDCKey(InOutWriter);
+				}
+			}
 			InOutWriter.Update(&Schema->DataPreprocessor, sizeof(Schema->DataPreprocessor));
 			InOutWriter.Update(&Schema->EffectiveDataPreprocessor, sizeof(Schema->EffectiveDataPreprocessor));
 			InOutWriter.Update(&Schema->SamplingInterval, sizeof(Schema->SamplingInterval));
diff --git a/Engine/Plugins/Experimental/Animation/PoseSearch/Source/Runtime/Private/PoseSearchDerivedData.h b/Engine/Plugins/Experimental/Animation/PoseSearch/Source/Runtime/Private/PoseSearchDerivedData.h
--- a/Engine/Plugins/Experimental/Animation/PoseSearch/Source/Runtime/Private/PoseSearchDerivedData.h
+++ b/Engine/Plugins/Experimental/Animation/PoseSearch/Source/Runtime/Private/PoseSearchDerivedData.h
@@ -24,6 +24,7 @@ struct FPoseSearchDatabaseDerivedData
 {
 	GENERATED_BODY()
 
+	UPROPERTY()
 	FPoseSearchIndex SearchIndex;
 
 #if WITH_EDITORONLY_DATA
diff --git a/Engine/Plugins/Experimental/Animation/PoseSearch/Source/Runtime/Private/PoseSearchFeatureChannels.cpp b/Engine/Plugins/Experimental/Animation/PoseSearch/Source/Runtime/Private/PoseSearchFeatureChannels.cpp
new file mode 100644
--- /dev/null
+++ b/Engine/Plugins/Experimental/Animation/PoseSearch/Source/Runtime/Private/PoseSearchFeatureChannels.cpp
@@ -0,0 +1,842 @@
+// Copyright Epic Games, Inc. All Rights Reserved.
+
+#include "PoseSearch/PoseSearchFeatureChannels.h"
+#include "PoseSearch/PoseSearch.h"
+#include "AnimationRuntime.h"
+#include "Animation/AnimationPoseData.h"
+#include "Animation/AttributesRuntime.h"
+#include "Animation/MirrorDataTable.h"
+#include "DrawDebugHelpers.h"
+#include "UObject/ObjectSaveContext.h"
+
+namespace UE::PoseSearch {
+
+//////////////////////////////////////////////////////////////////////////
+// Constants
+
+constexpr float DrawDebugLineThickness = 2.0f;
+constexpr float DrawDebugPointSize = 3.0f;
+constexpr float DrawDebugVelocityScale = 0.08f;
+constexpr float DrawDebugArrowSize = 30.0f;
+constexpr float DrawDebugSphereSize = 3.0f;
+constexpr int32 DrawDebugSphereSegments = 10;
+constexpr float DrawDebugGradientStrength = 0.8f;
+constexpr float DrawDebugSampleLabelFontScale = 1.0f;
+static const FVector DrawDebugSampleLabelOffset = FVector(0.0f, 0.0f, -10.0f);
+
+
+//////////////////////////////////////////////////////////////////////////
+// Drawing helpers
+
+static FLinearColor GetColorForFeature(FPoseSearchFeatureDesc Feature, const FPoseSearchFeatureVectorLayout* Layout)
+{
+	const float FeatureIdx = Layout->Features.IndexOfByKey(Feature);
+	const float FeatureCountIdx = Layout->Features.Num() - 1;
+	const float FeatureCountIdxHalf = FeatureCountIdx / 2.f;
+	check(FeatureIdx != INDEX_NONE);
+
+	const float Hue = FeatureIdx < FeatureCountIdxHalf
+		? FMath::GetMappedRangeValueUnclamped({ 0.f, FeatureCountIdxHalf }, FVector2f(60.f, 0.f), FeatureIdx)
+		: FMath::GetMappedRangeValueUnclamped({ FeatureCountIdxHalf, FeatureCountIdx }, FVector2f(280.f, 220.f), FeatureIdx);
+
+	const FLinearColor ColorHSV(Hue, 1.f, 1.f);
+	return ColorHSV.HSVToLinearRGB();
+}
+
+} // namespace UE::PoseSearch
+
+
+
+//////////////////////////////////////////////////////////////////////////
+// UPoseSearchFeatureChannel
+class USkeleton* UPoseSearchFeatureChannel::GetSkeleton(bool& bInvalidSkeletonIsError, const IPropertyHandle* PropertyHandle)
+{
+	bInvalidSkeletonIsError = false;
+
+	const UPoseSearchSchema* Schema = Cast<UPoseSearchSchema>(GetOuter());
+	return Schema ? Schema->Skeleton : nullptr;
+}
+
+
+
+//////////////////////////////////////////////////////////////////////////
+// UPoseSearchFeatureChannel_Pose
+
+void UPoseSearchFeatureChannel_Pose::PreSave(FObjectPreSaveContext ObjectSaveContext)
+{
+	SampleTimes.Sort(TLess<>());
+	Super::PreSave(ObjectSaveContext);
+}
+
+void UPoseSearchFeatureChannel_Pose::InitializeSchema(UE::PoseSearch::FSchemaInitializer& Initializer)
+{
+	auto AddFeatures = [this, &Initializer](EPoseSearchFeatureType Type)
+	{
+		FPoseSearchFeatureDesc Feature;
+		Feature.Type = Type;
+		Feature.ChannelIdx = GetChannelIndex();
+
+		const int32 NumBones = SampledBones.Num();
+		for (int32 SubsampleIdx = 0; SubsampleIdx != SampleTimes.Num(); ++SubsampleIdx)
+		{
+			Feature.SubsampleIdx = SubsampleIdx;
+
+			for (int32 ChannelBoneIdx = 0; ChannelBoneIdx != NumBones; ++ChannelBoneIdx)
+			{
+				Feature.ChannelFeatureId = ChannelBoneIdx;
+
+				const FPoseSearchBone& SampledBone = SampledBones[ChannelBoneIdx];
+				bool bAddFeature = (SampledBone.GetTypeMask() & (1 << static_cast<int>(Type))) != 0;
+				if (bAddFeature)
+				{
+					Initializer.AddFeatureDesc(Feature);
+				}
+			}
+		}
+	};
+
+	for (int32 FeatureType = 0; FeatureType != (int32)EPoseSearchFeatureType::Num; ++FeatureType)
+	{
+		AddFeatures(EPoseSearchFeatureType(FeatureType));
+	}
+
+	FeatureParams.Reset();
+	for (const FPoseSearchBone& Bone: SampledBones)
+	{
+		FPoseSearchPoseFeatureInfo FeatureInfo;
+		FeatureInfo.SchemaBoneIdx = Initializer.AddBoneReference(Bone.Reference);
+		FeatureParams.Add(FeatureInfo);
+	}
+}
+
+void UPoseSearchFeatureChannel_Pose::IndexAsset(const UE::PoseSearch::IAssetIndexer& Indexer,  UE::PoseSearch::FAssetIndexingOutput& IndexingOutput) const
+{
+	using namespace UE::PoseSearch;
+
+	const FAssetIndexingContext& IndexingContext = Indexer.GetIndexingContext();
+	for (int32 SampleIdx = IndexingContext.BeginSampleIdx; SampleIdx != IndexingContext.EndSampleIdx; ++SampleIdx)
+	{
+		int32 VectorIdx = SampleIdx - IndexingContext.BeginSampleIdx;
+		FPoseSearchFeatureVectorBuilder& FeatureVector = IndexingOutput.PoseVectors[VectorIdx];
+		AddPoseFeatures(Indexer, SampleIdx, FeatureVector);
+	}
+}
+
+void UPoseSearchFeatureChannel_Pose::AddPoseFeatures(const  UE::PoseSearch::IAssetIndexer& Indexer, int32 SampleIdx, FPoseSearchFeatureVectorBuilder& FeatureVector) const
+{
+	// This function samples the instantaneous pose at time t as well as the pose's velocity and acceleration at time t.
+	// Symmetric finite differences are used to approximate derivatives:
+	//	First symmetric derivative:   f'(t) ~ (f(t+h) - f(t-h)) / 2h
+	//	Second symmetric derivative: f''(t) ~ (f(t+h) - 2f(t) + f(t-h)) / h^2
+	// Where h is a constant time delta
+	// So this means three pose extractions are taken at time t-h, t, and t+h
+	constexpr int32 NumFiniteDiffTerms = 3;
+
+	using namespace UE::PoseSearch;
+	using FSampleInfo = IAssetIndexer::FSampleInfo;
+
+	if (SampledBones.IsEmpty() || SampleTimes.IsEmpty())
+	{
+		return;
+	}
+
+	const FAssetIndexingContext& IndexingContext = Indexer.GetIndexingContext();
+	const FAssetSamplingContext* SamplingContext = IndexingContext.SamplingContext;
+
+	FCompactPose Poses[NumFiniteDiffTerms];
+	FCSPose<FCompactPose> ComponentSpacePoses[NumFiniteDiffTerms];
+	FBlendedCurve UnusedCurves[NumFiniteDiffTerms];
+	UE::Anim::FStackAttributeContainer UnusedAtrributes[NumFiniteDiffTerms];
+
+	for (FCompactPose& Pose : Poses)
+	{
+		Pose.SetBoneContainer(&SamplingContext->BoneContainer);
+	}
+
+	for (FBlendedCurve& Curve : UnusedCurves)
+	{
+		Curve.InitFrom(SamplingContext->BoneContainer);
+	}
+
+	FAnimationPoseData AnimPoseData[NumFiniteDiffTerms] =
+	{
+		{Poses[0], UnusedCurves[0], UnusedAtrributes[0]},
+		{Poses[1], UnusedCurves[1], UnusedAtrributes[1]},
+		{Poses[2], UnusedCurves[2], UnusedAtrributes[2]},
+	};
+
+	FPoseSearchFeatureDesc Feature;
+	Feature.ChannelIdx = GetChannelIndex();
+
+	float SampleTime = FMath::Min(SampleIdx * IndexingContext.Schema->SamplingInterval, IndexingContext.MainSampler->GetPlayLength());
+	FSampleInfo Origin = Indexer.GetSampleInfo(SampleTime);
+
+	for (int32 SubsampleIdx = 0; SubsampleIdx != SampleTimes.Num(); ++SubsampleIdx)
+	{
+		Feature.SubsampleIdx = SubsampleIdx;
+
+		float SubsampleTime = SampleTime + SampleTimes[SubsampleIdx];
+
+		// For each pose subsample term, get the corresponding clip, accumulated root motion,
+		// and wrap the time parameter based on the clip's length.
+		FSampleInfo Samples[NumFiniteDiffTerms];
+		Samples[0] = Indexer.GetSampleInfoRelative(SubsampleTime - SamplingContext->FiniteDelta, Origin);
+		Samples[1] = Indexer.GetSampleInfoRelative(SubsampleTime, Origin);
+		Samples[2] = Indexer.GetSampleInfoRelative(SubsampleTime + SamplingContext->FiniteDelta, Origin);
+
+		// Get pose samples
+		for (int32 Term = 0; Term != NumFiniteDiffTerms; ++Term)
+		{
+			float CurrentTime = Samples[Term].ClipTime;
+			float PreviousTime = CurrentTime - SamplingContext->FiniteDelta;
+
+			FDeltaTimeRecord DeltaTimeRecord;
+			DeltaTimeRecord.Set(PreviousTime, CurrentTime - PreviousTime);
+			FAnimExtractContext ExtractionCtx(CurrentTime, true, DeltaTimeRecord, Samples[Term].Clip->IsLoopable());
+
+			Samples[Term].Clip->ExtractPose(ExtractionCtx, AnimPoseData[Term]);
+
+			if (IndexingContext.bMirrored)
+			{
+				FAnimationRuntime::MirrorPose(
+					AnimPoseData[Term].GetPose(),
+					IndexingContext.Schema->MirrorDataTable->MirrorAxis,
+					SamplingContext->CompactPoseMirrorBones,
+					SamplingContext->ComponentSpaceRefRotations
+				);
+				// Note curves and attributes are not used during the indexing process and therefore don't need to be mirrored
+			}
+
+			ComponentSpacePoses[Term].InitPose(Poses[Term]);
+		}
+
+		// Get each bone's component transform, velocity, and acceleration and add accumulated root motion at this time offset
+		// Think of this process as freezing the character in place (at SampleTime) and then tracing the paths of their joints
+		// as they move through space from past to present to future (at times indicated by PoseSampleTimes).
+		for (int32 ChannelBoneIdx = 0; ChannelBoneIdx != SampledBones.Num(); ++ChannelBoneIdx)
+		{
+			const FPoseSearchPoseFeatureInfo& PoseFeatureInfo = FeatureParams[ChannelBoneIdx];
+			const FBoneReference& BoneReference = IndexingContext.Schema->BoneReferences[PoseFeatureInfo.SchemaBoneIdx];
+
+			Feature.ChannelFeatureId = ChannelBoneIdx;
+
+			FCompactPoseBoneIndex CompactBoneIndex = SamplingContext->BoneContainer.MakeCompactPoseIndex(FMeshPoseBoneIndex(BoneReference.BoneIndex));
+
+			FTransform BoneTransforms[NumFiniteDiffTerms];
+			for (int32 Term = 0; Term != NumFiniteDiffTerms; ++Term)
+			{
+				BoneTransforms[Term] = ComponentSpacePoses[Term].GetComponentSpaceTransform(CompactBoneIndex);
+				BoneTransforms[Term] = BoneTransforms[Term] * Indexer.MirrorTransform(Samples[Term].RootTransform);
+			}
+
+			// Add properties to the feature vector for the pose at SampleIdx
+			FeatureVector.SetTransform(Feature, BoneTransforms[1]);
+
+			// We can get a better finite difference if we ignore samples that have
+			// been clamped at either side of the clip. However, if the central sample 
+			// itself is clamped, or there are no samples that are clamped, we can just 
+			// use the central difference as normal.
+			if (Samples[0].bClamped && !Samples[1].bClamped && !Samples[2].bClamped)
+			{
+				FeatureVector.SetTransformVelocity(Feature, BoneTransforms[2], BoneTransforms[1], SamplingContext->FiniteDelta);
+			}
+			else if (Samples[2].bClamped && !Samples[1].bClamped && !Samples[0].bClamped)
+			{
+				FeatureVector.SetTransformVelocity(Feature, BoneTransforms[1], BoneTransforms[0], SamplingContext->FiniteDelta);
+			}
+			else
+			{
+				FeatureVector.SetTransformVelocity(Feature, BoneTransforms[2], BoneTransforms[1], BoneTransforms[0], SamplingContext->FiniteDelta);
+			}
+		}
+	}
+}
+
+FFloatRange UPoseSearchFeatureChannel_Pose::GetHorizonRange(EPoseSearchFeatureDomain Domain) const
+{
+	FFloatRange Extent = FFloatRange::Empty();
+	if (Domain == EPoseSearchFeatureDomain::Time)
+	{
+		if (SampleTimes.Num())
+		{
+			Extent = FFloatRange::Inclusive(SampleTimes[0], SampleTimes.Last());
+		}
+	}
+
+	return Extent;
+}
+
+void UPoseSearchFeatureChannel_Pose::GenerateDDCKey(FBlake3& InOutKeyHasher) const
+{
+	InOutKeyHasher.Update(MakeMemoryView(SampledBones));
+	InOutKeyHasher.Update(MakeMemoryView(SampleTimes));
+}
+
+bool UPoseSearchFeatureChannel_Pose::BuildQuery(UE::PoseSearch::FQueryBuildingContext& Context) const
+{
+	check(Context.Schema);
+	
+	if (!Context.History)
+	{
+		return false;
+	}
+
+	FPoseSearchFeatureDesc Feature;
+	Feature.ChannelIdx = GetChannelIndex();
+
+	for (int32 SubsampleIdx = 0; SubsampleIdx != SampleTimes.Num(); ++SubsampleIdx)
+	{
+		Feature.SubsampleIdx = SubsampleIdx;
+
+		// Stop when we've reached future samples
+		float SampleTime = SampleTimes[SubsampleIdx];
+		if (SampleTime > 0.0f)
+		{
+			break;
+		}
+
+		float SecondsAgo = -SampleTime;
+		if (!Context.History->TrySamplePose(SecondsAgo, Context.Schema->Skeleton->GetReferenceSkeleton(), Context.Schema->BoneIndicesWithParents))
+		{
+			return false;
+		}
+
+		TArrayView<const FTransform> ComponentPose = Context.History->GetComponentPoseSample();
+		TArrayView<const FTransform> ComponentPrevPose = Context.History->GetPrevComponentPoseSample();
+		FTransform RootTransform = Context.History->GetRootTransformSample();
+		FTransform RootTransformPrev = Context.History->GetPrevRootTransformSample();
+		for (int32 SampledBoneIdx = 0; SampledBoneIdx != SampledBones.Num(); ++SampledBoneIdx)
+		{
+			Feature.ChannelFeatureId = SampledBoneIdx;
+
+			int32 SchemaBoneIdx = FeatureParams[SampledBoneIdx].SchemaBoneIdx;
+
+			int32 SkeletonBoneIndex = Context.Schema->BoneIndices[SchemaBoneIdx];
+
+			const FTransform& Transform = ComponentPose[SkeletonBoneIndex];
+			const FTransform& PrevTransform = ComponentPrevPose[SkeletonBoneIndex] * (RootTransformPrev * RootTransform.Inverse());
+			Context.Query.SetTransform(Feature, Transform);
+			Context.Query.SetTransformVelocity(Feature, Transform, PrevTransform, Context.History->GetSampleTimeInterval());
+		}
+	}
+
+	return true;
+}
+
+void UPoseSearchFeatureChannel_Pose::DebugDraw(const UE::PoseSearch::FDebugDrawParams& DrawParams, const  UE::PoseSearch::FFeatureVectorReader& Reader) const
+{
+	using namespace UE::PoseSearch;
+
+	const UPoseSearchSchema* Schema = DrawParams.GetSchema();
+	check(Schema && Schema->IsValid());
+
+	const float LifeTime = DrawParams.DefaultLifeTime;
+	const uint8 DepthPriority = ESceneDepthPriorityGroup::SDPG_Foreground + 2;
+	const bool bPersistent = EnumHasAnyFlags(DrawParams.Flags, EDebugDrawFlags::Persistent);
+
+	FPoseSearchFeatureDesc Feature;
+	Feature.ChannelIdx = GetChannelIndex();
+
+	const int32 NumSubsamples = SampleTimes.Num();
+	const int32 NumBones = SampledBones.Num();
+
+	if ((NumSubsamples * NumBones) == 0)
+	{
+		return;
+	}
+
+	for (int32 SubsampleIdx = 0; SubsampleIdx != NumSubsamples; ++SubsampleIdx)
+	{
+		Feature.SubsampleIdx = SubsampleIdx;
+
+		for (int32 ChannelBoneIdx = 0; ChannelBoneIdx != NumBones; ++ChannelBoneIdx)
+		{
+			Feature.ChannelFeatureId = ChannelBoneIdx;
+
+			FVector BonePos;
+			const bool bHaveBonePos = Reader.GetPosition(Feature, &BonePos);
+			if (bHaveBonePos)
+			{
+				Feature.Type = EPoseSearchFeatureType::Position;
+
+				FLinearColor LinearColor = DrawParams.Color ? *DrawParams.Color : GetColorForFeature(Feature, Reader.GetLayout());
+				FColor Color = LinearColor.ToFColor(true);
+
+				BonePos = DrawParams.RootTransform.TransformPosition(BonePos);
+				if (EnumHasAnyFlags(DrawParams.Flags, EDebugDrawFlags::DrawFast | EDebugDrawFlags::DrawSearchIndex))
+				{
+					DrawDebugPoint(DrawParams.World, BonePos, DrawParams.PointSize, Color, bPersistent, DrawParams.DefaultLifeTime, DepthPriority);
+				}
+				else
+				{
+					DrawDebugSphere(DrawParams.World, BonePos, DrawDebugSphereSize, DrawDebugSphereSegments, Color, bPersistent, LifeTime, DepthPriority);
+				}
+
+				if (EnumHasAnyFlags(DrawParams.Flags, EDebugDrawFlags::DrawBoneNames))
+				{
+					int32 SchemaBoneIdx = this->FeatureParams[ChannelBoneIdx].SchemaBoneIdx;
+					DrawDebugString(
+						DrawParams.World, BonePos + FVector(0.0, 0.0, 10.0),
+						Schema->BoneReferences[SchemaBoneIdx].BoneName.ToString(),
+						nullptr, Color, LifeTime, false, 1.0f);
+				}
+			}
+
+			FVector BoneVel;
+			if (bHaveBonePos && Reader.GetLinearVelocity(Feature, &BoneVel))
+			{
+				Feature.Type = EPoseSearchFeatureType::LinearVelocity;
+
+				FLinearColor LinearColor = DrawParams.Color ? *DrawParams.Color : GetColorForFeature(Feature, Reader.GetLayout());
+				FColor Color = LinearColor.ToFColor(true);
+
+				BoneVel *= DrawDebugVelocityScale;
+				BoneVel = DrawParams.RootTransform.TransformVector(BoneVel);
+				FVector BoneVelDirection = BoneVel.GetSafeNormal();
+
+				if (EnumHasAnyFlags(DrawParams.Flags, EDebugDrawFlags::DrawSearchIndex))
+				{
+					DrawDebugPoint(DrawParams.World, BoneVel, DrawParams.PointSize, Color, bPersistent, DrawParams.DefaultLifeTime, DepthPriority);
+				}
+				else
+				{
+					const float AdjustedThickness =
+						EnumHasAnyFlags(DrawParams.Flags, EDebugDrawFlags::DrawFast) ?
+						0.0f : DrawDebugLineThickness;
+
+					DrawDebugDirectionalArrow(
+						DrawParams.World,
+						BonePos + BoneVelDirection * DrawDebugSphereSize,
+						BonePos + BoneVel,
+						DrawDebugArrowSize,
+						Color,
+						bPersistent,
+						LifeTime,
+						DepthPriority,
+						AdjustedThickness);
+				}
+			}
+		}
+	}
+}
+
+
+
+//////////////////////////////////////////////////////////////////////////
+// UPoseSearchFeatureChannel_Trajectory
+void UPoseSearchFeatureChannel_Trajectory::PreSave(FObjectPreSaveContext ObjectSaveContext)
+{
+	SampleOffsets.Sort(TLess<>());
+
+	Super::PreSave(ObjectSaveContext);
+}
+
+void UPoseSearchFeatureChannel_Trajectory::InitializeSchema( UE::PoseSearch::FSchemaInitializer& Initializer)
+{
+	auto AddFeatures = [this, &Initializer](int32 NumSampleOffsets, EPoseSearchFeatureType Type)
+	{
+		FPoseSearchFeatureDesc Feature;
+		Feature.ChannelIdx = GetChannelIndex();
+		Feature.Type = Type;
+		Feature.ChannelFeatureId = 0; // Unused
+
+		for (Feature.SubsampleIdx = 0; Feature.SubsampleIdx != NumSampleOffsets; ++Feature.SubsampleIdx)
+		{
+			Initializer.AddFeatureDesc(Feature);
+		}
+	};
+
+	if (bUsePositions)
+	{
+		AddFeatures(SampleOffsets.Num(), EPoseSearchFeatureType::Position);
+	}
+
+	if (bUseLinearVelocities)
+	{
+		AddFeatures(SampleOffsets.Num(), EPoseSearchFeatureType::LinearVelocity);
+	}
+
+	if (bUseFacingDirections)
+	{
+		AddFeatures(SampleOffsets.Num(), EPoseSearchFeatureType::ForwardVector);
+	}
+}
+
+void UPoseSearchFeatureChannel_Trajectory::IndexAsset(const UE::PoseSearch::IAssetIndexer& Indexer,  UE::PoseSearch::FAssetIndexingOutput& IndexingOutput) const
+{
+	using namespace UE::PoseSearch;
+
+	auto IndexFeatures = [&Indexer, &IndexingOutput](auto FeatureIndexingFunc)
+	{
+		const FAssetIndexingContext& IndexingContext = Indexer.GetIndexingContext();
+		for (int32 SampleIdx = IndexingContext.BeginSampleIdx; SampleIdx != IndexingContext.EndSampleIdx; ++SampleIdx)
+		{
+			int32 VectorIdx = SampleIdx - IndexingContext.BeginSampleIdx;
+			FPoseSearchFeatureVectorBuilder& FeatureVector = IndexingOutput.PoseVectors[VectorIdx];
+			FeatureIndexingFunc(Indexer, SampleIdx, FeatureVector);
+		}
+	};
+
+	switch (Domain)
+	{
+		case EPoseSearchFeatureDomain::Time:
+			IndexFeatures([this](auto&&... Args){ return IndexTimeFeatures(Args...); });
+			break;
+
+		case EPoseSearchFeatureDomain::Distance:
+			IndexFeatures([this](auto&&... Args){ return IndexDistanceFeatures(Args...); });
+			break;
+
+		default:
+			checkNoEntry();
+	}
+}
+
+void UPoseSearchFeatureChannel_Trajectory::IndexTimeFeatures(const UE::PoseSearch::IAssetIndexer& Indexer, int32 SampleIdx,  FPoseSearchFeatureVectorBuilder& FeatureVector) const
+{
+	// This function samples the instantaneous trajectory at time t as well as the trajectory's velocity and acceleration at time t.
+	// Symmetric finite differences are used to approximate derivatives:
+	//	First symmetric derivative:   f'(t) ~ (f(t+h) - f(t-h)) / 2h
+	//	Second symmetric derivative: f''(t) ~ (f(t+h) - 2f(t) + f(t-h)) / h^2
+	// Where h is a constant time delta
+	// So this means three root motion extractions are taken at time t-h, t, and t+h
+
+	using namespace UE::PoseSearch;
+	using FSampleInfo = IAssetIndexer::FSampleInfo;
+
+	const FAssetIndexingContext& IndexingContext = Indexer.GetIndexingContext();
+
+
+	FPoseSearchFeatureDesc Feature;
+	Feature.ChannelIdx = GetChannelIndex();
+	Feature.ChannelFeatureId = 0; // Unused
+
+	float SampleTime = FMath::Min(SampleIdx * IndexingContext.Schema->SamplingInterval, IndexingContext.MainSampler->GetPlayLength());
+	FSampleInfo Origin = Indexer.GetSampleInfo(SampleTime);
+
+	for (int32 SubsampleIdx = 0; SubsampleIdx != SampleOffsets.Num(); ++SubsampleIdx)
+	{
+		Feature.SubsampleIdx = SubsampleIdx;
+
+		float SubsampleTime = SampleTime + SampleOffsets[SubsampleIdx];
+
+		// For each pose subsample term, get the corresponding clip, accumulated root motion,
+		// and wrap the time parameter based on the clip's length.
+		FSampleInfo Samples[3];
+		Samples[0] = Indexer.GetSampleInfoRelative(SubsampleTime - IndexingContext.SamplingContext->FiniteDelta, Origin);
+		Samples[1] = Indexer.GetSampleInfoRelative(SubsampleTime, Origin);
+		Samples[2] = Indexer.GetSampleInfoRelative(SubsampleTime + IndexingContext.SamplingContext->FiniteDelta, Origin);
+
+		// Mirror transforms if requested
+		FTransform MirroredRoots[3];
+		MirroredRoots[0] = Indexer.MirrorTransform(Samples[0].RootTransform);
+		MirroredRoots[1] = Indexer.MirrorTransform(Samples[1].RootTransform);
+		MirroredRoots[2] = Indexer.MirrorTransform(Samples[2].RootTransform);
+
+		// Add properties to the feature vector for the pose at SampleIdx
+		FeatureVector.SetTransform(Feature, MirroredRoots[1]);
+
+		// We can get a better finite difference if we ignore samples that have
+		// been clamped at either side of the clip. However, if the central sample 
+		// itself is clamped, or there are no samples that are clamped, we can just 
+		// use the central difference as normal.
+		if (Samples[0].bClamped && !Samples[1].bClamped && !Samples[2].bClamped)
+		{
+			FeatureVector.SetTransformVelocity(Feature, MirroredRoots[2], MirroredRoots[1], IndexingContext.SamplingContext->FiniteDelta);
+		}
+		else if (Samples[2].bClamped && !Samples[1].bClamped && !Samples[0].bClamped)
+		{
+			FeatureVector.SetTransformVelocity(Feature, MirroredRoots[1], MirroredRoots[0], IndexingContext.SamplingContext->FiniteDelta);
+		}
+		else
+		{
+			FeatureVector.SetTransformVelocity(Feature, MirroredRoots[2], MirroredRoots[1], MirroredRoots[0], IndexingContext.SamplingContext->FiniteDelta);
+		}
+	}
+}
+
+void UPoseSearchFeatureChannel_Trajectory::IndexDistanceFeatures(const  UE::PoseSearch::IAssetIndexer& Indexer, int32 SampleIdx, FPoseSearchFeatureVectorBuilder& FeatureVector) const
+{
+	// This function is very similar to AddTrajectoryTimeFeatures, but samples are taken in the distance domain
+	// instead of time domain.
+
+	using namespace UE::PoseSearch;
+	using FSampleInfo = IAssetIndexer::FSampleInfo;
+
+	const FAssetIndexingContext& IndexingContext = Indexer.GetIndexingContext();
+
+
+	FPoseSearchFeatureDesc Feature;
+	Feature.ChannelIdx = GetChannelIndex();
+	Feature.ChannelFeatureId = 0; // Unused
+
+	float SampleTime = FMath::Min(SampleIdx * IndexingContext.Schema->SamplingInterval, IndexingContext.MainSampler->GetPlayLength());
+	FSampleInfo Origin = Indexer.GetSampleInfo(SampleTime);
+
+	FQuat RootReferenceRot = IndexingContext.SamplingContext->BoneContainer.GetRefPoseTransform(FCompactPoseBoneIndex(0)).GetRotation();
+	for (int32 SubsampleIdx = 0; SubsampleIdx != SampleOffsets.Num(); ++SubsampleIdx)
+	{
+		Feature.SubsampleIdx = SubsampleIdx;
+
+		// For distance based sampling of the trajectory, we first have to look up the time value
+		// we're sampling given the desired travel distance of the root for this distance offset.
+		// Once we know the time, we can then carry on just like time-based sampling.
+		const float SubsampleDistance = Origin.RootDistance + SampleOffsets[SubsampleIdx];
+		float SubsampleTime = Indexer.GetSampleTimeFromDistance(SubsampleDistance);
+
+		// For each pose subsample term, get the corresponding clip, accumulated root motion,
+		// and wrap the time parameter based on the clip's length.
+		FSampleInfo Samples[3];
+		Samples[0] = Indexer.GetSampleInfoRelative(SubsampleTime - IndexingContext.SamplingContext->FiniteDelta, Origin);
+		Samples[1] = Indexer.GetSampleInfoRelative(SubsampleTime, Origin);
+		Samples[2] = Indexer.GetSampleInfoRelative(SubsampleTime + IndexingContext.SamplingContext->FiniteDelta, Origin);
+
+		// Mirror transforms if requested
+		FTransform MirroredRoots[3];
+		MirroredRoots[0] = Indexer.MirrorTransform(Samples[0].RootTransform);
+		MirroredRoots[1] = Indexer.MirrorTransform(Samples[1].RootTransform);
+		MirroredRoots[2] = Indexer.MirrorTransform(Samples[2].RootTransform);
+
+		// Add properties to the feature vector for the pose at SampleIdx
+		FeatureVector.SetTransform(Feature, MirroredRoots[1]);
+
+		// We can get a better finite difference if we ignore samples that have
+		// been clamped at either side of the clip. However, if the central sample 
+		// itself is clamped, or there are no samples that are clamped, we can just 
+		// use the central difference as normal.
+		if (Samples[0].bClamped && !Samples[1].bClamped && !Samples[2].bClamped)
+		{
+			FeatureVector.SetTransformVelocity(Feature, MirroredRoots[2], MirroredRoots[1], IndexingContext.SamplingContext->FiniteDelta);
+		}
+		else if (Samples[2].bClamped && !Samples[1].bClamped && !Samples[0].bClamped)
+		{
+			FeatureVector.SetTransformVelocity(Feature, MirroredRoots[1], MirroredRoots[0], IndexingContext.SamplingContext->FiniteDelta);
+		}
+		else
+		{
+			FeatureVector.SetTransformVelocity(Feature, MirroredRoots[2], MirroredRoots[1], MirroredRoots[0], IndexingContext.SamplingContext->FiniteDelta);
+		}
+	}
+}
+
+FFloatRange UPoseSearchFeatureChannel_Trajectory::GetHorizonRange(EPoseSearchFeatureDomain InDomain) const
+{
+	FFloatRange Extent = FFloatRange::Empty();
+	if (InDomain == Domain)
+	{
+		if (SampleOffsets.Num())
+		{
+			Extent = FFloatRange::Inclusive(SampleOffsets[0], SampleOffsets.Last());
+		}
+	}
+
+	return Extent;
+}
+
+void UPoseSearchFeatureChannel_Trajectory::GenerateDDCKey(FBlake3& InOutKeyHasher) const
+{
+	InOutKeyHasher.Update(&bUseLinearVelocities, sizeof(bUseLinearVelocities));
+	InOutKeyHasher.Update(&bUsePositions, sizeof(bUsePositions));
+	InOutKeyHasher.Update(&bUseFacingDirections, sizeof(bUseFacingDirections));
+	InOutKeyHasher.Update(&Domain, sizeof(Domain));
+	InOutKeyHasher.Update(MakeMemoryView(SampleOffsets));
+}
+
+bool UPoseSearchFeatureChannel_Trajectory::BuildQuery(UE::PoseSearch::FQueryBuildingContext& Context) const
+{
+	if (!Context.Trajectory)
+	{
+		return false;
+	}
+
+	ETrajectorySampleDomain SampleDomain;
+	switch (Domain)
+	{
+		case EPoseSearchFeatureDomain::Time:
+			SampleDomain = ETrajectorySampleDomain::Time;
+			break;
+
+		case EPoseSearchFeatureDomain::Distance:
+			SampleDomain = ETrajectorySampleDomain::Distance;
+			break;
+
+		default:
+			checkNoEntry();
+			return false;
+	}
+
+	FPoseSearchFeatureDesc Feature;
+	Feature.ChannelIdx = GetChannelIndex();
+
+	for (int32 Idx = 0, NextIterStartIdx = 0, Num = SampleOffsets.Num(); Idx < Num; ++Idx)
+	{
+		const float SampleOffset = SampleOffsets[Idx];
+		const FTrajectorySample Sample = FTrajectorySampleRange::IterSampleTrajectory(Context.Trajectory->Samples, SampleDomain, SampleOffset, NextIterStartIdx);
+
+		Feature.SubsampleIdx = Idx;
+
+		Feature.Type = EPoseSearchFeatureType::LinearVelocity;
+		Context.Query.SetVector(Feature, Sample.LinearVelocity);
+
+		Context.Query.SetTransform(Feature, Sample.Transform);
+	}
+
+	return true;
+}
+
+void UPoseSearchFeatureChannel_Trajectory::DebugDraw(const UE::PoseSearch::FDebugDrawParams& DrawParams, const UE::PoseSearch::FFeatureVectorReader& Reader) const
+{
+	using namespace UE::PoseSearch;
+
+	const float LifeTime = DrawParams.DefaultLifeTime;
+	const uint8 DepthPriority = ESceneDepthPriorityGroup::SDPG_Foreground + 2;
+	const bool bPersistent = EnumHasAnyFlags(DrawParams.Flags, EDebugDrawFlags::Persistent);
+
+	FPoseSearchFeatureDesc Feature;
+	Feature.ChannelIdx = GetChannelIndex();
+
+	const int32 NumSubsamples = SampleOffsets.Num();
+	if (NumSubsamples == 0)
+	{
+		return;
+	}
+
+	auto GetGradientColor = [](const FLinearColor& OriginalColor, int SampleIdx, int NumSamples, EDebugDrawFlags Flags) -> FLinearColor
+	{
+		int Denominator = NumSamples - 1;
+		if (Denominator <= 0 || !EnumHasAnyFlags(Flags, EDebugDrawFlags::DrawSamplesWithColorGradient))
+		{
+			return OriginalColor;
+		}
+
+		return OriginalColor * (1.0f - DrawDebugGradientStrength * (SampleIdx / (float)Denominator));
+	};
+
+	for (int32 SubsampleIdx = 0; SubsampleIdx != NumSubsamples; ++SubsampleIdx)
+	{
+		Feature.SubsampleIdx = SubsampleIdx;
+
+		FVector TrajectoryPos;
+		if (Reader.GetPosition(Feature, &TrajectoryPos))
+		{
+			Feature.Type = EPoseSearchFeatureType::Position;
+
+			FLinearColor LinearColor = DrawParams.Color ? *DrawParams.Color : GetColorForFeature(Feature, Reader.GetLayout());
+			FLinearColor GradientColor = GetGradientColor(LinearColor, SubsampleIdx, NumSubsamples, DrawParams.Flags);
+			FColor Color = GradientColor.ToFColor(true);
+
+			TrajectoryPos = DrawParams.RootTransform.TransformPosition(TrajectoryPos);
+			if (EnumHasAnyFlags(DrawParams.Flags, EDebugDrawFlags::DrawFast | EDebugDrawFlags::DrawSearchIndex))
+			{
+				DrawDebugPoint(DrawParams.World, TrajectoryPos, DrawParams.PointSize, Color, bPersistent, DrawParams.DefaultLifeTime, DepthPriority);
+			}
+			else
+			{
+				DrawDebugSphere(DrawParams.World, TrajectoryPos, DrawDebugSphereSize, DrawDebugSphereSegments, Color, bPersistent, LifeTime, DepthPriority);
+			}
+		}
+		else
+		{
+			TrajectoryPos = DrawParams.RootTransform.GetTranslation();
+		}
+
+		FVector TrajectoryVel;
+		if (Reader.GetLinearVelocity(Feature, &TrajectoryVel))
+		{
+			Feature.Type = EPoseSearchFeatureType::LinearVelocity;
+
+			FLinearColor LinearColor = DrawParams.Color ? *DrawParams.Color : GetColorForFeature(Feature, Reader.GetLayout());
+			FLinearColor GradientColor = GetGradientColor(LinearColor, SubsampleIdx, NumSubsamples, DrawParams.Flags);
+			FColor Color = GradientColor.ToFColor(true);
+
+			TrajectoryVel *= DrawDebugVelocityScale;
+			TrajectoryVel = DrawParams.RootTransform.TransformVector(TrajectoryVel);
+			FVector TrajectoryVelDirection = TrajectoryVel.GetSafeNormal();
+
+
+			if (EnumHasAnyFlags(DrawParams.Flags, EDebugDrawFlags::DrawSearchIndex))
+			{
+				DrawDebugPoint(DrawParams.World, TrajectoryVel, DrawParams.PointSize, Color, bPersistent, DrawParams.DefaultLifeTime, DepthPriority);
+			}
+			else
+			{
+				const float AdjustedThickness =
+					EnumHasAnyFlags(DrawParams.Flags, EDebugDrawFlags::DrawFast) ?
+					0.0f : DrawDebugLineThickness;
+
+				DrawDebugDirectionalArrow(
+					DrawParams.World,
+					TrajectoryPos + TrajectoryVelDirection * DrawDebugSphereSize,
+					TrajectoryPos + TrajectoryVel,
+					DrawDebugArrowSize,
+					Color,
+					bPersistent,
+					LifeTime,
+					DepthPriority,
+					AdjustedThickness
+				);
+			}
+		}
+
+		FVector TrajectoryForward;
+		if (Reader.GetForwardVector(Feature, &TrajectoryForward))
+		{
+			Feature.Type = EPoseSearchFeatureType::ForwardVector;
+
+			FLinearColor LinearColor = DrawParams.Color ? *DrawParams.Color : GetColorForFeature(Feature, Reader.GetLayout());
+			FLinearColor GradientColor = GetGradientColor(LinearColor, SubsampleIdx, NumSubsamples, DrawParams.Flags);
+			FColor Color = GradientColor.ToFColor(true);
+
+			TrajectoryForward = DrawParams.RootTransform.TransformVector(TrajectoryForward);
+
+			if (EnumHasAnyFlags(DrawParams.Flags, EDebugDrawFlags::DrawSearchIndex))
+			{
+				DrawDebugPoint(DrawParams.World, TrajectoryForward, DrawParams.PointSize, Color, bPersistent, DrawParams.DefaultLifeTime, DepthPriority);
+			}
+			else
+			{
+				const float AdjustedThickness =
+					EnumHasAnyFlags(DrawParams.Flags, EDebugDrawFlags::DrawFast) ?
+					0.0f : DrawDebugLineThickness;
+
+				DrawDebugDirectionalArrow(
+					DrawParams.World,
+					TrajectoryPos + TrajectoryForward * DrawDebugSphereSize,
+					TrajectoryPos + TrajectoryForward * DrawDebugSphereSize * 2.0f,
+					DrawDebugArrowSize,
+					Color,
+					bPersistent,
+					LifeTime,
+					DepthPriority,
+					AdjustedThickness
+				);
+			}
+		}
+
+		if (EnumHasAnyFlags(DrawParams.Flags, EDebugDrawFlags::DrawSampleLabels))
+		{
+			FLinearColor LinearColor = DrawParams.Color ? *DrawParams.Color : GetColorForFeature(Feature, Reader.GetLayout());
+			FLinearColor GradientColor = GetGradientColor(LinearColor, SubsampleIdx, NumSubsamples, DrawParams.Flags);
+			FColor Color = GradientColor.ToFColor(true);
+
+			FString SampleLabel;
+			if (DrawParams.LabelPrefix.IsEmpty())
+			{
+				SampleLabel = FString::Format(TEXT("{0}"), { SubsampleIdx });
+			}
+			else
+			{
+				SampleLabel = FString::Format(TEXT("{1}[{0}]"), { SubsampleIdx, DrawParams.LabelPrefix.GetData() });
+			}
+			DrawDebugString(
+				DrawParams.World,
+				TrajectoryPos + DrawDebugSampleLabelOffset,
+				SampleLabel,
+				nullptr,
+				Color,
+				LifeTime,
+				false,
+				DrawDebugSampleLabelFontScale);
+		}
+	}
+}
\ No newline at end of file
diff --git a/Engine/Plugins/Experimental/Animation/PoseSearch/Source/Runtime/Private/PoseSearchLibrary.cpp b/Engine/Plugins/Experimental/Animation/PoseSearch/Source/Runtime/Private/PoseSearchLibrary.cpp
--- a/Engine/Plugins/Experimental/Animation/PoseSearch/Source/Runtime/Private/PoseSearchLibrary.cpp
+++ b/Engine/Plugins/Experimental/Animation/PoseSearch/Source/Runtime/Private/PoseSearchLibrary.cpp
@@ -283,21 +283,6 @@ bool FMotionMatchingState::CanAdvance(float DeltaTime, bool& bOutAdvanceToFollow
 	return false;
 }
 
-void FMotionMatchingState::ComposeQuery(const UPoseSearchDatabase* Database, const FTrajectorySampleRange& Trajectory)
-{
-	FPoseSearchFeatureVectorBuilder Goal;
-	Goal.Init(Database->Schema);
-	Goal.BuildFromTrajectory(Trajectory);
-
-	// Merge goal features into the query vector
-	if (ComposedQuery.IsCompatible(Goal))
-	{
-		ComposedQuery.MergeReplace(Goal);
-	}
-
-	ComposedQuery.Normalize(*Database->GetSearchIndex());
-}
-
 static void RequestInertialBlend(const FAnimationUpdateContext& Context, float BlendTime)
 {
 	// Use inertial blending to smooth over the transition
@@ -373,6 +358,11 @@ void UpdateMotionMatchingState(
 	// Record Current Pose Index for Debugger
 	int32 CurrentPoseIdx = InOutMotionMatchingState.DbPoseIdx;
 
+	FQueryBuildingContext QueryBuildingContext(InOutMotionMatchingState.ComposedQuery);
+	QueryBuildingContext.Schema = Database->Schema;
+	QueryBuildingContext.History = nullptr;
+	QueryBuildingContext.Trajectory = &Trajectory;
+
 	// Build the search query. This is done even when we don't search 
 	// since we still want to record it for debugging purposes.
 	if (InOutMotionMatchingState.DbPoseIdx != INDEX_NONE)
@@ -384,29 +374,16 @@ void UpdateMotionMatchingState(
 	}
 	else
 	{
-		// When we don't have an active pose in the database, 
-		// initialize the search query from the pose history provider
-		bool bPoseHistoryFound = false;
-
+		InOutMotionMatchingState.ComposedQuery.ResetFeatures();
 		IPoseHistoryProvider* PoseHistoryProvider = Context.GetMessage<IPoseHistoryProvider>();
 		if (PoseHistoryProvider)
 		{
-			FPoseHistory& History = PoseHistoryProvider->GetPoseHistory();
-
-			bPoseHistoryFound = InOutMotionMatchingState.ComposedQuery.TrySetPoseFeatures(
-				&History,
-				Context.AnimInstanceProxy->GetRequiredBones());
-		}
-
-		// If no pose history provider or no pose found in history, then just reset/zero
-		if (!bPoseHistoryFound)
-		{
-			InOutMotionMatchingState.ComposedQuery.ResetFeatures();
+			QueryBuildingContext.History = &PoseHistoryProvider->GetPoseHistory();
 		}
 	}
 
-	// Update features in the query with the latest input trajectory
-	InOutMotionMatchingState.ComposeQuery(Database, Trajectory);
+	BuildQuery(QueryBuildingContext);
+	InOutMotionMatchingState.ComposedQuery.Normalize(*Database->GetSearchIndex());
 
 	// Check if we can advance. Includes the case where we can advance but only by switching to a follow up asset.
 	bool bAdvanceToFollowUpAsset = false;
@@ -431,7 +408,7 @@ void UpdateMotionMatchingState(
 		}
 
 		// Update weight groups
-		InOutMotionMatchingState.WeightsContext.Update(Settings.Weights, Database);
+		InOutMotionMatchingState.WeightsContext.Update(Database);
 
 		// Search the database for the nearest match to the updated query vector
 		FSearchResult SearchResult = Search(SearchContext);
@@ -571,7 +548,6 @@ void UpdateMotionMatchingState(
 		// @TODO: Change this to only be the previous query, not persistently updated (i.e. if throttled)?
 		TraceState.QueryVector = InOutMotionMatchingState.ComposedQuery.GetValues();
 		TraceState.QueryVectorNormalized = InOutMotionMatchingState.ComposedQuery.GetNormalizedValues();
-		TraceState.Weights = Settings.Weights;
 		TraceState.DbPoseIdx = InOutMotionMatchingState.DbPoseIdx;
 		TraceState.DatabaseId = FObjectTrace::GetObjectId(Database);
 		TraceState.ContinuingPoseIdx = CurrentPoseIdx;
diff --git a/Engine/Plugins/Experimental/Animation/PoseSearch/Source/Runtime/Private/Trace/PoseSearchTraceLogger.cpp b/Engine/Plugins/Experimental/Animation/PoseSearch/Source/Runtime/Private/Trace/PoseSearchTraceLogger.cpp
--- a/Engine/Plugins/Experimental/Animation/PoseSearch/Source/Runtime/Private/Trace/PoseSearchTraceLogger.cpp
+++ b/Engine/Plugins/Experimental/Animation/PoseSearch/Source/Runtime/Private/Trace/PoseSearchTraceLogger.cpp
@@ -22,10 +22,6 @@ UE_TRACE_EVENT_BEGIN(PoseSearch, MotionMatchingState)
 	UE_TRACE_EVENT_FIELD(int32, ContinuingPoseIdx)
 	UE_TRACE_EVENT_FIELD(float[], QueryVector)
 	UE_TRACE_EVENT_FIELD(float[], QueryVectorNormalized)
-	UE_TRACE_EVENT_FIELD(float[], ChannelWeightScales)
-	UE_TRACE_EVENT_FIELD(float[], HistoryWeightScales)
-	UE_TRACE_EVENT_FIELD(float[], PredictionWeightScales)
-	UE_TRACE_EVENT_FIELD(bool, DebugDisableWeights)
 	UE_TRACE_EVENT_FIELD(float, AssetPlayerTime)
 	UE_TRACE_EVENT_FIELD(float, DeltaTime)
 	UE_TRACE_EVENT_FIELD(float, SimLinearVelocity)
@@ -69,23 +65,6 @@ void FTraceMotionMatchingState::Output(const FAnimationBaseContext& InContext, c
 	TRACE_OBJECT(AnimInstance);
 	UObject* SkeletalMeshComponent = AnimInstance->GetOuter();
 
-	float ChannelWeightScales[] =
-	{
-		State.Weights.PoseDynamicWeights.ChannelWeightScale,
-		State.Weights.TrajectoryDynamicWeights.ChannelWeightScale
-	};
-
-	float HistoryWeightScales[] =
-	{
-		State.Weights.PoseDynamicWeights.HistoryWeightScale,
-		State.Weights.TrajectoryDynamicWeights.HistoryWeightScale
-	};
-	float PredictionWeightScales[] =
-	{
-		State.Weights.PoseDynamicWeights.PredictionWeightScale,
-		State.Weights.TrajectoryDynamicWeights.PredictionWeightScale
-	};
-
 	UE_TRACE_LOG(PoseSearch, MotionMatchingState, PoseSearchChannel)
 		<< MotionMatchingState.Cycle(FPlatformTime::Cycles64())
 		<< MotionMatchingState.FrameCounter(FObjectTrace::GetObjectWorldTickCounter(AnimInstance))
@@ -99,10 +78,6 @@ void FTraceMotionMatchingState::Output(const FAnimationBaseContext& InContext, c
 		<< MotionMatchingState.ContinuingPoseIdx(State.ContinuingPoseIdx)
 		<< MotionMatchingState.QueryVector(State.QueryVector.GetData(), State.QueryVector.Num())
 		<< MotionMatchingState.QueryVectorNormalized(State.QueryVectorNormalized.GetData(), State.QueryVectorNormalized.Num())
-		<< MotionMatchingState.ChannelWeightScales(ChannelWeightScales, UE_ARRAY_COUNT(ChannelWeightScales))
-		<< MotionMatchingState.HistoryWeightScales(HistoryWeightScales, UE_ARRAY_COUNT(HistoryWeightScales))
-		<< MotionMatchingState.PredictionWeightScales(PredictionWeightScales, UE_ARRAY_COUNT(PredictionWeightScales))
-		<< MotionMatchingState.DebugDisableWeights(State.Weights.bDebugDisableWeights)
 		<< MotionMatchingState.AssetPlayerTime(State.AssetPlayerTime)
 		<< MotionMatchingState.DeltaTime(State.DeltaTime)
 		<< MotionMatchingState.SimLinearVelocity(State.SimLinearVelocity)
diff --git a/Engine/Plugins/Experimental/Animation/PoseSearch/Source/Runtime/Private/Trace/PoseSearchTraceLogger.h b/Engine/Plugins/Experimental/Animation/PoseSearch/Source/Runtime/Private/Trace/PoseSearchTraceLogger.h
--- a/Engine/Plugins/Experimental/Animation/PoseSearch/Source/Runtime/Private/Trace/PoseSearchTraceLogger.h
+++ b/Engine/Plugins/Experimental/Animation/PoseSearch/Source/Runtime/Private/Trace/PoseSearchTraceLogger.h
@@ -44,9 +44,6 @@ struct POSESEARCH_API FTraceMotionMatchingState
 	TArrayView<const float> QueryVector;
 	TArrayView<const float> QueryVectorNormalized;
 
-	/** Runtime weights */
-	FPoseSearchDynamicWeightParams Weights;
-
 	/** If true, groups are being filtered by DatabaseGroupQuery */
 	TArrayView<const bool> DatabaseSequenceFilter;
 
diff --git a/Engine/Plugins/Experimental/Animation/PoseSearch/Source/Runtime/Public/PoseSearch/PoseSearch.h b/Engine/Plugins/Experimental/Animation/PoseSearch/Source/Runtime/Public/PoseSearch/PoseSearch.h
--- a/Engine/Plugins/Experimental/Animation/PoseSearch/Source/Runtime/Public/PoseSearch/PoseSearch.h
+++ b/Engine/Plugins/Experimental/Animation/PoseSearch/Source/Runtime/Public/PoseSearch/PoseSearch.h
@@ -21,23 +21,40 @@
 
 #include "PoseSearch.generated.h"
 
+DECLARE_LOG_CATEGORY_EXTERN(LogPoseSearch, Log, All);
+
+
+//////////////////////////////////////////////////////////////////////////
+// Forward declarations
+
 class UAnimSequence;
 class UBlendSpace;
 struct FCompactPose;
 struct FPoseContext;
 struct FReferenceSkeleton;
-
-DECLARE_LOG_CATEGORY_EXTERN(LogPoseSearch, Log, All);
-
 struct FPoseSearchDatabaseDerivedData;
+class UAnimNotifyState_PoseSearchBase;
+class UPoseSearchSchema;
+class FBlake3;
 
 namespace UE::PoseSearch {
+
 class FPoseHistory;
 struct FPoseSearchDatabaseAsyncCacheTask;
-}
+struct FDebugDrawParams;
+class FFeatureVectorReader;
+struct FSchemaInitializer;
+struct FQueryBuildingContext;
 
+} // namespace UE::PoseSearch
+
+
+
+//////////////////////////////////////////////////////////////////////////
+// Constants
+
 UENUM()
-enum class EPoseSearchFeatureType : int32
+enum class EPoseSearchFeatureType : int8
 {
 	Position,	
 	Rotation,
@@ -70,6 +87,111 @@
 	Invalid = Num UMETA(Hidden)
 };
 
+UENUM()
+enum class EPoseSearchDataPreprocessor : int32
+{
+	None,
+	Automatic,
+	Normalize,
+	Sphere UMETA(Hidden),
 
+	Num UMETA(Hidden),
+	Invalid = Num UMETA(Hidden)
+};
+
+UENUM()
+enum class EPoseSearchPoseFlags : uint32
+{
+	None = 0,
+
+	// Don't return this pose as a search result
+	BlockTransition = 1 << 0,
+};
+ENUM_CLASS_FLAGS(EPoseSearchPoseFlags);
+
+UENUM()
+enum class ESearchIndexAssetType : int32
+{
+	Invalid,
+	Sequence,
+	BlendSpace,
+};
+
+UENUM()
+enum class EPoseSearchMirrorOption : int32
+{
+	UnmirroredOnly UMETA(DisplayName = "Original Only"),
+	MirroredOnly UMETA(DisplayName = "Mirrored Only"),
+	UnmirroredAndMirrored UMETA(DisplayName = "Original and Mirrored"),
+
+	Num UMETA(Hidden),
+	Invalid = Num UMETA(Hidden)
+};
+
+
+
+//////////////////////////////////////////////////////////////////////////
+// Common structs
+
+USTRUCT()
+struct POSESEARCH_API FPoseSearchExtrapolationParameters
+{
+	GENERATED_BODY()
+
+	// If the angular root motion speed in degrees is below this value, it will be treated as zero.
+	UPROPERTY(EditAnywhere, Category = "Settings")
+	float AngularSpeedThreshold = 1.0f;
+
+	// If the root motion linear speed is below this value, it will be treated as zero.
+	UPROPERTY(EditAnywhere, Category = "Settings")
+	float LinearSpeedThreshold = 1.0f;
+
+	// Time from sequence start/end used to extrapolate the trajectory.
+	UPROPERTY(EditAnywhere, Category = "Settings")
+	float SampleTime = 0.05f;
+};
+
+USTRUCT()
+struct FPoseSearchBlockTransitionParameters
+{
+	GENERATED_BODY()
+
+	// Excluding the beginning of sequences can help ensure an exact past trajectory is used when building the features
+	UPROPERTY(EditAnywhere, Category = "Settings")
+	float SequenceStartInterval = 0.0f;
+
+	// Excluding the end of sequences help ensure an exact future trajectory, and also prevents the selection of
+	// a sequence which will end too soon to be worth selecting.
+	UPROPERTY(EditAnywhere, Category = "Settings")
+	float SequenceEndInterval = 0.0f;
+};
+
+USTRUCT()
+struct POSESEARCH_API FPoseSearchBone
+{
+	GENERATED_BODY()
+
+	UPROPERTY(EditAnywhere, Category = Config)
+	FBoneReference Reference;
+
+	UPROPERTY(EditAnywhere, Category = Config)
+	bool bUseVelocity = false;
+
+	UPROPERTY(EditAnywhere, Category = Config)
+	bool bUsePosition = false;
+
+	UPROPERTY(EditAnywhere, Category = Config)
+	bool bUseRotation = false;
+
+	// this function will return a mask out of EPoseSearchFeatureType based on which features were selected for
+	// the bone.
+	uint32 GetTypeMask() const;
+};
+
+
+
+//////////////////////////////////////////////////////////////////////////
+// Feature descriptions and vector layout
+
 /** Describes each feature of a vector, including data type, sampling options, and buffer offset. */
 USTRUCT()
@@ -77,26 +199,33 @@
 {
 	GENERATED_BODY()
 
-	static constexpr int32 TrajectoryBoneIndex = -1; 
+	// Index into UPoseSearchSchema::Channels
+	UPROPERTY()
+	int8 ChannelIdx = -1;
 
+	// Optional feature identifier within a channel
 	UPROPERTY()
-	int32 SchemaBoneIdx = 0;
+	int8 ChannelFeatureId = 0;
 
+	// Index into channel's sample offsets, if any
 	UPROPERTY()
-	int32 SubsampleIdx = 0;
+	int8 SubsampleIdx = 0;
 
+	// Value type of the feature
 	UPROPERTY()
 	EPoseSearchFeatureType Type = EPoseSearchFeatureType::Invalid;
 
-	UPROPERTY()
-	EPoseSearchFeatureDomain Domain = EPoseSearchFeatureDomain::Invalid;
-
-	UPROPERTY()
-	int8 ChannelIdx = 0;
-
 	// Set via FPoseSearchFeatureLayout::Init() and ignored by operator==
 	UPROPERTY()
-	int32 ValueOffset = 0;
+	int16 ValueOffset = 0;
 
+#if WITH_EDITORONLY_DATA
+	UPROPERTY()
+	EPoseSearchFeatureDomain Domain_DEPRECATED;
+
+	UPROPERTY()
+	int32 SchemaBoneIdx_DEPRECATED = 0;
+#endif
+
 	bool operator==(const FPoseSearchFeatureDesc& Other) const;
 
@@ -114,7 +243,7 @@
 {
 	GENERATED_BODY()
 
-	void Init();
+	void Finalize();
 	void Reset();
 
 	UPROPERTY()
@@ -123,48 +252,204 @@
 	UPROPERTY()
 	int32 NumFloats = 0;
 
-	UPROPERTY()
-	int32 NumChannels = 0;
+	bool IsValid(int32 ChannelCount) const;
 
-	bool IsValid(int32 MaxNumBones) const;
-
 	bool EnumerateBy(int32 ChannelIdx, EPoseSearchFeatureType Type, int32& InOutFeatureIdx) const;
 };
 
-UENUM()
-enum class EPoseSearchDataPreprocessor : int32
+
+
+//////////////////////////////////////////////////////////////////////////
+// Asset sampling and indexing
+
+namespace UE::PoseSearch {
+
+struct POSESEARCH_API FAssetSamplingContext
 {
-	None,
-	Automatic,
-	Normalize,
-	Sphere UMETA(Hidden),
+	// Time delta used for computing pose derivatives
+	static constexpr float FiniteDelta = 1 / 60.0f;
 
-	Num UMETA(Hidden),
-	Invalid = Num UMETA(Hidden)
+	FBoneContainer BoneContainer;
+
+	// Mirror data table pointer copied from Schema for convenience
+	TObjectPtr<UMirrorDataTable> MirrorDataTable = nullptr;
+
+	// Compact pose format of Mirror Bone Map
+	TCustomBoneIndexArray<FCompactPoseBoneIndex, FCompactPoseBoneIndex> CompactPoseMirrorBones;
+
+	// Pre-calculated component space rotations of reference pose, which allows mirror to work with any joint orientation
+	// Only initialized and used when a mirroring table is specified
+	TCustomBoneIndexArray<FQuat, FCompactPoseBoneIndex> ComponentSpaceRefRotations;
+
+
+	void Init(const UPoseSearchSchema* Schema);
+	FTransform MirrorTransform(const FTransform& Transform) const;
 };
 
-USTRUCT()
-struct POSESEARCH_API FPoseSearchBone
+/**
+ * Helper interface for sampling data from animation assets
+ */
+class POSESEARCH_API IAssetSampler
 {
+public:
+	virtual ~IAssetSampler() {};
+
+	virtual float GetPlayLength() const = 0;
+	virtual bool IsLoopable() const = 0;
+
+	// Gets the time associated with a particular root distance traveled
+	virtual float GetTimeFromRootDistance(float Distance) const = 0;
+
+	// Gets the total root distance traveled 
+	virtual float GetTotalRootDistance() const = 0;
+
+	// Gets the final root transformation at the end of the asset's playback time
+	virtual FTransform GetTotalRootTransform() const = 0;
+
+	// Extracts pose for this asset for a given context
+	virtual void ExtractPose(const FAnimExtractContext& ExtractionCtx, FAnimationPoseData& OutAnimPoseData) const = 0;
+
+	// Extracts the accumulated root distance at the given time, using the extremities of the sequence to extrapolate 
+	// beyond the sequence limits when Time is less than zero or greater than the sequence length
+	virtual float ExtractRootDistance(float Time) const = 0;
+
+	// Extracts root transform at the given time, using the extremities of the sequence to extrapolate beyond the 
+	// sequence limits when Time is less than zero or greater than the sequence length.
+	virtual FTransform ExtractRootTransform(float Time) const = 0;
+
+	// Extracts notify states inheriting from UAnimNotifyState_PoseSearchBase present in the sequence at Time.
+	// The function does not empty NotifyStates before adding new notifies!
+	virtual void ExtractPoseSearchNotifyStates(float Time, TArray<UAnimNotifyState_PoseSearchBase*>& NotifyStates) const = 0;
+};
+
+/**
+ * Inputs for asset indexing
+ */
+struct FAssetIndexingContext
+{
+	const FAssetSamplingContext* SamplingContext = nullptr;
+	const UPoseSearchSchema* Schema = nullptr;
+	const IAssetSampler* MainSampler = nullptr;
+	const IAssetSampler* LeadInSampler = nullptr;
+	const IAssetSampler* FollowUpSampler = nullptr;
+	bool bMirrored = false;
+	FFloatInterval RequestedSamplingRange = FFloatInterval(0.0f, 0.0f);
+	FPoseSearchBlockTransitionParameters BlockTransitionParameters;
+
+	// Index this asset's data from BeginPoseIdx up to but not including EndPoseIdx
+	int32 BeginSampleIdx = 0;
+	int32 EndSampleIdx = 0;
+};
+
+/**
+ * Output of indexer data for this asset
+ */
+struct FAssetIndexingOutput
+{
+	// Channel data should be written to this array of feature vector builders
+	// Size is EndPoseIdx - BeginPoseIdx and PoseVectors[0] contains data for BeginPoseIdx
+	const TArrayView<FPoseSearchFeatureVectorBuilder> PoseVectors;
+};
+
+class POSESEARCH_API IAssetIndexer
+{
+public:
+	struct FSampleInfo
+	{
+		const IAssetSampler* Clip = nullptr;
+		FTransform RootTransform;
+		float ClipTime = 0.0f;
+		float RootDistance = 0.0f;
+		bool bClamped = false;
+
+		bool IsValid() const { return Clip != nullptr; }
+	};
+
+	virtual ~IAssetIndexer() {}
+
+	virtual const FAssetIndexingContext& GetIndexingContext() const = 0;
+	virtual FSampleInfo GetSampleInfo(float SampleTime) const = 0;
+	virtual FSampleInfo GetSampleInfoRelative(float SampleTime, const FSampleInfo& Origin) const = 0;
+	virtual const float GetSampleTimeFromDistance(float Distance) const = 0;
+	virtual FTransform MirrorTransform(const FTransform& Transform) const = 0;
+};
+
+} // namespace UE::PoseSearch
+
+
+//////////////////////////////////////////////////////////////////////////
+// Feature channels interface
+
+UCLASS(Abstract, BlueprintType, EditInlineNew)
+class POSESEARCH_API UPoseSearchFeatureChannel : public UObject, public IBoneReferenceSkeletonProvider
+{
 	GENERATED_BODY()
 
-	UPROPERTY(EditAnywhere, Category = Config)
-	FBoneReference Reference;
-	
-	UPROPERTY(EditAnywhere, Category = Config)
-	bool bUseVelocity = false;
-	
-	UPROPERTY(EditAnywhere, Category = Config)
-	bool bUsePosition = false;
-	
-	UPROPERTY(EditAnywhere, Category = Config)
-	bool bUseRotation = false;
+public:
 
-	// this function will return a mask out of EPoseSearchFeatureType based on which features were selected for
-	// the bone.
-	uint32 GetTypeMask() const;
+	int32 GetChannelIndex() const { checkSlow(ChannelIdx >= 0); return ChannelIdx; }
+
+	// Called during UPoseSearchSchema::Finalize to prepare the schema for this channel
+	virtual void InitializeSchema(UE::PoseSearch::FSchemaInitializer& Initializer) PURE_VIRTUAL(UPoseSearchFeatureChannel::InitializeSchema, );
+
+	// Called at database build time to populate pose vectors with this channel's data
+	virtual void IndexAsset(const UE::PoseSearch::IAssetIndexer& Indexer, UE::PoseSearch::FAssetIndexingOutput& IndexingOutput) const PURE_VIRTUAL(UPoseSearchFeatureChannel::IndexAsset, );
+
+	// Return this channel's range of sampling offsets in the requested sampling domain.
+	// Returns empty range if the channel has no horizon in the requested domain.
+	virtual FFloatRange GetHorizonRange(EPoseSearchFeatureDomain Domain) const PURE_VIRTUAL(UPoseSearchFeatureChannel::GetHorizonRange, return FFloatRange::Empty(); );
+
+	// Return this channel's horizon sampling offsets
+	virtual TArrayView<const float> GetSampleOffsets () const PURE_VIRTUAL(UPoseSearchFeatureChannel::GetSampleOffsets, return {}; );
+
+	// Hash channel properties to produce a key for database derived data
+	virtual void GenerateDDCKey(FBlake3& InOutKeyHasher) const PURE_VIRTUAL(UPoseSearchFeatureChannel::GenerateDDCKey, );
+
+	// Called at runtime to add this channel's data to the query pose vector
+	virtual bool BuildQuery(UE::PoseSearch::FQueryBuildingContext& Context) const PURE_VIRTUAL(UPoseSearchFeatureChannel::BuildQuery, return false;);
+
+	// Draw this channel's data for the given pose vector
+	virtual void DebugDraw(const UE::PoseSearch::FDebugDrawParams& DrawParams, const UE::PoseSearch::FFeatureVectorReader& Reader) const PURE_VIRTUAL(UPoseSearchFeatureChannel::DebugDraw, );
+
+private:
+	// IBoneReferenceSkeletonProvider interface
+	// Note this function is exclusively for FBoneReference details customization
+	class USkeleton* GetSkeleton(bool& bInvalidSkeletonIsError, const IPropertyHandle* PropertyHandle) override;
+
+private:
+
+	friend class ::UPoseSearchSchema;
+
+	UPROPERTY()
+	int32 ChannelIdx = -1;
 };
 
+
+
+//////////////////////////////////////////////////////////////////////////
+// Schema
+
+namespace UE::PoseSearch {
+
+struct POSESEARCH_API FSchemaInitializer
+{
+public:
+
+	int32 AddBoneReference(const FBoneReference& BoneReference);
+	int32 AddFeatureDesc(const FPoseSearchFeatureDesc& FeatureDesc);
+
+private:
+	friend class ::UPoseSearchSchema;
+
+	int32 CurrentChannelIdx = 0;
+
+	TArray<FBoneReference> BoneReferences;
+	TArray<FPoseSearchFeatureDesc> Features;
+};
+
+} // namespace UE::PoseSearch
+
+
 /**
 * Specifies the format of a pose search index. At runtime, queries are built according to the schema for searching.
 */
@@ -176,6 +461,9 @@
 public:
 
 	static constexpr int32 DefaultSampleRate = 10;
+	static constexpr int32 MaxBoneReferences = MAX_int8;
+	static constexpr int32 MaxChannels = MAX_int8;
+	static constexpr int32 MaxFeatures = MAX_int8;
 
 	UPROPERTY(EditAnywhere, Category = "Schema")
 	TObjectPtr<USkeleton> Skeleton = nullptr;
@@ -183,35 +471,32 @@
 	UPROPERTY(EditAnywhere, meta = (ClampMin = "1", ClampMax = "60"), Category = "Schema")
 	int32 SampleRate = DefaultSampleRate;
 
+	UPROPERTY(EditAnywhere, Instanced, BlueprintReadWrite, Category = "Schema")
+	TArray<TObjectPtr<UPoseSearchFeatureChannel>> Channels;
+
+#if WITH_EDITORONLY_DATA
 	UPROPERTY()
-	bool bUseBoneVelocities_DEPRECATED = true;
+	bool bUseTrajectoryVelocities_DEPRECATED = true;
 
 	UPROPERTY()
-	bool bUseBonePositions_DEPRECATED = true;
+	bool bUseTrajectoryPositions_DEPRECATED = true;
 
-	UPROPERTY(EditAnywhere, Category = "Schema")
-	bool bUseTrajectoryVelocities = true;
+	UPROPERTY()
+	bool bUseTrajectoryForwardVectors_DEPRECATED = false;
 
-	UPROPERTY(EditAnywhere, Category = "Schema")
-	bool bUseTrajectoryPositions = true;
+	UPROPERTY()
+	TArray<FPoseSearchBone> SampledBones_DEPRECATED;
 
-	UPROPERTY(EditAnywhere, Category = "Schema", DisplayName="Use Facing Directions")
-	bool bUseTrajectoryForwardVectors = false;
-	
-	UPROPERTY(EditAnywhere, Category = "Schema")
-	TArray<FPoseSearchBone> SampledBones;
-	
 	UPROPERTY()
-	TArray<FBoneReference> Bones_DEPRECATED;
+	TArray<float> PoseSampleTimes_DEPRECATED;
 
-	UPROPERTY(EditAnywhere, Category = "Schema")
-	TArray<float> PoseSampleTimes;
+	UPROPERTY()
+	TArray<float> TrajectorySampleTimes_DEPRECATED;
 
-	UPROPERTY(EditAnywhere, Category = "Schema")
-	TArray<float> TrajectorySampleTimes;
+	UPROPERTY()
+	TArray<float> TrajectorySampleDistances_DEPRECATED;
+#endif // WITH_EDITOR
 
-	UPROPERTY(EditAnywhere, Category = "Schema")
-	TArray<float> TrajectorySampleDistances;
 
 	// If set, this schema will support mirroring pose search databases
 	UPROPERTY(EditAnywhere, Category = "Schema")
@@ -229,32 +514,22 @@
 	UPROPERTY()
 	FPoseSearchFeatureVectorLayout Layout;
 
+	UPROPERTY()
+	TArray<FBoneReference> BoneReferences;
+
 	UPROPERTY(Transient)
 	TArray<uint16> BoneIndices;
 
 	UPROPERTY(Transient)
 	TArray<uint16> BoneIndicesWithParents;
 
-
 	bool IsValid () const;
 
 	int32 GetNumBones () const { return BoneIndices.Num(); }
 
-	// Returns farthest future sample time >= 0.0f.
-	// Returns a negative value when there are no future sample times.
-	float GetTrajectoryFutureTimeHorizon () const;
+	// Returns global range of sampling offsets among all channels in requested sampling domain.
+	// Returns empty range if the channel has no horizon in the requested domain.
+	virtual FFloatRange GetHorizonRange(EPoseSearchFeatureDomain Domain) const;
 
-	// Returns farthest past sample time <= 0.0f.
-	// Returns a positive value when there are no past sample times.
-	float GetTrajectoryPastTimeHorizon () const;
-
-	// Returns farthest future sample distance >= 0.0f.
-	// Returns a negative value when there are no future sample distances.
-	float GetTrajectoryFutureDistanceHorizon () const;
-
-	// Returns farthest path sample distance <= 0.0f.
-	// Returns a positive value when there are no past sample distances.
-	float GetTrajectoryPastDistanceHorizon () const;
-
 	TArrayView<const float> GetChannelSampleOffsets (int32 ChannelIdx) const;
 
@@ -266,10 +541,15 @@
 	class USkeleton* GetSkeleton(bool& bInvalidSkeletonIsError, const IPropertyHandle* PropertyHandle) override { bInvalidSkeletonIsError = false; return Skeleton; }
 
 private:
-	void GenerateLayout();
+	void Finalize();
 	void ResolveBoneReferences();
 };
 
+
+
+//////////////////////////////////////////////////////////////////////////
+// Search index
+
 USTRUCT()
 struct POSESEARCH_API FPoseSearchIndexPreprocessInfo
 {
@@ -296,18 +576,10 @@
 	}
 };
 
-UENUM()
-enum class EPoseSearchPoseFlags : uint32
-{
-	None = 0,
 
-	// Don't return this pose as a search result
-	BlockTransition = 1 << 0,
-};
-ENUM_CLASS_FLAGS(EPoseSearchPoseFlags);
-
-// This is kept for each pose in the search index along side the feature vector values and is used to influence the
-// search.
+/**
+ * This is kept for each pose in the search index along side the feature vector values and is used to influence the search.
+ */
 USTRUCT()
 struct POSESEARCH_API FPoseSearchPoseMetadata
 {
@@ -321,14 +593,6 @@
 };
 
 
-UENUM()
-enum class ESearchIndexAssetType : int32
-{
-	Invalid,
-	Sequence,
-	BlendSpace,
-};
-
 /**
 * Information about a source animation asset used by a search index.
 * Some source animation entries may generate multiple FPoseSearchIndexAsset entries.
@@ -431,66 +695,11 @@
 	void InverseNormalize (TArrayView<float> PoseVector) const;
 };
 
-USTRUCT()
-struct FPoseSearchExtrapolationParameters
-{
-	GENERATED_BODY()
 
-public:
-	// If the angular root motion speed in degrees is below this value, it will be treated as zero.
-	UPROPERTY(EditAnywhere, Category = "Settings")
-	float AngularSpeedThreshold = 1.0f;
-	
-	// If the root motion linear speed is below this value, it will be treated as zero.
-	UPROPERTY(EditAnywhere, Category = "Settings")
-	float LinearSpeedThreshold = 1.0f;
 
-	// Time from sequence start/end used to extrapolate the trajectory.
-	UPROPERTY(EditAnywhere, Category = "Settings")
-	float SampleTime = 0.05f;
-};
+//////////////////////////////////////////////////////////////////////////
+// Database
 
-USTRUCT()
-struct FPoseSearchBlockTransitionParameters
-{
-	GENERATED_BODY()
-
-	// Excluding the beginning of sequences can help ensure an exact past trajectory is used when building the features
-	UPROPERTY(EditAnywhere, Category = "Settings")
-	float SequenceStartInterval = 0.0f;
-
-	// Excluding the end of sequences help ensure an exact future trajectory, and also prevents the selection of
-	// a sequence which will end too soon to be worth selecting.
-	UPROPERTY(EditAnywhere, Category = "Settings")
-	float SequenceEndInterval = 0.0f;
-};
-
-/** Animation metadata object for indexing a single animation. */
-UCLASS(BlueprintType, Category = "Animation|Pose Search", Experimental)
-class POSESEARCH_API UPoseSearchSequenceMetaData : public UAnimMetaData
-{
-	GENERATED_BODY()
-public:
-
-	UPROPERTY(EditAnywhere, Category="Settings")
-	TObjectPtr<const UPoseSearchSchema> Schema = nullptr;
-
-	UPROPERTY(EditAnywhere, Category="Settings")
-	FFloatInterval SamplingRange = FFloatInterval(0.0f, 0.0f);
-
-	UPROPERTY(EditAnywhere, Category = "Settings")
-	FPoseSearchExtrapolationParameters ExtrapolationParameters;
-
-	UPROPERTY()
-	FPoseSearchIndex SearchIndex;
-
-	bool IsValidForIndexing() const;
-	bool IsValidForSearch() const;
-
-public: // UObject
-	virtual void PreSave(FObjectPreSaveContext ObjectSaveContext) override;
-};
-
 USTRUCT(BlueprintType, Category = "Animation|Pose Search")
 struct POSESEARCH_API FPoseSearchChannelHorizonParams
 {
@@ -538,67 +747,16 @@
 };
 
 USTRUCT(BlueprintType, Category = "Animation|Pose Search")
-struct POSESEARCH_API FPoseSearchChannelDynamicWeightParams
-{
-	GENERATED_BODY()
-
-	// Multiplier for the contribution of this score component. Final weight will be normalized with other channels after scaling.
-	UPROPERTY(BlueprintReadWrite, EditAnywhere, Category = "Settings", meta = (ClampMin = "0.0"))
-	float ChannelWeightScale = 1.0f;
-
-	// Multiplier for history score contribution. Normalized with prediction weight after scaling.
-	UPROPERTY(BlueprintReadWrite, EditAnywhere, Category = "Settings", meta = (ClampMin = "0.0"))
-	float HistoryWeightScale = 1.0f;
-
-	// Multiplier for prediction score contribution. Normalized with history weight after scaling.
-	UPROPERTY(BlueprintReadWrite, EditAnywhere, Category = "Settings", meta = (ClampMin = "0.0"))
-	float PredictionWeightScale = 1.0f;
-
-	bool operator==(const FPoseSearchChannelDynamicWeightParams& Rhs) const
-	{
-		return (ChannelWeightScale == Rhs.ChannelWeightScale)
-			&& (HistoryWeightScale == Rhs.HistoryWeightScale)
-			&& (PredictionWeightScale == Rhs.PredictionWeightScale);
-	}
-	bool operator!=(const FPoseSearchChannelDynamicWeightParams& Rhs) const { return !(*this == Rhs); }
-};
-
-USTRUCT(BlueprintType, Category = "Animation|Pose Search")
 struct POSESEARCH_API FPoseSearchWeightParams
 {
 	GENERATED_BODY()
 
 	UPROPERTY(BlueprintReadWrite, EditAnywhere, Category = "Settings")
-	FPoseSearchChannelWeightParams PoseWeight;
+	TArray<FPoseSearchChannelWeightParams> ChannelWeights;
 
-	UPROPERTY(BlueprintReadWrite, EditAnywhere, Category = "Settings")
-	FPoseSearchChannelWeightParams TrajectoryWeight;
-
-	FPoseSearchWeightParams();
+	const FPoseSearchChannelWeightParams* GetChannelWeights(int32 ChannelIdx) const;
 };
 
-USTRUCT(BlueprintType, Category="Animation|Pose Search")
-struct POSESEARCH_API FPoseSearchDynamicWeightParams
-{
-	GENERATED_BODY()
 
-	UPROPERTY(BlueprintReadWrite, EditAnywhere, Category = "Settings")
-	FPoseSearchChannelDynamicWeightParams PoseDynamicWeights;
-
-	UPROPERTY(BlueprintReadWrite, EditAnywhere, Category = "Settings")
-	FPoseSearchChannelDynamicWeightParams TrajectoryDynamicWeights;
-
-	UPROPERTY(BlueprintReadWrite, EditAnywhere, Category = "Settings")
-	bool bDebugDisableWeights = false;
-
-	bool operator==(const FPoseSearchDynamicWeightParams& Rhs) const
-	{
-		return (PoseDynamicWeights == Rhs.PoseDynamicWeights)
-			&& (TrajectoryDynamicWeights == Rhs.TrajectoryDynamicWeights)
-			&& (bDebugDisableWeights == Rhs.bDebugDisableWeights);
-	}
-	bool operator!=(const FPoseSearchDynamicWeightParams& Rhs) const { return !(*this == Rhs); }
-};
-
 USTRUCT()
 struct POSESEARCH_API FPoseSearchWeights
@@ -609,7 +767,7 @@
 	TArray<float> Weights;
 
 	bool IsInitialized() const { return !Weights.IsEmpty(); }
-	void Init(const FPoseSearchWeightParams& WeightParams, const UPoseSearchSchema* Schema, const FPoseSearchDynamicWeightParams& RuntimeParams);
+	void Init(const FPoseSearchWeightParams& WeightParams, const UPoseSearchSchema* Schema);
 };
 
 USTRUCT()
@@ -618,8 +776,8 @@
 	GENERATED_BODY()
 
 public:
-	// Check if the database or runtime weight parameters have changed and then computes and caches new group weights
-	void Update(const FPoseSearchDynamicWeightParams& DynamicWeights, const UPoseSearchDatabase * Database);
+	// Computes and caches new group weights whenever the database changes
+	void Update(const UPoseSearchDatabase * Database);
 
 	const FPoseSearchWeights* GetGroupWeights (int32 WeightsGroupIdx) const;
 	
@@ -628,8 +786,5 @@
 	TWeakObjectPtr<const UPoseSearchDatabase> Database = nullptr;
 
 	UPROPERTY(Transient)
-	FPoseSearchDynamicWeightParams DynamicWeights;
-
-	UPROPERTY(Transient)
 	FPoseSearchWeights ComputedDefaultGroupWeights;
 
@@ -642,16 +797,6 @@
 #endif
 };
 
-UENUM()
-enum class EPoseSearchMirrorOption : int32
-{
-	UnmirroredOnly UMETA(DisplayName="Original Only"),
-	MirroredOnly UMETA(DisplayName="Mirrored Only"),
-	UnmirroredAndMirrored UMETA(DisplayName="Original and Mirrored"),
-	
-	Num UMETA(Hidden),
-	Invalid = Num UMETA(Hidden)
-};
 
 /** An entry in a UPoseSearchDatabase. */
 USTRUCT(BlueprintType, Category = "Animation|Pose Search")
@@ -761,7 +906,7 @@
 public:
 	// Motion Database Config asset to use with this database.
 	UPROPERTY(EditAnywhere, BlueprintReadOnly, Category="Database", DisplayName="Config")
-	const UPoseSearchSchema* Schema;
+	TObjectPtr<const UPoseSearchSchema> Schema = nullptr;
 
 	UPROPERTY(EditAnywhere, Category = "Database")
 	FPoseSearchWeightParams DefaultWeights;
@@ -870,12 +1015,73 @@
 	bool IsDerivedDataValid();
 };
 
+
+//////////////////////////////////////////////////////////////////////////
+// Sequence metadata
+
+/** Animation metadata object for indexing a single animation. */
+UCLASS(BlueprintType, Category = "Animation|Pose Search", Experimental)
+class POSESEARCH_API UPoseSearchSequenceMetaData : public UAnimMetaData
+{
+	GENERATED_BODY()
+public:
+
+	UPROPERTY(EditAnywhere, Category = "Settings")
+	TObjectPtr<const UPoseSearchSchema> Schema = nullptr;
+
+	UPROPERTY(EditAnywhere, Category = "Settings")
+	FFloatInterval SamplingRange = FFloatInterval(0.0f, 0.0f);
+
+	UPROPERTY(EditAnywhere, Category = "Settings")
+	FPoseSearchExtrapolationParameters ExtrapolationParameters;
+
+	UPROPERTY()
+	FPoseSearchIndex SearchIndex;
+
+	bool IsValidForIndexing() const;
+	bool IsValidForSearch() const;
+
+public: // UObject
+	virtual void PreSave(FObjectPreSaveContext ObjectSaveContext) override;
+};
+
+
+//////////////////////////////////////////////////////////////////////////
+// Feature vector reader and builder
+
+namespace UE::PoseSearch {
+
+/** Helper object for extracting features from a float buffer according to the feature vector layout. */
+class POSESEARCH_API FFeatureVectorReader
+{
+public:
+	void Init(const FPoseSearchFeatureVectorLayout* Layout);
+	void SetValues(TArrayView<const float> Values);
+	bool IsValid() const;
+
+	bool GetTransform(FPoseSearchFeatureDesc Feature, FTransform* OutTransform) const;
+	bool GetPosition(FPoseSearchFeatureDesc Feature, FVector* OutPosition) const;
+	bool GetRotation(FPoseSearchFeatureDesc Feature, FQuat* OutRotation) const;
+	bool GetForwardVector(FPoseSearchFeatureDesc Feature, FVector* OutForwardVector) const;
+	bool GetLinearVelocity(FPoseSearchFeatureDesc Feature, FVector* OutLinearVelocity) const;
+	bool GetAngularVelocity(FPoseSearchFeatureDesc Feature, FVector* OutAngularVelocity) const;
+	bool GetVector(FPoseSearchFeatureDesc Feature, FVector* OutVector) const;
+
+	const FPoseSearchFeatureVectorLayout* GetLayout() const { return Layout; }
+
+private:
+	const FPoseSearchFeatureVectorLayout* Layout = nullptr;
+	TArrayView<const float> Values;
+};
+
+} // namespace UE::PoseSearch
+
+
 /** 
 * Helper object for writing features into a float buffer according to a feature vector layout.
 * Keeps track of which features are present, allowing the feature vector to be built up piecemeal.
 * FFeatureVectorBuilder is used to build search queries at runtime and for adding samples during search index construction.
 */
-
 USTRUCT(BlueprintType, Category = "Animation|Pose Search")
 struct POSESEARCH_API FPoseSearchFeatureVectorBuilder
 {
@@ -900,8 +1106,6 @@
 	void SetAngularVelocity(FPoseSearchFeatureDesc Feature, const FTransform& Transform, const FTransform& PrevTransform, float DeltaTime);
 	void SetAngularVelocity(FPoseSearchFeatureDesc Feature, const FTransform& NextTransform, const FTransform& Transform, const FTransform& PrevTransform, float DeltaTime);
 	void SetVector(FPoseSearchFeatureDesc Feature, const FVector& Vector);
-	void BuildFromTrajectory(const FTrajectorySampleRange& Trajectory);
-	bool TrySetPoseFeatures(UE::PoseSearch::FPoseHistory* History, const FBoneContainer& BoneContainer);
 
 	void CopyFromSearchIndex(const FPoseSearchIndex& SearchIndex, int32 PoseIdx);
 	void CopyFeature(const FPoseSearchFeatureVectorBuilder& OtherBuilder, int32 FeatureIdx);
@@ -913,12 +1117,11 @@
 	bool IsComplete() const;
 	bool IsCompatible(const FPoseSearchFeatureVectorBuilder& OtherBuilder) const;
 
+	const TBitArray<>& GetFeaturesAdded() const;
+
 	void Normalize(const FPoseSearchIndex& ForSearchIndex);
 
 private:
-	void BuildFromTrajectoryTimeBased(const FTrajectorySampleRange& Trajectory);
-	void BuildFromTrajectoryDistanceBased(const FTrajectorySampleRange& Trajectory);
-
 	UPROPERTY(Transient)
 	TWeakObjectPtr<const UPoseSearchSchema> Schema = nullptr;
 
@@ -929,19 +1132,26 @@
 };
 
 
-namespace UE { namespace PoseSearch {
 
-UENUM()
-enum class EPoseHistoryRootUpdateMode : int32
-{
-	RootMotionDelta,
-	ComponentTransformDelta,
-};
+//////////////////////////////////////////////////////////////////////////
+// Pose history
 
-/** Records poses over time in a ring buffer. FFeatureVectorBuilder uses this to sample from the present or past poses according to the search schema. */
+namespace UE::PoseSearch {
+
+/**
+* Records poses over time in a ring buffer.
+* FFeatureVectorBuilder uses this to sample from the present or past poses according to the search schema.
+*/
 class POSESEARCH_API FPoseHistory
 {
 public:
+
+	enum class ERootUpdateMode
+	{
+		RootMotionDelta,
+		ComponentTransformDelta,
+	};
+
 	void Init(int32 InNumPoses, float InTimeHorizon);
 	void Init(const FPoseHistory& History);
 	bool TrySamplePose(float SecondsAgo, const FReferenceSkeleton& RefSkeleton, const TArray<FBoneIndexType>& RequiredBones);
@@ -951,7 +1161,7 @@
 		const FPoseContext& PoseContext, 
 		FTransform ComponentTransform, 
 		FText* OutError,
-		EPoseHistoryRootUpdateMode UpdateMode = EPoseHistoryRootUpdateMode::RootMotionDelta);
+		ERootUpdateMode UpdateMode = ERootUpdateMode::RootMotionDelta);
 
 	float GetSampleTimeInterval() const;
 	TArrayView<const FTransform> GetLocalPoseSample() const { return SampledLocalPose; }
@@ -997,66 +1207,34 @@
 };
 
 
-/** Helper object for extracting features from a float buffer according to the feature vector layout. */
-class POSESEARCH_API FFeatureVectorReader
-{
-public:
-	void Init(const FPoseSearchFeatureVectorLayout* Layout);
-	void SetValues(TArrayView<const float> Values);
-	bool IsValid() const;
 
-	bool GetTransform(FPoseSearchFeatureDesc Feature, FTransform* OutTransform) const;
-	bool GetPosition(FPoseSearchFeatureDesc Feature, FVector* OutPosition) const;
-	bool GetRotation(FPoseSearchFeatureDesc Feature, FQuat* OutRotation) const;
-	bool GetForwardVector(FPoseSearchFeatureDesc Feature, FVector* OutForwardVector) const;
-	bool GetLinearVelocity(FPoseSearchFeatureDesc Feature, FVector* OutLinearVelocity) const;
-	bool GetAngularVelocity(FPoseSearchFeatureDesc Feature, FVector* OutAngularVelocity) const;
-	bool GetVector(FPoseSearchFeatureDesc Feature, FVector* OutVector) const;
-
-	const FPoseSearchFeatureVectorLayout* GetLayout() const { return Layout; }
-
-private:
-	const FPoseSearchFeatureVectorLayout* Layout = nullptr;
-	TArrayView<const float> Values;
-};
-
-
 //////////////////////////////////////////////////////////////////////////
-// Main PoseSearch API
+// Debug visualization
 
 enum class EDebugDrawFlags : uint32
 {
-	None			    = 0,
+	None = 0,
 
 	// Draw the entire search index as a point cloud
-	DrawSearchIndex     = 1 << 0,
+	DrawSearchIndex = 1 << 0,
 
-	// Draw pose features for each pose vector
-	IncludePose         = 1 << 1,
-
-	// Draw trajectory features for each pose vector
-	IncludeTrajectory   = 1 << 2,
-
-	// Draw all pose vector features
-	IncludeAllFeatures  = IncludePose | IncludeTrajectory,
-
 	/**
 	 * Keep rendered data until the next call to FlushPersistentDebugLines().
 	 * Combine with DrawSearchIndex to draw the search index only once.
 	 */
-	Persistent = 1 << 3,
+	Persistent = 1 << 1,
 	
 	// Label samples with their indices
-	DrawSampleLabels = 1 << 4,
+	DrawSampleLabels = 1 << 2,
 
 	// Fade colors
-	DrawSamplesWithColorGradient = 1 << 5,
+	DrawSamplesWithColorGradient = 1 << 3,
 
-	// Draw Bone Names
-	DrawBoneNames = 1 << 6,
+	// Label bone names
+	DrawBoneNames = 1 << 4,
 
 	// Draws simpler shapes to improve performance
-	DrawFast = 1 << 7,
+	DrawFast = 1 << 5,
 };
 ENUM_CLASS_FLAGS(EDebugDrawFlags);
 
@@ -1065,7 +1243,8 @@
 	const UWorld* World = nullptr;
 	const UPoseSearchDatabase* Database = nullptr;
 	const UPoseSearchSequenceMetaData* SequenceMetaData = nullptr;
-	EDebugDrawFlags Flags = EDebugDrawFlags::IncludeAllFeatures | EDebugDrawFlags::DrawBoneNames;
+	EDebugDrawFlags Flags = EDebugDrawFlags::DrawBoneNames;
+	uint32 ChannelMask = (uint32)-1;
 
 	float DefaultLifeTime = 5.0f;
 	float PointSize = 1.0f;
@@ -1099,7 +1278,62 @@
 
 };
 
-struct FSearchResult
+/**
+* Visualize pose search debug information
+*
+* @param DrawParams		Visualization options
+*/
+POSESEARCH_API void Draw(const FDebugDrawParams& DrawParams);
+
+
+
+//////////////////////////////////////////////////////////////////////////
+// Index building
+
+/**
+* Creates a pose search index for an animation sequence
+*
+* @param Sequence			The input sequence create a search index for
+* @param SequenceMetaData	The input sequence indexing info and output search index
+*
+* @return Whether the index was built successfully
+*/
+POSESEARCH_API bool BuildIndex(const UAnimSequence* Sequence, UPoseSearchSequenceMetaData* SequenceMetaData);
+
+
+/**
+* Creates a pose search index for a collection of animations
+*
+* @param Database	The input collection of animations and output search index
+*
+* @return Whether the index was built successfully
+*/
+POSESEARCH_API bool BuildIndex(UPoseSearchDatabase* Database, FPoseSearchIndex& OutSearchIndex);
+
+
+//////////////////////////////////////////////////////////////////////////
+// Query building
+
+struct POSESEARCH_API FQueryBuildingContext
 {
+	FQueryBuildingContext(FPoseSearchFeatureVectorBuilder& InQuery) : Query(InQuery) {}
+
+	FPoseSearchFeatureVectorBuilder& Query;
+
+	const UPoseSearchSchema* Schema = nullptr;
+	FPoseHistory* History = nullptr;
+	const FTrajectorySampleRange* Trajectory = nullptr;
+
+	bool IsInitialized () const;
+};
+
+bool POSESEARCH_API BuildQuery(FQueryBuildingContext& QueryBuildingContext);
+
+
+//////////////////////////////////////////////////////////////////////////
+// Search
+
+struct POSESEARCH_API FSearchResult
+{
 	FPoseCost PoseCost;
 	int32 PoseIdx = INDEX_NONE;
@@ -1130,34 +1364,6 @@
 	float MirrorMismatchCost = 0.0f;
 };
 
-/**
-* Visualize pose search debug information
-*
-* @param DrawParams		Visualization options
-*/
-POSESEARCH_API void Draw(const FDebugDrawParams& DrawParams);
-
 
 /**
-* Creates a pose search index for an animation sequence
-* 
-* @param Sequence			The input sequence create a search index for
-* @param SequenceMetaData	The input sequence indexing info and output search index
-* 
-* @return Whether the index was built successfully
-*/
-POSESEARCH_API bool BuildIndex(const UAnimSequence* Sequence, UPoseSearchSequenceMetaData* SequenceMetaData);
-
-
-/**
-* Creates a pose search index for a collection of animations
-* 
-* @param Database	The input collection of animations and output search index
-* 
-* @return Whether the index was built successfully
-*/
-POSESEARCH_API bool BuildIndex(UPoseSearchDatabase* Database, FPoseSearchIndex& OutSearchIndex);
-
-
-/**
 * Performs a pose search on a UPoseSearchDatabase.
@@ -1169,6 +1375,9 @@
 POSESEARCH_API FSearchResult Search(FSearchContext& SearchContext);
 
 
+//////////////////////////////////////////////////////////////////////////
+// Pose comparison
+
 /**
 * Evaluate pose comparison metric between a pose in the search index and an input query
 *
@@ -1214,4 +1423,4 @@
 */
 POSESEARCH_API FPoseCost ComparePoses(int32 PoseIdx, FSearchContext& SearchContext, FPoseCostDetails& OutPoseCostDetails);
 
-}} // namespace UE::PoseSearch
+} // namespace UE::PoseSearch
diff --git a/Engine/Plugins/Experimental/Animation/PoseSearch/Source/Runtime/Public/PoseSearch/PoseSearchFeatureChannels.h b/Engine/Plugins/Experimental/Animation/PoseSearch/Source/Runtime/Public/PoseSearch/PoseSearchFeatureChannels.h
new file mode 100644
--- /dev/null
+++ b/Engine/Plugins/Experimental/Animation/PoseSearch/Source/Runtime/Public/PoseSearch/PoseSearchFeatureChannels.h
@@ -0,0 +1,99 @@
+// Copyright Epic Games, Inc. All Rights Reserved.
+
+#pragma once
+
+#include "CoreMinimal.h"
+#include "PoseSearch.h"
+#include "Interfaces/Interface_BoneReferenceSkeletonProvider.h"
+
+#include "PoseSearchFeatureChannels.generated.h"
+
+USTRUCT()
+struct FPoseSearchPoseFeatureInfo
+{
+	GENERATED_BODY()
+
+	UPROPERTY()
+	int8 SchemaBoneIdx = 0;
+};
+
+//////////////////////////////////////////////////////////////////////////
+// UPoseSearchFeatureChannel_Pose
+UCLASS(BlueprintType, EditInlineNew)
+class POSESEARCH_API UPoseSearchFeatureChannel_Pose : public UPoseSearchFeatureChannel
+{
+	GENERATED_BODY()
+
+public:
+
+	virtual ~UPoseSearchFeatureChannel_Pose() {}
+
+	UPROPERTY(EditAnywhere, Category = "Settings")
+	TArray<FPoseSearchBone> SampledBones;
+
+	UPROPERTY(EditAnywhere, Category = "Settings")
+	TArray<float> SampleTimes;
+
+	UPROPERTY()
+	TArray<FPoseSearchPoseFeatureInfo> FeatureParams;
+
+	// UObject interface
+	virtual void PreSave(FObjectPreSaveContext ObjectSaveContext) override;
+
+	// UPoseSearchFeatureChannel interface
+	virtual void InitializeSchema(UE::PoseSearch::FSchemaInitializer& Initializer) override;
+	virtual void IndexAsset(const UE::PoseSearch::IAssetIndexer& Indexer, UE::PoseSearch::FAssetIndexingOutput& IndexingOutput) const override;
+	virtual FFloatRange GetHorizonRange(EPoseSearchFeatureDomain Domain) const override;
+	virtual TArrayView<const float> GetSampleOffsets () const override { return SampleTimes; }
+	virtual void GenerateDDCKey(FBlake3& InOutKeyHasher) const override;
+	virtual bool BuildQuery(UE::PoseSearch::FQueryBuildingContext& Context) const override;
+	virtual void DebugDraw(const UE::PoseSearch::FDebugDrawParams& DrawParams, const UE::PoseSearch::FFeatureVectorReader& Reader) const override;
+
+protected:
+	virtual void AddPoseFeatures(const UE::PoseSearch::IAssetIndexer& Indexer, int32 SampleIdx, FPoseSearchFeatureVectorBuilder& FeatureVector) const;
+};
+
+
+
+//////////////////////////////////////////////////////////////////////////
+// UPoseSearchFeatureChannel_Trajectory
+UCLASS(BlueprintType, EditInlineNew)
+class POSESEARCH_API UPoseSearchFeatureChannel_Trajectory : public UPoseSearchFeatureChannel
+{
+	GENERATED_BODY()
+
+public:
+
+	virtual ~UPoseSearchFeatureChannel_Trajectory() {}
+
+	UPROPERTY(EditAnywhere, Category = "Settings")
+	bool bUseLinearVelocities = true;
+
+	UPROPERTY(EditAnywhere, Category = "Settings")
+	bool bUsePositions = true;
+
+	UPROPERTY(EditAnywhere, Category = "Settings")
+	bool bUseFacingDirections = false;
+
+	UPROPERTY(EditAnywhere, Category="Settings")
+	EPoseSearchFeatureDomain Domain = EPoseSearchFeatureDomain::Time;
+
+	UPROPERTY(EditAnywhere, Category = "Settings")
+	TArray<float> SampleOffsets;
+
+	// UObject interface
+	virtual void PreSave(FObjectPreSaveContext ObjectSaveContext) override;
+
+	// UPoseSearchFeatureChannel interface
+	virtual void InitializeSchema(UE::PoseSearch::FSchemaInitializer& Initializer) override;
+	virtual void IndexAsset(const UE::PoseSearch::IAssetIndexer& Indexer, UE::PoseSearch::FAssetIndexingOutput& IndexingOutput) const override;
+	virtual FFloatRange GetHorizonRange(EPoseSearchFeatureDomain Domain) const override;
+	virtual TArrayView<const float> GetSampleOffsets () const override { return SampleOffsets; }
+	virtual void GenerateDDCKey(FBlake3& InOutKeyHasher) const override;
+	virtual bool BuildQuery(UE::PoseSearch::FQueryBuildingContext& Context) const override;
+	virtual void DebugDraw(const UE::PoseSearch::FDebugDrawParams& DrawParams, const UE::PoseSearch::FFeatureVectorReader& Reader) const override;
+
+protected:
+	virtual void IndexTimeFeatures(const UE::PoseSearch::IAssetIndexer& Indexer, int32 SampleIdx, FPoseSearchFeatureVectorBuilder& FeatureVector) const;
+	virtual void IndexDistanceFeatures(const UE::PoseSearch::IAssetIndexer& Indexer, int32 SampleIdx, FPoseSearchFeatureVectorBuilder& FeatureVector) const;
+};
\ No newline at end of file
diff --git a/Engine/Plugins/Experimental/Animation/PoseSearch/Source/Runtime/Public/PoseSearch/PoseSearchLibrary.h b/Engine/Plugins/Experimental/Animation/PoseSearch/Source/Runtime/Public/PoseSearch/PoseSearchLibrary.h
--- a/Engine/Plugins/Experimental/Animation/PoseSearch/Source/Runtime/Public/PoseSearch/PoseSearchLibrary.h
+++ b/Engine/Plugins/Experimental/Animation/PoseSearch/Source/Runtime/Public/PoseSearch/PoseSearchLibrary.h
@@ -27,10 +27,6 @@ struct POSESEARCH_API FMotionMatchingSettings
 {
 	GENERATED_BODY()
 
-	// Dynamic weights for influencing pose selection
-	UPROPERTY(EditAnywhere, BlueprintReadWrite, Category=Settings, meta=(PinHiddenByDefault))
-	FPoseSearchDynamicWeightParams Weights;
-
 	// Time in seconds to blend out to the new pose. Uses inertial blending and requires an Inertialization node after this node.
 	UPROPERTY(EditAnywhere, BlueprintReadWrite, Category=Settings, meta=(ClampMin="0"))
 	float BlendTime = 0.2f;
@@ -73,9 +69,6 @@ struct POSESEARCH_API FMotionMatchingState
 	// state back to the default state.
 	void AdjustAssetTime(float AssetTime);
 
-	// Adds trajectory prediction and history information to ComposedQuery
-	void ComposeQuery(const UPoseSearchDatabase* Database, const FTrajectorySampleRange& Trajectory);
-
 	// Internally stores the 'jump' to a new pose/sequence index and asset time for evaluation
 	void JumpToPose(const FAnimationUpdateContext& Context, const FMotionMatchingSettings& Settings, const UE::PoseSearch::FSearchResult& Result);
 
