motion matching: performing searches with KDTree in PCA space

diff --git a/Engine/Plugins/Experimental/Animation/PoseSearch/Source/Runtime/Private/KDTree.cpp b/Engine/Plugins/Experimental/Animation/PoseSearch/Source/Runtime/Private/KDTree.cpp
new file mode 100644
--- /dev/null
+++ b/Engine/Plugins/Experimental/Animation/PoseSearch/Source/Runtime/Private/KDTree.cpp
@@ -0,0 +1,54 @@
+// Copyright Epic Games, Inc. All Rights Reserved.
+
+#include "PoseSearch/KDTree.h"
+
+namespace UE { namespace PoseSearch
+{
+
+FKDTree::FKDTree(int32 Count, int32 Dim, const float* Data, int32 MaxLeafSize)
+: DataSrc(Count, Dim, Data)
+, Impl(nullptr)
+{
+}
+
+FKDTree::FKDTree()
+: DataSrc(0, 0, nullptr)
+, Impl(nullptr)
+{
+}
+
+FKDTree::~FKDTree()
+{
+}
+
+FKDTree::FKDTree(const FKDTree& r)
+: DataSrc(r.DataSrc.PointCount, r.DataSrc.PointDim, r.DataSrc.Data)
+, Impl(nullptr)
+{
+}
+
+FKDTree& FKDTree::operator=(const FKDTree& r)
+{
+	this->~FKDTree();
+	new(this)FKDTree(r);
+	return *this;
+}
+
+void FKDTree::Construct(int32 Count, int32 Dim, const float* Data, int32 MaxLeafSize)
+{
+	this->~FKDTree();
+	new(this)FKDTree(Count, Dim, Data, MaxLeafSize);
+}
+
+bool FKDTree::FindNeighbors(KNNResultSet& Result, const float* Query) const
+{
+	check(false); // unimplemented
+	return false;
+}
+
+FArchive& Serialize(FArchive& Ar, FKDTree& KDTree, const float* KDTreeData)
+{
+	return Ar;
+}
+
+} }
diff --git a/Engine/Plugins/Experimental/Animation/PoseSearch/Source/Runtime/Private/PoseSearch.cpp b/Engine/Plugins/Experimental/Animation/PoseSearch/Source/Runtime/Private/PoseSearch.cpp
--- a/Engine/Plugins/Experimental/Animation/PoseSearch/Source/Runtime/Private/PoseSearch.cpp
+++ b/Engine/Plugins/Experimental/Animation/PoseSearch/Source/Runtime/Private/PoseSearch.cpp
@@ -45,6 +45,12 @@ IMPLEMENT_ANIMGRAPH_MESSAGE(UE::PoseSearch::IPoseHistoryProvider);
 
 DEFINE_LOG_CATEGORY(LogPoseSearch);
 
+DECLARE_STATS_GROUP(TEXT("PoseSearch"), STATGROUP_PoseSearch, STATCAT_Advanced);
+DECLARE_CYCLE_STAT_EXTERN(TEXT("Search Brute Force"), STAT_PoseSearchBruteForce, STATGROUP_PoseSearch, );
+DECLARE_CYCLE_STAT_EXTERN(TEXT("Search PCA/KNN"), STAT_PoseSearchPCAKNN, STATGROUP_PoseSearch, );
+DEFINE_STAT(STAT_PoseSearchBruteForce);
+DEFINE_STAT(STAT_PoseSearchPCAKNN);
+
 namespace UE { namespace PoseSearch {
 
 //////////////////////////////////////////////////////////////////////////
@@ -65,6 +71,18 @@ namespace UE { namespace PoseSearch {
 	constexpr EParallelForFlags ParallelForFlags = EParallelForFlags::None;
 #endif // UE_POSE_SEARCH_FORCE_SINGLE_THREAD
 
+using RowMajorVector = Eigen::Matrix<float, 1, Eigen::Dynamic, Eigen::RowMajor>;
+using RowMajorVectorMap = Eigen::Map<RowMajorVector, Eigen::RowMajor>;
+using RowMajorVectorMapConst = Eigen::Map<const RowMajorVector, Eigen::RowMajor>;
+
+using RowMajorMatrix = Eigen::Matrix<float, Eigen::Dynamic, Eigen::Dynamic, Eigen::RowMajor>;
+using RowMajorMatrixMap = Eigen::Map<RowMajorMatrix, Eigen::RowMajor>;
+using RowMajorMatrixMapConst = Eigen::Map<const RowMajorMatrix, Eigen::RowMajor>;
+
+using ColMajorMatrix = Eigen::Matrix<float, Eigen::Dynamic, Eigen::Dynamic, Eigen::ColMajor>;
+using ColMajorMatrixMap = Eigen::Map<ColMajorMatrix, Eigen::ColMajor>;
+using ColMajorMatrixMapConst = Eigen::Map<const ColMajorMatrix, Eigen::ColMajor>;
+
 static bool IsSamplingRangeValid(FFloatInterval Range)
 {
 	return Range.IsValid() && (Range.Min >= 0.0f);
@@ -1203,6 +1221,10 @@ const FString UPoseSearchDatabase::GetSourceAssetName(const FPoseSearchIndexAsse
 	}
 }
 
+int32 UPoseSearchDatabase::GetNumberOfPrincipalComponents() const
+{
+	return FMath::Min<int32>(NumberOfPrincipalComponents, Schema->Layout.NumFloats);
+}
 
 bool UPoseSearchDatabase::IsValidForIndexing() const
 {
@@ -1554,6 +1576,12 @@ bool UPoseSearchDatabase::TryInitSearchIndexAssets(FPoseSearchIndex& OutSearchIn
 		}
 	}
 
+	// @todo: change the above for loops fill OutSearchIndex.Assets already in ascnding group order
+	// sorting by ascending SourceGroupIdx
+	OutSearchIndex.Assets.Sort([](const FPoseSearchIndexAsset& InOne, const FPoseSearchIndexAsset& InTwo)
+	{
+		return InOne.SourceGroupIdx < InTwo.SourceGroupIdx;
+	});
 
 	if (bAnyMirrored && !Schema->MirrorDataTable)
 	{
@@ -4267,6 +4295,198 @@ static void PreprocessSearchIndex(FPoseSearchIndex* SearchIndex)
 	}
 }
 
+static void PreprocessGroupSearchIndexWeights(FGroupSearchIndex& GroupSearchIndex, const UPoseSearchDatabase* Database)
+{
+	const FPoseSearchWeightParams& WeightParams = GroupSearchIndex.GroupIndex == INDEX_NONE ? Database->DefaultWeights : Database->Groups[GroupSearchIndex.GroupIndex].Weights;
+	FPoseSearchWeights Weights;
+	Weights.Init(WeightParams, Database->Schema);
+	GroupSearchIndex.Weights = Weights.Weights;
+}
+
+// it calcualtes Mean and PCAProjectionMatrix
+static void PreprocessGroupSearchIndexPCAData(FGroupSearchIndex& GroupSearchIndex, const UPoseSearchDatabase* Database, const float* GroupValues, float* GroupPCAValues)
+{
+	// binding SearchIndex.Values and SearchIndex.PCAValues Eigen row major matrix maps
+	const int32 NumDimensions = Database->Schema->Layout.NumFloats;
+	const int32 NumGroupPoses = GroupSearchIndex.EndPoseIndex - GroupSearchIndex.StartPoseIndex;
+
+	const RowMajorVectorMapConst MapWeights(GroupSearchIndex.Weights.GetData(), 1, NumDimensions);
+	const RowMajorMatrixMapConst MapGroupValues(GroupValues, NumGroupPoses, NumDimensions);
+	const RowMajorMatrix WeightedGroupValues = MapGroupValues.array().rowwise() * MapWeights.array();
+	const uint32 NumberOfPrincipalComponents = Database->GetNumberOfPrincipalComponents();
+	RowMajorMatrixMap MapGroupPCAValues(GroupPCAValues, NumGroupPoses, NumberOfPrincipalComponents);
+
+	// calculating the mean
+	GroupSearchIndex.Mean.SetNumZeroed(NumDimensions);
+	RowMajorVectorMap Mean(GroupSearchIndex.Mean.GetData(), 1, NumDimensions);
+	Mean = WeightedGroupValues.colwise().mean();
+
+	// use the mean to center the data points
+	const RowMajorMatrix CenteredGroupValues = WeightedGroupValues.rowwise() - Mean;
+
+
+	// estimating the covariance matrix (with dimensionality of NumDimensions, NumDimensions)
+	// formula: https://en.wikipedia.org/wiki/Covariance_matrix#Estimation
+	// details: https://en.wikipedia.org/wiki/Estimation_of_covariance_matrices
+	const ColMajorMatrix CovariantMatrix = (CenteredGroupValues.transpose() * CenteredGroupValues) / float(NumGroupPoses - 1);
+	const Eigen::SelfAdjointEigenSolver<ColMajorMatrix> EigenSolver(CovariantMatrix);
+
+	check(EigenSolver.info() == Eigen::Success);
+
+	// validating EigenSolver results
+	const ColMajorMatrix EigenVectors = EigenSolver.eigenvectors().real();
+
+	if (Database->PoseSearchMode == EPoseSearchMode::PCAKDTree_Validate && NumberOfPrincipalComponents == NumDimensions)
+	{
+		const RowMajorVector ReciprocalWeights = MapWeights.cwiseInverse();
+		const RowMajorMatrix ProjectedGroupValues = CenteredGroupValues * EigenVectors;
+		for (Eigen::Index RowIndex = 0; RowIndex < MapGroupValues.rows(); ++RowIndex)
+		{
+			const RowMajorVector WeightedReconstructedPoint = ProjectedGroupValues.row(RowIndex) * EigenVectors.transpose() + Mean;
+			const RowMajorVector ReconstructedPoint = WeightedReconstructedPoint.array() * ReciprocalWeights.array();
+			const float Error = (ReconstructedPoint - MapGroupValues.row(RowIndex)).squaredNorm();
+			check(Error < UE_KINDA_SMALL_NUMBER);
+		}
+	}
+
+	// sorting EigenVectors by EigenValues, so we pick the most significant ones to compose our PCA projection matrix.
+	const RowMajorVector EigenValues = EigenSolver.eigenvalues().real();
+	TArray<size_t> Indexer;
+	Indexer.Reserve(NumDimensions);
+	for (size_t DimensionIndex = 0; DimensionIndex < NumDimensions; ++DimensionIndex)
+	{
+		Indexer.Push(DimensionIndex);
+	}
+	Indexer.Sort([&EigenValues](size_t a, size_t b)
+	{
+		return EigenValues[a] > EigenValues[b];
+	});
+
+	// composing the PCA projection matrix with the PCANumComponents most significant EigenVectors
+	GroupSearchIndex.PCAProjectionMatrix.SetNumZeroed(NumDimensions * NumberOfPrincipalComponents);
+	ColMajorMatrixMap PCAProjectionMatrix(GroupSearchIndex.PCAProjectionMatrix.GetData(), NumDimensions, NumberOfPrincipalComponents);
+	float AccumulatedVariance = 0.f;
+	for (size_t PCAComponentIndex = 0; PCAComponentIndex < NumberOfPrincipalComponents; ++PCAComponentIndex)
+	{
+		PCAProjectionMatrix.col(PCAComponentIndex) = EigenVectors.col(Indexer[PCAComponentIndex]);
+		AccumulatedVariance += EigenValues[Indexer[PCAComponentIndex]];
+	}
+
+	// calculating the total variance knowing that eigen values measure variance along the principal components:
+	const float TotalVariance = EigenValues.sum();
+	// and explained variance as ratio between AccumulatedVariance and TotalVariance: https://ro-che.info/articles/2017-12-11-pca-explained-variance
+	const float ExplainedVariance = TotalVariance > UE_KINDA_SMALL_NUMBER ? AccumulatedVariance / TotalVariance : 0.f;
+
+	MapGroupPCAValues = CenteredGroupValues * PCAProjectionMatrix;
+
+	if (Database->PoseSearchMode == EPoseSearchMode::PCAKDTree_Validate && NumberOfPrincipalComponents == NumDimensions)
+	{
+		const RowMajorVector ReciprocalWeights = MapWeights.cwiseInverse();
+		for (Eigen::Index RowIndex = 0; RowIndex < MapGroupValues.rows(); ++RowIndex)
+		{
+			const RowMajorVector WeightedReconstructedValues = MapGroupPCAValues.row(RowIndex) * PCAProjectionMatrix.transpose() + Mean;
+			const RowMajorVector ReconstructedValues = WeightedReconstructedValues.array() * ReciprocalWeights.array();
+			const float Error = (ReconstructedValues - MapGroupValues.row(RowIndex)).squaredNorm();
+			check(Error < UE_KINDA_SMALL_NUMBER);
+		}
+	}
+}
+
+static void PreprocessGroupSearchIndexKDTree(FGroupSearchIndex& GroupSearchIndex, const UPoseSearchDatabase* Database, const float* GroupValues, const float* GroupPCAValues)
+{
+	const int32 NumGroupPoses = GroupSearchIndex.EndPoseIndex - GroupSearchIndex.StartPoseIndex;
+	const uint32 NumberOfPrincipalComponents = Database->GetNumberOfPrincipalComponents();
+	GroupSearchIndex.KDTree.Construct(NumGroupPoses, NumberOfPrincipalComponents, GroupPCAValues, Database->KDTreeMaxLeafSize);
+
+	if (Database->PoseSearchMode == EPoseSearchMode::PCAKDTree_Validate)
+	{
+		// testing the KDTree is returning the proper searches for all the points in pca space
+		for (size_t PointIndex = 0; PointIndex < NumGroupPoses; ++PointIndex)
+		{
+			constexpr size_t NumResults = 10;
+			size_t ResultIndexes[NumResults + 1];
+			float ResultDistanceSqr[NumResults + 1];
+			FKDTree::KNNResultSet ResultSet(NumResults, ResultIndexes, ResultDistanceSqr);
+			GroupSearchIndex.KDTree.FindNeighbors(ResultSet, &GroupPCAValues[PointIndex * NumberOfPrincipalComponents]);
+
+			size_t ResultIndex = 0;
+			for (; ResultIndex < ResultSet.Num(); ++ResultIndex)
+			{
+				if (PointIndex == ResultIndexes[ResultIndex])
+				{
+					check(ResultDistanceSqr[ResultIndex] < UE_KINDA_SMALL_NUMBER);
+					break;
+				}
+			}
+			check(ResultIndex < ResultSet.Num());
+		}
+
+		// testing the KDTree is returning the proper searches for all the original points transformed in pca space
+		const int32 NumDimensions = Database->Schema->Layout.NumFloats;
+		for (size_t PointIndex = 0; PointIndex < NumGroupPoses; ++PointIndex)
+		{
+			constexpr size_t NumResults = 10;
+			size_t ResultIndexes[NumResults + 1];
+			float ResultDistanceSqr[NumResults + 1];
+			FKDTree::KNNResultSet ResultSet(NumResults, ResultIndexes, ResultDistanceSqr);
+
+			const RowMajorVectorMapConst MapGroupValues(&GroupValues[PointIndex * NumDimensions], 1, NumDimensions);
+			const RowMajorVectorMapConst MapWeights(GroupSearchIndex.Weights.GetData(), 1, NumDimensions);
+			const RowMajorVectorMapConst Mean(GroupSearchIndex.Mean.GetData(), 1, NumDimensions);
+			const ColMajorMatrixMapConst PCAProjectionMatrix(GroupSearchIndex.PCAProjectionMatrix.GetData(), NumDimensions, NumberOfPrincipalComponents);
+
+			const RowMajorMatrix WeightedGroupValues = MapGroupValues.array() * MapWeights.array();
+			const RowMajorMatrix CenteredGroupValues = WeightedGroupValues - Mean;
+			const RowMajorVector ProjectedGroupValues  = CenteredGroupValues * PCAProjectionMatrix;
+
+			GroupSearchIndex.KDTree.FindNeighbors(ResultSet, ProjectedGroupValues.data());
+
+			size_t ResultIndex = 0;
+			for (; ResultIndex < ResultSet.Num(); ++ResultIndex)
+			{
+				if (PointIndex == ResultIndexes[ResultIndex])
+				{
+					check(ResultDistanceSqr[ResultIndex] < UE_KINDA_SMALL_NUMBER);
+					break;
+				}
+			}
+			check(ResultIndex < ResultSet.Num());
+		}
+	}
+}
+
+static void PreprocessGroupSearchIndex(FPoseSearchIndex& SearchIndex, const UPoseSearchDatabase* Database)
+{
+	const uint32 NumberOfPrincipalComponents = Database->GetNumberOfPrincipalComponents();
+	if (NumberOfPrincipalComponents > 0)
+	{
+		const int32 NumDimensions = Database->Schema->Layout.NumFloats;
+
+		// preallocating the PCAValues for all the groups
+		SearchIndex.PCAValues.Reset();
+		SearchIndex.PCAValues.AddZeroed(SearchIndex.NumPoses * NumberOfPrincipalComponents);
+		
+		// original serial for loop: 
+		// for (int32 i = 0; i != SearchIndex.Groups.Num(); ++i)
+		ParallelFor(SearchIndex.Groups.Num(), [&](int32 i)
+		{
+			FGroupSearchIndex& GroupSearchIndex = SearchIndex.Groups[i];
+			
+			const float* GroupValues = SearchIndex.Values.GetData() + GroupSearchIndex.StartPoseIndex * NumDimensions;
+			float* GroupPCAValues = SearchIndex.PCAValues.GetData() + GroupSearchIndex.StartPoseIndex * NumberOfPrincipalComponents;
+
+			PreprocessGroupSearchIndexWeights(GroupSearchIndex, Database);
+			PreprocessGroupSearchIndexPCAData(GroupSearchIndex, Database, GroupValues, GroupPCAValues);
+			PreprocessGroupSearchIndexKDTree(GroupSearchIndex, Database, GroupValues, GroupPCAValues);
+		});
+	}
+	else
+	{
+		// we don't need groups
+		SearchIndex.Groups.Reset();
+	}
+}
+
 bool BuildIndex(const UAnimSequence* Sequence, UPoseSearchSequenceMetaData* SequenceMetaData)
 {
 	check(Sequence);
@@ -4534,25 +4754,53 @@ bool FDatabaseIndexingContext::IndexAssets()
 
 void FDatabaseIndexingContext::JoinIndex()
 {
-	// Write index info to sequence and count up total poses and storage required
+	// Write index info to asset and count up total poses and storage required
 	int32 TotalPoses = 0;
 	int32 TotalFloats = 0;
-	for (int32 AssetIdx = 0; AssetIdx != SearchIndex->Assets.Num(); ++AssetIdx)
+	
+	SearchIndex->Groups.Reset();
+
+	if (SearchIndex->Assets.Num() > 0)
 	{
-		const FAssetIndexer::FOutput& Output = Indexers[AssetIdx].Output;
+		SearchIndex->Groups.AddDefaulted();
+		SearchIndex->Groups.Last().GroupIndex = SearchIndex->Assets[0].SourceGroupIdx;
 
-		FPoseSearchIndexAsset& SearchIndexAsset = SearchIndex->Assets[AssetIdx];
-		SearchIndexAsset.NumPoses = Output.NumIndexedPoses;
-		SearchIndexAsset.FirstPoseIdx = TotalPoses;
+		for (int32 AssetIdx = 0; AssetIdx != SearchIndex->Assets.Num(); ++AssetIdx)
+		{
+			const FAssetIndexer::FOutput& Output = Indexers[AssetIdx].Output;
 
-		TotalPoses += Output.NumIndexedPoses;
-		TotalFloats += Output.FeatureVectorTable.Num();
-	}
+			FPoseSearchIndexAsset& SearchIndexAsset = SearchIndex->Assets[AssetIdx];
+			
+			if (SearchIndexAsset.SourceGroupIdx != SearchIndex->Groups.Last().GroupIndex)
+			{
+				// making sure groups are sorted correctly in ascending order
+				check(SearchIndexAsset.SourceGroupIdx > SearchIndex->Groups.Last().GroupIndex);
+				// finalizing the previous group before adding a new one
+				SearchIndex->Groups.Last().EndPoseIndex = TotalPoses;
+
+				SearchIndex->Groups.AddDefaulted();
+				SearchIndex->Groups.Last().GroupIndex = SearchIndexAsset.SourceGroupIdx;
+				SearchIndex->Groups.Last().StartPoseIndex = TotalPoses;
+			}
+			
+			SearchIndexAsset.NumPoses = Output.NumIndexedPoses;
+			SearchIndexAsset.FirstPoseIdx = TotalPoses;
 
+			TotalPoses += Output.NumIndexedPoses;
+			TotalFloats += Output.FeatureVectorTable.Num();
+		}
+		
+		// finalizing the last inserted group
+		SearchIndex->Groups.Last().EndPoseIndex = TotalPoses;
+	}
+	
+	check(TotalFloats == TotalPoses * Database->Schema->Layout.NumFloats);
+	
 	// Join animation data into a single search index
 	SearchIndex->Values.Reset(TotalFloats);
 	SearchIndex->PoseMetadata.Reset(TotalPoses);
-
+	SearchIndex->PCAValues.Reset();
+	
 	for (const FAssetIndexer& Indexer : Indexers)
 	{
 		const FAssetIndexer::FOutput& Output = Indexer.Output;
@@ -4594,6 +4842,8 @@ bool BuildIndex(UPoseSearchDatabase* Database, FPoseSearchIndex& OutSearchIndex)
 
 	PreprocessSearchIndex(&OutSearchIndex);
 
+	PreprocessGroupSearchIndex(OutSearchIndex, Database);
+	
 	return bSuccess;
 }
 
@@ -4619,31 +4869,131 @@ bool BuildQuery(FQueryBuildingContext& QueryBuildingContext)
 	return bSuccess;
 }
 
-FSearchResult Search(FSearchContext& SearchContext)
+FSearchResult SearchPCAKDTree(FSearchContext& SearchContext)
 {
+	QUICK_SCOPE_CYCLE_COUNTER(STAT_PoseSearch_PCA_KNN);
+	SCOPE_CYCLE_COUNTER(STAT_PoseSearchPCAKNN);
+
 	FSearchResult Result;
 
+	const UPoseSearchDatabase* Database = SearchContext.GetSourceDatabase();
+	check(Database);
+
+	const int32 NumDimensions = Database->Schema->Layout.NumFloats;
 	const FPoseSearchIndex* SearchIndex = SearchContext.GetSearchIndex();
-	if (!SearchIndex)
-	{
-		return Result;
-	}
+	check(SearchIndex);
 
-	if (!ensure(SearchIndex->IsValid() && !SearchIndex->IsEmpty()))
-	{
-		return Result;
-	}
+	const uint32 NumberOfPrincipalComponents = Database->GetNumberOfPrincipalComponents();
+	const uint32 KDTreeQueryNumNeighbors = FMath::Clamp<uint32>(Database->KDTreeQueryNumNeighbors, 1, SearchIndex->NumPoses);
 
-	if (!ensure(SearchContext.QueryValues.Num() == SearchIndex->Schema->Layout.NumFloats))
+	//stack allocated temporaries
+	TArrayView<size_t> ResultIndexes((size_t*)FMemory_Alloca((KDTreeQueryNumNeighbors + 1) * sizeof(size_t)), KDTreeQueryNumNeighbors + 1);
+	TArrayView<float> ResultDistanceSqr((float*)FMemory_Alloca((KDTreeQueryNumNeighbors + 1) * sizeof(float)), KDTreeQueryNumNeighbors + 1);
+	RowMajorVectorMap WeightedQueryValues((float*)FMemory_Alloca(NumDimensions * sizeof(float)), 1, NumDimensions);
+	RowMajorVectorMap CenteredQueryValues((float*)FMemory_Alloca(NumDimensions * sizeof(float)), 1, NumDimensions);
+	RowMajorVectorMap ProjectedQueryValues((float*)FMemory_Alloca(NumberOfPrincipalComponents * sizeof(float)), 1, NumberOfPrincipalComponents);
+
+	// KDTree in PCA space search
+	if (Database->PoseSearchMode == EPoseSearchMode::PCAKDTree_Validate)
 	{
-		return Result;
+		for (const FGroupSearchIndex& GroupSearchIndex : SearchIndex->Groups)
+		{
+			const RowMajorVectorMapConst MapWeights(GroupSearchIndex.Weights.GetData(), 1, NumDimensions);
+
+			// testing the KDTree is returning the proper searches for all the original points transformed in pca space
+			for (int32 PoseIdx = GroupSearchIndex.StartPoseIndex; PoseIdx < GroupSearchIndex.EndPoseIndex; ++PoseIdx)
+			{
+				FKDTree::KNNResultSet ResultSet(Database->KDTreeQueryNumNeighbors, ResultIndexes, ResultDistanceSqr);
+				TArrayView<const float> PoseValues = SearchIndex->GetPoseValues(PoseIdx);
+
+				const RowMajorVectorMapConst Mean(GroupSearchIndex.Mean.GetData(), 1, NumDimensions);
+				const ColMajorMatrixMapConst PCAProjectionMatrix(GroupSearchIndex.PCAProjectionMatrix.GetData(), NumDimensions, NumberOfPrincipalComponents);
+
+				const RowMajorVectorMapConst QueryValues(PoseValues.GetData(), 1, NumDimensions);
+				WeightedQueryValues = QueryValues.array() * MapWeights.array();
+				CenteredQueryValues.noalias() = WeightedQueryValues - Mean;
+				ProjectedQueryValues.noalias() = CenteredQueryValues * PCAProjectionMatrix;
+
+				GroupSearchIndex.KDTree.FindNeighbors(ResultSet, ProjectedQueryValues.data());
+
+				size_t ResultIndex = 0;
+				for (; ResultIndex < ResultSet.Num(); ++ResultIndex)
+				{
+					if ((PoseIdx - GroupSearchIndex.StartPoseIndex) == ResultIndexes[ResultIndex])
+					{
+						check(ResultDistanceSqr[ResultIndex] < UE_KINDA_SMALL_NUMBER);
+						break;
+					}
+				}
+				check(ResultIndex < ResultSet.Num());
+			}
+		}
 	}
 
+	// @todo: implement support for DatabaseTagQuery
 	FPoseCost BestPoseCost;
 	int32 BestPoseIdx = INDEX_NONE;
+	for (const FGroupSearchIndex& GroupSearchIndex : SearchIndex->Groups)
+	{
+		const RowMajorVectorMapConst MapWeights(GroupSearchIndex.Weights.GetData(), 1, NumDimensions);
+		FKDTree::KNNResultSet ResultSet(KDTreeQueryNumNeighbors, ResultIndexes, ResultDistanceSqr);
+
+		check(SearchContext.QueryValues.Num() == NumDimensions);
+
+		const RowMajorVectorMapConst Mean(GroupSearchIndex.Mean.GetData(), 1, NumDimensions);
+		const ColMajorMatrixMapConst PCAProjectionMatrix(GroupSearchIndex.PCAProjectionMatrix.GetData(), NumDimensions, NumberOfPrincipalComponents);
+
+		// transforming query values into PCA space to query the KDTree
+		const RowMajorVectorMapConst QueryValues(SearchContext.QueryValues.GetData(), 1, NumDimensions);
+		WeightedQueryValues = QueryValues.array() * MapWeights.array();
+		CenteredQueryValues.noalias() = WeightedQueryValues - Mean;
+		ProjectedQueryValues.noalias() = CenteredQueryValues * PCAProjectionMatrix;
+
+		GroupSearchIndex.KDTree.FindNeighbors(ResultSet, ProjectedQueryValues.data());
+
+		for (size_t ResultIndex = 0; ResultIndex < ResultSet.Num(); ++ResultIndex)
+		{
+			const int32 PoseIdx = ResultIndexes[ResultIndex] + GroupSearchIndex.StartPoseIndex;
+
+			const FPoseSearchPoseMetadata& Metadata = SearchIndex->PoseMetadata[PoseIdx];
+
+			if (EnumHasAnyFlags(Metadata.Flags, EPoseSearchPoseFlags::BlockTransition))
+			{
+				continue;
+			}
+
+			FPoseCost PoseCost = ComparePoses(PoseIdx, SearchContext, GroupSearchIndex.GroupIndex);
 
+			if (PoseCost < BestPoseCost)
+			{
+				BestPoseCost = PoseCost;
+				BestPoseIdx = PoseIdx;
+			}
+		}
+	}
+
+	Result.PoseCost = BestPoseCost;
+	Result.PoseIdx = BestPoseIdx;
+	Result.SearchIndexAsset = SearchIndex->FindAssetForPose(BestPoseIdx);
+	Result.AssetTime = SearchIndex->GetAssetTime(BestPoseIdx, Result.SearchIndexAsset);
+
+	return Result;
+}
+
+FSearchResult SearchBruteForce(FSearchContext& SearchContext)
+{
+	QUICK_SCOPE_CYCLE_COUNTER(STAT_PoseSearch_Brute_Force);
+	SCOPE_CYCLE_COUNTER(STAT_PoseSearchBruteForce);
+
+	FSearchResult Result;
+
+	const FPoseSearchIndex* SearchIndex = SearchContext.GetSearchIndex();
+	check(SearchIndex);
 	const UPoseSearchDatabase* Database = SearchContext.GetSourceDatabase();
+	check(Database);
 
+	FPoseCost BestPoseCost;
+	int32 BestPoseIdx = INDEX_NONE;
 	for (const FPoseSearchIndexAsset& Asset : SearchIndex->Assets)
 	{
 		if (Database && SearchContext.DatabaseTagQuery)
@@ -4679,6 +5029,40 @@ FSearchResult Search(FSearchContext& SearchContext)
 	Result.SearchIndexAsset = SearchIndex->FindAssetForPose(BestPoseIdx);
 	Result.AssetTime = SearchIndex->GetAssetTime(BestPoseIdx, Result.SearchIndexAsset);
 
+	return Result;
+}
+
+FSearchResult Search(FSearchContext& SearchContext)
+{
+	FSearchResult Result;
+
+	const FPoseSearchIndex* SearchIndex = SearchContext.GetSearchIndex();
+	if (!SearchIndex)
+	{
+		return Result;
+	}
+
+	if (!ensure(SearchIndex->IsValid() && !SearchIndex->IsEmpty()))
+	{
+		return Result;
+	}
+
+	if (!ensure(SearchContext.QueryValues.Num() == SearchIndex->Schema->Layout.NumFloats))
+	{
+		return Result;
+	}
+
+	const UPoseSearchDatabase* Database = SearchContext.GetSourceDatabase();
+	if (Database->PoseSearchMode != EPoseSearchMode::BruteForce)
+	{
+		Result = SearchPCAKDTree(SearchContext);
+	}
+	
+	if (Database->PoseSearchMode == EPoseSearchMode::BruteForce || Database->PoseSearchMode == EPoseSearchMode::PCAKDTree_Compare)
+	{
+		Result = SearchBruteForce(SearchContext);
+	}
+
 	SearchContext.DebugDrawParams.PoseVector = SearchContext.QueryValues;
 	SearchContext.DebugDrawParams.PoseIdx = Result.PoseIdx;
 	Draw(SearchContext.DebugDrawParams);
diff --git a/Engine/Plugins/Experimental/Animation/PoseSearch/Source/Runtime/Private/PoseSearchDerivedData.cpp b/Engine/Plugins/Experimental/Animation/PoseSearch/Source/Runtime/Private/PoseSearchDerivedData.cpp
--- a/Engine/Plugins/Experimental/Animation/PoseSearch/Source/Runtime/Private/PoseSearchDerivedData.cpp
+++ b/Engine/Plugins/Experimental/Animation/PoseSearch/Source/Runtime/Private/PoseSearchDerivedData.cpp
@@ -206,6 +206,11 @@ namespace UE::PoseSearch
 			AddDbBlendSpaceToWriter(Database.BlendSpaces[i], Writer);
 		}
 
+		Writer.Update(&Database.NumberOfPrincipalComponents, sizeof(Database.NumberOfPrincipalComponents));
+		Writer.Update(&Database.KDTreeMaxLeafSize, sizeof(Database.KDTreeMaxLeafSize));
+		Writer.Update(&Database.KDTreeQueryNumNeighbors, sizeof(Database.KDTreeQueryNumNeighbors));
+		Writer.Update(&Database.PoseSearchMode, sizeof(Database.PoseSearchMode));
+
 		return Writer.Finalize();
 	}
 
@@ -373,23 +378,31 @@ namespace UE::PoseSearch
 	FArchive& operator<<(FArchive& Ar, FPoseSearchIndex& Index)
 	{
 		int32 NumValues = 0;
+		int32 NumPCAValues = 0;
 		int32 NumAssets = 0;
+		int32 NumGroups = 0;
 
 		if (Ar.IsSaving())
 		{
 			NumValues = Index.Values.Num();
+			NumPCAValues = Index.PCAValues.Num();
 			NumAssets = Index.Assets.Num();
+			NumGroups = Index.Groups.Num();
 		}
 
 		Ar << Index.NumPoses;
 		Ar << NumValues;
+		Ar << NumPCAValues;
 		Ar << NumAssets;
+		Ar << NumGroups;
 
 		if (Ar.IsLoading())
 		{
 			Index.Values.SetNumUninitialized(NumValues);
+			Index.PCAValues.SetNumUninitialized(NumPCAValues);
 			Index.PoseMetadata.SetNumUninitialized(Index.NumPoses);
 			Index.Assets.SetNumUninitialized(NumAssets);
+			Index.Groups.SetNum(NumGroups);
 		}
 
 		if (Index.Values.Num() > 0)
@@ -397,6 +410,11 @@ namespace UE::PoseSearch
 			Ar.Serialize(&Index.Values[0], Index.Values.Num() * Index.Values.GetTypeSize());
 		}
 
+		if (Index.PCAValues.Num() > 0)
+		{
+			Ar.Serialize(&Index.PCAValues[0], Index.PCAValues.Num() * Index.PCAValues.GetTypeSize());
+		}
+
 		if (Index.PoseMetadata.Num() > 0)
 		{
 			Ar.Serialize(&Index.PoseMetadata[0], Index.PoseMetadata.Num() * Index.PoseMetadata.GetTypeSize());
@@ -409,6 +427,23 @@ namespace UE::PoseSearch
 
 		Ar << Index.PreprocessInfo;
 
+		for (int i = 0; i < Index.Groups.Num(); ++i)
+		{
+			FGroupSearchIndex& GroupSearchIndex = Index.Groups[i];
+			Ar << GroupSearchIndex.StartPoseIndex;
+			Ar << GroupSearchIndex.EndPoseIndex;
+			Ar << GroupSearchIndex.GroupIndex;
+			Ar << GroupSearchIndex.Weights;
+			Ar << GroupSearchIndex.Mean;
+			Ar << GroupSearchIndex.PCAProjectionMatrix;
+
+			check(GroupSearchIndex.PCAProjectionMatrix.Num() > 0 && GroupSearchIndex.Mean.Num() > 0);
+			check(GroupSearchIndex.PCAProjectionMatrix.Num() % GroupSearchIndex.Mean.Num() == 0);
+			const int NumberOfPrincipalComponents = GroupSearchIndex.PCAProjectionMatrix.Num() / GroupSearchIndex.Mean.Num();
+			float* GroupPCAValues = Index.PCAValues.GetData() + GroupSearchIndex.StartPoseIndex * NumberOfPrincipalComponents;
+			Serialize(Ar, GroupSearchIndex.KDTree, GroupPCAValues);
+		}
+
 		return Ar;
 	}
 
diff --git a/Engine/Plugins/Experimental/Animation/PoseSearch/Source/Runtime/Public/PoseSearch/KDTree.h b/Engine/Plugins/Experimental/Animation/PoseSearch/Source/Runtime/Public/PoseSearch/KDTree.h
new file mode 100644
--- /dev/null
+++ b/Engine/Plugins/Experimental/Animation/PoseSearch/Source/Runtime/Public/PoseSearch/KDTree.h
@@ -0,0 +1,120 @@
+// Copyright Epic Games, Inc. All Rights Reserved.
+
+#pragma once
+
+#include "HAL/Platform.h"
+#include "Containers/ArrayView.h"
+
+class FArchive;
+
+namespace UE { namespace PoseSearch
+{
+struct FKDTreeImplementation;
+
+struct POSESEARCH_API FKDTree
+{
+	struct DataSource
+	{
+		DataSource(int32 pointCount, int32 pointDim, const float* data)
+		: PointCount(pointCount)
+		, PointDim(pointDim)
+		, Data(data)
+		{
+		}
+
+		// Must return the number of data points
+		inline size_t kdtree_get_point_count() const { return PointCount; }
+
+		// this method is called by nanoflann::findNeighbors -> nanoflann::searchLevel, so it has to be as fast as possible
+		inline float kdtree_get_pt(const size_t idx, const size_t dim) const { return Data[idx * PointDim + dim]; }
+
+		// Optional bounding-box computation: return false to default to a standard
+		// bbox computation loop.
+		//   Return true if the BBOX was already computed by the class and returned
+		//   in "bb" so it can be avoided to redo it again. Look at bb.size() to
+		//   find out the expected dimensionality (e.g. 2 or 3 for point clouds)
+		template <class BBOX> bool kdtree_get_bbox(BBOX& /* bb */) const { return false; }
+
+		int32 PointCount = 0;
+		int32 PointDim = 0;
+		const float* Data = nullptr;
+	};
+
+	struct KNNResultSet
+	{
+		inline KNNResultSet(size_t InNumNeighbors, TArrayView<size_t> IndexesView, TArrayView<float> DistancesView)
+		: Indexes(IndexesView.GetData())
+		, Distances(DistancesView.GetData())
+		, NumNeighbors(InNumNeighbors)
+		, Count(0)
+		{
+			// by having IndexesView and DistancesView cardinality bigger than NumNeighbors, we can skip some if statements in the addPoint method
+			check(NumNeighbors > 0);
+			check(IndexesView.Num() > NumNeighbors);
+			check(DistancesView.Num() > NumNeighbors);
+
+			Distances[NumNeighbors - 1] = UE_BIG_NUMBER;
+		}
+
+		inline size_t Num() const
+		{
+			return Count;
+		}
+
+		inline bool full() const
+		{
+			return Count == NumNeighbors;
+		}
+
+		inline bool addPoint(size_t dist, size_t index)
+		{
+			// shifting Distances[i] and Indexes[i] to make space for "dist" and "index" at the right "i"th slot
+			size_t i;
+			for (i = Count; (i > 0) && (Distances[i - 1] > dist); --i)
+			{
+				// no need to check "if (i < capacity)" since dists and indices can contains more items than capacity_ 
+				Distances[i] = Distances[i - 1];
+				Indexes[i] = Indexes[i - 1];
+			}
+			
+			// inserting "dist" and "index" in a sorted manner
+			// no need to check "if (i < capacity)" since dists and indices can contains more items than capacity_ 
+			Distances[i] = dist;
+			Indexes[i] = index;
+			
+			if (Count < NumNeighbors)
+			{
+				Count++;
+			}
+
+			// tell caller that the search shall continue
+			return true;
+		}
+
+		inline float worstDist() const { return Distances[NumNeighbors - 1]; }
+
+	private:
+		size_t* Indexes;
+		float* Distances;
+		size_t NumNeighbors;
+		size_t Count;
+	};
+	
+	FKDTree(int32 Count, int32 Dim, const float* Data, int32 MaxLeafSize);
+	FKDTree();
+	FKDTree(const FKDTree& r);
+	FKDTree(FKDTree&& r) = delete;
+	
+	~FKDTree();
+
+	FKDTree& operator=(const FKDTree& r);
+
+	bool FindNeighbors(KNNResultSet& Result, const float* Query) const;
+	void Construct(int32 Count, int32 Dim, const float* Data, int32 MaxLeafSize);
+
+	DataSource DataSrc;
+	FKDTreeImplementation* Impl = nullptr;
+};
+
+FArchive& Serialize(FArchive& Ar, FKDTree& KDTree, const float* data);
+}} // namespace UE::PoseSearch
\ No newline at end of file
diff --git a/Engine/Plugins/Experimental/Animation/PoseSearch/Source/Runtime/Public/PoseSearch/PoseSearch.h b/Engine/Plugins/Experimental/Animation/PoseSearch/Source/Runtime/Public/PoseSearch/PoseSearch.h
--- a/Engine/Plugins/Experimental/Animation/PoseSearch/Source/Runtime/Public/PoseSearch/PoseSearch.h
+++ b/Engine/Plugins/Experimental/Animation/PoseSearch/Source/Runtime/Public/PoseSearch/PoseSearch.h
@@ -18,6 +18,7 @@
 #include "BoneIndices.h"
 #include "GameplayTagContainer.h"
 #include "Interfaces/Interface_BoneReferenceSkeletonProvider.h"
+#include "PoseSearch/KDTree.h"
 
 #include "PoseSearch.generated.h"
 
@@ -87,6 +88,18 @@ enum class EPoseSearchBooleanRequest : int32
 	Invalid = Num UMETA(Hidden)
 };
 
+UENUM()
+enum class EPoseSearchMode : int32
+{
+	BruteForce,
+	PCAKDTree,
+	PCAKDTree_Validate,	// runs PCAKDTree and performs validation tests
+	PCAKDTree_Compare,	// compares BruteForce vs PCAKDTree
+
+	Num UMETA(Hidden),
+	Invalid = Num UMETA(Hidden)
+};
+
 UENUM()
 enum class EPoseSearchDataPreprocessor : int32
 {
@@ -653,6 +666,32 @@ public:
 	}
 };
 
+USTRUCT()
+struct POSESEARCH_API FGroupSearchIndex
+{
+	GENERATED_BODY()
+
+	UE::PoseSearch::FKDTree KDTree;
+
+	UPROPERTY()
+	TArray<float> PCAProjectionMatrix;
+
+	UPROPERTY()
+	TArray<float> Mean;
+
+	UPROPERTY()
+	int32 StartPoseIndex = 0;
+	
+	UPROPERTY()
+	int32 EndPoseIndex = 0;
+
+	UPROPERTY()
+	int32 GroupIndex = 0;
+
+	UPROPERTY()
+	TArray<float> Weights;
+};
+
 /**
 * A search index for animation poses. The structure of the search index is determined by its UPoseSearchSchema.
 * May represent a single animation (see UPoseSearchSequenceMetaData) or a collection (see UPoseSearchDatabase).
@@ -668,6 +707,12 @@ struct POSESEARCH_API FPoseSearchIndex
 	UPROPERTY()
 	TArray<float> Values;
 
+	UPROPERTY()
+	TArray<float> PCAValues;
+
+	UPROPERTY()
+	TArray<FGroupSearchIndex> Groups;
+
 	UPROPERTY()
 	TArray<FPoseSearchPoseMetadata> PoseMetadata;
 
@@ -939,6 +984,18 @@ public:
 	UPROPERTY(EditAnywhere, Category = "Database")
 	TArray<FPoseSearchDatabaseBlendSpace> BlendSpaces;
 
+	UPROPERTY(EditAnywhere, Category = "Performance", meta = (ClampMin = "1", ClampMax = "64", UIMin = "1", UIMax = "64"))
+	int32 NumberOfPrincipalComponents = 4;
+
+	UPROPERTY(EditAnywhere, Category = "Performance", meta = (ClampMin = "1", ClampMax = "256", UIMin = "1", UIMax = "256"))
+	int32 KDTreeMaxLeafSize = 8;
+	
+	UPROPERTY(EditAnywhere, Category = "Performance", meta = (ClampMin = "1", ClampMax = "600", UIMin = "1", UIMax = "600"))
+	int32 KDTreeQueryNumNeighbors = 100;
+
+	UPROPERTY(EditAnywhere, Category = "Performance")
+	EPoseSearchMode PoseSearchMode = EPoseSearchMode::BruteForce;
+
 	FPoseSearchIndex* GetSearchIndex();
 	const FPoseSearchIndex* GetSearchIndex() const;
 
@@ -953,6 +1010,7 @@ public:
 	const bool IsSourceAssetLooping(const FPoseSearchIndexAsset* SearchIndexAsset) const;
 	const FGameplayTagContainer* GetSourceAssetGroupTags(const FPoseSearchIndexAsset* SearchIndexAsset) const;
 	const FString GetSourceAssetName(const FPoseSearchIndexAsset* SearchIndexAsset) const;
+	int32 GetNumberOfPrincipalComponents() const;
 
 public: // UObject
 	virtual void PostLoad() override;
