Fix for uninitialized FPoseSearchFeatureDesc member blocking QAGame cooks

- issue introduced with CL 20385207

diff --git a/Engine/Plugins/Experimental/Animation/PoseSearch/Source/Runtime/Public/PoseSearch/PoseSearch.h b/Engine/Plugins/Experimental/Animation/PoseSearch/Source/Runtime/Public/PoseSearch/PoseSearch.h
--- a/Engine/Plugins/Experimental/Animation/PoseSearch/Source/Runtime/Public/PoseSearch/PoseSearch.h
+++ b/Engine/Plugins/Experimental/Animation/PoseSearch/Source/Runtime/Public/PoseSearch/PoseSearch.h
@@ -234,7 +234,7 @@ struct POSESEARCH_API FPoseSearchFeatureDesc
 
 #if WITH_EDITORONLY_DATA
 	UPROPERTY()
-	EPoseSearchFeatureDomain Domain_DEPRECATED;
+		EPoseSearchFeatureDomain Domain_DEPRECATED = EPoseSearchFeatureDomain::Time;
 
 	UPROPERTY()
 	int32 SchemaBoneIdx_DEPRECATED = 0;
