Fix crash when using Pose Matching on Sequence Players ("Start From Matching Pose" option). Default to brute force search when there's no database.

diff --git a/Engine/Plugins/Experimental/Animation/PoseSearch/Source/Runtime/Private/PoseSearch.cpp b/Engine/Plugins/Experimental/Animation/PoseSearch/Source/Runtime/Private/PoseSearch.cpp
--- a/Engine/Plugins/Experimental/Animation/PoseSearch/Source/Runtime/Private/PoseSearch.cpp
+++ b/Engine/Plugins/Experimental/Animation/PoseSearch/Source/Runtime/Private/PoseSearch.cpp
@@ -4990,7 +4990,6 @@ FSearchResult SearchBruteForce(FSearchContext& SearchContext)
 	const FPoseSearchIndex* SearchIndex = SearchContext.GetSearchIndex();
 	check(SearchIndex);
 	const UPoseSearchDatabase* Database = SearchContext.GetSourceDatabase();
-	check(Database);
 
 	FPoseCost BestPoseCost;
 	int32 BestPoseIdx = INDEX_NONE;
@@ -5053,12 +5052,12 @@ FSearchResult Search(FSearchContext& SearchContext)
 	}
 
 	const UPoseSearchDatabase* Database = SearchContext.GetSourceDatabase();
-	if (Database->PoseSearchMode != EPoseSearchMode::BruteForce)
+	if (Database && Database->PoseSearchMode != EPoseSearchMode::BruteForce)
 	{
 		Result = SearchPCAKDTree(SearchContext);
 	}
-	
-	if (Database->PoseSearchMode == EPoseSearchMode::BruteForce || Database->PoseSearchMode == EPoseSearchMode::PCAKDTree_Compare)
+
+	if (!Database || Database->PoseSearchMode == EPoseSearchMode::BruteForce || Database->PoseSearchMode == EPoseSearchMode::PCAKDTree_Compare)
 	{
 		Result = SearchBruteForce(SearchContext);
 	}
