PoseSearch: Rename FDatabaseEditorToolkit to FDatabaseEditor

diff --git a/Engine/Plugins/Experimental/Animation/PoseSearch/Source/Editor/Private/PoseSearchDatabaseEdMode.cpp b/Engine/Plugins/Experimental/Animation/PoseSearch/Source/Editor/Private/PoseSearchDatabaseEdMode.cpp
--- a/Engine/Plugins/Experimental/Animation/PoseSearch/Source/Editor/Private/PoseSearchDatabaseEdMode.cpp
+++ b/Engine/Plugins/Experimental/Animation/PoseSearch/Source/Editor/Private/PoseSearchDatabaseEdMode.cpp
@@ -28,7 +28,7 @@ namespace UE::PoseSearch
 			static_cast<FDatabaseViewportClient*>(ViewportClient);
 		if (!ViewModel && PoseSearchDbViewportClient)
 		{
-			ViewModel = PoseSearchDbViewportClient->GetAssetEditorToolkit()->GetViewModel();
+			ViewModel = PoseSearchDbViewportClient->GetAssetEditor()->GetViewModel();
 		}
 
 		if (ViewModel)
diff --git a/Engine/Plugins/Experimental/Animation/PoseSearch/Source/Editor/Private/PoseSearchDatabaseEditorToolkit.cpp b/Engine/Plugins/Experimental/Animation/PoseSearch/Source/Editor/Private/PoseSearchDatabaseEditorToolkit.cpp
--- a/Engine/Plugins/Experimental/Animation/PoseSearch/Source/Editor/Private/PoseSearchDatabaseEditorToolkit.cpp
+++ b/Engine/Plugins/Experimental/Animation/PoseSearch/Source/Editor/Private/PoseSearchDatabaseEditorToolkit.cpp
@@ -43,11 +43,11 @@ namespace UE::PoseSearch
 	const FName FDatabaseEditorTabs::AssetTreeViewID(TEXT("PoseSearchDatabaseEditorAssetTreeViewTabID"));
 	const FName FDatabaseEditorTabs::SelectionDetailsID(TEXT("PoseSearchDatabaseEditorSelectionDetailsID"));
 
-	FDatabaseEditorToolkit::FDatabaseEditorToolkit()
+	FDatabaseEditor::FDatabaseEditor()
 	{
 	}
 
-	FDatabaseEditorToolkit::~FDatabaseEditorToolkit()
+	FDatabaseEditor::~FDatabaseEditor()
 	{
 		UPoseSearchDatabase* DatabaseAsset = ViewModel->GetPoseSearchDatabase();
 		if (IsValid(DatabaseAsset))
@@ -57,57 +57,57 @@ namespace UE::PoseSearch
 		}
 	}
 
-	const UPoseSearchDatabase* FDatabaseEditorToolkit::GetPoseSearchDatabase() const
+	const UPoseSearchDatabase* FDatabaseEditor::GetPoseSearchDatabase() const
 	{
 		return ViewModel.IsValid() ? ViewModel->GetPoseSearchDatabase() : nullptr;
 	}
 
-	UPoseSearchDatabase* FDatabaseEditorToolkit::GetPoseSearchDatabase()
+	UPoseSearchDatabase* FDatabaseEditor::GetPoseSearchDatabase()
 	{
 		return ViewModel.IsValid() ? ViewModel->GetPoseSearchDatabase() : nullptr;
 	}
 
-	void FDatabaseEditorToolkit::BuildSearchIndex()
+	void FDatabaseEditor::BuildSearchIndex()
 	{
 		ViewModel->BuildSearchIndex();
 	}
 
-	void FDatabaseEditorToolkit::PreviewBackwardEnd()
+	void FDatabaseEditor::PreviewBackwardEnd()
 	{
 		ViewModel->PreviewBackwardEnd();
 	}
 
-	void FDatabaseEditorToolkit::PreviewBackwardStep()
+	void FDatabaseEditor::PreviewBackwardStep()
 	{
 		ViewModel->PreviewBackwardStep();
 	}
 
-	void FDatabaseEditorToolkit::PreviewBackward()
+	void FDatabaseEditor::PreviewBackward()
 	{
 		ViewModel->PreviewBackward();
 	}
 
-	void FDatabaseEditorToolkit::PreviewPause()
+	void FDatabaseEditor::PreviewPause()
 	{
 		ViewModel->PreviewPause();
 	}
 
-	void FDatabaseEditorToolkit::PreviewForward()
+	void FDatabaseEditor::PreviewForward()
 	{
 		ViewModel->PreviewForward();
 	}
 
-	void FDatabaseEditorToolkit::PreviewForwardStep()
+	void FDatabaseEditor::PreviewForwardStep()
 	{
 		ViewModel->PreviewForwardStep();
 	}
 
-	void FDatabaseEditorToolkit::PreviewForwardEnd()
+	void FDatabaseEditor::PreviewForwardEnd()
 	{
 		ViewModel->PreviewForwardEnd();
 	}
 
-	void FDatabaseEditorToolkit::InitAssetEditor(
+	void FDatabaseEditor::InitAssetEditor(
 		const EToolkitMode::Type Mode,
 		const TSharedPtr<IToolkitHost>& InitToolkitHost,
 		UPoseSearchDatabase* DatabaseAsset)
@@ -124,7 +124,7 @@ namespace UE::PoseSearch
 					.AllowAudioPlayback(true)
 					.ShouldSimulatePhysics(true)
 					.ForceUseMovementComponentInNonGameWorld(true),
-					StaticCastSharedRef<FDatabaseEditorToolkit>(AsShared())));
+					StaticCastSharedRef<FDatabaseEditor>(AsShared())));
 
 			//Temporary fix for missing attached assets - MDW (Copied from FPersonaToolkit::CreatePreviewScene)
 			PreviewScene->GetWorld()->GetWorldSettings()->SetIsTemporarilyHiddenInEditor(false);
@@ -136,7 +136,7 @@ namespace UE::PoseSearch
 
 		// Create viewport widget
 		FDatabasePreviewRequiredArgs PreviewArgs(
-			StaticCastSharedRef<FDatabaseEditorToolkit>(AsShared()),
+			StaticCastSharedRef<FDatabaseEditor>(AsShared()),
 			PreviewScene.ToSharedRef());
 		PreviewWidget = SNew(SDatabasePreview, PreviewArgs)
 			.SliderScrubTime_Lambda([this]() { return ViewModel->GetPlayTime(); })
@@ -148,19 +148,19 @@ namespace UE::PoseSearch
 			{
 				ViewModel->SetPlayTime(NewScrubPosition, !bScrubbing);
 			})
-			.OnBackwardEnd_Raw(this, &FDatabaseEditorToolkit::PreviewBackwardEnd)
-			.OnBackwardStep_Raw(this, &FDatabaseEditorToolkit::PreviewBackwardStep)
-			.OnBackward_Raw(this, &FDatabaseEditorToolkit::PreviewBackward)
-			.OnPause_Raw(this, &FDatabaseEditorToolkit::PreviewPause)
-			.OnForward_Raw(this, &FDatabaseEditorToolkit::PreviewForward)
-			.OnForwardStep_Raw(this, &FDatabaseEditorToolkit::PreviewForwardStep)
-			.OnForwardEnd_Raw(this, &FDatabaseEditorToolkit::PreviewForwardEnd);
+			.OnBackwardEnd_Raw(this, &FDatabaseEditor::PreviewBackwardEnd)
+			.OnBackwardStep_Raw(this, &FDatabaseEditor::PreviewBackwardStep)
+			.OnBackward_Raw(this, &FDatabaseEditor::PreviewBackward)
+			.OnPause_Raw(this, &FDatabaseEditor::PreviewPause)
+			.OnForward_Raw(this, &FDatabaseEditor::PreviewForward)
+			.OnForwardStep_Raw(this, &FDatabaseEditor::PreviewForwardStep)
+			.OnForwardEnd_Raw(this, &FDatabaseEditor::PreviewForwardEnd);
 
 		AssetTreeWidget = SNew(SDatabaseAssetTree, ViewModel.ToSharedRef());
 		AssetTreeWidget->RegisterOnSelectionChanged(
 			SDatabaseAssetTree::FOnSelectionChanged::CreateSP(
 				this,
-				&FDatabaseEditorToolkit::OnAssetTreeSelectionChanged));
+				&FDatabaseEditor::OnAssetTreeSelectionChanged));
 		if (IsValid(DatabaseAsset))
 		{
 			DatabaseAsset->RegisterOnAssetChange(
@@ -201,7 +201,7 @@ namespace UE::PoseSearch
 			SelectionWidget->SetObject(nullptr);
 			SelectionWidget->OnFinishedChangingProperties().AddSP(
 				this,
-				&FDatabaseEditorToolkit::OnFinishedChangingSelectionProperties);
+				&FDatabaseEditor::OnFinishedChangingSelectionProperties);
 		}
 
 		// Define Editor Layout
@@ -274,17 +274,17 @@ namespace UE::PoseSearch
 		RegenerateMenusAndToolbars();
 	}
 
-	void FDatabaseEditorToolkit::BindCommands()
+	void FDatabaseEditor::BindCommands()
 	{
 		const FDatabaseEditorCommands& Commands = FDatabaseEditorCommands::Get();
 
 		ToolkitCommands->MapAction(
 			Commands.BuildSearchIndex,
-			FExecuteAction::CreateSP(this, &FDatabaseEditorToolkit::BuildSearchIndex),
+			FExecuteAction::CreateSP(this, &FDatabaseEditor::BuildSearchIndex),
 			EUIActionRepeatMode::RepeatDisabled);
 	}
 
-	void FDatabaseEditorToolkit::ExtendToolbar()
+	void FDatabaseEditor::ExtendToolbar()
 	{
 		TSharedPtr<FExtender> ToolbarExtender = MakeShareable(new FExtender);
 
@@ -294,10 +294,10 @@ namespace UE::PoseSearch
 			"Asset",
 			EExtensionHook::After,
 			GetToolkitCommands(),
-			FToolBarExtensionDelegate::CreateSP(this, &FDatabaseEditorToolkit::FillToolbar));
+			FToolBarExtensionDelegate::CreateSP(this, &FDatabaseEditor::FillToolbar));
 	}
 
-	void FDatabaseEditorToolkit::FillToolbar(FToolBarBuilder& ToolbarBuilder)
+	void FDatabaseEditor::FillToolbar(FToolBarBuilder& ToolbarBuilder)
 	{
 		ToolbarBuilder.AddToolBarButton(
 			FDatabaseEditorCommands::Get().BuildSearchIndex,
@@ -307,7 +307,7 @@ namespace UE::PoseSearch
 			FSlateIcon());
 	}
 
-	void FDatabaseEditorToolkit::RegisterTabSpawners(const TSharedRef<FTabManager>& InTabManager)
+	void FDatabaseEditor::RegisterTabSpawners(const TSharedRef<FTabManager>& InTabManager)
 	{
 		WorkspaceMenuCategory = InTabManager->AddLocalWorkspaceMenuCategory(
 			LOCTEXT("WorkspaceMenu_PoseSearchDbEditor", "Pose Search Database Editor"));
@@ -317,41 +317,41 @@ namespace UE::PoseSearch
 
 		InTabManager->RegisterTabSpawner(
 			FDatabaseEditorTabs::ViewportID,
-			FOnSpawnTab::CreateSP(this, &FDatabaseEditorToolkit::SpawnTab_Viewport))
+			FOnSpawnTab::CreateSP(this, &FDatabaseEditor::SpawnTab_Viewport))
 			.SetDisplayName(LOCTEXT("ViewportTab", "Viewport"))
 			.SetGroup(WorkspaceMenuCategoryRef)
 			.SetIcon(FSlateIcon(FAppStyle::GetAppStyleSetName(), "GraphEditor.EventGraph_16x"));
 
 		InTabManager->RegisterTabSpawner(
 			FDatabaseEditorTabs::AssetDetailsID,
-			FOnSpawnTab::CreateSP(this, &FDatabaseEditorToolkit::SpawnTab_AssetDetails))
+			FOnSpawnTab::CreateSP(this, &FDatabaseEditor::SpawnTab_AssetDetails))
 			.SetDisplayName(LOCTEXT("DatabaseDetailsTab", "Database Details"))
 			.SetGroup(WorkspaceMenuCategoryRef)
 			.SetIcon(FSlateIcon(FAppStyle::GetAppStyleSetName(), "LevelEditor.Tabs.Details"));
 
 		InTabManager->RegisterTabSpawner(
 			FDatabaseEditorTabs::PreviewSettingsID,
-			FOnSpawnTab::CreateSP(this, &FDatabaseEditorToolkit::SpawnTab_PreviewSettings))
+			FOnSpawnTab::CreateSP(this, &FDatabaseEditor::SpawnTab_PreviewSettings))
 			.SetDisplayName(LOCTEXT("PreviewSceneSettingsTab", "Preview Scene Settings"))
 			.SetGroup(WorkspaceMenuCategoryRef)
 			.SetIcon(FSlateIcon(FAppStyle::GetAppStyleSetName(), "LevelEditor.Tabs.Details"));
 
 		InTabManager->RegisterTabSpawner(
 			FDatabaseEditorTabs::AssetTreeViewID,
-			FOnSpawnTab::CreateSP(this, &FDatabaseEditorToolkit::SpawnTab_AssetTreeView))
+			FOnSpawnTab::CreateSP(this, &FDatabaseEditor::SpawnTab_AssetTreeView))
 			.SetDisplayName(LOCTEXT("TreeViewTab", "Tree View"))
 			.SetGroup(WorkspaceMenuCategoryRef)
 			.SetIcon(FSlateIcon(FAppStyle::GetAppStyleSetName(), "GraphEditor.EventGraph_16x"));
 
 		InTabManager->RegisterTabSpawner(
 			FDatabaseEditorTabs::SelectionDetailsID,
-			FOnSpawnTab::CreateSP(this, &FDatabaseEditorToolkit::SpawnTab_SelectionDetails))
+			FOnSpawnTab::CreateSP(this, &FDatabaseEditor::SpawnTab_SelectionDetails))
 			.SetDisplayName(LOCTEXT("SelectionDetailsTab", "Selection Details"))
 			.SetGroup(WorkspaceMenuCategoryRef)
 			.SetIcon(FSlateIcon(FAppStyle::GetAppStyleSetName(), "GraphEditor.EventGraph_16x"));
 	}
 
-	void FDatabaseEditorToolkit::UnregisterTabSpawners(const TSharedRef<FTabManager>& InTabManager)
+	void FDatabaseEditor::UnregisterTabSpawners(const TSharedRef<FTabManager>& InTabManager)
 	{
 		FAssetEditorToolkit::UnregisterTabSpawners(InTabManager);
 
@@ -362,34 +362,34 @@ namespace UE::PoseSearch
 		InTabManager->UnregisterTabSpawner(FDatabaseEditorTabs::SelectionDetailsID);
 	}
 
-	FName FDatabaseEditorToolkit::GetToolkitFName() const
+	FName FDatabaseEditor::GetToolkitFName() const
 	{
 		return FName("PoseSearchDatabaseEditor");
 	}
 
-	FText FDatabaseEditorToolkit::GetBaseToolkitName() const
+	FText FDatabaseEditor::GetBaseToolkitName() const
 	{
 		return LOCTEXT("PoseSearchDatabaseEditorAppLabel", "Pose Search Database Editor");
 	}
 
-	FText FDatabaseEditorToolkit::GetToolkitName() const
+	FText FDatabaseEditor::GetToolkitName() const
 	{
 		FFormatNamedArguments Args;
 		Args.Add(TEXT("AssetName"), FText::FromString(GetPoseSearchDatabase()->GetName()));
 		return FText::Format(LOCTEXT("PoseSearchDatabaseEditorToolkitName", "{AssetName}"), Args);
 	}
 
-	FLinearColor FDatabaseEditorToolkit::GetWorldCentricTabColorScale() const
+	FLinearColor FDatabaseEditor::GetWorldCentricTabColorScale() const
 	{
 		return FLinearColor::White;
 	}
 
-	FString FDatabaseEditorToolkit::GetWorldCentricTabPrefix() const
+	FString FDatabaseEditor::GetWorldCentricTabPrefix() const
 	{
 		return TEXT("PoseSearchDatabaseEditor");
 	}
 
-	TSharedRef<SDockTab> FDatabaseEditorToolkit::SpawnTab_Viewport(const FSpawnTabArgs& Args)
+	TSharedRef<SDockTab> FDatabaseEditor::SpawnTab_Viewport(const FSpawnTabArgs& Args)
 	{
 		check(Args.GetTabId() == FDatabaseEditorTabs::ViewportID);
 
@@ -403,7 +403,7 @@ namespace UE::PoseSearch
 		return SpawnedTab;
 	}
 
-	TSharedRef<SDockTab> FDatabaseEditorToolkit::SpawnTab_AssetDetails(const FSpawnTabArgs& Args)
+	TSharedRef<SDockTab> FDatabaseEditor::SpawnTab_AssetDetails(const FSpawnTabArgs& Args)
 	{
 		check(Args.GetTabId() == FDatabaseEditorTabs::AssetDetailsID);
 
@@ -414,7 +414,7 @@ namespace UE::PoseSearch
 			];
 	}
 
-	TSharedRef<SDockTab> FDatabaseEditorToolkit::SpawnTab_PreviewSettings(const FSpawnTabArgs& Args)
+	TSharedRef<SDockTab> FDatabaseEditor::SpawnTab_PreviewSettings(const FSpawnTabArgs& Args)
 	{
 		check(Args.GetTabId() == FDatabaseEditorTabs::PreviewSettingsID);
 
@@ -432,7 +432,7 @@ namespace UE::PoseSearch
 		return SpawnedTab;
 	}
 
-	TSharedRef<SDockTab> FDatabaseEditorToolkit::SpawnTab_AssetTreeView(const FSpawnTabArgs& Args)
+	TSharedRef<SDockTab> FDatabaseEditor::SpawnTab_AssetTreeView(const FSpawnTabArgs& Args)
 	{
 		check(Args.GetTabId() == FDatabaseEditorTabs::AssetTreeViewID);
 
@@ -443,7 +443,7 @@ namespace UE::PoseSearch
 			];
 	}
 
-	TSharedRef<SDockTab> FDatabaseEditorToolkit::SpawnTab_SelectionDetails(const FSpawnTabArgs& Args)
+	TSharedRef<SDockTab> FDatabaseEditor::SpawnTab_SelectionDetails(const FSpawnTabArgs& Args)
 	{
 		check(Args.GetTabId() == FDatabaseEditorTabs::SelectionDetailsID);
 
@@ -470,12 +470,12 @@ namespace UE::PoseSearch
 			];
 	}
 
-	void FDatabaseEditorToolkit::OnFinishedChangingSelectionProperties(const FPropertyChangedEvent& PropertyChangedEvent)
+	void FDatabaseEditor::OnFinishedChangingSelectionProperties(const FPropertyChangedEvent& PropertyChangedEvent)
 	{
 		ViewModel->BuildSearchIndex();
 	}
 
-	void FDatabaseEditorToolkit::OnAssetTreeSelectionChanged(
+	void FDatabaseEditor::OnAssetTreeSelectionChanged(
 		const TArrayView<TSharedPtr<FDatabaseAssetTreeNode>>& SelectedItems,
 		ESelectInfo::Type SelectionType)
 	{
diff --git a/Engine/Plugins/Experimental/Animation/PoseSearch/Source/Editor/Private/PoseSearchDatabasePreviewScene.cpp b/Engine/Plugins/Experimental/Animation/PoseSearch/Source/Editor/Private/PoseSearchDatabasePreviewScene.cpp
--- a/Engine/Plugins/Experimental/Animation/PoseSearch/Source/Editor/Private/PoseSearchDatabasePreviewScene.cpp
+++ b/Engine/Plugins/Experimental/Animation/PoseSearch/Source/Editor/Private/PoseSearchDatabasePreviewScene.cpp
@@ -14,9 +14,9 @@ namespace UE::PoseSearch
 {
 	FDatabasePreviewScene::FDatabasePreviewScene(
 		ConstructionValues CVs,
-		const TSharedRef<FDatabaseEditorToolkit>& EditorToolkit)
+		const TSharedRef<FDatabaseEditor>& Editor)
 		: FAdvancedPreviewScene(CVs)
-		, EditorToolkitPtr(EditorToolkit)
+		, EditorPtr(Editor)
 	{
 		// Disable killing actors outside of the world
 		AWorldSettings* WorldSettings = GetWorld()->GetWorldSettings(true);
@@ -57,7 +57,7 @@ namespace UE::PoseSearch
 			GetWorld()->Tick(LEVELTICK_All, InDeltaTime);
 		}
 
-		const FDatabaseViewModel* ViewModel = GetEditorToolkit()->GetViewModel();
+		const FDatabaseViewModel* ViewModel = GetEditor()->GetViewModel();
 		if (ViewModel->GetPoseSearchDatabase()->IsValidForSearch() &&
 			ViewModel->IsPoseFeaturesDrawMode(EFeaturesDrawMode::All))
 		{
diff --git a/Engine/Plugins/Experimental/Animation/PoseSearch/Source/Editor/Private/PoseSearchDatabaseViewportClient.cpp b/Engine/Plugins/Experimental/Animation/PoseSearch/Source/Editor/Private/PoseSearchDatabaseViewportClient.cpp
--- a/Engine/Plugins/Experimental/Animation/PoseSearch/Source/Editor/Private/PoseSearchDatabaseViewportClient.cpp
+++ b/Engine/Plugins/Experimental/Animation/PoseSearch/Source/Editor/Private/PoseSearchDatabaseViewportClient.cpp
@@ -13,10 +13,10 @@ namespace UE::PoseSearch
 	FDatabaseViewportClient::FDatabaseViewportClient(
 		const TSharedRef<FDatabasePreviewScene>& InPreviewScene,
 		const TSharedRef<SDatabaseViewport>& InViewport,
-		const TSharedRef<FDatabaseEditorToolkit>& InAssetEditorToolkit)
+		const TSharedRef<FDatabaseEditor>& InAssetEditor)
 		: FEditorViewportClient(nullptr, &InPreviewScene.Get(), StaticCastSharedRef<SEditorViewport>(InViewport))
 		, PreviewScenePtr(InPreviewScene)
-		, AssetEditorToolkitPtr(InAssetEditorToolkit)
+		, AssetEditorPtr(InAssetEditor)
 	{
 		Widget->SetUsesEditorModeTools(ModeTools.Get());
 		StaticCastSharedPtr<FAssetEditorModeManager>(ModeTools)->SetPreviewScene(&InPreviewScene.Get());
diff --git a/Engine/Plugins/Experimental/Animation/PoseSearch/Source/Editor/Private/PoseSearchTypeActions.cpp b/Engine/Plugins/Experimental/Animation/PoseSearch/Source/Editor/Private/PoseSearchTypeActions.cpp
--- a/Engine/Plugins/Experimental/Animation/PoseSearch/Source/Editor/Private/PoseSearchTypeActions.cpp
+++ b/Engine/Plugins/Experimental/Animation/PoseSearch/Source/Editor/Private/PoseSearchTypeActions.cpp
@@ -40,7 +40,7 @@ namespace UE::PoseSearch
 		{
 			if (UPoseSearchDatabase* PoseSearchDb = Cast<UPoseSearchDatabase>(*ObjIt))
 			{
-				TSharedRef<FDatabaseEditorToolkit> NewEditor(new FDatabaseEditorToolkit());
+				TSharedRef<FDatabaseEditor> NewEditor(new FDatabaseEditor());
 				NewEditor->InitAssetEditor(Mode, EditWithinLevelEditor, PoseSearchDb);
 			}
 		}
diff --git a/Engine/Plugins/Experimental/Animation/PoseSearch/Source/Editor/Private/SPoseSearchDatabaseViewport.cpp b/Engine/Plugins/Experimental/Animation/PoseSearch/Source/Editor/Private/SPoseSearchDatabaseViewport.cpp
--- a/Engine/Plugins/Experimental/Animation/PoseSearch/Source/Editor/Private/SPoseSearchDatabaseViewport.cpp
+++ b/Engine/Plugins/Experimental/Animation/PoseSearch/Source/Editor/Private/SPoseSearchDatabaseViewport.cpp
@@ -17,7 +17,7 @@ namespace UE::PoseSearch
 		const FDatabasePreviewRequiredArgs& InRequiredArgs)
 	{
 		PreviewScenePtr = InRequiredArgs.PreviewScene;
-		AssetEditorToolkitPtr = InRequiredArgs.AssetEditorToolkit;
+		AssetEditorPtr = InRequiredArgs.AssetEditor;
 
 		SEditorViewport::Construct(
 			SEditorViewport::FArguments()
@@ -33,7 +33,7 @@ namespace UE::PoseSearch
 		const FDatabaseEditorCommands& Commands = FDatabaseEditorCommands::Get();
 
 		TSharedRef<FDatabaseViewModel> ViewModelRef =
-			AssetEditorToolkitPtr.Pin()->GetViewModelSharedPtr().ToSharedRef();
+			AssetEditorPtr.Pin()->GetViewModelSharedPtr().ToSharedRef();
 
 		CommandList->MapAction(
 			Commands.ShowPoseFeaturesNone,
@@ -101,7 +101,7 @@ namespace UE::PoseSearch
 		ViewportClient = MakeShared<FDatabaseViewportClient>(
 			PreviewScenePtr.Pin().ToSharedRef(),
 			SharedThis(this),
-			AssetEditorToolkitPtr.Pin().ToSharedRef());
+			AssetEditorPtr.Pin().ToSharedRef());
 		ViewportClient->ViewportType = LVT_Perspective;
 		ViewportClient->bSetListenerPosition = false;
 		ViewportClient->SetViewLocation(EditorViewportDefs::DefaultPerspectiveViewLocation);
diff --git a/Engine/Plugins/Experimental/Animation/PoseSearch/Source/Editor/Public/PoseSearchDatabaseEditorToolkit.h b/Engine/Plugins/Experimental/Animation/PoseSearch/Source/Editor/Public/PoseSearchDatabaseEditorToolkit.h
--- a/Engine/Plugins/Experimental/Animation/PoseSearch/Source/Editor/Public/PoseSearchDatabaseEditorToolkit.h
+++ b/Engine/Plugins/Experimental/Animation/PoseSearch/Source/Editor/Public/PoseSearchDatabaseEditorToolkit.h
@@ -17,12 +17,12 @@ namespace UE::PoseSearch
 	class FDatabasePreviewScene;
 	class FDatabaseViewModel;
 
-	class FDatabaseEditorToolkit : public FAssetEditorToolkit, public FNotifyHook
+	class FDatabaseEditor : public FAssetEditorToolkit, public FNotifyHook
 	{
 	public:
 
-		FDatabaseEditorToolkit();
-		virtual ~FDatabaseEditorToolkit();
+		FDatabaseEditor();
+		virtual ~FDatabaseEditor();
 
 		void InitAssetEditor(
 			const EToolkitMode::Type Mode,
diff --git a/Engine/Plugins/Experimental/Animation/PoseSearch/Source/Editor/Public/PoseSearchDatabasePreviewScene.h b/Engine/Plugins/Experimental/Animation/PoseSearch/Source/Editor/Public/PoseSearchDatabasePreviewScene.h
--- a/Engine/Plugins/Experimental/Animation/PoseSearch/Source/Editor/Public/PoseSearchDatabasePreviewScene.h
+++ b/Engine/Plugins/Experimental/Animation/PoseSearch/Source/Editor/Public/PoseSearchDatabasePreviewScene.h
@@ -7,7 +7,7 @@
 
 namespace UE::PoseSearch
 {
-	class FDatabaseEditorToolkit;
+	class FDatabaseEditor;
 
 	class FDatabasePreviewScene : public FAdvancedPreviewScene
 	{
@@ -15,20 +15,20 @@ namespace UE::PoseSearch
 
 		FDatabasePreviewScene(
 			ConstructionValues CVs,
-			const TSharedRef<FDatabaseEditorToolkit>& EditorToolkit);
+			const TSharedRef<FDatabaseEditor>& Editor);
 		~FDatabasePreviewScene() {}
 
 		virtual void Tick(float InDeltaTime) override;
 
-		TSharedRef<FDatabaseEditorToolkit> GetEditorToolkit() const
+		TSharedRef<FDatabaseEditor> GetEditor() const
 		{
-			return EditorToolkitPtr.Pin().ToSharedRef();
+			return EditorPtr.Pin().ToSharedRef();
 		}
 
 	private:
 
-		/** The asset editor toolkit we are embedded in */
-		TWeakPtr<FDatabaseEditorToolkit> EditorToolkitPtr;
+		/** The asset editor we are embedded in */
+		TWeakPtr<FDatabaseEditor> EditorPtr;
 	};
 }
 
diff --git a/Engine/Plugins/Experimental/Animation/PoseSearch/Source/Editor/Public/PoseSearchDatabaseViewportClient.h b/Engine/Plugins/Experimental/Animation/PoseSearch/Source/Editor/Public/PoseSearchDatabaseViewportClient.h
--- a/Engine/Plugins/Experimental/Animation/PoseSearch/Source/Editor/Public/PoseSearchDatabaseViewportClient.h
+++ b/Engine/Plugins/Experimental/Animation/PoseSearch/Source/Editor/Public/PoseSearchDatabaseViewportClient.h
@@ -10,7 +10,7 @@ namespace UE::PoseSearch
 {
 	class FDatabasePreviewScene;
 	class SDatabaseViewport;
-	class FDatabaseEditorToolkit;
+	class FDatabaseEditor;
 
 	class FDatabaseViewportClient : public FEditorViewportClient
 	{
@@ -19,7 +19,7 @@ namespace UE::PoseSearch
 		FDatabaseViewportClient(
 			const TSharedRef<FDatabasePreviewScene>& InPreviewScene,
 			const TSharedRef<SDatabaseViewport>& InViewport,
-			const TSharedRef<FDatabaseEditorToolkit>& InAssetEditorToolkit);
+			const TSharedRef<FDatabaseEditor>& InAssetEditor);
 		virtual ~FDatabaseViewportClient() {}
 
 		// ~FEditorViewportClient interface
@@ -34,9 +34,9 @@ namespace UE::PoseSearch
 			return PreviewScenePtr.Pin().ToSharedRef();
 		}
 
-		TSharedRef<FDatabaseEditorToolkit> GetAssetEditorToolkit() const
+		TSharedRef<FDatabaseEditor> GetAssetEditor() const
 		{
-			return AssetEditorToolkitPtr.Pin().ToSharedRef();
+			return AssetEditorPtr.Pin().ToSharedRef();
 		}
 
 	private:
@@ -44,7 +44,7 @@ namespace UE::PoseSearch
 		/** Preview scene we are viewing */
 		TWeakPtr<FDatabasePreviewScene> PreviewScenePtr;
 
-		/** Asset editor toolkit we are embedded in */
-		TWeakPtr<FDatabaseEditorToolkit> AssetEditorToolkitPtr;
+		/** Asset editor we are embedded in */
+		TWeakPtr<FDatabaseEditor> AssetEditorPtr;
 	};
 }
diff --git a/Engine/Plugins/Experimental/Animation/PoseSearch/Source/Editor/Public/SPoseSearchDatabaseViewport.h b/Engine/Plugins/Experimental/Animation/PoseSearch/Source/Editor/Public/SPoseSearchDatabaseViewport.h
--- a/Engine/Plugins/Experimental/Animation/PoseSearch/Source/Editor/Public/SPoseSearchDatabaseViewport.h
+++ b/Engine/Plugins/Experimental/Animation/PoseSearch/Source/Editor/Public/SPoseSearchDatabaseViewport.h
@@ -11,20 +11,20 @@ namespace UE::PoseSearch
 {
 	class FDatabasePreviewScene;
 	class FDatabaseViewportClient;
-	class FDatabaseEditorToolkit;
+	class FDatabaseEditor;
 	class SPoseSearchDatabaseViewportToolBar;
 
 	struct FDatabasePreviewRequiredArgs
 	{
 		FDatabasePreviewRequiredArgs(
-			const TSharedRef<FDatabaseEditorToolkit>& InAssetEditorToolkit,
+			const TSharedRef<FDatabaseEditor>& InAssetEditor,
 			const TSharedRef<FDatabasePreviewScene>& InPreviewScene)
-			: AssetEditorToolkit(InAssetEditorToolkit)
+			: AssetEditor(InAssetEditor)
 			, PreviewScene(InPreviewScene)
 		{
 		}
 
-		TSharedRef<FDatabaseEditorToolkit> AssetEditorToolkit;
+		TSharedRef<FDatabaseEditor> AssetEditor;
 
 		TSharedRef<FDatabasePreviewScene> PreviewScene;
 	};
@@ -61,8 +61,8 @@ namespace UE::PoseSearch
 		/** The preview scene that we are viewing */
 		TWeakPtr<FDatabasePreviewScene> PreviewScenePtr;
 
-		/** Asset editor toolkit we are embedded in */
-		TWeakPtr<FDatabaseEditorToolkit> AssetEditorToolkitPtr;
+		/** Asset editor we are embedded in */
+		TWeakPtr<FDatabaseEditor> AssetEditorPtr;
 	};
 
 	class SDatabasePreview : public SCompoundWidget
