Initialize arrays to fix static analysis warnings

diff --git a/Engine/Plugins/Experimental/Animation/PoseSearch/Source/Runtime/Private/PoseSearch.cpp b/Engine/Plugins/Experimental/Animation/PoseSearch/Source/Runtime/Private/PoseSearch.cpp
--- a/Engine/Plugins/Experimental/Animation/PoseSearch/Source/Runtime/Private/PoseSearch.cpp
+++ b/Engine/Plugins/Experimental/Animation/PoseSearch/Source/Runtime/Private/PoseSearch.cpp
@@ -4404,8 +4404,8 @@ static void PreprocessGroupSearchIndexKDTree(FGroupSearchIndex& GroupSearchIndex
 		for (size_t PointIndex = 0; PointIndex < NumGroupPoses; ++PointIndex)
 		{
 			constexpr size_t NumResults = 10;
-			size_t ResultIndexes[NumResults + 1];
-			float ResultDistanceSqr[NumResults + 1];
+			size_t ResultIndexes[NumResults + 1] = { 0 };
+			float ResultDistanceSqr[NumResults + 1] = { 0.0f };
 			FKDTree::KNNResultSet ResultSet(NumResults, ResultIndexes, ResultDistanceSqr);
 			GroupSearchIndex.KDTree.FindNeighbors(ResultSet, &GroupPCAValues[PointIndex * NumberOfPrincipalComponents]);
 
