PoseSearch: Rename/move file(s)

diff --git a/Engine/Plugins/Experimental/Animation/PoseSearch/Source/Editor/Private/PoseSearchDatabaseEdMode.cpp b/Engine/Plugins/Experimental/Animation/PoseSearch/Source/Editor/Private/PoseSearchDatabaseEdMode.cpp
--- a/Engine/Plugins/Experimental/Animation/PoseSearch/Source/Editor/Private/PoseSearchDatabaseEdMode.cpp
+++ b/Engine/Plugins/Experimental/Animation/PoseSearch/Source/Editor/Private/PoseSearchDatabaseEdMode.cpp
@@ -2,7 +2,7 @@
 
 #include "PoseSearchDatabaseEdMode.h"
 #include "PoseSearchDatabaseViewportClient.h"
-#include "PoseSearchDatabaseEditorToolkit.h"
+#include "PoseSearchDatabaseEditor.h"
 #include "PoseSearchDatabaseViewModel.h"
 #include "PoseSearch/PoseSearch.h"
 
diff --git a/Engine/Plugins/Experimental/Animation/PoseSearch/Source/Editor/Private/PoseSearchDatabaseEditor.cpp b/Engine/Plugins/Experimental/Animation/PoseSearch/Source/Editor/Private/PoseSearchDatabaseEditor.cpp
new file mode 100644
--- /dev/null
+++ b/Engine/Plugins/Experimental/Animation/PoseSearch/Source/Editor/Private/PoseSearchDatabaseEditor.cpp
@@ -0,0 +1,551 @@
+// Copyright Epic Games, Inc. All Rights Reserved.
+
+#include "PoseSearchDatabaseEditor.h"
+#include "SPoseSearchDatabaseViewport.h"
+#include "SPoseSearchDatabaseAssetList.h"
+#include "PoseSearchDatabasePreviewScene.h"
+#include "PoseSearchDatabaseEditorCommands.h"
+#include "PoseSearchDatabaseViewModel.h"
+#include "PoseSearchDatabaseEditorReflection.h"
+#include "PoseSearchEditor.h"
+#include "PoseSearch/PoseSearch.h"
+#include "GameFramework/WorldSettings.h"
+#include "AdvancedPreviewSceneModule.h"
+#include "IStructureDetailsView.h"
+#include "PropertyEditorModule.h"
+#include "Modules/ModuleManager.h"
+#include "Widgets/Input/SButton.h"
+#include "Widgets/Docking/SDockTab.h"
+#include "Widgets/Layout/SScrollBox.h"
+#include "Styling/AppStyle.h"
+
+#define LOCTEXT_NAMESPACE "PoseSearchDatabaseEditor"
+
+namespace UE::PoseSearch
+{
+	const FName PoseSearchDatabaseEditorAppName = FName(TEXT("PoseSearchDatabaseEditorApp"));
+	constexpr int32 SelectionDetailsCount = 3; // sequence, blendspace, group
+	constexpr double ViewRangeSlack = 0.25;
+
+	// Tab identifiers
+	struct FDatabaseEditorTabs
+	{
+		static const FName AssetDetailsID;
+		static const FName ViewportID;
+		static const FName PreviewSettingsID;
+		static const FName AssetTreeViewID;
+		static const FName SelectionDetailsID;
+	};
+
+	const FName FDatabaseEditorTabs::AssetDetailsID(TEXT("PoseSearchDatabaseEditorAssetDetailsTabID"));
+	const FName FDatabaseEditorTabs::ViewportID(TEXT("PoseSearchDatabaseEditorViewportTabID"));
+	const FName FDatabaseEditorTabs::PreviewSettingsID(TEXT("PoseSearchDatabaseEditorPreviewSettingsTabID"));
+	const FName FDatabaseEditorTabs::AssetTreeViewID(TEXT("PoseSearchDatabaseEditorAssetTreeViewTabID"));
+	const FName FDatabaseEditorTabs::SelectionDetailsID(TEXT("PoseSearchDatabaseEditorSelectionDetailsID"));
+
+	FDatabaseEditor::FDatabaseEditor()
+	{
+	}
+
+	FDatabaseEditor::~FDatabaseEditor()
+	{
+		UPoseSearchDatabase* DatabaseAsset = ViewModel->GetPoseSearchDatabase();
+		if (IsValid(DatabaseAsset))
+		{
+			DatabaseAsset->UnregisterOnAssetChange(AssetTreeWidget.Get());
+			DatabaseAsset->UnregisterOnGroupChange(AssetTreeWidget.Get());
+		}
+	}
+
+	const UPoseSearchDatabase* FDatabaseEditor::GetPoseSearchDatabase() const
+	{
+		return ViewModel.IsValid() ? ViewModel->GetPoseSearchDatabase() : nullptr;
+	}
+
+	UPoseSearchDatabase* FDatabaseEditor::GetPoseSearchDatabase()
+	{
+		return ViewModel.IsValid() ? ViewModel->GetPoseSearchDatabase() : nullptr;
+	}
+
+	void FDatabaseEditor::BuildSearchIndex()
+	{
+		ViewModel->BuildSearchIndex();
+	}
+
+	void FDatabaseEditor::PreviewBackwardEnd()
+	{
+		ViewModel->PreviewBackwardEnd();
+	}
+
+	void FDatabaseEditor::PreviewBackwardStep()
+	{
+		ViewModel->PreviewBackwardStep();
+	}
+
+	void FDatabaseEditor::PreviewBackward()
+	{
+		ViewModel->PreviewBackward();
+	}
+
+	void FDatabaseEditor::PreviewPause()
+	{
+		ViewModel->PreviewPause();
+	}
+
+	void FDatabaseEditor::PreviewForward()
+	{
+		ViewModel->PreviewForward();
+	}
+
+	void FDatabaseEditor::PreviewForwardStep()
+	{
+		ViewModel->PreviewForwardStep();
+	}
+
+	void FDatabaseEditor::PreviewForwardEnd()
+	{
+		ViewModel->PreviewForwardEnd();
+	}
+
+	void FDatabaseEditor::InitAssetEditor(
+		const EToolkitMode::Type Mode,
+		const TSharedPtr<IToolkitHost>& InitToolkitHost,
+		UPoseSearchDatabase* DatabaseAsset)
+	{
+		// Bind Commands
+		BindCommands();
+
+		// Create Preview Scene
+		if (!PreviewScene.IsValid())
+		{
+			PreviewScene = MakeShareable(
+				new FDatabasePreviewScene(
+					FPreviewScene::ConstructionValues()
+					.AllowAudioPlayback(true)
+					.ShouldSimulatePhysics(true)
+					.ForceUseMovementComponentInNonGameWorld(true),
+					StaticCastSharedRef<FDatabaseEditor>(AsShared())));
+
+			//Temporary fix for missing attached assets - MDW (Copied from FPersonaToolkit::CreatePreviewScene)
+			PreviewScene->GetWorld()->GetWorldSettings()->SetIsTemporarilyHiddenInEditor(false);
+		}
+
+		// Create view model
+		ViewModel = MakeShared<FDatabaseViewModel>();
+		ViewModel->Initialize(DatabaseAsset, PreviewScene.ToSharedRef());
+
+		// Create viewport widget
+		FDatabasePreviewRequiredArgs PreviewArgs(
+			StaticCastSharedRef<FDatabaseEditor>(AsShared()),
+			PreviewScene.ToSharedRef());
+		PreviewWidget = SNew(SDatabasePreview, PreviewArgs)
+			.SliderScrubTime_Lambda([this]() { return ViewModel->GetPlayTime(); })
+			.SliderViewRange_Lambda([this]() 
+			{ 
+				return TRange<double>(-ViewRangeSlack, ViewModel->GetMaxPreviewPlayLength() + ViewRangeSlack);
+			})
+			.OnSliderScrubPositionChanged_Lambda([this](float NewScrubPosition, bool bScrubbing)
+			{
+				ViewModel->SetPlayTime(NewScrubPosition, !bScrubbing);
+			})
+			.OnBackwardEnd_Raw(this, &FDatabaseEditor::PreviewBackwardEnd)
+			.OnBackwardStep_Raw(this, &FDatabaseEditor::PreviewBackwardStep)
+			.OnBackward_Raw(this, &FDatabaseEditor::PreviewBackward)
+			.OnPause_Raw(this, &FDatabaseEditor::PreviewPause)
+			.OnForward_Raw(this, &FDatabaseEditor::PreviewForward)
+			.OnForwardStep_Raw(this, &FDatabaseEditor::PreviewForwardStep)
+			.OnForwardEnd_Raw(this, &FDatabaseEditor::PreviewForwardEnd);
+
+		AssetTreeWidget = SNew(SDatabaseAssetTree, ViewModel.ToSharedRef());
+		AssetTreeWidget->RegisterOnSelectionChanged(
+			SDatabaseAssetTree::FOnSelectionChanged::CreateSP(
+				this,
+				&FDatabaseEditor::OnAssetTreeSelectionChanged));
+		if (IsValid(DatabaseAsset))
+		{
+			DatabaseAsset->RegisterOnAssetChange(
+				UPoseSearchDatabase::FOnDerivedDataRebuild::CreateSP(
+					AssetTreeWidget.Get(),
+					&SDatabaseAssetTree::RefreshTreeView, false, false));
+			DatabaseAsset->RegisterOnGroupChange(
+				UPoseSearchDatabase::FOnDerivedDataRebuild::CreateSP(
+					AssetTreeWidget.Get(),
+					&SDatabaseAssetTree::RefreshTreeView, false, false));
+		}
+
+		// Create Asset Details widget
+		FPropertyEditorModule& PropertyModule = 
+			FModuleManager::LoadModuleChecked<FPropertyEditorModule>("PropertyEditor");
+
+		FDetailsViewArgs DatabaseDetailsArgs;
+		{
+			DatabaseDetailsArgs.bHideSelectionTip = true;
+			DatabaseDetailsArgs.NotifyHook = this;
+		}
+
+		EditingAssetWidget = PropertyModule.CreateDetailView(DatabaseDetailsArgs);
+		EditingAssetWidget->SetObject(DatabaseAsset);
+
+		FDetailsViewArgs SelectionDetailsArgs;
+		{
+			SelectionDetailsArgs.bHideSelectionTip = true;
+			SelectionDetailsArgs.NotifyHook = this;
+			SelectionDetailsArgs.bShowScrollBar = false;
+		}
+
+		SelectionWidgets.Reset(SelectionDetailsCount);
+		for (int32 i = 0; i < SelectionDetailsCount; ++i)
+		{
+			TSharedPtr<IDetailsView>& SelectionWidget = 
+				SelectionWidgets.Add_GetRef(PropertyModule.CreateDetailView(SelectionDetailsArgs));
+			SelectionWidget->SetObject(nullptr);
+			SelectionWidget->OnFinishedChangingProperties().AddSP(
+				this,
+				&FDatabaseEditor::OnFinishedChangingSelectionProperties);
+		}
+
+		// Define Editor Layout
+		const TSharedRef<FTabManager::FLayout> StandaloneDefaultLayout =
+		FTabManager::NewLayout("Standalone_PoseSearchDatabaseEditor_Layout_v0.07")
+			->AddArea
+			(
+				FTabManager::NewPrimaryArea()->SetOrientation(Orient_Vertical)
+				->Split
+				(
+					FTabManager::NewSplitter()->SetOrientation(Orient_Vertical)->SetSizeCoefficient(0.9f)
+					->Split
+					(
+						FTabManager::NewSplitter()->SetOrientation(Orient_Horizontal)->SetSizeCoefficient(0.9f)
+						->Split
+						(
+							FTabManager::NewSplitter()->SetOrientation(Orient_Vertical)->SetSizeCoefficient(0.9f)
+							->Split
+							(
+								FTabManager::NewStack()
+								->SetSizeCoefficient(0.5f)
+								->AddTab(FDatabaseEditorTabs::AssetTreeViewID, ETabState::OpenedTab)
+								->SetHideTabWell(false)
+							)
+							->Split
+							(
+								FTabManager::NewStack()
+								->SetSizeCoefficient(0.5f)
+								->AddTab(FDatabaseEditorTabs::AssetDetailsID, ETabState::OpenedTab)
+								->SetHideTabWell(false)
+							)
+						)
+						->Split
+						(
+							FTabManager::NewStack()
+							->SetSizeCoefficient(0.4f)
+							->AddTab(FDatabaseEditorTabs::ViewportID, ETabState::OpenedTab)
+							->SetHideTabWell(true)
+						)
+						->Split
+						(
+							FTabManager::NewSplitter()->SetOrientation(Orient_Vertical)
+							->Split
+							(
+								FTabManager::NewStack()
+								->SetSizeCoefficient(0.3f)
+								->AddTab(FDatabaseEditorTabs::SelectionDetailsID, ETabState::OpenedTab)
+								->AddTab(FDatabaseEditorTabs::PreviewSettingsID, ETabState::OpenedTab)
+							)
+						)
+					)
+				)
+			);
+
+		const bool bCreateDefaultStandaloneMenuParam = true;
+		const bool bCreateDefaultToolbarParam = true;
+		const bool bIsToolbarFocusableParam = false;
+		FAssetEditorToolkit::InitAssetEditor(
+			Mode,
+			InitToolkitHost,
+			PoseSearchDatabaseEditorAppName,
+			StandaloneDefaultLayout,
+			bCreateDefaultStandaloneMenuParam,
+			bCreateDefaultToolbarParam,
+			DatabaseAsset,
+			bIsToolbarFocusableParam);
+
+		ExtendToolbar();
+
+		RegenerateMenusAndToolbars();
+	}
+
+	void FDatabaseEditor::BindCommands()
+	{
+		const FDatabaseEditorCommands& Commands = FDatabaseEditorCommands::Get();
+
+		ToolkitCommands->MapAction(
+			Commands.BuildSearchIndex,
+			FExecuteAction::CreateSP(this, &FDatabaseEditor::BuildSearchIndex),
+			EUIActionRepeatMode::RepeatDisabled);
+	}
+
+	void FDatabaseEditor::ExtendToolbar()
+	{
+		TSharedPtr<FExtender> ToolbarExtender = MakeShareable(new FExtender);
+
+		AddToolbarExtender(ToolbarExtender);
+
+		ToolbarExtender->AddToolBarExtension(
+			"Asset",
+			EExtensionHook::After,
+			GetToolkitCommands(),
+			FToolBarExtensionDelegate::CreateSP(this, &FDatabaseEditor::FillToolbar));
+	}
+
+	void FDatabaseEditor::FillToolbar(FToolBarBuilder& ToolbarBuilder)
+	{
+		ToolbarBuilder.AddToolBarButton(
+			FDatabaseEditorCommands::Get().BuildSearchIndex,
+			NAME_None,
+			TAttribute<FText>(),
+			TAttribute<FText>(),
+			FSlateIcon());
+	}
+
+	void FDatabaseEditor::RegisterTabSpawners(const TSharedRef<FTabManager>& InTabManager)
+	{
+		WorkspaceMenuCategory = InTabManager->AddLocalWorkspaceMenuCategory(
+			LOCTEXT("WorkspaceMenu_PoseSearchDbEditor", "Pose Search Database Editor"));
+		auto WorkspaceMenuCategoryRef = WorkspaceMenuCategory.ToSharedRef();
+
+		FAssetEditorToolkit::RegisterTabSpawners(InTabManager);
+
+		InTabManager->RegisterTabSpawner(
+			FDatabaseEditorTabs::ViewportID,
+			FOnSpawnTab::CreateSP(this, &FDatabaseEditor::SpawnTab_Viewport))
+			.SetDisplayName(LOCTEXT("ViewportTab", "Viewport"))
+			.SetGroup(WorkspaceMenuCategoryRef)
+			.SetIcon(FSlateIcon(FAppStyle::GetAppStyleSetName(), "GraphEditor.EventGraph_16x"));
+
+		InTabManager->RegisterTabSpawner(
+			FDatabaseEditorTabs::AssetDetailsID,
+			FOnSpawnTab::CreateSP(this, &FDatabaseEditor::SpawnTab_AssetDetails))
+			.SetDisplayName(LOCTEXT("DatabaseDetailsTab", "Database Details"))
+			.SetGroup(WorkspaceMenuCategoryRef)
+			.SetIcon(FSlateIcon(FAppStyle::GetAppStyleSetName(), "LevelEditor.Tabs.Details"));
+
+		InTabManager->RegisterTabSpawner(
+			FDatabaseEditorTabs::PreviewSettingsID,
+			FOnSpawnTab::CreateSP(this, &FDatabaseEditor::SpawnTab_PreviewSettings))
+			.SetDisplayName(LOCTEXT("PreviewSceneSettingsTab", "Preview Scene Settings"))
+			.SetGroup(WorkspaceMenuCategoryRef)
+			.SetIcon(FSlateIcon(FAppStyle::GetAppStyleSetName(), "LevelEditor.Tabs.Details"));
+
+		InTabManager->RegisterTabSpawner(
+			FDatabaseEditorTabs::AssetTreeViewID,
+			FOnSpawnTab::CreateSP(this, &FDatabaseEditor::SpawnTab_AssetTreeView))
+			.SetDisplayName(LOCTEXT("TreeViewTab", "Tree View"))
+			.SetGroup(WorkspaceMenuCategoryRef)
+			.SetIcon(FSlateIcon(FAppStyle::GetAppStyleSetName(), "GraphEditor.EventGraph_16x"));
+
+		InTabManager->RegisterTabSpawner(
+			FDatabaseEditorTabs::SelectionDetailsID,
+			FOnSpawnTab::CreateSP(this, &FDatabaseEditor::SpawnTab_SelectionDetails))
+			.SetDisplayName(LOCTEXT("SelectionDetailsTab", "Selection Details"))
+			.SetGroup(WorkspaceMenuCategoryRef)
+			.SetIcon(FSlateIcon(FAppStyle::GetAppStyleSetName(), "GraphEditor.EventGraph_16x"));
+	}
+
+	void FDatabaseEditor::UnregisterTabSpawners(const TSharedRef<FTabManager>& InTabManager)
+	{
+		FAssetEditorToolkit::UnregisterTabSpawners(InTabManager);
+
+		InTabManager->UnregisterTabSpawner(FDatabaseEditorTabs::ViewportID);
+		InTabManager->UnregisterTabSpawner(FDatabaseEditorTabs::AssetDetailsID);
+		InTabManager->UnregisterTabSpawner(FDatabaseEditorTabs::PreviewSettingsID);
+		InTabManager->UnregisterTabSpawner(FDatabaseEditorTabs::AssetTreeViewID);
+		InTabManager->UnregisterTabSpawner(FDatabaseEditorTabs::SelectionDetailsID);
+	}
+
+	FName FDatabaseEditor::GetToolkitFName() const
+	{
+		return FName("PoseSearchDatabaseEditor");
+	}
+
+	FText FDatabaseEditor::GetBaseToolkitName() const
+	{
+		return LOCTEXT("PoseSearchDatabaseEditorAppLabel", "Pose Search Database Editor");
+	}
+
+	FText FDatabaseEditor::GetToolkitName() const
+	{
+		FFormatNamedArguments Args;
+		Args.Add(TEXT("AssetName"), FText::FromString(GetPoseSearchDatabase()->GetName()));
+		return FText::Format(LOCTEXT("PoseSearchDatabaseEditorToolkitName", "{AssetName}"), Args);
+	}
+
+	FLinearColor FDatabaseEditor::GetWorldCentricTabColorScale() const
+	{
+		return FLinearColor::White;
+	}
+
+	FString FDatabaseEditor::GetWorldCentricTabPrefix() const
+	{
+		return TEXT("PoseSearchDatabaseEditor");
+	}
+
+	TSharedRef<SDockTab> FDatabaseEditor::SpawnTab_Viewport(const FSpawnTabArgs& Args)
+	{
+		check(Args.GetTabId() == FDatabaseEditorTabs::ViewportID);
+
+		TSharedRef<SDockTab> SpawnedTab = SNew(SDockTab).Label(LOCTEXT("ViewportTab_Title", "Viewport"));
+
+		if (PreviewWidget.IsValid())
+		{
+			SpawnedTab->SetContent(PreviewWidget.ToSharedRef());
+		}
+
+		return SpawnedTab;
+	}
+
+	TSharedRef<SDockTab> FDatabaseEditor::SpawnTab_AssetDetails(const FSpawnTabArgs& Args)
+	{
+		check(Args.GetTabId() == FDatabaseEditorTabs::AssetDetailsID);
+
+		return SNew(SDockTab)
+			.Label(LOCTEXT("Database_Details_Title", "Database Details"))
+			[
+				EditingAssetWidget.ToSharedRef()
+			];
+	}
+
+	TSharedRef<SDockTab> FDatabaseEditor::SpawnTab_PreviewSettings(const FSpawnTabArgs& Args)
+	{
+		check(Args.GetTabId() == FDatabaseEditorTabs::PreviewSettingsID);
+
+		FAdvancedPreviewSceneModule& AdvancedPreviewSceneModule = 
+			FModuleManager::LoadModuleChecked<FAdvancedPreviewSceneModule>("AdvancedPreviewScene");
+		TSharedRef<SWidget> InWidget = 
+			AdvancedPreviewSceneModule.CreateAdvancedPreviewSceneSettingsWidget(PreviewScene.ToSharedRef());
+
+		TSharedRef<SDockTab> SpawnedTab = SNew(SDockTab)
+			.Label(LOCTEXT("PreviewSceneSettingsTab", "Preview Scene Settings"))
+			[
+				InWidget
+			];
+
+		return SpawnedTab;
+	}
+
+	TSharedRef<SDockTab> FDatabaseEditor::SpawnTab_AssetTreeView(const FSpawnTabArgs& Args)
+	{
+		check(Args.GetTabId() == FDatabaseEditorTabs::AssetTreeViewID);
+
+		return SNew(SDockTab)
+			.Label(LOCTEXT("AssetTreeView_Title", "Asset Tree"))
+			[
+				AssetTreeWidget.ToSharedRef()
+			];
+	}
+
+	TSharedRef<SDockTab> FDatabaseEditor::SpawnTab_SelectionDetails(const FSpawnTabArgs& Args)
+	{
+		check(Args.GetTabId() == FDatabaseEditorTabs::SelectionDetailsID);
+
+		TSharedPtr<SVerticalBox> DetailsContainer;
+		SAssignNew(DetailsContainer, SVerticalBox);
+
+		for (TSharedPtr<IDetailsView> SelectionWidget : SelectionWidgets)
+		{
+			DetailsContainer->AddSlot()
+				.AutoHeight()
+				[
+					SelectionWidget.ToSharedRef()
+				];
+		}
+
+		return SNew(SDockTab)
+			.Label(LOCTEXT("SelectionDetails_Title", "Selection Details"))
+			[
+				SNew(SScrollBox)
+				+SScrollBox::Slot()
+				[
+					DetailsContainer.ToSharedRef()
+				]
+			];
+	}
+
+	void FDatabaseEditor::OnFinishedChangingSelectionProperties(const FPropertyChangedEvent& PropertyChangedEvent)
+	{
+		ViewModel->BuildSearchIndex();
+	}
+
+	void FDatabaseEditor::OnAssetTreeSelectionChanged(
+		const TArrayView<TSharedPtr<FDatabaseAssetTreeNode>>& SelectedItems,
+		ESelectInfo::Type SelectionType)
+	{
+		enum ESelectionAssetTypeIndex
+		{
+			SequenceSelectionIndex = 0,
+			BlendSpaceSelectionIndex = 1,
+			GroupSelectionIndex = 2,
+			SelectionTypeCount = 3
+		};
+
+		TArray<TArray<TWeakObjectPtr<UObject>>> SelectionReflections;
+		SelectionReflections.SetNum(SelectionDetailsCount);
+
+		if (SelectedItems.Num() > 0)
+		{
+			for (TSharedPtr<FDatabaseAssetTreeNode>& SelectedItem : SelectedItems)
+			{
+				if (SelectedItem->SourceAssetType == ESearchIndexAssetType::Sequence)
+				{
+					UPoseSearchDatabaseSequenceReflection* NewSelectionReflection =
+						NewObject<UPoseSearchDatabaseSequenceReflection>();
+					NewSelectionReflection->AddToRoot();
+					NewSelectionReflection->Sequence =
+						GetPoseSearchDatabase()->Sequences[SelectedItem->SourceAssetIdx];
+					NewSelectionReflection->SetSourceLink(SelectedItem, AssetTreeWidget);
+					SelectionReflections[SequenceSelectionIndex].Add(NewSelectionReflection);
+				}
+				else if (SelectedItem->SourceAssetType == ESearchIndexAssetType::BlendSpace)
+				{
+					UPoseSearchDatabaseBlendSpaceReflection* NewSelectionReflection =
+						NewObject<UPoseSearchDatabaseBlendSpaceReflection>();
+					NewSelectionReflection->AddToRoot();
+					NewSelectionReflection->BlendSpace =
+						GetPoseSearchDatabase()->BlendSpaces[SelectedItem->SourceAssetIdx];
+					NewSelectionReflection->SetSourceLink(SelectedItem, AssetTreeWidget);
+					SelectionReflections[BlendSpaceSelectionIndex].Add(NewSelectionReflection);
+				}
+				else
+				{
+					if (SelectedItem->SourceAssetIdx != INDEX_NONE)
+					{
+						UPoseSearchDatabaseGroupReflection* NewSelectionReflection =
+							NewObject<UPoseSearchDatabaseGroupReflection>();
+						NewSelectionReflection->AddToRoot();
+						NewSelectionReflection->Group =
+							GetPoseSearchDatabase()->Groups[SelectedItem->SourceAssetIdx];
+						NewSelectionReflection->SetSourceLink(SelectedItem, AssetTreeWidget);
+						SelectionReflections[GroupSelectionIndex].Add(NewSelectionReflection);
+					}
+				}
+			}
+		}
+
+		int32 SelectionWidgetIdx = 0;
+		for (int32 SelectionTypeIdx = 0; SelectionTypeIdx < SelectionTypeCount; ++SelectionTypeIdx)
+		{
+			if (SelectionReflections[SelectionTypeIdx].Num() > 0)
+			{
+				SelectionWidgets[SelectionWidgetIdx++]->SetObjects(SelectionReflections[SelectionTypeIdx], true);
+			}
+		}
+
+		while (SelectionWidgetIdx < SelectionDetailsCount)
+		{
+			SelectionWidgets[SelectionWidgetIdx++]->SetObject(nullptr, true);
+		}
+
+		ViewModel->SetSelectedNodes(SelectedItems);
+	}
+}
+
+#undef LOCTEXT_NAMESPACE
diff --git a/Engine/Plugins/Experimental/Animation/PoseSearch/Source/Editor/Private/PoseSearchDatabaseEditorToolkit.cpp b/Engine/Plugins/Experimental/Animation/PoseSearch/Source/Editor/Private/PoseSearchDatabaseEditorToolkit.cpp
deleted file mode 100644
diff --git a/Engine/Plugins/Experimental/Animation/PoseSearch/Source/Editor/Private/PoseSearchDatabasePreviewScene.cpp b/Engine/Plugins/Experimental/Animation/PoseSearch/Source/Editor/Private/PoseSearchDatabasePreviewScene.cpp
--- a/Engine/Plugins/Experimental/Animation/PoseSearch/Source/Editor/Private/PoseSearchDatabasePreviewScene.cpp
+++ b/Engine/Plugins/Experimental/Animation/PoseSearch/Source/Editor/Private/PoseSearchDatabasePreviewScene.cpp
@@ -1,7 +1,7 @@
 // Copyright Epic Games, Inc. All Rights Reserved.
 
 #include "PoseSearchDatabasePreviewScene.h"
-#include "PoseSearchDatabaseEditorToolkit.h"
+#include "PoseSearchDatabaseEditor.h"
 #include "PoseSearchDatabaseViewModel.h"
 
 #include "PoseSearch/PoseSearch.h"
diff --git a/Engine/Plugins/Experimental/Animation/PoseSearch/Source/Editor/Private/PoseSearchDatabaseViewportClient.cpp b/Engine/Plugins/Experimental/Animation/PoseSearch/Source/Editor/Private/PoseSearchDatabaseViewportClient.cpp
--- a/Engine/Plugins/Experimental/Animation/PoseSearch/Source/Editor/Private/PoseSearchDatabaseViewportClient.cpp
+++ b/Engine/Plugins/Experimental/Animation/PoseSearch/Source/Editor/Private/PoseSearchDatabaseViewportClient.cpp
@@ -1,7 +1,7 @@
 // Copyright Epic Games, Inc. All Rights Reserved.
 
 #include "PoseSearchDatabaseViewportClient.h"
-#include "PoseSearchDatabaseEditorToolkit.h"
+#include "PoseSearchDatabaseEditor.h"
 #include "PoseSearchDatabasePreviewScene.h"
 #include "SPoseSearchDatabaseViewport.h"
 #include "AssetEditorModeManager.h"
diff --git a/Engine/Plugins/Experimental/Animation/PoseSearch/Source/Editor/Private/PoseSearchTypeActions.cpp b/Engine/Plugins/Experimental/Animation/PoseSearch/Source/Editor/Private/PoseSearchTypeActions.cpp
--- a/Engine/Plugins/Experimental/Animation/PoseSearch/Source/Editor/Private/PoseSearchTypeActions.cpp
+++ b/Engine/Plugins/Experimental/Animation/PoseSearch/Source/Editor/Private/PoseSearchTypeActions.cpp
@@ -1,7 +1,7 @@
 // Copyright Epic Games, Inc. All Rights Reserved.
 
 #include "PoseSearchTypeActions.h"
-#include "PoseSearchDatabaseEditorToolkit.h"
+#include "PoseSearchDatabaseEditor.h"
 #include "PoseSearch/PoseSearch.h"
 
 #define LOCTEXT_NAMESPACE "PoseSearchTypeActions"
diff --git a/Engine/Plugins/Experimental/Animation/PoseSearch/Source/Editor/Private/SPoseSearchDatabaseViewport.cpp b/Engine/Plugins/Experimental/Animation/PoseSearch/Source/Editor/Private/SPoseSearchDatabaseViewport.cpp
--- a/Engine/Plugins/Experimental/Animation/PoseSearch/Source/Editor/Private/SPoseSearchDatabaseViewport.cpp
+++ b/Engine/Plugins/Experimental/Animation/PoseSearch/Source/Editor/Private/SPoseSearchDatabaseViewport.cpp
@@ -3,7 +3,7 @@
 #include "SPoseSearchDatabaseViewport.h"
 #include "PoseSearchDatabaseViewportClient.h"
 #include "PoseSearchDatabaseViewModel.h"
-#include "PoseSearchDatabaseEditorToolkit.h"
+#include "PoseSearchDatabaseEditor.h"
 #include "SPoseSearchDatabaseViewportToolbar.h"
 #include "PoseSearchDatabaseEditorCommands.h"
 
diff --git a/Engine/Plugins/Experimental/Animation/PoseSearch/Source/Editor/Public/PoseSearchDatabaseEditor.h b/Engine/Plugins/Experimental/Animation/PoseSearch/Source/Editor/Public/PoseSearchDatabaseEditor.h
new file mode 100644
--- /dev/null
+++ b/Engine/Plugins/Experimental/Animation/PoseSearch/Source/Editor/Public/PoseSearchDatabaseEditor.h
@@ -0,0 +1,85 @@
+// Copyright Epic Games, Inc. All Rights Reserved.
+
+#pragma once
+
+#include "CoreMinimal.h"
+#include "Toolkits/AssetEditorToolkit.h"
+#include "Misc/NotifyHook.h"
+#include "SPoseSearchDatabaseAssetList.h"
+
+class IDetailsView;
+class FToolBarBuilder;
+class UPoseSearchDatabaseSelectionReflection;
+
+namespace UE::PoseSearch
+{
+	class SDatabasePreview;
+	class FDatabasePreviewScene;
+	class FDatabaseViewModel;
+
+	class FDatabaseEditor : public FAssetEditorToolkit, public FNotifyHook
+	{
+	public:
+
+		FDatabaseEditor();
+		virtual ~FDatabaseEditor();
+
+		void InitAssetEditor(
+			const EToolkitMode::Type Mode,
+			const TSharedPtr<IToolkitHost>& InitToolkitHost,
+			UPoseSearchDatabase* DatabaseAsset);
+
+		virtual void RegisterTabSpawners(const TSharedRef<FTabManager>& TabManager) override;
+		virtual void UnregisterTabSpawners(const TSharedRef<FTabManager>& TabManager) override;
+		virtual FName GetToolkitFName() const override;
+		virtual FText GetBaseToolkitName() const override;
+		virtual FText GetToolkitName() const override;
+		virtual FLinearColor GetWorldCentricTabColorScale() const override;
+		virtual FString GetWorldCentricTabPrefix() const override;
+
+		const UPoseSearchDatabase* GetPoseSearchDatabase() const;
+		UPoseSearchDatabase* GetPoseSearchDatabase();
+		FDatabaseViewModel* GetViewModel() const { return ViewModel.Get(); }
+		TSharedPtr<FDatabaseViewModel> GetViewModelSharedPtr() const { return ViewModel; }
+
+		void BuildSearchIndex();
+
+		void PreviewBackwardEnd();
+		void PreviewBackwardStep();
+		void PreviewBackward();
+		void PreviewPause();
+		void PreviewForward();
+		void PreviewForwardStep();
+		void PreviewForwardEnd();
+
+	private:
+
+		TSharedRef<SDockTab> SpawnTab_Viewport(const FSpawnTabArgs& Args);
+		TSharedRef<SDockTab> SpawnTab_AssetDetails(const FSpawnTabArgs& Args);
+		TSharedRef<SDockTab> SpawnTab_PreviewSettings(const FSpawnTabArgs& Args);
+		TSharedRef<SDockTab> SpawnTab_AssetTreeView(const FSpawnTabArgs& Args);
+		TSharedRef<SDockTab> SpawnTab_SelectionDetails(const FSpawnTabArgs& Args);
+
+		void BindCommands();
+		void ExtendToolbar();
+		void FillToolbar(FToolBarBuilder& ToolbarBuilder);
+
+		void OnFinishedChangingSelectionProperties(const FPropertyChangedEvent& PropertyChangedEvent);
+		void OnAssetTreeSelectionChanged(
+			const TArrayView<TSharedPtr<FDatabaseAssetTreeNode>>& SelectedItems,
+			ESelectInfo::Type SelectionType);
+
+		TSharedPtr<SDatabasePreview> PreviewWidget;
+
+		TSharedPtr<SDatabaseAssetTree> AssetTreeWidget;
+
+		TSharedPtr<IDetailsView> EditingAssetWidget;
+
+		TArray<TSharedPtr<IDetailsView>> SelectionWidgets;
+
+		TSharedPtr<FDatabasePreviewScene> PreviewScene;
+
+		TSharedPtr<FDatabaseViewModel> ViewModel;
+	};
+}
+
diff --git a/Engine/Plugins/Experimental/Animation/PoseSearch/Source/Editor/Public/PoseSearchDatabaseEditorToolkit.h b/Engine/Plugins/Experimental/Animation/PoseSearch/Source/Editor/Public/PoseSearchDatabaseEditorToolkit.h
deleted file mode 100644
