Pose search now uses the bLoop flag from sequences/blendspaces
#rnx

diff --git a/Engine/Plugins/Experimental/Animation/PoseSearch/Source/Runtime/Private/AnimNode_MotionMatching.cpp b/Engine/Plugins/Experimental/Animation/PoseSearch/Source/Runtime/Private/AnimNode_MotionMatching.cpp
--- a/Engine/Plugins/Experimental/Animation/PoseSearch/Source/Runtime/Private/AnimNode_MotionMatching.cpp
+++ b/Engine/Plugins/Experimental/Animation/PoseSearch/Source/Runtime/Private/AnimNode_MotionMatching.cpp
@@ -5,6 +5,7 @@
 #include "Animation/AnimInstanceProxy.h"
 #include "Animation/AnimNode_Inertialization.h"
 #include "Animation/AnimSequence.h"
+#include "Animation/BlendSpace.h"
 #include "Animation/MotionTrajectoryTypes.h"
 #include "Animation/AnimRootMotionProvider.h"
 #include "DynamicPlayRate/DynamicPlayRateLibrary.h"
@@ -120,7 +121,7 @@ void FAnimNode_MotionMatching::UpdateAssetPlayer(const FAnimationUpdateContext&
 			const FPoseSearchDatabaseSequence& ResultDbSequence = Database->GetSequenceSourceAsset(SearchIndexAsset);
 			SequencePlayerNode.SetAccumulatedTime(MotionMatchingState.AssetPlayerTime);
 			SequencePlayerNode.SetSequence(ResultDbSequence.Sequence);
-			SequencePlayerNode.SetLoopAnimation(ResultDbSequence.bLoopAnimation);
+			SequencePlayerNode.SetLoopAnimation(ResultDbSequence.Sequence->bLoop);
 			SequencePlayerNode.SetPlayRate(1.0f);
 		}
 		else if (SearchIndexAsset->Type == ESearchIndexAssetType::BlendSpace)
@@ -130,7 +131,7 @@ void FAnimNode_MotionMatching::UpdateAssetPlayer(const FAnimationUpdateContext&
 			const FPoseSearchDatabaseBlendSpace& ResultDbBlendSpace = Database->GetBlendSpaceSourceAsset(SearchIndexAsset);
 			BlendSpacePlayerNode.SetAccumulatedTime(MotionMatchingState.AssetPlayerTime);
 			BlendSpacePlayerNode.SetBlendSpace(ResultDbBlendSpace.BlendSpace);
-			BlendSpacePlayerNode.SetLoop(ResultDbBlendSpace.bLoopAnimation);
+			BlendSpacePlayerNode.SetLoop(ResultDbBlendSpace.BlendSpace->bLoop);
 			BlendSpacePlayerNode.SetPlayRate(1.0f);
 			BlendSpacePlayerNode.SetPosition(SearchIndexAsset->BlendParameters);
 		}
diff --git a/Engine/Plugins/Experimental/Animation/PoseSearch/Source/Runtime/Private/PoseSearch.cpp b/Engine/Plugins/Experimental/Animation/PoseSearch/Source/Runtime/Private/PoseSearch.cpp
--- a/Engine/Plugins/Experimental/Animation/PoseSearch/Source/Runtime/Private/PoseSearch.cpp
+++ b/Engine/Plugins/Experimental/Animation/PoseSearch/Source/Runtime/Private/PoseSearch.cpp
@@ -1174,11 +1174,11 @@ const bool UPoseSearchDatabase::IsSourceAssetLooping(const FPoseSearchIndexAsset
 {
 	if (SearchIndexAsset->Type == ESearchIndexAssetType::Sequence)
 	{
-		return GetSequenceSourceAsset(SearchIndexAsset).bLoopAnimation;
+		return GetSequenceSourceAsset(SearchIndexAsset).Sequence->bLoop;
 	}
 	else if (SearchIndexAsset->Type == ESearchIndexAssetType::BlendSpace)
 	{
-		return GetBlendSpaceSourceAsset(SearchIndexAsset).bLoopAnimation;
+		return GetBlendSpaceSourceAsset(SearchIndexAsset).BlendSpace->bLoop;
 	}
 	else
 	{
diff --git a/Engine/Plugins/Experimental/Animation/PoseSearch/Source/Runtime/Private/PoseSearchDerivedData.cpp b/Engine/Plugins/Experimental/Animation/PoseSearch/Source/Runtime/Private/PoseSearchDerivedData.cpp
--- a/Engine/Plugins/Experimental/Animation/PoseSearch/Source/Runtime/Private/PoseSearchDerivedData.cpp
+++ b/Engine/Plugins/Experimental/Animation/PoseSearch/Source/Runtime/Private/PoseSearchDerivedData.cpp
@@ -244,16 +244,25 @@ namespace UE::PoseSearch
 		// Main Sequence
 		AddRawSequenceToWriter(DbSequence.Sequence, InOutWriter);
 		InOutWriter.Update(&DbSequence.SamplingRange, sizeof(DbSequence.SamplingRange));
-		InOutWriter.Update(&DbSequence.bLoopAnimation, sizeof(DbSequence.bLoopAnimation));
+		if (DbSequence.Sequence)
+		{
+			InOutWriter.Update(&DbSequence.Sequence->bLoop, sizeof(DbSequence.Sequence->bLoop));
+		}
 		InOutWriter.Update(&DbSequence.MirrorOption, sizeof(DbSequence.MirrorOption));
 
 		// Lead in sequence
 		AddRawSequenceToWriter(DbSequence.LeadInSequence, InOutWriter);
-		InOutWriter.Update(&DbSequence.bLoopLeadInAnimation, sizeof(DbSequence.bLoopLeadInAnimation));
+		if (DbSequence.LeadInSequence)
+		{
+			InOutWriter.Update(&DbSequence.LeadInSequence->bLoop, sizeof(DbSequence.LeadInSequence->bLoop));
+		}
 
 		// Follow up sequence
 		AddRawSequenceToWriter(DbSequence.FollowUpSequence, InOutWriter);
-		InOutWriter.Update(&DbSequence.bLoopFollowUpAnimation, sizeof(DbSequence.bLoopFollowUpAnimation));
+		if (DbSequence.FollowUpSequence)
+		{
+			InOutWriter.Update(&DbSequence.FollowUpSequence->bLoop, sizeof(DbSequence.FollowUpSequence->bLoop));
+		}
 
 		// Tags
 		InOutWriter.Update(&DbSequence.GroupTags, sizeof(DbSequence.GroupTags));
@@ -323,7 +332,7 @@ namespace UE::PoseSearch
 				InOutWriter.Update(&Sample.RateScale, sizeof(Sample.RateScale));
 			}
 
-			InOutWriter.Update(&DbBlendSpace.bLoopAnimation, sizeof(DbBlendSpace.bLoopAnimation));
+			InOutWriter.Update(&DbBlendSpace.BlendSpace->bLoop, sizeof(DbBlendSpace.BlendSpace->bLoop));
 			InOutWriter.Update(&DbBlendSpace.MirrorOption, sizeof(DbBlendSpace.MirrorOption));
 			InOutWriter.Update(&DbBlendSpace.bUseGridForSampling, sizeof(DbBlendSpace.bUseGridForSampling));
 			InOutWriter.Update(&DbBlendSpace.NumberOfHorizontalSamples, sizeof(DbBlendSpace.NumberOfHorizontalSamples));
