Remove loop flags from pose search database and replace with loop flag from animation assets
#rnx

diff --git a/Engine/Plugins/Experimental/Animation/PoseSearch/Source/Editor/Private/PoseSearchDebugger.cpp b/Engine/Plugins/Experimental/Animation/PoseSearch/Source/Editor/Private/PoseSearchDebugger.cpp
--- a/Engine/Plugins/Experimental/Animation/PoseSearch/Source/Editor/Private/PoseSearchDebugger.cpp
+++ b/Engine/Plugins/Experimental/Animation/PoseSearch/Source/Editor/Private/PoseSearchDebugger.cpp
@@ -950,7 +950,7 @@ void SDebuggerDatabaseView::CreateRows(const UPoseSearchDatabase& Database)
 				Row->AssetTime = Time;
 				Row->AnimFrame = DbSequence.Sequence->GetFrameAtTime(Time);
 				Row->bMirrored = SearchIndexAsset.bMirrored;
-				Row->bLooping = DbSequence.bLoopAnimation;
+				Row->bLooping = DbSequence.Sequence->bLoop;
 				Row->BlendParameters = FVector::Zero();
 			}
 		}
@@ -971,7 +971,7 @@ void SDebuggerDatabaseView::CreateRows(const UPoseSearchDatabase& Database)
 				Row->AssetTime = Time;
 				Row->AnimFrame = 0; // There is no frame index associated with a blendspace
 				Row->bMirrored = SearchIndexAsset.bMirrored;
-				Row->bLooping = DbBlendSpace.bLoopAnimation;
+				Row->bLooping = DbBlendSpace.BlendSpace->bLoop;
 				Row->BlendParameters = SearchIndexAsset.BlendParameters;
 			}
 		}
@@ -2413,7 +2413,7 @@ void FDebuggerViewModel::OnDraw(FSkeletonDrawParams& DrawParams)
 			UpdateContext.StartTime = Skeletons[SelectedPose].Time;
 			UpdateContext.Time = Skeletons[SelectedPose].Time;
 			UpdateContext.bMirrored = Skeletons[SelectedPose].bMirrored;
-			UpdateContext.bLoop = DatabaseSequence->bLoopAnimation;
+			UpdateContext.bLoop = DatabaseSequence->Sequence->bLoop;
 		}
 		else if (Skeletons[SelectedPose].Type == ESearchIndexAssetType::BlendSpace)
 		{
@@ -2424,7 +2424,7 @@ void FDebuggerViewModel::OnDraw(FSkeletonDrawParams& DrawParams)
 			UpdateContext.StartTime = Skeletons[SelectedPose].Time;
 			UpdateContext.Time = Skeletons[SelectedPose].Time;
 			UpdateContext.bMirrored = Skeletons[SelectedPose].bMirrored;
-			UpdateContext.bLoop = DatabaseBlendSpace->bLoopAnimation;
+			UpdateContext.bLoop = DatabaseBlendSpace->BlendSpace->bLoop;
 			UpdateContext.BlendParameters = Skeletons[SelectedPose].BlendParameters;
 		}
 		else
@@ -2450,7 +2450,7 @@ void FDebuggerViewModel::OnDraw(FSkeletonDrawParams& DrawParams)
 			UpdateContext.StartTime = Skeletons[Asset].Time;
 			UpdateContext.Time = Skeletons[Asset].Time;
 			UpdateContext.bMirrored = Skeletons[Asset].bMirrored;
-			UpdateContext.bLoop = DatabaseSequence->bLoopAnimation;
+			UpdateContext.bLoop = DatabaseSequence->Sequence->bLoop;
 		}
 		else if (Skeletons[SelectedPose].Type == ESearchIndexAssetType::BlendSpace)
 		{
@@ -2461,7 +2461,7 @@ void FDebuggerViewModel::OnDraw(FSkeletonDrawParams& DrawParams)
 			UpdateContext.StartTime = Skeletons[Asset].Time;
 			UpdateContext.Time = Skeletons[Asset].Time;
 			UpdateContext.bMirrored = Skeletons[Asset].bMirrored;
-			UpdateContext.bLoop = DatabaseBlendSpace->bLoopAnimation;
+			UpdateContext.bLoop = DatabaseBlendSpace->BlendSpace->bLoop;
 			UpdateContext.BlendParameters = Skeletons[Asset].BlendParameters;
 		}
 		else
@@ -2578,13 +2578,13 @@ void FDebuggerViewModel::UpdateAsset()
 	{
 		const FPoseSearchDatabaseSequence* DatabaseSequence = GetAnimSequence(AssetSkeleton.AssetIdx);
 		AnimAsset = DatabaseSequence->Sequence;
-		bAssetLooping = DatabaseSequence->bLoopAnimation;
+		bAssetLooping = DatabaseSequence->Sequence->bLoop;
 	}
 	else if (AssetSkeleton.Type == ESearchIndexAssetType::BlendSpace)
 	{
 		const FPoseSearchDatabaseBlendSpace* DatabaseBlendSpace = GetBlendSpace(AssetSkeleton.AssetIdx);
 		AnimAsset = DatabaseBlendSpace->BlendSpace;
-		bAssetLooping = DatabaseBlendSpace->bLoopAnimation;
+		bAssetLooping = DatabaseBlendSpace->BlendSpace->bLoop;
 	}
 	else
 	{
diff --git a/Engine/Plugins/Experimental/Animation/PoseSearch/Source/Runtime/Private/PoseSearch.cpp b/Engine/Plugins/Experimental/Animation/PoseSearch/Source/Runtime/Private/PoseSearch.cpp
--- a/Engine/Plugins/Experimental/Animation/PoseSearch/Source/Runtime/Private/PoseSearch.cpp
+++ b/Engine/Plugins/Experimental/Animation/PoseSearch/Source/Runtime/Private/PoseSearch.cpp
@@ -4588,17 +4588,17 @@ void FDatabaseIndexingContext::PrepareSamplers()
 	{
 		if (DbSequence.Sequence)
 		{
-			AddSequenceSampler(DbSequence.Sequence, DbSequence.bLoopAnimation);
+			AddSequenceSampler(DbSequence.Sequence, DbSequence.Sequence->bLoop);
 		}
 
 		if (DbSequence.LeadInSequence)
 		{
-			AddSequenceSampler(DbSequence.LeadInSequence, DbSequence.bLoopLeadInAnimation);
+			AddSequenceSampler(DbSequence.LeadInSequence, DbSequence.LeadInSequence->bLoop);
 		}
 
 		if (DbSequence.FollowUpSequence)
 		{
-			AddSequenceSampler(DbSequence.FollowUpSequence, DbSequence.bLoopFollowUpAnimation);
+			AddSequenceSampler(DbSequence.FollowUpSequence, DbSequence.FollowUpSequence->bLoop);
 		}
 	}
 
@@ -4648,7 +4648,7 @@ void FDatabaseIndexingContext::PrepareSamplers()
 						Input.SamplingContext = &SamplingContext;
 						Input.ExtrapolationParameters = Database->ExtrapolationParameters;
 						Input.BlendSpace = DbBlendSpace.BlendSpace;
-						Input.bLoopable = DbBlendSpace.bLoopAnimation;
+						Input.bLoopable = DbBlendSpace.BlendSpace->bLoop;
 						Input.BlendParameters = BlendParameters;
 
 						BlendSpaceSamplers[BlendSpaceSamplerIdx].Init(Input);
diff --git a/Engine/Plugins/Experimental/Animation/PoseSearch/Source/Runtime/Private/PoseSearchLibrary.cpp b/Engine/Plugins/Experimental/Animation/PoseSearch/Source/Runtime/Private/PoseSearchLibrary.cpp
--- a/Engine/Plugins/Experimental/Animation/PoseSearch/Source/Runtime/Private/PoseSearchLibrary.cpp
+++ b/Engine/Plugins/Experimental/Animation/PoseSearch/Source/Runtime/Private/PoseSearchLibrary.cpp
@@ -200,7 +200,7 @@ bool FMotionMatchingState::CanAdvance(float DeltaTime, bool& bOutAdvanceToFollow
 
 		float SteppedTime = AssetPlayerTime;
 		ETypeAdvanceAnim AdvanceType = FAnimationRuntime::AdvanceTime(
-			DbSequence.bLoopAnimation,
+			DbSequence.Sequence->bLoop,
 			DeltaTime,
 			SteppedTime,
 			AssetLength);
@@ -265,7 +265,7 @@ bool FMotionMatchingState::CanAdvance(float DeltaTime, bool& bOutAdvanceToFollow
 		float RealTime = AssetPlayerTime * PlayLength;
 		float SteppedTime = RealTime;
 		ETypeAdvanceAnim AdvanceType = FAnimationRuntime::AdvanceTime(
-			DbBlendSpace.bLoopAnimation,
+			DbBlendSpace.BlendSpace->bLoop,
 			DeltaTime,
 			SteppedTime,
 			PlayLength);
diff --git a/Engine/Plugins/Experimental/Animation/PoseSearch/Source/Runtime/Public/PoseSearch/PoseSearch.h b/Engine/Plugins/Experimental/Animation/PoseSearch/Source/Runtime/Public/PoseSearch/PoseSearch.h
--- a/Engine/Plugins/Experimental/Animation/PoseSearch/Source/Runtime/Public/PoseSearch/PoseSearch.h
+++ b/Engine/Plugins/Experimental/Animation/PoseSearch/Source/Runtime/Public/PoseSearch/PoseSearch.h
@@ -855,9 +855,6 @@ struct POSESEARCH_API FPoseSearchDatabaseSequence
 	UPROPERTY(EditAnywhere, Category="Sequence")
 	FFloatInterval SamplingRange = FFloatInterval(0.0f, 0.0f);
 
-	UPROPERTY(EditAnywhere, Category="Sequence")
-	bool bLoopAnimation = false;
-
 	UPROPERTY(EditAnywhere, Category = "Sequence")
 	EPoseSearchMirrorOption MirrorOption = EPoseSearchMirrorOption::UnmirroredOnly;
 
@@ -868,9 +865,6 @@ struct POSESEARCH_API FPoseSearchDatabaseSequence
 	UPROPERTY(EditAnywhere, Category="Sequence")
 	TObjectPtr<UAnimSequence> LeadInSequence = nullptr;
 
-	UPROPERTY(EditAnywhere, Category="Sequence")
-	bool bLoopLeadInAnimation = false;
-
 	// Used for sampling future pose information at the end of the main sequence.
 	// This setting is intended for transitions between cycles. It is optional and only used
 	// for one shot anims with future sampling. When future sampling is used without a follow up sequence,
@@ -878,9 +872,6 @@ struct POSESEARCH_API FPoseSearchDatabaseSequence
 	UPROPERTY(EditAnywhere, Category="Sequence")
 	TObjectPtr<UAnimSequence> FollowUpSequence = nullptr;
 
-	UPROPERTY(EditAnywhere, Category="Sequence")
-	bool bLoopFollowUpAnimation = false;
-
 	UPROPERTY(EditAnywhere, Category = "Group")
 	FGameplayTagContainer GroupTags;
 
@@ -896,9 +887,6 @@ struct POSESEARCH_API FPoseSearchDatabaseBlendSpace
 	UPROPERTY(EditAnywhere, Category = "BlendSpace")
 	TObjectPtr<UBlendSpace> BlendSpace = nullptr;
 
-	UPROPERTY(EditAnywhere, Category = "BlendSpace")
-	bool bLoopAnimation = false;
-
 	UPROPERTY(EditAnywhere, Category = "BlendSpace")
 	EPoseSearchMirrorOption MirrorOption = EPoseSearchMirrorOption::UnmirroredOnly;
 
