PoseSearch: Database Group rework phase 1
- Database now implements SearchableAsset interface
- New DatabaseSet also implements SearchableAsset interface and allows using individual databases as groups
- Search function moved to the SearchableAsset interface
- Search queries now built inside search function
- Each entry in DatabaseSet can be dynamically enabled/disabled
- Each entry in the DatabaseSet can change the resulting pose cost

diff --git a/Engine/Plugins/Experimental/Animation/PoseSearch/PoseSearch.uplugin b/Engine/Plugins/Experimental/Animation/PoseSearch/PoseSearch.uplugin
--- a/Engine/Plugins/Experimental/Animation/PoseSearch/PoseSearch.uplugin
+++ b/Engine/Plugins/Experimental/Animation/PoseSearch/PoseSearch.uplugin
@@ -38,6 +38,10 @@
 		{
 			"Name": "AnimationWarping",
 			"Enabled": true
+		},
+		{
+			"Name": "MotionTrajectory",
+			"Enabled": true
 		}
 	]
 }
\ No newline at end of file
diff --git a/Engine/Plugins/Experimental/Animation/PoseSearch/Source/Editor/Private/PoseSearchDatabaseEdMode.cpp b/Engine/Plugins/Experimental/Animation/PoseSearch/Source/Editor/Private/PoseSearchDatabaseEdMode.cpp
--- a/Engine/Plugins/Experimental/Animation/PoseSearch/Source/Editor/Private/PoseSearchDatabaseEdMode.cpp
+++ b/Engine/Plugins/Experimental/Animation/PoseSearch/Source/Editor/Private/PoseSearchDatabaseEdMode.cpp
@@ -26,9 +26,15 @@ namespace UE::PoseSearch
 
 		FDatabaseViewportClient* PoseSearchDbViewportClient =
 			static_cast<FDatabaseViewportClient*>(ViewportClient);
-		if (!ViewModel && PoseSearchDbViewportClient)
+		if (PoseSearchDbViewportClient)
 		{
-			ViewModel = PoseSearchDbViewportClient->GetAssetEditor()->GetViewModel();
+			// ensure we redraw even if PIE is active
+			PoseSearchDbViewportClient->Invalidate();
+
+			if (!ViewModel)
+			{
+				ViewModel = PoseSearchDbViewportClient->GetAssetEditor()->GetViewModel();
+			}
 		}
 
 		if (ViewModel)
diff --git a/Engine/Plugins/Experimental/Animation/PoseSearch/Source/Editor/Private/PoseSearchDebugger.cpp b/Engine/Plugins/Experimental/Animation/PoseSearch/Source/Editor/Private/PoseSearchDebugger.cpp
--- a/Engine/Plugins/Experimental/Animation/PoseSearch/Source/Editor/Private/PoseSearchDebugger.cpp
+++ b/Engine/Plugins/Experimental/Animation/PoseSearch/Source/Editor/Private/PoseSearchDebugger.cpp
@@ -288,7 +288,7 @@ public:
 
 	void CalculateColors(FChannelCostRange TotalCostRange, TArrayView<const FChannelCostRange> ChannelCostRanges)
 	{
-		const float CostColorBlend = (PoseCostDetails.PoseCost.TotalCost - TotalCostRange.Min) / TotalCostRange.Delta;
+		const float CostColorBlend = (PoseCostDetails.PoseCost.GetTotalCost() - TotalCostRange.Min) / TotalCostRange.Delta;
 		CostColor = LinearColorBlend(FLinearColor::Green, FLinearColor::Red, CostColorBlend);
 
 		ChannelCostColors.SetNum(ChannelCostRanges.Num());
@@ -509,7 +509,7 @@ namespace DebuggerDatabaseColumns
 		
 		virtual FText GetRowText(const FRowDataRef& Row) const override
         {
-        	return FText::AsNumber(Row->PoseCostDetails.PoseCost.TotalCost);
+        	return FText::AsNumber(Row->PoseCostDetails.PoseCost.GetTotalCost());
         }
 
 		virtual FSlateColor GetColorAndOpacity(const FRowDataRef& Row) const override
@@ -1003,10 +1003,11 @@ void SDebuggerDatabaseView::UpdateRows(const FTraceMotionMatchingStateMessage& S
 	FPoseSearchWeightsContext StateWeights;
 	StateWeights.Update(&Database);
 
-	UE::PoseSearch::FSearchContext SearchContext;
-	SearchContext.SetSource(&Database);
-	SearchContext.QueryValues = State.QueryVectorNormalized;
-	SearchContext.WeightsContext = &StateWeights;
+	FPoseSearchContext SearchContext;
+
+	FPoseSearchWeightsContext WeightsContext;
+	WeightsContext.Update(&Database);
+
 	if (const FPoseSearchIndexAsset* CurrentIndexAsset = Database.GetSearchIndex()->FindAssetForPose(State.DbPoseIdx))
 	{
 		SearchContext.QueryMirrorRequest = CurrentIndexAsset->bMirrored ? 
@@ -1021,7 +1022,12 @@ void SDebuggerDatabaseView::UpdateRows(const FTraceMotionMatchingStateMessage& S
 		{
 			const int32 PoseIdx = Row->PoseIdx;
 
-			ComparePoses(PoseIdx, SearchContext, Row->PoseCostDetails);
+			Database.ComparePoses(
+				SearchContext, 
+				WeightsContext, 
+				PoseIdx, 
+				State.QueryVectorNormalized, 
+				Row->PoseCostDetails);
 
 			// If we are on the active pose for the frame
 			if (PoseIdx == State.DbPoseIdx)
@@ -1068,7 +1074,7 @@ void SDebuggerDatabaseView::ComputeFilteredDatabaseRowsColors()
 	TArray<FChannelCostRange> ChannelCostRanges;
 	for (const TSharedRef<FDebuggerDatabaseRowData>& Row : FilteredDatabaseView.Rows)
 	{
-		const float Cost = Row->PoseCostDetails.PoseCost.TotalCost;
+		const float Cost = Row->PoseCostDetails.PoseCost.GetTotalCost();
 		CostRange.Min = FMath::Min(CostRange.Min, Cost);
 		CostRange.Max = FMath::Max(CostRange.Max, Cost);
 
diff --git a/Engine/Plugins/Experimental/Animation/PoseSearch/Source/Runtime/PoseSearch.Build.cs b/Engine/Plugins/Experimental/Animation/PoseSearch/Source/Runtime/PoseSearch.Build.cs
--- a/Engine/Plugins/Experimental/Animation/PoseSearch/Source/Runtime/PoseSearch.Build.cs
+++ b/Engine/Plugins/Experimental/Animation/PoseSearch/Source/Runtime/PoseSearch.Build.cs
@@ -36,7 +36,8 @@ public class PoseSearch : ModuleRules
 				"TraceLog",
 				"AnimGraphRuntime",
 				"DynamicPlayRate",
-				"GameplayTags"
+				"GameplayTags",
+				"MotionTrajectory"
 			}
 		);
 
diff --git a/Engine/Plugins/Experimental/Animation/PoseSearch/Source/Runtime/Private/AnimNode_MotionMatching.cpp b/Engine/Plugins/Experimental/Animation/PoseSearch/Source/Runtime/Private/AnimNode_MotionMatching.cpp
--- a/Engine/Plugins/Experimental/Animation/PoseSearch/Source/Runtime/Private/AnimNode_MotionMatching.cpp
+++ b/Engine/Plugins/Experimental/Animation/PoseSearch/Source/Runtime/Private/AnimNode_MotionMatching.cpp
@@ -24,15 +24,9 @@ void FAnimNode_MotionMatching::Initialize_AnyThread(const FAnimationInitializeCo
 
 	GetEvaluateGraphExposedInputs().Execute(Context);
 
-	MotionMatchingState.InitNewDatabaseSearch(Database, nullptr /*OutError*/);
-
 	CurrentAssetPlayerNode = &SequencePlayerNode;
 
 	MirrorNode.SetSourceLinkNode(CurrentAssetPlayerNode);
-	if (Database && Database->Schema)
-	{
-		MirrorNode.SetMirrorDataTable(Database->Schema->MirrorDataTable.Get());
-	}
 
 	BlendSpacePlayerNode.SetResetPlayTimeWhenBlendSpaceChanges(false /*!bReset*/);
 
@@ -76,10 +70,6 @@ void FAnimNode_MotionMatching::UpdateAssetPlayer(const FAnimationUpdateContext&
 		UpdateCounter.HasEverBeenUpdated() &&
 		!UpdateCounter.WasSynchronizedCounter(Context.AnimInstanceProxy->GetUpdateCounter());
 
-#if WITH_EDITOR
-	bNeedsReset = bNeedsReset || MotionMatchingState.HasSearchIndexChanged();
-#endif
-
 	// If we just became relevant and haven't been initialized yet, then reset motion matching state, otherwise update the asset time using the player node.
 	if (bNeedsReset)
 	{
@@ -97,16 +87,19 @@ void FAnimNode_MotionMatching::UpdateAssetPlayer(const FAnimationUpdateContext&
 	// Execute core motion matching algorithm
 	UpdateMotionMatchingState(
 		Context,
-		Database,
+		Searchable,
 		bUseDatabaseTagQuery ? &DatabaseTagQuery : nullptr,
+		&ActiveTagsContainer,
 		Trajectory,
 		Settings,
 		MotionMatchingState
 	);
 
-	if (MotionMatchingState.CurrentDatabase.IsValid() && MotionMatchingState.CurrentDatabase->Schema)
+
+	if (MotionMatchingState.CurrentSearchResult.Database.IsValid() && 
+		MotionMatchingState.CurrentSearchResult.Database->Schema)
 	{
-		MirrorNode.SetMirrorDataTable(MotionMatchingState.CurrentDatabase->Schema->MirrorDataTable.Get());
+		MirrorNode.SetMirrorDataTable(MotionMatchingState.CurrentSearchResult.Database->Schema->MirrorDataTable.Get());
 	}
 
 	const FPoseSearchIndexAsset* SearchIndexAsset = MotionMatchingState.GetCurrentSearchIndexAsset();
@@ -114,6 +107,8 @@ void FAnimNode_MotionMatching::UpdateAssetPlayer(const FAnimationUpdateContext&
 	// If a new pose is requested, jump to the pose by updating the embedded sequence player node
 	if (SearchIndexAsset && (MotionMatchingState.Flags & EMotionMatchingFlags::JumpedToPose) == EMotionMatchingFlags::JumpedToPose)
 	{
+		const UPoseSearchDatabase* Database = MotionMatchingState.CurrentSearchResult.Database.Get();
+
 		if (SearchIndexAsset->Type == ESearchIndexAssetType::Sequence)
 		{
 			CurrentAssetPlayerNode = &SequencePlayerNode;
@@ -182,18 +177,18 @@ void FAnimNode_MotionMatching::PreUpdate(const UAnimInstance* InAnimInstance)
 
 		UE::PoseSearch::FDebugDrawParams DrawParams;
 		DrawParams.RootTransform = SkeletalMeshComponent->GetComponentTransform();
-		DrawParams.Database = Database;
+		DrawParams.Database = MotionMatchingState.CurrentSearchResult.Database.Get();
 		DrawParams.World = SkeletalMeshComponent->GetWorld();
 		DrawParams.DefaultLifeTime = 0.0f;
 
 		if (bDebugDrawMatch)
 		{
-			DrawParams.PoseIdx = MotionMatchingState.DbPoseIdx;
+			DrawParams.PoseIdx = MotionMatchingState.CurrentSearchResult.PoseIdx;
 		}
 
 		if (bDebugDrawQuery)
 		{
-			DrawParams.PoseVector = MotionMatchingState.ComposedQuery.GetValues();
+			DrawParams.PoseVector = MotionMatchingState.CurrentSearchResult.ComposedQuery.GetValues();
 		}
 
 		UE::PoseSearch::Draw(DrawParams);
diff --git a/Engine/Plugins/Experimental/Animation/PoseSearch/Source/Runtime/Private/PoseSearch.cpp b/Engine/Plugins/Experimental/Animation/PoseSearch/Source/Runtime/Private/PoseSearch.cpp
--- a/Engine/Plugins/Experimental/Animation/PoseSearch/Source/Runtime/Private/PoseSearch.cpp
+++ b/Engine/Plugins/Experimental/Animation/PoseSearch/Source/Runtime/Private/PoseSearch.cpp
@@ -552,6 +552,33 @@
 }
 
 
+bool UPoseSearchSchema::BuildQuery(
+	FPoseSearchContext& SearchContext,
+	FPoseSearchFeatureVectorBuilder& InOutQuery) const
+{
+	InOutQuery.Init(this);
 
+	// Copy search query directly from the database if we have an active pose
+	const bool bCopySearchResult = 
+		SearchContext.CurrentResult.IsValid() &&
+		SearchContext.CurrentResult.Database->Schema == this;
+
+	if (bCopySearchResult)
+	{
+		InOutQuery.CopyFromSearchIndex(
+			*SearchContext.CurrentResult.Database->GetSearchIndex(),
+			SearchContext.CurrentResult.PoseIdx);
+	}
+
+	bool bSuccess = true;
+	for (const TObjectPtr<UPoseSearchFeatureChannel>& Channel : Channels)
+	{
+		bool bChannelSuccess = Channel->BuildQuery(SearchContext, InOutQuery);
+		bSuccess &= bChannelSuccess;
+	}
+
+	return bSuccess;
+}
+
 //////////////////////////////////////////////////////////////////////////
 // FPoseSearchChannelWeightParams
@@ -889,6 +916,7 @@
 	return nullptr;
 }
 
+
 //////////////////////////////////////////////////////////////////////////
 // FPoseSearchIndexAsset
 // 
@@ -1087,6 +1115,86 @@
 	return IsValidForIndexing() && SearchIndex.IsValid() && !SearchIndex.IsEmpty();
 }
 
+UE::PoseSearch::FSearchResult UPoseSearchSequenceMetaData::Search(FPoseSearchContext& SearchContext) const
+{
+	using namespace UE::PoseSearch;
 
+	FSearchResult Result;
+
+	if (!ensure(SearchIndex.IsValid() && !SearchIndex.IsEmpty()))
+	{
+		return Result;
+	}
+
+	Schema->BuildQuery(SearchContext, Result.ComposedQuery);
+	Result.ComposedQuery.Normalize(SearchIndex);
+	TArrayView<const float> NormalizedQueryValues = Result.ComposedQuery.GetNormalizedValues();
+
+	if (!ensure(NormalizedQueryValues.Num() == SearchIndex.Schema->Layout.NumFloats))
+	{
+		return Result;
+	}
+
+	FPoseSearchCost BestPoseCost;
+	int32 BestPoseIdx = INDEX_NONE;
+	for (const FPoseSearchIndexAsset& Asset : SearchIndex.Assets)
+	{
+		const int32 EndIndex = Asset.FirstPoseIdx + Asset.NumPoses;
+		for (int32 PoseIdx = Asset.FirstPoseIdx; PoseIdx < EndIndex; ++PoseIdx)
+		{
+			const FPoseSearchPoseMetadata& Metadata = SearchIndex.PoseMetadata[PoseIdx];
+
+			if (EnumHasAnyFlags(Metadata.Flags, EPoseSearchPoseFlags::BlockTransition))
+			{
+				continue;
+			}
+
+			FPoseSearchCost PoseCost = ComparePoses(PoseIdx, NormalizedQueryValues);
+
+			if (PoseCost < BestPoseCost)
+			{
+				BestPoseCost = PoseCost;
+				BestPoseIdx = PoseIdx;
+			}
+		}
+	}
+
+	Result.PoseCost = BestPoseCost;
+	Result.PoseIdx = BestPoseIdx;
+	Result.SearchIndexAsset = SearchIndex.FindAssetForPose(BestPoseIdx);
+	Result.AssetTime = SearchIndex.GetAssetTime(BestPoseIdx, Result.SearchIndexAsset);
+	Result.Database = nullptr;
+	Result.Sequence = nullptr;
+
+	SearchContext.DebugDrawParams.PoseVector = NormalizedQueryValues;
+	SearchContext.DebugDrawParams.PoseIdx = Result.PoseIdx;
+	Draw(SearchContext.DebugDrawParams);
+
+	return Result;
+}
+
+FPoseSearchCost UPoseSearchSequenceMetaData::ComparePoses(
+	int32 PoseIdx, 
+	const TArrayView<const float>& QueryValues) const
+{
+	using namespace UE::PoseSearch;
+
+	FPoseSearchCost Result;
+
+	TArrayView<const float> PoseValues = SearchIndex.GetPoseValues(PoseIdx);
+	if (!ensure(PoseValues.Num() == QueryValues.Num()))
+	{
+		return Result;
+	}
+
+	Result.SetDissimilarity(CompareFeatureVectors(PoseValues.Num(), PoseValues.GetData(), QueryValues.GetData()));
+
+	const FPoseSearchPoseMetadata& PoseMetadata = SearchIndex.PoseMetadata[PoseIdx];
+	Result.SetCostAddend(PoseMetadata.CostAddend);
+
+	return Result;
+}
+
+
 //////////////////////////////////////////////////////////////////////////
 // FPoseSearchDatabaseSequence
@@ -1850,6 +1958,502 @@
 
 #endif // WITH_EDITOR
 
+FPoseSearchCost UPoseSearchDatabase::ComparePoses(
+	FPoseSearchContext& SearchContext,
+	FPoseSearchWeightsContext& WeightsContext,
+	int32 PoseIdx,
+	int32 GroupIdx,
+	const TArrayView<const float>& QueryValues) const
+{
+	using namespace UE::PoseSearch;
+
+	FPoseSearchCost Result;
+
+	const FPoseSearchIndex* SearchIndex = GetSearchIndex();
+	if (!ensure(SearchIndex))
+	{
+		return Result;
+	}
+
+	TArrayView<const float> PoseValues = SearchIndex->GetPoseValues(PoseIdx);
+	if (!ensure(PoseValues.Num() == QueryValues.Num()))
+	{
+		return Result;
+	}
+
+	if (GroupIdx == INDEX_NONE)
+	{
+		const FPoseSearchIndexAsset* SearchIndexAsset = SearchIndex->FindAssetForPose(PoseIdx);
+		if (!ensure(SearchIndexAsset))
+		{
+			return Result;
+		}
+
+		GroupIdx = SearchIndexAsset->SourceGroupIdx;
+	}
+
+	const FPoseSearchWeights* WeightsSet =
+		WeightsContext.GetGroupWeights(GroupIdx);
+	Result.SetDissimilarity(CompareFeatureVectors(
+		PoseValues.Num(),
+		PoseValues.GetData(),
+		QueryValues.GetData(),
+		WeightsSet->Weights.GetData()));
+
+	float NotifyAddend = 0.0f;
+	float MirrorMismatchAddend = 0.0f;
+	ComputePoseCostAddends(PoseIdx, SearchContext, NotifyAddend, MirrorMismatchAddend);
+	Result.SetCostAddend(NotifyAddend + MirrorMismatchAddend);
+
+	return Result;
+}
+
+FPoseSearchCost UPoseSearchDatabase::ComparePoses(
+	FPoseSearchContext& SearchContext,
+	FPoseSearchWeightsContext& WeightsContext,
+	int32 PoseIdx,
+	const TArrayView<const float>& QueryValues,
+	UE::PoseSearch::FPoseCostDetails& OutPoseCostDetails) const
+{
+	using namespace Eigen;
+	using namespace UE::PoseSearch;
+
+	FPoseSearchCost Result;
+
+	TArrayView<const float> PoseValues = GetSearchIndex()->GetPoseValues(PoseIdx);
+	const int32 Dims = PoseValues.Num();
+	if (!ensure(Dims == QueryValues.Num()))
+	{
+		return Result;
+	}
+
+	OutPoseCostDetails.CostVector.SetNum(Dims);
+
+	// Setup Eigen views onto our vectors
+	auto OutCostVector = Map<ArrayXf>(OutPoseCostDetails.CostVector.GetData(), Dims);
+	auto PoseVector = Map<const ArrayXf>(PoseValues.GetData(), Dims);
+	auto QueryVector = Map<const ArrayXf>(QueryValues.GetData(), Dims);
+
+	// Compute weighted squared difference vector
+	const FPoseSearchIndexAsset* SearchIndexAsset = GetSearchIndex()->FindAssetForPose(PoseIdx);
+	const FPoseSearchWeights* WeightsSet =
+		WeightsContext.GetGroupWeights(SearchIndexAsset->SourceGroupIdx);
+	check(WeightsSet);
+	check(WeightsSet->Weights.Num() == Dims);
+	auto WeightsVector = Map<const ArrayXf>(WeightsSet->Weights.GetData(), Dims);
+
+	OutCostVector = WeightsVector * (PoseVector - QueryVector).square();
+	Result.SetDissimilarity(OutCostVector.sum());
+
+	// Output result
+	float NotifyAddend = 0.0f;
+	float MirrorMismatchAddend = 0.0f;
+	ComputePoseCostAddends(PoseIdx, SearchContext, NotifyAddend, MirrorMismatchAddend);
+	Result.SetCostAddend(NotifyAddend + MirrorMismatchAddend);
+
+	// Output cost details
+	OutPoseCostDetails.NotifyCostAddend = NotifyAddend;
+	OutPoseCostDetails.MirrorMismatchCostAddend = MirrorMismatchAddend;
+	OutPoseCostDetails.PoseCost = Result;
+	CalcChannelCosts(
+		GetSearchIndex()->Schema,
+		OutPoseCostDetails.CostVector,
+		OutPoseCostDetails.ChannelCosts);
+
+#if DO_GUARD_SLOW
+	{
+		// Verify details pose comparator agrees with runtime pose comparator
+		FPoseSearchCost RuntimeComparatorCost = ComparePoses(
+			SearchContext, 
+			WeightsContext, 
+			PoseIdx, 
+			SearchIndexAsset->SourceGroupIdx,
+			QueryValues);
+		checkSlow(FMath::IsNearlyEqual(Result.GetTotalCost(), RuntimeComparatorCost.GetTotalCost(), 1e-3f));
+
+		// Verify channel cost decomposition agrees with runtime pose comparator
+		auto OutChannelCosts = Map<const ArrayXf>(
+			OutPoseCostDetails.ChannelCosts.GetData(), 
+			OutPoseCostDetails.ChannelCosts.Num());
+		checkSlow(FMath::IsNearlyEqual(OutChannelCosts.sum(), RuntimeComparatorCost.GetDissimilarity(), 1e-3f));
+	}
+#endif
+
+	return Result;
+}
+
+void UPoseSearchDatabase::ComputePoseCostAddends(
+	int32 PoseIdx,
+	FPoseSearchContext& SearchContext,
+	float& OutNotifyAddend,
+	float& OutMirrorMismatchAddend) const
+{
+	OutNotifyAddend = 0.0f;
+	OutMirrorMismatchAddend = 0.0f;
+
+	if (SearchContext.QueryMirrorRequest != EPoseSearchBooleanRequest::Indifferent)
+	{
+		const FPoseSearchIndexAsset* IndexAsset = GetSearchIndex()->FindAssetForPose(PoseIdx);
+		const bool bMirroringMismatch =
+			(IndexAsset->bMirrored && SearchContext.QueryMirrorRequest == EPoseSearchBooleanRequest::FalseValue) ||
+			(!IndexAsset->bMirrored && SearchContext.QueryMirrorRequest == EPoseSearchBooleanRequest::TrueValue);
+		if (bMirroringMismatch)
+		{
+			OutMirrorMismatchAddend = MirroringMismatchCost;
+		}
+	}
+
+	const FPoseSearchPoseMetadata& PoseMetadata = GetSearchIndex()->PoseMetadata[PoseIdx];
+	OutNotifyAddend = PoseMetadata.CostAddend;
+}
+
+UE::PoseSearch::FSearchResult UPoseSearchDatabase::Search(FPoseSearchContext& SearchContext) const
+{
+	using namespace UE::PoseSearch;
+
+	FSearchResult Result;
+
+#if WITH_EDITOR
+	if (IsDerivedDataBuildPending())
+	{
+		return Result;
+	}
+#endif
+
+	const FPoseSearchIndex* SearchIndex = GetSearchIndex();
+	if (!SearchIndex)
+	{
+		return Result;
+	}
+
+	if (!ensure(SearchIndex->IsValid() && !SearchIndex->IsEmpty()))
+	{
+		return Result;
+	}
+
+	if (PoseSearchMode != EPoseSearchMode::BruteForce)
+	{
+		Result = SearchPCAKDTree(SearchContext);
+	}
+
+	if (PoseSearchMode == EPoseSearchMode::BruteForce || PoseSearchMode == EPoseSearchMode::PCAKDTree_Compare)
+	{
+		Result = SearchBruteForce(SearchContext);
+	}
+
+	SearchContext.DebugDrawParams.PoseVector = Result.ComposedQuery.GetNormalizedValues();
+	SearchContext.DebugDrawParams.PoseIdx = Result.PoseIdx;
+	Draw(SearchContext.DebugDrawParams);
+
+	return Result;
+}
+
+UE::PoseSearch::FSearchResult UPoseSearchDatabase::SearchPCAKDTree(FPoseSearchContext& SearchContext) const
+{
+	QUICK_SCOPE_CYCLE_COUNTER(STAT_PoseSearch_PCA_KNN);
+	SCOPE_CYCLE_COUNTER(STAT_PoseSearchPCAKNN);
+
+	using namespace UE::PoseSearch;
+
+	FSearchResult Result;
+
+	const int32 NumDimensions = Schema->Layout.NumFloats;
+	const FPoseSearchIndex* SearchIndex = GetSearchIndex();
+	check(SearchIndex);
+
+	const uint32 ClampedNumberOfPrincipalComponents = GetNumberOfPrincipalComponents();
+	const uint32 ClampedKDTreeQueryNumNeighbors = FMath::Clamp<uint32>(KDTreeQueryNumNeighbors, 1, SearchIndex->NumPoses);
+
+	//stack allocated temporaries
+	TArrayView<size_t> ResultIndexes((size_t*)FMemory_Alloca((ClampedKDTreeQueryNumNeighbors + 1) * sizeof(size_t)), ClampedKDTreeQueryNumNeighbors + 1);
+	TArrayView<float> ResultDistanceSqr((float*)FMemory_Alloca((ClampedKDTreeQueryNumNeighbors + 1) * sizeof(float)), ClampedKDTreeQueryNumNeighbors + 1);
+	RowMajorVectorMap WeightedQueryValues((float*)FMemory_Alloca(NumDimensions * sizeof(float)), 1, NumDimensions);
+	RowMajorVectorMap CenteredQueryValues((float*)FMemory_Alloca(NumDimensions * sizeof(float)), 1, NumDimensions);
+	RowMajorVectorMap ProjectedQueryValues((float*)FMemory_Alloca(ClampedNumberOfPrincipalComponents * sizeof(float)), 1, ClampedNumberOfPrincipalComponents);
+
+	// KDTree in PCA space search
+	if (PoseSearchMode == EPoseSearchMode::PCAKDTree_Validate)
+	{
+		for (const FGroupSearchIndex& GroupSearchIndex : SearchIndex->Groups)
+		{
+			const RowMajorVectorMapConst MapWeights(GroupSearchIndex.Weights.GetData(), 1, NumDimensions);
+
+			// testing the KDTree is returning the proper searches for all the original points transformed in pca space
+			for (int32 PoseIdx = GroupSearchIndex.StartPoseIndex; PoseIdx < GroupSearchIndex.EndPoseIndex; ++PoseIdx)
+			{
+				FKDTree::KNNResultSet ResultSet(ClampedKDTreeQueryNumNeighbors, ResultIndexes, ResultDistanceSqr);
+				TArrayView<const float> PoseValues = SearchIndex->GetPoseValues(PoseIdx);
+
+				const RowMajorVectorMapConst Mean(GroupSearchIndex.Mean.GetData(), 1, NumDimensions);
+				const ColMajorMatrixMapConst PCAProjectionMatrix(GroupSearchIndex.PCAProjectionMatrix.GetData(), NumDimensions, ClampedNumberOfPrincipalComponents);
+
+				const RowMajorVectorMapConst QueryValues(PoseValues.GetData(), 1, NumDimensions);
+				WeightedQueryValues = QueryValues.array() * MapWeights.array();
+				CenteredQueryValues.noalias() = WeightedQueryValues - Mean;
+				ProjectedQueryValues.noalias() = CenteredQueryValues * PCAProjectionMatrix;
+
+				GroupSearchIndex.KDTree.FindNeighbors(ResultSet, ProjectedQueryValues.data());
+
+				size_t ResultIndex = 0;
+				for (; ResultIndex < ResultSet.Num(); ++ResultIndex)
+				{
+					if ((PoseIdx - GroupSearchIndex.StartPoseIndex) == ResultIndexes[ResultIndex])
+					{
+						check(ResultDistanceSqr[ResultIndex] < UE_KINDA_SMALL_NUMBER);
+						break;
+					}
+				}
+				check(ResultIndex < ResultSet.Num());
+			}
+		}
+	}
+
+	// @todo: implement support for DatabaseTagQuery
+	FPoseSearchCost BestPoseCost;
+	int32 BestPoseIdx = INDEX_NONE;
+
+	BuildQuery(SearchContext, Result.ComposedQuery);
+
+	FPoseSearchWeightsContext WeightsContext;
+	WeightsContext.Update(this);
+
+	TArrayView<const float> NormalizedQueryValues = Result.ComposedQuery.GetNormalizedValues();
+
+	if (SearchContext.CurrentResult.IsValid() && SearchContext.CurrentResult.Database == this)
+	{
+		Result.ContinuityPoseCost = ComparePoses(
+			SearchContext,
+			WeightsContext,
+			SearchContext.CurrentResult.PoseIdx,
+			SearchIndex->FindAssetForPose(SearchContext.CurrentResult.PoseIdx)->SourceGroupIdx,
+			NormalizedQueryValues);
+	}
+
+	for (const FGroupSearchIndex& GroupSearchIndex : SearchIndex->Groups)
+	{
+		const RowMajorVectorMapConst MapWeights(GroupSearchIndex.Weights.GetData(), 1, NumDimensions);
+		FKDTree::KNNResultSet ResultSet(ClampedKDTreeQueryNumNeighbors, ResultIndexes, ResultDistanceSqr);
+
+		check(NormalizedQueryValues.Num() == NumDimensions);
+
+		const RowMajorVectorMapConst Mean(GroupSearchIndex.Mean.GetData(), 1, NumDimensions);
+		const ColMajorMatrixMapConst PCAProjectionMatrix(GroupSearchIndex.PCAProjectionMatrix.GetData(), NumDimensions, ClampedNumberOfPrincipalComponents);
+
+		// transforming query values into PCA space to query the KDTree
+		const RowMajorVectorMapConst QueryValues(NormalizedQueryValues.GetData(), 1, NumDimensions);
+		WeightedQueryValues = QueryValues.array() * MapWeights.array();
+		CenteredQueryValues.noalias() = WeightedQueryValues - Mean;
+		ProjectedQueryValues.noalias() = CenteredQueryValues * PCAProjectionMatrix;
+
+		GroupSearchIndex.KDTree.FindNeighbors(ResultSet, ProjectedQueryValues.data());
+
+		for (size_t ResultIndex = 0; ResultIndex < ResultSet.Num(); ++ResultIndex)
+		{
+			const int32 PoseIdx = ResultIndexes[ResultIndex] + GroupSearchIndex.StartPoseIndex;
+
+			const FPoseSearchPoseMetadata& Metadata = SearchIndex->PoseMetadata[PoseIdx];
+
+			if (EnumHasAnyFlags(Metadata.Flags, EPoseSearchPoseFlags::BlockTransition))
+			{
+				continue;
+			}
+
+			FPoseSearchCost PoseCost = ComparePoses(
+				SearchContext, 
+				WeightsContext, 
+				PoseIdx, 
+				GroupSearchIndex.GroupIndex,
+				NormalizedQueryValues);
+
+			if (PoseCost < BestPoseCost)
+			{
+				BestPoseCost = PoseCost;
+				BestPoseIdx = PoseIdx;
+			}
+		}
+	}
+
+	Result.PoseCost = BestPoseCost;
+	Result.PoseIdx = BestPoseIdx;
+	Result.SearchIndexAsset = SearchIndex->FindAssetForPose(BestPoseIdx);
+	Result.AssetTime = SearchIndex->GetAssetTime(BestPoseIdx, Result.SearchIndexAsset);
+	Result.Database = this;
+	Result.Sequence = nullptr;
+#if WITH_EDITOR
+	Result.SearchIndexHash = GetSearchIndexHash();
+#endif // WITH_EDITOR
+
+	return Result;
+}
+
+UE::PoseSearch::FSearchResult UPoseSearchDatabase::SearchBruteForce(FPoseSearchContext& SearchContext) const
+{
+	QUICK_SCOPE_CYCLE_COUNTER(STAT_PoseSearch_Brute_Force);
+	SCOPE_CYCLE_COUNTER(STAT_PoseSearchBruteForce);
+	
+	using namespace UE::PoseSearch;
+	
+	FSearchResult Result;
+
+	const FPoseSearchIndex* SearchIndex = GetSearchIndex();
+	check(SearchIndex);
+
+	BuildQuery(SearchContext, Result.ComposedQuery);
+	TArrayView<const float> NormalizedQueryValues = Result.ComposedQuery.GetNormalizedValues();
+
+	FPoseSearchWeightsContext WeightsContext;
+	WeightsContext.Update(this);
+	
+	if (SearchContext.CurrentResult.IsValid() && SearchContext.CurrentResult.Database == this)
+	{
+		Result.ContinuityPoseCost = ComparePoses(
+			SearchContext,
+			WeightsContext,
+			SearchContext.CurrentResult.PoseIdx,
+			SearchIndex->FindAssetForPose(SearchContext.CurrentResult.PoseIdx)->SourceGroupIdx,
+			NormalizedQueryValues);
+	}
+
+	FPoseSearchCost BestPoseCost;
+	int32 BestPoseIdx = INDEX_NONE;
+	for (const FPoseSearchIndexAsset& Asset : SearchIndex->Assets)
+	{
+		if (SearchContext.DatabaseTagQuery)
+		{
+			if (!SearchContext.DatabaseTagQuery->Matches(*GetSourceAssetGroupTags(&Asset)))
+			{
+				continue;
+			}
+		}
+
+		const int32 EndIndex = Asset.FirstPoseIdx + Asset.NumPoses;
+		for (int32 PoseIdx = Asset.FirstPoseIdx; PoseIdx < EndIndex; ++PoseIdx)
+		{
+			const FPoseSearchPoseMetadata& Metadata = SearchIndex->PoseMetadata[PoseIdx];
+
+			if (EnumHasAnyFlags(Metadata.Flags, EPoseSearchPoseFlags::BlockTransition))
+			{
+				continue;
+			}
+
+			FPoseSearchCost PoseCost = ComparePoses(
+				SearchContext,
+				WeightsContext,
+				PoseIdx, 
+				Asset.SourceGroupIdx,
+				NormalizedQueryValues);
+
+			if (PoseCost < BestPoseCost)
+			{
+				BestPoseCost = PoseCost;
+				BestPoseIdx = PoseIdx;
+			}
+		}
+	}
+
+	Result.PoseCost = BestPoseCost;
+	Result.PoseIdx = BestPoseIdx;
+	Result.SearchIndexAsset = SearchIndex->FindAssetForPose(BestPoseIdx);
+	Result.AssetTime = SearchIndex->GetAssetTime(BestPoseIdx, Result.SearchIndexAsset);
+	Result.Database = this;
+	Result.Sequence = nullptr;
+#if WITH_EDITOR
+	Result.SearchIndexHash = GetSearchIndexHash();
+#endif // WITH_EDITOR
+
+	return Result;
+}
+
+void UPoseSearchDatabase::BuildQuery(
+	FPoseSearchContext& SearchContext, 
+	FPoseSearchFeatureVectorBuilder& OutQuery) const
+{
+	Schema->BuildQuery(SearchContext, OutQuery);
+	OutQuery.Normalize(*GetSearchIndex());
+}
+
+UE::PoseSearch::FSearchResult UPoseSearchDatabaseSet::Search(FPoseSearchContext& SearchContext) const
+{
+	using namespace UE::PoseSearch;
+
+	FSearchResult Result;
+	FPoseSearchCost ContinuityCost;
+
+	TArray<const FPoseSearchDatabaseSetEntry*, TInlineAllocator<8>> FallBackEntries;
+
+	auto ProcessActiveEntry = 
+		[&SearchContext, &Result, &ContinuityCost](const FPoseSearchDatabaseSetEntry& Entry) 
+		-> EPoseSearchPostSearchStatus
+	{
+		EPoseSearchPostSearchStatus PostSearchStatus = EPoseSearchPostSearchStatus::Continue;
+		FSearchResult EntryResult = Entry.Searchable->Search(SearchContext);
+		if (EntryResult.IsValid())
+		{
+			if (IsValid(Entry.PostProcessor))
+			{
+				PostSearchStatus = Entry.PostProcessor->PostProcess(EntryResult.PoseCost);
+				if (EntryResult.ContinuityPoseCost.IsValid())
+				{
+					Entry.PostProcessor->PostProcess(EntryResult.ContinuityPoseCost);
+				}
+			}
+
+			if (!Result.IsValid() || EntryResult.PoseCost.GetTotalCost() < Result.PoseCost.GetTotalCost())
+			{
+				Result = EntryResult;
+			}
+
+			if (EntryResult.ContinuityPoseCost.IsValid())
+			{
+				if (!ContinuityCost.IsValid() ||
+					EntryResult.ContinuityPoseCost.GetTotalCost() < ContinuityCost.GetTotalCost())
+				{
+					ContinuityCost = EntryResult.ContinuityPoseCost;
+				}
+			}
+		}
+
+		return PostSearchStatus;
+	};
+
+	for (const FPoseSearchDatabaseSetEntry& Entry : AssetsToSearch)
+	{
+		if (!IsValid(Entry.Searchable))
+		{
+			UE_LOG(LogPoseSearch, Warning, TEXT("Invalid entry in Database Set %s"), *GetName());
+			continue;
+		}
+
+		const bool bSearchEntry =
+			!Entry.Tag.IsValid() ||
+			SearchContext.ActiveTagsContainer == nullptr ||
+			SearchContext.ActiveTagsContainer->IsEmpty() ||
+			SearchContext.ActiveTagsContainer->HasTag(Entry.Tag);
+
+		if (bSearchEntry)
+		{
+			EPoseSearchPostSearchStatus PostSearchStatus = ProcessActiveEntry(Entry);
+			if (PostSearchStatus == EPoseSearchPostSearchStatus::Stop)
+			{
+				break;
+			}
+		}
+	}
+
+	if (Result.IsValid())
+	{
+		Result.ContinuityPoseCost = ContinuityCost;
+	}
+	else
+	{
+		UE_LOG(
+			LogPoseSearch, 
+			Error, 
+			TEXT("Invalid result searching %s"), *GetName());
+	}
+
+	return Result;
+}
 
 //////////////////////////////////////////////////////////////////////////
 // FPoseSearchFeatureVectorBuilder
@@ -2079,6 +2683,7 @@
 	ForSearchIndex.Normalize(ValuesNormalized);
 }
 
+
 namespace UE { namespace PoseSearch {
 
 //////////////////////////////////////////////////////////////////////////
@@ -2443,52 +3048,89 @@
 	return nullptr;
 }
 
+//////////////////////////////////////////////////////////////////////////
+// FSearchResult
+// 
 
-//////////////////////////////////////////////////////////////////////////
-// FSearchContext
+void FSearchResult::Update(float NewAssetTime)
+{
+	if (!IsValid())
+	{
+		Reset();
+		return;
+	}
 
-void FSearchContext::SetSource(const UPoseSearchDatabase* InSourceDatabase)
-{
-	SearchIndex = nullptr;
-	DebugDrawParams.Database = nullptr;
-	DebugDrawParams.SequenceMetaData = nullptr;
+	if (SearchIndexAsset->Type == ESearchIndexAssetType::Sequence)
+	{
+		const FPoseSearchDatabaseSequence& DbSequence = Database->GetSequenceSourceAsset(SearchIndexAsset);
 
-	SourceDatabase = InSourceDatabase;
-	if (SourceDatabase)
+		if (SearchIndexAsset->SamplingInterval.Contains(NewAssetTime))
+		{
+			PoseIdx = Database->GetPoseIndexFromTime(NewAssetTime, SearchIndexAsset);
+			AssetTime = NewAssetTime;
+		}
+		else
+		{
+			Reset();
+		}
+	}
+	else if (SearchIndexAsset->Type == ESearchIndexAssetType::BlendSpace)
 	{
-		if (ensure(SourceDatabase->IsValidForSearch()))
+		const FPoseSearchDatabaseBlendSpace& DbBlendSpace = Database->GetBlendSpaceSourceAsset(SearchIndexAsset);
+
+		TArray<FBlendSampleData> BlendSamples;
+		int32 TriangulationIndex = 0;
+		DbBlendSpace.BlendSpace->GetSamplesFromBlendInput(SearchIndexAsset->BlendParameters, BlendSamples, TriangulationIndex, true);
+
+		float PlayLength = DbBlendSpace.BlendSpace->GetAnimationLengthFromSampleData(BlendSamples);
+
+		// Asset player time for blendspaces is normalized [0, 1] so we need to convert 
+		// to a real time before we advance it
+		float RealTime = NewAssetTime * PlayLength;
+
+		if (SearchIndexAsset->SamplingInterval.Contains(RealTime))
 		{
-			SearchIndex = SourceDatabase->GetSearchIndex();
-			DebugDrawParams.Database = SourceDatabase;
-			MirrorMismatchCost = SourceDatabase->MirroringMismatchCost;
+			PoseIdx = Database->GetPoseIndexFromTime(RealTime, SearchIndexAsset);
+			AssetTime = NewAssetTime;
 		}
+		else
+		{
+			Reset();
+		}
+	}
+	else
+	{
+		checkNoEntry();
 	}
 }
 
-void FSearchContext::SetSource(const UAnimSequenceBase* InSourceSequence)
+bool FSearchResult::IsValid() const
 {
-	SearchIndex = nullptr;
-	DebugDrawParams.Database = nullptr;
-	DebugDrawParams.SequenceMetaData = nullptr;
+	bool bIsValid =
+		PoseIdx != INDEX_NONE &&
+		Database.IsValid();
 
-	SourceSequence = InSourceSequence;
-	const UPoseSearchSequenceMetaData* MetaData =
-		SourceSequence->FindMetaDataByClass<UPoseSearchSequenceMetaData>();
-	if (MetaData && MetaData->IsValidForSearch())
-	{
-		SearchIndex = &MetaData->SearchIndex;
-		DebugDrawParams.SequenceMetaData = MetaData;
-	}
+#if WITH_EDITOR
+	bIsValid = bIsValid &&
+		!Database->IsDerivedDataBuildPending() &&
+		Database->GetSearchIndexHash() == SearchIndexHash;
+#endif // WITH_EDITOR
+
+	return bIsValid;
 }
 
-const FPoseSearchIndex* FSearchContext::GetSearchIndex() const
+void FSearchResult::Reset()
 {
-	return SearchIndex;
-}
+	PoseIdx = INDEX_NONE;
+	SearchIndexAsset = nullptr;
+	Database = nullptr;
+	Sequence = nullptr;
+	ComposedQuery.Reset();
+	AssetTime = 0.0f;
 
-float FSearchContext::GetMirrorMismatchCost() const
-{
-	return MirrorMismatchCost;
+#if WITH_EDITOR
+	SearchIndexHash = FIoHash::Zero;
+#endif // WITH_EDITOR
 }
 
 
@@ -2762,249 +3404,3 @@
 	const float ClampedTime = FMath::Clamp(Time, 0.0f, PlayLength);
 	const float ExtrapolationTime = Time - ClampedTime;
 
-	FTransform RootTransform = FTransform::Identity;
-
-	// If Time is less than zero, ExtrapolationTime will be negative. In this case, we extrapolate the beginning of the 
-	// animation to estimate where the root would be at Time
-	if (ExtrapolationTime < -SMALL_NUMBER)
-	{
-		FTransform SampleToExtrapolate = Input.Sequence->ExtractRootMotionFromRange(0.0f, ExtrapolationSampleTime);
-
-		const FTransform ExtrapolatedRootMotion = ExtrapolateRootMotion(
-			SampleToExtrapolate,
-			0.0f, ExtrapolationSampleTime, 
-			ExtrapolationTime,
-			Input.ExtrapolationParameters);
-		RootTransform = ExtrapolatedRootMotion;
-	}
-	else
-	{
-		RootTransform = Input.Sequence->ExtractRootMotionFromRange(0.0f, ClampedTime);
-
-		// If Time is greater than PlayLength, ExtrapolationTIme will be a positive number. In this case, we extrapolate
-		// the end of the animation to estimate where the root would be at Time
-		if (ExtrapolationTime > SMALL_NUMBER)
-		{
-			FTransform SampleToExtrapolate = Input.Sequence->ExtractRootMotionFromRange(PlayLength - ExtrapolationSampleTime, PlayLength);
-
-			const FTransform ExtrapolatedRootMotion = ExtrapolateRootMotion(
-				SampleToExtrapolate,
-				PlayLength - ExtrapolationSampleTime, PlayLength,
-				ExtrapolationTime,
-				Input.ExtrapolationParameters);
-			RootTransform = ExtrapolatedRootMotion * RootTransform;
-		}
-	}
-
-	return RootTransform;
-}
-
-float FSequenceSampler::ExtractRootDistance(float Time) const
-{
-	return ExtractAccumulatedRootDistance(
-		Input.RootDistanceSamplingRate,
-		AccumulatedRootDistance,
-		Input.Sequence->GetPlayLength(),
-		Time,
-		Input.ExtrapolationParameters);
-}
-
-void FSequenceSampler::ExtractPoseSearchNotifyStates(
-	float Time, 
-	TArray<UAnimNotifyState_PoseSearchBase*>& NotifyStates) const
-{
-	// getting pose search notifies in an interval of size ExtractionInterval, centered on Time
-	constexpr float ExtractionInterval = 1.0f / 120.0f;
-	FAnimNotifyContext NotifyContext;
-	Input.Sequence->GetAnimNotifies(Time - (ExtractionInterval * 0.5f), ExtractionInterval, NotifyContext);
-
-	// check which notifies actually overlap Time and are of the right base type
-	for (const FAnimNotifyEventReference& EventReference : NotifyContext.ActiveNotifies)
-	{
-		const FAnimNotifyEvent* NotifyEvent = EventReference.GetNotify();
-		if (!NotifyEvent)
-		{
-			continue;
-		}
-
-		if (NotifyEvent->GetTriggerTime() > Time ||
-			NotifyEvent->GetEndTriggerTime() < Time)
-		{
-			continue;
-		}
-
-		UAnimNotifyState_PoseSearchBase* PoseSearchAnimNotify = 
-			Cast<UAnimNotifyState_PoseSearchBase>(NotifyEvent->NotifyStateClass);
-		if (PoseSearchAnimNotify)
-		{
-			NotifyStates.Add(PoseSearchAnimNotify);
-		}
-	}
-}
-
-void FSequenceSampler::ProcessRootDistance()
-{
-	// Note the distance sampling interval is independent of the schema's sampling interval
-	const float DistanceSamplingInterval = 1.0f / Input.RootDistanceSamplingRate;
-
-	const FTransform InitialRootTransform = Input.Sequence->ExtractRootTrackTransform(0.0f, nullptr);
-
-	uint32 NumDistanceSamples = FMath::CeilToInt(Input.Sequence->GetPlayLength() * Input.RootDistanceSamplingRate) + 1;
-	AccumulatedRootDistance.Reserve(NumDistanceSamples);
-
-	// Build a distance lookup table by sampling root motion at a fixed rate and accumulating
-	// absolute translation deltas. During indexing we'll bsearch this table and interpolate
-	// between samples in order to convert distance offsets to time offsets.
-	// See also FAssetIndexer::AddTrajectoryDistanceFeatures().
-
-	double TotalAccumulatedRootDistance = 0.0;
-	FTransform LastRootTransform = InitialRootTransform;
-	float SampleTime = 0.0f;
-	for (int32 SampleIdx = 0; SampleIdx != NumDistanceSamples; ++SampleIdx)
-	{
-		SampleTime = FMath::Min(SampleIdx * DistanceSamplingInterval, Input.Sequence->GetPlayLength());
-
-		FTransform RootTransform = Input.Sequence->ExtractRootTrackTransform(SampleTime, nullptr);
-		FTransform LocalRootMotion = RootTransform.GetRelativeTransform(LastRootTransform);
-		LastRootTransform = RootTransform;
-
-		TotalAccumulatedRootDistance += LocalRootMotion.GetTranslation().Size();
-		AccumulatedRootDistance.Add((float)TotalAccumulatedRootDistance);
-	}
-
-	// Verify we sampled the final frame of the clip
-	check(SampleTime == Input.Sequence->GetPlayLength());
-
-	// Also emit root motion summary info to help with sample wrapping in 
-	// FAssetIndexer::GetSampleTimeFromDistance() and FAssetIndexer::GetSampleInfo()
-	TotalRootTransform = LastRootTransform.GetRelativeTransform(InitialRootTransform);
-	TotalRootDistance = AccumulatedRootDistance.Last();
-}
-
-//////////////////////////////////////////////////////////////////////////
-// FBlendSpaceSampler
-
-struct FBlendSpaceSampler : public IAssetSampler
-{
-public:
-	struct FInput
-	{
-		const FAssetSamplingContext* SamplingContext = nullptr;
-		const UBlendSpace* BlendSpace = nullptr;
-		bool bLoopable = false;
-		int32 RootDistanceSamplingRate = 60;
-		int32 RootTransformSamplingRate = 60;
-		FPoseSearchExtrapolationParameters ExtrapolationParameters;
-		FVector BlendParameters;
-	} Input;
-
-	void Init(const FInput& Input);
-
-	void Process();
-
-	float GetPlayLength() const override { return PlayLength; };
-	bool IsLoopable() const override { return Input.bLoopable; };
-
-	float GetTimeFromRootDistance(float Distance) const override;
-
-	float GetTotalRootDistance() const override { return TotalRootDistance; };
-	FTransform GetTotalRootTransform() const override { return TotalRootTransform; }
-
-	virtual void ExtractPose(const FAnimExtractContext& ExtractionCtx, FAnimationPoseData& OutAnimPoseData) const override;
-	virtual float ExtractRootDistance(float Time) const override;
-	virtual FTransform ExtractRootTransform(float Time) const override;
-	virtual void ExtractPoseSearchNotifyStates(float Time, TArray<class UAnimNotifyState_PoseSearchBase*>& NotifyStates) const override;
-
-private:
-	float PlayLength = 0.0f;
-	float TotalRootDistance = 0.0f;
-	FTransform TotalRootTransform = FTransform::Identity;
-	TArray<float> AccumulatedRootDistance;
-	TArray<FTransform> AccumulatedRootTransform;
-	
-	void ProcessPlayLength();
-	void ProcessRootDistance();
-	void ProcessRootTransform();
-
-	// Extracts the pre-computed blend space root transform. ProcessRootTransform must be run first.
-	FTransform ExtractBlendSpaceRootTrackTransform(float Time) const;
-	FTransform ExtractBlendSpaceRootMotion(float StartTime, float DeltaTime, bool bAllowLooping) const;
-	FTransform ExtractBlendSpaceRootMotionFromRange(float StartTrackPosition, float EndTrackPosition) const;
-};
-
-void FBlendSpaceSampler::Init(const FInput& InInput)
-{
-	check(InInput.BlendSpace);
-
-	Input = InInput;
-}
-
-void FBlendSpaceSampler::Process()
-{
-	FMemMark Mark(FMemStack::Get());
-
-	ProcessPlayLength();
-	ProcessRootTransform();
-	ProcessRootDistance();
-}
-
-float FBlendSpaceSampler::GetTimeFromRootDistance(float Distance) const
-{
-	int32 NextSampleIdx = 1;
-	int32 PrevSampleIdx = 0;
-	if (Distance > 0.0f)
-	{
-		// Search for the distance value. Because the values will be extrapolated if necessary
-		// LowerBound might go past the end of the array, in which case the last valid index is used
-		int32 ClipDistanceLowerBoundIndex = Algo::LowerBound(AccumulatedRootDistance, Distance);
-		NextSampleIdx = FMath::Min(
-			ClipDistanceLowerBoundIndex,
-			AccumulatedRootDistance.Num() - 1);
-
-		// Compute distance interpolation amount
-		PrevSampleIdx = FMath::Max(0, NextSampleIdx - 1);
-	}
-
-	float NextDistance = AccumulatedRootDistance[NextSampleIdx];
-	float PrevDistance = AccumulatedRootDistance[PrevSampleIdx];
-	float DistanceSampleAlpha = FMath::GetRangePct(PrevDistance, NextDistance, Distance);
-
-	// Convert to time
-	float ClipTime = (float(NextSampleIdx) - (1.0f - DistanceSampleAlpha)) / Input.RootDistanceSamplingRate;
-	return ClipTime;
-}
-
-void FBlendSpaceSampler::ExtractPose(const FAnimExtractContext& ExtractionCtx, FAnimationPoseData& OutAnimPoseData) const
-{
-	TArray<FBlendSampleData> BlendSamples;
-	int32 TriangulationIndex = 0;
-	Input.BlendSpace->GetSamplesFromBlendInput(Input.BlendParameters, BlendSamples, TriangulationIndex, true);
-
-	for (int32 BlendSampleIdex = 0; BlendSampleIdex < BlendSamples.Num(); BlendSampleIdex++)
-	{
-		float Scale = BlendSamples[BlendSampleIdex].Animation->GetPlayLength() / PlayLength;
-
-		FDeltaTimeRecord BlendSampleDeltaTimeRecord;
-		BlendSampleDeltaTimeRecord.Set(ExtractionCtx.DeltaTimeRecord.GetPrevious() * Scale, ExtractionCtx.DeltaTimeRecord.Delta * Scale);
-
-		BlendSamples[BlendSampleIdex].DeltaTimeRecord = BlendSampleDeltaTimeRecord;
-		BlendSamples[BlendSampleIdex].PreviousTime = ExtractionCtx.DeltaTimeRecord.GetPrevious() * Scale;
-		BlendSamples[BlendSampleIdex].Time = ExtractionCtx.CurrentTime * Scale;
-	}
-
-	Input.BlendSpace->GetAnimationPose(BlendSamples, ExtractionCtx, OutAnimPoseData);
-}
-
-FTransform FBlendSpaceSampler::ExtractRootTransform(float Time) const
-{
-	if (Input.bLoopable)
-	{
-		FTransform LoopableRootTransform = ExtractBlendSpaceRootMotion(0.0f, Time, true);
-		return LoopableRootTransform;
-	}
-
-	const float ExtrapolationSampleTime = Input.ExtrapolationParameters.SampleTime;
-
-	const float ClampedTime = FMath::Clamp(Time, 0.0f, PlayLength);
-	const float ExtrapolationTime = Time - ClampedTime;
-
@@ -3014,24 +3410,24 @@
 	// animation to estimate where the root would be at Time
 	if (ExtrapolationTime < -SMALL_NUMBER)
 	{
-		FTransform SampleToExtrapolate = ExtractBlendSpaceRootMotionFromRange(0.0f, ExtrapolationSampleTime);
+		FTransform SampleToExtrapolate = Input.Sequence->ExtractRootMotionFromRange(0.0f, ExtrapolationSampleTime);
 
 		const FTransform ExtrapolatedRootMotion = ExtrapolateRootMotion(
 			SampleToExtrapolate,
-			0.0f, ExtrapolationSampleTime,
+			0.0f, ExtrapolationSampleTime, 
 			ExtrapolationTime,
 			Input.ExtrapolationParameters);
 		RootTransform = ExtrapolatedRootMotion;
 	}
 	else
 	{
-		RootTransform = ExtractBlendSpaceRootMotionFromRange(0.0f, ClampedTime);
+		RootTransform = Input.Sequence->ExtractRootMotionFromRange(0.0f, ClampedTime);
 
 		// If Time is greater than PlayLength, ExtrapolationTIme will be a positive number. In this case, we extrapolate
 		// the end of the animation to estimate where the root would be at Time
 		if (ExtrapolationTime > SMALL_NUMBER)
 		{
-			FTransform SampleToExtrapolate = ExtractBlendSpaceRootMotionFromRange(PlayLength - ExtrapolationSampleTime, PlayLength);
+			FTransform SampleToExtrapolate = Input.Sequence->ExtractRootMotionFromRange(PlayLength - ExtrapolationSampleTime, PlayLength);
 
 			const FTransform ExtrapolatedRootMotion = ExtrapolateRootMotion(
 				SampleToExtrapolate,
@@ -3045,6 +3441,252 @@
 	return RootTransform;
 }
 
+float FSequenceSampler::ExtractRootDistance(float Time) const
+{
+	return ExtractAccumulatedRootDistance(
+		Input.RootDistanceSamplingRate,
+		AccumulatedRootDistance,
+		Input.Sequence->GetPlayLength(),
+		Time,
+		Input.ExtrapolationParameters);
+}
+
+void FSequenceSampler::ExtractPoseSearchNotifyStates(
+	float Time, 
+	TArray<UAnimNotifyState_PoseSearchBase*>& NotifyStates) const
+{
+	// getting pose search notifies in an interval of size ExtractionInterval, centered on Time
+	constexpr float ExtractionInterval = 1.0f / 120.0f;
+	FAnimNotifyContext NotifyContext;
+	Input.Sequence->GetAnimNotifies(Time - (ExtractionInterval * 0.5f), ExtractionInterval, NotifyContext);
+
+	// check which notifies actually overlap Time and are of the right base type
+	for (const FAnimNotifyEventReference& EventReference : NotifyContext.ActiveNotifies)
+	{
+		const FAnimNotifyEvent* NotifyEvent = EventReference.GetNotify();
+		if (!NotifyEvent)
+		{
+			continue;
+		}
+
+		if (NotifyEvent->GetTriggerTime() > Time ||
+			NotifyEvent->GetEndTriggerTime() < Time)
+		{
+			continue;
+		}
+
+		UAnimNotifyState_PoseSearchBase* PoseSearchAnimNotify = 
+			Cast<UAnimNotifyState_PoseSearchBase>(NotifyEvent->NotifyStateClass);
+		if (PoseSearchAnimNotify)
+		{
+			NotifyStates.Add(PoseSearchAnimNotify);
+		}
+	}
+}
+
+void FSequenceSampler::ProcessRootDistance()
+{
+	// Note the distance sampling interval is independent of the schema's sampling interval
+	const float DistanceSamplingInterval = 1.0f / Input.RootDistanceSamplingRate;
+
+	const FTransform InitialRootTransform = Input.Sequence->ExtractRootTrackTransform(0.0f, nullptr);
+
+	uint32 NumDistanceSamples = FMath::CeilToInt(Input.Sequence->GetPlayLength() * Input.RootDistanceSamplingRate) + 1;
+	AccumulatedRootDistance.Reserve(NumDistanceSamples);
+
+	// Build a distance lookup table by sampling root motion at a fixed rate and accumulating
+	// absolute translation deltas. During indexing we'll bsearch this table and interpolate
+	// between samples in order to convert distance offsets to time offsets.
+	// See also FAssetIndexer::AddTrajectoryDistanceFeatures().
+
+	double TotalAccumulatedRootDistance = 0.0;
+	FTransform LastRootTransform = InitialRootTransform;
+	float SampleTime = 0.0f;
+	for (int32 SampleIdx = 0; SampleIdx != NumDistanceSamples; ++SampleIdx)
+	{
+		SampleTime = FMath::Min(SampleIdx * DistanceSamplingInterval, Input.Sequence->GetPlayLength());
+
+		FTransform RootTransform = Input.Sequence->ExtractRootTrackTransform(SampleTime, nullptr);
+		FTransform LocalRootMotion = RootTransform.GetRelativeTransform(LastRootTransform);
+		LastRootTransform = RootTransform;
+
+		TotalAccumulatedRootDistance += LocalRootMotion.GetTranslation().Size();
+		AccumulatedRootDistance.Add((float)TotalAccumulatedRootDistance);
+	}
+
+	// Verify we sampled the final frame of the clip
+	check(SampleTime == Input.Sequence->GetPlayLength());
+
+	// Also emit root motion summary info to help with sample wrapping in 
+	// FAssetIndexer::GetSampleTimeFromDistance() and FAssetIndexer::GetSampleInfo()
+	TotalRootTransform = LastRootTransform.GetRelativeTransform(InitialRootTransform);
+	TotalRootDistance = AccumulatedRootDistance.Last();
+}
+
+//////////////////////////////////////////////////////////////////////////
+// FBlendSpaceSampler
+
+struct FBlendSpaceSampler : public IAssetSampler
+{
+public:
+	struct FInput
+	{
+		const FAssetSamplingContext* SamplingContext = nullptr;
+		const UBlendSpace* BlendSpace = nullptr;
+		bool bLoopable = false;
+		int32 RootDistanceSamplingRate = 60;
+		int32 RootTransformSamplingRate = 60;
+		FPoseSearchExtrapolationParameters ExtrapolationParameters;
+		FVector BlendParameters;
+	} Input;
+
+	void Init(const FInput& Input);
+
+	void Process();
+
+	float GetPlayLength() const override { return PlayLength; };
+	bool IsLoopable() const override { return Input.bLoopable; };
+
+	float GetTimeFromRootDistance(float Distance) const override;
+
+	float GetTotalRootDistance() const override { return TotalRootDistance; };
+	FTransform GetTotalRootTransform() const override { return TotalRootTransform; }
+
+	virtual void ExtractPose(const FAnimExtractContext& ExtractionCtx, FAnimationPoseData& OutAnimPoseData) const override;
+	virtual float ExtractRootDistance(float Time) const override;
+	virtual FTransform ExtractRootTransform(float Time) const override;
+	virtual void ExtractPoseSearchNotifyStates(float Time, TArray<class UAnimNotifyState_PoseSearchBase*>& NotifyStates) const override;
+
+private:
+	float PlayLength = 0.0f;
+	float TotalRootDistance = 0.0f;
+	FTransform TotalRootTransform = FTransform::Identity;
+	TArray<float> AccumulatedRootDistance;
+	TArray<FTransform> AccumulatedRootTransform;
+	
+	void ProcessPlayLength();
+	void ProcessRootDistance();
+	void ProcessRootTransform();
+
+	// Extracts the pre-computed blend space root transform. ProcessRootTransform must be run first.
+	FTransform ExtractBlendSpaceRootTrackTransform(float Time) const;
+	FTransform ExtractBlendSpaceRootMotion(float StartTime, float DeltaTime, bool bAllowLooping) const;
+	FTransform ExtractBlendSpaceRootMotionFromRange(float StartTrackPosition, float EndTrackPosition) const;
+};
+
+void FBlendSpaceSampler::Init(const FInput& InInput)
+{
+	check(InInput.BlendSpace);
+
+	Input = InInput;
+}
+
+void FBlendSpaceSampler::Process()
+{
+	FMemMark Mark(FMemStack::Get());
+
+	ProcessPlayLength();
+	ProcessRootTransform();
+	ProcessRootDistance();
+}
+
+float FBlendSpaceSampler::GetTimeFromRootDistance(float Distance) const
+{
+	int32 NextSampleIdx = 1;
+	int32 PrevSampleIdx = 0;
+	if (Distance > 0.0f)
+	{
+		// Search for the distance value. Because the values will be extrapolated if necessary
+		// LowerBound might go past the end of the array, in which case the last valid index is used
+		int32 ClipDistanceLowerBoundIndex = Algo::LowerBound(AccumulatedRootDistance, Distance);
+		NextSampleIdx = FMath::Min(
+			ClipDistanceLowerBoundIndex,
+			AccumulatedRootDistance.Num() - 1);
+
+		// Compute distance interpolation amount
+		PrevSampleIdx = FMath::Max(0, NextSampleIdx - 1);
+	}
+
+	float NextDistance = AccumulatedRootDistance[NextSampleIdx];
+	float PrevDistance = AccumulatedRootDistance[PrevSampleIdx];
+	float DistanceSampleAlpha = FMath::GetRangePct(PrevDistance, NextDistance, Distance);
+
+	// Convert to time
+	float ClipTime = (float(NextSampleIdx) - (1.0f - DistanceSampleAlpha)) / Input.RootDistanceSamplingRate;
+	return ClipTime;
+}
+
+void FBlendSpaceSampler::ExtractPose(const FAnimExtractContext& ExtractionCtx, FAnimationPoseData& OutAnimPoseData) const
+{
+	TArray<FBlendSampleData> BlendSamples;
+	int32 TriangulationIndex = 0;
+	Input.BlendSpace->GetSamplesFromBlendInput(Input.BlendParameters, BlendSamples, TriangulationIndex, true);
+
+	for (int32 BlendSampleIdex = 0; BlendSampleIdex < BlendSamples.Num(); BlendSampleIdex++)
+	{
+		float Scale = BlendSamples[BlendSampleIdex].Animation->GetPlayLength() / PlayLength;
+
+		FDeltaTimeRecord BlendSampleDeltaTimeRecord;
+		BlendSampleDeltaTimeRecord.Set(ExtractionCtx.DeltaTimeRecord.GetPrevious() * Scale, ExtractionCtx.DeltaTimeRecord.Delta * Scale);
+
+		BlendSamples[BlendSampleIdex].DeltaTimeRecord = BlendSampleDeltaTimeRecord;
+		BlendSamples[BlendSampleIdex].PreviousTime = ExtractionCtx.DeltaTimeRecord.GetPrevious() * Scale;
+		BlendSamples[BlendSampleIdex].Time = ExtractionCtx.CurrentTime * Scale;
+	}
+
+	Input.BlendSpace->GetAnimationPose(BlendSamples, ExtractionCtx, OutAnimPoseData);
+}
+
+FTransform FBlendSpaceSampler::ExtractRootTransform(float Time) const
+{
+	if (Input.bLoopable)
+	{
+		FTransform LoopableRootTransform = ExtractBlendSpaceRootMotion(0.0f, Time, true);
+		return LoopableRootTransform;
+	}
+
+	const float ExtrapolationSampleTime = Input.ExtrapolationParameters.SampleTime;
+
+	const float ClampedTime = FMath::Clamp(Time, 0.0f, PlayLength);
+	const float ExtrapolationTime = Time - ClampedTime;
+
+	FTransform RootTransform = FTransform::Identity;
+
+	// If Time is less than zero, ExtrapolationTime will be negative. In this case, we extrapolate the beginning of the 
+	// animation to estimate where the root would be at Time
+	if (ExtrapolationTime < -SMALL_NUMBER)
+	{
+		FTransform SampleToExtrapolate = ExtractBlendSpaceRootMotionFromRange(0.0f, ExtrapolationSampleTime);
+
+		const FTransform ExtrapolatedRootMotion = ExtrapolateRootMotion(
+			SampleToExtrapolate,
+			0.0f, ExtrapolationSampleTime,
+			ExtrapolationTime,
+			Input.ExtrapolationParameters);
+		RootTransform = ExtrapolatedRootMotion;
+	}
+	else
+	{
+		RootTransform = ExtractBlendSpaceRootMotionFromRange(0.0f, ClampedTime);
+
+		// If Time is greater than PlayLength, ExtrapolationTIme will be a positive number. In this case, we extrapolate
+		// the end of the animation to estimate where the root would be at Time
+		if (ExtrapolationTime > SMALL_NUMBER)
+		{
+			FTransform SampleToExtrapolate = ExtractBlendSpaceRootMotionFromRange(PlayLength - ExtrapolationSampleTime, PlayLength);
+
+			const FTransform ExtrapolatedRootMotion = ExtrapolateRootMotion(
+				SampleToExtrapolate,
+				PlayLength - ExtrapolationSampleTime, PlayLength,
+				ExtrapolationTime,
+				Input.ExtrapolationParameters);
+			RootTransform = ExtrapolatedRootMotion * RootTransform;
+		}
+	}
+
+	return RootTransform;
+}
+
 float FBlendSpaceSampler::ExtractRootDistance(float Time) const
 {
 	return ExtractAccumulatedRootDistance(
@@ -4405,7 +5047,7 @@
 		{
 			constexpr size_t NumResults = 10;
 			size_t ResultIndexes[NumResults + 1] = { 0 };
-			float ResultDistanceSqr[NumResults + 1] = { 0.0f };
+			float ResultDistanceSqr[NumResults + 1] = { 0 };
 			FKDTree::KNNResultSet ResultSet(NumResults, ResultIndexes, ResultDistanceSqr);
 			GroupSearchIndex.KDTree.FindNeighbors(ResultSet, &GroupPCAValues[PointIndex * NumberOfPrincipalComponents]);
 
@@ -4427,7 +5069,7 @@
 		{
 			constexpr size_t NumResults = 10;
 			size_t ResultIndexes[NumResults + 1] = { 0 };
-			float ResultDistanceSqr[NumResults + 1] = { 0.0f };
+			float ResultDistanceSqr[NumResults + 1] = { 0 };
 			FKDTree::KNNResultSet ResultSet(NumResults, ResultIndexes, ResultDistanceSqr);
 
 			const RowMajorVectorMapConst MapGroupValues(&GroupValues[PointIndex * NumDimensions], 1, NumDimensions);
@@ -4847,375 +5489,6 @@
 	return bSuccess;
 }
 
-bool FQueryBuildingContext::IsInitialized() const
-{
-	return Schema && Schema->IsValid() && Query.IsInitializedForSchema(Schema);
-}
-
-bool BuildQuery(FQueryBuildingContext& QueryBuildingContext)
-{
-	if (!QueryBuildingContext.IsInitialized())
-	{
-		return false;
-	}
-
-	bool bSuccess = true;
-	for (const TObjectPtr<UPoseSearchFeatureChannel>& Channel : QueryBuildingContext.Schema->Channels)
-	{
-		bool bChannelSuccess = Channel->BuildQuery(QueryBuildingContext);
-		bSuccess &= bChannelSuccess;
-	}
-
-	return bSuccess;
-}
-
-FSearchResult SearchPCAKDTree(FSearchContext& SearchContext)
-{
-	QUICK_SCOPE_CYCLE_COUNTER(STAT_PoseSearch_PCA_KNN);
-	SCOPE_CYCLE_COUNTER(STAT_PoseSearchPCAKNN);
-
-	FSearchResult Result;
-
-	const UPoseSearchDatabase* Database = SearchContext.GetSourceDatabase();
-	check(Database);
-
-	const int32 NumDimensions = Database->Schema->Layout.NumFloats;
-	const FPoseSearchIndex* SearchIndex = SearchContext.GetSearchIndex();
-	check(SearchIndex);
-
-	const uint32 NumberOfPrincipalComponents = Database->GetNumberOfPrincipalComponents();
-	const uint32 KDTreeQueryNumNeighbors = FMath::Clamp<uint32>(Database->KDTreeQueryNumNeighbors, 1, SearchIndex->NumPoses);
-
-	//stack allocated temporaries
-	TArrayView<size_t> ResultIndexes((size_t*)FMemory_Alloca((KDTreeQueryNumNeighbors + 1) * sizeof(size_t)), KDTreeQueryNumNeighbors + 1);
-	TArrayView<float> ResultDistanceSqr((float*)FMemory_Alloca((KDTreeQueryNumNeighbors + 1) * sizeof(float)), KDTreeQueryNumNeighbors + 1);
-	RowMajorVectorMap WeightedQueryValues((float*)FMemory_Alloca(NumDimensions * sizeof(float)), 1, NumDimensions);
-	RowMajorVectorMap CenteredQueryValues((float*)FMemory_Alloca(NumDimensions * sizeof(float)), 1, NumDimensions);
-	RowMajorVectorMap ProjectedQueryValues((float*)FMemory_Alloca(NumberOfPrincipalComponents * sizeof(float)), 1, NumberOfPrincipalComponents);
-
-	// KDTree in PCA space search
-	if (Database->PoseSearchMode == EPoseSearchMode::PCAKDTree_Validate)
-	{
-		for (const FGroupSearchIndex& GroupSearchIndex : SearchIndex->Groups)
-		{
-			const RowMajorVectorMapConst MapWeights(GroupSearchIndex.Weights.GetData(), 1, NumDimensions);
-
-			// testing the KDTree is returning the proper searches for all the original points transformed in pca space
-			for (int32 PoseIdx = GroupSearchIndex.StartPoseIndex; PoseIdx < GroupSearchIndex.EndPoseIndex; ++PoseIdx)
-			{
-				FKDTree::KNNResultSet ResultSet(Database->KDTreeQueryNumNeighbors, ResultIndexes, ResultDistanceSqr);
-				TArrayView<const float> PoseValues = SearchIndex->GetPoseValues(PoseIdx);
-
-				const RowMajorVectorMapConst Mean(GroupSearchIndex.Mean.GetData(), 1, NumDimensions);
-				const ColMajorMatrixMapConst PCAProjectionMatrix(GroupSearchIndex.PCAProjectionMatrix.GetData(), NumDimensions, NumberOfPrincipalComponents);
-
-				const RowMajorVectorMapConst QueryValues(PoseValues.GetData(), 1, NumDimensions);
-				WeightedQueryValues = QueryValues.array() * MapWeights.array();
-				CenteredQueryValues.noalias() = WeightedQueryValues - Mean;
-				ProjectedQueryValues.noalias() = CenteredQueryValues * PCAProjectionMatrix;
-
-				GroupSearchIndex.KDTree.FindNeighbors(ResultSet, ProjectedQueryValues.data());
-
-				size_t ResultIndex = 0;
-				for (; ResultIndex < ResultSet.Num(); ++ResultIndex)
-				{
-					if ((PoseIdx - GroupSearchIndex.StartPoseIndex) == ResultIndexes[ResultIndex])
-					{
-						check(ResultDistanceSqr[ResultIndex] < UE_KINDA_SMALL_NUMBER);
-						break;
-					}
-				}
-				check(ResultIndex < ResultSet.Num());
-			}
-		}
-	}
-
-	// @todo: implement support for DatabaseTagQuery
-	FPoseCost BestPoseCost;
-	int32 BestPoseIdx = INDEX_NONE;
-	for (const FGroupSearchIndex& GroupSearchIndex : SearchIndex->Groups)
-	{
-		const RowMajorVectorMapConst MapWeights(GroupSearchIndex.Weights.GetData(), 1, NumDimensions);
-		FKDTree::KNNResultSet ResultSet(KDTreeQueryNumNeighbors, ResultIndexes, ResultDistanceSqr);
-
-		check(SearchContext.QueryValues.Num() == NumDimensions);
-
-		const RowMajorVectorMapConst Mean(GroupSearchIndex.Mean.GetData(), 1, NumDimensions);
-		const ColMajorMatrixMapConst PCAProjectionMatrix(GroupSearchIndex.PCAProjectionMatrix.GetData(), NumDimensions, NumberOfPrincipalComponents);
-
-		// transforming query values into PCA space to query the KDTree
-		const RowMajorVectorMapConst QueryValues(SearchContext.QueryValues.GetData(), 1, NumDimensions);
-		WeightedQueryValues = QueryValues.array() * MapWeights.array();
-		CenteredQueryValues.noalias() = WeightedQueryValues - Mean;
-		ProjectedQueryValues.noalias() = CenteredQueryValues * PCAProjectionMatrix;
-
-		GroupSearchIndex.KDTree.FindNeighbors(ResultSet, ProjectedQueryValues.data());
-
-		for (size_t ResultIndex = 0; ResultIndex < ResultSet.Num(); ++ResultIndex)
-		{
-			const int32 PoseIdx = ResultIndexes[ResultIndex] + GroupSearchIndex.StartPoseIndex;
-
-			const FPoseSearchPoseMetadata& Metadata = SearchIndex->PoseMetadata[PoseIdx];
-
-			if (EnumHasAnyFlags(Metadata.Flags, EPoseSearchPoseFlags::BlockTransition))
-			{
-				continue;
-			}
-
-			FPoseCost PoseCost = ComparePoses(PoseIdx, SearchContext, GroupSearchIndex.GroupIndex);
-
-			if (PoseCost < BestPoseCost)
-			{
-				BestPoseCost = PoseCost;
-				BestPoseIdx = PoseIdx;
-			}
-		}
-	}
-
-	Result.PoseCost = BestPoseCost;
-	Result.PoseIdx = BestPoseIdx;
-	Result.SearchIndexAsset = SearchIndex->FindAssetForPose(BestPoseIdx);
-	Result.AssetTime = SearchIndex->GetAssetTime(BestPoseIdx, Result.SearchIndexAsset);
-
-	return Result;
-}
-
-FSearchResult SearchBruteForce(FSearchContext& SearchContext)
-{
-	QUICK_SCOPE_CYCLE_COUNTER(STAT_PoseSearch_Brute_Force);
-	SCOPE_CYCLE_COUNTER(STAT_PoseSearchBruteForce);
-
-	FSearchResult Result;
-
-	const FPoseSearchIndex* SearchIndex = SearchContext.GetSearchIndex();
-	check(SearchIndex);
-	const UPoseSearchDatabase* Database = SearchContext.GetSourceDatabase();
-
-	FPoseCost BestPoseCost;
-	int32 BestPoseIdx = INDEX_NONE;
-	for (const FPoseSearchIndexAsset& Asset : SearchIndex->Assets)
-	{
-		if (Database && SearchContext.DatabaseTagQuery)
-		{
-			if (!SearchContext.DatabaseTagQuery->Matches(*Database->GetSourceAssetGroupTags(&Asset)))
-			{
-				continue;
-			}
-		}
-
-		const int32 EndIndex = Asset.FirstPoseIdx + Asset.NumPoses;
-		for (int32 PoseIdx = Asset.FirstPoseIdx; PoseIdx < EndIndex; ++PoseIdx)
-		{
-			const FPoseSearchPoseMetadata& Metadata = SearchIndex->PoseMetadata[PoseIdx];
-
-			if (EnumHasAnyFlags(Metadata.Flags, EPoseSearchPoseFlags::BlockTransition))
-			{
-				continue;
-			}
-
-			FPoseCost PoseCost = ComparePoses(PoseIdx, SearchContext, Asset.SourceGroupIdx);
-
-			if (PoseCost < BestPoseCost)
-			{
-				BestPoseCost = PoseCost;
-				BestPoseIdx = PoseIdx;
-			}
-		}
-	}
-
-	Result.PoseCost = BestPoseCost;
-	Result.PoseIdx = BestPoseIdx;
-	Result.SearchIndexAsset = SearchIndex->FindAssetForPose(BestPoseIdx);
-	Result.AssetTime = SearchIndex->GetAssetTime(BestPoseIdx, Result.SearchIndexAsset);
-
-	return Result;
-}
-
-FSearchResult Search(FSearchContext& SearchContext)
-{
-	FSearchResult Result;
-
-	const FPoseSearchIndex* SearchIndex = SearchContext.GetSearchIndex();
-	if (!SearchIndex)
-	{
-		return Result;
-	}
-
-	if (!ensure(SearchIndex->IsValid() && !SearchIndex->IsEmpty()))
-	{
-		return Result;
-	}
-
-	if (!ensure(SearchContext.QueryValues.Num() == SearchIndex->Schema->Layout.NumFloats))
-	{
-		return Result;
-	}
-
-	const UPoseSearchDatabase* Database = SearchContext.GetSourceDatabase();
-	if (Database && Database->PoseSearchMode != EPoseSearchMode::BruteForce)
-	{
-		Result = SearchPCAKDTree(SearchContext);
-	}
-
-	if (!Database || Database->PoseSearchMode == EPoseSearchMode::BruteForce || Database->PoseSearchMode == EPoseSearchMode::PCAKDTree_Compare)
-	{
-		Result = SearchBruteForce(SearchContext);
-	}
-
-	SearchContext.DebugDrawParams.PoseVector = SearchContext.QueryValues;
-	SearchContext.DebugDrawParams.PoseIdx = Result.PoseIdx;
-	Draw(SearchContext.DebugDrawParams);
-
-	return Result;
-}
-
-static void ComputePoseCostAddends(
-	int32 PoseIdx, 
-	FSearchContext& SearchContext, 
-	float& OutNotifyAddend, 
-	float& OutMirrorMismatchAddend)
-{
-	OutNotifyAddend = 0.0f;
-	OutMirrorMismatchAddend = 0.0f;
-
-	if (SearchContext.QueryMirrorRequest != EPoseSearchBooleanRequest::Indifferent)
-	{
-		const FPoseSearchIndexAsset* IndexAsset = SearchContext.GetSearchIndex()->FindAssetForPose(PoseIdx);
-		const bool bMirroringMismatch =
-			(IndexAsset->bMirrored && SearchContext.QueryMirrorRequest == EPoseSearchBooleanRequest::FalseValue) ||
-			(!IndexAsset->bMirrored && SearchContext.QueryMirrorRequest == EPoseSearchBooleanRequest::TrueValue);
-		if (bMirroringMismatch)
-		{
-			OutMirrorMismatchAddend = SearchContext.GetMirrorMismatchCost();
-		}
-	}
-
-	const FPoseSearchPoseMetadata& PoseMetadata = SearchContext.GetSearchIndex()->PoseMetadata[PoseIdx];
-	OutNotifyAddend = PoseMetadata.CostAddend;
-}
-
-
-FPoseCost ComparePoses(int32 PoseIdx, FSearchContext& SearchContext, int32 GroupIdx)
-{
-	FPoseCost Result;
-
-	const FPoseSearchIndex* SearchIndex = SearchContext.GetSearchIndex();
-	if (!ensure(SearchIndex))
-	{
-		return Result;
-	}
-
-	TArrayView<const float> PoseValues = SearchIndex->GetPoseValues(PoseIdx);
-	if (!ensure(PoseValues.Num() == SearchContext.QueryValues.Num()))
-	{
-		return Result;
-	}
-
-	if (SearchContext.WeightsContext)
-	{
-		if (GroupIdx == INDEX_NONE)
-		{
-			const FPoseSearchIndexAsset* SearchIndexAsset = SearchIndex->FindAssetForPose(PoseIdx);
-			if (!ensure(SearchIndexAsset))
-			{
-				return Result;
-			}
-
-			GroupIdx = SearchIndexAsset->SourceGroupIdx;
-		}
-
-		const FPoseSearchWeights* WeightsSet =
-			SearchContext.WeightsContext->GetGroupWeights(GroupIdx);
-		Result.Dissimilarity = CompareFeatureVectors(
-			PoseValues.Num(),
-			PoseValues.GetData(),
-			SearchContext.QueryValues.GetData(),
-			WeightsSet->Weights.GetData());
-	}
-	else
-	{
-		Result.Dissimilarity = CompareFeatureVectors(PoseValues.Num(), PoseValues.GetData(), SearchContext.QueryValues.GetData());
-	}
-
-	float NotifyAddend = 0.0f;
-	float MirrorMismatchAddend = 0.0f;
-	ComputePoseCostAddends(PoseIdx, SearchContext, NotifyAddend, MirrorMismatchAddend);
-	Result.CostAddend = NotifyAddend + MirrorMismatchAddend;
-	Result.TotalCost = Result.Dissimilarity + Result.CostAddend;
-
-	return Result;
-}
-
-
-FPoseCost ComparePoses(int32 PoseIdx, FSearchContext& SearchContext, FPoseCostDetails& OutPoseCostDetails)
-{
-	using namespace Eigen;
-
-	FPoseCost Result;
-
-	TArrayView<const float> PoseValues = SearchContext.GetSearchIndex()->GetPoseValues(PoseIdx);
-	const int32 Dims = PoseValues.Num();
-	if (!ensure(Dims == SearchContext.QueryValues.Num()))
-	{
-		return Result;
-	}
-
-	OutPoseCostDetails.CostVector.SetNum(Dims);
-
-	// Setup Eigen views onto our vectors
-	auto OutCostVector = Map<ArrayXf>(OutPoseCostDetails.CostVector.GetData(), Dims);
-	auto PoseVector = Map<const ArrayXf>(PoseValues.GetData(), Dims);
-	auto QueryVector = Map<const ArrayXf>(SearchContext.QueryValues.GetData(), Dims);
-	
-	// Compute weighted squared difference vector
-	const FPoseSearchIndexAsset* SearchIndexAsset = SearchContext.GetSearchIndex()->FindAssetForPose(PoseIdx);
-	if (SearchContext.WeightsContext)
-	{
-		const FPoseSearchWeights* WeightsSet = 
-			SearchContext.WeightsContext->GetGroupWeights(SearchIndexAsset->SourceGroupIdx);
-		check(WeightsSet);
-		check(WeightsSet->Weights.Num() == Dims);
-		auto WeightsVector = Map<const ArrayXf>(WeightsSet->Weights.GetData(), Dims);
-
-		OutCostVector = WeightsVector * (PoseVector - QueryVector).square();
-		Result.Dissimilarity = OutCostVector.sum();
-	}
-	else
-	{
-		OutCostVector = (PoseVector - QueryVector).square();
-		Result.Dissimilarity = OutCostVector.sum();
-	}
-
-	// Output result
-	float NotifyAddend = 0.0f;
-	float MirrorMismatchAddend = 0.0f;
-	ComputePoseCostAddends(PoseIdx, SearchContext, NotifyAddend, MirrorMismatchAddend);
-	Result.CostAddend = NotifyAddend + MirrorMismatchAddend;
-	Result.TotalCost = Result.Dissimilarity + Result.CostAddend;
-
-	// Output cost details
-	OutPoseCostDetails.NotifyCostAddend = NotifyAddend;
-	OutPoseCostDetails.MirrorMismatchCostAddend = MirrorMismatchAddend;
-	OutPoseCostDetails.PoseCost = Result;
-	CalcChannelCosts(SearchContext.GetSearchIndex()->Schema, OutPoseCostDetails.CostVector, OutPoseCostDetails.ChannelCosts);
-
-
-#if DO_GUARD_SLOW
-	{
-		// Verify details pose comparator agrees with runtime pose comparator
-		FPoseCost RuntimeComparatorCost = ComparePoses(PoseIdx, SearchContext, SearchIndexAsset->SourceGroupIdx);
-		checkSlow(FMath::IsNearlyEqual(Result.TotalCost, RuntimeComparatorCost.TotalCost, 1e-3f));
-
-		// Verify channel cost decomposition agrees with runtime pose comparator
-		auto OutChannelCosts = Map<const ArrayXf>(OutPoseCostDetails.ChannelCosts.GetData(), OutPoseCostDetails.ChannelCosts.Num());
-		checkSlow(FMath::IsNearlyEqual(OutChannelCosts.sum(), RuntimeComparatorCost.Dissimilarity, 1e-3f));
-	}
-#endif
-
-	return Result;
-}
-
-
 //////////////////////////////////////////////////////////////////////////
 // FModule
 
@@ -5257,6 +5530,8 @@
 {
 	UE::Anim::IPoseSearchProvider::FSearchResult ProviderResult;
 
+	using namespace UE::PoseSearch;
+
 	const UPoseSearchSequenceMetaData* MetaData = Sequence ? Sequence->FindMetaDataByClass<UPoseSearchSequenceMetaData>() : nullptr;
 	if (!MetaData || !MetaData->IsValidForSearch())
 	{
@@ -5270,27 +5545,15 @@
 	}
 
 	FPoseHistory& PoseHistory = PoseHistoryProvider->GetPoseHistory();
-	FPoseSearchFeatureVectorBuilder& QueryBuilder = PoseHistory.GetQueryBuilder();
-	QueryBuilder.Init(MetaData->Schema);
 
-	UE::PoseSearch::FQueryBuildingContext QueryBuildingContext(QueryBuilder);
-	QueryBuildingContext.Schema = MetaData->Schema;
-	QueryBuildingContext.History = &PoseHistory;
-	QueryBuildingContext.Trajectory = nullptr;
+	FPoseSearchContext SearchContext;
+	SearchContext.OwningComponent = GraphContext.AnimInstanceProxy->GetSkelMeshComponent();
+	SearchContext.BoneContainer = &GraphContext.AnimInstanceProxy->GetRequiredBones();
+	SearchContext.History = &PoseHistoryProvider->GetPoseHistory();
 
-	if (!UE::PoseSearch::BuildQuery(QueryBuildingContext))
-	{
-		return ProviderResult;
-	}
+	UE::PoseSearch::FSearchResult Result = MetaData->Search(SearchContext);
 
-	QueryBuilder.Normalize(MetaData->SearchIndex);
-
-	::UE::PoseSearch::FSearchContext SearchContext;
-	SearchContext.SetSource(Sequence);
-	SearchContext.QueryValues = QueryBuilder.GetNormalizedValues();
-	::UE::PoseSearch::FSearchResult Result = ::UE::PoseSearch::Search(SearchContext);
-
-	ProviderResult.Dissimilarity = Result.PoseCost.TotalCost;
+	ProviderResult.Dissimilarity = Result.PoseCost.GetTotalCost();
 	ProviderResult.PoseIdx = Result.PoseIdx;
 	ProviderResult.TimeOffsetSeconds = Result.AssetTime;
 	return ProviderResult;
@@ -5387,6 +5650,13 @@
 
 }} // namespace UE::PoseSearch
 
+EPoseSearchPostSearchStatus UPoseSearchPostProcessor::PostProcess_Implementation(
+	FPoseSearchCost& InOutCost) const
+{
+	return EPoseSearchPostSearchStatus::Continue;
+}
+
+
 #undef LOCTEXT_NAMESPACE
 
 IMPLEMENT_MODULE(UE::PoseSearch::FModule, PoseSearch)
diff --git a/Engine/Plugins/Experimental/Animation/PoseSearch/Source/Runtime/Private/PoseSearchDerivedData.cpp b/Engine/Plugins/Experimental/Animation/PoseSearch/Source/Runtime/Private/PoseSearchDerivedData.cpp
--- a/Engine/Plugins/Experimental/Animation/PoseSearch/Source/Runtime/Private/PoseSearchDerivedData.cpp
+++ b/Engine/Plugins/Experimental/Animation/PoseSearch/Source/Runtime/Private/PoseSearchDerivedData.cpp
@@ -34,6 +34,7 @@ void FPoseSearchDatabaseDerivedData::Cache(UPoseSearchDatabase& Database, bool b
 		SearchIndex.Reset();
 		SearchIndex.Schema = Database.Schema;
 		DerivedDataKey = { UE::DerivedData::FCacheBucket(), FIoHash::Zero };
+		PendingDerivedDataKey = FIoHash::Zero;
 	}
 }
 
diff --git a/Engine/Plugins/Experimental/Animation/PoseSearch/Source/Runtime/Private/PoseSearchFeatureChannels.cpp b/Engine/Plugins/Experimental/Animation/PoseSearch/Source/Runtime/Private/PoseSearchFeatureChannels.cpp
--- a/Engine/Plugins/Experimental/Animation/PoseSearch/Source/Runtime/Private/PoseSearchFeatureChannels.cpp
+++ b/Engine/Plugins/Experimental/Animation/PoseSearch/Source/Runtime/Private/PoseSearchFeatureChannels.cpp
@@ -272,11 +272,23 @@ void UPoseSearchFeatureChannel_Pose::GenerateDDCKey(FBlake3& InOutKeyHasher) con
 	InOutKeyHasher.Update(MakeMemoryView(SampleTimes));
 }
 
-bool UPoseSearchFeatureChannel_Pose::BuildQuery(UE::PoseSearch::FQueryBuildingContext& Context) const
+bool UPoseSearchFeatureChannel_Pose::BuildQuery(
+	FPoseSearchContext& SearchContext,
+	FPoseSearchFeatureVectorBuilder& InOutQuery) const
 {
-	check(Context.Schema);
-	
-	if (!Context.History)
+	const bool bSkip =
+		SearchContext.CurrentResult.IsValid() &&
+		SearchContext.CurrentResult.Database->Schema == InOutQuery.GetSchema();
+
+	if (bSkip)
+	{
+		// todo: instead of skipping because the pose should have been copied already (that's currently happening in
+		// UPoseSearchSchema::BuildIndex()), consider making the copy here, but only copy the right values, not the
+		// whole query
+		return true;
+	}
+
+	if (!SearchContext.History)
 	{
 		return false;
 	}
@@ -296,27 +308,27 @@ bool UPoseSearchFeatureChannel_Pose::BuildQuery(UE::PoseSearch::FQueryBuildingCo
 		}
 
 		float SecondsAgo = -SampleTime;
-		if (!Context.History->TrySamplePose(SecondsAgo, Context.Schema->Skeleton->GetReferenceSkeleton(), Context.Schema->BoneIndicesWithParents))
+		if (!SearchContext.History->TrySamplePose(SecondsAgo, InOutQuery.GetSchema()->Skeleton->GetReferenceSkeleton(), InOutQuery.GetSchema()->BoneIndicesWithParents))
 		{
 			return false;
 		}
 
-		TArrayView<const FTransform> ComponentPose = Context.History->GetComponentPoseSample();
-		TArrayView<const FTransform> ComponentPrevPose = Context.History->GetPrevComponentPoseSample();
-		FTransform RootTransform = Context.History->GetRootTransformSample();
-		FTransform RootTransformPrev = Context.History->GetPrevRootTransformSample();
+		TArrayView<const FTransform> ComponentPose = SearchContext.History->GetComponentPoseSample();
+		TArrayView<const FTransform> ComponentPrevPose = SearchContext.History->GetPrevComponentPoseSample();
+		FTransform RootTransform = SearchContext.History->GetRootTransformSample();
+		FTransform RootTransformPrev = SearchContext.History->GetPrevRootTransformSample();
 		for (int32 SampledBoneIdx = 0; SampledBoneIdx != SampledBones.Num(); ++SampledBoneIdx)
 		{
 			Feature.ChannelFeatureId = SampledBoneIdx;
 
 			int32 SchemaBoneIdx = FeatureParams[SampledBoneIdx].SchemaBoneIdx;
 
-			int32 SkeletonBoneIndex = Context.Schema->BoneIndices[SchemaBoneIdx];
+			int32 SkeletonBoneIndex = InOutQuery.GetSchema()->BoneIndices[SchemaBoneIdx];
 
 			const FTransform& Transform = ComponentPose[SkeletonBoneIndex];
 			const FTransform& PrevTransform = ComponentPrevPose[SkeletonBoneIndex] * (RootTransformPrev * RootTransform.Inverse());
-			Context.Query.SetTransform(Feature, Transform);
-			Context.Query.SetTransformVelocity(Feature, Transform, PrevTransform, Context.History->GetSampleTimeInterval());
+			InOutQuery.SetTransform(Feature, Transform);
+			InOutQuery.SetTransformVelocity(Feature, Transform, PrevTransform, SearchContext.History->GetSampleTimeInterval());
 		}
 	}
 
@@ -642,9 +654,11 @@ void UPoseSearchFeatureChannel_Trajectory::GenerateDDCKey(FBlake3& InOutKeyHashe
 	InOutKeyHasher.Update(MakeMemoryView(SampleOffsets));
 }
 
-bool UPoseSearchFeatureChannel_Trajectory::BuildQuery(UE::PoseSearch::FQueryBuildingContext& Context) const
+bool UPoseSearchFeatureChannel_Trajectory::BuildQuery(
+	FPoseSearchContext& SearchContext,
+	FPoseSearchFeatureVectorBuilder& InOutQuery) const
 {
-	if (!Context.Trajectory)
+	if (!SearchContext.Trajectory)
 	{
 		return false;
 	}
@@ -671,14 +685,18 @@ bool UPoseSearchFeatureChannel_Trajectory::BuildQuery(UE::PoseSearch::FQueryBuil
 	for (int32 Idx = 0, NextIterStartIdx = 0, Num = SampleOffsets.Num(); Idx < Num; ++Idx)
 	{
 		const float SampleOffset = SampleOffsets[Idx];
-		const FTrajectorySample Sample = FTrajectorySampleRange::IterSampleTrajectory(Context.Trajectory->Samples, SampleDomain, SampleOffset, NextIterStartIdx);
+		const FTrajectorySample Sample = FTrajectorySampleRange::IterSampleTrajectory(
+			SearchContext.Trajectory->Samples,
+			SampleDomain, 
+			SampleOffset, 
+			NextIterStartIdx);
 
 		Feature.SubsampleIdx = Idx;
 
 		Feature.Type = EPoseSearchFeatureType::LinearVelocity;
-		Context.Query.SetVector(Feature, Sample.LinearVelocity);
+		InOutQuery.SetVector(Feature, Sample.LinearVelocity);
 
-		Context.Query.SetTransform(Feature, Sample.Transform);
+		InOutQuery.SetTransform(Feature, Sample.Transform);
 	}
 
 	return true;
diff --git a/Engine/Plugins/Experimental/Animation/PoseSearch/Source/Runtime/Private/PoseSearchLibrary.cpp b/Engine/Plugins/Experimental/Animation/PoseSearch/Source/Runtime/Private/PoseSearchLibrary.cpp
--- a/Engine/Plugins/Experimental/Animation/PoseSearch/Source/Runtime/Private/PoseSearchLibrary.cpp
+++ b/Engine/Plugins/Experimental/Animation/PoseSearch/Source/Runtime/Private/PoseSearchLibrary.cpp
@@ -9,6 +9,7 @@
 #include "Animation/BlendSpace.h"
 #include "PoseSearch/AnimNode_MotionMatching.h"
 #include "Trace/PoseSearchTraceLogger.h"
+#include "MotionTrajectoryLibrary.h"
 
 #define LOCTEXT_NAMESPACE "PoseSearchLibrary"
 
@@ -66,119 +67,18 @@ static void ComputeDatabaseBlendSpaceFilter(
 //////////////////////////////////////////////////////////////////////////
 // FMotionMatchingState
 
-bool FMotionMatchingState::InitNewDatabaseSearch(
-	const UPoseSearchDatabase* Database,
-	FText* OutError
-)
-{
-	bool bValidDatabase = Database && Database->IsValidForSearch();
-
-	if (bValidDatabase)
-	{
-		CurrentDatabase = Database;
-#if WITH_EDITOR
-		CurrentSearchIndexHash = Database->GetSearchIndexHash();
-#endif
-	}
-	else
-	{
-		CurrentDatabase = nullptr;
-#if WITH_EDITOR
-		CurrentSearchIndexHash = FIoHash::Zero;
-#endif	
-	}
-
-	if (!bValidDatabase && OutError)
-	{
-		if (Database)
-		{
-			*OutError = FText::Format(LOCTEXT("InvalidDatabase", "Invalid database for motion matching. Try re-saving {0}."), FText::FromString(Database->GetPathName()));
-		}
-		else
-		{
-			*OutError = LOCTEXT("NoDatabase", "No database provided for motion matching.");
-		}
-	}
-
-	Reset();
-
-	return bValidDatabase;
-}
-
 void FMotionMatchingState::Reset()
 {
-	DbPoseIdx = INDEX_NONE;
-	SearchIndexAssetIdx = INDEX_NONE;
+	CurrentSearchResult.Reset();
 	AssetPlayerTime = 0.0f;
 	// Set the elapsed time to INFINITY to trigger a search right away
 	ElapsedPoseJumpTime = INFINITY;
-
-	if (CurrentDatabase.IsValid() && CurrentDatabase->IsValidForSearch())
-	{
-		ComposedQuery.Init(CurrentDatabase->Schema);
-	}
-	else
-	{
-		ComposedQuery.Reset();
-	}
-
 }
 
 void FMotionMatchingState::AdjustAssetTime(float AssetTime)
 {
-	if (SearchIndexAssetIdx == INDEX_NONE || !CurrentDatabase.IsValid() || !(CurrentDatabase->GetSearchIndex()))
-	{
-		return;
-	}
-
-	const FPoseSearchIndexAsset* SearchIndexAsset = GetCurrentSearchIndexAsset();
-
-	if (SearchIndexAsset->Type == ESearchIndexAssetType::Sequence)
-	{
-		const FPoseSearchDatabaseSequence& DbSequence = CurrentDatabase->GetSequenceSourceAsset(SearchIndexAsset);
-
-		if (SearchIndexAsset->SamplingInterval.Contains(AssetTime))
-		{
-			DbPoseIdx = CurrentDatabase->GetPoseIndexFromTime(AssetTime, SearchIndexAsset);
-			AssetPlayerTime = AssetTime;
-		}
-		else
-		{
-			DbPoseIdx = INDEX_NONE;
-			AssetPlayerTime = 0.0f;
-			SearchIndexAssetIdx = INDEX_NONE;
-		}
-	}
-	else if (SearchIndexAsset->Type == ESearchIndexAssetType::BlendSpace)
-	{
-		const FPoseSearchDatabaseBlendSpace& DbBlendSpace = CurrentDatabase->GetBlendSpaceSourceAsset(SearchIndexAsset);
-
-		TArray<FBlendSampleData> BlendSamples;
-		int32 TriangulationIndex = 0;
-		DbBlendSpace.BlendSpace->GetSamplesFromBlendInput(SearchIndexAsset->BlendParameters, BlendSamples, TriangulationIndex, true);
-
-		float PlayLength = DbBlendSpace.BlendSpace->GetAnimationLengthFromSampleData(BlendSamples);
-
-		// Asset player time for blendspaces is normalized [0, 1] so we need to convert 
-		// to a real time before we advance it
-		float RealTime = AssetTime * PlayLength;
-
-		if (SearchIndexAsset->SamplingInterval.Contains(RealTime))
-		{
-			DbPoseIdx = CurrentDatabase->GetPoseIndexFromTime(RealTime, SearchIndexAsset);
-			AssetPlayerTime = AssetTime;
-		}
-		else
-		{
-			DbPoseIdx = INDEX_NONE;
-			AssetPlayerTime = 0.0f;
-			SearchIndexAssetIdx = INDEX_NONE;
-		}
-	}
-	else
-	{
-		checkNoEntry();
-	}
+	CurrentSearchResult.Update(AssetTime);
+	AssetPlayerTime = CurrentSearchResult.AssetTime;
 }
 
 bool FMotionMatchingState::CanAdvance(float DeltaTime, bool& bOutAdvanceToFollowUpAsset, UE::PoseSearch::FSearchResult& OutFollowUpAsset) const
@@ -186,7 +86,7 @@ bool FMotionMatchingState::CanAdvance(float DeltaTime, bool& bOutAdvanceToFollow
 	bOutAdvanceToFollowUpAsset = false;
 	OutFollowUpAsset = UE::PoseSearch::FSearchResult();
 
-	if (SearchIndexAssetIdx == INDEX_NONE || !CurrentDatabase.IsValid() || !(CurrentDatabase->GetSearchIndex()))
+	if (!CurrentSearchResult.IsValid())
 	{
 		return false;
 	}
@@ -195,7 +95,8 @@ bool FMotionMatchingState::CanAdvance(float DeltaTime, bool& bOutAdvanceToFollow
 
 	if (SearchIndexAsset->Type == ESearchIndexAssetType::Sequence)
 	{
-		const FPoseSearchDatabaseSequence& DbSequence = CurrentDatabase->GetSequenceSourceAsset(SearchIndexAsset);
+		const FPoseSearchDatabaseSequence& DbSequence = 
+			CurrentSearchResult.Database->GetSequenceSourceAsset(SearchIndexAsset);
 		const float AssetLength = DbSequence.Sequence->GetPlayLength();
 
 		float SteppedTime = AssetPlayerTime;
@@ -212,13 +113,13 @@ bool FMotionMatchingState::CanAdvance(float DeltaTime, bool& bOutAdvanceToFollow
 		else
 		{
 			// check if there's a follow-up that can be used
-			int32 FollowUpDbSequenceIdx = CurrentDatabase->Sequences.IndexOfByPredicate(
+			int32 FollowUpDbSequenceIdx = CurrentSearchResult.Database->Sequences.IndexOfByPredicate(
 				[&](const FPoseSearchDatabaseSequence& Entry)
 				{
 					return Entry.Sequence == DbSequence.FollowUpSequence;
 				});
 
-			int32 FollowUpSearchIndexAssetIdx = CurrentDatabase->GetSearchIndex()->Assets.IndexOfByPredicate(
+			int32 FollowUpSearchIndexAssetIdx = CurrentSearchResult.Database->GetSearchIndex()->Assets.IndexOfByPredicate(
 				[&](const FPoseSearchIndexAsset& Entry)
 				{
 					const bool bIsMatch =
@@ -233,7 +134,7 @@ bool FMotionMatchingState::CanAdvance(float DeltaTime, bool& bOutAdvanceToFollow
 				bOutAdvanceToFollowUpAsset = true;
 
 				const FPoseSearchIndexAsset* FollowUpSearchIndexAsset =
-					&CurrentDatabase->GetSearchIndex()->Assets[FollowUpSearchIndexAssetIdx];
+					&CurrentSearchResult.Database->GetSearchIndex()->Assets[FollowUpSearchIndexAssetIdx];
 
 				// Follow up asset time will start slightly before the beginning of the sequence as 
 				// this is essentially what the matching time in the corresponding main sequence is.
@@ -243,7 +144,7 @@ bool FMotionMatchingState::CanAdvance(float DeltaTime, bool& bOutAdvanceToFollow
 
 				// There is no correspoding pose index when we switch due to what is mentioned above
 				// so for now we just take whatever pose index is associated with the first frame.
-				OutFollowUpAsset.PoseIdx = CurrentDatabase->GetPoseIndexFromTime(FollowUpSearchIndexAsset->SamplingInterval.Min, FollowUpSearchIndexAsset);
+				OutFollowUpAsset.PoseIdx = CurrentSearchResult.Database->GetPoseIndexFromTime(FollowUpSearchIndexAsset->SamplingInterval.Min, FollowUpSearchIndexAsset);
 				OutFollowUpAsset.SearchIndexAsset = FollowUpSearchIndexAsset;
 				OutFollowUpAsset.AssetTime = FollowUpAssetTime;
 				return true;
@@ -252,7 +153,8 @@ bool FMotionMatchingState::CanAdvance(float DeltaTime, bool& bOutAdvanceToFollow
 	}
 	else if (SearchIndexAsset->Type == ESearchIndexAssetType::BlendSpace)
 	{
-		const FPoseSearchDatabaseBlendSpace& DbBlendSpace = CurrentDatabase->GetBlendSpaceSourceAsset(SearchIndexAsset);
+		const FPoseSearchDatabaseBlendSpace& DbBlendSpace = 
+			CurrentSearchResult.Database->GetBlendSpaceSourceAsset(SearchIndexAsset);
 
 		TArray<FBlendSampleData> BlendSamples;
 		int32 TriangulationIndex = 0;
@@ -300,8 +202,7 @@ static void RequestInertialBlend(const FAnimationUpdateContext& Context, float B
 void FMotionMatchingState::JumpToPose(const FAnimationUpdateContext& Context, const FMotionMatchingSettings& Settings, const UE::PoseSearch::FSearchResult& Result)
 {
 	// Remember which pose and sequence we're playing from the database
-	DbPoseIdx = Result.PoseIdx;
-	SearchIndexAssetIdx = CurrentDatabase->GetSearchIndex()->FindAssetIndex(Result.SearchIndexAsset);
+	CurrentSearchResult = Result;
 
 	ElapsedPoseJumpTime = 0.0f;
 	AssetPlayerTime = Result.AssetTime;
@@ -313,8 +214,9 @@ void FMotionMatchingState::JumpToPose(const FAnimationUpdateContext& Context, co
 
 void UpdateMotionMatchingState(
 	const FAnimationUpdateContext& Context,
-	const UPoseSearchDatabase* Database,
+	const UPoseSearchSearchableAsset* Searchable,
 	const FGameplayTagQuery* DatabaseTagQuery,
+	const FGameplayTagContainer* ActiveTagsContainer,
 	const FTrajectorySampleRange& Trajectory,
 	const FMotionMatchingSettings& Settings,
 	FMotionMatchingState& InOutMotionMatchingState
@@ -322,33 +224,13 @@ void UpdateMotionMatchingState(
 {
 	using namespace UE::PoseSearch;
 
-	if (!Database)
-	{
-		Context.LogMessage(EMessageSeverity::Error, LOCTEXT("NoDatabase", "No database provided for motion matching."));
-		return;
-	}
-	
-	// Check for a switch in the database
-	if (!InOutMotionMatchingState.IsCompatibleDatabase(Database))
-	{
-		FText InitError;
-		if (!InOutMotionMatchingState.InitNewDatabaseSearch(Database, &InitError))
-		{
-			Context.LogMessage(EMessageSeverity::Error, InitError);
-			return;
-		}
-	}
-
-#if WITH_EDITOR
-	if (InOutMotionMatchingState.CurrentDatabase->IsDerivedDataBuildPending())
+	if (!Searchable)
 	{
-		InOutMotionMatchingState.Reset();
 		Context.LogMessage(
 			EMessageSeverity::Error, 
-			LOCTEXT("PendingDerivedData", "Derived data build pending."));
+			LOCTEXT("NoSearchable", "No searchable asset provided for motion matching."));
 		return;
 	}
-#endif
 
 	const float DeltaTime = Context.GetDeltaTime();
 
@@ -356,34 +238,7 @@ void UpdateMotionMatchingState(
 	InOutMotionMatchingState.Flags = EMotionMatchingFlags::None;
 
 	// Record Current Pose Index for Debugger
-	int32 CurrentPoseIdx = InOutMotionMatchingState.DbPoseIdx;
-
-	FQueryBuildingContext QueryBuildingContext(InOutMotionMatchingState.ComposedQuery);
-	QueryBuildingContext.Schema = Database->Schema;
-	QueryBuildingContext.History = nullptr;
-	QueryBuildingContext.Trajectory = &Trajectory;
-
-	// Build the search query. This is done even when we don't search 
-	// since we still want to record it for debugging purposes.
-	if (InOutMotionMatchingState.DbPoseIdx != INDEX_NONE)
-	{
-		// Copy search query directly from the database if we have an active pose
-		InOutMotionMatchingState.ComposedQuery.CopyFromSearchIndex(
-			*Database->GetSearchIndex(),
-			InOutMotionMatchingState.DbPoseIdx);
-	}
-	else
-	{
-		InOutMotionMatchingState.ComposedQuery.ResetFeatures();
-		IPoseHistoryProvider* PoseHistoryProvider = Context.GetMessage<IPoseHistoryProvider>();
-		if (PoseHistoryProvider)
-		{
-			QueryBuildingContext.History = &PoseHistoryProvider->GetPoseHistory();
-		}
-	}
-
-	BuildQuery(QueryBuildingContext);
-	InOutMotionMatchingState.ComposedQuery.Normalize(*Database->GetSearchIndex());
+	const FSearchResult LastResult = InOutMotionMatchingState.CurrentSearchResult;
 
 	// Check if we can advance. Includes the case where we can advance but only by switching to a follow up asset.
 	bool bAdvanceToFollowUpAsset = false;
@@ -394,11 +249,22 @@ void UpdateMotionMatchingState(
 	if (!bCanAdvance || (InOutMotionMatchingState.ElapsedPoseJumpTime >= Settings.SearchThrottleTime))
 	{
 		// Build the search context
-		FSearchContext SearchContext;
-		SearchContext.SetSource(InOutMotionMatchingState.CurrentDatabase.Get());
-		SearchContext.QueryValues = InOutMotionMatchingState.ComposedQuery.GetNormalizedValues();
-		SearchContext.WeightsContext = &InOutMotionMatchingState.WeightsContext;
+		FPoseSearchContext SearchContext;
 		SearchContext.DatabaseTagQuery = DatabaseTagQuery;
+		SearchContext.ActiveTagsContainer = ActiveTagsContainer;
+		SearchContext.Trajectory = &Trajectory;
+		SearchContext.OwningComponent = Context.AnimInstanceProxy->GetSkelMeshComponent();
+		SearchContext.BoneContainer = &Context.AnimInstanceProxy->GetRequiredBones();
+
+		
+		SearchContext.CurrentResult = InOutMotionMatchingState.CurrentSearchResult;
+
+		IPoseHistoryProvider* PoseHistoryProvider = Context.GetMessage<IPoseHistoryProvider>();
+		if (PoseHistoryProvider)
+		{
+			SearchContext.History = &PoseHistoryProvider->GetPoseHistory();
+		}
+
 		if (const FPoseSearchIndexAsset* CurrentIndexAsset = InOutMotionMatchingState.GetCurrentSearchIndexAsset())
 		{
 			SearchContext.QueryMirrorRequest =
@@ -407,11 +273,8 @@ void UpdateMotionMatchingState(
 				EPoseSearchBooleanRequest::FalseValue;
 		}
 
-		// Update weight groups
-		InOutMotionMatchingState.WeightsContext.Update(Database);
-
 		// Search the database for the nearest match to the updated query vector
-		FSearchResult SearchResult = Search(SearchContext);
+		FSearchResult SearchResult = Searchable->Search(SearchContext);
 
 		if (SearchResult.IsValid())
 		{
@@ -423,27 +286,24 @@ void UpdateMotionMatchingState(
 			// Otherwise we need to check if the result is a good improvement over the current pose
 			else
 			{
-				// Determine how much the updated query vector deviates from the current pose vector
-				FPoseCost CurrentPoseCost;
-				if (InOutMotionMatchingState.DbPoseIdx != INDEX_NONE)
-				{
-					const FPoseSearchIndexAsset* SearchIndexAsset = &Database->GetSearchIndex()->Assets[InOutMotionMatchingState.SearchIndexAssetIdx];
-					CurrentPoseCost = ComparePoses(InOutMotionMatchingState.DbPoseIdx, SearchContext, SearchIndexAsset->SourceGroupIdx);
-				}
-
 				// Consider the search result better if it is more similar to the query than the current pose we're playing back from the database
-				check(SearchResult.PoseCost.Dissimilarity >= 0.0f);
+				check(SearchResult.PoseCost.GetDissimilarity() >= 0.0f);
 				bool bBetterPose = true;
-				if (CurrentPoseCost.IsValid())
+				if (SearchResult.ContinuityPoseCost.IsValid())
 				{
-					if ((CurrentPoseCost.TotalCost <= SearchResult.PoseCost.TotalCost) || (CurrentPoseCost.Dissimilarity <= SearchResult.PoseCost.Dissimilarity))
+					if ((SearchResult.ContinuityPoseCost.GetTotalCost() <= SearchResult.PoseCost.GetTotalCost()) || 
+						(SearchResult.ContinuityPoseCost.GetDissimilarity() <= SearchResult.PoseCost.GetDissimilarity()))
 					{
 						bBetterPose = false;
 					}
 					else
 					{
-						checkSlow(CurrentPoseCost.Dissimilarity > 0.0f && CurrentPoseCost.Dissimilarity > SearchResult.PoseCost.Dissimilarity);
-						const float RelativeSimilarityGain = -1.0f * (SearchResult.PoseCost.Dissimilarity - CurrentPoseCost.Dissimilarity) / CurrentPoseCost.Dissimilarity;
+						checkSlow(
+							SearchResult.ContinuityPoseCost.GetDissimilarity() > 0.0f && 
+							 SearchResult.ContinuityPoseCost.GetDissimilarity() > SearchResult.PoseCost.GetDissimilarity());
+						const float RelativeSimilarityGain = -1.0f * 
+							(SearchResult.PoseCost.GetDissimilarity() - SearchResult.ContinuityPoseCost.GetDissimilarity()) / 
+							SearchResult.ContinuityPoseCost.GetDissimilarity();
 						bBetterPose = RelativeSimilarityGain >= Settings.MinPercentImprovement / 100.0f;
 					}
 				}
@@ -456,12 +316,20 @@ void UpdateMotionMatchingState(
 					// We need to check in terms of PoseIdx rather than AssetTime because
 					// for blendspaces, AssetTime is not in seconds, but in the normalized range 
 					// [0, 1] so comparing to `PoseJumpThresholdTime` will not make sense		
-					bNearbyPose = FMath::Abs(InOutMotionMatchingState.DbPoseIdx - SearchResult.PoseIdx) * Database->Schema->SamplingInterval < Settings.PoseJumpThresholdTime;
+					bNearbyPose = 
+						FMath::Abs(InOutMotionMatchingState.CurrentSearchResult.PoseIdx - SearchResult.PoseIdx) *
+						SearchResult.Database->Schema->SamplingInterval < Settings.PoseJumpThresholdTime;
 
 					// Handle looping anims when checking for the pose being too close
-					if (!bNearbyPose && Database->IsSourceAssetLooping(StateSearchIndexAsset))
+					if (!bNearbyPose && SearchResult.Database->IsSourceAssetLooping(StateSearchIndexAsset))
 					{
-						bNearbyPose = FMath::Abs(StateSearchIndexAsset->NumPoses - InOutMotionMatchingState.DbPoseIdx - SearchResult.PoseIdx) * Database->Schema->SamplingInterval < Settings.PoseJumpThresholdTime;
+						const float Time =
+							FMath::Abs(
+								StateSearchIndexAsset->NumPoses -
+								InOutMotionMatchingState.CurrentSearchResult.PoseIdx -
+								SearchResult.PoseIdx) *
+							SearchResult.Database->Schema->SamplingInterval;
+						bNearbyPose = Time < Settings.PoseJumpThresholdTime;
 					}
 				}
 
@@ -493,7 +361,7 @@ void UpdateMotionMatchingState(
 
 	// Record debugger details
 #if UE_POSE_SEARCH_TRACE_ENABLED
-	if (InOutMotionMatchingState.DbPoseIdx != INDEX_NONE)
+	if (InOutMotionMatchingState.CurrentSearchResult.IsValid())
 	{
 		float SimLinearVelocity, SimAngularVelocity, AnimLinearVelocity, AnimAngularVelocity;
 
@@ -533,10 +401,16 @@ void UpdateMotionMatchingState(
 		}
 
 		TArray<bool> DatabaseSequenceFilter;
-		ComputeDatabaseSequenceFilter(Database, DatabaseTagQuery, DatabaseSequenceFilter);
+		ComputeDatabaseSequenceFilter(
+			InOutMotionMatchingState.CurrentSearchResult.Database.Get(), 
+			DatabaseTagQuery, 
+			DatabaseSequenceFilter);
 
 		TArray<bool> DatabaseBlendSpaceFilter;
-		ComputeDatabaseBlendSpaceFilter(Database, DatabaseTagQuery, DatabaseBlendSpaceFilter);
+		ComputeDatabaseBlendSpaceFilter(
+			InOutMotionMatchingState.CurrentSearchResult.Database.Get(),
+			DatabaseTagQuery, 
+			DatabaseBlendSpaceFilter);
 
 		FTraceMotionMatchingState TraceState;
 		if (EnumHasAnyFlags(InOutMotionMatchingState.Flags, EMotionMatchingFlags::JumpedToFollowUp))
@@ -546,11 +420,11 @@ void UpdateMotionMatchingState(
 
 		TraceState.ElapsedPoseJumpTime = InOutMotionMatchingState.ElapsedPoseJumpTime;
 		// @TODO: Change this to only be the previous query, not persistently updated (i.e. if throttled)?
-		TraceState.QueryVector = InOutMotionMatchingState.ComposedQuery.GetValues();
-		TraceState.QueryVectorNormalized = InOutMotionMatchingState.ComposedQuery.GetNormalizedValues();
-		TraceState.DbPoseIdx = InOutMotionMatchingState.DbPoseIdx;
-		TraceState.DatabaseId = FObjectTrace::GetObjectId(Database);
-		TraceState.ContinuingPoseIdx = CurrentPoseIdx;
+		TraceState.QueryVector = InOutMotionMatchingState.CurrentSearchResult.ComposedQuery.GetValues();
+		TraceState.QueryVectorNormalized = InOutMotionMatchingState.CurrentSearchResult.ComposedQuery.GetNormalizedValues();
+		TraceState.DbPoseIdx = InOutMotionMatchingState.CurrentSearchResult.PoseIdx;
+		TraceState.DatabaseId = FObjectTrace::GetObjectId(InOutMotionMatchingState.CurrentSearchResult.Database.Get());
+		TraceState.ContinuingPoseIdx = LastResult.PoseIdx;
 
 		TraceState.AssetPlayerTime = InOutMotionMatchingState.AssetPlayerTime;
 		TraceState.DeltaTime = DeltaTime;
@@ -567,14 +441,12 @@ void UpdateMotionMatchingState(
 
 const FPoseSearchIndexAsset* FMotionMatchingState::GetCurrentSearchIndexAsset() const
 {
-	if (!CurrentDatabase.Get() || 
-		!CurrentDatabase->GetSearchIndex() ||
-		!CurrentDatabase->GetSearchIndex()->Assets.IsValidIndex(SearchIndexAssetIdx))
+	if (CurrentSearchResult.IsValid())
 	{
-		return nullptr;
+		return CurrentSearchResult.SearchIndexAsset;
 	}
 
-	return &CurrentDatabase->GetSearchIndex()->Assets[SearchIndexAssetIdx];
+	return nullptr;
 }
 
 float FMotionMatchingState::ComputeJumpBlendTime(
@@ -597,32 +469,12 @@ float FMotionMatchingState::ComputeJumpBlendTime(
 	return JumpBlendTime;
 }
 
-bool FMotionMatchingState::IsCompatibleDatabase(const UPoseSearchDatabase* Database) const
+EPoseSearchPostSearchStatus UPoseSearchPostProcessor_Bias::PostProcess_Implementation(FPoseSearchCost& InOutCost) const
 {
-	if (Database != CurrentDatabase)
-	{
-		return false;
-	}
+	InOutCost.SetDissimilarity(Multiplier * InOutCost.GetDissimilarity());
+	InOutCost.SetCostAddend(Addend + InOutCost.GetCostAddend());
 
-#if WITH_EDITOR
-	if (Database->GetSearchIndexHash() != CurrentSearchIndexHash)
-	{
-		return false;
-	}
-#endif
-
-	return true;
-}
-
-#if WITH_EDITOR
-bool FMotionMatchingState::HasSearchIndexChanged() const
-{
-	const bool bIsConsistent = 
-		!CurrentDatabase.IsValid() ||
-		CurrentDatabase->IsDerivedDataBuildPending() ||
-		CurrentDatabase->GetSearchIndexHash() != CurrentSearchIndexHash;
-	return bIsConsistent;
+	return EPoseSearchPostSearchStatus::Continue;
 }
-#endif
 
 #undef LOCTEXT_NAMESPACE
\ No newline at end of file
diff --git a/Engine/Plugins/Experimental/Animation/PoseSearch/Source/Runtime/Public/PoseSearch/AnimNode_MotionMatching.h b/Engine/Plugins/Experimental/Animation/PoseSearch/Source/Runtime/Public/PoseSearch/AnimNode_MotionMatching.h
--- a/Engine/Plugins/Experimental/Animation/PoseSearch/Source/Runtime/Public/PoseSearch/AnimNode_MotionMatching.h
+++ b/Engine/Plugins/Experimental/Animation/PoseSearch/Source/Runtime/Public/PoseSearch/AnimNode_MotionMatching.h
@@ -24,7 +24,7 @@ public:
 
 	// Collection of animations for motion matching
 	UPROPERTY(EditAnywhere, BlueprintReadWrite, Category=Settings, meta=(PinShownByDefault))
-	TObjectPtr<const UPoseSearchDatabase> Database = nullptr;
+	TObjectPtr<const UPoseSearchSearchableAsset> Searchable = nullptr;
 
 
 	UPROPERTY(EditAnywhere, BlueprintReadWrite, Category = Settings, meta = (PinHiddenByDefault))
@@ -34,6 +34,9 @@ public:
 	UPROPERTY(EditAnywhere, BlueprintReadWrite, Category = Settings, meta = (PinShownByDefault))
 	FGameplayTagQuery DatabaseTagQuery;
 
+	UPROPERTY(EditAnywhere, BlueprintReadWrite, Category = Settings, meta = (PinShownByDefault))
+	FGameplayTagContainer ActiveTagsContainer;
+
 	// Motion trajectory samples for pose search queries
 	UPROPERTY(EditAnywhere, BlueprintReadWrite, Category=Settings, meta=(PinShownByDefault))
 	FTrajectorySampleRange Trajectory;
diff --git a/Engine/Plugins/Experimental/Animation/PoseSearch/Source/Runtime/Public/PoseSearch/PoseSearch.h b/Engine/Plugins/Experimental/Animation/PoseSearch/Source/Runtime/Public/PoseSearch/PoseSearch.h
--- a/Engine/Plugins/Experimental/Animation/PoseSearch/Source/Runtime/Public/PoseSearch/PoseSearch.h
+++ b/Engine/Plugins/Experimental/Animation/PoseSearch/Source/Runtime/Public/PoseSearch/PoseSearch.h
@@ -77,8 +77,8 @@
 	Invalid = Num UMETA(Hidden)
 };
 
-UENUM()
-enum class EPoseSearchBooleanRequest : int32
+UENUM(BlueprintType)
+enum class EPoseSearchBooleanRequest : uint8
 {
 	FalseValue,
 	TrueValue,
@@ -419,7 +419,10 @@
 	virtual void GenerateDDCKey(FBlake3& InOutKeyHasher) const PURE_VIRTUAL(UPoseSearchFeatureChannel::GenerateDDCKey, );
 
 	// Called at runtime to add this channel's data to the query pose vector
-	virtual bool BuildQuery(UE::PoseSearch::FQueryBuildingContext& Context) const PURE_VIRTUAL(UPoseSearchFeatureChannel::BuildQuery, return false;);
+	virtual bool BuildQuery(
+		FPoseSearchContext& SearchContext,
+		FPoseSearchFeatureVectorBuilder& InOutQuery) const
+		PURE_VIRTUAL(UPoseSearchFeatureChannel::BuildQuery, return false;);
 
 	// Draw this channel's data for the given pose vector
 	virtual void DebugDraw(const UE::PoseSearch::FDebugDrawParams& DrawParams, const UE::PoseSearch::FFeatureVectorReader& Reader) const PURE_VIRTUAL(UPoseSearchFeatureChannel::DebugDraw, );
@@ -556,6 +559,9 @@
 private:
 	void Finalize();
 	void ResolveBoneReferences();
+
+public:
+	bool BuildQuery(FPoseSearchContext& SearchContext, FPoseSearchFeatureVectorBuilder& InOutQuery) const;
 };
 
 
@@ -931,9 +937,270 @@
 	FPoseSearchWeightParams Weights;
 };
 
+USTRUCT(BlueprintType)
+struct FPoseSearchCost
+{
+	GENERATED_BODY()
+
+public:
+	FPoseSearchCost() = default;
+	FPoseSearchCost(float InDissimilarity, float InCostAddend)
+		: Dissimilarity(InDissimilarity)
+		, CostAddend(InCostAddend)
+		, TotalCost(Dissimilarity + CostAddend)
+	{}
+
+	bool IsValid() const { return TotalCost != MAX_flt; }
+
+	void Set(float InDissimilarity, float InCostAddend)
+	{
+		Dissimilarity = InDissimilarity;
+		CostAddend = InCostAddend;
+		TotalCost = Dissimilarity + CostAddend;
+	}
+
+	float GetDissimilarity() const { return Dissimilarity; }
+
+	void SetDissimilarity(float InDissimilarity)
+	{
+		Dissimilarity = InDissimilarity;
+		TotalCost = Dissimilarity + CostAddend;
+	}
+
+	float GetCostAddend() const { return CostAddend; }
+
+	void SetCostAddend(float InCostAddend)
+	{
+		CostAddend = InCostAddend;
+		TotalCost = Dissimilarity + CostAddend;
+	}
+
+	float GetTotalCost() const { return TotalCost; }
+
+	bool operator<(const FPoseSearchCost& Other) const { return TotalCost < Other.TotalCost; }
+
+	void Reset()
+	{
+		Dissimilarity = MAX_flt;
+		CostAddend = 0.0f;
+		TotalCost = MAX_flt;
+	}
+
+protected:
+
+	UPROPERTY()
+	float Dissimilarity = MAX_flt;
+
+	UPROPERTY()
+	float CostAddend = 0.0f;
+
+	UPROPERTY()
+	float TotalCost = MAX_flt;
+};
+
+/**
+* Helper object for writing features into a float buffer according to a feature vector layout.
+* Keeps track of which features are present, allowing the feature vector to be built up piecemeal.
+* FFeatureVectorBuilder is used to build search queries at runtime and for adding samples during search index construction.
+*/
+USTRUCT(BlueprintType, Category = "Animation|Pose Search")
+struct POSESEARCH_API FPoseSearchFeatureVectorBuilder
+{
+	GENERATED_BODY()
+
+public:
+	void Init(const UPoseSearchSchema* Schema);
+	void Reset();
+	void ResetFeatures();
+
+	const UPoseSearchSchema* GetSchema() const { return Schema.Get(); }
+
+	TArrayView<const float> GetValues() const { return Values; }
+	TArrayView<const float> GetNormalizedValues() const { return ValuesNormalized; }
+
+	void SetTransform(FPoseSearchFeatureDesc Feature, const FTransform& Transform);
+	void SetTransformVelocity(FPoseSearchFeatureDesc Feature, const FTransform& Transform, const FTransform& PrevTransform, float DeltaTime);
+	void SetTransformVelocity(FPoseSearchFeatureDesc Feature, const FTransform& NextTransform, const FTransform& Transform, const FTransform& PrevTransform, float DeltaTime);
+	void SetPosition(FPoseSearchFeatureDesc Feature, const FVector& Translation);
+	void SetRotation(FPoseSearchFeatureDesc Feature, const FQuat& Rotation);
+	void SetLinearVelocity(FPoseSearchFeatureDesc Feature, const FTransform& Transform, const FTransform& PrevTransform, float DeltaTime);
+	void SetLinearVelocity(FPoseSearchFeatureDesc Feature, const FTransform& NextTransform, const FTransform& Transform, const FTransform& PrevTransform, float DeltaTime);
+	void SetAngularVelocity(FPoseSearchFeatureDesc Feature, const FTransform& Transform, const FTransform& PrevTransform, float DeltaTime);
+	void SetAngularVelocity(FPoseSearchFeatureDesc Feature, const FTransform& NextTransform, const FTransform& Transform, const FTransform& PrevTransform, float DeltaTime);
+	void SetVector(FPoseSearchFeatureDesc Feature, const FVector& Vector);
+
+	void CopyFromSearchIndex(const FPoseSearchIndex& SearchIndex, int32 PoseIdx);
+	void CopyFeature(const FPoseSearchFeatureVectorBuilder& OtherBuilder, int32 FeatureIdx);
+
+	void MergeReplace(const FPoseSearchFeatureVectorBuilder& OtherBuilder);
+
+	bool IsInitialized() const;
+	bool IsInitializedForSchema(const UPoseSearchSchema* Schema) const;
+	bool IsComplete() const;
+	bool IsCompatible(const FPoseSearchFeatureVectorBuilder& OtherBuilder) const;
+
+	const TBitArray<>& GetFeaturesAdded() const;
+
+	void Normalize(const FPoseSearchIndex& ForSearchIndex);
+
+private:
+	UPROPERTY(Transient)
+	TWeakObjectPtr<const UPoseSearchSchema> Schema = nullptr;
+
+	TArray<float> Values;
+	TArray<float> ValuesNormalized;
+	TBitArray<> FeaturesAdded;
+	int32 NumFeaturesAdded = 0;
+};
+
+
+namespace UE::PoseSearch
+{
+
+enum class EDebugDrawFlags : uint32
+{
+	None			    = 0,
+
+	// Draw the entire search index as a point cloud
+	DrawSearchIndex     = 1 << 0,
+
+	// Draw pose features for each pose vector
+	IncludePose         = 1 << 1,
+
+	// Draw trajectory features for each pose vector
+	IncludeTrajectory   = 1 << 2,
+
+	// Draw all pose vector features
+	IncludeAllFeatures  = IncludePose | IncludeTrajectory,
+
+	/**
+	 * Keep rendered data until the next call to FlushPersistentDebugLines().
+	 * Combine with DrawSearchIndex to draw the search index only once.
+	 */
+	Persistent = 1 << 3,
+	
+	// Label samples with their indices
+	DrawSampleLabels = 1 << 4,
+
+	// Fade colors
+	DrawSamplesWithColorGradient = 1 << 5,
+
+	// Draw Bone Names
+	DrawBoneNames = 1 << 6,
+
+	// Draws simpler shapes to improve performance
+	DrawFast = 1 << 7,
+};
+ENUM_CLASS_FLAGS(EDebugDrawFlags);
+
+struct POSESEARCH_API FDebugDrawParams
+{
+	const UWorld* World = nullptr;
+	const UPoseSearchDatabase* Database = nullptr;
+	const UPoseSearchSequenceMetaData* SequenceMetaData = nullptr;
+	EDebugDrawFlags Flags = EDebugDrawFlags::DrawBoneNames;
+	uint32 ChannelMask = (uint32)-1;
+
+	float DefaultLifeTime = 5.0f;
+	float PointSize = 1.0f;
+
+	FTransform RootTransform = FTransform::Identity;
+
+	// If set, draw the corresponding pose from the search index
+	int32 PoseIdx = INDEX_NONE;
+
+	// If set, draw using this uniform color instead of feature-based coloring
+	const FLinearColor* Color = nullptr;
+
+	// If set, interpret the buffer as a pose vector and draw it
+	TArrayView<const float> PoseVector;
+
+	// Optional prefix for sample labels
+	FStringView LabelPrefix;
+
+	bool CanDraw() const;
+	const FPoseSearchIndex* GetSearchIndex() const;
+	const UPoseSearchSchema* GetSchema() const;
+};
+
+class IPoseHistoryProvider : public UE::Anim::IGraphMessage
+{
+	DECLARE_ANIMGRAPH_MESSAGE(IPoseHistoryProvider);
+
+public:
+
+	virtual const FPoseHistory& GetPoseHistory() const = 0;
+	virtual FPoseHistory& GetPoseHistory() = 0;
+};
+
+
+struct FSearchResult
+{
+	FPoseSearchCost PoseCost;
+	int32 PoseIdx = INDEX_NONE;
+	const FPoseSearchIndexAsset* SearchIndexAsset = nullptr;
+	TWeakObjectPtr<const UPoseSearchDatabase> Database = nullptr;
+	TWeakObjectPtr<const UAnimSequenceBase> Sequence = nullptr;
+	FPoseSearchFeatureVectorBuilder ComposedQuery;
+
+	// cost of the current pose with the query from database in the result, if possible
+	FPoseSearchCost ContinuityPoseCost; 
+
+	float AssetTime = 0.0f;
+
+#if WITH_EDITOR
+	FIoHash SearchIndexHash = FIoHash::Zero;
+#endif // WITH_EDITOR
+
+	// Attempts to set the internal state to match the provided asset time including updating the internal DbPoseIdx. 
+	// If the provided asset time is out of bounds for the currently playing asset then this function will reset the 
+	// state back to the default state.
+	void Update(float NewAssetTime);
+
+	bool IsValid() const;
+
+	void Reset();
+};
+
+/**
+ * Cost details for pose analysis in the rewind debugger
+ */
+struct FPoseCostDetails
+{
+	FPoseSearchCost PoseCost;
+
+	// Contribution from ModifyCost anim notify
+	float NotifyCostAddend = 0.0f;
+
+	// Contribution from mirroring cost
+	float MirrorMismatchCostAddend = 0.0f;
+
+	// Cost breakdown per channel (e.g. pose cost, time-based trajectory cost, distance-based trajectory cost, etc.)
+	TArray<float> ChannelCosts;
+
+	// Difference vector computed as W*((P-Q)^2) without the cost modifier applied
+	// Where P is the pose vector, Q is the query vector, W is the weights vector, and multiplication/exponentiation are element-wise operations
+	TArray<float> CostVector;
+};
+
+}
+
+
 /** A data asset for indexing a collection of animation sequences. */
+UCLASS(Abstract, BlueprintType, Experimental)
+class POSESEARCH_API UPoseSearchSearchableAsset : public UDataAsset
+{
+	GENERATED_BODY()
+public:
+
+	virtual UE::PoseSearch::FSearchResult Search(FPoseSearchContext& SearchContext) const 
+		PURE_VIRTUAL(UPoseSearchSearchableAsset::Search, return UE::PoseSearch::FSearchResult(););
+};
+
+
+/** A data asset for indexing a collection of animation sequences. */
 UCLASS(BlueprintType, Category = "Animation|Pose Search", Experimental, meta = (DisplayName = "Motion Database"))
-class POSESEARCH_API UPoseSearchDatabase : public UDataAsset
+class POSESEARCH_API UPoseSearchDatabase : public UPoseSearchSearchableAsset
 {
 	GENERATED_BODY()
 public:
@@ -1059,6 +1326,35 @@
 #endif // WITH_EDITOR
 
 	bool IsDerivedDataValid();
+
+public:
+	virtual UE::PoseSearch::FSearchResult Search(FPoseSearchContext& SearchContext) const override;
+
+	void BuildQuery(FPoseSearchContext& SearchContext, FPoseSearchFeatureVectorBuilder& OutQuery) const;
+
+	FPoseSearchCost ComparePoses(
+		FPoseSearchContext& SearchContext,
+		FPoseSearchWeightsContext& WeightsContext,
+		int32 PoseIdx,
+		int32 GroupIdx,
+		const TArrayView<const float>& QueryValues) const;
+
+	FPoseSearchCost ComparePoses(
+		FPoseSearchContext& SearchContext,
+		FPoseSearchWeightsContext& WeightsContext,
+		int32 PoseIdx,
+		const TArrayView<const float>& QueryValues,
+		UE::PoseSearch::FPoseCostDetails& OutPoseCostDetails) const;
+
+	void ComputePoseCostAddends(
+		int32 PoseIdx,
+		FPoseSearchContext& SearchContext,
+		float& OutNotifyAddend,
+		float& OutMirrorMismatchAddend) const;
+
+protected:
+	UE::PoseSearch::FSearchResult SearchPCAKDTree(FPoseSearchContext& SearchContext) const;
+	UE::PoseSearch::FSearchResult SearchBruteForce(FPoseSearchContext& SearchContext) const;
 };
 
 
@@ -1087,8 +1383,14 @@
 	bool IsValidForIndexing() const;
 	bool IsValidForSearch() const;
 
+	UE::PoseSearch::FSearchResult Search(FPoseSearchContext& SearchContext) const;
+
+protected:
+	FPoseSearchCost ComparePoses(int32 PoseIdx, const TArrayView<const float>& QueryValues) const;
+
 public: // UObject
 	virtual void PreSave(FObjectPreSaveContext ObjectSaveContext) override;
+
 };
 
 
@@ -1120,70 +1422,6 @@
 	TArrayView<const float> Values;
 };
 
-} // namespace UE::PoseSearch
-
-
-/** 
-* Helper object for writing features into a float buffer according to a feature vector layout.
-* Keeps track of which features are present, allowing the feature vector to be built up piecemeal.
-* FFeatureVectorBuilder is used to build search queries at runtime and for adding samples during search index construction.
-*/
-USTRUCT(BlueprintType, Category = "Animation|Pose Search")
-struct POSESEARCH_API FPoseSearchFeatureVectorBuilder
-{
-	GENERATED_BODY()
-public:
-	void Init(const UPoseSearchSchema* Schema);
-	void Reset();
-	void ResetFeatures();
-
-	const UPoseSearchSchema* GetSchema() const { return Schema.Get(); }
-
-	TArrayView<const float> GetValues() const { return Values; }
-	TArrayView<const float> GetNormalizedValues() const { return ValuesNormalized; }
-
-	void SetTransform(FPoseSearchFeatureDesc Feature, const FTransform& Transform);
-	void SetTransformVelocity(FPoseSearchFeatureDesc Feature, const FTransform& Transform, const FTransform& PrevTransform, float DeltaTime);
-	void SetTransformVelocity(FPoseSearchFeatureDesc Feature, const FTransform& NextTransform, const FTransform& Transform, const FTransform& PrevTransform, float DeltaTime);
-	void SetPosition(FPoseSearchFeatureDesc Feature, const FVector& Translation);
-	void SetRotation(FPoseSearchFeatureDesc Feature, const FQuat& Rotation);
-	void SetLinearVelocity(FPoseSearchFeatureDesc Feature, const FTransform& Transform, const FTransform& PrevTransform, float DeltaTime);
-	void SetLinearVelocity(FPoseSearchFeatureDesc Feature, const FTransform& NextTransform, const FTransform& Transform, const FTransform& PrevTransform, float DeltaTime);
-	void SetAngularVelocity(FPoseSearchFeatureDesc Feature, const FTransform& Transform, const FTransform& PrevTransform, float DeltaTime);
-	void SetAngularVelocity(FPoseSearchFeatureDesc Feature, const FTransform& NextTransform, const FTransform& Transform, const FTransform& PrevTransform, float DeltaTime);
-	void SetVector(FPoseSearchFeatureDesc Feature, const FVector& Vector);
-
-	void CopyFromSearchIndex(const FPoseSearchIndex& SearchIndex, int32 PoseIdx);
-	void CopyFeature(const FPoseSearchFeatureVectorBuilder& OtherBuilder, int32 FeatureIdx);
-
-	void MergeReplace(const FPoseSearchFeatureVectorBuilder& OtherBuilder);
-
-	bool IsInitialized() const;
-	bool IsInitializedForSchema(const UPoseSearchSchema* Schema) const;
-	bool IsComplete() const;
-	bool IsCompatible(const FPoseSearchFeatureVectorBuilder& OtherBuilder) const;
-
-	const TBitArray<>& GetFeaturesAdded() const;
-
-	void Normalize(const FPoseSearchIndex& ForSearchIndex);
-
-private:
-	UPROPERTY(Transient)
-	TWeakObjectPtr<const UPoseSearchSchema> Schema = nullptr;
-
-	TArray<float> Values;
-	TArray<float> ValuesNormalized;
-	TBitArray<> FeaturesAdded;
-	int32 NumFeaturesAdded = 0;
-};
-
-
-
-//////////////////////////////////////////////////////////////////////////
-// Pose history
-
-namespace UE::PoseSearch {
-
 /**
 * Records poses over time in a ring buffer.
 * FFeatureVectorBuilder uses this to sample from the present or past poses according to the search schema.
@@ -1203,9 +1441,9 @@
 	bool TrySamplePose(float SecondsAgo, const FReferenceSkeleton& RefSkeleton, const TArray<FBoneIndexType>& RequiredBones);
 
 	bool Update(
-		float SecondsElapsed, 
-		const FPoseContext& PoseContext, 
-		FTransform ComponentTransform, 
+		float SecondsElapsed,
+		const FPoseContext& PoseContext,
+		FTransform ComponentTransform,
 		FText* OutError,
 		ERootUpdateMode UpdateMode = ERootUpdateMode::RootMotionDelta);
 
@@ -1242,87 +1480,35 @@
 	float TimeHorizon = 0.0f;
 };
 
-class IPoseHistoryProvider : public UE::Anim::IGraphMessage
+} // namespace UE::PoseSearch
+
+
+USTRUCT(BlueprintType)
+struct POSESEARCH_API FPoseSearchContext
 {
-	DECLARE_ANIMGRAPH_MESSAGE(IPoseHistoryProvider);
+	GENERATED_BODY()
 
 public:
-
-	virtual const FPoseHistory& GetPoseHistory() const = 0;
-	virtual FPoseHistory& GetPoseHistory() = 0;
+	EPoseSearchBooleanRequest QueryMirrorRequest = EPoseSearchBooleanRequest::Indifferent;
+	const FGameplayTagQuery* DatabaseTagQuery = nullptr;
+	UE::PoseSearch::FDebugDrawParams DebugDrawParams;
+	UE::PoseSearch::FPoseHistory* History = nullptr;
+	const FTrajectorySampleRange* Trajectory = nullptr;
+	TObjectPtr<const USkeletalMeshComponent> OwningComponent = nullptr;
+	UE::PoseSearch::FSearchResult CurrentResult;
+	const FBoneContainer* BoneContainer = nullptr;
+	const FGameplayTagContainer* ActiveTagsContainer = nullptr;
 };
 
 
+//////////////////////////////////////////////////////////////////////////
+// Pose history
+
+namespace UE::PoseSearch {
 
 //////////////////////////////////////////////////////////////////////////
-// Debug visualization
+// Main PoseSearch API
 
-enum class EDebugDrawFlags : uint32
-{
-	None = 0,
 
-	// Draw the entire search index as a point cloud
-	DrawSearchIndex = 1 << 0,
-
-	/**
-	 * Keep rendered data until the next call to FlushPersistentDebugLines().
-	 * Combine with DrawSearchIndex to draw the search index only once.
-	 */
-	Persistent = 1 << 1,
-	
-	// Label samples with their indices
-	DrawSampleLabels = 1 << 2,
-
-	// Fade colors
-	DrawSamplesWithColorGradient = 1 << 3,
-
-	// Label bone names
-	DrawBoneNames = 1 << 4,
-
-	// Draws simpler shapes to improve performance
-	DrawFast = 1 << 5,
-};
-ENUM_CLASS_FLAGS(EDebugDrawFlags);
-
-struct POSESEARCH_API FDebugDrawParams
-{
-	const UWorld* World = nullptr;
-	const UPoseSearchDatabase* Database = nullptr;
-	const UPoseSearchSequenceMetaData* SequenceMetaData = nullptr;
-	EDebugDrawFlags Flags = EDebugDrawFlags::DrawBoneNames;
-	uint32 ChannelMask = (uint32)-1;
-
-	float DefaultLifeTime = 5.0f;
-	float PointSize = 1.0f;
-
-	FTransform RootTransform = FTransform::Identity;
-
-	// If set, draw the corresponding pose from the search index
-	int32 PoseIdx = INDEX_NONE;
-
-	// If set, draw using this uniform color instead of feature-based coloring
-	const FLinearColor* Color = nullptr;
-
-	// If set, interpret the buffer as a pose vector and draw it
-	TArrayView<const float> PoseVector;
-
-	// Optional prefix for sample labels
-	FStringView LabelPrefix;
-
-	bool CanDraw() const;
-	const FPoseSearchIndex* GetSearchIndex() const;
-	const UPoseSearchSchema* GetSchema() const;
-};
-
-struct FPoseCost
-{
-	float Dissimilarity = MAX_flt;
-	float CostAddend = 0.0f;
-	float TotalCost = MAX_flt;
-	bool operator<(const FPoseCost& Other) const { return TotalCost < Other.TotalCost; }
-	bool IsValid() const { return TotalCost != MAX_flt; }
-
-};
-
 /**
 * Visualize pose search debug information
@@ -1332,16 +1518,12 @@
 POSESEARCH_API void Draw(const FDebugDrawParams& DrawParams);
 
 
-
-//////////////////////////////////////////////////////////////////////////
-// Index building
-
 /**
 * Creates a pose search index for an animation sequence
-*
+* 
 * @param Sequence			The input sequence create a search index for
 * @param SequenceMetaData	The input sequence indexing info and output search index
-*
+* 
 * @return Whether the index was built successfully
 */
 POSESEARCH_API bool BuildIndex(const UAnimSequence* Sequence, UPoseSearchSequenceMetaData* SequenceMetaData);
@@ -1349,124 +1531,63 @@
 
 /**
 * Creates a pose search index for a collection of animations
-*
+* 
 * @param Database	The input collection of animations and output search index
-*
+* 
 * @return Whether the index was built successfully
 */
 POSESEARCH_API bool BuildIndex(UPoseSearchDatabase* Database, FPoseSearchIndex& OutSearchIndex);
 
-
-//////////////////////////////////////////////////////////////////////////
-// Query building
-
-struct POSESEARCH_API FQueryBuildingContext
-{
-	FQueryBuildingContext(FPoseSearchFeatureVectorBuilder& InQuery) : Query(InQuery) {}
-
-	FPoseSearchFeatureVectorBuilder& Query;
-
-	const UPoseSearchSchema* Schema = nullptr;
-	FPoseHistory* History = nullptr;
-	const FTrajectorySampleRange* Trajectory = nullptr;
-
-	bool IsInitialized () const;
-};
-
-bool POSESEARCH_API BuildQuery(FQueryBuildingContext& QueryBuildingContext);
-
-
-//////////////////////////////////////////////////////////////////////////
-// Search
-
-struct POSESEARCH_API FSearchResult
-{
-	FPoseCost PoseCost;
-	int32 PoseIdx = INDEX_NONE;
-	const FPoseSearchIndexAsset* SearchIndexAsset = nullptr;
-	float AssetTime = 0.0f;
-
-	bool IsValid() const { return PoseIdx != INDEX_NONE; }
-};
-
-struct POSESEARCH_API FSearchContext
-{
-	TArrayView<const float> QueryValues;
-	EPoseSearchBooleanRequest QueryMirrorRequest = EPoseSearchBooleanRequest::Indifferent;
-	const FPoseSearchWeightsContext* WeightsContext = nullptr;
-	const FGameplayTagQuery* DatabaseTagQuery = nullptr;
-	FDebugDrawParams DebugDrawParams;
-
-	void SetSource(const UPoseSearchDatabase* InSourceDatabase);
-	void SetSource(const UAnimSequenceBase* InSourceSequence);
-	const FPoseSearchIndex* GetSearchIndex() const;
-	float GetMirrorMismatchCost() const;
-	const UPoseSearchDatabase* GetSourceDatabase() { return SourceDatabase; }
-
-private:
-	const UPoseSearchDatabase* SourceDatabase = nullptr;
-	const UAnimSequenceBase* SourceSequence = nullptr;
-	const FPoseSearchIndex* SearchIndex = nullptr;
-	float MirrorMismatchCost = 0.0f;
-};
-
-
-/**
-* Performs a pose search on a UPoseSearchDatabase.
-*
-* @param SearchContext	Structure containing search parameters
-* 
-* @return The pose in the database that most closely matches the Query.
-*/
-POSESEARCH_API FSearchResult Search(FSearchContext& SearchContext);
-
-
-//////////////////////////////////////////////////////////////////////////
-// Pose comparison
-
-/**
-* Evaluate pose comparison metric between a pose in the search index and an input query
-*
-* @param PoseIdx			The index of the pose in the search index to compare to the query
-* @param SearchContext		Structure containing search parameters
-* @param GroupIdx			Indicates the group for this pose's source asset. Specify INDEX_NONE to lookup the pose's group.
-*
-* @return Dissimilarity between the two poses
-*/
-
-POSESEARCH_API FPoseCost ComparePoses(int32 PoseIdx, FSearchContext& SearchContext, int32 GroupIdx);
-
-/**
- * Cost details for pose analysis in the rewind debugger
- */
-struct FPoseCostDetails
-{
-	FPoseCost PoseCost;
-
-	// Contribution from ModifyCost anim notify
-	float NotifyCostAddend = 0.0f;
-
-	// Contribution from mirroring cost
-	float MirrorMismatchCostAddend = 0.0f;
-
-	// Cost breakdown per channel (e.g. pose cost, time-based trajectory cost, distance-based trajectory cost, etc.)
-	TArray<float> ChannelCosts;
-
-	// Difference vector computed as W*((P-Q)^2) without the cost modifier applied
-	// Where P is the pose vector, Q is the query vector, W is the weights vector, and multiplication/exponentiation are element-wise operations
-	TArray<float> CostVector;
-};
-
-/**
-* Evaluate pose comparison metric between a pose in the search index and an input query with cost details
-*
-* @param SearchIndex		The search index containing the pose to compare to the query
-* @param PoseIdx			The index of the pose in the search index to compare to the query
-* @param SearchContext		Structure containing search parameters
-* &param OutPoseCostDetails	Cost details for analysis in the debugger
-*
-* @return Dissimilarity between the two poses
-*/
-POSESEARCH_API FPoseCost ComparePoses(int32 PoseIdx, FSearchContext& SearchContext, FPoseCostDetails& OutPoseCostDetails);
-
 } // namespace UE::PoseSearch
+
+
+UENUM(BlueprintType)
+enum class EPoseSearchPostSearchStatus : uint8
+{
+	// Continue looking for results 
+	Continue,
+
+	// Halt and return the best result
+	Stop
+};
+
+
+UCLASS(Abstract, EditInlineNew, Blueprintable)
+class POSESEARCH_API UPoseSearchPostProcessor : public UDataAsset
+{
+	GENERATED_BODY()
+
+public:
+	UFUNCTION(BlueprintNativeEvent, BlueprintPure, Category=Settings)
+	EPoseSearchPostSearchStatus PostProcess(FPoseSearchCost& InOutCost) const;
+};
+
+USTRUCT(BlueprintType, Category = "Animation|Pose Search")
+struct POSESEARCH_API FPoseSearchDatabaseSetEntry
+{
+	GENERATED_BODY()
+
+public:
+	UPROPERTY(EditAnywhere, Category = Settings)
+	TObjectPtr<UPoseSearchSearchableAsset> Searchable;
+
+	UPROPERTY(EditAnywhere, Category = Settings)
+	FGameplayTag Tag;
+
+	UPROPERTY(EditAnywhere, Instanced, Category = Settings)
+	TObjectPtr<UPoseSearchPostProcessor> PostProcessor;
+};
+
+/** A data asset which holds a collection searchable assets. */
+UCLASS(BlueprintType, Category = "Animation|Pose Search", Experimental)
+class POSESEARCH_API UPoseSearchDatabaseSet : public UPoseSearchSearchableAsset
+{
+	GENERATED_BODY()
+
+public:
+	UPROPERTY(EditAnywhere, Category = Settings)
+	TArray<FPoseSearchDatabaseSetEntry> AssetsToSearch;
+
+public:
+	virtual UE::PoseSearch::FSearchResult Search(FPoseSearchContext& SearchContext) const override;
+};
diff --git a/Engine/Plugins/Experimental/Animation/PoseSearch/Source/Runtime/Public/PoseSearch/PoseSearchFeatureChannels.h b/Engine/Plugins/Experimental/Animation/PoseSearch/Source/Runtime/Public/PoseSearch/PoseSearchFeatureChannels.h
--- a/Engine/Plugins/Experimental/Animation/PoseSearch/Source/Runtime/Public/PoseSearch/PoseSearchFeatureChannels.h
+++ b/Engine/Plugins/Experimental/Animation/PoseSearch/Source/Runtime/Public/PoseSearch/PoseSearchFeatureChannels.h
@@ -46,7 +46,9 @@ public:
 	virtual FFloatRange GetHorizonRange(EPoseSearchFeatureDomain Domain) const override;
 	virtual TArrayView<const float> GetSampleOffsets () const override { return SampleTimes; }
 	virtual void GenerateDDCKey(FBlake3& InOutKeyHasher) const override;
-	virtual bool BuildQuery(UE::PoseSearch::FQueryBuildingContext& Context) const override;
+	virtual bool BuildQuery(
+		FPoseSearchContext& SearchContext,
+		FPoseSearchFeatureVectorBuilder& InOutQuery) const override;
 	virtual void DebugDraw(const UE::PoseSearch::FDebugDrawParams& DrawParams, const UE::PoseSearch::FFeatureVectorReader& Reader) const override;
 
 protected:
@@ -90,7 +92,9 @@ public:
 	virtual FFloatRange GetHorizonRange(EPoseSearchFeatureDomain Domain) const override;
 	virtual TArrayView<const float> GetSampleOffsets () const override { return SampleOffsets; }
 	virtual void GenerateDDCKey(FBlake3& InOutKeyHasher) const override;
-	virtual bool BuildQuery(UE::PoseSearch::FQueryBuildingContext& Context) const override;
+	virtual bool BuildQuery(
+		FPoseSearchContext& SearchContext,
+		FPoseSearchFeatureVectorBuilder& InOutQuery) const override;
 	virtual void DebugDraw(const UE::PoseSearch::FDebugDrawParams& DrawParams, const UE::PoseSearch::FFeatureVectorReader& Reader) const override;
 
 protected:
diff --git a/Engine/Plugins/Experimental/Animation/PoseSearch/Source/Runtime/Public/PoseSearch/PoseSearchLibrary.h b/Engine/Plugins/Experimental/Animation/PoseSearch/Source/Runtime/Public/PoseSearch/PoseSearchLibrary.h
--- a/Engine/Plugins/Experimental/Animation/PoseSearch/Source/Runtime/Public/PoseSearch/PoseSearchLibrary.h
+++ b/Engine/Plugins/Experimental/Animation/PoseSearch/Source/Runtime/Public/PoseSearch/PoseSearchLibrary.h
@@ -53,9 +53,6 @@ struct POSESEARCH_API FMotionMatchingState
 {
 	GENERATED_BODY()
 
-	// Initializes the minimum required motion matching state
-	bool InitNewDatabaseSearch(const UPoseSearchDatabase* Database, FText* OutError);
-
 	// Reset the state to a default state using the current Database
 	void Reset();
 
@@ -76,35 +73,7 @@ struct POSESEARCH_API FMotionMatchingState
 
 	float ComputeJumpBlendTime(const UE::PoseSearch::FSearchResult& Result, const FMotionMatchingSettings& Settings) const;
 
-	bool IsCompatibleDatabase(const UPoseSearchDatabase* Database) const;
-
-#if WITH_EDITOR
-	bool HasSearchIndexChanged() const;
-#endif 
-
-	// The current pose we're playing from the database
-	UPROPERTY(Transient)
-	int32 DbPoseIdx = INDEX_NONE;
-
-	// The current animation we're playing from the database
-	UPROPERTY(Transient)
-	int32 SearchIndexAssetIdx = INDEX_NONE;
-
-	// The current query feature vector used to search the database for pose candidates
-	UPROPERTY(Transient)
-	FPoseSearchFeatureVectorBuilder ComposedQuery;
-
-	// Precomputed runtime weights
-	UPROPERTY(Transient)
-	FPoseSearchWeightsContext WeightsContext;
-
-	// When the database changes, the search parameters are reset
-	UPROPERTY(Transient)
-	TWeakObjectPtr<const UPoseSearchDatabase> CurrentDatabase = nullptr;
-
-#if WITH_EDITOR
-	FIoHash CurrentSearchIndexHash = FIoHash::Zero;
-#endif // _DEBUG
+	UE::PoseSearch::FSearchResult CurrentSearchResult;
 
 	// Time since the last pose jump
 	UPROPERTY(Transient)
@@ -133,10 +102,28 @@ struct POSESEARCH_API FMotionMatchingState
 * @param Settings					Input motion matching algorithm configuration settings
 * @param InOutMotionMatchingState	Input/Output encapsulated motion matching algorithm and state
 */
-POSESEARCH_API void UpdateMotionMatchingState(const FAnimationUpdateContext& Context
-	, const UPoseSearchDatabase* Database
-	, const FGameplayTagQuery* DatabaseTagQuery
-	, const FTrajectorySampleRange& Trajectory
-	, const FMotionMatchingSettings& Settings
-	, FMotionMatchingState& InOutMotionMatchingState
+POSESEARCH_API void UpdateMotionMatchingState(
+	const FAnimationUpdateContext& Context,
+	const UPoseSearchSearchableAsset* Searchable,
+	const FGameplayTagQuery* DatabaseTagQuery,
+	const FGameplayTagContainer* ActiveTagsContainer,
+	const FTrajectorySampleRange& Trajectory,
+	const FMotionMatchingSettings& Settings,
+	FMotionMatchingState& InOutMotionMatchingState
 );
+
+UCLASS(DisplayName = "Bias")
+class POSESEARCH_API UPoseSearchPostProcessor_Bias : public UPoseSearchPostProcessor
+{
+	GENERATED_BODY()
+
+public:
+	virtual EPoseSearchPostSearchStatus PostProcess_Implementation(FPoseSearchCost& InOutCost) const override;
+
+	UPROPERTY(EditAnywhere, Category = Parameters)
+	float Multiplier = 1.0f;
+
+	UPROPERTY(EditAnywhere, Category = Parameters)
+	float Addend = 0.0f;
+};
+
