defaulting MM velocities in character space

diff --git a/Engine/Plugins/Experimental/Animation/PoseSearch/Source/Runtime/Private/PoseSearchFeatureChannels.cpp b/Engine/Plugins/Experimental/Animation/PoseSearch/Source/Runtime/Private/PoseSearchFeatureChannels.cpp
--- a/Engine/Plugins/Experimental/Animation/PoseSearch/Source/Runtime/Private/PoseSearchFeatureChannels.cpp
+++ b/Engine/Plugins/Experimental/Animation/PoseSearch/Source/Runtime/Private/PoseSearchFeatureChannels.cpp
@@ -24,6 +24,7 @@ constexpr float DrawDebugGradientStrength = 0.8f;
 constexpr float DrawDebugSampleLabelFontScale = 1.0f;
 static const FVector DrawDebugSampleLabelOffset = FVector(0.0f, 0.0f, -10.0f);
 
+constexpr bool UseCharacterSpaceVelocities = true;
 
 //////////////////////////////////////////////////////////////////////////
 // Drawing helpers
@@ -180,10 +181,22 @@ void UPoseSearchFeatureChannel_Pose::AddPoseFeatures(const  UE::PoseSearch::IAss
 		// For each pose subsample term, get the corresponding clip, accumulated root motion,
 		// and wrap the time parameter based on the clip's length.
 		FSampleInfo Samples[NumFiniteDiffTerms];
-		Samples[0] = Indexer.GetSampleInfoRelative(SubsampleTime - SamplingContext->FiniteDelta, Origin);
-		Samples[1] = Indexer.GetSampleInfoRelative(SubsampleTime, Origin);
-		Samples[2] = Indexer.GetSampleInfoRelative(SubsampleTime + SamplingContext->FiniteDelta, Origin);
 
+		if (UseCharacterSpaceVelocities)
+		{
+			// character space velocity
+			Samples[0] = Indexer.GetSampleInfoRelative(SubsampleTime - SamplingContext->FiniteDelta, Indexer.GetSampleInfo(SubsampleTime - SamplingContext->FiniteDelta));
+			Samples[1] = Indexer.GetSampleInfoRelative(SubsampleTime, Indexer.GetSampleInfo(SampleTime));
+			Samples[2] = Indexer.GetSampleInfoRelative(SubsampleTime + SamplingContext->FiniteDelta, Indexer.GetSampleInfo(SubsampleTime + SamplingContext->FiniteDelta));
+		}
+		else
+		{
+			// animation space velocity
+			Samples[0] = Indexer.GetSampleInfoRelative(SubsampleTime - SamplingContext->FiniteDelta, Origin);
+			Samples[1] = Indexer.GetSampleInfoRelative(SubsampleTime, Origin);
+			Samples[2] = Indexer.GetSampleInfoRelative(SubsampleTime + SamplingContext->FiniteDelta, Origin);
+		}
+		
 		// Get pose samples
 		for (int32 Term = 0; Term != NumFiniteDiffTerms; ++Term)
 		{
@@ -326,7 +339,19 @@ bool UPoseSearchFeatureChannel_Pose::BuildQuery(
 			int32 SkeletonBoneIndex = InOutQuery.GetSchema()->BoneIndices[SchemaBoneIdx];
 
 			const FTransform& Transform = ComponentPose[SkeletonBoneIndex];
-			const FTransform& PrevTransform = ComponentPrevPose[SkeletonBoneIndex] * (RootTransformPrev * RootTransform.Inverse());
+
+			FTransform PrevTransform;
+			if (UE::PoseSearch::UseCharacterSpaceVelocities)
+			{
+				// character space velocity
+				PrevTransform = ComponentPrevPose[SkeletonBoneIndex];
+			}
+			else
+			{
+				// animation space velocity
+				PrevTransform = ComponentPrevPose[SkeletonBoneIndex] * (RootTransformPrev * RootTransform.Inverse());
+			}
+			
 			InOutQuery.SetTransform(Feature, Transform);
 			InOutQuery.SetTransformVelocity(Feature, Transform, PrevTransform, SearchContext.History->GetSampleTimeInterval());
 		}
