Fix uninitialized struct members in PoseSearch

diff --git a/Engine/Plugins/Experimental/Animation/PoseSearch/Source/Runtime/Public/PoseSearch/PoseSearch.h b/Engine/Plugins/Experimental/Animation/PoseSearch/Source/Runtime/Public/PoseSearch/PoseSearch.h
--- a/Engine/Plugins/Experimental/Animation/PoseSearch/Source/Runtime/Public/PoseSearch/PoseSearch.h
+++ b/Engine/Plugins/Experimental/Animation/PoseSearch/Source/Runtime/Public/PoseSearch/PoseSearch.h
@@ -1569,13 +1569,13 @@ struct POSESEARCH_API FPoseSearchDatabaseSetEntry
 
 public:
 	UPROPERTY(EditAnywhere, Category = Settings)
-	TObjectPtr<UPoseSearchSearchableAsset> Searchable;
+	TObjectPtr<UPoseSearchSearchableAsset> Searchable = nullptr;
 
 	UPROPERTY(EditAnywhere, Category = Settings)
 	FGameplayTag Tag;
 
 	UPROPERTY(EditAnywhere, Instanced, Category = Settings)
-	TObjectPtr<UPoseSearchPostProcessor> PostProcessor;
+	TObjectPtr<UPoseSearchPostProcessor> PostProcessor = nullptr;
 };
 
 /** A data asset which holds a collection searchable assets. */
