motion matching: phases as PoseSearchFeatureChannels input

diff --git a/Engine/Plugins/Experimental/Animation/PoseSearch/Source/Editor/Private/PoseSearchDatabasePreviewScene.cpp b/Engine/Plugins/Experimental/Animation/PoseSearch/Source/Editor/Private/PoseSearchDatabasePreviewScene.cpp
--- a/Engine/Plugins/Experimental/Animation/PoseSearch/Source/Editor/Private/PoseSearchDatabasePreviewScene.cpp
+++ b/Engine/Plugins/Experimental/Animation/PoseSearch/Source/Editor/Private/PoseSearchDatabasePreviewScene.cpp
@@ -72,6 +72,8 @@ namespace UE::PoseSearch
 					DrawParams.DefaultLifeTime = 0.0f;
 					DrawParams.PoseIdx = PreviewActor.CurrentPoseIndex;
 					DrawParams.PointSize = 5.0f;
+					DrawParams.Mesh = PreviewActor.Mesh;
+
 					EnumAddFlags(DrawParams.Flags, UE::PoseSearch::EDebugDrawFlags::DrawFast);
 
 					UE::PoseSearch::Draw(DrawParams);
diff --git a/Engine/Plugins/Experimental/Animation/PoseSearch/Source/Editor/Private/PoseSearchDebugger.cpp b/Engine/Plugins/Experimental/Animation/PoseSearch/Source/Editor/Private/PoseSearchDebugger.cpp
--- a/Engine/Plugins/Experimental/Animation/PoseSearch/Source/Editor/Private/PoseSearchDebugger.cpp
+++ b/Engine/Plugins/Experimental/Animation/PoseSearch/Source/Editor/Private/PoseSearchDebugger.cpp
@@ -69,23 +69,47 @@ void FPoseSearchDebuggerPoseVector::ExtractFeatures(const UE::PoseSearch::FFeatu
 			continue;
 		}
 
-		FVector Vector;
-		if (Reader.GetVector(Feature, &Vector))
+		switch (Feature.Type)
 		{
-			if (Feature.Type == EPoseSearchFeatureType::Position)
+		case EPoseSearchFeatureType::Position:
+		{
+			FVector Vector;
+			if (Reader.GetVector(Feature, &Vector))
 			{
 				Channels[Feature.ChannelIdx]->Positions.Add(Vector);
 			}
-			else if (Feature.Type == EPoseSearchFeatureType::LinearVelocity)
+			break;
+		}
+		case EPoseSearchFeatureType::LinearVelocity:
+		{
+			FVector Vector;
+			if (Reader.GetVector(Feature, &Vector))
 			{
 				Channels[Feature.ChannelIdx]->LinearVelocities.Add(Vector);
 			}
-			else if (Feature.Type == EPoseSearchFeatureType::ForwardVector)
+			break;
+
+		}
+		case EPoseSearchFeatureType::ForwardVector:
+		{
+			FVector Vector;
+			if (Reader.GetVector(Feature, &Vector))
 			{
 				Channels[Feature.ChannelIdx]->FacingDirections.Add(Vector);
 			}
+			break;
+		}
+
+		case EPoseSearchFeatureType::Phase:
+		{
+			FVector2D Phase;
+			if (Reader.GetPhase(Feature, &Phase))
+			{
+				Channels[Feature.ChannelIdx]->Phases.Add(Phase);
+			}
+			break;
+		}
 		}
-		
 	}
 
 	for (FPoseSearchDebuggerPoseVectorChannel* Channel : Channels)
@@ -93,6 +117,7 @@ void FPoseSearchDebuggerPoseVector::ExtractFeatures(const UE::PoseSearch::FFeatu
 		Channel->bShowPositions = !Channel->Positions.IsEmpty();
 		Channel->bShowLinearVelocities = !Channel->LinearVelocities.IsEmpty();
 		Channel->bShowFacingDirections = !Channel->FacingDirections.IsEmpty();
+		Channel->bShowPhases = !Channel->Phases.IsEmpty();
 	}
 
 	bShowPose = !Pose.IsEmpty();
@@ -1783,7 +1808,7 @@ void SDebuggerView::DrawVisualization() const
 	const FTransform* Transform = ViewModel.Get()->GetRootTransform();
 	if (State && Database && Transform)
 	{
-		DrawFeatures(*DebuggerWorld, *State, *Database, *Transform);
+		DrawFeatures(*DebuggerWorld, *State, *Database, *Transform, ViewModel.Get()->GetMeshComponent());
 	}
 }
 
@@ -1801,7 +1826,8 @@ void SDebuggerView::DrawFeatures(
 	const UWorld& DebuggerWorld,
 	const FTraceMotionMatchingStateMessage& State,
 	const UPoseSearchDatabase& Database,
-	const FTransform& Transform
+	const FTransform& Transform,
+	const USkinnedMeshComponent* Mesh
 ) const
 {
 	// Set shared state
@@ -1812,7 +1838,7 @@ void SDebuggerView::DrawFeatures(
 	DrawParams.RootTransform = Transform;
 	// Single frame render
 	DrawParams.DefaultLifeTime = 0.0f;
-
+	DrawParams.Mesh = Mesh;
 	const TObjectPtr<UPoseSearchDebuggerReflection>& Reflection = DetailsView->GetReflection();
 
 	auto SetDrawFlags = [](FDebugDrawParams& InDrawParams, const FPoseSearchDebuggerFeatureDrawOptions& Options)
@@ -2654,6 +2680,15 @@ void FDebuggerViewModel::UpdateAsset()
 	}
 }
 
+const USkinnedMeshComponent* FDebuggerViewModel::GetMeshComponent() const
+{
+	if (Skeletons.Num() > FDebuggerViewModel::Asset)
+	{
+		return Skeletons[FDebuggerViewModel::Asset].Component;
+	}
+	return nullptr;
+}
+
 void FDebuggerViewModel::FillCompactPoseAndComponentRefRotations()
 {
 	const UPoseSearchDatabase* Database = GetPoseSearchDatabase();
diff --git a/Engine/Plugins/Experimental/Animation/PoseSearch/Source/Editor/Private/PoseSearchDebugger.h b/Engine/Plugins/Experimental/Animation/PoseSearch/Source/Editor/Private/PoseSearchDebugger.h
--- a/Engine/Plugins/Experimental/Animation/PoseSearch/Source/Editor/Private/PoseSearchDebugger.h
+++ b/Engine/Plugins/Experimental/Animation/PoseSearch/Source/Editor/Private/PoseSearchDebugger.h
@@ -90,6 +90,9 @@ struct FPoseSearchDebuggerPoseVectorChannel
 	UPROPERTY(VisibleDefaultsOnly, Category = "Channel")
 	bool bShowFacingDirections = false;
 
+	UPROPERTY(VisibleDefaultsOnly, Category = "Channel")
+	bool bShowPhases = false;
+
     UPROPERTY(VisibleAnywhere, EditFixedSize, Category="Channel", meta=(EditCondition="bShowPositions", EditConditionHides))
 	TArray<FVector> Positions;
 
@@ -99,6 +102,9 @@ struct FPoseSearchDebuggerPoseVectorChannel
     UPROPERTY(VisibleAnywhere, EditFixedSize, Category="Channel", meta=(EditCondition="bShowFacingDirections", EditConditionHides))
 	TArray<FVector> FacingDirections;
 
+	UPROPERTY(VisibleAnywhere, EditFixedSize, Category = "Channel", meta = (EditCondition = "bShowPhases", EditConditionHides))
+	TArray<FVector2D> Phases;
+
 	void Reset();
 	bool IsEmpty() const { return Positions.IsEmpty() && LinearVelocities.IsEmpty() && FacingDirections.IsEmpty(); }
 };
@@ -444,7 +450,7 @@ public:
 
 private:
 	/** Called each frame to draw features of the query vector & database selections */
-	void DrawFeatures(const UWorld& DebuggerWorld, const FTraceMotionMatchingStateMessage& State, const UPoseSearchDatabase& Database, const FTransform& Transform) const;
+	void DrawFeatures(const UWorld& DebuggerWorld, const FTraceMotionMatchingStateMessage& State, const UPoseSearchDatabase& Database, const FTransform& Transform, const USkinnedMeshComponent* Mesh) const;
 	
 	/** Check if a node selection was made, true if a node is selected */
 	bool UpdateSelection();
@@ -591,6 +597,8 @@ public:
 	/** Updates the current playing asset */
 	void UpdateAsset();
 
+	const USkinnedMeshComponent* GetMeshComponent() const;
+
 private:
 	/** Update the list of states for this frame */
 	void UpdateFromTimeline();
diff --git a/Engine/Plugins/Experimental/Animation/PoseSearch/Source/Runtime/Private/PoseSearch.cpp b/Engine/Plugins/Experimental/Animation/PoseSearch/Source/Runtime/Private/PoseSearch.cpp
--- a/Engine/Plugins/Experimental/Animation/PoseSearch/Source/Runtime/Private/PoseSearch.cpp
+++ b/Engine/Plugins/Experimental/Animation/PoseSearch/Source/Runtime/Private/PoseSearch.cpp
@@ -186,6 +186,7 @@ static constexpr FFeatureTypeTraits FeatureTypeTraits[] =
 	{ EPoseSearchFeatureType::LinearVelocity, 3 },
 	{ EPoseSearchFeatureType::AngularVelocity, 3 },
 	{ EPoseSearchFeatureType::ForwardVector, 3 },
+	{ EPoseSearchFeatureType::Phase, 2 },
 };
 
 static FFeatureTypeTraits GetFeatureTypeTraits(EPoseSearchFeatureType Type)
@@ -304,7 +305,7 @@ bool FPoseSearchFeatureVectorLayout::EnumerateBy(int32 ChannelIdx, EPoseSearchFe
 
 	for (int32 Size = Features.Num(); ++InOutFeatureIdx < Size; )
 	{
-		const FPoseSearchFeatureDesc& Feature = Features[InOutFeatureIdx];
+		const FPoseSearchFeatureDesc& Feature = Features[InOutFeatureIdx]; 
 
 		bool bChannelMatch = IsChannelMatch(ChannelIdx, Feature);
 		bool bTypeMatch = IsTypeMatch(Type, Feature);
@@ -338,6 +339,10 @@ uint32 FPoseSearchBone::GetTypeMask() const
 	{
 		Mask |= 1 << static_cast<int>(EPoseSearchFeatureType::Rotation);
 	}
+	if (bUsePhase)
+	{
+		Mask |= 1 << static_cast<int>(EPoseSearchFeatureType::Phase);
+	}
 
 	return Mask;
 }
@@ -2508,6 +2513,26 @@ void FPoseSearchFeatureVectorBuilder::SetPosition(FPoseSearchFeatureDesc Feature
 	SetVector(Feature, Position);
 }
 
+void FPoseSearchFeatureVectorBuilder::SetPhase(FPoseSearchFeatureDesc Feature, const FVector2D& Phase)
+{
+	Feature.Type = EPoseSearchFeatureType::Phase;
+	
+	int32 ElementIndex = Schema->Layout.Features.Find(Feature);
+	if (ElementIndex >= 0)
+	{
+		const FPoseSearchFeatureDesc& FoundElement = Schema->Layout.Features[ElementIndex];
+
+		Values[FoundElement.ValueOffset + 0] = Phase.X;
+		Values[FoundElement.ValueOffset + 1] = Phase.Y;
+
+		if (!FeaturesAdded[ElementIndex])
+		{
+			FeaturesAdded[ElementIndex] = true;
+			++NumFeaturesAdded;
+		}
+	}
+}
+
 void FPoseSearchFeatureVectorBuilder::SetRotation(FPoseSearchFeatureDesc Feature, const FQuat& Rotation)
 {
 	Feature.Type = EPoseSearchFeatureType::Rotation;
@@ -2979,6 +3004,23 @@ bool FFeatureVectorReader::GetAngularVelocity(FPoseSearchFeatureDesc Element, FV
 	return GetVector(Element, OutAngularVelocity);
 }
 
+bool FFeatureVectorReader::GetPhase(FPoseSearchFeatureDesc Element, FVector2D* OutPhase) const
+{
+	Element.Type = EPoseSearchFeatureType::Phase;
+	int32 ElementIndex = IsValid() ? Layout->Features.Find(Element) : -1;
+	if (ElementIndex >= 0)
+	{
+		const FPoseSearchFeatureDesc& FoundElement = Layout->Features[ElementIndex];
+
+		OutPhase->X = Values[FoundElement.ValueOffset + 0];
+		OutPhase->Y = Values[FoundElement.ValueOffset + 1];
+		return true;
+	}
+
+	*OutPhase = FVector2D::ZeroVector;
+	return false;
+}
+
 bool FFeatureVectorReader::GetVector(FPoseSearchFeatureDesc Element, FVector* OutVector) const
 {
 	int32 ElementIndex = IsValid() ? Layout->Features.Find(Element) : -1;
@@ -4091,6 +4133,10 @@ void FAssetIndexer::Init(const FAssetIndexingContext& InIndexingContext)
 
 bool FAssetIndexer::Process()
 {
+	check(IndexingContext.Schema);
+	check(IndexingContext.Schema->IsValid());
+	check(IndexingContext.MainSampler);
+
 	FMemMark Mark(FMemStack::Get());
 
 	IndexingContext.BeginSampleIdx = Output.FirstIndexedSample;
@@ -4493,10 +4539,14 @@ void Draw(const FDebugDrawParams& DebugDrawParams)
 			const FPoseSearchIndex* SearchIndex = DebugDrawParams.GetSearchIndex();
 			check(SearchIndex);
 
-			TArray<float> PoseVector;
-			PoseVector = SearchIndex->GetPoseValues(DebugDrawParams.PoseIdx);
-			SearchIndex->InverseNormalize(PoseVector);
-			DrawFeatureVector(DebugDrawParams, PoseVector);
+			// PreprocessInfo happens to be invalid when updating the database
+			if (SearchIndex->PreprocessInfo.IsValid())
+			{
+				TArray<float> PoseVector;
+				PoseVector = SearchIndex->GetPoseValues(DebugDrawParams.PoseIdx);
+				SearchIndex->InverseNormalize(PoseVector);
+				DrawFeatureVector(DebugDrawParams, PoseVector);
+			}
 		}
 		if (!DebugDrawParams.PoseVector.IsEmpty())
 		{
@@ -4939,7 +4989,7 @@ static void PreprocessSearchIndex(FPoseSearchIndex* SearchIndex)
 
 static void PreprocessGroupSearchIndexWeights(FGroupSearchIndex& GroupSearchIndex, const UPoseSearchDatabase* Database)
 {
-	const FPoseSearchWeightParams& WeightParams = GroupSearchIndex.GroupIndex == INDEX_NONE ? Database->DefaultWeights : Database->Groups[GroupSearchIndex.GroupIndex].Weights;
+	const FPoseSearchWeightParams& WeightParams = (GroupSearchIndex.GroupIndex == INDEX_NONE || Database->Groups.IsEmpty())  ? Database->DefaultWeights : Database->Groups[GroupSearchIndex.GroupIndex].Weights;
 	FPoseSearchWeights Weights;
 	Weights.Init(WeightParams, Database->Schema);
 	GroupSearchIndex.Weights = Weights.Weights;
diff --git a/Engine/Plugins/Experimental/Animation/PoseSearch/Source/Runtime/Private/PoseSearchFeatureChannels.cpp b/Engine/Plugins/Experimental/Animation/PoseSearch/Source/Runtime/Private/PoseSearchFeatureChannels.cpp
--- a/Engine/Plugins/Experimental/Animation/PoseSearch/Source/Runtime/Private/PoseSearchFeatureChannels.cpp
+++ b/Engine/Plugins/Experimental/Animation/PoseSearch/Source/Runtime/Private/PoseSearchFeatureChannels.cpp
@@ -44,6 +44,223 @@ static FLinearColor GetColorForFeature(FPoseSearchFeatureDesc Feature, const FPo
 	return ColorHSV.HSVToLinearRGB();
 }
 
+
+struct LocalMinMax
+{
+	enum
+	{
+		Min,
+		Max
+	} Type = Min;
+	int32 Index = 0;
+	float SignalValue = 0.f;
+};
+
+template <typename T>
+T GetValueAtIndex(int32 Sample, const TArray<T>& Values)
+{
+	const int32 Num = Values.Num();
+	check(Num > 1);
+
+	if (Sample < 0)
+	{
+		return (Values[1] - Values[0]) * Sample + Values[0];
+	}
+
+	if (Sample < Num)
+	{
+		return Values[Sample];
+	}
+
+	return (Values[Num - 1] - Values[Num - 2]) * (Sample - (Num - 1)) + Values[Num - 1];
+}
+
+static void CalculateSignal(const TArray<FVector>& BonePositions, TArray<float>& Signal, int32 offset = 1)
+{
+	Signal.Reset();
+	Signal.AddDefaulted(BonePositions.Num());
+
+	for (int32 SampleIdx = 0; SampleIdx != BonePositions.Num(); ++SampleIdx)
+	{
+		Signal[SampleIdx] = (GetValueAtIndex(SampleIdx + offset, BonePositions) - GetValueAtIndex(SampleIdx - offset, BonePositions)).Length();
+	}
+}
+
+static void SmoothSignal(const TArray<float>& Signal, TArray<float>& SmoothedSignal, int32 offset = 1)
+{
+	SmoothedSignal.Reset();
+	SmoothedSignal.AddDefaulted(Signal.Num());
+
+	for (int32 SampleIdx = -offset; SampleIdx != offset; ++SampleIdx)
+	{
+		SmoothedSignal[0] += GetValueAtIndex(SampleIdx, Signal);
+	}
+
+	for (int32 SampleIdx = 1; SampleIdx != Signal.Num(); ++SampleIdx)
+	{
+		SmoothedSignal[SampleIdx] = SmoothedSignal[SampleIdx - 1] - GetValueAtIndex(SampleIdx - offset - 1, Signal) + GetValueAtIndex(SampleIdx + offset, Signal);
+	}
+
+	for (int32 SampleIdx = 0; SampleIdx != Signal.Num(); ++SampleIdx)
+	{
+		SmoothedSignal[SampleIdx] /= 2 * offset + 1;
+	}
+}
+
+static void FindLocalMinMax(const TArray<float>& Signal, TArray<LocalMinMax>& MinMax, int32 offset = 1)
+{
+	check(offset > 0);
+	MinMax.Reset();
+	for (int32 i = 0; i < Signal.Num(); ++i)
+	{
+		const float Previous = GetValueAtIndex(i - offset, Signal);
+		const float Current = GetValueAtIndex(i, Signal);
+		const float Next = GetValueAtIndex(i + offset, Signal);
+
+		const float DeltaSignalValueBackward = Previous - Current;
+		const float DeltaSignalValueForward = Next - Current;
+
+		const float Sign = DeltaSignalValueBackward * DeltaSignalValueForward;
+		if (Sign >= 0.f && DeltaSignalValueBackward != 0.f)
+		{
+			LocalMinMax LocalMinMax;
+			LocalMinMax.Type = DeltaSignalValueForward < 0.f ? LocalMinMax::Max : LocalMinMax::Min;
+			LocalMinMax.Index = i;
+			LocalMinMax.SignalValue = Signal[i];
+
+			check(MinMax.IsEmpty() || MinMax.Last().Type != LocalMinMax.Type);
+			MinMax.Add(LocalMinMax);
+		}
+	}
+}
+
+static void ExtrapolateLocalMinMaxBoundaries(TArray<LocalMinMax>& MinMax, const TArray<float>& Signal)
+{
+	const int32 Num = MinMax.Num();
+
+	check(Signal.Num() > 0);
+
+	LocalMinMax InitialMinMax;
+	LocalMinMax FinalMinMax;
+
+	if (Num == 0)
+	{
+		const bool IsInitialMax = Signal[0] > Signal[Signal.Num() - 1];
+
+		InitialMinMax.Index = 0;
+		InitialMinMax.SignalValue = Signal[0];
+		InitialMinMax.Type = IsInitialMax ? LocalMinMax::Max : LocalMinMax::Min;
+
+		FinalMinMax.Index = Signal.Num() - 1;
+		FinalMinMax.SignalValue = Signal[Signal.Num() - 1];
+		FinalMinMax.Type = IsInitialMax ? LocalMinMax::Min : LocalMinMax::Max;
+
+		MinMax.Add(InitialMinMax);
+		MinMax.Add(FinalMinMax);
+	}
+	else
+	{
+		int32 InitialDelta = 0;
+		int32 FinalDelta = 0;
+		if (Num > 2)
+		{
+			InitialDelta = MinMax[2].Index - MinMax[1].Index;
+			FinalDelta = MinMax[Num - 2].Index - MinMax[Num - 3].Index;
+		}
+		else if (Num > 1)
+		{
+			InitialDelta = MinMax[1].Index - MinMax[0].Index;
+			FinalDelta = MinMax[Num - 1].Index - MinMax[Num - 2].Index;
+		}
+		else
+		{
+			InitialDelta = MinMax[0].Index;
+			FinalDelta = (Signal.Num() - 1) - MinMax[0].Index;
+		}
+
+		InitialMinMax.SignalValue = Num > 1 ? MinMax[1].SignalValue : Signal[0];
+		InitialMinMax.Type = MinMax[0].Type == LocalMinMax::Min ? LocalMinMax::Max : LocalMinMax::Min;
+		InitialMinMax.Index = FMath::Min(MinMax[0].Index - InitialDelta, 0);
+
+		FinalMinMax.SignalValue = Num > 1 ? MinMax[Num - 2].SignalValue : Signal[Signal.Num() - 1];
+		FinalMinMax.Type = MinMax[Num - 1].Type == LocalMinMax::Min ? LocalMinMax::Max : LocalMinMax::Min;
+		FinalMinMax.Index = FMath::Max(MinMax[Num - 1].Index + FinalDelta, Signal.Num() - 1);
+
+		// there's no point in adding an InitialMinMax if the first MinMax is at the first frame of the signal
+		if (MinMax[0].Index > 0)
+		{
+			MinMax.Insert(InitialMinMax, 0);
+		}
+
+		// there's no point in adding a FinalMinMax if the last MinMax is at the last frame of the signal
+		if (MinMax[Num - 1].Index < Signal.Num() - 1)
+		{
+			MinMax.Add(FinalMinMax);
+		}
+	}
+}
+
+static void ValidateLocalMinMax(const TArray<LocalMinMax>& MinMax)
+{
+	for (int32 i = 1; i < MinMax.Num(); ++i)
+	{
+		check(MinMax[i].Type != MinMax[i - 1].Type);
+		check(MinMax[i].Index > MinMax[i - 1].Index);
+		if (MinMax[i].Type == LocalMinMax::Min)
+		{
+			check(MinMax[i].SignalValue < MinMax[i - 1].SignalValue);
+		}
+		else
+		{
+			check(MinMax[i].SignalValue > MinMax[i - 1].SignalValue);
+		}
+	}
+}
+
+static void CalculatePhaseAndCertainty(int32 Index, const TArray<LocalMinMax>& MinMax, int32 SignalSize, float& Phase, float& Certainty)
+{
+	// @todo: expose them via UI
+	static float CertaintyMin = 1.f;
+	static float CertaintyMult = 0.1f;
+
+	const int32 LastIndex = MinMax.Num() - 1;
+	for (int32 i = 1; i < MinMax.Num(); ++i)
+	{
+		const int32 MinMaxIndex = MinMax[i].Index;
+		if (Index < MinMaxIndex)
+		{
+			const int32 PrevMinMaxIndex = MinMax[i - 1].Index;
+			check(MinMaxIndex > PrevMinMaxIndex);
+			const float Ratio = static_cast<float>((Index - PrevMinMaxIndex)) / static_cast<float>((MinMaxIndex - PrevMinMaxIndex));
+			const float PhaseOffset = MinMax[i - 1].Type == LocalMinMax::Min ? 0.f : 0.5f;
+			Phase = PhaseOffset + Ratio * 0.5f;
+
+			const float DeltaSignalValue = FMath::Abs(MinMax[i - 1].SignalValue - MinMax[i].SignalValue);
+			const float NextDeltaSignalValue = i < LastIndex ? FMath::Abs(MinMax[i].SignalValue - MinMax[i + 1].SignalValue) : DeltaSignalValue;
+			Certainty = CertaintyMin + (DeltaSignalValue * (1.f - Ratio) + NextDeltaSignalValue * Ratio) * CertaintyMult;
+			return;
+		}
+	}
+
+	Phase = MinMax[LastIndex].Type == LocalMinMax::Min ? 0.f : 0.5f;
+	Certainty = CertaintyMin + (LastIndex > 0 ? FMath::Abs(MinMax[LastIndex].SignalValue - MinMax[LastIndex - 1].SignalValue) : 0.f) * CertaintyMult;
+}
+
+static void CalculatePhasesFromLocalMinMax(const TArray<LocalMinMax>& MinMax, TArray<FVector2D>& Phases, int32 SignalSize)
+{
+	Phases.Reset();
+	Phases.AddDefaulted(SignalSize);
+
+	float Certainty = 1.f;
+	float Phase = 0.f;
+	for (int32 i = 0; i < SignalSize; ++i)
+	{
+		CalculatePhaseAndCertainty(i, MinMax, SignalSize, Phase, Certainty);
+		FMath::SinCos(&Phases[i].X, &Phases[i].Y, Phase * TWO_PI);
+		Phases[i] *= Certainty;
+	}
+}
+
 } // namespace UE::PoseSearch
 
 
@@ -110,20 +327,121 @@ void UPoseSearchFeatureChannel_Pose::InitializeSchema(UE::PoseSearch::FSchemaIni
 	}
 }
 
+// @todo: do we really need to use double(s) in all this math?
+void UPoseSearchFeatureChannel_Pose::CalculatePhases(const UE::PoseSearch::IAssetIndexer& Indexer, UE::PoseSearch::FAssetIndexingOutput& IndexingOutput, TArray<TArray<FVector2D>>& OutPhases) const
+{
+	// @todo: expose them via UI
+	static float BoneSamplingCentralDifferencesTime = 0.2f; // seconds
+	static float SmoothingWindowTime = 0.3f; // seconds
+
+	using namespace UE::PoseSearch;
+	using FSampleInfo = IAssetIndexer::FSampleInfo;
+
+	const FAssetIndexingContext& IndexingContext = Indexer.GetIndexingContext();
+
+	const float SampleTimeStart = FMath::Min(IndexingContext.BeginSampleIdx * IndexingContext.Schema->SamplingInterval, IndexingContext.MainSampler->GetPlayLength());
+	const FAssetSamplingContext* SamplingContext = IndexingContext.SamplingContext;
+	const float FiniteDelta = IndexingContext.Schema->SamplingInterval;
+
+	const int32 NumSamples = IndexingContext.EndSampleIdx - IndexingContext.BeginSampleIdx;
+
+	TArray<TArray<FVector>> BonePositions;
+	BonePositions.AddDefaulted(SampledBones.Num());
+	for (int32 ChannelBoneIdx = 0; ChannelBoneIdx != SampledBones.Num(); ++ChannelBoneIdx)
+	{
+		BonePositions[ChannelBoneIdx].AddDefaulted(NumSamples);
+	}
+
+	FDeltaTimeRecord DeltaTimeRecord;
+	FCompactPose Pose;
+	FCSPose<FCompactPose> ComponentSpacePose;
+	FBlendedCurve UnusedCurve;
+	UE::Anim::FStackAttributeContainer UnusedAtrribute;
+	Pose.SetBoneContainer(&SamplingContext->BoneContainer);
+
+	// collecting all the bone transforms
+	FSampleInfo Origin = Indexer.GetSampleInfo(SampleTimeStart);
+	for (int32 SampleIdx = 0; SampleIdx != NumSamples; ++SampleIdx)
+	{
+		const float SampleTime = SampleTimeStart + SampleIdx * FiniteDelta;
+		const float PreviousTime = SampleTime - FiniteDelta; // @todo: clamp to zero?
+
+		FSampleInfo Sample = Indexer.GetSampleInfoRelative(SampleTime, Origin);
+
+		DeltaTimeRecord.Set(PreviousTime, SampleTime - PreviousTime);
+		FAnimExtractContext ExtractionCtx(static_cast<double>(SampleTime), true, DeltaTimeRecord, Sample.Clip->IsLoopable());
+		
+		UnusedCurve.InitFrom(SamplingContext->BoneContainer);
+		FAnimationPoseData AnimPoseData = { Pose, UnusedCurve, UnusedAtrribute };
+		Sample.Clip->ExtractPose(ExtractionCtx, AnimPoseData);
+
+		if (IndexingContext.bMirrored)
+		{
+			FAnimationRuntime::MirrorPose(
+				AnimPoseData.GetPose(),
+				IndexingContext.Schema->MirrorDataTable->MirrorAxis,
+				SamplingContext->CompactPoseMirrorBones,
+				SamplingContext->ComponentSpaceRefRotations
+			);
+			// Note curves and attributes are not used during the indexing process and therefore don't need to be mirrored
+		}
+
+		ComponentSpacePose.InitPose(Pose);
+		for (int32 ChannelBoneIdx = 0; ChannelBoneIdx != SampledBones.Num(); ++ChannelBoneIdx)
+		{
+			const FPoseSearchPoseFeatureInfo& PoseFeatureInfo = FeatureParams[ChannelBoneIdx];
+			const FBoneReference& BoneReference = IndexingContext.Schema->BoneReferences[PoseFeatureInfo.SchemaBoneIdx];
+			const FCompactPoseBoneIndex CompactBoneIndex = SamplingContext->BoneContainer.MakeCompactPoseIndex(FMeshPoseBoneIndex(BoneReference.BoneIndex));
+			const FTransform BoneTransform = ComponentSpacePose.GetComponentSpaceTransform(CompactBoneIndex) * Indexer.MirrorTransform(Sample.RootTransform);
+			BonePositions[ChannelBoneIdx][SampleIdx] = BoneTransform.GetTranslation();
+		}
+	}
+
+	OutPhases.Reset();
+	OutPhases.AddDefaulted(SampledBones.Num());
+
+	const int32 BoneSamplingCentralDifferencesOffset = FMath::Max(FMath::CeilToInt(BoneSamplingCentralDifferencesTime / FiniteDelta), 1);
+	const int32 SmoothingWindowOffset = FMath::Max(FMath::CeilToInt(SmoothingWindowTime / FiniteDelta), 1);
+	
+	TArray<float> Signal;
+	TArray<float> SmoothedSignal;
+	TArray<LocalMinMax> LocalMinMax;
+	for (int32 ChannelBoneIdx = 0; ChannelBoneIdx != SampledBones.Num(); ++ChannelBoneIdx)
+	{
+		// @todo: have different way of calculating signals, for example: height of the bone transform, acceleration, etc?
+		CalculateSignal(BonePositions[ChannelBoneIdx], Signal, BoneSamplingCentralDifferencesOffset);
+		
+		SmoothSignal(Signal, SmoothedSignal, SmoothingWindowOffset);
+
+		FindLocalMinMax(SmoothedSignal, LocalMinMax);
+		ValidateLocalMinMax(LocalMinMax);
+
+		ExtrapolateLocalMinMaxBoundaries(LocalMinMax, SmoothedSignal);
+		ValidateLocalMinMax(LocalMinMax);
+		CalculatePhasesFromLocalMinMax(LocalMinMax, OutPhases[ChannelBoneIdx], SmoothedSignal.Num());
+	}
+}
+
 void UPoseSearchFeatureChannel_Pose::IndexAsset(const UE::PoseSearch::IAssetIndexer& Indexer,  UE::PoseSearch::FAssetIndexingOutput& IndexingOutput) const
 {
 	using namespace UE::PoseSearch;
+	 
+	// Phases is an array of array with cardinality SampledBones.Num() times NumSamples (IndexingContext.EndSampleIdx - IndexingContext.BeginSampleIdx)
+	// of 2 dimensional vectors (FVector2D) representing phases in an Eucledean space with phase angle sin/cos as direction and certainty of the signal as magnitude,
+	// where certainty is a function of the amplitude of the signal used as input
+	TArray<TArray<FVector2D>> Phases;
+	CalculatePhases(Indexer, IndexingOutput, Phases);
 
 	const FAssetIndexingContext& IndexingContext = Indexer.GetIndexingContext();
 	for (int32 SampleIdx = IndexingContext.BeginSampleIdx; SampleIdx != IndexingContext.EndSampleIdx; ++SampleIdx)
 	{
 		int32 VectorIdx = SampleIdx - IndexingContext.BeginSampleIdx;
 		FPoseSearchFeatureVectorBuilder& FeatureVector = IndexingOutput.PoseVectors[VectorIdx];
-		AddPoseFeatures(Indexer, SampleIdx, FeatureVector);
+		AddPoseFeatures(Indexer, SampleIdx, FeatureVector, Phases);
 	}
 }
 
-void UPoseSearchFeatureChannel_Pose::AddPoseFeatures(const  UE::PoseSearch::IAssetIndexer& Indexer, int32 SampleIdx, FPoseSearchFeatureVectorBuilder& FeatureVector) const
+void UPoseSearchFeatureChannel_Pose::AddPoseFeatures(const  UE::PoseSearch::IAssetIndexer& Indexer, int32 SampleIdx, FPoseSearchFeatureVectorBuilder& FeatureVector, const TArray<TArray<FVector2D>>& Phases) const
 {
 	// This function samples the instantaneous pose at time t as well as the pose's velocity and acceleration at time t.
 	// Symmetric finite differences are used to approximate derivatives:
@@ -261,6 +579,9 @@ void UPoseSearchFeatureChannel_Pose::AddPoseFeatures(const  UE::PoseSearch::IAss
 			{
 				FeatureVector.SetTransformVelocity(Feature, BoneTransforms[2], BoneTransforms[1], BoneTransforms[0], SamplingContext->FiniteDelta);
 			}
+
+			// @todo: support for SubsampleIdx
+			FeatureVector.SetPhase(Feature, Phases[ChannelBoneIdx][SampleIdx]);
 		}
 	}
 }
@@ -391,7 +712,7 @@ void UPoseSearchFeatureChannel_Pose::DebugDraw(const UE::PoseSearch::FDebugDrawP
 			Feature.ChannelFeatureId = ChannelBoneIdx;
 
 			FVector BonePos;
-			const bool bHaveBonePos = Reader.GetPosition(Feature, &BonePos);
+			bool bHaveBonePos = Reader.GetPosition(Feature, &BonePos);
 			if (bHaveBonePos)
 			{
 				Feature.Type = EPoseSearchFeatureType::Position;
@@ -418,6 +739,11 @@ void UPoseSearchFeatureChannel_Pose::DebugDraw(const UE::PoseSearch::FDebugDrawP
 						nullptr, Color, LifeTime, false, 1.0f);
 				}
 			}
+			else if (DrawParams.Mesh != nullptr)
+			{
+				BonePos = DrawParams.Mesh->GetSocketTransform(SampledBones[ChannelBoneIdx].Reference.BoneName).GetLocation();
+				bHaveBonePos = true;
+			}
 
 			FVector BoneVel;
 			if (bHaveBonePos && Reader.GetLinearVelocity(Feature, &BoneVel))
@@ -453,6 +779,29 @@ void UPoseSearchFeatureChannel_Pose::DebugDraw(const UE::PoseSearch::FDebugDrawP
 						AdjustedThickness);
 				}
 			}
+
+			FVector2D Phase;
+			if (bHaveBonePos && Reader.GetPhase(Feature, &Phase))
+			{
+				Feature.Type = EPoseSearchFeatureType::Phase;
+
+				FLinearColor LinearColor = DrawParams.Color ? *DrawParams.Color : GetColorForFeature(Feature, Reader.GetLayout());
+				FColor Color = LinearColor.ToFColor(true);
+
+				static float ScaleFactor = 1.f;
+
+				const FVector TransformXAxisVector = DrawParams.RootTransform.TransformVector(FVector::XAxisVector);
+				const FVector TransformYAxisVector = DrawParams.RootTransform.TransformVector(FVector::YAxisVector);
+				const FVector TransformZAxisVector = DrawParams.RootTransform.TransformVector(FVector::ZAxisVector);
+
+				const FVector PhaseVector = (TransformZAxisVector * Phase.X + TransformYAxisVector * Phase.Y) * ScaleFactor;
+				DrawDebugLine(DrawParams.World, BonePos, BonePos + PhaseVector, Color, bPersistent, LifeTime, DepthPriority, 0.f);
+
+				static int32 Segments = 64;
+				FMatrix CircleTransform;
+				CircleTransform.SetAxes(&TransformXAxisVector, &TransformYAxisVector, &TransformZAxisVector, &BonePos);
+				DrawDebugCircle(DrawParams.World, CircleTransform, PhaseVector.Length(), Segments, Color, bPersistent, LifeTime, DepthPriority, 0.f, false);
+			}
 		}
 	}
 }
diff --git a/Engine/Plugins/Experimental/Animation/PoseSearch/Source/Runtime/Public/PoseSearch/PoseSearch.h b/Engine/Plugins/Experimental/Animation/PoseSearch/Source/Runtime/Public/PoseSearch/PoseSearch.h
--- a/Engine/Plugins/Experimental/Animation/PoseSearch/Source/Runtime/Public/PoseSearch/PoseSearch.h
+++ b/Engine/Plugins/Experimental/Animation/PoseSearch/Source/Runtime/Public/PoseSearch/PoseSearch.h
@@ -62,6 +62,7 @@ enum class EPoseSearchFeatureType : int8
 	LinearVelocity,
 	AngularVelocity,
 	ForwardVector,
+	Phase,
 
 	Num UMETA(Hidden),
 	Invalid = Num UMETA(Hidden)
@@ -196,6 +197,9 @@ struct POSESEARCH_API FPoseSearchBone
 	UPROPERTY(EditAnywhere, Category = Config)
 	bool bUseRotation = false;
 
+	UPROPERTY(EditAnywhere, Category = Config)
+	bool bUsePhase = false;
+
 	// this function will return a mask out of EPoseSearchFeatureType based on which features were selected for
 	// the bone.
 	uint32 GetTypeMask() const;
@@ -593,6 +597,11 @@ struct POSESEARCH_API FPoseSearchIndexPreprocessInfo
 		InverseTransformationMatrix.Reset();
 		SampleMean.Reset();
 	}
+
+	bool IsValid() const
+	{
+		return NumDimensions > 0;
+	}
 };
 
 
@@ -1028,6 +1037,7 @@ public:
 	void SetAngularVelocity(FPoseSearchFeatureDesc Feature, const FTransform& Transform, const FTransform& PrevTransform, float DeltaTime);
 	void SetAngularVelocity(FPoseSearchFeatureDesc Feature, const FTransform& NextTransform, const FTransform& Transform, const FTransform& PrevTransform, float DeltaTime);
 	void SetVector(FPoseSearchFeatureDesc Feature, const FVector& Vector);
+	void SetPhase(FPoseSearchFeatureDesc Feature, const FVector2D& Phase);
 
 	void CopyFromSearchIndex(const FPoseSearchIndex& SearchIndex, int32 PoseIdx);
 	void CopyFeature(const FPoseSearchFeatureVectorBuilder& OtherBuilder, int32 FeatureIdx);
@@ -1118,6 +1128,9 @@ struct POSESEARCH_API FDebugDrawParams
 	// Optional prefix for sample labels
 	FStringView LabelPrefix;
 
+	// Optional Mesh for gathering SocketTransform(s)
+	TWeakObjectPtr<const USkinnedMeshComponent> Mesh = nullptr;
+
 	bool CanDraw() const;
 	const FPoseSearchIndex* GetSearchIndex() const;
 	const UPoseSearchSchema* GetSchema() const;
@@ -1414,6 +1427,7 @@ public:
 	bool GetLinearVelocity(FPoseSearchFeatureDesc Feature, FVector* OutLinearVelocity) const;
 	bool GetAngularVelocity(FPoseSearchFeatureDesc Feature, FVector* OutAngularVelocity) const;
 	bool GetVector(FPoseSearchFeatureDesc Feature, FVector* OutVector) const;
+	bool GetPhase(FPoseSearchFeatureDesc Element, FVector2D* OutPhase) const;
 
 	const FPoseSearchFeatureVectorLayout* GetLayout() const { return Layout; }
 
diff --git a/Engine/Plugins/Experimental/Animation/PoseSearch/Source/Runtime/Public/PoseSearch/PoseSearchFeatureChannels.h b/Engine/Plugins/Experimental/Animation/PoseSearch/Source/Runtime/Public/PoseSearch/PoseSearchFeatureChannels.h
--- a/Engine/Plugins/Experimental/Animation/PoseSearch/Source/Runtime/Public/PoseSearch/PoseSearchFeatureChannels.h
+++ b/Engine/Plugins/Experimental/Animation/PoseSearch/Source/Runtime/Public/PoseSearch/PoseSearchFeatureChannels.h
@@ -52,7 +52,8 @@ public:
 	virtual void DebugDraw(const UE::PoseSearch::FDebugDrawParams& DrawParams, const UE::PoseSearch::FFeatureVectorReader& Reader) const override;
 
 protected:
-	virtual void AddPoseFeatures(const UE::PoseSearch::IAssetIndexer& Indexer, int32 SampleIdx, FPoseSearchFeatureVectorBuilder& FeatureVector) const;
+	virtual void AddPoseFeatures(const UE::PoseSearch::IAssetIndexer& Indexer, int32 SampleIdx, FPoseSearchFeatureVectorBuilder& FeatureVector, const TArray<TArray<FVector2D>>& Phases) const;
+	void CalculatePhases(const UE::PoseSearch::IAssetIndexer& Indexer, UE::PoseSearch::FAssetIndexingOutput& IndexingOutput, TArray<TArray<FVector2D>>& OutPhases) const;
 };
 
 
