motion matching - drawing cost comparison histograms when in EPoseSearchMode::PCAKDTree_Compare

diff --git a/Engine/Plugins/Experimental/Animation/PoseSearch/Source/Runtime/Private/AnimNode_MotionMatching.cpp b/Engine/Plugins/Experimental/Animation/PoseSearch/Source/Runtime/Private/AnimNode_MotionMatching.cpp
--- a/Engine/Plugins/Experimental/Animation/PoseSearch/Source/Runtime/Private/AnimNode_MotionMatching.cpp
+++ b/Engine/Plugins/Experimental/Animation/PoseSearch/Source/Runtime/Private/AnimNode_MotionMatching.cpp
@@ -191,6 +191,9 @@ void FAnimNode_MotionMatching::PreUpdate(const UAnimInstance* InAnimInstance)
 			DrawParams.PoseVector = MotionMatchingState.CurrentSearchResult.ComposedQuery.GetValues();
 		}
 
+		DrawParams.SearchCostHistoryKDTree = &MotionMatchingState.SearchCostHistoryKDTree;
+		DrawParams.SearchCostHistoryBruteForce = &MotionMatchingState.SearchCostHistoryBruteForce;
+
 		UE::PoseSearch::Draw(DrawParams);
 	}
 
diff --git a/Engine/Plugins/Experimental/Animation/PoseSearch/Source/Runtime/Private/PoseSearch.cpp b/Engine/Plugins/Experimental/Animation/PoseSearch/Source/Runtime/Private/PoseSearch.cpp
--- a/Engine/Plugins/Experimental/Animation/PoseSearch/Source/Runtime/Private/PoseSearch.cpp
+++ b/Engine/Plugins/Experimental/Animation/PoseSearch/Source/Runtime/Private/PoseSearch.cpp
@@ -561,6 +561,8 @@ bool UPoseSearchSchema::BuildQuery(
 	FPoseSearchContext& SearchContext,
 	FPoseSearchFeatureVectorBuilder& InOutQuery) const
 {
+	QUICK_SCOPE_CYCLE_COUNTER(STAT_PoseSearch_BuildQuery);
+
 	InOutQuery.Init(this);
 
 	// Copy search query directly from the database if we have an active pose
@@ -2139,11 +2141,35 @@ UE::PoseSearch::FSearchResult UPoseSearchDatabase::Search(FPoseSearchContext& Se
 	if (PoseSearchMode != EPoseSearchMode::BruteForce)
 	{
 		Result = SearchPCAKDTree(SearchContext);
+
+#if WITH_EDITORONLY_DATA
+		if (PoseSearchMode == EPoseSearchMode::PCAKDTree_Compare && SearchContext.DebugDrawParams.SearchCostHistoryKDTree && Result.IsValid())
+		{
+			SearchContext.DebugDrawParams.SearchCostHistoryKDTree->AddSample(Result.PoseCost.GetTotalCost());
+		}
+#endif
 	}
 
 	if (PoseSearchMode == EPoseSearchMode::BruteForce || PoseSearchMode == EPoseSearchMode::PCAKDTree_Compare)
 	{
 		Result = SearchBruteForce(SearchContext);
+
+#if WITH_EDITORONLY_DATA
+		if (PoseSearchMode == EPoseSearchMode::PCAKDTree_Compare && SearchContext.DebugDrawParams.SearchCostHistoryBruteForce && Result.IsValid())
+		{
+			SearchContext.DebugDrawParams.SearchCostHistoryBruteForce->AddSample(Result.PoseCost.GetTotalCost());
+
+			// making SearchCostHistoryKDTree and SearchCostHistoryBruteForce min max consistent
+			if (SearchContext.DebugDrawParams.SearchCostHistoryKDTree)
+			{
+				SearchContext.DebugDrawParams.SearchCostHistoryKDTree->MinValue = FMath::Min(SearchContext.DebugDrawParams.SearchCostHistoryKDTree->MinValue, SearchContext.DebugDrawParams.SearchCostHistoryBruteForce->MinValue);
+				SearchContext.DebugDrawParams.SearchCostHistoryKDTree->MaxValue = FMath::Max(SearchContext.DebugDrawParams.SearchCostHistoryKDTree->MaxValue, SearchContext.DebugDrawParams.SearchCostHistoryBruteForce->MaxValue);
+
+				SearchContext.DebugDrawParams.SearchCostHistoryBruteForce->MinValue = SearchContext.DebugDrawParams.SearchCostHistoryKDTree->MinValue;
+				SearchContext.DebugDrawParams.SearchCostHistoryBruteForce->MaxValue = SearchContext.DebugDrawParams.SearchCostHistoryKDTree->MaxValue;
+			}
+		}
+#endif
 	}
 
 	SearchContext.DebugDrawParams.PoseVector = Result.ComposedQuery.GetNormalizedValues();
@@ -4556,6 +4582,19 @@ void Draw(const FDebugDrawParams& DebugDrawParams)
 		{
 			DrawSearchIndex(DebugDrawParams);
 		}
+
+#if WITH_EDITORONLY_DATA
+		if (DebugDrawParams.Database->PoseSearchMode == EPoseSearchMode::PCAKDTree_Compare)
+		{
+			if (DebugDrawParams.SearchCostHistoryKDTree && DebugDrawParams.SearchCostHistoryBruteForce)
+			{
+				const FTransform OffsetTransform(FRotator(0.f, 90.f, 0.f), FVector(-50.f, 0.f, 100.f));
+				FVector2D DrawSize(150.f, 50.f);
+				DrawDebugFloatHistory(*DebugDrawParams.World, *DebugDrawParams.SearchCostHistoryKDTree, OffsetTransform * DebugDrawParams.RootTransform, DrawSize, FColor::Red);
+				DrawDebugFloatHistory(*DebugDrawParams.World, *DebugDrawParams.SearchCostHistoryBruteForce, OffsetTransform * DebugDrawParams.RootTransform, DrawSize, FColor::Blue);
+			}
+		}
+#endif
 	}
 }
 
diff --git a/Engine/Plugins/Experimental/Animation/PoseSearch/Source/Runtime/Private/PoseSearchLibrary.cpp b/Engine/Plugins/Experimental/Animation/PoseSearch/Source/Runtime/Private/PoseSearchLibrary.cpp
--- a/Engine/Plugins/Experimental/Animation/PoseSearch/Source/Runtime/Private/PoseSearchLibrary.cpp
+++ b/Engine/Plugins/Experimental/Animation/PoseSearch/Source/Runtime/Private/PoseSearchLibrary.cpp
@@ -256,7 +256,11 @@ void UpdateMotionMatchingState(
 		SearchContext.OwningComponent = Context.AnimInstanceProxy->GetSkelMeshComponent();
 		SearchContext.BoneContainer = &Context.AnimInstanceProxy->GetRequiredBones();
 
-		
+#if WITH_EDITORONLY_DATA
+		SearchContext.DebugDrawParams.SearchCostHistoryBruteForce = &InOutMotionMatchingState.SearchCostHistoryBruteForce;
+		SearchContext.DebugDrawParams.SearchCostHistoryKDTree = &InOutMotionMatchingState.SearchCostHistoryKDTree;
+#endif
+
 		SearchContext.CurrentResult = InOutMotionMatchingState.CurrentSearchResult;
 
 		IPoseHistoryProvider* PoseHistoryProvider = Context.GetMessage<IPoseHistoryProvider>();
diff --git a/Engine/Plugins/Experimental/Animation/PoseSearch/Source/Runtime/Public/PoseSearch/PoseSearch.h b/Engine/Plugins/Experimental/Animation/PoseSearch/Source/Runtime/Public/PoseSearch/PoseSearch.h
--- a/Engine/Plugins/Experimental/Animation/PoseSearch/Source/Runtime/Public/PoseSearch/PoseSearch.h
+++ b/Engine/Plugins/Experimental/Animation/PoseSearch/Source/Runtime/Public/PoseSearch/PoseSearch.h
@@ -1128,6 +1128,11 @@ struct POSESEARCH_API FDebugDrawParams
 	// Optional prefix for sample labels
 	FStringView LabelPrefix;
 
+#if WITH_EDITORONLY_DATA
+	FDebugFloatHistory* SearchCostHistoryBruteForce = nullptr;
+	FDebugFloatHistory* SearchCostHistoryKDTree = nullptr;
+#endif
+
 	// Optional Mesh for gathering SocketTransform(s)
 	TWeakObjectPtr<const USkinnedMeshComponent> Mesh = nullptr;
 
diff --git a/Engine/Plugins/Experimental/Animation/PoseSearch/Source/Runtime/Public/PoseSearch/PoseSearchLibrary.h b/Engine/Plugins/Experimental/Animation/PoseSearch/Source/Runtime/Public/PoseSearch/PoseSearchLibrary.h
--- a/Engine/Plugins/Experimental/Animation/PoseSearch/Source/Runtime/Public/PoseSearch/PoseSearchLibrary.h
+++ b/Engine/Plugins/Experimental/Animation/PoseSearch/Source/Runtime/Public/PoseSearch/PoseSearchLibrary.h
@@ -91,6 +91,11 @@ struct POSESEARCH_API FMotionMatchingState
 	// when UE_POSE_SEARCH_TRACE_ENABLED is active
 	UPROPERTY(Transient)
 	FTransform RootMotionTransformDelta = FTransform::Identity;
+
+#if WITH_EDITORONLY_DATA
+	FDebugFloatHistory SearchCostHistoryBruteForce;
+	FDebugFloatHistory SearchCostHistoryKDTree;
+#endif
 };
 
 /**
