Motion Matching - Move nearby pose check into search loop

diff --git a/Engine/Plugins/Experimental/Animation/PoseSearch/Source/Runtime/Private/PoseSearch.cpp b/Engine/Plugins/Experimental/Animation/PoseSearch/Source/Runtime/Private/PoseSearch.cpp
--- a/Engine/Plugins/Experimental/Animation/PoseSearch/Source/Runtime/Private/PoseSearch.cpp
+++ b/Engine/Plugins/Experimental/Animation/PoseSearch/Source/Runtime/Private/PoseSearch.cpp
@@ -222,6 +222,73 @@ static void CalcChannelCosts(const UPoseSearchSchema* Schema, TArrayView<const f
 		}
 	}
 }
+
+static int32 PopulateNonSelectableIdx(size_t* NonSelectableIdx, int32 NonSelectableIdxSize, const FPoseSearchContext& SearchContext, int32 IdxOffset = 0)
+{
+	check(SearchContext.CurrentResult.IsValid());
+
+	int32 NonSelectableIdxUsedSize = 0;
+	if (SearchContext.PoseJumpThresholdTime > 0.f)
+	{
+		const UPoseSearchDatabase* Database = SearchContext.CurrentResult.Database.Get();
+		const FPoseSearchIndexAsset* CurrentIndexAsset = SearchContext.CurrentResult.SearchIndexAsset;
+		const int32 PoseJumpIndexThreshold = FMath::FloorToInt(SearchContext.PoseJumpThresholdTime / Database->Schema->SamplingInterval);
+		const bool IsLooping = Database->IsSourceAssetLooping(CurrentIndexAsset);
+
+		for (int32 i = -PoseJumpIndexThreshold; i <= -1; ++i)
+		{
+			const int32 PoseIdx = SearchContext.CurrentResult.PoseIdx + i;
+			const float DeltaTime = i * Database->Schema->SamplingInterval;
+
+			// @todo: should we use the quantized time associated to InOutMotionMatchingState.DbPoseIdx instead of InOutMotionMatchingState.AssetPlayerTime?
+			float PoseAssetPlayerTime = SearchContext.CurrentResult.AssetTime + DeltaTime;
+			if (IsLooping)
+			{
+				while (PoseAssetPlayerTime < CurrentIndexAsset->SamplingInterval.Min)
+				{
+					PoseAssetPlayerTime += CurrentIndexAsset->SamplingInterval.Size();
+				}
+			}
+
+			if (CurrentIndexAsset->SamplingInterval.Contains(PoseAssetPlayerTime) && NonSelectableIdxUsedSize < NonSelectableIdxSize)
+			{
+				NonSelectableIdx[NonSelectableIdxUsedSize++] = PoseIdx + IdxOffset;
+			}
+			else
+			{
+				break;
+			}
+		}
+
+		for (int32 i = 1; i <= PoseJumpIndexThreshold; ++i)
+		{
+			const int32 PoseIdx = SearchContext.CurrentResult.PoseIdx + i;
+			const float DeltaTime = i * Database->Schema->SamplingInterval;
+
+			// @todo: should we use the quantized time associated to InOutMotionMatchingState.DbPoseIdx instead of InOutMotionMatchingState.AssetPlayerTime?
+			float PoseAssetPlayerTime = SearchContext.CurrentResult.AssetTime + DeltaTime;
+			if (IsLooping)
+			{
+				while (PoseAssetPlayerTime > CurrentIndexAsset->SamplingInterval.Max)
+				{
+					PoseAssetPlayerTime -= CurrentIndexAsset->SamplingInterval.Size();
+				}
+			}
+
+			if (CurrentIndexAsset->SamplingInterval.Contains(PoseAssetPlayerTime) && NonSelectableIdxUsedSize < NonSelectableIdxSize)
+			{
+				NonSelectableIdx[NonSelectableIdxUsedSize++] = PoseIdx + IdxOffset;
+			}
+			else
+			{
+				break;
+			}
+		}
+	}
+
+	return NonSelectableIdxUsedSize;
+}
+
 }} // namespace UE::PoseSearch
 
 
@@ -2138,6 +2205,7 @@ UE::PoseSearch::FSearchResult UPoseSearchDatabase::Search(FPoseSearchContext& Se
 		return Result;
 	}
 
+
 	if (PoseSearchMode != EPoseSearchMode::BruteForce)
 	{
 		Result = SearchPCAKDTree(SearchContext);
@@ -2195,13 +2263,17 @@ UE::PoseSearch::FSearchResult UPoseSearchDatabase::SearchPCAKDTree(FPoseSearchCo
 	const uint32 ClampedNumberOfPrincipalComponents = GetNumberOfPrincipalComponents();
 	const uint32 ClampedKDTreeQueryNumNeighbors = FMath::Clamp<uint32>(KDTreeQueryNumNeighbors, 1, SearchIndex->NumPoses);
 
+
 	//stack allocated temporaries
 	TArrayView<size_t> ResultIndexes((size_t*)FMemory_Alloca((ClampedKDTreeQueryNumNeighbors + 1) * sizeof(size_t)), ClampedKDTreeQueryNumNeighbors + 1);
 	TArrayView<float> ResultDistanceSqr((float*)FMemory_Alloca((ClampedKDTreeQueryNumNeighbors + 1) * sizeof(float)), ClampedKDTreeQueryNumNeighbors + 1);
 	RowMajorVectorMap WeightedQueryValues((float*)FMemory_Alloca(NumDimensions * sizeof(float)), 1, NumDimensions);
 	RowMajorVectorMap CenteredQueryValues((float*)FMemory_Alloca(NumDimensions * sizeof(float)), 1, NumDimensions);
 	RowMajorVectorMap ProjectedQueryValues((float*)FMemory_Alloca(ClampedNumberOfPrincipalComponents * sizeof(float)), 1, ClampedNumberOfPrincipalComponents);
-
+	
+	constexpr int NonSelectableIdxDataSize = 128;
+	size_t* NonSelectableIdxData((size_t*)FMemory_Alloca(NonSelectableIdxDataSize * sizeof(size_t)));
+	
 	// KDTree in PCA space search
 	if (PoseSearchMode == EPoseSearchMode::PCAKDTree_Validate)
 	{
@@ -2250,7 +2322,8 @@ UE::PoseSearch::FSearchResult UPoseSearchDatabase::SearchPCAKDTree(FPoseSearchCo
 
 	TArrayView<const float> NormalizedQueryValues = Result.ComposedQuery.GetNormalizedValues();
 
-	if (SearchContext.CurrentResult.IsValid() && SearchContext.CurrentResult.Database == this)
+	const bool IsCurrentResultFromThisDatabase = SearchContext.CurrentResult.IsValid() && SearchContext.CurrentResult.Database == this;
+	if (IsCurrentResultFromThisDatabase)
 	{
 		Result.ContinuityPoseCost = ComparePoses(
 			SearchContext,
@@ -2262,8 +2335,11 @@ UE::PoseSearch::FSearchResult UPoseSearchDatabase::SearchPCAKDTree(FPoseSearchCo
 
 	for (const FGroupSearchIndex& GroupSearchIndex : SearchIndex->Groups)
 	{
+		// we're offsetting the NonSelectableIdx by "-GroupSearchIndex.StartPoseIndex" to have the indexes in kdtree space rather than in PoseIdx database space (kdtree index 0 is PoseIdx of GroupSearchIndex.StartPoseIndex)
+		const int NonSelectableIdxUsedSize = IsCurrentResultFromThisDatabase ? PopulateNonSelectableIdx(NonSelectableIdxData, NonSelectableIdxDataSize, SearchContext, -GroupSearchIndex.StartPoseIndex) : 0;
+		TArrayView<size_t> NonSelectableIdx(NonSelectableIdxData, NonSelectableIdxUsedSize);
 		const RowMajorVectorMapConst MapWeights(GroupSearchIndex.Weights.GetData(), 1, NumDimensions);
-		FKDTree::KNNResultSet ResultSet(ClampedKDTreeQueryNumNeighbors, ResultIndexes, ResultDistanceSqr);
+		FKDTree::KNNResultSet ResultSet(ClampedKDTreeQueryNumNeighbors, ResultIndexes, ResultDistanceSqr, NonSelectableIdx);
 
 		check(NormalizedQueryValues.Num() == NumDimensions);
 
@@ -2335,7 +2411,11 @@ UE::PoseSearch::FSearchResult UPoseSearchDatabase::SearchBruteForce(FPoseSearchC
 	FPoseSearchWeightsContext WeightsContext;
 	WeightsContext.Update(this);
 	
-	if (SearchContext.CurrentResult.IsValid() && SearchContext.CurrentResult.Database == this)
+	constexpr int NonSelectableIdxDataSize = 128;
+	size_t* NonSelectableIdxData((size_t*)FMemory_Alloca(NonSelectableIdxDataSize * sizeof(size_t)));
+	int NonSelectableIdxUsedSize = 0;
+	const bool IsCurrentResultFromThisDatabase = SearchContext.CurrentResult.IsValid() && SearchContext.CurrentResult.Database == this;
+	if (IsCurrentResultFromThisDatabase)
 	{
 		Result.ContinuityPoseCost = ComparePoses(
 			SearchContext,
@@ -2343,8 +2423,12 @@ UE::PoseSearch::FSearchResult UPoseSearchDatabase::SearchBruteForce(FPoseSearchC
 			SearchContext.CurrentResult.PoseIdx,
 			SearchIndex->FindAssetForPose(SearchContext.CurrentResult.PoseIdx)->SourceGroupIdx,
 			NormalizedQueryValues);
+
+		NonSelectableIdxUsedSize = PopulateNonSelectableIdx(NonSelectableIdxData, NonSelectableIdxDataSize, SearchContext);
 	}
 
+	TArrayView<size_t> NonSelectableIdx(NonSelectableIdxData, NonSelectableIdxUsedSize);
+
 	FPoseSearchCost BestPoseCost;
 	int32 BestPoseIdx = INDEX_NONE;
 	for (const FPoseSearchIndexAsset& Asset : SearchIndex->Assets)
@@ -2357,6 +2441,7 @@ UE::PoseSearch::FSearchResult UPoseSearchDatabase::SearchBruteForce(FPoseSearchC
 			}
 		}
 
+		const bool CheckForNonSelectableIdx = IsCurrentResultFromThisDatabase && (&Asset == SearchContext.CurrentResult.SearchIndexAsset);
 		const int32 EndIndex = Asset.FirstPoseIdx + Asset.NumPoses;
 		for (int32 PoseIdx = Asset.FirstPoseIdx; PoseIdx < EndIndex; ++PoseIdx)
 		{
@@ -2367,6 +2452,11 @@ UE::PoseSearch::FSearchResult UPoseSearchDatabase::SearchBruteForce(FPoseSearchC
 				continue;
 			}
 
+			if (CheckForNonSelectableIdx && NonSelectableIdx.Contains(PoseIdx))
+			{
+				continue;
+			}
+
 			FPoseSearchCost PoseCost = ComparePoses(
 				SearchContext,
 				WeightsContext,
diff --git a/Engine/Plugins/Experimental/Animation/PoseSearch/Source/Runtime/Private/PoseSearchLibrary.cpp b/Engine/Plugins/Experimental/Animation/PoseSearch/Source/Runtime/Private/PoseSearchLibrary.cpp
--- a/Engine/Plugins/Experimental/Animation/PoseSearch/Source/Runtime/Private/PoseSearchLibrary.cpp
+++ b/Engine/Plugins/Experimental/Animation/PoseSearch/Source/Runtime/Private/PoseSearchLibrary.cpp
@@ -222,6 +222,8 @@ void UpdateMotionMatchingState(
 	FMotionMatchingState& InOutMotionMatchingState
 )
 {
+	QUICK_SCOPE_CYCLE_COUNTER(STAT_PoseSearch_Update);
+
 	using namespace UE::PoseSearch;
 
 	if (!Searchable)
@@ -262,6 +264,7 @@ void UpdateMotionMatchingState(
 #endif
 
 		SearchContext.CurrentResult = InOutMotionMatchingState.CurrentSearchResult;
+		SearchContext.PoseJumpThresholdTime = Settings.PoseJumpThresholdTime;
 
 		IPoseHistoryProvider* PoseHistoryProvider = Context.GetMessage<IPoseHistoryProvider>();
 		if (PoseHistoryProvider)
@@ -312,33 +315,8 @@ void UpdateMotionMatchingState(
 					}
 				}
 
-				// Ignore the candidate poses from the same anim when they are too near to the current pose
-				bool bNearbyPose = false;
-				const FPoseSearchIndexAsset* StateSearchIndexAsset = InOutMotionMatchingState.GetCurrentSearchIndexAsset();
-				if (StateSearchIndexAsset == SearchResult.SearchIndexAsset)
-				{
-					// We need to check in terms of PoseIdx rather than AssetTime because
-					// for blendspaces, AssetTime is not in seconds, but in the normalized range 
-					// [0, 1] so comparing to `PoseJumpThresholdTime` will not make sense		
-					bNearbyPose = 
-						FMath::Abs(InOutMotionMatchingState.CurrentSearchResult.PoseIdx - SearchResult.PoseIdx) *
-						SearchResult.Database->Schema->SamplingInterval < Settings.PoseJumpThresholdTime;
-
-					// Handle looping anims when checking for the pose being too close
-					if (!bNearbyPose && SearchResult.Database->IsSourceAssetLooping(StateSearchIndexAsset))
-					{
-						const float Time =
-							FMath::Abs(
-								StateSearchIndexAsset->NumPoses -
-								InOutMotionMatchingState.CurrentSearchResult.PoseIdx -
-								SearchResult.PoseIdx) *
-							SearchResult.Database->Schema->SamplingInterval;
-						bNearbyPose = Time < Settings.PoseJumpThresholdTime;
-					}
-				}
-
 				// Jump to candidate pose if there was a better option
-				if (bBetterPose && !bNearbyPose)
+				if (bBetterPose)
 				{
 					InOutMotionMatchingState.JumpToPose(Context, Settings, SearchResult);
 				}
diff --git a/Engine/Plugins/Experimental/Animation/PoseSearch/Source/Runtime/Public/PoseSearch/KDTree.h b/Engine/Plugins/Experimental/Animation/PoseSearch/Source/Runtime/Public/PoseSearch/KDTree.h
--- a/Engine/Plugins/Experimental/Animation/PoseSearch/Source/Runtime/Public/PoseSearch/KDTree.h
+++ b/Engine/Plugins/Experimental/Animation/PoseSearch/Source/Runtime/Public/PoseSearch/KDTree.h
@@ -42,16 +42,17 @@ struct POSESEARCH_API FKDTree
 
 	struct KNNResultSet
 	{
-		inline KNNResultSet(size_t InNumNeighbors, TArrayView<size_t> IndexesView, TArrayView<float> DistancesView)
-		: Indexes(IndexesView.GetData())
-		, Distances(DistancesView.GetData())
+		inline KNNResultSet(size_t InNumNeighbors, TArrayView<size_t> InIndexes, TArrayView<float> InDistances, TArrayView<size_t> InExcludeFromSearchIndexes = TArrayView<size_t>())
+		: Indexes(InIndexes)
+		, Distances(InDistances)
 		, NumNeighbors(InNumNeighbors)
 		, Count(0)
+		, ExcludeFromSearchIndexes(InExcludeFromSearchIndexes)
 		{
 			// by having IndexesView and DistancesView cardinality bigger than NumNeighbors, we can skip some if statements in the addPoint method
 			check(NumNeighbors > 0);
-			check(IndexesView.Num() > NumNeighbors);
-			check(DistancesView.Num() > NumNeighbors);
+			check(InIndexes.Num() > NumNeighbors);
+			check(InDistances.Num() > NumNeighbors);
 
 			Distances[NumNeighbors - 1] = UE_BIG_NUMBER;
 		}
@@ -68,6 +69,9 @@ struct POSESEARCH_API FKDTree
 
 		inline bool addPoint(size_t dist, size_t index)
 		{
+			if (ExcludeFromSearchIndexes.Contains(index))
+				return true;
+
 			// shifting Distances[i] and Indexes[i] to make space for "dist" and "index" at the right "i"th slot
 			size_t i;
 			for (i = Count; (i > 0) && (Distances[i - 1] > dist); --i)
@@ -94,12 +98,14 @@ struct POSESEARCH_API FKDTree
 		inline float worstDist() const { return Distances[NumNeighbors - 1]; }
 
 	private:
-		size_t* Indexes;
-		float* Distances;
+		TArrayView<size_t> Indexes;
+		TArrayView<float> Distances;
 		size_t NumNeighbors;
 		size_t Count;
+		const TArrayView<size_t> ExcludeFromSearchIndexes; // @todo: perhaps make it an hash if it gets too big
 	};
 	
+
 	FKDTree(int32 Count, int32 Dim, const float* Data, int32 MaxLeafSize);
 	FKDTree();
 	FKDTree(const FKDTree& r);
diff --git a/Engine/Plugins/Experimental/Animation/PoseSearch/Source/Runtime/Public/PoseSearch/PoseSearch.h b/Engine/Plugins/Experimental/Animation/PoseSearch/Source/Runtime/Public/PoseSearch/PoseSearch.h
--- a/Engine/Plugins/Experimental/Animation/PoseSearch/Source/Runtime/Public/PoseSearch/PoseSearch.h
+++ b/Engine/Plugins/Experimental/Animation/PoseSearch/Source/Runtime/Public/PoseSearch/PoseSearch.h
@@ -1517,6 +1517,7 @@ public:
 	UE::PoseSearch::FSearchResult CurrentResult;
 	const FBoneContainer* BoneContainer = nullptr;
 	const FGameplayTagContainer* ActiveTagsContainer = nullptr;
+	float PoseJumpThresholdTime = 0.f;
 };
 
 
