motion matching - deprecating FFeatureTypeTraits and start commonizing FPoseSearchFeatureVectorLayout with UPoseSearchFeatureChannel

diff --git a/Engine/Plugins/Experimental/Animation/PoseSearch/Source/Runtime/Private/PoseSearch.cpp b/Engine/Plugins/Experimental/Animation/PoseSearch/Source/Runtime/Private/PoseSearch.cpp
--- a/Engine/Plugins/Experimental/Animation/PoseSearch/Source/Runtime/Private/PoseSearch.cpp
+++ b/Engine/Plugins/Experimental/Animation/PoseSearch/Source/Runtime/Private/PoseSearch.cpp
@@ -168,57 +168,17 @@ FORCEINLINE auto LowerBound(IteratorType First, IteratorType Last, const ValueTy
 	return LowerBound(First, Last, Value, FIdentityFunctor(), TLess<>());
 }
 
-
-//////////////////////////////////////////////////////////////////////////
-// FFeatureTypeTraits
-
-struct FFeatureTypeTraits
-{
-	EPoseSearchFeatureType Type = EPoseSearchFeatureType::Invalid;
-	uint32 NumFloats = 0;
-};
-
-// Could upgrade to class objects in the future with value reader/writer functions
-static constexpr FFeatureTypeTraits FeatureTypeTraits[] =
-{
-	{ EPoseSearchFeatureType::Position, 3 },
-	{ EPoseSearchFeatureType::Rotation, 6 },
-	{ EPoseSearchFeatureType::LinearVelocity, 3 },
-	{ EPoseSearchFeatureType::AngularVelocity, 3 },
-	{ EPoseSearchFeatureType::ForwardVector, 3 },
-	{ EPoseSearchFeatureType::Phase, 2 },
-};
-
-static FFeatureTypeTraits GetFeatureTypeTraits(EPoseSearchFeatureType Type)
-{
-	// Could allow external registration to a TSet of traits in the future
-	// For now just use a simple local array
-	for (const FFeatureTypeTraits& Traits : FeatureTypeTraits)
-	{
-		if (Traits.Type == Type)
-		{
-			return Traits;
-		}
-	}
-
-	return FFeatureTypeTraits();
-}
-
 static void CalcChannelCosts(const UPoseSearchSchema* Schema, TArrayView<const float> CostVector, TArray<float>& OutChannelCosts)
 {
 	OutChannelCosts.Reset();
 	OutChannelCosts.AddZeroed(Schema->Channels.Num());
 	for (int ChannelIdx = 0; ChannelIdx != Schema->Channels.Num(); ++ChannelIdx)
 	{
-		for (int FeatureIdx = INDEX_NONE; Schema->Layout.EnumerateBy(ChannelIdx, EPoseSearchFeatureType::Invalid, FeatureIdx); /*empty*/)
+		const UPoseSearchFeatureChannel* Channel = Schema->Channels[ChannelIdx].Get();
+		const int32 ValueTerm = Channel->GetChannelDataOffset() + Channel->GetChannelCardinality();
+		for (int32 ValueIdx = Channel->GetChannelDataOffset(); ValueIdx != ValueTerm; ++ValueIdx)
 		{
-			const FPoseSearchFeatureDesc& Feature = Schema->Layout.Features[FeatureIdx];
-			int32 ValueSize = GetFeatureTypeTraits(Feature.Type).NumFloats;
-			int32 ValueTerm = Feature.ValueOffset + ValueSize;
-			for (int32 ValueIdx = Feature.ValueOffset; ValueIdx != ValueTerm; ++ValueIdx)
-			{
-				OutChannelCosts[ChannelIdx] += CostVector[ValueIdx];
-			}
+			OutChannelCosts[ChannelIdx] += CostVector[ValueIdx];
 		}
 	}
 }
@@ -292,6 +252,14 @@ static int32 PopulateNonSelectableIdx(size_t* NonSelectableIdx, int32 NonSelecta
 }} // namespace UE::PoseSearch
 
 
+//////////////////////////////////////////////////////////////////////////
+// UPoseSearchFeatureChannel
+void UPoseSearchFeatureChannel::InitializeSchema(UE::PoseSearch::FSchemaInitializer& Initializer)
+{
+	ChannelIdx = Initializer.GetCurrentChannelIdx();
+	ChannelDataOffset = Initializer.GetCurrentChannelDataOffset();
+}
+
 //////////////////////////////////////////////////////////////////////////
 // FPoseSearchFeatureDesc
 
@@ -301,6 +269,8 @@ bool FPoseSearchFeatureDesc::operator==(const FPoseSearchFeatureDesc& Other) con
 		(ChannelIdx == Other.ChannelIdx) &&
 		(ChannelFeatureId == Other.ChannelFeatureId) &&
 		(SubsampleIdx == Other.SubsampleIdx) &&
+		// ignored or else Schema->Layout.Features.Find(Feature); called by FPoseSearchFeatureVectorBuilder will fail since FPoseSearchFeatureDesc doesn't get created fully :(
+		//(Cardinality == Other.Cardinality) &&
 		(Type == Other.Type);
 }
 
@@ -308,28 +278,6 @@ bool FPoseSearchFeatureDesc::operator==(const FPoseSearchFeatureDesc& Other) con
 //////////////////////////////////////////////////////////////////////////
 // FPoseSearchFeatureVectorLayout
 
-void FPoseSearchFeatureVectorLayout::Finalize()
-{
-	uint32 FloatCount = 0;
-
-	// Initialize value offsets
-	for (FPoseSearchFeatureDesc& Feature : Features)
-	{
-		Feature.ValueOffset = FloatCount;
-
-		uint32 FeatureNumFloats = UE::PoseSearch::GetFeatureTypeTraits(Feature.Type).NumFloats;
-		FloatCount += FeatureNumFloats;
-	}
-
-	NumFloats = FloatCount;
-}
-
-void FPoseSearchFeatureVectorLayout::Reset()
-{
-	Features.Reset();
-	NumFloats = 0;
-}
-
 bool FPoseSearchFeatureVectorLayout::IsValid(int32 ChannelCount) const
 {
 	if (NumFloats == 0)
@@ -386,35 +334,6 @@ bool FPoseSearchFeatureVectorLayout::EnumerateBy(int32 ChannelIdx, EPoseSearchFe
 	return false;
 }
 
-
-//////////////////////////////////////////////////////////////////////////
-// FPoseSearchBone
-
-uint32 FPoseSearchBone::GetTypeMask() const
-{
-	uint32 Mask = 0;
-	
-	if (bUsePosition)
-	{
-		Mask |= 1 << static_cast<int>(EPoseSearchFeatureType::Position);
-	}
-	if (bUseVelocity)
-	{
-		Mask |= 1 << static_cast<int>(EPoseSearchFeatureType::LinearVelocity);
-	}
-	if (bUseRotation)
-	{
-		Mask |= 1 << static_cast<int>(EPoseSearchFeatureType::Rotation);
-	}
-	if (bUsePhase)
-	{
-		Mask |= 1 << static_cast<int>(EPoseSearchFeatureType::Phase);
-	}
-
-	return Mask;
-}
-
-
 //////////////////////////////////////////////////////////////////////////
 // UPoseSearchSchema
 
@@ -426,14 +345,48 @@ int32 FSchemaInitializer::AddBoneReference(const FBoneReference& BoneReference)
 	return BoneReferences.AddUnique(BoneReference);
 }
 
+constexpr uint32 GetPoseSearchFeatureTypeNumFloats(EPoseSearchFeatureType Type)
+{
+	switch (Type)
+	{
+	case EPoseSearchFeatureType::Position: return 3;
+	case EPoseSearchFeatureType::Rotation: return 6;
+	case EPoseSearchFeatureType::LinearVelocity: return 3;
+	case EPoseSearchFeatureType::AngularVelocity: return 3;
+	case EPoseSearchFeatureType::ForwardVector: return 3;
+	case EPoseSearchFeatureType::Phase: return 2;
+	default: return 0;
+	}
+}
+
 int32 FSchemaInitializer::AddFeatureDesc(const FPoseSearchFeatureDesc& FeatureDesc)
 {
 	check(FeatureDesc.ChannelIdx == CurrentChannelIdx);
 	check(Features.Num() + 1 <= UPoseSearchSchema::MaxFeatures);
 	check(!Features.Contains(FeatureDesc));
-	return Features.Add(FeatureDesc);
+	const int32 FeatureDescIndex = Features.Add(FeatureDesc);
+
+	const int32 NumFloats = GetPoseSearchFeatureTypeNumFloats(FeatureDesc.Type);
+
+	Features[FeatureDescIndex].ValueOffset = CurrentChannelDataOffset;
+	Features[FeatureDescIndex].Cardinality = NumFloats;
+	CurrentChannelDataOffset += NumFloats;
+	return FeatureDescIndex;
 }
 
+void FSchemaInitializer::AddFeatures(int8 ChannelIdx, EPoseSearchFeatureType Type, int32 ChannelFeatureId, int32 NumSampleOffsets)
+{
+	FPoseSearchFeatureDesc Feature;
+	Feature.ChannelIdx = ChannelIdx;
+	Feature.Type = Type;
+	Feature.ChannelFeatureId = ChannelFeatureId;
+
+	for (Feature.SubsampleIdx = 0; Feature.SubsampleIdx != NumSampleOffsets; ++Feature.SubsampleIdx)
+	{
+		AddFeatureDesc(Feature);
+	}
+};
+
 } // namespace UE::PoseSearch
 
 void UPoseSearchSchema::Finalize()
@@ -446,23 +399,21 @@ void UPoseSearchSchema::Finalize()
 	// Discard null channels
 	Channels.RemoveAll([](TObjectPtr<UPoseSearchFeatureChannel>& Channel) { return Channel.IsNull(); });
 
-	Layout.Reset();
 	BoneReferences.Reset();
 
 	FSchemaInitializer Initializer;
 	for (int32 ChannelIdx = 0; ChannelIdx != Channels.Num(); ++ChannelIdx)
 	{
 		TObjectPtr<UPoseSearchFeatureChannel>& Channel = Channels[ChannelIdx];
-		Channel->ChannelIdx = ChannelIdx;
-
 		Initializer.CurrentChannelIdx = ChannelIdx;
 		Channel->InitializeSchema(Initializer);
 	}
 
+	// @todo: replace FPoseSearchFeatureVectorLayout with channels after itemizing UPoseSearchFeatureChannel_Pose into position, velocity, orientation / heading, phase etc 
 	Layout.Features = MoveTemp(Initializer.Features);
-	BoneReferences = MoveTemp(Initializer.BoneReferences);
+	Layout.NumFloats = Initializer.GetCurrentChannelDataOffset();
 
-	Layout.Finalize();
+	BoneReferences = MoveTemp(Initializer.BoneReferences);
 
 	EffectiveDataPreprocessor = DataPreprocessor;
 	if (EffectiveDataPreprocessor == EPoseSearchDataPreprocessor::Automatic)
@@ -530,6 +481,31 @@ void UPoseSearchSchema::PostLoad()
 		SampledBones_DEPRECATED.Empty();
 		PoseSampleTimes_DEPRECATED.Empty();
 	}
+
+	if (!bNeedFinalize)
+	{
+		for (const FPoseSearchFeatureDesc& Feature : Layout.Features)
+		{
+			if (Feature.Cardinality == 0)
+			{
+				bNeedFinalize = true;
+				break;
+			}
+		}
+	}
+
+	if (!bNeedFinalize)
+	{
+		for (const UPoseSearchFeatureChannel* Channel : Channels)
+		{
+			if (Channel->ChannelDataOffset == -1 || Channel->ChannelCardinality == -1)
+			{
+				bNeedFinalize = true;
+				break;
+			}
+		}
+	}
+
 #endif // WITH_EDITORONLY_DATA
 
 	if (bNeedFinalize)
@@ -925,8 +901,7 @@ void FPoseSearchWeights::Init(const FPoseSearchWeightParams& WeightParams, const
 		for (int FeatureIdx = INDEX_NONE; Schema->Layout.EnumerateBy(ChannelIdx, EPoseSearchFeatureType::Invalid, FeatureIdx); /*empty*/)
 		{
 			const FPoseSearchFeatureDesc& Feature = Schema->Layout.Features[FeatureIdx];
-			int32 ValueSize = GetFeatureTypeTraits(Feature.Type).NumFloats;
-			int32 ValueTerm = Feature.ValueOffset + ValueSize;
+			const int32 ValueTerm = Feature.ValueOffset + Feature.Cardinality;
 			for (int32 ValueIdx = Feature.ValueOffset; ValueIdx != ValueTerm; ++ValueIdx)
 			{
 				Weights[ValueIdx] = WeightsByFeature[FeatureIdx];
@@ -2734,7 +2709,7 @@ void FPoseSearchFeatureVectorBuilder::SetVector(FPoseSearchFeatureDesc Feature,
 	int32 ElementIndex = Schema->Layout.Features.Find(Feature);
 	if (ElementIndex >= 0)
 	{
-		const FPoseSearchFeatureDesc& FoundElement = Schema->Layout.Features[ElementIndex];
+		const FPoseSearchFeatureDesc& FoundElement = Schema->Layout.Features[ElementIndex]; 
 
 		Values[FoundElement.ValueOffset + 0] = Vector[0];
 		Values[FoundElement.ValueOffset + 1] = Vector[1];
@@ -2762,37 +2737,6 @@ void FPoseSearchFeatureVectorBuilder::CopyFromSearchIndex(const FPoseSearchIndex
 	FeaturesAdded.SetRange(0, FeaturesAdded.Num(), true);
 }
 
-void FPoseSearchFeatureVectorBuilder::CopyFeature(const FPoseSearchFeatureVectorBuilder& OtherBuilder, int32 FeatureIdx)
-{
-	check(IsCompatible(OtherBuilder));
-	check(OtherBuilder.FeaturesAdded[FeatureIdx]);
-
-	const FPoseSearchFeatureDesc& FeatureDesc = Schema->Layout.Features[FeatureIdx];
-	const int32 FeatureNumFloats = UE::PoseSearch::GetFeatureTypeTraits(FeatureDesc.Type).NumFloats;
-	const int32 FeatureValueOffset = FeatureDesc.ValueOffset;
-
-	for(int32 FeatureValueIdx = FeatureValueOffset; FeatureValueIdx != FeatureValueOffset + FeatureNumFloats; ++FeatureValueIdx)
-	{
-		Values[FeatureValueIdx] = OtherBuilder.Values[FeatureValueIdx];
-	}
-
-	if (!FeaturesAdded[FeatureIdx])
-	{
-		FeaturesAdded[FeatureIdx] = true;
-		++NumFeaturesAdded;
-	}
-}
-
-void FPoseSearchFeatureVectorBuilder::MergeReplace(const FPoseSearchFeatureVectorBuilder& OtherBuilder)
-{
-	check(IsCompatible(OtherBuilder));
-
-	for (TConstSetBitIterator<> Iter(OtherBuilder.FeaturesAdded); Iter; ++Iter)
-	{
-		CopyFeature(OtherBuilder, Iter.GetIndex());
-	}
-}
-
 bool FPoseSearchFeatureVectorBuilder::IsInitialized() const
 {
 	return (Schema != nullptr) && (Values.Num() == Schema->Layout.NumFloats);
@@ -4743,17 +4687,15 @@ static inline Eigen::VectorXd ComputeFeatureMeanDeviations(const Eigen::MatrixXd
 	MeanDeviations.setConstant(1.0);
 	for (const FPoseSearchFeatureDesc& Feature : Layout.Features)
 	{
-		int32 FeatureDims = GetFeatureTypeTraits(Feature.Type).NumFloats;
-
 		// Construct a submatrix for the feature and find the average distance to the feature's centroid.
 		// Since we've already mean centered the data, the average distance to the centroid is simply the average norm.
-		double FeatureMeanDeviation = CenteredPoseMatrix.block(Feature.ValueOffset, 0, FeatureDims, NumPoses).colwise().norm().mean();
+		double FeatureMeanDeviation = CenteredPoseMatrix.block(Feature.ValueOffset, 0, Feature.Cardinality, NumPoses).colwise().norm().mean();
 
 		// Fill the feature's corresponding scaling axes with the average distance
 		// Avoid scaling by zero by leaving near-zero deviations as 1.0
 		if (FeatureMeanDeviation > KINDA_SMALL_NUMBER)
 		{
-			MeanDeviations.segment(Feature.ValueOffset, FeatureDims).setConstant(FeatureMeanDeviation);
+			MeanDeviations.segment(Feature.ValueOffset, Feature.Cardinality).setConstant(FeatureMeanDeviation);
 		}
 	}
 
diff --git a/Engine/Plugins/Experimental/Animation/PoseSearch/Source/Runtime/Private/PoseSearchFeatureChannels.cpp b/Engine/Plugins/Experimental/Animation/PoseSearch/Source/Runtime/Private/PoseSearchFeatureChannels.cpp
--- a/Engine/Plugins/Experimental/Animation/PoseSearch/Source/Runtime/Private/PoseSearchFeatureChannels.cpp
+++ b/Engine/Plugins/Experimental/Animation/PoseSearch/Source/Runtime/Private/PoseSearchFeatureChannels.cpp
@@ -263,8 +263,6 @@ static void CalculatePhasesFromLocalMinMax(const TArray<LocalMinMax>& MinMax, TA
 
 } // namespace UE::PoseSearch
 
-
-
 //////////////////////////////////////////////////////////////////////////
 // UPoseSearchFeatureChannel
 class USkeleton* UPoseSearchFeatureChannel::GetSkeleton(bool& bInvalidSkeletonIsError, const IPropertyHandle* PropertyHandle)
@@ -288,38 +286,34 @@ void UPoseSearchFeatureChannel_Pose::PreSave(FObjectPreSaveContext ObjectSaveCon
 
 void UPoseSearchFeatureChannel_Pose::InitializeSchema(UE::PoseSearch::FSchemaInitializer& Initializer)
 {
-	auto AddFeatures = [this, &Initializer](EPoseSearchFeatureType Type)
-	{
-		FPoseSearchFeatureDesc Feature;
-		Feature.Type = Type;
-		Feature.ChannelIdx = GetChannelIndex();
+	Super::InitializeSchema(Initializer);
 
-		const int32 NumBones = SampledBones.Num();
-		for (int32 SubsampleIdx = 0; SubsampleIdx != SampleTimes.Num(); ++SubsampleIdx)
+	const int32 NumBones = SampledBones.Num();
+	for (int32 ChannelBoneIdx = 0; ChannelBoneIdx != NumBones; ++ChannelBoneIdx)
+	{
+		const FPoseSearchBone& SampledBone = SampledBones[ChannelBoneIdx];
+		if (SampledBone.bUsePosition)
 		{
-			Feature.SubsampleIdx = SubsampleIdx;
-
-			for (int32 ChannelBoneIdx = 0; ChannelBoneIdx != NumBones; ++ChannelBoneIdx)
-			{
-				Feature.ChannelFeatureId = ChannelBoneIdx;
-
-				const FPoseSearchBone& SampledBone = SampledBones[ChannelBoneIdx];
-				bool bAddFeature = (SampledBone.GetTypeMask() & (1 << static_cast<int>(Type))) != 0;
-				if (bAddFeature)
-				{
-					Initializer.AddFeatureDesc(Feature);
-				}
-			}
+			Initializer.AddFeatures(GetChannelIndex(), EPoseSearchFeatureType::Position, ChannelBoneIdx, SampleTimes.Num());
+		}
+		if (SampledBone.bUseRotation)
+		{
+			Initializer.AddFeatures(GetChannelIndex(), EPoseSearchFeatureType::Rotation, ChannelBoneIdx, SampleTimes.Num());
+		}
+		if (SampledBone.bUseVelocity)
+		{
+			Initializer.AddFeatures(GetChannelIndex(), EPoseSearchFeatureType::LinearVelocity, ChannelBoneIdx, SampleTimes.Num());
+		}
+		if (SampledBone.bUsePhase)
+		{
+			Initializer.AddFeatures(GetChannelIndex(), EPoseSearchFeatureType::Phase, ChannelBoneIdx, SampleTimes.Num());
 		}
-	};
-
-	for (int32 FeatureType = 0; FeatureType != (int32)EPoseSearchFeatureType::Num; ++FeatureType)
-	{
-		AddFeatures(EPoseSearchFeatureType(FeatureType));
 	}
 
+	ChannelCardinality = Initializer.GetCurrentCardinalityFrom(ChannelDataOffset);
+
 	FeatureParams.Reset();
-	for (const FPoseSearchBone& Bone: SampledBones)
+	for (const FPoseSearchBone& Bone : SampledBones)
 	{
 		FPoseSearchPoseFeatureInfo FeatureInfo;
 		FeatureInfo.SchemaBoneIdx = Initializer.AddBoneReference(Bone.Reference);
@@ -819,33 +813,24 @@ void UPoseSearchFeatureChannel_Trajectory::PreSave(FObjectPreSaveContext ObjectS
 
 void UPoseSearchFeatureChannel_Trajectory::InitializeSchema( UE::PoseSearch::FSchemaInitializer& Initializer)
 {
-	auto AddFeatures = [this, &Initializer](int32 NumSampleOffsets, EPoseSearchFeatureType Type)
-	{
-		FPoseSearchFeatureDesc Feature;
-		Feature.ChannelIdx = GetChannelIndex();
-		Feature.Type = Type;
-		Feature.ChannelFeatureId = 0; // Unused
-
-		for (Feature.SubsampleIdx = 0; Feature.SubsampleIdx != NumSampleOffsets; ++Feature.SubsampleIdx)
-		{
-			Initializer.AddFeatureDesc(Feature);
-		}
-	};
+	Super::InitializeSchema(Initializer);
 
 	if (bUsePositions)
 	{
-		AddFeatures(SampleOffsets.Num(), EPoseSearchFeatureType::Position);
+		Initializer.AddFeatures(GetChannelIndex(), EPoseSearchFeatureType::Position, 0, SampleOffsets.Num());
 	}
 
 	if (bUseLinearVelocities)
 	{
-		AddFeatures(SampleOffsets.Num(), EPoseSearchFeatureType::LinearVelocity);
+		Initializer.AddFeatures(GetChannelIndex(), EPoseSearchFeatureType::LinearVelocity, 0, SampleOffsets.Num());
 	}
 
 	if (bUseFacingDirections)
 	{
-		AddFeatures(SampleOffsets.Num(), EPoseSearchFeatureType::ForwardVector);
+		Initializer.AddFeatures(GetChannelIndex(), EPoseSearchFeatureType::ForwardVector, 0, SampleOffsets.Num());
 	}
+
+	ChannelCardinality = Initializer.GetCurrentCardinalityFrom(ChannelDataOffset);
 }
 
 void UPoseSearchFeatureChannel_Trajectory::IndexAsset(const UE::PoseSearch::IAssetIndexer& Indexer,  UE::PoseSearch::FAssetIndexingOutput& IndexingOutput) const
diff --git a/Engine/Plugins/Experimental/Animation/PoseSearch/Source/Runtime/Public/PoseSearch/PoseSearch.h b/Engine/Plugins/Experimental/Animation/PoseSearch/Source/Runtime/Public/PoseSearch/PoseSearch.h
--- a/Engine/Plugins/Experimental/Animation/PoseSearch/Source/Runtime/Public/PoseSearch/PoseSearch.h
+++ b/Engine/Plugins/Experimental/Animation/PoseSearch/Source/Runtime/Public/PoseSearch/PoseSearch.h
@@ -199,10 +199,6 @@ struct POSESEARCH_API FPoseSearchBone
 
 	UPROPERTY(EditAnywhere, Category = Config)
 	bool bUsePhase = false;
-
-	// this function will return a mask out of EPoseSearchFeatureType based on which features were selected for
-	// the bone.
-	uint32 GetTypeMask() const;
 };
 
 
@@ -232,17 +228,13 @@ struct POSESEARCH_API FPoseSearchFeatureDesc
 	UPROPERTY()
 	EPoseSearchFeatureType Type = EPoseSearchFeatureType::Invalid;
 
-	// Set via FPoseSearchFeatureLayout::Init() and ignored by operator==
+	// Set via FSchemaInitializer::AddFeatureDesc() and ignored by operator==
 	UPROPERTY()
 	int16 ValueOffset = 0;
 
-#if WITH_EDITORONLY_DATA
+	// Set via FSchemaInitializer::AddFeatureDesc() and ignored by operator== or else Schema->Layout.Features.Find(Feature); called by FPoseSearchFeatureVectorBuilder will fail
 	UPROPERTY()
-		EPoseSearchFeatureDomain Domain_DEPRECATED = EPoseSearchFeatureDomain::Time;
-
-	UPROPERTY()
-	int32 SchemaBoneIdx_DEPRECATED = 0;
-#endif
+	int16 Cardinality = 0;
 
 	bool operator==(const FPoseSearchFeatureDesc& Other) const;
 
@@ -260,9 +252,6 @@ struct POSESEARCH_API FPoseSearchFeatureVectorLayout
 {
 	GENERATED_BODY()
 
-	void Finalize();
-	void Reset();
-
 	UPROPERTY()
 	TArray<FPoseSearchFeatureDesc> Features;
 
@@ -406,9 +395,12 @@ public:
 
 	int32 GetChannelIndex() const { checkSlow(ChannelIdx >= 0); return ChannelIdx; }
 
-	// Called during UPoseSearchSchema::Finalize to prepare the schema for this channel
-	virtual void InitializeSchema(UE::PoseSearch::FSchemaInitializer& Initializer) PURE_VIRTUAL(UPoseSearchFeatureChannel::InitializeSchema, );
+	int32 GetChannelCardinality() const { checkSlow(ChannelCardinality >= 0); return ChannelCardinality; }
+	int32 GetChannelDataOffset() const { checkSlow(ChannelDataOffset >= 0); return ChannelDataOffset; }
 
+	// Called during UPoseSearchSchema::Finalize to prepare the schema for this channel
+	virtual void InitializeSchema(UE::PoseSearch::FSchemaInitializer& Initializer);
+	
 	// Called at database build time to populate pose vectors with this channel's data
 	virtual void IndexAsset(const UE::PoseSearch::IAssetIndexer& Indexer, UE::PoseSearch::FAssetIndexingOutput& IndexingOutput) const PURE_VIRTUAL(UPoseSearchFeatureChannel::IndexAsset, );
 
@@ -442,6 +434,15 @@ private:
 
 	UPROPERTY()
 	int32 ChannelIdx = -1;
+
+protected:
+	// WIP: mimicking the FPoseSearchFeatureDesc to be able to deprecate it in favor of UPoseSearchFeatureChannel later on 
+	UPROPERTY()
+	int32 ChannelDataOffset = -1;
+
+	// WIP: mimicking the FPoseSearchFeatureDesc to be able to deprecate it in favor of UPoseSearchFeatureChannel later on 
+	UPROPERTY()
+	int32 ChannelCardinality = -1;
 };
 
 
@@ -454,15 +455,23 @@ namespace UE::PoseSearch {
 struct POSESEARCH_API FSchemaInitializer
 {
 public:
-
 	int32 AddBoneReference(const FBoneReference& BoneReference);
-	int32 AddFeatureDesc(const FPoseSearchFeatureDesc& FeatureDesc);
+
+	// Gets the index into the schema's channel array for the channel currently being initialized
+	int32 GetCurrentChannelIdx() const { return CurrentChannelIdx; }
+
+	int32 GetCurrentCardinalityFrom(int32 ChannelDataOffset) const { check(CurrentChannelDataOffset >= ChannelDataOffset);  return CurrentChannelDataOffset - ChannelDataOffset; }
+	void AddFeatures(int8 ChannelIdx, EPoseSearchFeatureType Type, int32 ChannelFeatureId, int32 NumSampleOffsets);
+	int32 GetCurrentChannelDataOffset() const { return CurrentChannelDataOffset; }
 
 private:
+	int32 AddFeatureDesc(const FPoseSearchFeatureDesc& FeatureDesc);
+
 	friend class ::UPoseSearchSchema;
 
 	int32 CurrentChannelIdx = 0;
-
+	int32 CurrentChannelDataOffset = 0;
+	
 	TArray<FBoneReference> BoneReferences;
 	TArray<FPoseSearchFeatureDesc> Features;
 };
@@ -1040,9 +1049,6 @@ public:
 	void SetPhase(FPoseSearchFeatureDesc Feature, const FVector2D& Phase);
 
 	void CopyFromSearchIndex(const FPoseSearchIndex& SearchIndex, int32 PoseIdx);
-	void CopyFeature(const FPoseSearchFeatureVectorBuilder& OtherBuilder, int32 FeatureIdx);
-
-	void MergeReplace(const FPoseSearchFeatureVectorBuilder& OtherBuilder);
 
 	bool IsInitialized() const;
 	bool IsInitializedForSchema(const UPoseSearchSchema* Schema) const;
