motion matching - moving user weights from pose database into channels and bake them in FPoseSearchDatabaseDerivedData

diff --git a/Engine/Plugins/Experimental/Animation/PoseSearch/Source/Editor/Private/PoseSearchDebugger.cpp b/Engine/Plugins/Experimental/Animation/PoseSearch/Source/Editor/Private/PoseSearchDebugger.cpp
--- a/Engine/Plugins/Experimental/Animation/PoseSearch/Source/Editor/Private/PoseSearchDebugger.cpp
+++ b/Engine/Plugins/Experimental/Animation/PoseSearch/Source/Editor/Private/PoseSearchDebugger.cpp
@@ -1025,14 +1025,7 @@ void SDebuggerDatabaseView::UpdateRows(const FTraceMotionMatchingStateMessage& S
 	check(ActiveView.Rows.Num() == 1);
 	check(ContinuingPoseView.Rows.Num() == 1);
 
-	FPoseSearchWeightsContext StateWeights;
-	StateWeights.Update(&Database);
-
 	FPoseSearchContext SearchContext;
-
-	FPoseSearchWeightsContext WeightsContext;
-	WeightsContext.Update(&Database);
-
 	if (const FPoseSearchIndexAsset* CurrentIndexAsset = Database.GetSearchIndex()->FindAssetForPose(State.DbPoseIdx))
 	{
 		SearchContext.QueryMirrorRequest = CurrentIndexAsset->bMirrored ? 
@@ -1049,7 +1042,6 @@ void SDebuggerDatabaseView::UpdateRows(const FTraceMotionMatchingStateMessage& S
 
 			Database.ComparePoses(
 				SearchContext, 
-				WeightsContext, 
 				PoseIdx, 
 				State.QueryVectorNormalized, 
 				Row->PoseCostDetails);
diff --git a/Engine/Plugins/Experimental/Animation/PoseSearch/Source/Runtime/Private/PoseSearch.cpp b/Engine/Plugins/Experimental/Animation/PoseSearch/Source/Runtime/Private/PoseSearch.cpp
--- a/Engine/Plugins/Experimental/Animation/PoseSearch/Source/Runtime/Private/PoseSearch.cpp
+++ b/Engine/Plugins/Experimental/Animation/PoseSearch/Source/Runtime/Private/PoseSearch.cpp
@@ -2,9 +2,12 @@
 
 #include "PoseSearch/PoseSearch.h"
 #include "PoseSearch/PoseSearchAnimNotifies.h"
+
 #if WITH_EDITORONLY_DATA
+// @todo: deleted this include after deleting the *_DEPRECATED properies
 #include "PoseSearch/PoseSearchFeatureChannels.h"
 #endif // WITH_EDITORONLY_DATA
+
 #include "PoseSearchDerivedData.h"
 #include "PoseSearchEigenHelper.h"
 
@@ -345,48 +348,18 @@ int32 FSchemaInitializer::AddBoneReference(const FBoneReference& BoneReference)
 	return BoneReferences.AddUnique(BoneReference);
 }
 
-constexpr uint32 GetPoseSearchFeatureTypeNumFloats(EPoseSearchFeatureType Type)
-{
-	switch (Type)
-	{
-	case EPoseSearchFeatureType::Position: return 3;
-	case EPoseSearchFeatureType::Rotation: return 6;
-	case EPoseSearchFeatureType::LinearVelocity: return 3;
-	case EPoseSearchFeatureType::AngularVelocity: return 3;
-	case EPoseSearchFeatureType::ForwardVector: return 3;
-	case EPoseSearchFeatureType::Phase: return 2;
-	default: return 0;
-	}
-}
-
 int32 FSchemaInitializer::AddFeatureDesc(const FPoseSearchFeatureDesc& FeatureDesc)
 {
 	check(FeatureDesc.ChannelIdx == CurrentChannelIdx);
 	check(Features.Num() + 1 <= UPoseSearchSchema::MaxFeatures);
 	check(!Features.Contains(FeatureDesc));
+	check(FeatureDesc.Cardinality > 0);
 	const int32 FeatureDescIndex = Features.Add(FeatureDesc);
-
-	const int32 NumFloats = GetPoseSearchFeatureTypeNumFloats(FeatureDesc.Type);
-
 	Features[FeatureDescIndex].ValueOffset = CurrentChannelDataOffset;
-	Features[FeatureDescIndex].Cardinality = NumFloats;
-	CurrentChannelDataOffset += NumFloats;
+	CurrentChannelDataOffset += FeatureDesc.Cardinality;
 	return FeatureDescIndex;
 }
 
-void FSchemaInitializer::AddFeatures(int8 ChannelIdx, EPoseSearchFeatureType Type, int32 ChannelFeatureId, int32 NumSampleOffsets)
-{
-	FPoseSearchFeatureDesc Feature;
-	Feature.ChannelIdx = ChannelIdx;
-	Feature.Type = Type;
-	Feature.ChannelFeatureId = ChannelFeatureId;
-
-	for (Feature.SubsampleIdx = 0; Feature.SubsampleIdx != NumSampleOffsets; ++Feature.SubsampleIdx)
-	{
-		AddFeatureDesc(Feature);
-	}
-};
-
 } // namespace UE::PoseSearch
 
 void UPoseSearchSchema::Finalize()
@@ -554,20 +527,6 @@ FFloatRange UPoseSearchSchema::GetHorizonRange(EPoseSearchFeatureDomain Domain)
 	return GlobalRange;
 }
 
-TArrayView<const float> UPoseSearchSchema::GetChannelSampleOffsets(int32 ChannelIdx) const
-{
-	if (Channels.IsValidIndex(ChannelIdx))
-	{
-		const auto & Channel = Channels[ChannelIdx];
-		if (Channel)
-		{
-			return Channel->GetSampleOffsets();
-		}
-	}
-
-	return {};
-}
-
 void UPoseSearchSchema::ResolveBoneReferences()
 {
 	// Initialize references to obtain bone indices
@@ -630,346 +589,9 @@ bool UPoseSearchSchema::BuildQuery(
 	return bSuccess;
 }
 
-//////////////////////////////////////////////////////////////////////////
-// FPoseSearchChannelWeightParams
-
-FPoseSearchChannelWeightParams::FPoseSearchChannelWeightParams()
-{
-	TypeWeights.Empty((int32)EPoseSearchFeatureType::Num);
-	for (int32 Type = 0; Type != (int32)EPoseSearchFeatureType::Num; ++Type)
-	{
-		TypeWeights.Add((EPoseSearchFeatureType)Type, 1.0f);
-	}
-}
-
-
-
-//////////////////////////////////////////////////////////////////////////
-// FPoseSearchWeightParams
-const FPoseSearchChannelWeightParams* FPoseSearchWeightParams::GetChannelWeights(int32 ChannelIdx) const
-{
-	return ChannelWeights.IsValidIndex(ChannelIdx) ? &ChannelWeights[ChannelIdx] : nullptr;
-}
-
-
-
-//////////////////////////////////////////////////////////////////////////
-// FPoseSearchWeights
-
-void FPoseSearchWeights::Init(const FPoseSearchWeightParams& WeightParams, const UPoseSearchSchema* Schema)
-{
-	using namespace UE::PoseSearch;
-	using namespace Eigen;
-
-	// Convenience enum for indexing by horizon
-	enum EHorizon : int
-	{
-		History,
-		Prediction,
-		Num
-	};
-
-	// Initialize weights
-	Weights.Reset();
-	Weights.SetNumZeroed(Schema->Layout.NumFloats);
-
-	// Setup channel indexable weight params
-	const int32 ChannelNum = Schema->Channels.Num();
-
-	// Normalize channel weights
-	Eigen::Array<float, Dynamic, 1> NormalizedChannelWeights(ChannelNum);
-	for (int ChannelIdx = 0; ChannelIdx != ChannelNum; ++ChannelIdx)
-	{
-		const FPoseSearchChannelWeightParams* ChannelWeights = WeightParams.GetChannelWeights(ChannelIdx);
-		NormalizedChannelWeights[ChannelIdx] = ChannelWeights ? ChannelWeights->ChannelWeight : 1.0f;
-
-		// Zero the channel weight if there are no features in this channel
-		int32 FeatureIdx = INDEX_NONE;
-		if (!Schema->Layout.EnumerateBy(ChannelIdx, EPoseSearchFeatureType::Invalid, FeatureIdx))
-		{
-			NormalizedChannelWeights[ChannelIdx] = 0.0f;
-		}
-	}
-
-	const float ChannelWeightSum = NormalizedChannelWeights.sum();
-	if (!FMath::IsNearlyZero(ChannelWeightSum))
-	{
-		NormalizedChannelWeights *= (1.0f / ChannelWeightSum);
-	}	
-
-
-	// Determine maximum number of channel sample offsets for allocation
-	int32 MaxChannelSampleOffsets = 0;
-	for (int ChannelIdx = 0; ChannelIdx != ChannelNum; ++ChannelIdx)
-	{
-		TArrayView<const float> ChannelSampleOffsets = Schema->GetChannelSampleOffsets(ChannelIdx);
-		MaxChannelSampleOffsets = FMath::Max(MaxChannelSampleOffsets, ChannelSampleOffsets.Num());
-	}
-
-	// WeightsByFeature is indexed by FeatureIdx in a Layout
-	TArray<float, TInlineAllocator<32>> WeightsByFeatureStorage;
-	WeightsByFeatureStorage.SetNum(Schema->Layout.Features.Num());
-	Eigen::Map<ArrayXf> WeightsByFeature(WeightsByFeatureStorage.GetData(), WeightsByFeatureStorage.Num());
-
-	// HorizonWeightsBySample is indexed by the channel's sample offsets in a Schema
-	TArray<float, TInlineAllocator<16>> HorizonWeightsBySampleStorage;
-	HorizonWeightsBySampleStorage.SetNum(MaxChannelSampleOffsets);
-	Eigen::Map<ArrayXf> HorizonWeightsBySample(HorizonWeightsBySampleStorage.GetData(), HorizonWeightsBySampleStorage.Num());
-
-	// WeightsByType is indexed by feature type
-	Eigen::Array<float, (int)EPoseSearchFeatureType::Num, 1> WeightsByType;
-
-
-	// Determine each channel's feature weights
-	for (int ChannelIdx = 0; ChannelIdx != ChannelNum; ++ChannelIdx)
-	{
-		// Ignore this channel entirely if it has no weight
-		if (FMath::IsNearlyZero(NormalizedChannelWeights[ChannelIdx]))
-		{
-			continue;
-		}
-
-		// Get channel info
-		const FPoseSearchChannelWeightParams* ChannelWeights = WeightParams.GetChannelWeights(ChannelIdx);
-		TArrayView<const float> ChannelSampleOffsets = Schema->GetChannelSampleOffsets(ChannelIdx);
-
-		// Reset scratch weights
-		WeightsByFeature.setConstant(0.0f);
-		WeightsByType.setConstant(0.0f);
-		HorizonWeightsBySample.setConstant(0.0f);
-
-
-		// Initialize weights by type lookup
-		for (int Type = 0; Type != (int)EPoseSearchFeatureType::Num; ++Type)
-		{
-			WeightsByType[Type] = ChannelWeights ? ChannelWeights->TypeWeights.FindRef((EPoseSearchFeatureType)Type) : 1.0f;
-
-			// Zero the weight if this channel doesn't have any features using this type
-			int32 FeatureIdx = INDEX_NONE;
-			if (!Schema->Layout.EnumerateBy(ChannelIdx, (EPoseSearchFeatureType)Type, FeatureIdx))
-			{
-				WeightsByType[Type] = 0.0f;
-			}
-		}
-
-		// Normalize type weights
-		float TypeWeightsSum = WeightsByType.sum();
-		if (!FMath::IsNearlyZero(TypeWeightsSum))
-		{
-			WeightsByType *= (1.0f / TypeWeightsSum);
-		}
-		else
-		{
-			// Ignore this channel entirely if there are no types contributing weight
-			continue;
-		}
-
-
-		// Determine the range of sample offsets that make up the history and prediction horizons
-		FInt32Range HorizonSampleIdxRanges[EHorizon::Num];
-		{
-			int32 IdxUpper = Algo::UpperBound(ChannelSampleOffsets, 0.0f);
-			int32 IdxLower = ChannelSampleOffsets[0] <= 0.0f ? 0 : IdxUpper;
-			HorizonSampleIdxRanges[EHorizon::History] = FInt32Range(IdxLower, IdxUpper);
-
-			IdxLower = IdxUpper;
-			IdxUpper = ChannelSampleOffsets.Num();
-			HorizonSampleIdxRanges[EHorizon::Prediction] = FInt32Range(IdxLower, IdxUpper);
-		}
-
-
-		// Initialize horizon weights
-		Eigen::Array<float, 1, EHorizon::Num> NormalizedHorizonWeights;
-		NormalizedHorizonWeights.setConstant(0.0f);
-
-		if (!HorizonSampleIdxRanges[EHorizon::History].IsEmpty())
-		{
-			NormalizedHorizonWeights[EHorizon::History] = ChannelWeights ? ChannelWeights->HistoryParams.Weight : 1.0f;
-		}
-		if (!HorizonSampleIdxRanges[EHorizon::Prediction].IsEmpty())
-		{
-			NormalizedHorizonWeights[EHorizon::Prediction] = ChannelWeights ? ChannelWeights->PredictionParams.Weight : 1.0f;
-		}
-		
-		// Normalize horizon weights
-		float HorizonWeightSum = NormalizedHorizonWeights.sum();
-		if (!FMath::IsNearlyZero(HorizonWeightSum))
-		{
-			NormalizedHorizonWeights *= (1.0f / HorizonWeightSum);
-		}
-		else
-		{
-			// Ignore this channel entirely if the horizons don't contribute any weight
-			continue;
-		}
-
-
-		auto SetHorizonSampleWeights = [&HorizonWeightsBySample, &ChannelSampleOffsets](FInt32Range SampleIdxRange, const FPoseSearchChannelHorizonParams* HorizonParams)
-		{
-			// Segment length represents the number of sample offsets in the span that make up this horizon
-			int32 SegmentLength = SampleIdxRange.Size<int32>();
-
-			if (SegmentLength > 0)
-			{
-				int32 SegmentBegin = SampleIdxRange.GetLowerBoundValue();
-				if (HorizonParams && HorizonParams->bInterpolate && SegmentLength > 1)
-				{
-					// We'll map the range spanned by the horizon's sample offsets to the interpolation range
-					// The interpolation range is 0 to 1 unless an initial value was set
-					// The initial value allows the user to set a minimum weight or reverse the lerp direction
-					// We'll normalize these weights in the next step
-					FVector2f InputRange(ChannelSampleOffsets[SegmentBegin], ChannelSampleOffsets[SegmentBegin + SegmentLength - 1]);
-					FVector2f OutputRange(HorizonParams->InitialValue, 1.0f - HorizonParams->InitialValue);
-
-					for (int32 OffsetIdx = SegmentBegin; OffsetIdx != SegmentBegin + SegmentLength; ++OffsetIdx)
-					{
-						float SampleOffset = ChannelSampleOffsets[OffsetIdx];
-						float Alpha = FMath::GetMappedRangeValueUnclamped(InputRange, OutputRange, SampleOffset);
-						float Weight = FAlphaBlend::AlphaToBlendOption(Alpha, HorizonParams->InterpolationMethod);
-						HorizonWeightsBySample[OffsetIdx] = Weight;
-					}
-				}
-				else
-				{
-					// If we're not interpolating weights across this horizon, just give them all equal weight
-					HorizonWeightsBySample.segment(SegmentBegin, SegmentLength).setConstant(1.0f);
-				}
-
-				// Normalize weights within the horizon's segment of sample offsets
-				float HorizonSum = HorizonWeightsBySample.segment(SegmentBegin, SegmentLength).sum();
-				if (!FMath::IsNearlyZero(HorizonSum))
-				{
-					HorizonWeightsBySample.segment(SegmentBegin, SegmentLength) *= 1.0f / HorizonSum;
-				}
-			}
-		};
-
-		SetHorizonSampleWeights(
-			HorizonSampleIdxRanges[EHorizon::History],
-			ChannelWeights ? &ChannelWeights->HistoryParams : nullptr
-		);
-		SetHorizonSampleWeights(
-			HorizonSampleIdxRanges[EHorizon::Prediction],
-			ChannelWeights ? &ChannelWeights->PredictionParams : nullptr
-		);
-
-		// Now set this channel's weights for every feature in each horizon
-		Eigen::Array<float, 1, EHorizon::Num> HorizonSums;
-		HorizonSums = 0.0f;
-		for (int FeatureIdx = INDEX_NONE; Schema->Layout.EnumerateBy(ChannelIdx, EPoseSearchFeatureType::Invalid, FeatureIdx); /*empty*/)
-		{
-			const FPoseSearchFeatureDesc& Feature = Schema->Layout.Features[FeatureIdx];
-
-			for (int HorizonIdx = 0; HorizonIdx != EHorizon::Num; ++HorizonIdx)
-			{
-				if (HorizonSampleIdxRanges[HorizonIdx].Contains(Feature.SubsampleIdx))
-				{
-					int HorizonSize = HorizonSampleIdxRanges[HorizonIdx].Size<int>();
-					
-					// Bug: HorizonSize * WeightsByType[(int)Feature.Type] assumes horizons are composed of homogenous features, not true any more for pose channel
-					WeightsByFeature[FeatureIdx] = HorizonWeightsBySample[Feature.SubsampleIdx] * (HorizonSize * WeightsByType[(int)Feature.Type]);
-					HorizonSums[HorizonIdx] += WeightsByFeature[FeatureIdx];
-					break;
-				}
-			}
-		}
-
-		// Scale feature weights within horizons so that they have the desired total horizon weight
-		for (int FeatureIdx = INDEX_NONE; Schema->Layout.EnumerateBy(ChannelIdx, EPoseSearchFeatureType::Invalid, FeatureIdx); /*empty*/)
-		{
-			const FPoseSearchFeatureDesc& Feature = Schema->Layout.Features[FeatureIdx];
-
-			for (int HorizonIdx = 0; HorizonIdx != EHorizon::Num; ++HorizonIdx)
-			{
-				if (HorizonSampleIdxRanges[HorizonIdx].Contains(Feature.SubsampleIdx))
-				{
-					float HorizonWeight = NormalizedHorizonWeights[HorizonIdx] / HorizonSums[HorizonIdx];
-					WeightsByFeature[FeatureIdx] *= HorizonWeight;
-					break;
-				}
-			}
-		}
-
-		// Scale all features in all horizons so they have the desired channel weight
-		WeightsByFeature *= NormalizedChannelWeights[ChannelIdx];
-
-		// Weights should sum to channel weight at this point
-		ensure(FMath::IsNearlyEqual(WeightsByFeature.sum(), NormalizedChannelWeights[ChannelIdx], KINDA_SMALL_NUMBER));
-
-		// Merge feature weights for channel into per-value weights buffer
-		// Weights are replicated per feature dimension so the cost function can directly index weights by value index
-		for (int FeatureIdx = INDEX_NONE; Schema->Layout.EnumerateBy(ChannelIdx, EPoseSearchFeatureType::Invalid, FeatureIdx); /*empty*/)
-		{
-			const FPoseSearchFeatureDesc& Feature = Schema->Layout.Features[FeatureIdx];
-			const int32 ValueTerm = Feature.ValueOffset + Feature.Cardinality;
-			for (int32 ValueIdx = Feature.ValueOffset; ValueIdx != ValueTerm; ++ValueIdx)
-			{
-				Weights[ValueIdx] = WeightsByFeature[FeatureIdx];
-			}
-		}
-	}
-}
-
-
-//////////////////////////////////////////////////////////////////////////
-// FPoseSearchWeightsContext
-
-void FPoseSearchWeightsContext::Update(const UPoseSearchDatabase* ActiveDatabase)
-{
-	bool bRecomputeWeights = false;
-	if (Database != ActiveDatabase)
-	{
-		Database = ActiveDatabase;
-#if WITH_EDITOR
-		SearchIndexHash = Database->GetSearchIndexHash();
-#endif
-		bRecomputeWeights = true;
-	}
-
-#if WITH_EDITOR
-	if (Database.IsValid() && Database->GetSearchIndexHash() != SearchIndexHash)
-	{
-		SearchIndexHash = Database->GetSearchIndexHash();
-		bRecomputeWeights = true;
-	}
-#endif
-
-	if (bRecomputeWeights)
-	{
-		ComputedDefaultGroupWeights.Init(Database->DefaultWeights, Database->Schema);
-
-		int32 NumGroups = ActiveDatabase ? ActiveDatabase->Groups.Num() : 0;
-		ComputedGroupWeights.SetNum(NumGroups);
-
-		for (int GroupIdx = 0; GroupIdx < NumGroups; ++GroupIdx)
-		{
-			ComputedGroupWeights[GroupIdx].Init(
-				Database->Groups[GroupIdx].Weights, 
-				Database->Schema);
-		}
-	}
-}
-
-const FPoseSearchWeights* FPoseSearchWeightsContext::GetGroupWeights(int32 WeightsGroupIdx) const
-{
-	if (WeightsGroupIdx == INDEX_NONE)
-	{
-		return &ComputedDefaultGroupWeights;
-	}
-
-	if (ComputedGroupWeights.IsValidIndex(WeightsGroupIdx))
-	{
-		return &ComputedGroupWeights[WeightsGroupIdx];
-	}
-
-	return nullptr;
-}
-
-
 //////////////////////////////////////////////////////////////////////////
 // FPoseSearchIndexAsset
 // 
-
 int32 FPoseSearchIndex::FindAssetIndex(const FPoseSearchIndexAsset* Asset) const
 {
 	if (Asset == nullptr || Assets.Num() == 0)
@@ -988,6 +610,19 @@ int32 FPoseSearchIndex::FindAssetIndex(const FPoseSearchIndexAsset* Asset) const
 	return Result;
 }
 
+const FGroupSearchIndex* FPoseSearchIndex::FindGroup(int32 GroupIndex) const
+{
+	for (const FGroupSearchIndex& Group : Groups)
+	{
+		if (Group.GroupIndex == GroupIndex)
+		{
+			return &Group;
+		}
+	}
+	check(false);
+	return nullptr;
+}
+
 const FPoseSearchIndexAsset* FPoseSearchIndex::FindAssetForPose(int32 PoseIdx) const
 {
 	auto Predicate = [PoseIdx](const FPoseSearchIndexAsset& Asset)
@@ -1532,10 +1167,6 @@ static inline void CollectGroupIndices(
 		{
 			BadGroupIndices.Add(Index);
 		}
-		else if (Groups[GroupIndex].bUseGroupWeights)
-		{
-			GroupIndices.Add(GroupIndex);
-		}
 	}
 
 	if (GroupIndices.Num() == 0)
@@ -2009,7 +1640,6 @@ bool UPoseSearchDatabase::IsCachedCookedPlatformDataLoaded(const ITargetPlatform
 
 FPoseSearchCost UPoseSearchDatabase::ComparePoses(
 	FPoseSearchContext& SearchContext,
-	FPoseSearchWeightsContext& WeightsContext,
 	int32 PoseIdx,
 	int32 GroupIdx,
 	const TArrayView<const float>& QueryValues) const
@@ -2041,13 +1671,11 @@ FPoseSearchCost UPoseSearchDatabase::ComparePoses(
 		GroupIdx = SearchIndexAsset->SourceGroupIdx;
 	}
 
-	const FPoseSearchWeights* WeightsSet =
-		WeightsContext.GetGroupWeights(GroupIdx);
 	Result.SetDissimilarity(CompareFeatureVectors(
 		PoseValues.Num(),
 		PoseValues.GetData(),
 		QueryValues.GetData(),
-		WeightsSet->Weights.GetData()));
+		SearchIndex->FindGroup(GroupIdx)->Weights.GetData()));
 
 	float NotifyAddend = 0.0f;
 	float MirrorMismatchAddend = 0.0f;
@@ -2059,7 +1687,6 @@ FPoseSearchCost UPoseSearchDatabase::ComparePoses(
 
 FPoseSearchCost UPoseSearchDatabase::ComparePoses(
 	FPoseSearchContext& SearchContext,
-	FPoseSearchWeightsContext& WeightsContext,
 	int32 PoseIdx,
 	const TArrayView<const float>& QueryValues,
 	UE::PoseSearch::FPoseCostDetails& OutPoseCostDetails) const
@@ -2085,11 +1712,10 @@ FPoseSearchCost UPoseSearchDatabase::ComparePoses(
 
 	// Compute weighted squared difference vector
 	const FPoseSearchIndexAsset* SearchIndexAsset = GetSearchIndex()->FindAssetForPose(PoseIdx);
-	const FPoseSearchWeights* WeightsSet =
-		WeightsContext.GetGroupWeights(SearchIndexAsset->SourceGroupIdx);
-	check(WeightsSet);
-	check(WeightsSet->Weights.Num() == Dims);
-	auto WeightsVector = Map<const ArrayXf>(WeightsSet->Weights.GetData(), Dims);
+	const TArray<float>& Weights = GetSearchIndex()->FindGroup(SearchIndexAsset->SourceGroupIdx)->Weights;
+
+	check(Weights.Num() == Dims);
+	auto WeightsVector = Map<const ArrayXf>(Weights.GetData(), Dims);
 
 	OutCostVector = WeightsVector * (PoseVector - QueryVector).square();
 	Result.SetDissimilarity(OutCostVector.sum());
@@ -2114,7 +1740,6 @@ FPoseSearchCost UPoseSearchDatabase::ComparePoses(
 		// Verify details pose comparator agrees with runtime pose comparator
 		FPoseSearchCost RuntimeComparatorCost = ComparePoses(
 			SearchContext, 
-			WeightsContext, 
 			PoseIdx, 
 			SearchIndexAsset->SourceGroupIdx,
 			QueryValues);
@@ -2292,9 +1917,6 @@ UE::PoseSearch::FSearchResult UPoseSearchDatabase::SearchPCAKDTree(FPoseSearchCo
 
 	BuildQuery(SearchContext, Result.ComposedQuery);
 
-	FPoseSearchWeightsContext WeightsContext;
-	WeightsContext.Update(this);
-
 	TArrayView<const float> NormalizedQueryValues = Result.ComposedQuery.GetNormalizedValues();
 
 	const bool IsCurrentResultFromThisDatabase = SearchContext.CurrentResult.IsValid() && SearchContext.CurrentResult.Database == this;
@@ -2302,7 +1924,6 @@ UE::PoseSearch::FSearchResult UPoseSearchDatabase::SearchPCAKDTree(FPoseSearchCo
 	{
 		Result.ContinuityPoseCost = ComparePoses(
 			SearchContext,
-			WeightsContext,
 			SearchContext.CurrentResult.PoseIdx,
 			SearchIndex->FindAssetForPose(SearchContext.CurrentResult.PoseIdx)->SourceGroupIdx,
 			NormalizedQueryValues);
@@ -2342,7 +1963,6 @@ UE::PoseSearch::FSearchResult UPoseSearchDatabase::SearchPCAKDTree(FPoseSearchCo
 
 			FPoseSearchCost PoseCost = ComparePoses(
 				SearchContext, 
-				WeightsContext, 
 				PoseIdx, 
 				GroupSearchIndex.GroupIndex,
 				NormalizedQueryValues);
@@ -2383,9 +2003,6 @@ UE::PoseSearch::FSearchResult UPoseSearchDatabase::SearchBruteForce(FPoseSearchC
 	BuildQuery(SearchContext, Result.ComposedQuery);
 	TArrayView<const float> NormalizedQueryValues = Result.ComposedQuery.GetNormalizedValues();
 
-	FPoseSearchWeightsContext WeightsContext;
-	WeightsContext.Update(this);
-	
 	constexpr int NonSelectableIdxDataSize = 128;
 	size_t* NonSelectableIdxData((size_t*)FMemory_Alloca(NonSelectableIdxDataSize * sizeof(size_t)));
 	int NonSelectableIdxUsedSize = 0;
@@ -2394,7 +2011,6 @@ UE::PoseSearch::FSearchResult UPoseSearchDatabase::SearchBruteForce(FPoseSearchC
 	{
 		Result.ContinuityPoseCost = ComparePoses(
 			SearchContext,
-			WeightsContext,
 			SearchContext.CurrentResult.PoseIdx,
 			SearchIndex->FindAssetForPose(SearchContext.CurrentResult.PoseIdx)->SourceGroupIdx,
 			NormalizedQueryValues);
@@ -2434,7 +2050,6 @@ UE::PoseSearch::FSearchResult UPoseSearchDatabase::SearchBruteForce(FPoseSearchC
 
 			FPoseSearchCost PoseCost = ComparePoses(
 				SearchContext,
-				WeightsContext,
 				PoseIdx, 
 				Asset.SourceGroupIdx,
 				NormalizedQueryValues);
@@ -5060,10 +4675,22 @@ static void PreprocessSearchIndex(FPoseSearchIndex* SearchIndex)
 
 static void PreprocessGroupSearchIndexWeights(FGroupSearchIndex& GroupSearchIndex, const UPoseSearchDatabase* Database)
 {
-	const FPoseSearchWeightParams& WeightParams = (GroupSearchIndex.GroupIndex == INDEX_NONE || Database->Groups.IsEmpty())  ? Database->DefaultWeights : Database->Groups[GroupSearchIndex.GroupIndex].Weights;
-	FPoseSearchWeights Weights;
-	Weights.Init(WeightParams, Database->Schema);
-	GroupSearchIndex.Weights = Weights.Weights;
+	const int32 NumDimensions = Database->Schema->Layout.NumFloats;
+	GroupSearchIndex.Weights.Init(1.f, NumDimensions);
+
+	for (int ChannelIdx = 0; ChannelIdx != Database->Schema->Channels.Num(); ++ChannelIdx)
+	{
+		const UPoseSearchFeatureChannel* Channel = Database->Schema->Channels[ChannelIdx].Get();
+		Channel->FillWeights(GroupSearchIndex.Weights);
+	}
+
+	// normalizing user weights
+	RowMajorVectorMap MapWeights(GroupSearchIndex.Weights.GetData(), 1, NumDimensions);
+	const float WeightsSum = MapWeights.sum();
+	if (!FMath::IsNearlyZero(WeightsSum))
+	{
+		MapWeights *= (1.0f / WeightsSum);
+	}
 }
 
 // it calcualtes Mean and PCAProjectionMatrix
diff --git a/Engine/Plugins/Experimental/Animation/PoseSearch/Source/Runtime/Private/PoseSearchFeatureChannels.cpp b/Engine/Plugins/Experimental/Animation/PoseSearch/Source/Runtime/Private/PoseSearchFeatureChannels.cpp
--- a/Engine/Plugins/Experimental/Animation/PoseSearch/Source/Runtime/Private/PoseSearchFeatureChannels.cpp
+++ b/Engine/Plugins/Experimental/Animation/PoseSearch/Source/Runtime/Private/PoseSearchFeatureChannels.cpp
@@ -1,7 +1,6 @@
 // Copyright Epic Games, Inc. All Rights Reserved.
 
 #include "PoseSearch/PoseSearchFeatureChannels.h"
-#include "PoseSearch/PoseSearch.h"
 #include "AnimationRuntime.h"
 #include "Animation/AnimationPoseData.h"
 #include "Animation/AttributesRuntime.h"
@@ -294,19 +293,31 @@ void UPoseSearchFeatureChannel_Pose::InitializeSchema(UE::PoseSearch::FSchemaIni
 		const FPoseSearchBone& SampledBone = SampledBones[ChannelBoneIdx];
 		if (SampledBone.bUsePosition)
 		{
-			Initializer.AddFeatures(GetChannelIndex(), EPoseSearchFeatureType::Position, ChannelBoneIdx, SampleTimes.Num());
+			for (int32 SubsampleIdx = 0; SubsampleIdx != SampleTimes.Num(); ++SubsampleIdx)
+			{
+				Initializer.AddFeatureDesc(FPoseSearchFeatureDesc::Construct(GetChannelIndex(), ChannelBoneIdx, SubsampleIdx, EPoseSearchFeatureType::Position, PositionCardinality));
+			}
 		}
 		if (SampledBone.bUseRotation)
 		{
-			Initializer.AddFeatures(GetChannelIndex(), EPoseSearchFeatureType::Rotation, ChannelBoneIdx, SampleTimes.Num());
+			for (int32 SubsampleIdx = 0; SubsampleIdx != SampleTimes.Num(); ++SubsampleIdx)
+			{
+				Initializer.AddFeatureDesc(FPoseSearchFeatureDesc::Construct(GetChannelIndex(), ChannelBoneIdx, SubsampleIdx, EPoseSearchFeatureType::Rotation, RotationCardinality));
+			}
 		}
 		if (SampledBone.bUseVelocity)
 		{
-			Initializer.AddFeatures(GetChannelIndex(), EPoseSearchFeatureType::LinearVelocity, ChannelBoneIdx, SampleTimes.Num());
+			for (int32 SubsampleIdx = 0; SubsampleIdx != SampleTimes.Num(); ++SubsampleIdx)
+			{
+				Initializer.AddFeatureDesc(FPoseSearchFeatureDesc::Construct(GetChannelIndex(), ChannelBoneIdx, SubsampleIdx, EPoseSearchFeatureType::LinearVelocity, LinearVelocityCardinality));
+			}
 		}
 		if (SampledBone.bUsePhase)
 		{
-			Initializer.AddFeatures(GetChannelIndex(), EPoseSearchFeatureType::Phase, ChannelBoneIdx, SampleTimes.Num());
+			for (int32 SubsampleIdx = 0; SubsampleIdx != SampleTimes.Num(); ++SubsampleIdx)
+			{
+				Initializer.AddFeatureDesc(FPoseSearchFeatureDesc::Construct(GetChannelIndex(), ChannelBoneIdx, SubsampleIdx, EPoseSearchFeatureType::Phase, PhaseCardinality));
+			}
 		}
 	}
 
@@ -321,6 +332,51 @@ void UPoseSearchFeatureChannel_Pose::InitializeSchema(UE::PoseSearch::FSchemaIni
 	}
 }
 
+void UPoseSearchFeatureChannel_Pose::FillWeights(TArray<float>& Weights) const
+{
+	int32 DataOffset = ChannelDataOffset;
+
+	const int32 NumBones = SampledBones.Num();
+	for (int32 ChannelBoneIdx = 0; ChannelBoneIdx != NumBones; ++ChannelBoneIdx)
+	{
+		const FPoseSearchBone& SampledBone = SampledBones[ChannelBoneIdx];
+		if (SampledBone.bUsePosition)
+		{
+			for (int32 SubsampleIdx = 0; SubsampleIdx != SampleTimes.Num(); ++SubsampleIdx)
+			{
+				Weights[DataOffset] = SampledBone.Weight;
+				DataOffset += PositionCardinality;
+			}
+		}
+		if (SampledBone.bUseRotation)
+		{
+			for (int32 SubsampleIdx = 0; SubsampleIdx != SampleTimes.Num(); ++SubsampleIdx)
+			{
+				Weights[DataOffset] = SampledBone.Weight;
+				DataOffset += RotationCardinality;
+			}
+		}
+		if (SampledBone.bUseVelocity)
+		{
+			for (int32 SubsampleIdx = 0; SubsampleIdx != SampleTimes.Num(); ++SubsampleIdx)
+			{
+				Weights[DataOffset] = SampledBone.Weight;
+				DataOffset += LinearVelocityCardinality;
+			}
+		}
+		if (SampledBone.bUsePhase)
+		{
+			for (int32 SubsampleIdx = 0; SubsampleIdx != SampleTimes.Num(); ++SubsampleIdx)
+			{
+				Weights[DataOffset] = SampledBone.Weight;
+				DataOffset += PhaseCardinality;
+			}
+		}
+	}
+
+	check(DataOffset == ChannelDataOffset + ChannelCardinality);
+}
+
 // @todo: do we really need to use double(s) in all this math?
 void UPoseSearchFeatureChannel_Pose::CalculatePhases(const UE::PoseSearch::IAssetIndexer& Indexer, UE::PoseSearch::FAssetIndexingOutput& IndexingOutput, TArray<TArray<FVector2D>>& OutPhases) const
 {
@@ -800,8 +856,6 @@ void UPoseSearchFeatureChannel_Pose::DebugDraw(const UE::PoseSearch::FDebugDrawP
 	}
 }
 
-
-
 //////////////////////////////////////////////////////////////////////////
 // UPoseSearchFeatureChannel_Trajectory
 void UPoseSearchFeatureChannel_Trajectory::PreSave(FObjectPreSaveContext ObjectSaveContext)
@@ -817,22 +871,41 @@ void UPoseSearchFeatureChannel_Trajectory::InitializeSchema( UE::PoseSearch::FSc
 
 	if (bUsePositions)
 	{
-		Initializer.AddFeatures(GetChannelIndex(), EPoseSearchFeatureType::Position, 0, SampleOffsets.Num());
+		for (int32 SubsampleIdx = 0; SubsampleIdx != SampleOffsets.Num(); ++SubsampleIdx)
+		{
+			Initializer.AddFeatureDesc(FPoseSearchFeatureDesc::Construct(GetChannelIndex(), 0, SubsampleIdx, EPoseSearchFeatureType::Position, PositionCardinality));
+		}
 	}
 
 	if (bUseLinearVelocities)
 	{
-		Initializer.AddFeatures(GetChannelIndex(), EPoseSearchFeatureType::LinearVelocity, 0, SampleOffsets.Num());
+		for (int32 SubsampleIdx = 0; SubsampleIdx != SampleOffsets.Num(); ++SubsampleIdx)
+		{
+			Initializer.AddFeatureDesc(FPoseSearchFeatureDesc::Construct(GetChannelIndex(), 0, SubsampleIdx, EPoseSearchFeatureType::LinearVelocity, LinearVelocityCardinality));
+		}
 	}
 
 	if (bUseFacingDirections)
 	{
-		Initializer.AddFeatures(GetChannelIndex(), EPoseSearchFeatureType::ForwardVector, 0, SampleOffsets.Num());
+		for (int32 SubsampleIdx = 0; SubsampleIdx != SampleOffsets.Num(); ++SubsampleIdx)
+		{
+			Initializer.AddFeatureDesc(FPoseSearchFeatureDesc::Construct(GetChannelIndex(), 0, SubsampleIdx, EPoseSearchFeatureType::ForwardVector, ForwardVectorCardinality));
+		}
 	}
 
 	ChannelCardinality = Initializer.GetCurrentCardinalityFrom(ChannelDataOffset);
 }
 
+void UPoseSearchFeatureChannel_Trajectory::FillWeights(TArray<float>& Weights) const
+{
+	const int32 Begin = ChannelDataOffset;
+	const int32 End = ChannelDataOffset + ChannelCardinality;
+	for (int i = Begin; i < End; ++i)
+	{
+		Weights[i] = Weight;
+	}
+}
+
 void UPoseSearchFeatureChannel_Trajectory::IndexAsset(const UE::PoseSearch::IAssetIndexer& Indexer,  UE::PoseSearch::FAssetIndexingOutput& IndexingOutput) const
 {
 	using namespace UE::PoseSearch;
@@ -1011,6 +1084,7 @@ void UPoseSearchFeatureChannel_Trajectory::GenerateDDCKey(FBlake3& InOutKeyHashe
 	InOutKeyHasher.Update(&bUseFacingDirections, sizeof(bUseFacingDirections));
 	InOutKeyHasher.Update(&Domain, sizeof(Domain));
 	InOutKeyHasher.Update(MakeMemoryView(SampleOffsets));
+	InOutKeyHasher.Update(&Weight, sizeof(Weight));
 }
 
 bool UPoseSearchFeatureChannel_Trajectory::BuildQuery(
diff --git a/Engine/Plugins/Experimental/Animation/PoseSearch/Source/Runtime/Public/PoseSearch/PoseSearch.h b/Engine/Plugins/Experimental/Animation/PoseSearch/Source/Runtime/Public/PoseSearch/PoseSearch.h
--- a/Engine/Plugins/Experimental/Animation/PoseSearch/Source/Runtime/Public/PoseSearch/PoseSearch.h
+++ b/Engine/Plugins/Experimental/Animation/PoseSearch/Source/Runtime/Public/PoseSearch/PoseSearch.h
@@ -180,6 +180,7 @@ struct FPoseSearchBlockTransitionParameters
 	float SequenceEndInterval = 0.0f;
 };
 
+// @todo: move it into PoseSearchFeatureChannels after removing SampledBones_DEPRECATED 
 USTRUCT()
 struct POSESEARCH_API FPoseSearchBone
 {
@@ -199,6 +200,10 @@ struct POSESEARCH_API FPoseSearchBone
 
 	UPROPERTY(EditAnywhere, Category = Config)
 	bool bUsePhase = false;
+
+	// @todo: temporary location for the channel bone weight to help the weights refactoring
+	UPROPERTY(EditAnywhere, Category = Config)
+	float Weight = 1.f;
 };
 
 
@@ -228,17 +233,29 @@ struct POSESEARCH_API FPoseSearchFeatureDesc
 	UPROPERTY()
 	EPoseSearchFeatureType Type = EPoseSearchFeatureType::Invalid;
 
-	// Set via FSchemaInitializer::AddFeatureDesc() and ignored by operator==
-	UPROPERTY()
-	int16 ValueOffset = 0;
-
 	// Set via FSchemaInitializer::AddFeatureDesc() and ignored by operator== or else Schema->Layout.Features.Find(Feature); called by FPoseSearchFeatureVectorBuilder will fail
 	UPROPERTY()
 	int16 Cardinality = 0;
 
+	// Set via FSchemaInitializer::AddFeatureDesc() and ignored by operator==
+	UPROPERTY()
+	int16 ValueOffset = 0;
+
 	bool operator==(const FPoseSearchFeatureDesc& Other) const;
 
 	bool IsValid() const { return Type != EPoseSearchFeatureType::Invalid; }
+
+	static FPoseSearchFeatureDesc Construct(int8 ChannelIdx, int8 ChannelFeatureId, int8 SubsampleIdx, EPoseSearchFeatureType Type, int16 Cardinality)
+	{
+		FPoseSearchFeatureDesc PoseSearchFeatureDesc;
+		PoseSearchFeatureDesc.ChannelIdx = ChannelIdx;
+		PoseSearchFeatureDesc.ChannelFeatureId = ChannelFeatureId;
+		PoseSearchFeatureDesc.SubsampleIdx = SubsampleIdx;
+		PoseSearchFeatureDesc.Type = Type;
+		PoseSearchFeatureDesc.Cardinality = Cardinality;
+		//PoseSearchFeatureDesc.ValueOffset
+		return PoseSearchFeatureDesc;
+	}
 };
 
 
@@ -401,6 +418,8 @@ public:
 	// Called during UPoseSearchSchema::Finalize to prepare the schema for this channel
 	virtual void InitializeSchema(UE::PoseSearch::FSchemaInitializer& Initializer);
 	
+	virtual void FillWeights(TArray<float>& Weights) const PURE_VIRTUAL(UPoseSearchFeatureChannel::FillWeights, );
+
 	// Called at database build time to populate pose vectors with this channel's data
 	virtual void IndexAsset(const UE::PoseSearch::IAssetIndexer& Indexer, UE::PoseSearch::FAssetIndexingOutput& IndexingOutput) const PURE_VIRTUAL(UPoseSearchFeatureChannel::IndexAsset, );
 
@@ -408,9 +427,6 @@ public:
 	// Returns empty range if the channel has no horizon in the requested domain.
 	virtual FFloatRange GetHorizonRange(EPoseSearchFeatureDomain Domain) const PURE_VIRTUAL(UPoseSearchFeatureChannel::GetHorizonRange, return FFloatRange::Empty(); );
 
-	// Return this channel's horizon sampling offsets
-	virtual TArrayView<const float> GetSampleOffsets () const PURE_VIRTUAL(UPoseSearchFeatureChannel::GetSampleOffsets, return {}; );
-
 	// Hash channel properties to produce a key for database derived data
 	virtual void GenerateDDCKey(FBlake3& InOutKeyHasher) const PURE_VIRTUAL(UPoseSearchFeatureChannel::GenerateDDCKey, );
 
@@ -461,12 +477,10 @@ public:
 	int32 GetCurrentChannelIdx() const { return CurrentChannelIdx; }
 
 	int32 GetCurrentCardinalityFrom(int32 ChannelDataOffset) const { check(CurrentChannelDataOffset >= ChannelDataOffset);  return CurrentChannelDataOffset - ChannelDataOffset; }
-	void AddFeatures(int8 ChannelIdx, EPoseSearchFeatureType Type, int32 ChannelFeatureId, int32 NumSampleOffsets);
 	int32 GetCurrentChannelDataOffset() const { return CurrentChannelDataOffset; }
-
-private:
 	int32 AddFeatureDesc(const FPoseSearchFeatureDesc& FeatureDesc);
 
+private:
 	friend class ::UPoseSearchSchema;
 
 	int32 CurrentChannelIdx = 0;
@@ -560,8 +574,6 @@ public:
 	// Returns empty range if the channel has no horizon in the requested domain.
 	virtual FFloatRange GetHorizonRange(EPoseSearchFeatureDomain Domain) const;
 
-	TArrayView<const float> GetChannelSampleOffsets (int32 ChannelIdx) const;
-
 public: // UObject
 	virtual void PreSave(FObjectPreSaveContext ObjectSaveContext) override;
 	virtual void PostLoad() override;
@@ -755,6 +767,8 @@ struct POSESEARCH_API FPoseSearchIndex
 	TArrayView<const float> GetPoseValues(int32 PoseIdx) const;
 
 	int32 FindAssetIndex(const FPoseSearchIndexAsset* Asset) const;
+	const FGroupSearchIndex* FindGroup(int32 GroupIndex) const;
+
 	const FPoseSearchIndexAsset* FindAssetForPose(int32 PoseIdx) const;
 	float GetAssetTime(int32 PoseIdx, const FPoseSearchIndexAsset* Asset) const;
 
@@ -769,104 +783,6 @@ struct POSESEARCH_API FPoseSearchIndex
 //////////////////////////////////////////////////////////////////////////
 // Database
 
-USTRUCT(BlueprintType, Category = "Animation|Pose Search")
-struct POSESEARCH_API FPoseSearchChannelHorizonParams
-{
-	GENERATED_BODY()
-
-	// Total score contribution of all samples within this horizon, normalized with other horizons
-	UPROPERTY(BlueprintReadWrite, EditAnywhere, Category = "Settings", meta = (ClampMin = "0.0"))
-	float Weight = 1.0f;
-
-	// Whether to interpolate samples within this horizon
-	UPROPERTY(BlueprintReadWrite, EditAnywhere, Category = "Advanced")
-	bool bInterpolate = false;
-
-	// Horizon sample weights will be interpolated from InitialValue to 1.0 - InitialValue and then normalized
-	UPROPERTY(BlueprintReadWrite, EditAnywhere, Category = "Advanced", meta = (EditCondition = "bInterpolate", ClampMin="0.0", ClampMax="1.0"))
-	float InitialValue = 0.1f;
-
-	// Curve type for horizon interpolation 
-	UPROPERTY(BlueprintReadWrite, EditAnywhere, Category = "Advanced", meta = (EditCondition = "bInterpolate"))
-	EAlphaBlendOption InterpolationMethod = EAlphaBlendOption::Linear;
-};
-
-USTRUCT(BlueprintType, Category = "Animation|Pose Search")
-struct POSESEARCH_API FPoseSearchChannelWeightParams
-{
-	GENERATED_BODY()
-
-	// Contribution of this score component. Normalized with other channels.
-	UPROPERTY(BlueprintReadWrite, EditAnywhere, Category = "Settings", meta = (ClampMin = "0.0"))
-	float ChannelWeight = 1.0f;
-
-	// History horizon params (for sample offsets <= 0)
-	UPROPERTY(BlueprintReadWrite, EditAnywhere, Category = "Settings")
-	FPoseSearchChannelHorizonParams HistoryParams;
-
-	// Prediction horizon params (for sample offsets > 0)
-	UPROPERTY(BlueprintReadWrite, EditAnywhere, Category = "Settings")
-	FPoseSearchChannelHorizonParams PredictionParams;
-
-	// Contribution of each type within this channel
-	UPROPERTY(BlueprintReadWrite, EditAnywhere, Category = "Settings")
-	TMap<EPoseSearchFeatureType, float> TypeWeights;
-
-	FPoseSearchChannelWeightParams();
-};
-
-USTRUCT(BlueprintType, Category = "Animation|Pose Search")
-struct POSESEARCH_API FPoseSearchWeightParams
-{
-	GENERATED_BODY()
-
-	UPROPERTY(BlueprintReadWrite, EditAnywhere, Category = "Settings")
-	TArray<FPoseSearchChannelWeightParams> ChannelWeights;
-
-	const FPoseSearchChannelWeightParams* GetChannelWeights(int32 ChannelIdx) const;
-};
-
-
-USTRUCT()
-struct POSESEARCH_API FPoseSearchWeights
-{
-	GENERATED_BODY()
-
-	UPROPERTY(Transient)
-	TArray<float> Weights;
-
-	bool IsInitialized() const { return !Weights.IsEmpty(); }
-	void Init(const FPoseSearchWeightParams& WeightParams, const UPoseSearchSchema* Schema);
-};
-
-USTRUCT()
-struct POSESEARCH_API FPoseSearchWeightsContext
-{
-	GENERATED_BODY()
-
-public:
-	// Computes and caches new group weights whenever the database changes
-	void Update(const UPoseSearchDatabase * Database);
-
-	const FPoseSearchWeights* GetGroupWeights (int32 WeightsGroupIdx) const;
-	
-private:
-	UPROPERTY(Transient)
-	TWeakObjectPtr<const UPoseSearchDatabase> Database = nullptr;
-
-	UPROPERTY(Transient)
-	FPoseSearchWeights ComputedDefaultGroupWeights;
-
-	UPROPERTY(Transient)
-	TArray<FPoseSearchWeights> ComputedGroupWeights;
-
-#if WITH_EDITOR
-	// used to check if the data has changed, which requires the weights to be recomputed
-	FIoHash SearchIndexHash = FIoHash::Zero;
-#endif
-};
-
-
 /** An entry in a UPoseSearchDatabase. */
 USTRUCT(BlueprintType, Category = "Animation|Pose Search")
 struct POSESEARCH_API FPoseSearchDatabaseSequence
@@ -947,12 +863,6 @@ struct POSESEARCH_API FPoseSearchDatabaseGroup
 public:
 	UPROPERTY(EditAnywhere, Category = "Settings")
 	FGameplayTag Tag;
-
-	UPROPERTY(EditAnywhere, Category = "Settings")
-	bool bUseGroupWeights = false;
-
-	UPROPERTY(EditAnywhere, Category = "Settings", meta=(EditCondition="bUseGroupWeights", EditConditionHides))
-	FPoseSearchWeightParams Weights;
 };
 
 USTRUCT(BlueprintType)
@@ -1232,9 +1142,6 @@ public:
 	UPROPERTY(EditAnywhere, BlueprintReadOnly, Category="Database", DisplayName="Config")
 	TObjectPtr<const UPoseSearchSchema> Schema = nullptr;
 
-	UPROPERTY(EditAnywhere, Category = "Database")
-	FPoseSearchWeightParams DefaultWeights;
-
 	// If there's a mirroring mismatch between the currently playing sequence and a search candidate, this cost will be 
 	// added to the candidate, making it less likely to be selected
 	UPROPERTY(EditAnywhere, BlueprintReadWrite, Category="Database")
@@ -1358,14 +1265,12 @@ public:
 
 	FPoseSearchCost ComparePoses(
 		FPoseSearchContext& SearchContext,
-		FPoseSearchWeightsContext& WeightsContext,
 		int32 PoseIdx,
 		int32 GroupIdx,
 		const TArrayView<const float>& QueryValues) const;
 
 	FPoseSearchCost ComparePoses(
 		FPoseSearchContext& SearchContext,
-		FPoseSearchWeightsContext& WeightsContext,
 		int32 PoseIdx,
 		const TArrayView<const float>& QueryValues,
 		UE::PoseSearch::FPoseCostDetails& OutPoseCostDetails) const;
diff --git a/Engine/Plugins/Experimental/Animation/PoseSearch/Source/Runtime/Public/PoseSearch/PoseSearchFeatureChannels.h b/Engine/Plugins/Experimental/Animation/PoseSearch/Source/Runtime/Public/PoseSearch/PoseSearchFeatureChannels.h
--- a/Engine/Plugins/Experimental/Animation/PoseSearch/Source/Runtime/Public/PoseSearch/PoseSearchFeatureChannels.h
+++ b/Engine/Plugins/Experimental/Animation/PoseSearch/Source/Runtime/Public/PoseSearch/PoseSearchFeatureChannels.h
@@ -42,9 +42,9 @@ public:
 
 	// UPoseSearchFeatureChannel interface
 	virtual void InitializeSchema(UE::PoseSearch::FSchemaInitializer& Initializer) override;
+	virtual void FillWeights(TArray<float>& Weights) const override;
 	virtual void IndexAsset(const UE::PoseSearch::IAssetIndexer& Indexer, UE::PoseSearch::FAssetIndexingOutput& IndexingOutput) const override;
 	virtual FFloatRange GetHorizonRange(EPoseSearchFeatureDomain Domain) const override;
-	virtual TArrayView<const float> GetSampleOffsets () const override { return SampleTimes; }
 	virtual void GenerateDDCKey(FBlake3& InOutKeyHasher) const override;
 	virtual bool BuildQuery(
 		FPoseSearchContext& SearchContext,
@@ -54,6 +54,11 @@ public:
 protected:
 	virtual void AddPoseFeatures(const UE::PoseSearch::IAssetIndexer& Indexer, int32 SampleIdx, FPoseSearchFeatureVectorBuilder& FeatureVector, const TArray<TArray<FVector2D>>& Phases) const;
 	void CalculatePhases(const UE::PoseSearch::IAssetIndexer& Indexer, UE::PoseSearch::FAssetIndexingOutput& IndexingOutput, TArray<TArray<FVector2D>>& OutPhases) const;
+
+	enum { PositionCardinality = 3 };
+	enum { RotationCardinality = 6 };
+	enum { LinearVelocityCardinality = 3 };
+	enum { PhaseCardinality = 2 };
 };
 
 
@@ -81,6 +86,10 @@ public:
 	UPROPERTY(EditAnywhere, Category="Settings")
 	EPoseSearchFeatureDomain Domain = EPoseSearchFeatureDomain::Time;
 
+	// @todo: temporary location for the channel bone weight to help the weights refactoring (later on it makes sense to have a weight per SampleOffsets, but right now for semplicity the traj will have just one weight)
+	UPROPERTY(EditAnywhere, Category = "Settings")
+	float Weight = 1.f;
+
 	UPROPERTY(EditAnywhere, Category = "Settings")
 	TArray<float> SampleOffsets;
 
@@ -89,9 +98,9 @@ public:
 
 	// UPoseSearchFeatureChannel interface
 	virtual void InitializeSchema(UE::PoseSearch::FSchemaInitializer& Initializer) override;
+	virtual void FillWeights(TArray<float>& Weights) const override;
 	virtual void IndexAsset(const UE::PoseSearch::IAssetIndexer& Indexer, UE::PoseSearch::FAssetIndexingOutput& IndexingOutput) const override;
 	virtual FFloatRange GetHorizonRange(EPoseSearchFeatureDomain Domain) const override;
-	virtual TArrayView<const float> GetSampleOffsets () const override { return SampleOffsets; }
 	virtual void GenerateDDCKey(FBlake3& InOutKeyHasher) const override;
 	virtual bool BuildQuery(
 		FPoseSearchContext& SearchContext,
@@ -101,4 +110,8 @@ public:
 protected:
 	virtual void IndexTimeFeatures(const UE::PoseSearch::IAssetIndexer& Indexer, int32 SampleIdx, FPoseSearchFeatureVectorBuilder& FeatureVector) const;
 	virtual void IndexDistanceFeatures(const UE::PoseSearch::IAssetIndexer& Indexer, int32 SampleIdx, FPoseSearchFeatureVectorBuilder& FeatureVector) const;
+
+	enum { PositionCardinality = 3 };
+	enum { LinearVelocityCardinality = 3 };
+	enum { ForwardVectorCardinality = 3 };
 };
\ No newline at end of file
