motion matching - UPoseSearchFeatureChannel_Pose character space velocity calculation fix for SubsampleTime != SampleTime bug fix

diff --git a/Engine/Plugins/Experimental/Animation/PoseSearch/Source/Runtime/Private/PoseSearchFeatureChannels.cpp b/Engine/Plugins/Experimental/Animation/PoseSearch/Source/Runtime/Private/PoseSearchFeatureChannels.cpp
--- a/Engine/Plugins/Experimental/Animation/PoseSearch/Source/Runtime/Private/PoseSearchFeatureChannels.cpp
+++ b/Engine/Plugins/Experimental/Animation/PoseSearch/Source/Runtime/Private/PoseSearchFeatureChannels.cpp
@@ -554,7 +554,7 @@ void UPoseSearchFeatureChannel_Pose::AddPoseFeatures(const  UE::PoseSearch::IAss
 		{
 			// character space velocity
 			Samples[0] = Indexer.GetSampleInfoRelative(SubsampleTime - SamplingContext->FiniteDelta, Indexer.GetSampleInfo(SubsampleTime - SamplingContext->FiniteDelta));
-			Samples[1] = Indexer.GetSampleInfoRelative(SubsampleTime, Indexer.GetSampleInfo(SampleTime));
+			Samples[1] = Indexer.GetSampleInfoRelative(SubsampleTime, Indexer.GetSampleInfo(SubsampleTime));
 			Samples[2] = Indexer.GetSampleInfoRelative(SubsampleTime + SamplingContext->FiniteDelta, Indexer.GetSampleInfo(SubsampleTime + SamplingContext->FiniteDelta));
 		}
 		else
