motion matching
- WIP deprecation of FPoseSearchFeatureDesc, FPoseSearchFeatureVectorLayout, FPoseSearchFeatureVectorBuilder, FFeatureVectorReader in favor of UPoseSearchFeatureChannel
- deprecated FPoseSearchDebuggerPoseVectorChannel
- introduced UE::PoseSearch::BoneTransformsCache to speed up UPoseSearchFeatureChannel_Pose::IndexAsset

diff --git a/Engine/Plugins/Experimental/Animation/PoseSearch/Source/Editor/Private/PoseSearchDebugger.cpp b/Engine/Plugins/Experimental/Animation/PoseSearch/Source/Editor/Private/PoseSearchDebugger.cpp
--- a/Engine/Plugins/Experimental/Animation/PoseSearch/Source/Editor/Private/PoseSearchDebugger.cpp
+++ b/Engine/Plugins/Experimental/Animation/PoseSearch/Source/Editor/Private/PoseSearchDebugger.cpp
@@ -35,96 +35,6 @@ static FLinearColor LinearColorBlend(FLinearColor LinearColorA, FLinearColor Lin
 	return LinearColorA + (LinearColorB - LinearColorA) * BlendParam;
 }
 
-void FPoseSearchDebuggerPoseVectorChannel::Reset()
-{
-	Positions.Reset();
-	LinearVelocities.Reset();
-	FacingDirections.Reset();
-}
-
-void FPoseSearchDebuggerPoseVector::Reset()
-{
-	Pose.Reset();
-	TrajectoryTimeBased.Reset();
-	TrajectoryDistanceBased.Reset();
-}
-
-void FPoseSearchDebuggerPoseVector::ExtractFeatures(const UE::PoseSearch::FFeatureVectorReader& Reader)
-{
-	using namespace UE::PoseSearch;
-
-	Reset();
-
-	FPoseSearchDebuggerPoseVectorChannel* Channels[] = 
-	{
-		&Pose,
-		&TrajectoryTimeBased,
-		&TrajectoryDistanceBased
-	};
-
-	for (const FPoseSearchFeatureDesc& Feature: Reader.GetLayout()->Features)
-	{
-		if (Feature.ChannelIdx >= UE_ARRAY_COUNT(Channels))
-		{
-			continue;
-		}
-
-		switch (Feature.Type)
-		{
-		case EPoseSearchFeatureType::Position:
-		{
-			FVector Vector;
-			if (Reader.GetVector(Feature, &Vector))
-			{
-				Channels[Feature.ChannelIdx]->Positions.Add(Vector);
-			}
-			break;
-		}
-		case EPoseSearchFeatureType::LinearVelocity:
-		{
-			FVector Vector;
-			if (Reader.GetVector(Feature, &Vector))
-			{
-				Channels[Feature.ChannelIdx]->LinearVelocities.Add(Vector);
-			}
-			break;
-
-		}
-		case EPoseSearchFeatureType::ForwardVector:
-		{
-			FVector Vector;
-			if (Reader.GetVector(Feature, &Vector))
-			{
-				Channels[Feature.ChannelIdx]->FacingDirections.Add(Vector);
-			}
-			break;
-		}
-
-		case EPoseSearchFeatureType::Phase:
-		{
-			FVector2D Phase;
-			if (Reader.GetPhase(Feature, &Phase))
-			{
-				Channels[Feature.ChannelIdx]->Phases.Add(Phase);
-			}
-			break;
-		}
-		}
-	}
-
-	for (FPoseSearchDebuggerPoseVectorChannel* Channel : Channels)
-	{
-		Channel->bShowPositions = !Channel->Positions.IsEmpty();
-		Channel->bShowLinearVelocities = !Channel->LinearVelocities.IsEmpty();
-		Channel->bShowFacingDirections = !Channel->FacingDirections.IsEmpty();
-		Channel->bShowPhases = !Channel->Phases.IsEmpty();
-	}
-
-	bShowPose = !Pose.IsEmpty();
-	bShowTrajectoryTimeBased = !TrajectoryTimeBased.IsEmpty();
-	bShowTrajectoryDistanceBased = !TrajectoryDistanceBased.IsEmpty();
-}
-
 void UPoseSearchMeshComponent::Initialize(const FTransform& InComponentToWorld)
 {
 	SetComponentToWorld(InComponentToWorld);
@@ -1557,18 +1467,12 @@ void SDebuggerDetailsView::UpdateReflection(const FTraceMotionMatchingStateMessa
 	Reflection->AnimLinearVelocity = State.AnimLinearVelocity;
 	Reflection->AnimAngularVelocity = State.AnimAngularVelocity;
 
-	FFeatureVectorReader Reader;
-	Reader.Init(&Schema->Layout);
-
 	// Query pose
-	Reader.SetValues(State.QueryVector);
-	Reflection->QueryPoseVector.ExtractFeatures(Reader);
+	Reflection->QueryPoseVector = State.QueryVector;
 
 	// Active pose
-	TArray<float> Pose(Database.GetSearchIndex()->GetPoseValues(State.DbPoseIdx));
-	Database.GetSearchIndex()->InverseNormalize(Pose);
-	Reader.SetValues(Pose);
-	Reflection->ActivePoseVector.ExtractFeatures(Reader);
+	Reflection->ActivePoseVector = Database.GetSearchIndex()->GetPoseValues(State.DbPoseIdx);
+	Database.GetSearchIndex()->InverseNormalize(Reflection->ActivePoseVector);
 
 	auto DebuggerView = ParentDebuggerViewPtr.Pin();
 	if (DebuggerView.IsValid()) 
@@ -1577,27 +1481,19 @@ void SDebuggerDetailsView::UpdateReflection(const FTraceMotionMatchingStateMessa
 		if (!SelectedRows.IsEmpty())
 		{
 			const TSharedRef<FDebuggerDatabaseRowData>& Selected = SelectedRows[0];
-			Pose = Database.GetSearchIndex()->GetPoseValues(Selected->PoseIdx);
-			Database.GetSearchIndex()->InverseNormalize(Pose);
-			Reader.SetValues(Pose);
-			Reflection->SelectedPoseVector.ExtractFeatures(Reader);
+			Reflection->SelectedPoseVector = Database.GetSearchIndex()->GetPoseValues(Selected->PoseIdx);
+			Database.GetSearchIndex()->InverseNormalize(Reflection->SelectedPoseVector);
 
-			Pose = Selected->PoseCostDetails.CostVector;
-			//Database.SearchIndex.InverseNormalize(Pose);
-			Reader.SetValues(Pose);
-			Reflection->CostVector.ExtractFeatures(Reader);
+			Reflection->CostVector = Selected->PoseCostDetails.CostVector;
+			//Database.SearchIndex.InverseNormalize(Reflection->CostVector);
 
-			const TSharedRef<FDebuggerDatabaseRowData>& ActiveRow = 
-				DebuggerView->GetPoseIdxDatabaseRow(State.DbPoseIdx);
+			const TSharedRef<FDebuggerDatabaseRowData>& ActiveRow = DebuggerView->GetPoseIdxDatabaseRow(State.DbPoseIdx);
 
-			TArray<float> ActiveCostDifference(Pose);
-			for (int i = 0; i < ActiveCostDifference.Num(); ++i)
+			Reflection->CostVectorDifference = Reflection->CostVector;
+			for (int i = 0; i < Reflection->CostVectorDifference.Num(); ++i)
 			{
-				ActiveCostDifference[i] -= ActiveRow->PoseCostDetails.CostVector[i];
+				Reflection->CostVectorDifference[i] -= ActiveRow->PoseCostDetails.CostVector[i];
 			}
-
-			Reader.SetValues(ActiveCostDifference);
-			Reflection->CostVectorDifference.ExtractFeatures(Reader);
 		}
 	}
 }
diff --git a/Engine/Plugins/Experimental/Animation/PoseSearch/Source/Editor/Private/PoseSearchDebugger.h b/Engine/Plugins/Experimental/Animation/PoseSearch/Source/Editor/Private/PoseSearchDebugger.h
--- a/Engine/Plugins/Experimental/Animation/PoseSearch/Source/Editor/Private/PoseSearchDebugger.h
+++ b/Engine/Plugins/Experimental/Animation/PoseSearch/Source/Editor/Private/PoseSearchDebugger.h
@@ -70,74 +70,6 @@ public:
 	FTransform LastRootMotionDelta = FTransform::Identity;
 };
 
-
-/**
- * Used by the reflection UObject to encompass a set of feature vectors
- */
-USTRUCT()
-struct FPoseSearchDebuggerPoseVectorChannel
-{
-	GENERATED_BODY()
-
-	// @TODO: Should be ideally enumerated based on all possible schema features
-
-	UPROPERTY(VisibleDefaultsOnly, Category = "Channel")
-	bool bShowPositions = false;
-	
-	UPROPERTY(VisibleDefaultsOnly, Category = "Channel")
-	bool bShowLinearVelocities = false;
-
-	UPROPERTY(VisibleDefaultsOnly, Category = "Channel")
-	bool bShowFacingDirections = false;
-
-	UPROPERTY(VisibleDefaultsOnly, Category = "Channel")
-	bool bShowPhases = false;
-
-    UPROPERTY(VisibleAnywhere, EditFixedSize, Category="Channel", meta=(EditCondition="bShowPositions", EditConditionHides))
-	TArray<FVector> Positions;
-
-    UPROPERTY(VisibleAnywhere, EditFixedSize, Category="Channel", meta=(EditCondition="bShowLinearVelocities", EditConditionHides))
-	TArray<FVector> LinearVelocities;
-	
-    UPROPERTY(VisibleAnywhere, EditFixedSize, Category="Channel", meta=(EditCondition="bShowFacingDirections", EditConditionHides))
-	TArray<FVector> FacingDirections;
-
-	UPROPERTY(VisibleAnywhere, EditFixedSize, Category = "Channel", meta = (EditCondition = "bShowPhases", EditConditionHides))
-	TArray<FVector2D> Phases;
-
-	void Reset();
-	bool IsEmpty() const { return Positions.IsEmpty() && LinearVelocities.IsEmpty() && FacingDirections.IsEmpty(); }
-};
-
-USTRUCT()
-struct FPoseSearchDebuggerPoseVector
-{
-	GENERATED_BODY()
-
-
-	UPROPERTY(VisibleDefaultsOnly, Category = "Pose Vector")
-	bool bShowPose = false;
-
-	UPROPERTY(VisibleDefaultsOnly, Category = "Pose Vector")
-	bool bShowTrajectoryTimeBased = false;
-
-	UPROPERTY(VisibleDefaultsOnly, Category = "Pose Vector")
-	bool bShowTrajectoryDistanceBased = false;
-
-	UPROPERTY(VisibleAnywhere, Category="Pose Vector", meta=(EditCondition="bShowPose", EditConditionHides))
-	FPoseSearchDebuggerPoseVectorChannel Pose;
-
-	UPROPERTY(VisibleAnywhere, Category="Pose Vector", meta=(EditCondition="bShowTrajectoryTimeBased", EditConditionHides, DisplayName="Trajectory (Time)"))
-	FPoseSearchDebuggerPoseVectorChannel TrajectoryTimeBased;
-
-	UPROPERTY(VisibleAnywhere, Category="Pose Vector", meta=(EditCondition="bShowTrajectoryDistanceBased", EditConditionHides, DisplayName="Trajectory (Distance)"))
-	FPoseSearchDebuggerPoseVectorChannel TrajectoryDistanceBased;
-
-	void Reset();
-	void ExtractFeatures(const UE::PoseSearch::FFeatureVectorReader& Reader);
-	bool IsEmpty() const { return Pose.IsEmpty() && TrajectoryTimeBased.IsEmpty() && TrajectoryDistanceBased.IsEmpty(); }
-};
-
 /**
  * Used by the reflection UObject to encompass draw options for the query and database selections
  */
@@ -214,21 +146,21 @@ public:
 	bool bDrawSelectedSkeleton = true;
 
     UPROPERTY(VisibleAnywhere, Category="Pose Vectors")
-	FPoseSearchDebuggerPoseVector QueryPoseVector;
+	TArray<float> QueryPoseVector;
     	
     UPROPERTY(VisibleAnywhere, Category="Pose Vectors")
-	FPoseSearchDebuggerPoseVector ActivePoseVector;
+	TArray<float> ActivePoseVector;
 
 	UPROPERTY(VisibleAnywhere, Category="Pose Vectors")
-	FPoseSearchDebuggerPoseVector SelectedPoseVector;
+	TArray<float> SelectedPoseVector;
 
 	UPROPERTY(VisibleAnywhere, Category="Pose Vectors")
-	FPoseSearchDebuggerPoseVector CostVector;
+	TArray<float> CostVector;
 
 	// Cost comparison of selected and active poses. A negative value indicates the cost in the selected pose is lower;
 	// a positive value indicates the cost in the selected pose is higher.
 	UPROPERTY(VisibleAnywhere, Category="Pose Vectors")
-	FPoseSearchDebuggerPoseVector CostVectorDifference;
+	TArray<float> CostVectorDifference;
 };
 
 
diff --git a/Engine/Plugins/Experimental/Animation/PoseSearch/Source/Runtime/Private/PoseSearch.cpp b/Engine/Plugins/Experimental/Animation/PoseSearch/Source/Runtime/Private/PoseSearch.cpp
--- a/Engine/Plugins/Experimental/Animation/PoseSearch/Source/Runtime/Private/PoseSearch.cpp
+++ b/Engine/Plugins/Experimental/Animation/PoseSearch/Source/Runtime/Private/PoseSearch.cpp
@@ -299,44 +299,6 @@ bool FPoseSearchFeatureVectorLayout::IsValid(int32 ChannelCount) const
 	return true;
 }
 
-bool FPoseSearchFeatureVectorLayout::EnumerateBy(int32 ChannelIdx, EPoseSearchFeatureType Type, int32& InOutFeatureIdx) const
-{
-	auto IsChannelMatch = [](int32 ChannelIdx, const FPoseSearchFeatureDesc& Feature) -> bool
-	{
-		bool bChannelMatch = true;
-		if (ChannelIdx >= 0)
-		{
-			bChannelMatch = Feature.ChannelIdx == ChannelIdx;
-		}
-		return bChannelMatch;
-	};
-
-	auto IsTypeMatch = [](EPoseSearchFeatureType Type, const FPoseSearchFeatureDesc& Feature) -> bool
-	{
-		bool bTypeMatch = true;
-		if (Type != EPoseSearchFeatureType::Invalid)
-		{
-			bTypeMatch = Feature.Type == Type;
-		}
-		return bTypeMatch;
-	};
-
-	for (int32 Size = Features.Num(); ++InOutFeatureIdx < Size; )
-	{
-		const FPoseSearchFeatureDesc& Feature = Features[InOutFeatureIdx]; 
-
-		bool bChannelMatch = IsChannelMatch(ChannelIdx, Feature);
-		bool bTypeMatch = IsTypeMatch(Type, Feature);
-
-		if (bChannelMatch && bTypeMatch)
-		{
-			return true;
-		}
-	}
-
-	return false;
-}
-
 //////////////////////////////////////////////////////////////////////////
 // UPoseSearchSchema
 
@@ -354,10 +316,9 @@ int32 FSchemaInitializer::AddFeatureDesc(const FPoseSearchFeatureDesc& FeatureDe
 	check(Features.Num() + 1 <= UPoseSearchSchema::MaxFeatures);
 	check(!Features.Contains(FeatureDesc));
 	check(FeatureDesc.Cardinality > 0);
-	const int32 FeatureDescIndex = Features.Add(FeatureDesc);
-	Features[FeatureDescIndex].ValueOffset = CurrentChannelDataOffset;
+	check(FeatureDesc.ValueOffset == CurrentChannelDataOffset);
 	CurrentChannelDataOffset += FeatureDesc.Cardinality;
-	return FeatureDescIndex;
+	return Features.Add(FeatureDesc);
 }
 
 } // namespace UE::PoseSearch
@@ -2195,144 +2156,75 @@ void FPoseSearchFeatureVectorBuilder::ResetFeatures()
 	FeaturesAdded.Init(false, Schema->Layout.Features.Num());
 }
 
-void FPoseSearchFeatureVectorBuilder::SetTransform(FPoseSearchFeatureDesc Feature, const FTransform& Transform)
+void FPoseSearchFeatureVectorBuilder::SetPhase(const FPoseSearchFeatureDesc& Feature, const FVector2D& Phase)
 {
-	SetPosition(Feature, Transform.GetTranslation());
-	SetRotation(Feature, Transform.GetRotation());
-}
+	check(Feature.Type == EPoseSearchFeatureType::Phase);
+	check(Feature.Cardinality == 2);
 
-void FPoseSearchFeatureVectorBuilder::SetTransformVelocity(FPoseSearchFeatureDesc Feature, const FTransform& Transform, const FTransform& PrevTransform, float DeltaTime)
-{
-	SetLinearVelocity(Feature, Transform, PrevTransform, DeltaTime);
-	SetAngularVelocity(Feature, Transform, PrevTransform, DeltaTime);
-}
-
-void FPoseSearchFeatureVectorBuilder::SetTransformVelocity(FPoseSearchFeatureDesc Feature, const FTransform& NextTransform, const FTransform& Transform, const FTransform& PrevTransform, float DeltaTime)
-{
-	SetLinearVelocity(Feature, NextTransform, Transform, PrevTransform, DeltaTime);
-	SetAngularVelocity(Feature, NextTransform, Transform, PrevTransform, DeltaTime);
-}
-
-void FPoseSearchFeatureVectorBuilder::SetPosition(FPoseSearchFeatureDesc Feature, const FVector& Position)
-{
-	Feature.Type = EPoseSearchFeatureType::Position;
-	SetVector(Feature, Position);
-}
-
-void FPoseSearchFeatureVectorBuilder::SetPhase(FPoseSearchFeatureDesc Feature, const FVector2D& Phase)
-{
-	Feature.Type = EPoseSearchFeatureType::Phase;
-	
-	int32 ElementIndex = Schema->Layout.Features.Find(Feature);
-	if (ElementIndex >= 0)
+	const int32 FeatureIndex = Schema->Layout.Features.Find(Feature);
+	check(FeatureIndex >= 0);
+	if (FeatureIndex >= 0)
 	{
-		const FPoseSearchFeatureDesc& FoundElement = Schema->Layout.Features[ElementIndex];
+		const FPoseSearchFeatureDesc& FoundFeature = Schema->Layout.Features[FeatureIndex];
+		check(Feature.ValueOffset < 0 || Feature.ValueOffset == FoundFeature.ValueOffset);
 
-		Values[FoundElement.ValueOffset + 0] = Phase.X;
-		Values[FoundElement.ValueOffset + 1] = Phase.Y;
+		Values[FoundFeature.ValueOffset + 0] = Phase.X;
+		Values[FoundFeature.ValueOffset + 1] = Phase.Y;
 
-		if (!FeaturesAdded[ElementIndex])
+		if (!FeaturesAdded[FeatureIndex])
 		{
-			FeaturesAdded[ElementIndex] = true;
+			FeaturesAdded[FeatureIndex] = true;
 			++NumFeaturesAdded;
 		}
 	}
 }
 
-void FPoseSearchFeatureVectorBuilder::SetRotation(FPoseSearchFeatureDesc Feature, const FQuat& Rotation)
+void FPoseSearchFeatureVectorBuilder::SetRotation(const FPoseSearchFeatureDesc& Feature, const FQuat& Rotation)
 {
-	Feature.Type = EPoseSearchFeatureType::Rotation;
-	int32 ElementIndex = Schema->Layout.Features.Find(Feature);
-	if (ElementIndex >= 0)
+	check(Feature.Type == EPoseSearchFeatureType::Rotation);
+	check(Feature.Cardinality == 6);
+	const int32 FeatureIndex = Schema->Layout.Features.Find(Feature);
+	check(FeatureIndex >= 0);
+	if (FeatureIndex >= 0)
 	{
 		FVector X = Rotation.GetAxisX();
 		FVector Y = Rotation.GetAxisY();
 
-		const FPoseSearchFeatureDesc& FoundElement = Schema->Layout.Features[ElementIndex];
+		const FPoseSearchFeatureDesc& FoundFeature = Schema->Layout.Features[FeatureIndex];
+		check(Feature.ValueOffset < 0 || Feature.ValueOffset == FoundFeature.ValueOffset);
 
-		Values[FoundElement.ValueOffset + 0] = X.X;
-		Values[FoundElement.ValueOffset + 1] = X.Y;
-		Values[FoundElement.ValueOffset + 2] = X.Z;
-		Values[FoundElement.ValueOffset + 3] = Y.X;
-		Values[FoundElement.ValueOffset + 4] = Y.Y;
-		Values[FoundElement.ValueOffset + 5] = Y.Z;
+		Values[FoundFeature.ValueOffset + 0] = X.X;
+		Values[FoundFeature.ValueOffset + 1] = X.Y;
+		Values[FoundFeature.ValueOffset + 2] = X.Z;
+		Values[FoundFeature.ValueOffset + 3] = Y.X;
+		Values[FoundFeature.ValueOffset + 4] = Y.Y;
+		Values[FoundFeature.ValueOffset + 5] = Y.Z;
 
-		if (!FeaturesAdded[ElementIndex])
+		if (!FeaturesAdded[FeatureIndex])
 		{
-			FeaturesAdded[ElementIndex] = true;
+			FeaturesAdded[FeatureIndex] = true;
 			++NumFeaturesAdded;
 		}
 	}
-
-	Feature.Type = EPoseSearchFeatureType::ForwardVector;
-	SetVector(Feature, Rotation.GetAxisY());
 }
 
-void FPoseSearchFeatureVectorBuilder::SetLinearVelocity(FPoseSearchFeatureDesc Feature, const FTransform& Transform, const FTransform& PrevTransform, float DeltaTime)
+void FPoseSearchFeatureVectorBuilder::SetVector(const FPoseSearchFeatureDesc& Feature, const FVector& Vector)
 {
-	Feature.Type = EPoseSearchFeatureType::LinearVelocity;
-	FVector LinearVelocity = (Transform.GetTranslation() - PrevTransform.GetTranslation()) / DeltaTime;
-	SetVector(Feature, LinearVelocity);
-}
-
-void FPoseSearchFeatureVectorBuilder::SetLinearVelocity(FPoseSearchFeatureDesc Feature, const FTransform& NextTransform, const FTransform& Transform, const FTransform& PrevTransform, float DeltaTime)
-{
-	Feature.Type = EPoseSearchFeatureType::LinearVelocity;
-	FVector LinearVelocityNext = (NextTransform.GetTranslation() - Transform.GetTranslation()) / DeltaTime;
-	FVector LinearVelocityPrev = (Transform.GetTranslation() - PrevTransform.GetTranslation()) / DeltaTime;
-	SetVector(Feature, (LinearVelocityNext + LinearVelocityPrev) / 2.0f);
-}
-
-static inline FVector QuaternionAngularVelocity(const FQuat& Rotation, const FQuat& PrevRotation, float DeltaTime)
-{
-	FQuat Q0 = PrevRotation;
-	FQuat Q1 = Rotation;
-	Q1.EnforceShortestArcWith(Q0);
-
-	// Given angular velocity vector w, quaternion differentiation can be represented as
-	//   dq/dt = (w * q)/2
-	// Solve for w
-	//   w = 2 * dq/dt * q^-1
-	// And let dq/dt be expressed as the finite difference
-	//   dq/dt = (q(t+h) - q(t)) / h
-	FQuat DQDt = (Q1 - Q0) / DeltaTime;
-	FQuat QInv = Q0.Inverse();
-	FQuat W = (DQDt * QInv) * 2.0f;
-
-	FVector AngularVelocity(W.X, W.Y, W.Z);
-
-	return AngularVelocity;
-}
-
-void FPoseSearchFeatureVectorBuilder::SetAngularVelocity(FPoseSearchFeatureDesc Feature, const FTransform& Transform, const FTransform& PrevTransform, float DeltaTime)
-{
-	Feature.Type = EPoseSearchFeatureType::AngularVelocity;
-	FVector AngularVelocity = QuaternionAngularVelocity(Transform.GetRotation(), PrevTransform.GetRotation(), DeltaTime);
-	SetVector(Feature, AngularVelocity);
-}
-
-void FPoseSearchFeatureVectorBuilder::SetAngularVelocity(FPoseSearchFeatureDesc Feature, const FTransform& NextTransform, const FTransform& Transform, const FTransform& PrevTransform, float DeltaTime)
-{
-	Feature.Type = EPoseSearchFeatureType::AngularVelocity;
-	FVector AngularVelocityNext = QuaternionAngularVelocity(NextTransform.GetRotation(), Transform.GetRotation(), DeltaTime);
-	FVector AngularVelocityPrev = QuaternionAngularVelocity(Transform.GetRotation(), PrevTransform.GetRotation(), DeltaTime);
-	SetVector(Feature, (AngularVelocityNext + AngularVelocityPrev) / 2.0f);
-}
-
-void FPoseSearchFeatureVectorBuilder::SetVector(FPoseSearchFeatureDesc Feature, const FVector& Vector)
-{
-	int32 ElementIndex = Schema->Layout.Features.Find(Feature);
-	if (ElementIndex >= 0)
+	check(Feature.Cardinality == 3);
+	const int32 FeatureIndex = Schema->Layout.Features.Find(Feature);
+	check(FeatureIndex >= 0);
+	if (FeatureIndex >= 0)
 	{
-		const FPoseSearchFeatureDesc& FoundElement = Schema->Layout.Features[ElementIndex]; 
+		const FPoseSearchFeatureDesc& FoundFeature = Schema->Layout.Features[FeatureIndex]; 
+		check(Feature.ValueOffset < 0 || Feature.ValueOffset == FoundFeature.ValueOffset);
 
-		Values[FoundElement.ValueOffset + 0] = Vector[0];
-		Values[FoundElement.ValueOffset + 1] = Vector[1];
-		Values[FoundElement.ValueOffset + 2] = Vector[2];
+		Values[FoundFeature.ValueOffset + 0] = Vector[0];
+		Values[FoundFeature.ValueOffset + 1] = Vector[1];
+		Values[FoundFeature.ValueOffset + 2] = Vector[2];
 
-		if (!FeaturesAdded[ElementIndex])
+		if (!FeaturesAdded[FeatureIndex])
 		{
-			FeaturesAdded[ElementIndex] = true;
+			FeaturesAdded[FeatureIndex] = true;
 			++NumFeaturesAdded;
 		}
 	}
@@ -2426,7 +2318,7 @@ void FPoseHistory::Init(const FPoseHistory& History)
 	TimeHorizon = History.TimeHorizon;
 }
 
-bool FPoseHistory::TrySampleLocalPose(float SecondsAgo, const TArray<FBoneIndexType>& RequiredBones, TArray<FTransform>& LocalPose, FTransform& RootTransform)
+bool FPoseHistory::TrySampleLocalPose(float SecondsAgo, const TArray<FBoneIndexType>& RequiredBones, TArray<FTransform>& LocalPose, FTransform& RootTransform) const
 {
 	int32 NextIdx = LowerBound(Knots.begin(), Knots.end(), SecondsAgo, TGreater<>());
 	if (NextIdx <= 0 || NextIdx >= Knots.Num())
@@ -2610,41 +2502,25 @@ bool FFeatureVectorReader::IsValid() const
 	return Layout && (Layout->NumFloats == Values.Num());
 }
 
-bool FFeatureVectorReader::GetTransform(FPoseSearchFeatureDesc Element, FTransform* OutTransform) const
+bool FFeatureVectorReader::GetRotation(const FPoseSearchFeatureDesc& Feature, FQuat* OutRotation) const
 {
-	FVector Position;
-	bool bResult = GetPosition(Element, &Position);
-
-	FQuat Rotation;
-	bResult |= GetRotation(Element, &Rotation);
-
-	OutTransform->SetComponents(Rotation, Position, FVector::OneVector);
-	return bResult;
-}
-
-bool FFeatureVectorReader::GetPosition(FPoseSearchFeatureDesc Element, FVector* OutPosition) const
-{
-	Element.Type = EPoseSearchFeatureType::Position;
-	return GetVector(Element, OutPosition);
-}
-
-bool FFeatureVectorReader::GetRotation(FPoseSearchFeatureDesc Element, FQuat* OutRotation) const
-{
-	Element.Type = EPoseSearchFeatureType::Rotation;
-	int32 ElementIndex = IsValid() ? Layout->Features.Find(Element) : -1;
-	if (ElementIndex >= 0)
+	check(Feature.Type == EPoseSearchFeatureType::Rotation);
+	const int32 FeatureIndex = IsValid() ? Layout->Features.Find(Feature) : -1;
+	check(FeatureIndex >= 0);
+	if (FeatureIndex >= 0)
 	{
-		const FPoseSearchFeatureDesc& FoundElement = Layout->Features[ElementIndex];
+		const FPoseSearchFeatureDesc& FoundFeature = Layout->Features[FeatureIndex];
+		check(Feature.ValueOffset < 0 || Feature.ValueOffset == FoundFeature.ValueOffset);
 
 		FVector X;
 		FVector Y;
 
-		X.X = Values[FoundElement.ValueOffset + 0];
-		X.Y = Values[FoundElement.ValueOffset + 1];
-		X.Z = Values[FoundElement.ValueOffset + 2];
-		Y.X = Values[FoundElement.ValueOffset + 3];
-		Y.Y = Values[FoundElement.ValueOffset + 4];
-		Y.Z = Values[FoundElement.ValueOffset + 5];
+		X.X = Values[FoundFeature.ValueOffset + 0];
+		X.Y = Values[FoundFeature.ValueOffset + 1];
+		X.Z = Values[FoundFeature.ValueOffset + 2];
+		Y.X = Values[FoundFeature.ValueOffset + 3];
+		Y.Y = Values[FoundFeature.ValueOffset + 4];
+		Y.Z = Values[FoundFeature.ValueOffset + 5];
 
 		FVector Z = FVector::CrossProduct(X, Y);
 
@@ -2661,34 +2537,19 @@ bool FFeatureVectorReader::GetRotation(FPoseSearchFeatureDesc Element, FQuat* Ou
 	return false;
 }
 
-bool FFeatureVectorReader::GetForwardVector(FPoseSearchFeatureDesc Element, FVector* OutForwardVector) const
-{
-	Element.Type = EPoseSearchFeatureType::ForwardVector;
-	return GetVector(Element, OutForwardVector);
-}
-
-bool FFeatureVectorReader::GetLinearVelocity(FPoseSearchFeatureDesc Element, FVector* OutLinearVelocity) const
-{
-	Element.Type = EPoseSearchFeatureType::LinearVelocity;
-	return GetVector(Element, OutLinearVelocity);
-}
-
-bool FFeatureVectorReader::GetAngularVelocity(FPoseSearchFeatureDesc Element, FVector* OutAngularVelocity) const
-{
-	Element.Type = EPoseSearchFeatureType::AngularVelocity;
-	return GetVector(Element, OutAngularVelocity);
-}
-
-bool FFeatureVectorReader::GetPhase(FPoseSearchFeatureDesc Element, FVector2D* OutPhase) const
+bool FFeatureVectorReader::GetPhase(const FPoseSearchFeatureDesc& Feature, FVector2D* OutPhase) const
 {
-	Element.Type = EPoseSearchFeatureType::Phase;
-	int32 ElementIndex = IsValid() ? Layout->Features.Find(Element) : -1;
-	if (ElementIndex >= 0)
+	check(Feature.Cardinality == 2);
+	check(Feature.Type == EPoseSearchFeatureType::Phase);
+	const int32 FeatureIndex = IsValid() ? Layout->Features.Find(Feature) : -1;
+	check(FeatureIndex >= 0);
+	if (FeatureIndex >= 0)
 	{
-		const FPoseSearchFeatureDesc& FoundElement = Layout->Features[ElementIndex];
+		const FPoseSearchFeatureDesc& FoundFeature = Layout->Features[FeatureIndex];
+		check(Feature.ValueOffset < 0 || Feature.ValueOffset == FoundFeature.ValueOffset);
 
-		OutPhase->X = Values[FoundElement.ValueOffset + 0];
-		OutPhase->Y = Values[FoundElement.ValueOffset + 1];
+		OutPhase->X = Values[FoundFeature.ValueOffset + 0];
+		OutPhase->Y = Values[FoundFeature.ValueOffset + 1];
 		return true;
 	}
 
@@ -2696,17 +2557,20 @@ bool FFeatureVectorReader::GetPhase(FPoseSearchFeatureDesc Element, FVector2D* O
 	return false;
 }
 
-bool FFeatureVectorReader::GetVector(FPoseSearchFeatureDesc Element, FVector* OutVector) const
+bool FFeatureVectorReader::GetVector(const FPoseSearchFeatureDesc& Feature, FVector* OutVector) const
 {
-	int32 ElementIndex = IsValid() ? Layout->Features.Find(Element) : -1;
-	if (ElementIndex >= 0)
+	check(Feature.Cardinality == 3);
+	const int32 FeatureIndex = IsValid() ? Layout->Features.Find(Feature) : -1;
+	check(FeatureIndex >= 0);
+	if (FeatureIndex >= 0)
 	{
-		const FPoseSearchFeatureDesc& FoundElement = Layout->Features[ElementIndex];
+		const FPoseSearchFeatureDesc& FoundFeature = Layout->Features[FeatureIndex];
+		check(Feature.ValueOffset < 0 || Feature.ValueOffset == FoundFeature.ValueOffset);
 
 		FVector V;
-		V.X = Values[FoundElement.ValueOffset + 0];
-		V.Y = Values[FoundElement.ValueOffset + 1];
-		V.Z = Values[FoundElement.ValueOffset + 2];
+		V.X = Values[FoundFeature.ValueOffset + 0];
+		V.Y = Values[FoundFeature.ValueOffset + 1];
+		V.Z = Values[FoundFeature.ValueOffset + 2];
 		*OutVector = V;
 		return true;
 	}
diff --git a/Engine/Plugins/Experimental/Animation/PoseSearch/Source/Runtime/Private/PoseSearchFeatureChannels.cpp b/Engine/Plugins/Experimental/Animation/PoseSearch/Source/Runtime/Private/PoseSearchFeatureChannels.cpp
--- a/Engine/Plugins/Experimental/Animation/PoseSearch/Source/Runtime/Private/PoseSearchFeatureChannels.cpp
+++ b/Engine/Plugins/Experimental/Animation/PoseSearch/Source/Runtime/Private/PoseSearchFeatureChannels.cpp
@@ -28,7 +28,7 @@ constexpr bool UseCharacterSpaceVelocities = true;
 //////////////////////////////////////////////////////////////////////////
 // Drawing helpers
 
-static FLinearColor GetColorForFeature(FPoseSearchFeatureDesc Feature, const FPoseSearchFeatureVectorLayout* Layout)
+static FLinearColor GetColorForFeature(const FPoseSearchFeatureDesc& Feature, const FPoseSearchFeatureVectorLayout* Layout)
 {
 	const float FeatureIdx = Layout->Features.IndexOfByKey(Feature);
 	const float FeatureCountIdx = Layout->Features.Num() - 1;
@@ -260,6 +260,92 @@ static void CalculatePhasesFromLocalMinMax(const TArray<LocalMinMax>& MinMax, TA
 	}
 }
 
+struct BoneTransformsCache
+{
+	BoneTransformsCache(const IAssetIndexer& InIndexer)
+		: Indexer(InIndexer)
+	{
+	}
+
+	FTransform Get(float SampleTime, float OriginTime, int8 SchemaBoneIdx, bool& Clamped)
+	{
+		// @todo: use an hashmap if we end up having too many entries
+		Entry* Entry = Entries.FindByPredicate([SampleTime, OriginTime](const BoneTransformsCache::Entry& Entry)
+		{
+			return Entry.SampleTime == SampleTime && Entry.OriginTime == OriginTime;
+		});
+
+		const FAssetIndexingContext& IndexingContext = Indexer.GetIndexingContext();
+		const FAssetSamplingContext* SamplingContext = IndexingContext.SamplingContext;
+
+		if (!Entry)
+		{
+			Entry = &Entries[Entries.AddDefaulted()];
+
+			Entry->SampleTime = SampleTime;
+			Entry->OriginTime = OriginTime;
+
+			Entry->Pose.SetBoneContainer(&SamplingContext->BoneContainer);
+			Entry->UnusedCurve.InitFrom(SamplingContext->BoneContainer);
+
+			IAssetIndexer::FSampleInfo Origin = Indexer.GetSampleInfo(OriginTime);
+			IAssetIndexer::FSampleInfo Sample = Indexer.GetSampleInfoRelative(SampleTime, Origin);
+
+			float CurrentTime = Sample.ClipTime;
+			float PreviousTime = CurrentTime - SamplingContext->FiniteDelta;
+
+			FDeltaTimeRecord DeltaTimeRecord;
+			DeltaTimeRecord.Set(PreviousTime, CurrentTime - PreviousTime);
+			FAnimExtractContext ExtractionCtx(CurrentTime, true, DeltaTimeRecord, Sample.Clip->IsLoopable());
+
+			Sample.Clip->ExtractPose(ExtractionCtx, Entry->AnimPoseData);
+
+			if (IndexingContext.bMirrored)
+			{
+				FAnimationRuntime::MirrorPose(
+					Entry->AnimPoseData.GetPose(),
+					IndexingContext.Schema->MirrorDataTable->MirrorAxis,
+					SamplingContext->CompactPoseMirrorBones,
+					SamplingContext->ComponentSpaceRefRotations
+				);
+				// Note curves and attributes are not used during the indexing process and therefore don't need to be mirrored
+			}
+
+			Entry->ComponentSpacePose.InitPose(Entry->Pose);
+			Entry->RootTransform = Sample.RootTransform;
+			Entry->Clamped = Sample.bClamped;
+		}
+
+		const FBoneReference& BoneReference = IndexingContext.Schema->BoneReferences[SchemaBoneIdx];
+		FCompactPoseBoneIndex CompactBoneIndex = SamplingContext->BoneContainer.MakeCompactPoseIndex(FMeshPoseBoneIndex(BoneReference.BoneIndex));
+
+		const FTransform BoneTransform = Entry->ComponentSpacePose.GetComponentSpaceTransform(CompactBoneIndex) * Indexer.MirrorTransform(Entry->RootTransform);
+		Clamped = Entry->Clamped;
+
+		return BoneTransform;
+	}
+
+	const UE::PoseSearch::IAssetIndexer& Indexer;
+
+	struct Entry
+	{
+		float SampleTime;
+		float OriginTime;
+		bool Clamped;
+
+		// @todo: minimize the Entry memory footprint
+		FTransform RootTransform;
+		FCompactPose Pose;
+		FCSPose<FCompactPose> ComponentSpacePose;
+		FBlendedCurve UnusedCurve;
+		UE::Anim::FStackAttributeContainer UnusedAtrribute;
+		FAnimationPoseData AnimPoseData = { Pose, UnusedCurve, UnusedAtrribute };
+	};
+
+	TArray<Entry> Entries;
+};
+
+
 } // namespace UE::PoseSearch
 
 //////////////////////////////////////////////////////////////////////////
@@ -272,8 +358,6 @@ class USkeleton* UPoseSearchFeatureChannel::GetSkeleton(bool& bInvalidSkeletonIs
 	return Schema ? Schema->Skeleton : nullptr;
 }
 
-
-
 //////////////////////////////////////////////////////////////////////////
 // UPoseSearchFeatureChannel_Pose
 
@@ -287,36 +371,40 @@ void UPoseSearchFeatureChannel_Pose::InitializeSchema(UE::PoseSearch::FSchemaIni
 {
 	Super::InitializeSchema(Initializer);
 
-	const int32 NumBones = SampledBones.Num();
-	for (int32 ChannelBoneIdx = 0; ChannelBoneIdx != NumBones; ++ChannelBoneIdx)
+	int32 DataOffset = ChannelDataOffset;
+	for (int32 ChannelBoneIdx = 0; ChannelBoneIdx != SampledBones.Num(); ++ChannelBoneIdx)
 	{
 		const FPoseSearchBone& SampledBone = SampledBones[ChannelBoneIdx];
 		if (SampledBone.bUsePosition)
 		{
 			for (int32 SubsampleIdx = 0; SubsampleIdx != SampleTimes.Num(); ++SubsampleIdx)
 			{
-				Initializer.AddFeatureDesc(FPoseSearchFeatureDesc::Construct(GetChannelIndex(), ChannelBoneIdx, SubsampleIdx, EPoseSearchFeatureType::Position, PositionCardinality));
+				Initializer.AddFeatureDesc(FPoseSearchFeatureDesc::Construct(GetChannelIndex(), ChannelBoneIdx, SubsampleIdx, EPoseSearchFeatureType::Position, PositionCardinality, DataOffset));
+				DataOffset += PositionCardinality;
 			}
 		}
 		if (SampledBone.bUseRotation)
 		{
 			for (int32 SubsampleIdx = 0; SubsampleIdx != SampleTimes.Num(); ++SubsampleIdx)
 			{
-				Initializer.AddFeatureDesc(FPoseSearchFeatureDesc::Construct(GetChannelIndex(), ChannelBoneIdx, SubsampleIdx, EPoseSearchFeatureType::Rotation, RotationCardinality));
+				Initializer.AddFeatureDesc(FPoseSearchFeatureDesc::Construct(GetChannelIndex(), ChannelBoneIdx, SubsampleIdx, EPoseSearchFeatureType::Rotation, RotationCardinality, DataOffset));
+				DataOffset += RotationCardinality;
 			}
 		}
 		if (SampledBone.bUseVelocity)
 		{
 			for (int32 SubsampleIdx = 0; SubsampleIdx != SampleTimes.Num(); ++SubsampleIdx)
 			{
-				Initializer.AddFeatureDesc(FPoseSearchFeatureDesc::Construct(GetChannelIndex(), ChannelBoneIdx, SubsampleIdx, EPoseSearchFeatureType::LinearVelocity, LinearVelocityCardinality));
+				Initializer.AddFeatureDesc(FPoseSearchFeatureDesc::Construct(GetChannelIndex(), ChannelBoneIdx, SubsampleIdx, EPoseSearchFeatureType::LinearVelocity, LinearVelocityCardinality, DataOffset));
+				DataOffset += LinearVelocityCardinality;
 			}
 		}
 		if (SampledBone.bUsePhase)
 		{
 			for (int32 SubsampleIdx = 0; SubsampleIdx != SampleTimes.Num(); ++SubsampleIdx)
 			{
-				Initializer.AddFeatureDesc(FPoseSearchFeatureDesc::Construct(GetChannelIndex(), ChannelBoneIdx, SubsampleIdx, EPoseSearchFeatureType::Phase, PhaseCardinality));
+				Initializer.AddFeatureDesc(FPoseSearchFeatureDesc::Construct(GetChannelIndex(), ChannelBoneIdx, SubsampleIdx, EPoseSearchFeatureType::Phase, PhaseCardinality, DataOffset));
+				DataOffset += PhaseCardinality;
 			}
 		}
 	}
@@ -378,19 +466,17 @@ void UPoseSearchFeatureChannel_Pose::FillWeights(TArray<float>& Weights) const
 }
 
 // @todo: do we really need to use double(s) in all this math?
-void UPoseSearchFeatureChannel_Pose::CalculatePhases(const UE::PoseSearch::IAssetIndexer& Indexer, UE::PoseSearch::FAssetIndexingOutput& IndexingOutput, TArray<TArray<FVector2D>>& OutPhases) const
+void UPoseSearchFeatureChannel_Pose::CalculatePhases(UE::PoseSearch::BoneTransformsCache& BoneTransformsCache, UE::PoseSearch::FAssetIndexingOutput& IndexingOutput, TArray<TArray<FVector2D>>& OutPhases) const
 {
 	// @todo: expose them via UI
 	static float BoneSamplingCentralDifferencesTime = 0.2f; // seconds
 	static float SmoothingWindowTime = 0.3f; // seconds
 
 	using namespace UE::PoseSearch;
-	using FSampleInfo = IAssetIndexer::FSampleInfo;
 
-	const FAssetIndexingContext& IndexingContext = Indexer.GetIndexingContext();
+	const FAssetIndexingContext& IndexingContext = BoneTransformsCache.Indexer.GetIndexingContext();
 
 	const float SampleTimeStart = FMath::Min(IndexingContext.BeginSampleIdx * IndexingContext.Schema->SamplingInterval, IndexingContext.MainSampler->GetPlayLength());
-	const FAssetSamplingContext* SamplingContext = IndexingContext.SamplingContext;
 	const float FiniteDelta = IndexingContext.Schema->SamplingInterval;
 
 	const int32 NumSamples = IndexingContext.EndSampleIdx - IndexingContext.BeginSampleIdx;
@@ -402,48 +488,15 @@ void UPoseSearchFeatureChannel_Pose::CalculatePhases(const UE::PoseSearch::IAsse
 		BonePositions[ChannelBoneIdx].AddDefaulted(NumSamples);
 	}
 
-	FDeltaTimeRecord DeltaTimeRecord;
-	FCompactPose Pose;
-	FCSPose<FCompactPose> ComponentSpacePose;
-	FBlendedCurve UnusedCurve;
-	UE::Anim::FStackAttributeContainer UnusedAtrribute;
-	Pose.SetBoneContainer(&SamplingContext->BoneContainer);
-
 	// collecting all the bone transforms
-	FSampleInfo Origin = Indexer.GetSampleInfo(SampleTimeStart);
+	IAssetIndexer::FSampleInfo Origin = BoneTransformsCache.Indexer.GetSampleInfo(SampleTimeStart);
 	for (int32 SampleIdx = 0; SampleIdx != NumSamples; ++SampleIdx)
 	{
 		const float SampleTime = SampleTimeStart + SampleIdx * FiniteDelta;
-		const float PreviousTime = SampleTime - FiniteDelta; // @todo: clamp to zero?
-
-		FSampleInfo Sample = Indexer.GetSampleInfoRelative(SampleTime, Origin);
-
-		DeltaTimeRecord.Set(PreviousTime, SampleTime - PreviousTime);
-		FAnimExtractContext ExtractionCtx(static_cast<double>(SampleTime), true, DeltaTimeRecord, Sample.Clip->IsLoopable());
-		
-		UnusedCurve.InitFrom(SamplingContext->BoneContainer);
-		FAnimationPoseData AnimPoseData = { Pose, UnusedCurve, UnusedAtrribute };
-		Sample.Clip->ExtractPose(ExtractionCtx, AnimPoseData);
-
-		if (IndexingContext.bMirrored)
-		{
-			FAnimationRuntime::MirrorPose(
-				AnimPoseData.GetPose(),
-				IndexingContext.Schema->MirrorDataTable->MirrorAxis,
-				SamplingContext->CompactPoseMirrorBones,
-				SamplingContext->ComponentSpaceRefRotations
-			);
-			// Note curves and attributes are not used during the indexing process and therefore don't need to be mirrored
-		}
-
-		ComponentSpacePose.InitPose(Pose);
 		for (int32 ChannelBoneIdx = 0; ChannelBoneIdx != SampledBones.Num(); ++ChannelBoneIdx)
 		{
-			const FPoseSearchPoseFeatureInfo& PoseFeatureInfo = FeatureParams[ChannelBoneIdx];
-			const FBoneReference& BoneReference = IndexingContext.Schema->BoneReferences[PoseFeatureInfo.SchemaBoneIdx];
-			const FCompactPoseBoneIndex CompactBoneIndex = SamplingContext->BoneContainer.MakeCompactPoseIndex(FMeshPoseBoneIndex(BoneReference.BoneIndex));
-			const FTransform BoneTransform = ComponentSpacePose.GetComponentSpaceTransform(CompactBoneIndex) * Indexer.MirrorTransform(Sample.RootTransform);
-			BonePositions[ChannelBoneIdx][SampleIdx] = BoneTransform.GetTranslation();
+			bool Unused;
+			BonePositions[ChannelBoneIdx][SampleIdx] = BoneTransformsCache.Get(SampleTime, SampleTimeStart, FeatureParams[ChannelBoneIdx].SchemaBoneIdx, Unused).GetTranslation();
 		}
 	}
 
@@ -476,22 +529,24 @@ void UPoseSearchFeatureChannel_Pose::IndexAsset(const UE::PoseSearch::IAssetInde
 {
 	using namespace UE::PoseSearch;
 	 
+	UE::PoseSearch::BoneTransformsCache BoneTransformsCache(Indexer);
+
 	// Phases is an array of array with cardinality SampledBones.Num() times NumSamples (IndexingContext.EndSampleIdx - IndexingContext.BeginSampleIdx)
 	// of 2 dimensional vectors (FVector2D) representing phases in an Eucledean space with phase angle sin/cos as direction and certainty of the signal as magnitude,
 	// where certainty is a function of the amplitude of the signal used as input
 	TArray<TArray<FVector2D>> Phases;
-	CalculatePhases(Indexer, IndexingOutput, Phases);
+	CalculatePhases(BoneTransformsCache, IndexingOutput, Phases);
 
 	const FAssetIndexingContext& IndexingContext = Indexer.GetIndexingContext();
 	for (int32 SampleIdx = IndexingContext.BeginSampleIdx; SampleIdx != IndexingContext.EndSampleIdx; ++SampleIdx)
 	{
 		int32 VectorIdx = SampleIdx - IndexingContext.BeginSampleIdx;
 		FPoseSearchFeatureVectorBuilder& FeatureVector = IndexingOutput.PoseVectors[VectorIdx];
-		AddPoseFeatures(Indexer, SampleIdx, FeatureVector, Phases);
+		AddPoseFeatures(BoneTransformsCache, SampleIdx, FeatureVector, Phases);
 	}
 }
 
-void UPoseSearchFeatureChannel_Pose::AddPoseFeatures(const  UE::PoseSearch::IAssetIndexer& Indexer, int32 SampleIdx, FPoseSearchFeatureVectorBuilder& FeatureVector, const TArray<TArray<FVector2D>>& Phases) const
+void UPoseSearchFeatureChannel_Pose::AddPoseFeatures(UE::PoseSearch::BoneTransformsCache& BoneTransformsCache, int32 SampleIdx, FPoseSearchFeatureVectorBuilder& FeatureVector, const TArray<TArray<FVector2D>>& Phases) const
 {
 	// This function samples the instantaneous pose at time t as well as the pose's velocity and acceleration at time t.
 	// Symmetric finite differences are used to approximate derivatives:
@@ -499,8 +554,7 @@ void UPoseSearchFeatureChannel_Pose::AddPoseFeatures(const  UE::PoseSearch::IAss
 	//	Second symmetric derivative: f''(t) ~ (f(t+h) - 2f(t) + f(t-h)) / h^2
 	// Where h is a constant time delta
 	// So this means three pose extractions are taken at time t-h, t, and t+h
-	constexpr int32 NumFiniteDiffTerms = 3;
-
+	
 	using namespace UE::PoseSearch;
 	using FSampleInfo = IAssetIndexer::FSampleInfo;
 
@@ -509,129 +563,97 @@ void UPoseSearchFeatureChannel_Pose::AddPoseFeatures(const  UE::PoseSearch::IAss
 		return;
 	}
 
-	const FAssetIndexingContext& IndexingContext = Indexer.GetIndexingContext();
+	const FAssetIndexingContext& IndexingContext = BoneTransformsCache.Indexer.GetIndexingContext();
 	const FAssetSamplingContext* SamplingContext = IndexingContext.SamplingContext;
 
-	FCompactPose Poses[NumFiniteDiffTerms];
-	FCSPose<FCompactPose> ComponentSpacePoses[NumFiniteDiffTerms];
-	FBlendedCurve UnusedCurves[NumFiniteDiffTerms];
-	UE::Anim::FStackAttributeContainer UnusedAtrributes[NumFiniteDiffTerms];
+	const float SampleTime = FMath::Min(SampleIdx * IndexingContext.Schema->SamplingInterval, IndexingContext.MainSampler->GetPlayLength());
 
-	for (FCompactPose& Pose : Poses)
-	{
-		Pose.SetBoneContainer(&SamplingContext->BoneContainer);
-	}
-
-	for (FBlendedCurve& Curve : UnusedCurves)
-	{
-		Curve.InitFrom(SamplingContext->BoneContainer);
-	}
-
-	FAnimationPoseData AnimPoseData[NumFiniteDiffTerms] =
+	int32 DataOffset = ChannelDataOffset;
+	for (int32 ChannelBoneIdx = 0; ChannelBoneIdx != SampledBones.Num(); ++ChannelBoneIdx)
 	{
-		{Poses[0], UnusedCurves[0], UnusedAtrributes[0]},
-		{Poses[1], UnusedCurves[1], UnusedAtrributes[1]},
-		{Poses[2], UnusedCurves[2], UnusedAtrributes[2]},
-	};
-
-	FPoseSearchFeatureDesc Feature;
-	Feature.ChannelIdx = GetChannelIndex();
-
-	float SampleTime = FMath::Min(SampleIdx * IndexingContext.Schema->SamplingInterval, IndexingContext.MainSampler->GetPlayLength());
-	FSampleInfo Origin = Indexer.GetSampleInfo(SampleTime);
+		const FPoseSearchPoseFeatureInfo& PoseFeatureInfo = FeatureParams[ChannelBoneIdx];
+		const FPoseSearchBone& SampledBone = SampledBones[ChannelBoneIdx];
 
-	for (int32 SubsampleIdx = 0; SubsampleIdx != SampleTimes.Num(); ++SubsampleIdx)
-	{
-		Feature.SubsampleIdx = SubsampleIdx;
+		// Get each bone's component transform, velocity, and acceleration and add accumulated root motion at this time offset
+		// Think of this process as freezing the character in place (at SampleTime) and then tracing the paths of their joints
+		// as they move through space from past to present to future (at times indicated by PoseSampleTimes).
 
-		float SubsampleTime = SampleTime + SampleTimes[SubsampleIdx];
+		if (SampledBone.bUsePosition)
+		{
+			for (int32 SubsampleIdx = 0; SubsampleIdx != SampleTimes.Num(); ++SubsampleIdx)
+			{
+				const float SubsampleTime = SampleTime + SampleTimes[SubsampleIdx];
 
-		// For each pose subsample term, get the corresponding clip, accumulated root motion,
-		// and wrap the time parameter based on the clip's length.
-		FSampleInfo Samples[NumFiniteDiffTerms];
+				bool ClampedPresent;
+				const FTransform BoneTransformsPresent = BoneTransformsCache.Get(SubsampleTime, UseCharacterSpaceVelocities ? SubsampleTime : SampleTime, PoseFeatureInfo.SchemaBoneIdx, ClampedPresent);
 
-		if (UseCharacterSpaceVelocities)
-		{
-			// character space velocity
-			Samples[0] = Indexer.GetSampleInfoRelative(SubsampleTime - SamplingContext->FiniteDelta, Indexer.GetSampleInfo(SubsampleTime - SamplingContext->FiniteDelta));
-			Samples[1] = Indexer.GetSampleInfoRelative(SubsampleTime, Indexer.GetSampleInfo(SubsampleTime));
-			Samples[2] = Indexer.GetSampleInfoRelative(SubsampleTime + SamplingContext->FiniteDelta, Indexer.GetSampleInfo(SubsampleTime + SamplingContext->FiniteDelta));
-		}
-		else
-		{
-			// animation space velocity
-			Samples[0] = Indexer.GetSampleInfoRelative(SubsampleTime - SamplingContext->FiniteDelta, Origin);
-			Samples[1] = Indexer.GetSampleInfoRelative(SubsampleTime, Origin);
-			Samples[2] = Indexer.GetSampleInfoRelative(SubsampleTime + SamplingContext->FiniteDelta, Origin);
+				const FPoseSearchFeatureDesc Feature = FPoseSearchFeatureDesc::Construct(GetChannelIndex(), ChannelBoneIdx, SubsampleIdx, EPoseSearchFeatureType::Position, PositionCardinality, DataOffset);
+				DataOffset += PositionCardinality;
+				FeatureVector.SetVector(Feature, BoneTransformsPresent.GetTranslation());
+			}
 		}
-		
-		// Get pose samples
-		for (int32 Term = 0; Term != NumFiniteDiffTerms; ++Term)
-		{
-			float CurrentTime = Samples[Term].ClipTime;
-			float PreviousTime = CurrentTime - SamplingContext->FiniteDelta;
 
-			FDeltaTimeRecord DeltaTimeRecord;
-			DeltaTimeRecord.Set(PreviousTime, CurrentTime - PreviousTime);
-			FAnimExtractContext ExtractionCtx(CurrentTime, true, DeltaTimeRecord, Samples[Term].Clip->IsLoopable());
+		if (SampledBone.bUseRotation)
+		{
+			for (int32 SubsampleIdx = 0; SubsampleIdx != SampleTimes.Num(); ++SubsampleIdx)
+			{
+				const float SubsampleTime = SampleTime + SampleTimes[SubsampleIdx];
 
-			Samples[Term].Clip->ExtractPose(ExtractionCtx, AnimPoseData[Term]);
+				bool ClampedPresent;
+				const FTransform BoneTransformsPresent = BoneTransformsCache.Get(SubsampleTime, UseCharacterSpaceVelocities ? SubsampleTime : SampleTime, PoseFeatureInfo.SchemaBoneIdx, ClampedPresent);
+				const FPoseSearchFeatureDesc Feature = FPoseSearchFeatureDesc::Construct(GetChannelIndex(), ChannelBoneIdx, SubsampleIdx, EPoseSearchFeatureType::Rotation, RotationCardinality, DataOffset);
+				DataOffset += RotationCardinality;
 
-			if (IndexingContext.bMirrored)
-			{
-				FAnimationRuntime::MirrorPose(
-					AnimPoseData[Term].GetPose(),
-					IndexingContext.Schema->MirrorDataTable->MirrorAxis,
-					SamplingContext->CompactPoseMirrorBones,
-					SamplingContext->ComponentSpaceRefRotations
-				);
-				// Note curves and attributes are not used during the indexing process and therefore don't need to be mirrored
+				FeatureVector.SetRotation(Feature, BoneTransformsPresent.GetRotation());
 			}
-
-			ComponentSpacePoses[Term].InitPose(Poses[Term]);
 		}
 
-		// Get each bone's component transform, velocity, and acceleration and add accumulated root motion at this time offset
-		// Think of this process as freezing the character in place (at SampleTime) and then tracing the paths of their joints
-		// as they move through space from past to present to future (at times indicated by PoseSampleTimes).
-		for (int32 ChannelBoneIdx = 0; ChannelBoneIdx != SampledBones.Num(); ++ChannelBoneIdx)
+		if (SampledBone.bUseVelocity)
 		{
-			const FPoseSearchPoseFeatureInfo& PoseFeatureInfo = FeatureParams[ChannelBoneIdx];
-			const FBoneReference& BoneReference = IndexingContext.Schema->BoneReferences[PoseFeatureInfo.SchemaBoneIdx];
-
-			Feature.ChannelFeatureId = ChannelBoneIdx;
+			for (int32 SubsampleIdx = 0; SubsampleIdx != SampleTimes.Num(); ++SubsampleIdx)
+			{
+				const float SubsampleTime = SampleTime + SampleTimes[SubsampleIdx];
 
-			FCompactPoseBoneIndex CompactBoneIndex = SamplingContext->BoneContainer.MakeCompactPoseIndex(FMeshPoseBoneIndex(BoneReference.BoneIndex));
+				bool ClampedPast, ClampedPresent, ClampedFuture;
+				const FTransform BoneTransformsPast = BoneTransformsCache.Get(SubsampleTime - SamplingContext->FiniteDelta, UseCharacterSpaceVelocities ? SubsampleTime - SamplingContext->FiniteDelta : SampleTime, PoseFeatureInfo.SchemaBoneIdx, ClampedPast);
+				const FTransform BoneTransformsPresent = BoneTransformsCache.Get(SubsampleTime, UseCharacterSpaceVelocities ? SubsampleTime : SampleTime, PoseFeatureInfo.SchemaBoneIdx, ClampedPresent);
+				const FTransform BoneTransformsFuture = BoneTransformsCache.Get(SubsampleTime + SamplingContext->FiniteDelta, UseCharacterSpaceVelocities ? SubsampleTime + SamplingContext->FiniteDelta : SampleTime, PoseFeatureInfo.SchemaBoneIdx, ClampedFuture);
 
-			FTransform BoneTransforms[NumFiniteDiffTerms];
-			for (int32 Term = 0; Term != NumFiniteDiffTerms; ++Term)
-			{
-				BoneTransforms[Term] = ComponentSpacePoses[Term].GetComponentSpaceTransform(CompactBoneIndex);
-				BoneTransforms[Term] = BoneTransforms[Term] * Indexer.MirrorTransform(Samples[Term].RootTransform);
-			}
+				const FPoseSearchFeatureDesc Feature = FPoseSearchFeatureDesc::Construct(GetChannelIndex(), ChannelBoneIdx, SubsampleIdx, EPoseSearchFeatureType::LinearVelocity, LinearVelocityCardinality, DataOffset);
+				DataOffset += LinearVelocityCardinality;
 
-			// Add properties to the feature vector for the pose at SampleIdx
-			FeatureVector.SetTransform(Feature, BoneTransforms[1]);
+				// We can get a better finite difference if we ignore samples that have
+				// been clamped at either side of the clip. However, if the central sample 
+				// itself is clamped, or there are no samples that are clamped, we can just 
+				// use the central difference as normal.
+				FVector LinearVelocity;
+				if (ClampedPast && !ClampedPresent && !ClampedFuture)
+				{
+					LinearVelocity = (BoneTransformsFuture.GetTranslation() - BoneTransformsPresent.GetTranslation()) / SamplingContext->FiniteDelta;
+				}
+				else if (ClampedFuture && !ClampedPresent && !ClampedPast)
+				{
+					LinearVelocity = (BoneTransformsPresent.GetTranslation() - BoneTransformsPast.GetTranslation()) / SamplingContext->FiniteDelta;
+				}
+				else
+				{
+					LinearVelocity = (BoneTransformsFuture.GetTranslation() - BoneTransformsPast.GetTranslation()) / (SamplingContext->FiniteDelta * 2.0f);
+				}
 
-			// We can get a better finite difference if we ignore samples that have
-			// been clamped at either side of the clip. However, if the central sample 
-			// itself is clamped, or there are no samples that are clamped, we can just 
-			// use the central difference as normal.
-			if (Samples[0].bClamped && !Samples[1].bClamped && !Samples[2].bClamped)
-			{
-				FeatureVector.SetTransformVelocity(Feature, BoneTransforms[2], BoneTransforms[1], SamplingContext->FiniteDelta);
-			}
-			else if (Samples[2].bClamped && !Samples[1].bClamped && !Samples[0].bClamped)
-			{
-				FeatureVector.SetTransformVelocity(Feature, BoneTransforms[1], BoneTransforms[0], SamplingContext->FiniteDelta);
+				FeatureVector.SetVector(Feature, LinearVelocity);
 			}
-			else
+		}
+
+		if (SampledBone.bUsePhase)
+		{
+			for (int32 SubsampleIdx = 0; SubsampleIdx != SampleTimes.Num(); ++SubsampleIdx)
 			{
-				FeatureVector.SetTransformVelocity(Feature, BoneTransforms[2], BoneTransforms[1], BoneTransforms[0], SamplingContext->FiniteDelta);
-			}
+				const FPoseSearchFeatureDesc Feature = FPoseSearchFeatureDesc::Construct(GetChannelIndex(), ChannelBoneIdx, SubsampleIdx, EPoseSearchFeatureType::Phase, PhaseCardinality, DataOffset);
+				DataOffset += PhaseCardinality;
 
-			// @todo: support for SubsampleIdx
-			FeatureVector.SetPhase(Feature, Phases[ChannelBoneIdx][SampleIdx]);
+				// @todo: support for SubsampleIdx
+				FeatureVector.SetPhase(Feature, Phases[ChannelBoneIdx][SampleIdx]);
+			}
 		}
 	}
 }
@@ -677,13 +699,17 @@ bool UPoseSearchFeatureChannel_Pose::BuildQuery(
 		return false;
 	}
 
-	FPoseSearchFeatureDesc Feature;
-	Feature.ChannelIdx = GetChannelIndex();
+	struct CachedTransforms
+	{
+		FTransform Current;
+		FTransform Previous;
+		bool Valid = false;
+	};
+	TArray<CachedTransforms> CachedTransforms;
+	CachedTransforms.AddUninitialized(SampleTimes.Num() * SampledBones.Num());
 
 	for (int32 SubsampleIdx = 0; SubsampleIdx != SampleTimes.Num(); ++SubsampleIdx)
 	{
-		Feature.SubsampleIdx = SubsampleIdx;
-
 		// Stop when we've reached future samples
 		float SampleTime = SampleTimes[SubsampleIdx];
 		if (SampleTime > 0.0f)
@@ -697,34 +723,90 @@ bool UPoseSearchFeatureChannel_Pose::BuildQuery(
 			return false;
 		}
 
-		TArrayView<const FTransform> ComponentPose = SearchContext.History->GetComponentPoseSample();
-		TArrayView<const FTransform> ComponentPrevPose = SearchContext.History->GetPrevComponentPoseSample();
-		FTransform RootTransform = SearchContext.History->GetRootTransformSample();
-		FTransform RootTransformPrev = SearchContext.History->GetPrevRootTransformSample();
+		const TArrayView<const FTransform> ComponentPose = SearchContext.History->GetComponentPoseSample();
+		const TArrayView<const FTransform> ComponentPrevPose = SearchContext.History->GetPrevComponentPoseSample();
+		const FTransform RootTransform = SearchContext.History->GetRootTransformSample();
+		const FTransform RootTransformPrev = SearchContext.History->GetPrevRootTransformSample();
 		for (int32 SampledBoneIdx = 0; SampledBoneIdx != SampledBones.Num(); ++SampledBoneIdx)
 		{
-			Feature.ChannelFeatureId = SampledBoneIdx;
-
-			int32 SchemaBoneIdx = FeatureParams[SampledBoneIdx].SchemaBoneIdx;
-
-			int32 SkeletonBoneIndex = InOutQuery.GetSchema()->BoneIndices[SchemaBoneIdx];
+			const int32 SchemaBoneIdx = FeatureParams[SampledBoneIdx].SchemaBoneIdx;
+			const int32 SkeletonBoneIndex = InOutQuery.GetSchema()->BoneIndices[SchemaBoneIdx];
+			const int32 CachedTransformsIndex = SubsampleIdx * SampledBones.Num() + SampledBoneIdx;
+			CachedTransforms[CachedTransformsIndex].Current = ComponentPose[SkeletonBoneIndex];
 
-			const FTransform& Transform = ComponentPose[SkeletonBoneIndex];
-
-			FTransform PrevTransform;
 			if (UE::PoseSearch::UseCharacterSpaceVelocities)
 			{
 				// character space velocity
-				PrevTransform = ComponentPrevPose[SkeletonBoneIndex];
+				CachedTransforms[CachedTransformsIndex].Previous = ComponentPrevPose[SkeletonBoneIndex];
 			}
 			else
 			{
 				// animation space velocity
-				PrevTransform = ComponentPrevPose[SkeletonBoneIndex] * (RootTransformPrev * RootTransform.Inverse());
+				CachedTransforms[CachedTransformsIndex].Previous = ComponentPrevPose[SkeletonBoneIndex] * (RootTransformPrev * RootTransform.Inverse());
+			}
+		}
+	}
+
+	int32 DataOffset = ChannelDataOffset;
+	for (int32 SampledBoneIdx = 0; SampledBoneIdx != SampledBones.Num(); ++SampledBoneIdx)
+	{
+		const FPoseSearchBone& SampledBone = SampledBones[SampledBoneIdx];
+		if (SampledBone.bUsePosition)
+		{
+			for (int32 SubsampleIdx = 0; SubsampleIdx != SampleTimes.Num(); ++SubsampleIdx)
+			{
+				const int32 CachedTransformsIndex = SubsampleIdx * SampledBones.Num() + SampledBoneIdx;
+				const FPoseSearchFeatureDesc Feature = FPoseSearchFeatureDesc::Construct(GetChannelIndex(), SampledBoneIdx, SubsampleIdx, EPoseSearchFeatureType::Position, PositionCardinality, DataOffset);
+				DataOffset += PositionCardinality;
+				if (CachedTransforms[CachedTransformsIndex].Valid)
+				{
+					InOutQuery.SetVector(Feature, CachedTransforms[CachedTransformsIndex].Current.GetTranslation());
+				}
+			}
+		}
+
+		if (SampledBone.bUseRotation)
+		{
+			for (int32 SubsampleIdx = 0; SubsampleIdx != SampleTimes.Num(); ++SubsampleIdx)
+			{
+				const int32 CachedTransformsIndex = SubsampleIdx * SampledBones.Num() + SampledBoneIdx;
+				const FPoseSearchFeatureDesc Feature = FPoseSearchFeatureDesc::Construct(GetChannelIndex(), SampledBoneIdx, SubsampleIdx, EPoseSearchFeatureType::Rotation, RotationCardinality, DataOffset);
+				DataOffset += RotationCardinality;
+
+				if (CachedTransforms[CachedTransformsIndex].Valid)
+				{
+					InOutQuery.SetRotation(Feature, CachedTransforms[CachedTransformsIndex].Current.GetRotation());
+				}
+			}
+		}
+
+		if (SampledBone.bUseVelocity)
+		{
+			for (int32 SubsampleIdx = 0; SubsampleIdx != SampleTimes.Num(); ++SubsampleIdx)
+			{
+				const int32 CachedTransformsIndex = SubsampleIdx * SampledBones.Num() + SampledBoneIdx;
+				const FPoseSearchFeatureDesc Feature = FPoseSearchFeatureDesc::Construct(GetChannelIndex(), SampledBoneIdx, SubsampleIdx, EPoseSearchFeatureType::LinearVelocity, LinearVelocityCardinality, DataOffset);
+				DataOffset += LinearVelocityCardinality;
+
+				if (CachedTransforms[CachedTransformsIndex].Valid)
+				{
+					const FVector LinearVelocity = (CachedTransforms[CachedTransformsIndex].Current.GetTranslation() - CachedTransforms[CachedTransformsIndex].Previous.GetTranslation()) / SearchContext.History->GetSampleTimeInterval();
+					InOutQuery.SetVector(Feature, LinearVelocity);
+				}
+			}
+		}
+
+		if (SampledBone.bUsePhase)
+		{
+			for (int32 SubsampleIdx = 0; SubsampleIdx != SampleTimes.Num(); ++SubsampleIdx)
+			{
+				const int32 CachedTransformsIndex = SubsampleIdx * SampledBones.Num() + SampledBoneIdx;
+				const FPoseSearchFeatureDesc Feature = FPoseSearchFeatureDesc::Construct(GetChannelIndex(), SampledBoneIdx, SubsampleIdx, EPoseSearchFeatureType::Phase, PhaseCardinality, DataOffset);
+				DataOffset += PhaseCardinality;
+
+				// @todo: Support phase in BuildQuery
+				//InOutQuery.SetPhase(Feature, ???);
 			}
-			
-			InOutQuery.SetTransform(Feature, Transform);
-			InOutQuery.SetTransformVelocity(Feature, Transform, PrevTransform, SearchContext.History->GetSampleTimeInterval());
 		}
 	}
 
@@ -742,9 +824,6 @@ void UPoseSearchFeatureChannel_Pose::DebugDraw(const UE::PoseSearch::FDebugDrawP
 	const uint8 DepthPriority = ESceneDepthPriorityGroup::SDPG_Foreground + 2;
 	const bool bPersistent = EnumHasAnyFlags(DrawParams.Flags, EDebugDrawFlags::Persistent);
 
-	FPoseSearchFeatureDesc Feature;
-	Feature.ChannelIdx = GetChannelIndex();
-
 	const int32 NumSubsamples = SampleTimes.Num();
 	const int32 NumBones = SampledBones.Num();
 
@@ -753,52 +832,81 @@ void UPoseSearchFeatureChannel_Pose::DebugDraw(const UE::PoseSearch::FDebugDrawP
 		return;
 	}
 
-	for (int32 SubsampleIdx = 0; SubsampleIdx != NumSubsamples; ++SubsampleIdx)
+	int32 DataOffset = ChannelDataOffset;
+	for (int32 ChannelBoneIdx = 0; ChannelBoneIdx != NumBones; ++ChannelBoneIdx)
 	{
-		Feature.SubsampleIdx = SubsampleIdx;
+		const FPoseSearchBone& SampledBone = SampledBones[ChannelBoneIdx];
 
-		for (int32 ChannelBoneIdx = 0; ChannelBoneIdx != NumBones; ++ChannelBoneIdx)
+		TArray<FVector> BonePos;
+		BonePos.AddUninitialized(NumSubsamples);
+		if (SampledBone.bUsePosition)
 		{
-			Feature.ChannelFeatureId = ChannelBoneIdx;
-
-			FVector BonePos;
-			bool bHaveBonePos = Reader.GetPosition(Feature, &BonePos);
-			if (bHaveBonePos)
+			for (int32 SubsampleIdx = 0; SubsampleIdx != NumSubsamples; ++SubsampleIdx)
 			{
-				Feature.Type = EPoseSearchFeatureType::Position;
+				const FPoseSearchFeatureDesc Feature = FPoseSearchFeatureDesc::Construct(GetChannelIndex(), ChannelBoneIdx, SubsampleIdx, EPoseSearchFeatureType::Position, PositionCardinality, DataOffset);
+				DataOffset += PositionCardinality;
+
+				const bool Found = Reader.GetVector(Feature, &BonePos[SubsampleIdx]);
+				check(Found);
 
 				FLinearColor LinearColor = DrawParams.Color ? *DrawParams.Color : GetColorForFeature(Feature, Reader.GetLayout());
 				FColor Color = LinearColor.ToFColor(true);
 
-				BonePos = DrawParams.RootTransform.TransformPosition(BonePos);
+				BonePos[SubsampleIdx] = DrawParams.RootTransform.TransformPosition(BonePos[SubsampleIdx]);
 				if (EnumHasAnyFlags(DrawParams.Flags, EDebugDrawFlags::DrawFast | EDebugDrawFlags::DrawSearchIndex))
 				{
-					DrawDebugPoint(DrawParams.World, BonePos, DrawParams.PointSize, Color, bPersistent, DrawParams.DefaultLifeTime, DepthPriority);
+					DrawDebugPoint(DrawParams.World, BonePos[SubsampleIdx], DrawParams.PointSize, Color, bPersistent, DrawParams.DefaultLifeTime, DepthPriority);
 				}
 				else
 				{
-					DrawDebugSphere(DrawParams.World, BonePos, DrawDebugSphereSize, DrawDebugSphereSegments, Color, bPersistent, LifeTime, DepthPriority);
+					DrawDebugSphere(DrawParams.World, BonePos[SubsampleIdx], DrawDebugSphereSize, DrawDebugSphereSegments, Color, bPersistent, LifeTime, DepthPriority);
 				}
 
 				if (EnumHasAnyFlags(DrawParams.Flags, EDebugDrawFlags::DrawBoneNames))
 				{
 					int32 SchemaBoneIdx = this->FeatureParams[ChannelBoneIdx].SchemaBoneIdx;
 					DrawDebugString(
-						DrawParams.World, BonePos + FVector(0.0, 0.0, 10.0),
+						DrawParams.World, BonePos[SubsampleIdx] + FVector(0.0, 0.0, 10.0),
 						Schema->BoneReferences[SchemaBoneIdx].BoneName.ToString(),
 						nullptr, Color, LifeTime, false, 1.0f);
 				}
 			}
-			else if (DrawParams.Mesh != nullptr)
+		}
+		else
+		{
+			for (int32 SubsampleIdx = 0; SubsampleIdx != NumSubsamples; ++SubsampleIdx)
 			{
-				BonePos = DrawParams.Mesh->GetSocketTransform(SampledBones[ChannelBoneIdx].Reference.BoneName).GetLocation();
-				bHaveBonePos = true;
+				DataOffset += PositionCardinality;
+
+				// @todo: initialize with the character position instead of FVector::Zero?
+				BonePos[SubsampleIdx] = DrawParams.Mesh != nullptr ? DrawParams.Mesh->GetSocketTransform(SampledBones[ChannelBoneIdx].Reference.BoneName).GetLocation() : FVector::Zero();
 			}
+		}
 
-			FVector BoneVel;
-			if (bHaveBonePos && Reader.GetLinearVelocity(Feature, &BoneVel))
+		if (SampledBone.bUseRotation)
+		{
+			for (int32 SubsampleIdx = 0; SubsampleIdx != NumSubsamples; ++SubsampleIdx)
 			{
-				Feature.Type = EPoseSearchFeatureType::LinearVelocity;
+				const FPoseSearchFeatureDesc Feature = FPoseSearchFeatureDesc::Construct(GetChannelIndex(), ChannelBoneIdx, SubsampleIdx, EPoseSearchFeatureType::Rotation, RotationCardinality, DataOffset);
+				DataOffset += RotationCardinality;
+
+				FQuat BoneRot;
+				const bool Found = Reader.GetRotation(Feature, &BoneRot);
+				check(Found);
+				// @todo: debug draw rotation
+			}
+		}
+
+		if (SampledBone.bUseVelocity)
+		{
+			for (int32 SubsampleIdx = 0; SubsampleIdx != NumSubsamples; ++SubsampleIdx)
+			{
+				const FPoseSearchFeatureDesc Feature = FPoseSearchFeatureDesc::Construct(GetChannelIndex(), ChannelBoneIdx, SubsampleIdx, EPoseSearchFeatureType::LinearVelocity, LinearVelocityCardinality, DataOffset);
+				DataOffset += LinearVelocityCardinality;
+
+				FVector BoneVel;
+				const bool Found = Reader.GetVector(Feature, &BoneVel);
+				check(Found);
 
 				FLinearColor LinearColor = DrawParams.Color ? *DrawParams.Color : GetColorForFeature(Feature, Reader.GetLayout());
 				FColor Color = LinearColor.ToFColor(true);
@@ -819,8 +927,8 @@ void UPoseSearchFeatureChannel_Pose::DebugDraw(const UE::PoseSearch::FDebugDrawP
 
 					DrawDebugDirectionalArrow(
 						DrawParams.World,
-						BonePos + BoneVelDirection * DrawDebugSphereSize,
-						BonePos + BoneVel,
+						BonePos[SubsampleIdx] + BoneVelDirection * DrawDebugSphereSize,
+						BonePos[SubsampleIdx] + BoneVel,
 						DrawDebugArrowSize,
 						Color,
 						bPersistent,
@@ -829,11 +937,18 @@ void UPoseSearchFeatureChannel_Pose::DebugDraw(const UE::PoseSearch::FDebugDrawP
 						AdjustedThickness);
 				}
 			}
+		}
 
-			FVector2D Phase;
-			if (bHaveBonePos && Reader.GetPhase(Feature, &Phase))
+		if (SampledBone.bUsePhase)
+		{
+			for (int32 SubsampleIdx = 0; SubsampleIdx != NumSubsamples; ++SubsampleIdx)
 			{
-				Feature.Type = EPoseSearchFeatureType::Phase;
+				const FPoseSearchFeatureDesc Feature = FPoseSearchFeatureDesc::Construct(GetChannelIndex(), ChannelBoneIdx, SubsampleIdx, EPoseSearchFeatureType::Phase, PhaseCardinality, DataOffset);
+				DataOffset += PhaseCardinality;
+
+				FVector2D Phase;
+				const bool Found = Reader.GetPhase(Feature, &Phase);
+				check(Found);
 
 				FLinearColor LinearColor = DrawParams.Color ? *DrawParams.Color : GetColorForFeature(Feature, Reader.GetLayout());
 				FColor Color = LinearColor.ToFColor(true);
@@ -845,11 +960,11 @@ void UPoseSearchFeatureChannel_Pose::DebugDraw(const UE::PoseSearch::FDebugDrawP
 				const FVector TransformZAxisVector = DrawParams.RootTransform.TransformVector(FVector::ZAxisVector);
 
 				const FVector PhaseVector = (TransformZAxisVector * Phase.X + TransformYAxisVector * Phase.Y) * ScaleFactor;
-				DrawDebugLine(DrawParams.World, BonePos, BonePos + PhaseVector, Color, bPersistent, LifeTime, DepthPriority, 0.f);
+				DrawDebugLine(DrawParams.World, BonePos[SubsampleIdx], BonePos[SubsampleIdx] + PhaseVector, Color, bPersistent, LifeTime, DepthPriority, 0.f);
 
 				static int32 Segments = 64;
 				FMatrix CircleTransform;
-				CircleTransform.SetAxes(&TransformXAxisVector, &TransformYAxisVector, &TransformZAxisVector, &BonePos);
+				CircleTransform.SetAxes(&TransformXAxisVector, &TransformYAxisVector, &TransformZAxisVector, &BonePos[SubsampleIdx]);
 				DrawDebugCircle(DrawParams.World, CircleTransform, PhaseVector.Length(), Segments, Color, bPersistent, LifeTime, DepthPriority, 0.f, false);
 			}
 		}
@@ -869,11 +984,13 @@ void UPoseSearchFeatureChannel_Trajectory::InitializeSchema( UE::PoseSearch::FSc
 {
 	Super::InitializeSchema(Initializer);
 
+	int32 DataOffset = ChannelDataOffset;
 	if (bUsePositions)
 	{
 		for (int32 SubsampleIdx = 0; SubsampleIdx != SampleOffsets.Num(); ++SubsampleIdx)
 		{
-			Initializer.AddFeatureDesc(FPoseSearchFeatureDesc::Construct(GetChannelIndex(), 0, SubsampleIdx, EPoseSearchFeatureType::Position, PositionCardinality));
+			Initializer.AddFeatureDesc(FPoseSearchFeatureDesc::Construct(GetChannelIndex(), 0, SubsampleIdx, EPoseSearchFeatureType::Position, PositionCardinality, DataOffset));
+			DataOffset += PositionCardinality;
 		}
 	}
 
@@ -881,7 +998,8 @@ void UPoseSearchFeatureChannel_Trajectory::InitializeSchema( UE::PoseSearch::FSc
 	{
 		for (int32 SubsampleIdx = 0; SubsampleIdx != SampleOffsets.Num(); ++SubsampleIdx)
 		{
-			Initializer.AddFeatureDesc(FPoseSearchFeatureDesc::Construct(GetChannelIndex(), 0, SubsampleIdx, EPoseSearchFeatureType::LinearVelocity, LinearVelocityCardinality));
+			Initializer.AddFeatureDesc(FPoseSearchFeatureDesc::Construct(GetChannelIndex(), 0, SubsampleIdx, EPoseSearchFeatureType::LinearVelocity, LinearVelocityCardinality, DataOffset));
+			DataOffset += LinearVelocityCardinality;
 		}
 	}
 
@@ -889,7 +1007,8 @@ void UPoseSearchFeatureChannel_Trajectory::InitializeSchema( UE::PoseSearch::FSc
 	{
 		for (int32 SubsampleIdx = 0; SubsampleIdx != SampleOffsets.Num(); ++SubsampleIdx)
 		{
-			Initializer.AddFeatureDesc(FPoseSearchFeatureDesc::Construct(GetChannelIndex(), 0, SubsampleIdx, EPoseSearchFeatureType::ForwardVector, ForwardVectorCardinality));
+			Initializer.AddFeatureDesc(FPoseSearchFeatureDesc::Construct(GetChannelIndex(), 0, SubsampleIdx, EPoseSearchFeatureType::ForwardVector, ForwardVectorCardinality, DataOffset));
+			DataOffset += ForwardVectorCardinality;
 		}
 	}
 
@@ -908,35 +1027,32 @@ void UPoseSearchFeatureChannel_Trajectory::FillWeights(TArray<float>& Weights) c
 
 void UPoseSearchFeatureChannel_Trajectory::IndexAsset(const UE::PoseSearch::IAssetIndexer& Indexer,  UE::PoseSearch::FAssetIndexingOutput& IndexingOutput) const
 {
-	using namespace UE::PoseSearch;
-
-	auto IndexFeatures = [&Indexer, &IndexingOutput](auto FeatureIndexingFunc)
+	const UE::PoseSearch::FAssetIndexingContext& IndexingContext = Indexer.GetIndexingContext();
+	for (int32 SampleIdx = IndexingContext.BeginSampleIdx; SampleIdx != IndexingContext.EndSampleIdx; ++SampleIdx)
 	{
-		const FAssetIndexingContext& IndexingContext = Indexer.GetIndexingContext();
-		for (int32 SampleIdx = IndexingContext.BeginSampleIdx; SampleIdx != IndexingContext.EndSampleIdx; ++SampleIdx)
-		{
-			int32 VectorIdx = SampleIdx - IndexingContext.BeginSampleIdx;
-			FPoseSearchFeatureVectorBuilder& FeatureVector = IndexingOutput.PoseVectors[VectorIdx];
-			FeatureIndexingFunc(Indexer, SampleIdx, FeatureVector);
-		}
-	};
+		const int32 VectorIdx = SampleIdx - IndexingContext.BeginSampleIdx;
+		IndexAssetPrivate(Indexer, SampleIdx, IndexingOutput.PoseVectors[VectorIdx]);
+	}
+}
 
+float UPoseSearchFeatureChannel_Trajectory::GetSampleTime(const UE::PoseSearch::IAssetIndexer& Indexer, int32 SubsampleIdx, float SampleTime, float RootDistance) const
+{
 	switch (Domain)
 	{
-		case EPoseSearchFeatureDomain::Time:
-			IndexFeatures([this](auto&&... Args){ return IndexTimeFeatures(Args...); });
-			break;
+	case EPoseSearchFeatureDomain::Time:
+		return SampleTime + SampleOffsets[SubsampleIdx];
 
-		case EPoseSearchFeatureDomain::Distance:
-			IndexFeatures([this](auto&&... Args){ return IndexDistanceFeatures(Args...); });
-			break;
+	case EPoseSearchFeatureDomain::Distance:
+		return Indexer.GetSampleTimeFromDistance(RootDistance + SampleOffsets[SubsampleIdx]);
 
-		default:
-			checkNoEntry();
+	default:
+		checkNoEntry();
 	}
+	
+	return 0.0f;
 }
 
-void UPoseSearchFeatureChannel_Trajectory::IndexTimeFeatures(const UE::PoseSearch::IAssetIndexer& Indexer, int32 SampleIdx,  FPoseSearchFeatureVectorBuilder& FeatureVector) const
+void UPoseSearchFeatureChannel_Trajectory::IndexAssetPrivate(const UE::PoseSearch::IAssetIndexer& Indexer, int32 SampleIdx,  FPoseSearchFeatureVectorBuilder& FeatureVector) const
 {
 	// This function samples the instantaneous trajectory at time t as well as the trajectory's velocity and acceleration at time t.
 	// Symmetric finite differences are used to approximate derivatives:
@@ -949,116 +1065,72 @@ void UPoseSearchFeatureChannel_Trajectory::IndexTimeFeatures(const UE::PoseSearc
 	using FSampleInfo = IAssetIndexer::FSampleInfo;
 
 	const FAssetIndexingContext& IndexingContext = Indexer.GetIndexingContext();
-
-
-	FPoseSearchFeatureDesc Feature;
-	Feature.ChannelIdx = GetChannelIndex();
-	Feature.ChannelFeatureId = 0; // Unused
-
 	float SampleTime = FMath::Min(SampleIdx * IndexingContext.Schema->SamplingInterval, IndexingContext.MainSampler->GetPlayLength());
 	FSampleInfo Origin = Indexer.GetSampleInfo(SampleTime);
 
-	for (int32 SubsampleIdx = 0; SubsampleIdx != SampleOffsets.Num(); ++SubsampleIdx)
+	int32 DataOffset = ChannelDataOffset;
+	if (bUsePositions)
 	{
-		Feature.SubsampleIdx = SubsampleIdx;
-
-		float SubsampleTime = SampleTime + SampleOffsets[SubsampleIdx];
-
-		// For each pose subsample term, get the corresponding clip, accumulated root motion,
-		// and wrap the time parameter based on the clip's length.
-		FSampleInfo Samples[3];
-		Samples[0] = Indexer.GetSampleInfoRelative(SubsampleTime - IndexingContext.SamplingContext->FiniteDelta, Origin);
-		Samples[1] = Indexer.GetSampleInfoRelative(SubsampleTime, Origin);
-		Samples[2] = Indexer.GetSampleInfoRelative(SubsampleTime + IndexingContext.SamplingContext->FiniteDelta, Origin);
-
-		// Mirror transforms if requested
-		FTransform MirroredRoots[3];
-		MirroredRoots[0] = Indexer.MirrorTransform(Samples[0].RootTransform);
-		MirroredRoots[1] = Indexer.MirrorTransform(Samples[1].RootTransform);
-		MirroredRoots[2] = Indexer.MirrorTransform(Samples[2].RootTransform);
-
-		// Add properties to the feature vector for the pose at SampleIdx
-		FeatureVector.SetTransform(Feature, MirroredRoots[1]);
-
-		// We can get a better finite difference if we ignore samples that have
-		// been clamped at either side of the clip. However, if the central sample 
-		// itself is clamped, or there are no samples that are clamped, we can just 
-		// use the central difference as normal.
-		if (Samples[0].bClamped && !Samples[1].bClamped && !Samples[2].bClamped)
-		{
-			FeatureVector.SetTransformVelocity(Feature, MirroredRoots[2], MirroredRoots[1], IndexingContext.SamplingContext->FiniteDelta);
-		}
-		else if (Samples[2].bClamped && !Samples[1].bClamped && !Samples[0].bClamped)
-		{
-			FeatureVector.SetTransformVelocity(Feature, MirroredRoots[1], MirroredRoots[0], IndexingContext.SamplingContext->FiniteDelta);
-		}
-		else
+		for (int32 SubsampleIdx = 0; SubsampleIdx != SampleOffsets.Num(); ++SubsampleIdx)
 		{
-			FeatureVector.SetTransformVelocity(Feature, MirroredRoots[2], MirroredRoots[1], MirroredRoots[0], IndexingContext.SamplingContext->FiniteDelta);
+			const float SubsampleTime = GetSampleTime(Indexer, SubsampleIdx, SampleTime, Origin.RootDistance);
+			const FSampleInfo SamplePresent = Indexer.GetSampleInfoRelative(SubsampleTime, Origin);
+			const FPoseSearchFeatureDesc Feature = FPoseSearchFeatureDesc::Construct(GetChannelIndex(), 0, SubsampleIdx, EPoseSearchFeatureType::Position, PositionCardinality, DataOffset);
+			DataOffset += PositionCardinality;
+			FeatureVector.SetVector(Feature, Indexer.MirrorTransform(SamplePresent.RootTransform).GetTranslation());
 		}
 	}
-}
 
-void UPoseSearchFeatureChannel_Trajectory::IndexDistanceFeatures(const  UE::PoseSearch::IAssetIndexer& Indexer, int32 SampleIdx, FPoseSearchFeatureVectorBuilder& FeatureVector) const
-{
-	// This function is very similar to AddTrajectoryTimeFeatures, but samples are taken in the distance domain
-	// instead of time domain.
+	if (bUseLinearVelocities)
+	{
+		for (int32 SubsampleIdx = 0; SubsampleIdx != SampleOffsets.Num(); ++SubsampleIdx)
+		{
+			const float SubsampleTime = GetSampleTime(Indexer, SubsampleIdx, SampleTime, Origin.RootDistance);
 
-	using namespace UE::PoseSearch;
-	using FSampleInfo = IAssetIndexer::FSampleInfo;
+			// For each pose subsample term, get the corresponding clip, accumulated root motion,
+			// and wrap the time parameter based on the clip's length.
+			const FSampleInfo SamplePast = Indexer.GetSampleInfoRelative(SubsampleTime - IndexingContext.SamplingContext->FiniteDelta, Origin);
+			const FSampleInfo SamplePresent = Indexer.GetSampleInfoRelative(SubsampleTime, Origin);
+			const FSampleInfo SampleFuture = Indexer.GetSampleInfoRelative(SubsampleTime + IndexingContext.SamplingContext->FiniteDelta, Origin);
 
-	const FAssetIndexingContext& IndexingContext = Indexer.GetIndexingContext();
+			// Mirror transforms if requested
+			const FTransform MirroredRootPast = Indexer.MirrorTransform(SamplePast.RootTransform);
+			const FTransform MirroredRootPresent = Indexer.MirrorTransform(SamplePresent.RootTransform);
+			const FTransform MirroredRootFuture = Indexer.MirrorTransform(SampleFuture.RootTransform);
 
+			// We can get a better finite difference if we ignore samples that have
+			// been clamped at either side of the clip. However, if the central sample 
+			// itself is clamped, or there are no samples that are clamped, we can just 
+			// use the central difference as normal.
+			FVector LinearVelocity;
+			if (SamplePast.bClamped && !SamplePresent.bClamped && !SampleFuture.bClamped)
+			{
+				LinearVelocity = (MirroredRootFuture.GetTranslation() - MirroredRootPresent.GetTranslation()) / IndexingContext.SamplingContext->FiniteDelta;
+			}
+			else if (SampleFuture.bClamped && !SamplePresent.bClamped && !SamplePast.bClamped)
+			{
+				LinearVelocity = (MirroredRootPresent.GetTranslation() - MirroredRootPast.GetTranslation()) / IndexingContext.SamplingContext->FiniteDelta;
+			}
+			else
+			{
+				LinearVelocity = (MirroredRootFuture.GetTranslation() - MirroredRootPast.GetTranslation()) / (IndexingContext.SamplingContext->FiniteDelta * 2.0f);
+			}
 
-	FPoseSearchFeatureDesc Feature;
-	Feature.ChannelIdx = GetChannelIndex();
-	Feature.ChannelFeatureId = 0; // Unused
-
-	float SampleTime = FMath::Min(SampleIdx * IndexingContext.Schema->SamplingInterval, IndexingContext.MainSampler->GetPlayLength());
-	FSampleInfo Origin = Indexer.GetSampleInfo(SampleTime);
+			const FPoseSearchFeatureDesc Feature = FPoseSearchFeatureDesc::Construct(GetChannelIndex(), 0, SubsampleIdx, EPoseSearchFeatureType::LinearVelocity, LinearVelocityCardinality, DataOffset);
+			DataOffset += LinearVelocityCardinality;
+			FeatureVector.SetVector(Feature, LinearVelocity);
+		}
+	}
 
-	FQuat RootReferenceRot = IndexingContext.SamplingContext->BoneContainer.GetRefPoseTransform(FCompactPoseBoneIndex(0)).GetRotation();
-	for (int32 SubsampleIdx = 0; SubsampleIdx != SampleOffsets.Num(); ++SubsampleIdx)
+	if (bUseFacingDirections)
 	{
-		Feature.SubsampleIdx = SubsampleIdx;
-
-		// For distance based sampling of the trajectory, we first have to look up the time value
-		// we're sampling given the desired travel distance of the root for this distance offset.
-		// Once we know the time, we can then carry on just like time-based sampling.
-		const float SubsampleDistance = Origin.RootDistance + SampleOffsets[SubsampleIdx];
-		float SubsampleTime = Indexer.GetSampleTimeFromDistance(SubsampleDistance);
-
-		// For each pose subsample term, get the corresponding clip, accumulated root motion,
-		// and wrap the time parameter based on the clip's length.
-		FSampleInfo Samples[3];
-		Samples[0] = Indexer.GetSampleInfoRelative(SubsampleTime - IndexingContext.SamplingContext->FiniteDelta, Origin);
-		Samples[1] = Indexer.GetSampleInfoRelative(SubsampleTime, Origin);
-		Samples[2] = Indexer.GetSampleInfoRelative(SubsampleTime + IndexingContext.SamplingContext->FiniteDelta, Origin);
-
-		// Mirror transforms if requested
-		FTransform MirroredRoots[3];
-		MirroredRoots[0] = Indexer.MirrorTransform(Samples[0].RootTransform);
-		MirroredRoots[1] = Indexer.MirrorTransform(Samples[1].RootTransform);
-		MirroredRoots[2] = Indexer.MirrorTransform(Samples[2].RootTransform);
-
-		// Add properties to the feature vector for the pose at SampleIdx
-		FeatureVector.SetTransform(Feature, MirroredRoots[1]);
-
-		// We can get a better finite difference if we ignore samples that have
-		// been clamped at either side of the clip. However, if the central sample 
-		// itself is clamped, or there are no samples that are clamped, we can just 
-		// use the central difference as normal.
-		if (Samples[0].bClamped && !Samples[1].bClamped && !Samples[2].bClamped)
-		{
-			FeatureVector.SetTransformVelocity(Feature, MirroredRoots[2], MirroredRoots[1], IndexingContext.SamplingContext->FiniteDelta);
-		}
-		else if (Samples[2].bClamped && !Samples[1].bClamped && !Samples[0].bClamped)
-		{
-			FeatureVector.SetTransformVelocity(Feature, MirroredRoots[1], MirroredRoots[0], IndexingContext.SamplingContext->FiniteDelta);
-		}
-		else
+		for (int32 SubsampleIdx = 0; SubsampleIdx != SampleOffsets.Num(); ++SubsampleIdx)
 		{
-			FeatureVector.SetTransformVelocity(Feature, MirroredRoots[2], MirroredRoots[1], MirroredRoots[0], IndexingContext.SamplingContext->FiniteDelta);
+			const float SubsampleTime = GetSampleTime(Indexer, SubsampleIdx, SampleTime, Origin.RootDistance);
+			const FSampleInfo SamplePresent = Indexer.GetSampleInfoRelative(SubsampleTime, Origin);
+			const FPoseSearchFeatureDesc Feature = FPoseSearchFeatureDesc::Construct(GetChannelIndex(), 0, SubsampleIdx, EPoseSearchFeatureType::ForwardVector, ForwardVectorCardinality, DataOffset);
+			DataOffset += ForwardVectorCardinality;
+			FeatureVector.SetVector(Feature, Indexer.MirrorTransform(SamplePresent.RootTransform).GetRotation().GetAxisY());
 		}
 	}
 }
@@ -1112,24 +1184,43 @@ bool UPoseSearchFeatureChannel_Trajectory::BuildQuery(
 			return false;
 	}
 
-	FPoseSearchFeatureDesc Feature;
-	Feature.ChannelIdx = GetChannelIndex();
-
+	// @todo: optimize dynamic allocation
+	TArray<FTrajectorySample> Samples;
+	Samples.Reserve(SampleOffsets.Num());
 	for (int32 Idx = 0, NextIterStartIdx = 0, Num = SampleOffsets.Num(); Idx < Num; ++Idx)
 	{
-		const float SampleOffset = SampleOffsets[Idx];
-		const FTrajectorySample Sample = FTrajectorySampleRange::IterSampleTrajectory(
-			SearchContext.Trajectory->Samples,
-			SampleDomain, 
-			SampleOffset, 
-			NextIterStartIdx);
+		Samples.Add(FTrajectorySampleRange::IterSampleTrajectory(SearchContext.Trajectory->Samples, SampleDomain, SampleOffsets[Idx], NextIterStartIdx));
+	}
 
-		Feature.SubsampleIdx = Idx;
+	int32 DataOffset = ChannelDataOffset;
+	if (bUsePositions)
+	{
+		for (int32 Idx = 0, Num = SampleOffsets.Num(); Idx < Num; ++Idx)
+		{
+			const FPoseSearchFeatureDesc Feature = FPoseSearchFeatureDesc::Construct(GetChannelIndex(), 0, Idx, EPoseSearchFeatureType::Position, PositionCardinality, DataOffset);
+			DataOffset += PositionCardinality;
+			InOutQuery.SetVector(Feature, Samples[Idx].Transform.GetTranslation());
+		}
+	}
 
-		Feature.Type = EPoseSearchFeatureType::LinearVelocity;
-		InOutQuery.SetVector(Feature, Sample.LinearVelocity);
+	if (bUseLinearVelocities)
+	{
+		for (int32 Idx = 0, Num = SampleOffsets.Num(); Idx < Num; ++Idx)
+		{
+			const FPoseSearchFeatureDesc Feature = FPoseSearchFeatureDesc::Construct(GetChannelIndex(), 0, Idx, EPoseSearchFeatureType::LinearVelocity, LinearVelocityCardinality, DataOffset);
+			DataOffset += LinearVelocityCardinality;
+			InOutQuery.SetVector(Feature, Samples[Idx].LinearVelocity);
+		}
+	}
 
-		InOutQuery.SetTransform(Feature, Sample.Transform);
+	if (bUseFacingDirections)
+	{ 
+		for (int32 Idx = 0, Num = SampleOffsets.Num(); Idx < Num; ++Idx)
+		{
+			const FPoseSearchFeatureDesc Feature = FPoseSearchFeatureDesc::Construct(GetChannelIndex(), 0, Idx, EPoseSearchFeatureType::ForwardVector, ForwardVectorCardinality, DataOffset);
+			DataOffset += ForwardVectorCardinality;
+			InOutQuery.SetVector(Feature, Samples[Idx].Transform.GetRotation().GetAxisY());
+		}
 	}
 
 	return true;
@@ -1143,9 +1234,6 @@ void UPoseSearchFeatureChannel_Trajectory::DebugDraw(const UE::PoseSearch::FDebu
 	const uint8 DepthPriority = ESceneDepthPriorityGroup::SDPG_Foreground + 2;
 	const bool bPersistent = EnumHasAnyFlags(DrawParams.Flags, EDebugDrawFlags::Persistent);
 
-	FPoseSearchFeatureDesc Feature;
-	Feature.ChannelIdx = GetChannelIndex();
-
 	const int32 NumSubsamples = SampleOffsets.Num();
 	if (NumSubsamples == 0)
 	{
@@ -1163,112 +1251,127 @@ void UPoseSearchFeatureChannel_Trajectory::DebugDraw(const UE::PoseSearch::FDebu
 		return OriginalColor * (1.0f - DrawDebugGradientStrength * (SampleIdx / (float)Denominator));
 	};
 
+	TArray<FVector> TrajectoryPos;
+	TrajectoryPos.AddUninitialized(NumSubsamples);
+	int32 DataOffset = ChannelDataOffset;
 	for (int32 SubsampleIdx = 0; SubsampleIdx != NumSubsamples; ++SubsampleIdx)
 	{
-		Feature.SubsampleIdx = SubsampleIdx;
-
-		FVector TrajectoryPos;
-		if (Reader.GetPosition(Feature, &TrajectoryPos))
+		const FPoseSearchFeatureDesc Feature = FPoseSearchFeatureDesc::Construct(GetChannelIndex(), 0, SubsampleIdx, EPoseSearchFeatureType::Position, PositionCardinality, DataOffset);
+		DataOffset += PositionCardinality;
+		if (bUsePositions && Reader.GetVector(Feature, &TrajectoryPos[SubsampleIdx]))
 		{
-			Feature.Type = EPoseSearchFeatureType::Position;
-
 			FLinearColor LinearColor = DrawParams.Color ? *DrawParams.Color : GetColorForFeature(Feature, Reader.GetLayout());
 			FLinearColor GradientColor = GetGradientColor(LinearColor, SubsampleIdx, NumSubsamples, DrawParams.Flags);
 			FColor Color = GradientColor.ToFColor(true);
 
-			TrajectoryPos = DrawParams.RootTransform.TransformPosition(TrajectoryPos);
+			TrajectoryPos[SubsampleIdx] = DrawParams.RootTransform.TransformPosition(TrajectoryPos[SubsampleIdx]);
 			if (EnumHasAnyFlags(DrawParams.Flags, EDebugDrawFlags::DrawFast | EDebugDrawFlags::DrawSearchIndex))
 			{
-				DrawDebugPoint(DrawParams.World, TrajectoryPos, DrawParams.PointSize, Color, bPersistent, DrawParams.DefaultLifeTime, DepthPriority);
+				DrawDebugPoint(DrawParams.World, TrajectoryPos[SubsampleIdx], DrawParams.PointSize, Color, bPersistent, DrawParams.DefaultLifeTime, DepthPriority);
 			}
 			else
 			{
-				DrawDebugSphere(DrawParams.World, TrajectoryPos, DrawDebugSphereSize, DrawDebugSphereSegments, Color, bPersistent, LifeTime, DepthPriority);
+				DrawDebugSphere(DrawParams.World, TrajectoryPos[SubsampleIdx], DrawDebugSphereSize, DrawDebugSphereSegments, Color, bPersistent, LifeTime, DepthPriority);
 			}
 		}
 		else
 		{
-			TrajectoryPos = DrawParams.RootTransform.GetTranslation();
+			TrajectoryPos[SubsampleIdx] = DrawParams.RootTransform.GetTranslation();
 		}
+	}
 
-		FVector TrajectoryVel;
-		if (Reader.GetLinearVelocity(Feature, &TrajectoryVel))
+	if (bUseLinearVelocities)
+	{
+		for (int32 SubsampleIdx = 0; SubsampleIdx != NumSubsamples; ++SubsampleIdx)
 		{
-			Feature.Type = EPoseSearchFeatureType::LinearVelocity;
-
-			FLinearColor LinearColor = DrawParams.Color ? *DrawParams.Color : GetColorForFeature(Feature, Reader.GetLayout());
-			FLinearColor GradientColor = GetGradientColor(LinearColor, SubsampleIdx, NumSubsamples, DrawParams.Flags);
-			FColor Color = GradientColor.ToFColor(true);
+			const FPoseSearchFeatureDesc Feature = FPoseSearchFeatureDesc::Construct(GetChannelIndex(), 0, SubsampleIdx, EPoseSearchFeatureType::LinearVelocity, LinearVelocityCardinality, DataOffset);
+			DataOffset += LinearVelocityCardinality;
+			FVector TrajectoryVel;
+			if (Reader.GetVector(Feature, &TrajectoryVel))
+			{
+				const FLinearColor LinearColor = DrawParams.Color ? *DrawParams.Color : GetColorForFeature(Feature, Reader.GetLayout());
+				const FLinearColor GradientColor = GetGradientColor(LinearColor, SubsampleIdx, NumSubsamples, DrawParams.Flags);
+				const FColor Color = GradientColor.ToFColor(true);
 
-			TrajectoryVel *= DrawDebugVelocityScale;
-			TrajectoryVel = DrawParams.RootTransform.TransformVector(TrajectoryVel);
-			FVector TrajectoryVelDirection = TrajectoryVel.GetSafeNormal();
+				TrajectoryVel *= DrawDebugVelocityScale;
+				TrajectoryVel = DrawParams.RootTransform.TransformVector(TrajectoryVel);
+				const FVector TrajectoryVelDirection = TrajectoryVel.GetSafeNormal();
 
+				if (EnumHasAnyFlags(DrawParams.Flags, EDebugDrawFlags::DrawSearchIndex))
+				{
+					DrawDebugPoint(DrawParams.World, TrajectoryVel, DrawParams.PointSize, Color, bPersistent, DrawParams.DefaultLifeTime, DepthPriority);
+				}
+				else
+				{
+					const float AdjustedThickness =
+						EnumHasAnyFlags(DrawParams.Flags, EDebugDrawFlags::DrawFast) ?
+						0.0f : DrawDebugLineThickness;
 
-			if (EnumHasAnyFlags(DrawParams.Flags, EDebugDrawFlags::DrawSearchIndex))
-			{
-				DrawDebugPoint(DrawParams.World, TrajectoryVel, DrawParams.PointSize, Color, bPersistent, DrawParams.DefaultLifeTime, DepthPriority);
-			}
-			else
-			{
-				const float AdjustedThickness =
-					EnumHasAnyFlags(DrawParams.Flags, EDebugDrawFlags::DrawFast) ?
-					0.0f : DrawDebugLineThickness;
-
-				DrawDebugDirectionalArrow(
-					DrawParams.World,
-					TrajectoryPos + TrajectoryVelDirection * DrawDebugSphereSize,
-					TrajectoryPos + TrajectoryVel,
-					DrawDebugArrowSize,
-					Color,
-					bPersistent,
-					LifeTime,
-					DepthPriority,
-					AdjustedThickness
-				);
+					DrawDebugDirectionalArrow(
+						DrawParams.World,
+						TrajectoryPos[SubsampleIdx] + TrajectoryVelDirection * DrawDebugSphereSize,
+						TrajectoryPos[SubsampleIdx] + TrajectoryVel,
+						DrawDebugArrowSize,
+						Color,
+						bPersistent,
+						LifeTime,
+						DepthPriority,
+						AdjustedThickness
+					);
+				}
 			}
 		}
+	}
 
-		FVector TrajectoryForward;
-		if (Reader.GetForwardVector(Feature, &TrajectoryForward))
+	if (bUseFacingDirections)
+	{
+		for (int32 SubsampleIdx = 0; SubsampleIdx != NumSubsamples; ++SubsampleIdx)
 		{
-			Feature.Type = EPoseSearchFeatureType::ForwardVector;
+			const FPoseSearchFeatureDesc Feature = FPoseSearchFeatureDesc::Construct(GetChannelIndex(), 0, SubsampleIdx, EPoseSearchFeatureType::ForwardVector, ForwardVectorCardinality, DataOffset);
+			DataOffset += ForwardVectorCardinality;
+			FVector TrajectoryForward;
+			if (Reader.GetVector(Feature, &TrajectoryForward))
+			{
+				const FLinearColor LinearColor = DrawParams.Color ? *DrawParams.Color : GetColorForFeature(Feature, Reader.GetLayout());
+				const FLinearColor GradientColor = GetGradientColor(LinearColor, SubsampleIdx, NumSubsamples, DrawParams.Flags);
+				const FColor Color = GradientColor.ToFColor(true);
 
-			FLinearColor LinearColor = DrawParams.Color ? *DrawParams.Color : GetColorForFeature(Feature, Reader.GetLayout());
-			FLinearColor GradientColor = GetGradientColor(LinearColor, SubsampleIdx, NumSubsamples, DrawParams.Flags);
-			FColor Color = GradientColor.ToFColor(true);
+				TrajectoryForward = DrawParams.RootTransform.TransformVector(TrajectoryForward);
 
-			TrajectoryForward = DrawParams.RootTransform.TransformVector(TrajectoryForward);
+				if (EnumHasAnyFlags(DrawParams.Flags, EDebugDrawFlags::DrawSearchIndex))
+				{
+					DrawDebugPoint(DrawParams.World, TrajectoryForward, DrawParams.PointSize, Color, bPersistent, DrawParams.DefaultLifeTime, DepthPriority);
+				}
+				else
+				{
+					const float AdjustedThickness =
+						EnumHasAnyFlags(DrawParams.Flags, EDebugDrawFlags::DrawFast) ?
+						0.0f : DrawDebugLineThickness;
 
-			if (EnumHasAnyFlags(DrawParams.Flags, EDebugDrawFlags::DrawSearchIndex))
-			{
-				DrawDebugPoint(DrawParams.World, TrajectoryForward, DrawParams.PointSize, Color, bPersistent, DrawParams.DefaultLifeTime, DepthPriority);
-			}
-			else
-			{
-				const float AdjustedThickness =
-					EnumHasAnyFlags(DrawParams.Flags, EDebugDrawFlags::DrawFast) ?
-					0.0f : DrawDebugLineThickness;
-
-				DrawDebugDirectionalArrow(
-					DrawParams.World,
-					TrajectoryPos + TrajectoryForward * DrawDebugSphereSize,
-					TrajectoryPos + TrajectoryForward * DrawDebugSphereSize * 2.0f,
-					DrawDebugArrowSize,
-					Color,
-					bPersistent,
-					LifeTime,
-					DepthPriority,
-					AdjustedThickness
-				);
+					DrawDebugDirectionalArrow(
+						DrawParams.World,
+						TrajectoryPos[SubsampleIdx] + TrajectoryForward * DrawDebugSphereSize,
+						TrajectoryPos[SubsampleIdx] + TrajectoryForward * DrawDebugSphereSize * 2.0f,
+						DrawDebugArrowSize,
+						Color,
+						bPersistent,
+						LifeTime,
+						DepthPriority,
+						AdjustedThickness
+					);
+				}
 			}
 		}
+	}
 
-		if (EnumHasAnyFlags(DrawParams.Flags, EDebugDrawFlags::DrawSampleLabels))
+	if (EnumHasAnyFlags(DrawParams.Flags, EDebugDrawFlags::DrawSampleLabels))
+	{
+		for (int32 SubsampleIdx = 0; SubsampleIdx != NumSubsamples; ++SubsampleIdx)
 		{
-			FLinearColor LinearColor = DrawParams.Color ? *DrawParams.Color : GetColorForFeature(Feature, Reader.GetLayout());
-			FLinearColor GradientColor = GetGradientColor(LinearColor, SubsampleIdx, NumSubsamples, DrawParams.Flags);
-			FColor Color = GradientColor.ToFColor(true);
+			const FPoseSearchFeatureDesc Feature = FPoseSearchFeatureDesc::Construct(GetChannelIndex(), 0, SubsampleIdx, EPoseSearchFeatureType::Position, PositionCardinality, -1);
+			const FLinearColor LinearColor = DrawParams.Color ? *DrawParams.Color : GetColorForFeature(Feature, Reader.GetLayout());
+			const FLinearColor GradientColor = GetGradientColor(LinearColor, SubsampleIdx, NumSubsamples, DrawParams.Flags);
+			const FColor Color = GradientColor.ToFColor(true);
 
 			FString SampleLabel;
 			if (DrawParams.LabelPrefix.IsEmpty())
@@ -1281,7 +1384,7 @@ void UPoseSearchFeatureChannel_Trajectory::DebugDraw(const UE::PoseSearch::FDebu
 			}
 			DrawDebugString(
 				DrawParams.World,
-				TrajectoryPos + DrawDebugSampleLabelOffset,
+				TrajectoryPos[SubsampleIdx] + DrawDebugSampleLabelOffset,
 				SampleLabel,
 				nullptr,
 				Color,
diff --git a/Engine/Plugins/Experimental/Animation/PoseSearch/Source/Runtime/Public/PoseSearch/PoseSearch.h b/Engine/Plugins/Experimental/Animation/PoseSearch/Source/Runtime/Public/PoseSearch/PoseSearch.h
--- a/Engine/Plugins/Experimental/Animation/PoseSearch/Source/Runtime/Public/PoseSearch/PoseSearch.h
+++ b/Engine/Plugins/Experimental/Animation/PoseSearch/Source/Runtime/Public/PoseSearch/PoseSearch.h
@@ -245,7 +245,7 @@ struct POSESEARCH_API FPoseSearchFeatureDesc
 
 	bool IsValid() const { return Type != EPoseSearchFeatureType::Invalid; }
 
-	static FPoseSearchFeatureDesc Construct(int8 ChannelIdx, int8 ChannelFeatureId, int8 SubsampleIdx, EPoseSearchFeatureType Type, int16 Cardinality)
+	static FPoseSearchFeatureDesc Construct(int8 ChannelIdx, int8 ChannelFeatureId, int8 SubsampleIdx, EPoseSearchFeatureType Type, int16 Cardinality, int16 ValueOffset)
 	{
 		FPoseSearchFeatureDesc PoseSearchFeatureDesc;
 		PoseSearchFeatureDesc.ChannelIdx = ChannelIdx;
@@ -253,7 +253,7 @@ struct POSESEARCH_API FPoseSearchFeatureDesc
 		PoseSearchFeatureDesc.SubsampleIdx = SubsampleIdx;
 		PoseSearchFeatureDesc.Type = Type;
 		PoseSearchFeatureDesc.Cardinality = Cardinality;
-		//PoseSearchFeatureDesc.ValueOffset
+		PoseSearchFeatureDesc.ValueOffset = ValueOffset;
 		return PoseSearchFeatureDesc;
 	}
 };
@@ -276,12 +276,9 @@ struct POSESEARCH_API FPoseSearchFeatureVectorLayout
 	int32 NumFloats = 0;
 
 	bool IsValid(int32 ChannelCount) const;
-
-	bool EnumerateBy(int32 ChannelIdx, EPoseSearchFeatureType Type, int32& InOutFeatureIdx) const;
 };
 
 
-
 //////////////////////////////////////////////////////////////////////////
 // Asset sampling and indexing
 
@@ -946,17 +943,9 @@ public:
 	TArrayView<const float> GetValues() const { return Values; }
 	TArrayView<const float> GetNormalizedValues() const { return ValuesNormalized; }
 
-	void SetTransform(FPoseSearchFeatureDesc Feature, const FTransform& Transform);
-	void SetTransformVelocity(FPoseSearchFeatureDesc Feature, const FTransform& Transform, const FTransform& PrevTransform, float DeltaTime);
-	void SetTransformVelocity(FPoseSearchFeatureDesc Feature, const FTransform& NextTransform, const FTransform& Transform, const FTransform& PrevTransform, float DeltaTime);
-	void SetPosition(FPoseSearchFeatureDesc Feature, const FVector& Translation);
-	void SetRotation(FPoseSearchFeatureDesc Feature, const FQuat& Rotation);
-	void SetLinearVelocity(FPoseSearchFeatureDesc Feature, const FTransform& Transform, const FTransform& PrevTransform, float DeltaTime);
-	void SetLinearVelocity(FPoseSearchFeatureDesc Feature, const FTransform& NextTransform, const FTransform& Transform, const FTransform& PrevTransform, float DeltaTime);
-	void SetAngularVelocity(FPoseSearchFeatureDesc Feature, const FTransform& Transform, const FTransform& PrevTransform, float DeltaTime);
-	void SetAngularVelocity(FPoseSearchFeatureDesc Feature, const FTransform& NextTransform, const FTransform& Transform, const FTransform& PrevTransform, float DeltaTime);
-	void SetVector(FPoseSearchFeatureDesc Feature, const FVector& Vector);
-	void SetPhase(FPoseSearchFeatureDesc Feature, const FVector2D& Phase);
+	void SetRotation(const FPoseSearchFeatureDesc& Feature, const FQuat& Rotation);
+	void SetVector(const FPoseSearchFeatureDesc& Feature, const FVector& Vector);
+	void SetPhase(const FPoseSearchFeatureDesc& Feature, const FVector2D& Phase);
 
 	void CopyFromSearchIndex(const FPoseSearchIndex& SearchIndex, int32 PoseIdx);
 
@@ -1336,14 +1325,9 @@ public:
 	void SetValues(TArrayView<const float> Values);
 	bool IsValid() const;
 
-	bool GetTransform(FPoseSearchFeatureDesc Feature, FTransform* OutTransform) const;
-	bool GetPosition(FPoseSearchFeatureDesc Feature, FVector* OutPosition) const;
-	bool GetRotation(FPoseSearchFeatureDesc Feature, FQuat* OutRotation) const;
-	bool GetForwardVector(FPoseSearchFeatureDesc Feature, FVector* OutForwardVector) const;
-	bool GetLinearVelocity(FPoseSearchFeatureDesc Feature, FVector* OutLinearVelocity) const;
-	bool GetAngularVelocity(FPoseSearchFeatureDesc Feature, FVector* OutAngularVelocity) const;
-	bool GetVector(FPoseSearchFeatureDesc Feature, FVector* OutVector) const;
-	bool GetPhase(FPoseSearchFeatureDesc Element, FVector2D* OutPhase) const;
+	bool GetRotation(const FPoseSearchFeatureDesc& Feature, FQuat* OutRotation) const;
+	bool GetVector(const FPoseSearchFeatureDesc& Feature, FVector* OutVector) const;
+	bool GetPhase(const FPoseSearchFeatureDesc& Feature, FVector2D* OutPhase) const;
 
 	const FPoseSearchFeatureVectorLayout* GetLayout() const { return Layout; }
 
@@ -1388,7 +1372,7 @@ public:
 	FPoseSearchFeatureVectorBuilder& GetQueryBuilder() { return QueryBuilder; }
 
 private:
-	bool TrySampleLocalPose(float Time, const TArray<FBoneIndexType>& RequiredBones, TArray<FTransform>& LocalPose, FTransform& RootTransform);
+	bool TrySampleLocalPose(float Time, const TArray<FBoneIndexType>& RequiredBones, TArray<FTransform>& LocalPose, FTransform& RootTransform) const;
 
 	struct FPose
 	{
diff --git a/Engine/Plugins/Experimental/Animation/PoseSearch/Source/Runtime/Public/PoseSearch/PoseSearchFeatureChannels.h b/Engine/Plugins/Experimental/Animation/PoseSearch/Source/Runtime/Public/PoseSearch/PoseSearchFeatureChannels.h
--- a/Engine/Plugins/Experimental/Animation/PoseSearch/Source/Runtime/Public/PoseSearch/PoseSearchFeatureChannels.h
+++ b/Engine/Plugins/Experimental/Animation/PoseSearch/Source/Runtime/Public/PoseSearch/PoseSearchFeatureChannels.h
@@ -8,6 +8,11 @@
 
 #include "PoseSearchFeatureChannels.generated.h"
 
+namespace UE::PoseSearch
+{
+	struct BoneTransformsCache;
+}
+
 USTRUCT()
 struct FPoseSearchPoseFeatureInfo
 {
@@ -52,8 +57,8 @@ public:
 	virtual void DebugDraw(const UE::PoseSearch::FDebugDrawParams& DrawParams, const UE::PoseSearch::FFeatureVectorReader& Reader) const override;
 
 protected:
-	virtual void AddPoseFeatures(const UE::PoseSearch::IAssetIndexer& Indexer, int32 SampleIdx, FPoseSearchFeatureVectorBuilder& FeatureVector, const TArray<TArray<FVector2D>>& Phases) const;
-	void CalculatePhases(const UE::PoseSearch::IAssetIndexer& Indexer, UE::PoseSearch::FAssetIndexingOutput& IndexingOutput, TArray<TArray<FVector2D>>& OutPhases) const;
+	void AddPoseFeatures(UE::PoseSearch::BoneTransformsCache& BoneTransformsCache, int32 SampleIdx, FPoseSearchFeatureVectorBuilder& FeatureVector, const TArray<TArray<FVector2D>>& Phases) const;
+	void CalculatePhases(UE::PoseSearch::BoneTransformsCache& BoneTransformsCache, UE::PoseSearch::FAssetIndexingOutput& IndexingOutput, TArray<TArray<FVector2D>>& OutPhases) const;
 
 	enum { PositionCardinality = 3 };
 	enum { RotationCardinality = 6 };
@@ -108,8 +113,8 @@ public:
 	virtual void DebugDraw(const UE::PoseSearch::FDebugDrawParams& DrawParams, const UE::PoseSearch::FFeatureVectorReader& Reader) const override;
 
 protected:
-	virtual void IndexTimeFeatures(const UE::PoseSearch::IAssetIndexer& Indexer, int32 SampleIdx, FPoseSearchFeatureVectorBuilder& FeatureVector) const;
-	virtual void IndexDistanceFeatures(const UE::PoseSearch::IAssetIndexer& Indexer, int32 SampleIdx, FPoseSearchFeatureVectorBuilder& FeatureVector) const;
+	virtual void IndexAssetPrivate(const UE::PoseSearch::IAssetIndexer& Indexer, int32 SampleIdx, FPoseSearchFeatureVectorBuilder& FeatureVector) const;
+	float GetSampleTime(const UE::PoseSearch::IAssetIndexer& Indexer, int32 SubsampleIdx, float SampleTime, float RootDistance) const;
 
 	enum { PositionCardinality = 3 };
 	enum { LinearVelocityCardinality = 3 };
