motion matching
- deprecation of FPoseSearchFeatureDesc
- deprecation of FPoseSearchFeatureVectorLayout
- moving ComputeMeanDeviations into channels

diff --git a/Engine/Plugins/Experimental/Animation/PoseSearch/Source/Editor/Private/PoseSearchDebugger.h b/Engine/Plugins/Experimental/Animation/PoseSearch/Source/Editor/Private/PoseSearchDebugger.h
--- a/Engine/Plugins/Experimental/Animation/PoseSearch/Source/Editor/Private/PoseSearchDebugger.h
+++ b/Engine/Plugins/Experimental/Animation/PoseSearch/Source/Editor/Private/PoseSearchDebugger.h
@@ -36,7 +36,6 @@ namespace UE::PoseSearch
 	namespace DebuggerDatabaseColumns { struct IColumn; }
 	struct FDebugDrawParams;
 	struct FTraceMotionMatchingStateMessage;
-	class FFeatureVectorReader;
 	class FDebuggerDatabaseRowData;
 } // namespace UE::PoseSearch
 
diff --git a/Engine/Plugins/Experimental/Animation/PoseSearch/Source/Runtime/Private/PoseSearch.cpp b/Engine/Plugins/Experimental/Animation/PoseSearch/Source/Runtime/Private/PoseSearch.cpp
--- a/Engine/Plugins/Experimental/Animation/PoseSearch/Source/Runtime/Private/PoseSearch.cpp
+++ b/Engine/Plugins/Experimental/Animation/PoseSearch/Source/Runtime/Private/PoseSearch.cpp
@@ -74,18 +74,6 @@ namespace UE { namespace PoseSearch {
 	constexpr EParallelForFlags ParallelForFlags = EParallelForFlags::None;
 #endif // UE_POSE_SEARCH_FORCE_SINGLE_THREAD
 
-using RowMajorVector = Eigen::Matrix<float, 1, Eigen::Dynamic, Eigen::RowMajor>;
-using RowMajorVectorMap = Eigen::Map<RowMajorVector, Eigen::RowMajor>;
-using RowMajorVectorMapConst = Eigen::Map<const RowMajorVector, Eigen::RowMajor>;
-
-using RowMajorMatrix = Eigen::Matrix<float, Eigen::Dynamic, Eigen::Dynamic, Eigen::RowMajor>;
-using RowMajorMatrixMap = Eigen::Map<RowMajorMatrix, Eigen::RowMajor>;
-using RowMajorMatrixMapConst = Eigen::Map<const RowMajorMatrix, Eigen::RowMajor>;
-
-using ColMajorMatrix = Eigen::Matrix<float, Eigen::Dynamic, Eigen::Dynamic, Eigen::ColMajor>;
-using ColMajorMatrixMap = Eigen::Map<ColMajorMatrix, Eigen::ColMajor>;
-using ColMajorMatrixMapConst = Eigen::Map<const ColMajorMatrix, Eigen::ColMajor>;
-
 static bool IsSamplingRangeValid(FFloatInterval Range)
 {
 	return Range.IsValid() && (Range.Min >= 0.0f);
@@ -263,40 +251,15 @@ void UPoseSearchFeatureChannel::InitializeSchema(UE::PoseSearch::FSchemaInitiali
 	ChannelDataOffset = Initializer.GetCurrentChannelDataOffset();
 }
 
-//////////////////////////////////////////////////////////////////////////
-// FPoseSearchFeatureDesc
-
-bool FPoseSearchFeatureDesc::operator==(const FPoseSearchFeatureDesc& Other) const
-{
-	return
-		(ChannelIdx == Other.ChannelIdx) &&
-		(ChannelFeatureId == Other.ChannelFeatureId) &&
-		(SubsampleIdx == Other.SubsampleIdx) &&
-		// ignored or else Schema->Layout.Features.Find(Feature); called by FPoseSearchFeatureVectorBuilder will fail since FPoseSearchFeatureDesc doesn't get created fully :(
-		//(Cardinality == Other.Cardinality) &&
-		(Type == Other.Type);
-}
-
-
-//////////////////////////////////////////////////////////////////////////
-// FPoseSearchFeatureVectorLayout
-
-bool FPoseSearchFeatureVectorLayout::IsValid(int32 ChannelCount) const
+// base implementation calculating a single mean deviation value (replicated ChannelCardinality times into MeanDeviations starting at DataOffset index) from all the features data associated to this channel
+void UPoseSearchFeatureChannel::ComputeMeanDeviations(const Eigen::MatrixXd& CenteredPoseMatrix, Eigen::VectorXd& MeanDeviations) const
 {
-	if (NumFloats == 0)
-	{
-		return false;
-	}
+	using namespace UE::PoseSearch;
 
-	for (const FPoseSearchFeatureDesc& Feature : Features)
-	{
-		if (Feature.ChannelIdx >= ChannelCount)
-		{
-			return false;
-		}
-	}
+	int32 DataOffset = ChannelDataOffset;
+	FFeatureVectorHelper::ComputeMeanDeviations(CenteredPoseMatrix, MeanDeviations, DataOffset, ChannelCardinality);
 
-	return true;
+	check(DataOffset == ChannelDataOffset + ChannelCardinality);
 }
 
 //////////////////////////////////////////////////////////////////////////
@@ -310,17 +273,6 @@ int32 FSchemaInitializer::AddBoneReference(const FBoneReference& BoneReference)
 	return BoneReferences.AddUnique(BoneReference);
 }
 
-int32 FSchemaInitializer::AddFeatureDesc(const FPoseSearchFeatureDesc& FeatureDesc)
-{
-	check(FeatureDesc.ChannelIdx == CurrentChannelIdx);
-	check(Features.Num() + 1 <= UPoseSearchSchema::MaxFeatures);
-	check(!Features.Contains(FeatureDesc));
-	check(FeatureDesc.Cardinality > 0);
-	check(FeatureDesc.ValueOffset == CurrentChannelDataOffset);
-	CurrentChannelDataOffset += FeatureDesc.Cardinality;
-	return Features.Add(FeatureDesc);
-}
-
 } // namespace UE::PoseSearch
 
 void UPoseSearchSchema::Finalize()
@@ -343,9 +295,7 @@ void UPoseSearchSchema::Finalize()
 		Channel->InitializeSchema(Initializer);
 	}
 
-	// @todo: replace FPoseSearchFeatureVectorLayout with channels after itemizing UPoseSearchFeatureChannel_Pose into position, velocity, orientation / heading, phase etc 
-	Layout.Features = MoveTemp(Initializer.Features);
-	Layout.NumFloats = Initializer.GetCurrentChannelDataOffset();
+	SchemaCardinality = Initializer.GetCurrentChannelDataOffset();
 
 	BoneReferences = MoveTemp(Initializer.BoneReferences);
 
@@ -373,6 +323,11 @@ void UPoseSearchSchema::PostLoad()
 
 	// Migrate deprecated schema properties into channels
 #if WITH_EDITORONLY_DATA
+	if (SchemaCardinality == 0)
+	{
+		bNeedFinalize = true;
+	}
+
 	if (!TrajectorySampleTimes_DEPRECATED.IsEmpty())
 	{
 		bNeedFinalize = true;
@@ -416,18 +371,6 @@ void UPoseSearchSchema::PostLoad()
 		PoseSampleTimes_DEPRECATED.Empty();
 	}
 
-	if (!bNeedFinalize)
-	{
-		for (const FPoseSearchFeatureDesc& Feature : Layout.Features)
-		{
-			if (Feature.Cardinality == 0)
-			{
-				bNeedFinalize = true;
-				break;
-			}
-		}
-	}
-
 	if (!bNeedFinalize)
 	{
 		for (const UPoseSearchFeatureChannel* Channel : Channels)
@@ -468,7 +411,7 @@ bool UPoseSearchSchema::IsValid() const
 
 	bValid &= (BoneReferences.Num() == BoneIndices.Num());
 
-	bValid &= Layout.IsValid(Channels.Num());
+	bValid &= SchemaCardinality > 0;
 	
 	return bValid;
 }
@@ -645,7 +588,7 @@ float FPoseSearchIndex::GetAssetTime(int32 PoseIdx, const FPoseSearchIndexAsset*
 bool FPoseSearchIndex::IsValid() const
 {
 	const bool bSchemaValid = Schema && Schema->IsValid();
-	const bool bSearchIndexValid = bSchemaValid && (NumPoses * Schema->Layout.NumFloats == Values.Num());
+	const bool bSearchIndexValid = bSchemaValid && (NumPoses * Schema->SchemaCardinality == Values.Num());
 
 	return bSearchIndexValid;
 }
@@ -659,8 +602,8 @@ bool FPoseSearchIndex::IsEmpty() const
 TArrayView<const float> FPoseSearchIndex::GetPoseValues(int32 PoseIdx) const
 {
 	check(PoseIdx < NumPoses);
-	int32 ValueOffset = PoseIdx * Schema->Layout.NumFloats;
-	return MakeArrayView(&Values[ValueOffset], Schema->Layout.NumFloats);
+	int32 ValueOffset = PoseIdx * Schema->SchemaCardinality;
+	return MakeArrayView(&Values[ValueOffset], Schema->SchemaCardinality);
 }
 
 void FPoseSearchIndex::Reset()
@@ -775,7 +718,7 @@ UE::PoseSearch::FSearchResult UPoseSearchSequenceMetaData::Search(FPoseSearchCon
 	Result.ComposedQuery.Normalize(SearchIndex);
 	TArrayView<const float> NormalizedQueryValues = Result.ComposedQuery.GetNormalizedValues();
 
-	if (!ensure(NormalizedQueryValues.Num() == SearchIndex.Schema->Layout.NumFloats))
+	if (!ensure(NormalizedQueryValues.Num() == SearchIndex.Schema->SchemaCardinality))
 	{
 		return Result;
 	}
@@ -976,7 +919,7 @@ const FString UPoseSearchDatabase::GetSourceAssetName(const FPoseSearchIndexAsse
 
 int32 UPoseSearchDatabase::GetNumberOfPrincipalComponents() const
 {
-	return FMath::Min<int32>(NumberOfPrincipalComponents, Schema->Layout.NumFloats);
+	return FMath::Min<int32>(NumberOfPrincipalComponents, Schema->SchemaCardinality);
 }
 
 bool UPoseSearchDatabase::IsValidForIndexing() const
@@ -1817,7 +1760,7 @@ UE::PoseSearch::FSearchResult UPoseSearchDatabase::SearchPCAKDTree(FPoseSearchCo
 
 	FSearchResult Result;
 
-	const int32 NumDimensions = Schema->Layout.NumFloats;
+	const int32 NumDimensions = Schema->SchemaCardinality;
 	const FPoseSearchIndex* SearchIndex = GetSearchIndex();
 	check(SearchIndex);
 
@@ -2051,8 +1994,6 @@ UE::PoseSearch::FSearchResult UPoseSearchDatabaseSet::Search(FPoseSearchContext&
 	FSearchResult Result;
 	FPoseSearchCost ContinuityCost;
 
-	TArray<const FPoseSearchDatabaseSetEntry*, TInlineAllocator<8>> FallBackEntries;
-
 	auto ProcessActiveEntry = 
 		[&SearchContext, &Result, &ContinuityCost](const FPoseSearchDatabaseSetEntry& Entry) 
 		-> EPoseSearchPostSearchStatus
@@ -2142,92 +2083,14 @@ void FPoseSearchFeatureVectorBuilder::Reset()
 	Schema = nullptr;
 	Values.Reset(0);
 	ValuesNormalized.Reset(0);
-	NumFeaturesAdded = 0;
-	FeaturesAdded.Reset();
 }
 
 void FPoseSearchFeatureVectorBuilder::ResetFeatures()
 {
 	Values.Reset(0);
-	Values.SetNumZeroed(Schema->Layout.NumFloats);
+	Values.SetNumZeroed(Schema->SchemaCardinality);
 	ValuesNormalized.Reset(0);
-	ValuesNormalized.SetNumZeroed(Schema->Layout.NumFloats);
-	NumFeaturesAdded = 0;
-	FeaturesAdded.Init(false, Schema->Layout.Features.Num());
-}
-
-void FPoseSearchFeatureVectorBuilder::SetPhase(const FPoseSearchFeatureDesc& Feature, const FVector2D& Phase)
-{
-	check(Feature.Type == EPoseSearchFeatureType::Phase);
-	check(Feature.Cardinality == 2);
-
-	const int32 FeatureIndex = Schema->Layout.Features.Find(Feature);
-	check(FeatureIndex >= 0);
-	if (FeatureIndex >= 0)
-	{
-		const FPoseSearchFeatureDesc& FoundFeature = Schema->Layout.Features[FeatureIndex];
-		check(Feature.ValueOffset < 0 || Feature.ValueOffset == FoundFeature.ValueOffset);
-
-		Values[FoundFeature.ValueOffset + 0] = Phase.X;
-		Values[FoundFeature.ValueOffset + 1] = Phase.Y;
-
-		if (!FeaturesAdded[FeatureIndex])
-		{
-			FeaturesAdded[FeatureIndex] = true;
-			++NumFeaturesAdded;
-		}
-	}
-}
-
-void FPoseSearchFeatureVectorBuilder::SetRotation(const FPoseSearchFeatureDesc& Feature, const FQuat& Rotation)
-{
-	check(Feature.Type == EPoseSearchFeatureType::Rotation);
-	check(Feature.Cardinality == 6);
-	const int32 FeatureIndex = Schema->Layout.Features.Find(Feature);
-	check(FeatureIndex >= 0);
-	if (FeatureIndex >= 0)
-	{
-		FVector X = Rotation.GetAxisX();
-		FVector Y = Rotation.GetAxisY();
-
-		const FPoseSearchFeatureDesc& FoundFeature = Schema->Layout.Features[FeatureIndex];
-		check(Feature.ValueOffset < 0 || Feature.ValueOffset == FoundFeature.ValueOffset);
-
-		Values[FoundFeature.ValueOffset + 0] = X.X;
-		Values[FoundFeature.ValueOffset + 1] = X.Y;
-		Values[FoundFeature.ValueOffset + 2] = X.Z;
-		Values[FoundFeature.ValueOffset + 3] = Y.X;
-		Values[FoundFeature.ValueOffset + 4] = Y.Y;
-		Values[FoundFeature.ValueOffset + 5] = Y.Z;
-
-		if (!FeaturesAdded[FeatureIndex])
-		{
-			FeaturesAdded[FeatureIndex] = true;
-			++NumFeaturesAdded;
-		}
-	}
-}
-
-void FPoseSearchFeatureVectorBuilder::SetVector(const FPoseSearchFeatureDesc& Feature, const FVector& Vector)
-{
-	check(Feature.Cardinality == 3);
-	const int32 FeatureIndex = Schema->Layout.Features.Find(Feature);
-	check(FeatureIndex >= 0);
-	if (FeatureIndex >= 0)
-	{
-		const FPoseSearchFeatureDesc& FoundFeature = Schema->Layout.Features[FeatureIndex]; 
-		check(Feature.ValueOffset < 0 || Feature.ValueOffset == FoundFeature.ValueOffset);
-
-		Values[FoundFeature.ValueOffset + 0] = Vector[0];
-		Values[FoundFeature.ValueOffset + 1] = Vector[1];
-		Values[FoundFeature.ValueOffset + 2] = Vector[2];
-
-		if (!FeaturesAdded[FeatureIndex])
-		{
-			FeaturesAdded[FeatureIndex] = true;
-			++NumFeaturesAdded;
-		}
-	}
+	ValuesNormalized.SetNumZeroed(Schema->SchemaCardinality);
 }
 
 void FPoseSearchFeatureVectorBuilder::CopyFromSearchIndex(const FPoseSearchIndex& SearchIndex, int32 PoseIdx)
@@ -2239,14 +2102,11 @@ void FPoseSearchFeatureVectorBuilder::CopyFromSearchIndex(const FPoseSearchIndex
 	ValuesNormalized = FeatureVector;
 	Values = FeatureVector;
 	SearchIndex.InverseNormalize(Values);
-
-	NumFeaturesAdded = Schema->Layout.Features.Num();
-	FeaturesAdded.SetRange(0, FeaturesAdded.Num(), true);
 }
 
 bool FPoseSearchFeatureVectorBuilder::IsInitialized() const
 {
-	return (Schema != nullptr) && (Values.Num() == Schema->Layout.NumFloats);
+	return (Schema != nullptr) && (Values.Num() == Schema->SchemaCardinality);
 }
 
 bool FPoseSearchFeatureVectorBuilder::IsInitializedForSchema(const UPoseSearchSchema* InSchema) const
@@ -2254,21 +2114,11 @@ bool FPoseSearchFeatureVectorBuilder::IsInitializedForSchema(const UPoseSearchSc
 	return (Schema == InSchema) && IsInitialized();
 }
 
-bool FPoseSearchFeatureVectorBuilder::IsComplete() const
-{
-	return NumFeaturesAdded == Schema->Layout.Features.Num();
-}
-
 bool FPoseSearchFeatureVectorBuilder::IsCompatible(const FPoseSearchFeatureVectorBuilder& OtherBuilder) const
 {
 	return IsInitialized() && (Schema == OtherBuilder.Schema);
 }
 
-const TBitArray<>& FPoseSearchFeatureVectorBuilder::GetFeaturesAdded() const
-{
-	return FeaturesAdded;
-}
-
 void FPoseSearchFeatureVectorBuilder::Normalize(const FPoseSearchIndex& ForSearchIndex)
 {
 	ValuesNormalized = Values;
@@ -2480,109 +2330,84 @@ float FPoseHistory::GetSampleTimeInterval() const
 	return TimeHorizon / (Knots.Max() - 1);
 }
 
-
 //////////////////////////////////////////////////////////////////////////
-// FPoseSearchFeatureVectorReader
-
-void FFeatureVectorReader::Init(const FPoseSearchFeatureVectorLayout* InLayout)
+// FFeatureVectorHelper
+void FFeatureVectorHelper::EncodeQuat(TArrayView<float> Values, int32& DataOffset, const FQuat& Quat)
 {
-	check(InLayout);
-	Layout = InLayout;
-}
+	const FVector X = Quat.GetAxisX();
+	const FVector Y = Quat.GetAxisY();
 
-void FFeatureVectorReader::SetValues(TArrayView<const float> InValues)
-{
-	check(Layout);
-	check(Layout->NumFloats == InValues.Num());
-	Values = InValues;
-}
+	Values[DataOffset + 0] = X.X;
+	Values[DataOffset + 1] = X.Y;
+	Values[DataOffset + 2] = X.Z;
+	Values[DataOffset + 3] = Y.X;
+	Values[DataOffset + 4] = Y.Y;
+	Values[DataOffset + 5] = Y.Z;
 
-bool FFeatureVectorReader::IsValid() const
-{
-	return Layout && (Layout->NumFloats == Values.Num());
+	DataOffset += EncodeQuatCardinality;
 }
 
-bool FFeatureVectorReader::GetRotation(const FPoseSearchFeatureDesc& Feature, FQuat* OutRotation) const
+FQuat FFeatureVectorHelper::DecodeQuat(TArrayView<const float> Values, int32& DataOffset)
 {
-	check(Feature.Type == EPoseSearchFeatureType::Rotation);
-	const int32 FeatureIndex = IsValid() ? Layout->Features.Find(Feature) : -1;
-	check(FeatureIndex >= 0);
-	if (FeatureIndex >= 0)
-	{
-		const FPoseSearchFeatureDesc& FoundFeature = Layout->Features[FeatureIndex];
-		check(Feature.ValueOffset < 0 || Feature.ValueOffset == FoundFeature.ValueOffset);
-
-		FVector X;
-		FVector Y;
-
-		X.X = Values[FoundFeature.ValueOffset + 0];
-		X.Y = Values[FoundFeature.ValueOffset + 1];
-		X.Z = Values[FoundFeature.ValueOffset + 2];
-		Y.X = Values[FoundFeature.ValueOffset + 3];
-		Y.Y = Values[FoundFeature.ValueOffset + 4];
-		Y.Z = Values[FoundFeature.ValueOffset + 5];
+	const FVector X(Values[DataOffset + 0], Values[DataOffset + 1], Values[DataOffset + 2]);
+	const FVector Y(Values[DataOffset + 3], Values[DataOffset + 4], Values[DataOffset + 5]);
+	const FVector Z = FVector::CrossProduct(X, Y);
 
-		FVector Z = FVector::CrossProduct(X, Y);
+	FMatrix M(FMatrix::Identity);
+	M.SetColumn(0, X);
+	M.SetColumn(1, Y);
+	M.SetColumn(2, Z);
 
-		FMatrix M(FMatrix::Identity);
-		M.SetColumn(0, X);
-		M.SetColumn(1, Y);
-		M.SetColumn(2, Z);
-
-		*OutRotation = FQuat(M);
-		return true;
-	}
+	DataOffset += EncodeQuatCardinality;
+	return FQuat(M);
+}
 
-	*OutRotation = FQuat::Identity;
-	return false;
+void FFeatureVectorHelper::EncodeVector(TArrayView<float> Values, int32& DataOffset, const FVector& Vector)
+{
+	Values[DataOffset + 0] = Vector.X;
+	Values[DataOffset + 1] = Vector.Y;
+	Values[DataOffset + 2] = Vector.Z;
+	DataOffset += EncodeVectorCardinality;
 }
 
-bool FFeatureVectorReader::GetPhase(const FPoseSearchFeatureDesc& Feature, FVector2D* OutPhase) const
+FVector FFeatureVectorHelper::DecodeVector(TArrayView<const float> Values, int32& DataOffset)
 {
-	check(Feature.Cardinality == 2);
-	check(Feature.Type == EPoseSearchFeatureType::Phase);
-	const int32 FeatureIndex = IsValid() ? Layout->Features.Find(Feature) : -1;
-	check(FeatureIndex >= 0);
-	if (FeatureIndex >= 0)
-	{
-		const FPoseSearchFeatureDesc& FoundFeature = Layout->Features[FeatureIndex];
-		check(Feature.ValueOffset < 0 || Feature.ValueOffset == FoundFeature.ValueOffset);
+	const FVector Vector(Values[DataOffset + 0], Values[DataOffset + 1], Values[DataOffset + 2]);
+	DataOffset += EncodeVectorCardinality;
+	return Vector;
+}
 
-		OutPhase->X = Values[FoundFeature.ValueOffset + 0];
-		OutPhase->Y = Values[FoundFeature.ValueOffset + 1];
-		return true;
-	}
+void FFeatureVectorHelper::EncodeVector2D(TArrayView<float> Values, int32& DataOffset, const FVector2D& Vector2D)
+{
+	Values[DataOffset + 0] = Vector2D.X;
+	Values[DataOffset + 1] = Vector2D.Y;
+	DataOffset += EncodeVector2DCardinality;
+}
 
-	*OutPhase = FVector2D::ZeroVector;
-	return false;
+FVector2D FFeatureVectorHelper::DecodeVector2D(TArrayView<const float> Values, int32& DataOffset)
+{
+	const FVector2D Vector2D(Values[DataOffset + 0], Values[DataOffset + 1]);
+	DataOffset += EncodeVector2DCardinality;
+	return Vector2D;
 }
 
-bool FFeatureVectorReader::GetVector(const FPoseSearchFeatureDesc& Feature, FVector* OutVector) const
+void FFeatureVectorHelper::ComputeMeanDeviations(const Eigen::MatrixXd& CenteredPoseMatrix, Eigen::VectorXd& MeanDeviations, int32& DataOffset, int32 Cardinality)
 {
-	check(Feature.Cardinality == 3);
-	const int32 FeatureIndex = IsValid() ? Layout->Features.Find(Feature) : -1;
-	check(FeatureIndex >= 0);
-	if (FeatureIndex >= 0)
-	{
-		const FPoseSearchFeatureDesc& FoundFeature = Layout->Features[FeatureIndex];
-		check(Feature.ValueOffset < 0 || Feature.ValueOffset == FoundFeature.ValueOffset);
+	const int32 NumPoses = MeanDeviations.size();
 
-		FVector V;
-		V.X = Values[FoundFeature.ValueOffset + 0];
-		V.Y = Values[FoundFeature.ValueOffset + 1];
-		V.Z = Values[FoundFeature.ValueOffset + 2];
-		*OutVector = V;
-		return true;
-	}
+	// Construct a submatrix for the feature and find the average distance to the feature's centroid.
+	// Since we've already mean centered the data, the average distance to the centroid is simply the average norm.
+	const double FeatureMeanDeviation = CenteredPoseMatrix.block(DataOffset, 0, Cardinality, NumPoses).colwise().norm().mean();
 
-	*OutVector = FVector::ZeroVector;
-	return false;
-}
+	// Fill the feature's corresponding scaling axes with the average distance
+	// Avoid scaling by zero by leaving near-zero deviations as 1.0
+	MeanDeviations.segment(DataOffset, Cardinality).setConstant(FeatureMeanDeviation > KINDA_SMALL_NUMBER ? FeatureMeanDeviation : 1.f);
 
+	DataOffset += Cardinality;
+}
 
 //////////////////////////////////////////////////////////////////////////
 // FDebugDrawParams
-
 bool FDebugDrawParams::CanDraw() const
 {
 	if (!World)
@@ -2599,6 +2424,31 @@ bool FDebugDrawParams::CanDraw() const
 	return SearchIndex->IsValid() && !SearchIndex->IsEmpty();
 }
 
+FLinearColor FDebugDrawParams::GetColor(const UPoseSearchFeatureChannel* channel) const
+{
+	if (Color)
+	{
+		return *Color;
+	}
+
+	const UPoseSearchSchema* Schema = GetSchema();
+	if (!Schema || Schema->SchemaCardinality <= 0)
+	{
+		return FLinearColor::Red;
+	}
+
+	const float TotalData = Schema->SchemaCardinality;
+	const float ChannelData = channel->GetChannelDataOffset();
+	const float HalfData = TotalData * 0.5f;
+
+	const float Hue = ChannelData < HalfData
+		? FMath::GetMappedRangeValueUnclamped({ 0.f, HalfData }, FVector2f(60.f, 0.f), ChannelData)
+		: FMath::GetMappedRangeValueUnclamped({ HalfData, TotalData }, FVector2f(280.f, 220.f), ChannelData);
+
+	const FLinearColor ColorHSV(Hue, 1.f, 1.f);
+	return ColorHSV.HSVToLinearRGB();
+}
+
 const FPoseSearchIndex* FDebugDrawParams::GetSearchIndex() const
 {
 	if (Database)
@@ -3608,7 +3458,6 @@ public:
 		int32 NumIndexedPoses = 0;
 		TArray<float> FeatureVectorTable;
 		TArray<FPoseSearchPoseMetadata> PoseMetadata;
-		TSet<int32> InvalidChannels;
 		TBitArray<> AllFeaturesNotAdded;
 	} Output;
 
@@ -3640,7 +3489,6 @@ void FAssetIndexer::Reset()
 
 	Output.FeatureVectorTable.Reset(0);
 	Output.PoseMetadata.Reset(0);
-	Output.InvalidChannels.Reset();
 	Output.AllFeaturesNotAdded.Reset();
 }
 
@@ -3659,7 +3507,7 @@ void FAssetIndexer::Init(const FAssetIndexingContext& InIndexingContext)
 		FMath::Max(0, FMath::CeilToInt(IndexingContext.RequestedSamplingRange.Max * IndexingContext.Schema->SampleRate));
 	Output.NumIndexedPoses = Output.LastIndexedSample - Output.FirstIndexedSample + 1;
 	
-	Output.FeatureVectorTable.SetNumZeroed(IndexingContext.Schema->Layout.NumFloats * Output.NumIndexedPoses);
+	Output.FeatureVectorTable.SetNumZeroed(IndexingContext.Schema->SchemaCardinality * Output.NumIndexedPoses);
 
 	Output.PoseMetadata.SetNum(Output.NumIndexedPoses);
 
@@ -3691,40 +3539,16 @@ bool FAssetIndexer::Process()
 		Channel->IndexAsset(*this, AssetIndexingOutput);
 	}
 
-	// Verify all channels provided a complete feature vector for all poses
-	TBitArray<> PoseFeaturesNotAdded;
-	for (int32 VectorIdx = 0; VectorIdx != NumSamplesInRange; ++VectorIdx)
-	{
-		if (!FeatureVectorBuilders[VectorIdx].IsComplete())
-		{
-			PoseFeaturesNotAdded = FeatureVectorBuilders[VectorIdx].GetFeaturesAdded();
-			PoseFeaturesNotAdded.BitwiseNOT();
-			TBitArray<>::BitwiseOR(Output.AllFeaturesNotAdded, PoseFeaturesNotAdded, EBitwiseOperatorFlags::MaxSize);
-
-			for (TConstSetBitIterator<> FeatureIter(PoseFeaturesNotAdded); FeatureIter; ++FeatureIter)
-			{
-				int32 FeatureIdx = FeatureIter.GetIndex();
-				const FPoseSearchFeatureDesc& FeatureDesc = IndexingContext.Schema->Layout.Features[FeatureIdx];
-				Output.InvalidChannels.Add(FeatureDesc.ChannelIdx);
-			}
-		}
-	}
-
-	Output.InvalidChannels.Sort(TLess<>());
-
 	// Merge spans of feature vectors into contiguous buffer
 	for (int32 VectorIdx = 0; VectorIdx != NumSamplesInRange; ++VectorIdx)
 	{
-		if (FeatureVectorBuilders[VectorIdx].IsComplete())
-		{
-			const int32 SampleIdx = VectorIdx + IndexingContext.BeginSampleIdx;
-			const int32 PoseIdx = SampleIdx - Output.FirstIndexedSample;
-			const int32 FirstValueIdx = PoseIdx * IndexingContext.Schema->Layout.NumFloats;
-			TArrayView<float> WriteValues = MakeArrayView(&Output.FeatureVectorTable[FirstValueIdx], IndexingContext.Schema->Layout.NumFloats);
-			TArrayView<const float> ReadValues = FeatureVectorBuilders[VectorIdx].GetValues();
-			check(WriteValues.Num() == ReadValues.Num());
-			FMemory::Memcpy(WriteValues.GetData(), ReadValues.GetData(), WriteValues.Num() * WriteValues.GetTypeSize());
-		}
+		const int32 SampleIdx = VectorIdx + IndexingContext.BeginSampleIdx;
+		const int32 PoseIdx = SampleIdx - Output.FirstIndexedSample;
+		const int32 FirstValueIdx = PoseIdx * IndexingContext.Schema->SchemaCardinality;
+		TArrayView<float> WriteValues = MakeArrayView(&Output.FeatureVectorTable[FirstValueIdx], IndexingContext.Schema->SchemaCardinality);
+		TArrayView<const float> ReadValues = FeatureVectorBuilders[VectorIdx].GetValues();
+		check(WriteValues.Num() == ReadValues.Num());
+		FMemory::Memcpy(WriteValues.GetData(), ReadValues.GetData(), WriteValues.Num() * WriteValues.GetTypeSize());
 	}
 
 	// Generate pose metadata
@@ -3735,7 +3559,7 @@ bool FAssetIndexer::Process()
 		Output.PoseMetadata[PoseIdx] = Metadata;
 	}
 
-	return Output.InvalidChannels.IsEmpty();
+	return true;
 }
 
 const float FAssetIndexer::GetSampleTimeFromDistance(float SampleDistance) const
@@ -4018,45 +3842,30 @@ void FAssetIndexer::AddMetadata(int32 SampleIdx)
 //////////////////////////////////////////////////////////////////////////
 // PoseSearch API
 
-static void DrawFeatureVector(const FDebugDrawParams& DrawParams, const FFeatureVectorReader& Reader)
-{
-	const UPoseSearchSchema* Schema = DrawParams.GetSchema();
-
-	for (int32 ChannelIdx = 0; ChannelIdx != Schema->Channels.Num(); ++ChannelIdx)
-	{
-		if (DrawParams.ChannelMask & (1 << ChannelIdx))
-		{
-			Schema->Channels[ChannelIdx]->DebugDraw(DrawParams, Reader);
-		}
-	}
-}
-
 static void DrawFeatureVector(const FDebugDrawParams& DrawParams, TArrayView<const float> PoseVector)
 {
 	const UPoseSearchSchema* Schema = DrawParams.GetSchema();
 	check(Schema)
 
-	if (PoseVector.Num() != Schema->Layout.NumFloats)
+	if (PoseVector.Num() != Schema->SchemaCardinality)
 	{
 		return;
 	}
 
-	FFeatureVectorReader Reader;
-	Reader.Init(&Schema->Layout);
-	Reader.SetValues(PoseVector);
-	DrawFeatureVector(DrawParams, Reader);
+	for (int32 ChannelIdx = 0; ChannelIdx != Schema->Channels.Num(); ++ChannelIdx)
+	{
+		if (DrawParams.ChannelMask & (1 << ChannelIdx))
+		{
+			Schema->Channels[ChannelIdx]->DebugDraw(DrawParams, PoseVector);
+		}
+	}
 }
 
 static void DrawSearchIndex(const FDebugDrawParams& DrawParams)
 {
-	const UPoseSearchSchema* Schema = DrawParams.GetSchema();
 	const FPoseSearchIndex* SearchIndex = DrawParams.GetSearchIndex();
-	check(Schema);
 	check(SearchIndex);
 
-	FFeatureVectorReader Reader;
-	Reader.Init(&Schema->Layout);
-
 	const int32 LastPoseIdx = SearchIndex->NumPoses;
 
 	TArray<float> PoseVector;
@@ -4064,8 +3873,7 @@ static void DrawSearchIndex(const FDebugDrawParams& DrawParams)
 	{
 		PoseVector = SearchIndex->GetPoseValues(PoseIdx);
 		SearchIndex->InverseNormalize(PoseVector);
-		Reader.SetValues(PoseVector);
-		DrawFeatureVector(DrawParams, Reader);
+		DrawFeatureVector(DrawParams, PoseVector);
 	}
 }
 
@@ -4124,10 +3932,8 @@ static void PreprocessSearchIndexNone(FPoseSearchIndex* SearchIndex)
 	FPoseSearchIndexPreprocessInfo& Info = SearchIndex->PreprocessInfo;
 	Info.Reset();
 
-	const FPoseSearchFeatureVectorLayout& Layout = SearchIndex->Schema->Layout;
-
 	const int32 NumPoses = SearchIndex->NumPoses;
-	const int32 NumDimensions = Layout.NumFloats;
+	const int32 NumDimensions = SearchIndex->Schema->SchemaCardinality;
 
 	Info.NumDimensions = NumDimensions;
 	Info.TransformationMatrix.SetNumZeroed(NumDimensions * NumPoses);
@@ -4155,7 +3961,7 @@ static void PreprocessSearchIndexNone(FPoseSearchIndex* SearchIndex)
 	SampleMeanMap = VectorXf::Zero(NumDimensions);
 }
 
-static inline Eigen::VectorXd ComputeFeatureMeanDeviations(const Eigen::MatrixXd& CenteredPoseMatrix, const FPoseSearchFeatureVectorLayout& Layout)
+static inline Eigen::VectorXd ComputeChannelsMeanDeviations(const Eigen::MatrixXd& CenteredPoseMatrix, const UPoseSearchSchema* Schema)
 {
 	using namespace Eigen;
 
@@ -4164,18 +3970,11 @@ static inline Eigen::VectorXd ComputeFeatureMeanDeviations(const Eigen::MatrixXd
 
 	VectorXd MeanDeviations(NumDimensions);
 	MeanDeviations.setConstant(1.0);
-	for (const FPoseSearchFeatureDesc& Feature : Layout.Features)
-	{
-		// Construct a submatrix for the feature and find the average distance to the feature's centroid.
-		// Since we've already mean centered the data, the average distance to the centroid is simply the average norm.
-		double FeatureMeanDeviation = CenteredPoseMatrix.block(Feature.ValueOffset, 0, Feature.Cardinality, NumPoses).colwise().norm().mean();
 
-		// Fill the feature's corresponding scaling axes with the average distance
-		// Avoid scaling by zero by leaving near-zero deviations as 1.0
-		if (FeatureMeanDeviation > KINDA_SMALL_NUMBER)
-		{
-			MeanDeviations.segment(Feature.ValueOffset, Feature.Cardinality).setConstant(FeatureMeanDeviation);
-		}
+	for (int ChannelIdx = 0; ChannelIdx != Schema->Channels.Num(); ++ChannelIdx)
+	{
+		const UPoseSearchFeatureChannel* Channel = Schema->Channels[ChannelIdx].Get();
+		Channel->ComputeMeanDeviations(CenteredPoseMatrix, MeanDeviations);
 	}
 
 	return MeanDeviations;
@@ -4233,10 +4032,8 @@ static void PreprocessSearchIndexNormalize(FPoseSearchIndex* SearchIndex)
 	FPoseSearchIndexPreprocessInfo& Info = SearchIndex->PreprocessInfo;
 	Info.Reset();
 
-	const FPoseSearchFeatureVectorLayout& Layout = SearchIndex->Schema->Layout;
-
 	const int32 NumPoses = SearchIndex->NumPoses;
-	const int32 NumDimensions = Layout.NumFloats;
+	const int32 NumDimensions = SearchIndex->Schema->SchemaCardinality;
 
 	// Map input buffer
 	auto PoseMatrixSourceMap = Map<Matrix<float, Dynamic, Dynamic, RowMajor>>(
@@ -4259,7 +4056,7 @@ static void PreprocessSearchIndexNormalize(FPoseSearchIndex* SearchIndex)
 	PoseMatrix = PoseMatrix.colwise() - SampleMean;
 
 	// Compute per-feature average distances
-	VectorXd MeanDeviations = ComputeFeatureMeanDeviations(PoseMatrix, Layout);
+	VectorXd MeanDeviations = ComputeChannelsMeanDeviations(PoseMatrix, SearchIndex->Schema);
 
 	// Construct a scaling matrix that uniformly scales each feature by its average distance from the mean
 	MatrixXd ScalingMatrix = MeanDeviations.cwiseInverse().asDiagonal();
@@ -4391,10 +4188,8 @@ static void PreprocessSearchIndexSphere(FPoseSearchIndex* SearchIndex)
 	FPoseSearchIndexPreprocessInfo& Info = SearchIndex->PreprocessInfo;
 	Info.Reset();
 
-	const FPoseSearchFeatureVectorLayout& Layout = SearchIndex->Schema->Layout;
-
 	const int32 NumPoses = SearchIndex->NumPoses;
-	const int32 NumDimensions = Layout.NumFloats;
+	const int32 NumDimensions = SearchIndex->Schema->SchemaCardinality;
 
 	// Map input buffer
 	auto PoseMatrixSourceMap = Map<Matrix<float, Dynamic, Dynamic, RowMajor>>(
@@ -4417,7 +4212,7 @@ static void PreprocessSearchIndexSphere(FPoseSearchIndex* SearchIndex)
 	PoseMatrix = PoseMatrix.colwise() - SampleMean;
 
 	// Compute per-feature average distances
-	VectorXd MeanDeviations = ComputeFeatureMeanDeviations(PoseMatrix, Layout);
+	VectorXd MeanDeviations = ComputeChannelsMeanDeviations(PoseMatrix, SearchIndex->Schema);
 
 	// Rescale data by transforming it with the scaling matrix
 	// Now each feature has an average Euclidean length = 1.
@@ -4539,7 +4334,7 @@ static void PreprocessSearchIndex(FPoseSearchIndex* SearchIndex)
 
 static void PreprocessGroupSearchIndexWeights(FGroupSearchIndex& GroupSearchIndex, const UPoseSearchDatabase* Database)
 {
-	const int32 NumDimensions = Database->Schema->Layout.NumFloats;
+	const int32 NumDimensions = Database->Schema->SchemaCardinality;
 	GroupSearchIndex.Weights.Init(1.f, NumDimensions);
 
 	for (int ChannelIdx = 0; ChannelIdx != Database->Schema->Channels.Num(); ++ChannelIdx)
@@ -4561,7 +4356,7 @@ static void PreprocessGroupSearchIndexWeights(FGroupSearchIndex& GroupSearchInde
 static void PreprocessGroupSearchIndexPCAData(FGroupSearchIndex& GroupSearchIndex, const UPoseSearchDatabase* Database, const float* GroupValues, float* GroupPCAValues)
 {
 	// binding SearchIndex.Values and SearchIndex.PCAValues Eigen row major matrix maps
-	const int32 NumDimensions = Database->Schema->Layout.NumFloats;
+	const int32 NumDimensions = Database->Schema->SchemaCardinality;
 	const int32 NumGroupPoses = GroupSearchIndex.EndPoseIndex - GroupSearchIndex.StartPoseIndex;
 
 	const RowMajorVectorMapConst MapWeights(GroupSearchIndex.Weights.GetData(), 1, NumDimensions);
@@ -4676,7 +4471,7 @@ static void PreprocessGroupSearchIndexKDTree(FGroupSearchIndex& GroupSearchIndex
 		}
 
 		// testing the KDTree is returning the proper searches for all the original points transformed in pca space
-		const int32 NumDimensions = Database->Schema->Layout.NumFloats;
+		const int32 NumDimensions = Database->Schema->SchemaCardinality;
 		for (size_t PointIndex = 0; PointIndex < NumGroupPoses; ++PointIndex)
 		{
 			constexpr size_t NumResults = 10;
@@ -4714,7 +4509,7 @@ static void PreprocessGroupSearchIndex(FPoseSearchIndex& SearchIndex, const UPos
 	const uint32 NumberOfPrincipalComponents = Database->GetNumberOfPrincipalComponents();
 	if (NumberOfPrincipalComponents > 0)
 	{
-		const int32 NumDimensions = Database->Schema->Layout.NumFloats;
+		const int32 NumDimensions = Database->Schema->SchemaCardinality;
 
 		// preallocating the PCAValues for all the groups
 		SearchIndex.PCAValues.Reset();
@@ -4971,39 +4766,13 @@ void FDatabaseIndexingContext::PrepareIndexers()
 
 bool FDatabaseIndexingContext::IndexAssets()
 {
-	bool bSuccess = true;
-
-	// Index sequence data
+	// Index asset data
 	ParallelFor(
 		Indexers.Num(),
-		[this](int32 SequenceIdx) { Indexers[SequenceIdx].Process(); },
+		[this](int32 AssetIdx) { Indexers[AssetIdx].Process(); },
 		ParallelForFlags
 	);
-
-	// Validate indexing process
-	TSet<int32> AllInvalidChannels;
-	TBitArray<> AllInvalidFeatures;
-	for (FAssetIndexer& Indexer : Indexers)
-	{
-		AllInvalidChannels.Union(Indexer.Output.InvalidChannels);
-		TBitArray<>::BitwiseOR(AllInvalidFeatures, Indexer.Output.AllFeaturesNotAdded, EBitwiseOperatorFlags::MaxSize);
-	}
-
-	if (!AllInvalidChannels.IsEmpty())
-	{
-		bSuccess = false;
-	}
-
-	for (int32 InvalidChannelIdx : AllInvalidChannels)
-	{
-		const TObjectPtr<UPoseSearchFeatureChannel> ChannelPtr = Database->Schema->Channels[InvalidChannelIdx];
-		if (ChannelPtr)
-		{
-			UE_LOG(LogPoseSearch, Error, TEXT("Feature channel '%f' failed to index all features from all assets"), *ChannelPtr->GetName());
-		}
-	}
-
-	return bSuccess;
+	return true;
 }
 
 void FDatabaseIndexingContext::JoinIndex()
@@ -5048,7 +4817,7 @@ void FDatabaseIndexingContext::JoinIndex()
 		SearchIndex->Groups.Last().EndPoseIndex = TotalPoses;
 	}
 	
-	check(TotalFloats == TotalPoses * Database->Schema->Layout.NumFloats);
+	check(TotalFloats == TotalPoses * Database->Schema->SchemaCardinality);
 	
 	// Join animation data into a single search index
 	SearchIndex->Values.Reset(TotalFloats);
diff --git a/Engine/Plugins/Experimental/Animation/PoseSearch/Source/Runtime/Private/PoseSearchEigenHelper.h b/Engine/Plugins/Experimental/Animation/PoseSearch/Source/Runtime/Private/PoseSearchEigenHelper.h
--- a/Engine/Plugins/Experimental/Animation/PoseSearch/Source/Runtime/Private/PoseSearchEigenHelper.h
+++ b/Engine/Plugins/Experimental/Animation/PoseSearch/Source/Runtime/Private/PoseSearchEigenHelper.h
@@ -48,15 +48,33 @@ PRAGMA_DEFAULT_VISIBILITY_END
 #include <sstream>
 #endif
 
-namespace UE { namespace PoseSearch {
+namespace UE { namespace PoseSearch
+{
+	using RowMajorVector = Eigen::Matrix<float, 1, Eigen::Dynamic, Eigen::RowMajor>;
+	using RowMajorVectorMap = Eigen::Map<RowMajorVector, Eigen::RowMajor>;
+	using RowMajorVectorMapConst = Eigen::Map<const RowMajorVector, Eigen::RowMajor>;
+
+	using RowMajorMatrix = Eigen::Matrix<float, Eigen::Dynamic, Eigen::Dynamic, Eigen::RowMajor>;
+	using RowMajorMatrixMap = Eigen::Map<RowMajorMatrix, Eigen::RowMajor>;
+	using RowMajorMatrixMapConst = Eigen::Map<const RowMajorMatrix, Eigen::RowMajor>;
+
+	using ColMajorMatrix = Eigen::Matrix<float, Eigen::Dynamic, Eigen::Dynamic, Eigen::ColMajor>;
+	using ColMajorMatrixMap = Eigen::Map<ColMajorMatrix, Eigen::ColMajor>;
+	using ColMajorMatrixMapConst = Eigen::Map<const ColMajorMatrix, Eigen::ColMajor>;
+
+	struct ColMajorVectord : public Eigen::VectorXd
+	{
+		ColMajorVectord(int32 Size) : Eigen::VectorXd(Size) {}
+	};
+	struct ColMajorMatrixd : public Eigen::MatrixXd {};
 
 #if UE_POSE_SEARCH_EIGEN_DEBUG
-template<typename EigenDenseBaseDerivedType>
-FString EigenMatrixToString(const Eigen::DenseBase<EigenDenseBaseDerivedType>& Matrix){
-	std::stringstream StringStream;
-    StringStream << Matrix;
-	return StringStream.str().c_str();
-}
+	template<typename EigenDenseBaseDerivedType>
+	FString EigenMatrixToString(const Eigen::DenseBase<EigenDenseBaseDerivedType>& Matrix){
+		std::stringstream StringStream;
+		StringStream << Matrix;
+		return StringStream.str().c_str();
+	}
 #endif
 
 }}
\ No newline at end of file
diff --git a/Engine/Plugins/Experimental/Animation/PoseSearch/Source/Runtime/Private/PoseSearchFeatureChannels.cpp b/Engine/Plugins/Experimental/Animation/PoseSearch/Source/Runtime/Private/PoseSearchFeatureChannels.cpp
--- a/Engine/Plugins/Experimental/Animation/PoseSearch/Source/Runtime/Private/PoseSearchFeatureChannels.cpp
+++ b/Engine/Plugins/Experimental/Animation/PoseSearch/Source/Runtime/Private/PoseSearchFeatureChannels.cpp
@@ -6,6 +6,7 @@
 #include "Animation/AttributesRuntime.h"
 #include "Animation/MirrorDataTable.h"
 #include "DrawDebugHelpers.h"
+#include "PoseSearchEigenHelper.h"
 #include "UObject/ObjectSaveContext.h"
 
 namespace UE::PoseSearch {
@@ -25,25 +26,6 @@ static const FVector DrawDebugSampleLabelOffset = FVector(0.0f, 0.0f, -10.0f);
 
 constexpr bool UseCharacterSpaceVelocities = true;
 
-//////////////////////////////////////////////////////////////////////////
-// Drawing helpers
-
-static FLinearColor GetColorForFeature(const FPoseSearchFeatureDesc& Feature, const FPoseSearchFeatureVectorLayout* Layout)
-{
-	const float FeatureIdx = Layout->Features.IndexOfByKey(Feature);
-	const float FeatureCountIdx = Layout->Features.Num() - 1;
-	const float FeatureCountIdxHalf = FeatureCountIdx / 2.f;
-	check(FeatureIdx != (float)INDEX_NONE);
-
-	const float Hue = FeatureIdx < FeatureCountIdxHalf
-		? FMath::GetMappedRangeValueUnclamped({ 0.f, FeatureCountIdxHalf }, FVector2f(60.f, 0.f), FeatureIdx)
-		: FMath::GetMappedRangeValueUnclamped({ FeatureCountIdxHalf, FeatureCountIdx }, FVector2f(280.f, 220.f), FeatureIdx);
-
-	const FLinearColor ColorHSV(Hue, 1.f, 1.f);
-	return ColorHSV.HSVToLinearRGB();
-}
-
-
 struct LocalMinMax
 {
 	enum
@@ -74,6 +56,109 @@ T GetValueAtIndex(int32 Sample, const TArray<T>& Values)
 	return (Values[Num - 1] - Values[Num - 2]) * (Sample - (Num - 1)) + Values[Num - 1];
 }
 
+struct BoneTransformsCache
+{
+	BoneTransformsCache(const IAssetIndexer& InIndexer)
+		: Indexer(InIndexer)
+	{
+	}
+
+	FTransform Get(float SampleTime, float OriginTime, int8 SchemaBoneIdx, bool& Clamped)
+	{
+		// @todo: use an hashmap if we end up having too many entries
+		Entry* Entry = Entries.FindByPredicate([SampleTime, OriginTime](const BoneTransformsCache::Entry& Entry)
+			{
+				return Entry.SampleTime == SampleTime && Entry.OriginTime == OriginTime;
+			});
+
+		const FAssetIndexingContext& IndexingContext = Indexer.GetIndexingContext();
+		const FAssetSamplingContext* SamplingContext = IndexingContext.SamplingContext;
+
+		if (!Entry)
+		{
+			Entry = &Entries[Entries.AddDefaulted()];
+
+			Entry->SampleTime = SampleTime;
+			Entry->OriginTime = OriginTime;
+
+			Entry->Pose.SetBoneContainer(&SamplingContext->BoneContainer);
+			Entry->UnusedCurve.InitFrom(SamplingContext->BoneContainer);
+
+			IAssetIndexer::FSampleInfo Origin = Indexer.GetSampleInfo(OriginTime);
+			IAssetIndexer::FSampleInfo Sample = Indexer.GetSampleInfoRelative(SampleTime, Origin);
+
+			float CurrentTime = Sample.ClipTime;
+			float PreviousTime = CurrentTime - SamplingContext->FiniteDelta;
+
+			FDeltaTimeRecord DeltaTimeRecord;
+			DeltaTimeRecord.Set(PreviousTime, CurrentTime - PreviousTime);
+			FAnimExtractContext ExtractionCtx(CurrentTime, true, DeltaTimeRecord, Sample.Clip->IsLoopable());
+
+			Sample.Clip->ExtractPose(ExtractionCtx, Entry->AnimPoseData);
+
+			if (IndexingContext.bMirrored)
+			{
+				FAnimationRuntime::MirrorPose(
+					Entry->AnimPoseData.GetPose(),
+					IndexingContext.Schema->MirrorDataTable->MirrorAxis,
+					SamplingContext->CompactPoseMirrorBones,
+					SamplingContext->ComponentSpaceRefRotations
+				);
+				// Note curves and attributes are not used during the indexing process and therefore don't need to be mirrored
+			}
+
+			Entry->ComponentSpacePose.InitPose(Entry->Pose);
+			Entry->RootTransform = Sample.RootTransform;
+			Entry->Clamped = Sample.bClamped;
+		}
+
+		const FBoneReference& BoneReference = IndexingContext.Schema->BoneReferences[SchemaBoneIdx];
+		FCompactPoseBoneIndex CompactBoneIndex = SamplingContext->BoneContainer.MakeCompactPoseIndex(FMeshPoseBoneIndex(BoneReference.BoneIndex));
+
+		const FTransform BoneTransform = Entry->ComponentSpacePose.GetComponentSpaceTransform(CompactBoneIndex) * Indexer.MirrorTransform(Entry->RootTransform);
+		Clamped = Entry->Clamped;
+
+		return BoneTransform;
+	}
+
+	const UE::PoseSearch::IAssetIndexer& Indexer;
+
+	struct Entry
+	{
+		float SampleTime;
+		float OriginTime;
+		bool Clamped;
+
+		// @todo: minimize the Entry memory footprint
+		FTransform RootTransform;
+		FCompactPose Pose;
+		FCSPose<FCompactPose> ComponentSpacePose;
+		FBlendedCurve UnusedCurve;
+		UE::Anim::FStackAttributeContainer UnusedAtrribute;
+		FAnimationPoseData AnimPoseData = { Pose, UnusedCurve, UnusedAtrribute };
+	};
+
+	TArray<Entry> Entries;
+};
+
+static void CollectBonePositions(TArray<FVector>& BonePositions, BoneTransformsCache& BoneTransformsCache, int8 SchemaBoneIdx)
+{
+	const FAssetIndexingContext& IndexingContext = BoneTransformsCache.Indexer.GetIndexingContext();
+	const float FiniteDelta = IndexingContext.Schema->SamplingInterval;
+	const float SampleTimeStart = FMath::Min(IndexingContext.BeginSampleIdx * FiniteDelta, IndexingContext.MainSampler->GetPlayLength());
+	const int32 NumSamples = IndexingContext.EndSampleIdx - IndexingContext.BeginSampleIdx;
+
+	// collecting all the bone transforms
+	BonePositions.Reset();
+	BonePositions.AddDefaulted(NumSamples);
+	for (int32 SampleIdx = 0; SampleIdx != NumSamples; ++SampleIdx)
+	{
+		const float SampleTime = SampleTimeStart + SampleIdx * FiniteDelta;
+		bool Unused;
+		BonePositions[SampleIdx] = BoneTransformsCache.Get(SampleTime, SampleTimeStart, SchemaBoneIdx, Unused).GetTranslation();
+	}
+}
+
 static void CalculateSignal(const TArray<FVector>& BonePositions, TArray<float>& Signal, int32 offset = 1)
 {
 	Signal.Reset();
@@ -260,92 +345,6 @@ static void CalculatePhasesFromLocalMinMax(const TArray<LocalMinMax>& MinMax, TA
 	}
 }
 
-struct BoneTransformsCache
-{
-	BoneTransformsCache(const IAssetIndexer& InIndexer)
-		: Indexer(InIndexer)
-	{
-	}
-
-	FTransform Get(float SampleTime, float OriginTime, int8 SchemaBoneIdx, bool& Clamped)
-	{
-		// @todo: use an hashmap if we end up having too many entries
-		Entry* Entry = Entries.FindByPredicate([SampleTime, OriginTime](const BoneTransformsCache::Entry& Entry)
-		{
-			return Entry.SampleTime == SampleTime && Entry.OriginTime == OriginTime;
-		});
-
-		const FAssetIndexingContext& IndexingContext = Indexer.GetIndexingContext();
-		const FAssetSamplingContext* SamplingContext = IndexingContext.SamplingContext;
-
-		if (!Entry)
-		{
-			Entry = &Entries[Entries.AddDefaulted()];
-
-			Entry->SampleTime = SampleTime;
-			Entry->OriginTime = OriginTime;
-
-			Entry->Pose.SetBoneContainer(&SamplingContext->BoneContainer);
-			Entry->UnusedCurve.InitFrom(SamplingContext->BoneContainer);
-
-			IAssetIndexer::FSampleInfo Origin = Indexer.GetSampleInfo(OriginTime);
-			IAssetIndexer::FSampleInfo Sample = Indexer.GetSampleInfoRelative(SampleTime, Origin);
-
-			float CurrentTime = Sample.ClipTime;
-			float PreviousTime = CurrentTime - SamplingContext->FiniteDelta;
-
-			FDeltaTimeRecord DeltaTimeRecord;
-			DeltaTimeRecord.Set(PreviousTime, CurrentTime - PreviousTime);
-			FAnimExtractContext ExtractionCtx(CurrentTime, true, DeltaTimeRecord, Sample.Clip->IsLoopable());
-
-			Sample.Clip->ExtractPose(ExtractionCtx, Entry->AnimPoseData);
-
-			if (IndexingContext.bMirrored)
-			{
-				FAnimationRuntime::MirrorPose(
-					Entry->AnimPoseData.GetPose(),
-					IndexingContext.Schema->MirrorDataTable->MirrorAxis,
-					SamplingContext->CompactPoseMirrorBones,
-					SamplingContext->ComponentSpaceRefRotations
-				);
-				// Note curves and attributes are not used during the indexing process and therefore don't need to be mirrored
-			}
-
-			Entry->ComponentSpacePose.InitPose(Entry->Pose);
-			Entry->RootTransform = Sample.RootTransform;
-			Entry->Clamped = Sample.bClamped;
-		}
-
-		const FBoneReference& BoneReference = IndexingContext.Schema->BoneReferences[SchemaBoneIdx];
-		FCompactPoseBoneIndex CompactBoneIndex = SamplingContext->BoneContainer.MakeCompactPoseIndex(FMeshPoseBoneIndex(BoneReference.BoneIndex));
-
-		const FTransform BoneTransform = Entry->ComponentSpacePose.GetComponentSpaceTransform(CompactBoneIndex) * Indexer.MirrorTransform(Entry->RootTransform);
-		Clamped = Entry->Clamped;
-
-		return BoneTransform;
-	}
-
-	const UE::PoseSearch::IAssetIndexer& Indexer;
-
-	struct Entry
-	{
-		float SampleTime;
-		float OriginTime;
-		bool Clamped;
-
-		// @todo: minimize the Entry memory footprint
-		FTransform RootTransform;
-		FCompactPose Pose;
-		FCSPose<FCompactPose> ComponentSpacePose;
-		FBlendedCurve UnusedCurve;
-		UE::Anim::FStackAttributeContainer UnusedAtrribute;
-		FAnimationPoseData AnimPoseData = { Pose, UnusedCurve, UnusedAtrribute };
-	};
-
-	TArray<Entry> Entries;
-};
-
-
 } // namespace UE::PoseSearch
 
 //////////////////////////////////////////////////////////////////////////
@@ -369,6 +368,8 @@ void UPoseSearchFeatureChannel_Pose::PreSave(FObjectPreSaveContext ObjectSaveCon
 
 void UPoseSearchFeatureChannel_Pose::InitializeSchema(UE::PoseSearch::FSchemaInitializer& Initializer)
 {
+	using namespace UE::PoseSearch;
+
 	Super::InitializeSchema(Initializer);
 
 	int32 DataOffset = ChannelDataOffset;
@@ -379,37 +380,35 @@ void UPoseSearchFeatureChannel_Pose::InitializeSchema(UE::PoseSearch::FSchemaIni
 		{
 			for (int32 SubsampleIdx = 0; SubsampleIdx != SampleTimes.Num(); ++SubsampleIdx)
 			{
-				Initializer.AddFeatureDesc(FPoseSearchFeatureDesc::Construct(GetChannelIndex(), ChannelBoneIdx, SubsampleIdx, EPoseSearchFeatureType::Position, PositionCardinality, DataOffset));
-				DataOffset += PositionCardinality;
+				DataOffset += FFeatureVectorHelper::EncodeVectorCardinality;
 			}
 		}
 		if (SampledBone.bUseRotation)
 		{
 			for (int32 SubsampleIdx = 0; SubsampleIdx != SampleTimes.Num(); ++SubsampleIdx)
 			{
-				Initializer.AddFeatureDesc(FPoseSearchFeatureDesc::Construct(GetChannelIndex(), ChannelBoneIdx, SubsampleIdx, EPoseSearchFeatureType::Rotation, RotationCardinality, DataOffset));
-				DataOffset += RotationCardinality;
+				DataOffset += FFeatureVectorHelper::EncodeQuatCardinality;
 			}
 		}
 		if (SampledBone.bUseVelocity)
 		{
 			for (int32 SubsampleIdx = 0; SubsampleIdx != SampleTimes.Num(); ++SubsampleIdx)
 			{
-				Initializer.AddFeatureDesc(FPoseSearchFeatureDesc::Construct(GetChannelIndex(), ChannelBoneIdx, SubsampleIdx, EPoseSearchFeatureType::LinearVelocity, LinearVelocityCardinality, DataOffset));
-				DataOffset += LinearVelocityCardinality;
+				DataOffset += FFeatureVectorHelper::EncodeVectorCardinality;
 			}
 		}
 		if (SampledBone.bUsePhase)
 		{
 			for (int32 SubsampleIdx = 0; SubsampleIdx != SampleTimes.Num(); ++SubsampleIdx)
 			{
-				Initializer.AddFeatureDesc(FPoseSearchFeatureDesc::Construct(GetChannelIndex(), ChannelBoneIdx, SubsampleIdx, EPoseSearchFeatureType::Phase, PhaseCardinality, DataOffset));
-				DataOffset += PhaseCardinality;
+				DataOffset += FFeatureVectorHelper::EncodeVector2DCardinality;
 			}
 		}
 	}
 
-	ChannelCardinality = Initializer.GetCurrentCardinalityFrom(ChannelDataOffset);
+	ChannelCardinality = DataOffset - ChannelDataOffset;
+
+	Initializer.SetCurrentChannelDataOffset(DataOffset);
 
 	FeatureParams.Reset();
 	for (const FPoseSearchBone& Bone : SampledBones)
@@ -422,6 +421,8 @@ void UPoseSearchFeatureChannel_Pose::InitializeSchema(UE::PoseSearch::FSchemaIni
 
 void UPoseSearchFeatureChannel_Pose::FillWeights(TArray<float>& Weights) const
 {
+	using namespace UE::PoseSearch;
+
 	int32 DataOffset = ChannelDataOffset;
 
 	const int32 NumBones = SampledBones.Num();
@@ -432,32 +433,44 @@ void UPoseSearchFeatureChannel_Pose::FillWeights(TArray<float>& Weights) const
 		{
 			for (int32 SubsampleIdx = 0; SubsampleIdx != SampleTimes.Num(); ++SubsampleIdx)
 			{
-				Weights[DataOffset] = SampledBone.Weight;
-				DataOffset += PositionCardinality;
+				for (int32 i = 0; i != FFeatureVectorHelper::EncodeVectorCardinality; ++i)
+				{
+					Weights[DataOffset + i] = SampledBone.Weight;
+				}
+				DataOffset += FFeatureVectorHelper::EncodeVectorCardinality;
 			}
 		}
 		if (SampledBone.bUseRotation)
 		{
 			for (int32 SubsampleIdx = 0; SubsampleIdx != SampleTimes.Num(); ++SubsampleIdx)
 			{
-				Weights[DataOffset] = SampledBone.Weight;
-				DataOffset += RotationCardinality;
+				for (int32 i = 0; i != FFeatureVectorHelper::EncodeQuatCardinality; ++i)
+				{
+					Weights[DataOffset + i] = SampledBone.Weight;
+				}
+				DataOffset += FFeatureVectorHelper::EncodeQuatCardinality;
 			}
 		}
 		if (SampledBone.bUseVelocity)
 		{
 			for (int32 SubsampleIdx = 0; SubsampleIdx != SampleTimes.Num(); ++SubsampleIdx)
 			{
-				Weights[DataOffset] = SampledBone.Weight;
-				DataOffset += LinearVelocityCardinality;
+				for (int32 i = 0; i != FFeatureVectorHelper::EncodeVectorCardinality; ++i)
+				{
+					Weights[DataOffset + i] = SampledBone.Weight;
+				}
+				DataOffset += FFeatureVectorHelper::EncodeVectorCardinality;
 			}
 		}
 		if (SampledBone.bUsePhase)
 		{
 			for (int32 SubsampleIdx = 0; SubsampleIdx != SampleTimes.Num(); ++SubsampleIdx)
 			{
-				Weights[DataOffset] = SampledBone.Weight;
-				DataOffset += PhaseCardinality;
+				for (int32 i = 0; i != FFeatureVectorHelper::EncodeVector2DCardinality; ++i)
+				{
+					Weights[DataOffset + i] = SampledBone.Weight;
+				}
+				DataOffset += FFeatureVectorHelper::EncodeVector2DCardinality;
 			}
 		}
 	}
@@ -473,55 +486,37 @@ void UPoseSearchFeatureChannel_Pose::CalculatePhases(UE::PoseSearch::BoneTransfo
 	static float SmoothingWindowTime = 0.3f; // seconds
 
 	using namespace UE::PoseSearch;
-
-	const FAssetIndexingContext& IndexingContext = BoneTransformsCache.Indexer.GetIndexingContext();
-
-	const float SampleTimeStart = FMath::Min(IndexingContext.BeginSampleIdx * IndexingContext.Schema->SamplingInterval, IndexingContext.MainSampler->GetPlayLength());
-	const float FiniteDelta = IndexingContext.Schema->SamplingInterval;
-
-	const int32 NumSamples = IndexingContext.EndSampleIdx - IndexingContext.BeginSampleIdx;
-
-	TArray<TArray<FVector>> BonePositions;
-	BonePositions.AddDefaulted(SampledBones.Num());
-	for (int32 ChannelBoneIdx = 0; ChannelBoneIdx != SampledBones.Num(); ++ChannelBoneIdx)
-	{
-		BonePositions[ChannelBoneIdx].AddDefaulted(NumSamples);
-	}
-
-	// collecting all the bone transforms
-	IAssetIndexer::FSampleInfo Origin = BoneTransformsCache.Indexer.GetSampleInfo(SampleTimeStart);
-	for (int32 SampleIdx = 0; SampleIdx != NumSamples; ++SampleIdx)
-	{
-		const float SampleTime = SampleTimeStart + SampleIdx * FiniteDelta;
-		for (int32 ChannelBoneIdx = 0; ChannelBoneIdx != SampledBones.Num(); ++ChannelBoneIdx)
-		{
-			bool Unused;
-			BonePositions[ChannelBoneIdx][SampleIdx] = BoneTransformsCache.Get(SampleTime, SampleTimeStart, FeatureParams[ChannelBoneIdx].SchemaBoneIdx, Unused).GetTranslation();
-		}
-	}
-
+	
 	OutPhases.Reset();
 	OutPhases.AddDefaulted(SampledBones.Num());
-
-	const int32 BoneSamplingCentralDifferencesOffset = FMath::Max(FMath::CeilToInt(BoneSamplingCentralDifferencesTime / FiniteDelta), 1);
-	const int32 SmoothingWindowOffset = FMath::Max(FMath::CeilToInt(SmoothingWindowTime / FiniteDelta), 1);
 	
+	const float FiniteDelta = BoneTransformsCache.Indexer.GetIndexingContext().Schema->SamplingInterval;
+
 	TArray<float> Signal;
 	TArray<float> SmoothedSignal;
 	TArray<LocalMinMax> LocalMinMax;
+	TArray<FVector> BonePositions;
 	for (int32 ChannelBoneIdx = 0; ChannelBoneIdx != SampledBones.Num(); ++ChannelBoneIdx)
 	{
-		// @todo: have different way of calculating signals, for example: height of the bone transform, acceleration, etc?
-		CalculateSignal(BonePositions[ChannelBoneIdx], Signal, BoneSamplingCentralDifferencesOffset);
-		
-		SmoothSignal(Signal, SmoothedSignal, SmoothingWindowOffset);
+		const FPoseSearchBone& SampledBone = SampledBones[ChannelBoneIdx];
+		if (SampledBone.bUsePhase)
+		{
+			CollectBonePositions(BonePositions, BoneTransformsCache, FeatureParams[ChannelBoneIdx].SchemaBoneIdx);
 
-		FindLocalMinMax(SmoothedSignal, LocalMinMax);
-		ValidateLocalMinMax(LocalMinMax);
+			// @todo: have different way of calculating signals, for example: height of the bone transform, acceleration, etc?
+			const int32 BoneSamplingCentralDifferencesOffset = FMath::Max(FMath::CeilToInt(BoneSamplingCentralDifferencesTime / FiniteDelta), 1);
+			CalculateSignal(BonePositions, Signal, BoneSamplingCentralDifferencesOffset);
 
-		ExtrapolateLocalMinMaxBoundaries(LocalMinMax, SmoothedSignal);
-		ValidateLocalMinMax(LocalMinMax);
-		CalculatePhasesFromLocalMinMax(LocalMinMax, OutPhases[ChannelBoneIdx], SmoothedSignal.Num());
+			const int32 SmoothingWindowOffset = FMath::Max(FMath::CeilToInt(SmoothingWindowTime / FiniteDelta), 1);
+			SmoothSignal(Signal, SmoothedSignal, SmoothingWindowOffset);
+
+			FindLocalMinMax(SmoothedSignal, LocalMinMax);
+			ValidateLocalMinMax(LocalMinMax);
+
+			ExtrapolateLocalMinMaxBoundaries(LocalMinMax, SmoothedSignal);
+			ValidateLocalMinMax(LocalMinMax);
+			CalculatePhasesFromLocalMinMax(LocalMinMax, OutPhases[ChannelBoneIdx], SmoothedSignal.Num());
+		}
 	}
 }
 
@@ -546,6 +541,46 @@ void UPoseSearchFeatureChannel_Pose::IndexAsset(const UE::PoseSearch::IAssetInde
 	}
 }
 
+void UPoseSearchFeatureChannel_Pose::ComputeMeanDeviations(const Eigen::MatrixXd& CenteredPoseMatrix, Eigen::VectorXd& MeanDeviations) const
+{
+	using namespace UE::PoseSearch;
+
+	int32 DataOffset = ChannelDataOffset;
+
+	const int32 NumBones = SampledBones.Num();
+	for (int32 ChannelBoneIdx = 0; ChannelBoneIdx != NumBones; ++ChannelBoneIdx)
+	{
+		const FPoseSearchBone& SampledBone = SampledBones[ChannelBoneIdx];
+		if (SampledBone.bUsePosition)
+		{
+			for (int32 SubsampleIdx = 0; SubsampleIdx != SampleTimes.Num(); ++SubsampleIdx)
+			{
+				FFeatureVectorHelper::ComputeMeanDeviations(CenteredPoseMatrix, MeanDeviations, DataOffset, FFeatureVectorHelper::EncodeVectorCardinality);
+			}
+		}
+		if (SampledBone.bUseRotation)
+		{
+			for (int32 SubsampleIdx = 0; SubsampleIdx != SampleTimes.Num(); ++SubsampleIdx)
+			{
+				FFeatureVectorHelper::ComputeMeanDeviations(CenteredPoseMatrix, MeanDeviations, DataOffset, FFeatureVectorHelper::EncodeQuatCardinality);
+			}
+		}
+		if (SampledBone.bUseVelocity)
+		{
+			for (int32 SubsampleIdx = 0; SubsampleIdx != SampleTimes.Num(); ++SubsampleIdx)
+			{
+				FFeatureVectorHelper::ComputeMeanDeviations(CenteredPoseMatrix, MeanDeviations, DataOffset, FFeatureVectorHelper::EncodeVectorCardinality);
+			}
+		}
+		if (SampledBone.bUsePhase)
+		{
+			FFeatureVectorHelper::ComputeMeanDeviations(CenteredPoseMatrix, MeanDeviations, DataOffset, FFeatureVectorHelper::EncodeVector2DCardinality);
+		}
+	}
+
+	check(DataOffset == ChannelDataOffset + ChannelCardinality);
+}
+
 void UPoseSearchFeatureChannel_Pose::AddPoseFeatures(UE::PoseSearch::BoneTransformsCache& BoneTransformsCache, int32 SampleIdx, FPoseSearchFeatureVectorBuilder& FeatureVector, const TArray<TArray<FVector2D>>& Phases) const
 {
 	// This function samples the instantaneous pose at time t as well as the pose's velocity and acceleration at time t.
@@ -586,10 +621,7 @@ void UPoseSearchFeatureChannel_Pose::AddPoseFeatures(UE::PoseSearch::BoneTransfo
 
 				bool ClampedPresent;
 				const FTransform BoneTransformsPresent = BoneTransformsCache.Get(SubsampleTime, UseCharacterSpaceVelocities ? SubsampleTime : SampleTime, PoseFeatureInfo.SchemaBoneIdx, ClampedPresent);
-
-				const FPoseSearchFeatureDesc Feature = FPoseSearchFeatureDesc::Construct(GetChannelIndex(), ChannelBoneIdx, SubsampleIdx, EPoseSearchFeatureType::Position, PositionCardinality, DataOffset);
-				DataOffset += PositionCardinality;
-				FeatureVector.SetVector(Feature, BoneTransformsPresent.GetTranslation());
+				FFeatureVectorHelper::EncodeVector(FeatureVector.EditValues(), DataOffset, BoneTransformsPresent.GetTranslation());
 			}
 		}
 
@@ -601,10 +633,7 @@ void UPoseSearchFeatureChannel_Pose::AddPoseFeatures(UE::PoseSearch::BoneTransfo
 
 				bool ClampedPresent;
 				const FTransform BoneTransformsPresent = BoneTransformsCache.Get(SubsampleTime, UseCharacterSpaceVelocities ? SubsampleTime : SampleTime, PoseFeatureInfo.SchemaBoneIdx, ClampedPresent);
-				const FPoseSearchFeatureDesc Feature = FPoseSearchFeatureDesc::Construct(GetChannelIndex(), ChannelBoneIdx, SubsampleIdx, EPoseSearchFeatureType::Rotation, RotationCardinality, DataOffset);
-				DataOffset += RotationCardinality;
-
-				FeatureVector.SetRotation(Feature, BoneTransformsPresent.GetRotation());
+				FFeatureVectorHelper::EncodeQuat(FeatureVector.EditValues(), DataOffset, BoneTransformsPresent.GetRotation());
 			}
 		}
 
@@ -619,9 +648,6 @@ void UPoseSearchFeatureChannel_Pose::AddPoseFeatures(UE::PoseSearch::BoneTransfo
 				const FTransform BoneTransformsPresent = BoneTransformsCache.Get(SubsampleTime, UseCharacterSpaceVelocities ? SubsampleTime : SampleTime, PoseFeatureInfo.SchemaBoneIdx, ClampedPresent);
 				const FTransform BoneTransformsFuture = BoneTransformsCache.Get(SubsampleTime + SamplingContext->FiniteDelta, UseCharacterSpaceVelocities ? SubsampleTime + SamplingContext->FiniteDelta : SampleTime, PoseFeatureInfo.SchemaBoneIdx, ClampedFuture);
 
-				const FPoseSearchFeatureDesc Feature = FPoseSearchFeatureDesc::Construct(GetChannelIndex(), ChannelBoneIdx, SubsampleIdx, EPoseSearchFeatureType::LinearVelocity, LinearVelocityCardinality, DataOffset);
-				DataOffset += LinearVelocityCardinality;
-
 				// We can get a better finite difference if we ignore samples that have
 				// been clamped at either side of the clip. However, if the central sample 
 				// itself is clamped, or there are no samples that are clamped, we can just 
@@ -640,7 +666,7 @@ void UPoseSearchFeatureChannel_Pose::AddPoseFeatures(UE::PoseSearch::BoneTransfo
 					LinearVelocity = (BoneTransformsFuture.GetTranslation() - BoneTransformsPast.GetTranslation()) / (SamplingContext->FiniteDelta * 2.0f);
 				}
 
-				FeatureVector.SetVector(Feature, LinearVelocity);
+				FFeatureVectorHelper::EncodeVector(FeatureVector.EditValues(), DataOffset, LinearVelocity);
 			}
 		}
 
@@ -648,11 +674,8 @@ void UPoseSearchFeatureChannel_Pose::AddPoseFeatures(UE::PoseSearch::BoneTransfo
 		{
 			for (int32 SubsampleIdx = 0; SubsampleIdx != SampleTimes.Num(); ++SubsampleIdx)
 			{
-				const FPoseSearchFeatureDesc Feature = FPoseSearchFeatureDesc::Construct(GetChannelIndex(), ChannelBoneIdx, SubsampleIdx, EPoseSearchFeatureType::Phase, PhaseCardinality, DataOffset);
-				DataOffset += PhaseCardinality;
-
 				// @todo: support for SubsampleIdx
-				FeatureVector.SetPhase(Feature, Phases[ChannelBoneIdx][SampleIdx]);
+				FFeatureVectorHelper::EncodeVector2D(FeatureVector.EditValues(), DataOffset, Phases[ChannelBoneIdx][SampleIdx]);
 			}
 		}
 	}
@@ -682,6 +705,8 @@ bool UPoseSearchFeatureChannel_Pose::BuildQuery(
 	FPoseSearchContext& SearchContext,
 	FPoseSearchFeatureVectorBuilder& InOutQuery) const
 {
+	using namespace UE::PoseSearch;
+
 	const bool bSkip =
 		SearchContext.CurrentResult.IsValid() &&
 		SearchContext.CurrentResult.Database->Schema == InOutQuery.GetSchema();
@@ -756,11 +781,14 @@ bool UPoseSearchFeatureChannel_Pose::BuildQuery(
 			for (int32 SubsampleIdx = 0; SubsampleIdx != SampleTimes.Num(); ++SubsampleIdx)
 			{
 				const int32 CachedTransformsIndex = SubsampleIdx * SampledBones.Num() + SampledBoneIdx;
-				const FPoseSearchFeatureDesc Feature = FPoseSearchFeatureDesc::Construct(GetChannelIndex(), SampledBoneIdx, SubsampleIdx, EPoseSearchFeatureType::Position, PositionCardinality, DataOffset);
-				DataOffset += PositionCardinality;
 				if (CachedTransforms[CachedTransformsIndex].Valid)
 				{
-					InOutQuery.SetVector(Feature, CachedTransforms[CachedTransformsIndex].Current.GetTranslation());
+					FFeatureVectorHelper::EncodeVector(InOutQuery.EditValues(), DataOffset, CachedTransforms[CachedTransformsIndex].Current.GetTranslation());
+				}
+				else
+				{
+					// preserve the InOutQuery.EditValues() and increase the DataOffset
+					DataOffset += FFeatureVectorHelper::EncodeVectorCardinality;
 				}
 			}
 		}
@@ -770,12 +798,14 @@ bool UPoseSearchFeatureChannel_Pose::BuildQuery(
 			for (int32 SubsampleIdx = 0; SubsampleIdx != SampleTimes.Num(); ++SubsampleIdx)
 			{
 				const int32 CachedTransformsIndex = SubsampleIdx * SampledBones.Num() + SampledBoneIdx;
-				const FPoseSearchFeatureDesc Feature = FPoseSearchFeatureDesc::Construct(GetChannelIndex(), SampledBoneIdx, SubsampleIdx, EPoseSearchFeatureType::Rotation, RotationCardinality, DataOffset);
-				DataOffset += RotationCardinality;
-
 				if (CachedTransforms[CachedTransformsIndex].Valid)
 				{
-					InOutQuery.SetRotation(Feature, CachedTransforms[CachedTransformsIndex].Current.GetRotation());
+					FFeatureVectorHelper::EncodeQuat(InOutQuery.EditValues(), DataOffset, CachedTransforms[CachedTransformsIndex].Current.GetRotation());
+				}
+				else
+				{
+					// preserve the InOutQuery.EditValues() and increase the DataOffset
+					DataOffset += FFeatureVectorHelper::EncodeQuatCardinality;
 				}
 			}
 		}
@@ -785,13 +815,15 @@ bool UPoseSearchFeatureChannel_Pose::BuildQuery(
 			for (int32 SubsampleIdx = 0; SubsampleIdx != SampleTimes.Num(); ++SubsampleIdx)
 			{
 				const int32 CachedTransformsIndex = SubsampleIdx * SampledBones.Num() + SampledBoneIdx;
-				const FPoseSearchFeatureDesc Feature = FPoseSearchFeatureDesc::Construct(GetChannelIndex(), SampledBoneIdx, SubsampleIdx, EPoseSearchFeatureType::LinearVelocity, LinearVelocityCardinality, DataOffset);
-				DataOffset += LinearVelocityCardinality;
-
 				if (CachedTransforms[CachedTransformsIndex].Valid)
 				{
 					const FVector LinearVelocity = (CachedTransforms[CachedTransformsIndex].Current.GetTranslation() - CachedTransforms[CachedTransformsIndex].Previous.GetTranslation()) / SearchContext.History->GetSampleTimeInterval();
-					InOutQuery.SetVector(Feature, LinearVelocity);
+					FFeatureVectorHelper::EncodeVector(InOutQuery.EditValues(), DataOffset, LinearVelocity);
+				}
+				else
+				{
+					// preserve the InOutQuery.EditValues() and increase the DataOffset
+					DataOffset += FFeatureVectorHelper::EncodeVectorCardinality;
 				}
 			}
 		}
@@ -801,11 +833,10 @@ bool UPoseSearchFeatureChannel_Pose::BuildQuery(
 			for (int32 SubsampleIdx = 0; SubsampleIdx != SampleTimes.Num(); ++SubsampleIdx)
 			{
 				const int32 CachedTransformsIndex = SubsampleIdx * SampledBones.Num() + SampledBoneIdx;
-				const FPoseSearchFeatureDesc Feature = FPoseSearchFeatureDesc::Construct(GetChannelIndex(), SampledBoneIdx, SubsampleIdx, EPoseSearchFeatureType::Phase, PhaseCardinality, DataOffset);
-				DataOffset += PhaseCardinality;
+				DataOffset += FFeatureVectorHelper::EncodeVector2DCardinality;
 
 				// @todo: Support phase in BuildQuery
-				//InOutQuery.SetPhase(Feature, ???);
+				// FFeatureVectorHelper::EncodeVector2D(InOutQuery.EditValues(), DataOffset, ???);
 			}
 		}
 	}
@@ -813,7 +844,7 @@ bool UPoseSearchFeatureChannel_Pose::BuildQuery(
 	return true;
 }
 
-void UPoseSearchFeatureChannel_Pose::DebugDraw(const UE::PoseSearch::FDebugDrawParams& DrawParams, const  UE::PoseSearch::FFeatureVectorReader& Reader) const
+void UPoseSearchFeatureChannel_Pose::DebugDraw(const UE::PoseSearch::FDebugDrawParams& DrawParams, TArrayView<const float> PoseVector) const
 {
 	using namespace UE::PoseSearch;
 
@@ -843,14 +874,10 @@ void UPoseSearchFeatureChannel_Pose::DebugDraw(const UE::PoseSearch::FDebugDrawP
 		{
 			for (int32 SubsampleIdx = 0; SubsampleIdx != NumSubsamples; ++SubsampleIdx)
 			{
-				const FPoseSearchFeatureDesc Feature = FPoseSearchFeatureDesc::Construct(GetChannelIndex(), ChannelBoneIdx, SubsampleIdx, EPoseSearchFeatureType::Position, PositionCardinality, DataOffset);
-				DataOffset += PositionCardinality;
-
-				const bool Found = Reader.GetVector(Feature, &BonePos[SubsampleIdx]);
-				check(Found);
+				BonePos[SubsampleIdx] = FFeatureVectorHelper::DecodeVector(PoseVector, DataOffset);
 
-				FLinearColor LinearColor = DrawParams.Color ? *DrawParams.Color : GetColorForFeature(Feature, Reader.GetLayout());
-				FColor Color = LinearColor.ToFColor(true);
+				const FLinearColor LinearColor = DrawParams.GetColor(this);
+				const FColor Color = LinearColor.ToFColor(true);
 
 				BonePos[SubsampleIdx] = DrawParams.RootTransform.TransformPosition(BonePos[SubsampleIdx]);
 				if (EnumHasAnyFlags(DrawParams.Flags, EDebugDrawFlags::DrawFast | EDebugDrawFlags::DrawSearchIndex))
@@ -876,8 +903,6 @@ void UPoseSearchFeatureChannel_Pose::DebugDraw(const UE::PoseSearch::FDebugDrawP
 		{
 			for (int32 SubsampleIdx = 0; SubsampleIdx != NumSubsamples; ++SubsampleIdx)
 			{
-				DataOffset += PositionCardinality;
-
 				// @todo: initialize with the character position instead of FVector::Zero?
 				BonePos[SubsampleIdx] = DrawParams.Mesh != nullptr ? DrawParams.Mesh->GetSocketTransform(SampledBones[ChannelBoneIdx].Reference.BoneName).GetLocation() : FVector::Zero();
 			}
@@ -887,12 +912,8 @@ void UPoseSearchFeatureChannel_Pose::DebugDraw(const UE::PoseSearch::FDebugDrawP
 		{
 			for (int32 SubsampleIdx = 0; SubsampleIdx != NumSubsamples; ++SubsampleIdx)
 			{
-				const FPoseSearchFeatureDesc Feature = FPoseSearchFeatureDesc::Construct(GetChannelIndex(), ChannelBoneIdx, SubsampleIdx, EPoseSearchFeatureType::Rotation, RotationCardinality, DataOffset);
-				DataOffset += RotationCardinality;
+				const FQuat BoneRot = FFeatureVectorHelper::DecodeQuat(PoseVector, DataOffset);
 
-				FQuat BoneRot;
-				const bool Found = Reader.GetRotation(Feature, &BoneRot);
-				check(Found);
 				// @todo: debug draw rotation
 			}
 		}
@@ -901,15 +922,10 @@ void UPoseSearchFeatureChannel_Pose::DebugDraw(const UE::PoseSearch::FDebugDrawP
 		{
 			for (int32 SubsampleIdx = 0; SubsampleIdx != NumSubsamples; ++SubsampleIdx)
 			{
-				const FPoseSearchFeatureDesc Feature = FPoseSearchFeatureDesc::Construct(GetChannelIndex(), ChannelBoneIdx, SubsampleIdx, EPoseSearchFeatureType::LinearVelocity, LinearVelocityCardinality, DataOffset);
-				DataOffset += LinearVelocityCardinality;
+				FVector BoneVel = FFeatureVectorHelper::DecodeVector(PoseVector, DataOffset);
 
-				FVector BoneVel;
-				const bool Found = Reader.GetVector(Feature, &BoneVel);
-				check(Found);
-
-				FLinearColor LinearColor = DrawParams.Color ? *DrawParams.Color : GetColorForFeature(Feature, Reader.GetLayout());
-				FColor Color = LinearColor.ToFColor(true);
+				const FLinearColor LinearColor = DrawParams.GetColor(this);
+				const FColor Color = LinearColor.ToFColor(true);
 
 				BoneVel *= DrawDebugVelocityScale;
 				BoneVel = DrawParams.RootTransform.TransformVector(BoneVel);
@@ -943,15 +959,10 @@ void UPoseSearchFeatureChannel_Pose::DebugDraw(const UE::PoseSearch::FDebugDrawP
 		{
 			for (int32 SubsampleIdx = 0; SubsampleIdx != NumSubsamples; ++SubsampleIdx)
 			{
-				const FPoseSearchFeatureDesc Feature = FPoseSearchFeatureDesc::Construct(GetChannelIndex(), ChannelBoneIdx, SubsampleIdx, EPoseSearchFeatureType::Phase, PhaseCardinality, DataOffset);
-				DataOffset += PhaseCardinality;
-
-				FVector2D Phase;
-				const bool Found = Reader.GetPhase(Feature, &Phase);
-				check(Found);
+				const FVector2D Phase = FFeatureVectorHelper::DecodeVector2D(PoseVector, DataOffset);
 
-				FLinearColor LinearColor = DrawParams.Color ? *DrawParams.Color : GetColorForFeature(Feature, Reader.GetLayout());
-				FColor Color = LinearColor.ToFColor(true);
+				const FLinearColor LinearColor = DrawParams.GetColor(this);
+				const FColor Color = LinearColor.ToFColor(true);
 
 				static float ScaleFactor = 1.f;
 
@@ -969,6 +980,8 @@ void UPoseSearchFeatureChannel_Pose::DebugDraw(const UE::PoseSearch::FDebugDrawP
 			}
 		}
 	}
+
+	check(DataOffset == ChannelDataOffset + ChannelCardinality);
 }
 
 //////////////////////////////////////////////////////////////////////////
@@ -980,8 +993,10 @@ void UPoseSearchFeatureChannel_Trajectory::PreSave(FObjectPreSaveContext ObjectS
 	Super::PreSave(ObjectSaveContext);
 }
 
-void UPoseSearchFeatureChannel_Trajectory::InitializeSchema( UE::PoseSearch::FSchemaInitializer& Initializer)
+void UPoseSearchFeatureChannel_Trajectory::InitializeSchema(UE::PoseSearch::FSchemaInitializer& Initializer)
 {
+	using namespace UE::PoseSearch;
+
 	Super::InitializeSchema(Initializer);
 
 	int32 DataOffset = ChannelDataOffset;
@@ -989,8 +1004,7 @@ void UPoseSearchFeatureChannel_Trajectory::InitializeSchema( UE::PoseSearch::FSc
 	{
 		for (int32 SubsampleIdx = 0; SubsampleIdx != SampleOffsets.Num(); ++SubsampleIdx)
 		{
-			Initializer.AddFeatureDesc(FPoseSearchFeatureDesc::Construct(GetChannelIndex(), 0, SubsampleIdx, EPoseSearchFeatureType::Position, PositionCardinality, DataOffset));
-			DataOffset += PositionCardinality;
+			DataOffset += FFeatureVectorHelper::EncodeVectorCardinality;
 		}
 	}
 
@@ -998,8 +1012,7 @@ void UPoseSearchFeatureChannel_Trajectory::InitializeSchema( UE::PoseSearch::FSc
 	{
 		for (int32 SubsampleIdx = 0; SubsampleIdx != SampleOffsets.Num(); ++SubsampleIdx)
 		{
-			Initializer.AddFeatureDesc(FPoseSearchFeatureDesc::Construct(GetChannelIndex(), 0, SubsampleIdx, EPoseSearchFeatureType::LinearVelocity, LinearVelocityCardinality, DataOffset));
-			DataOffset += LinearVelocityCardinality;
+			DataOffset += FFeatureVectorHelper::EncodeVectorCardinality;
 		}
 	}
 
@@ -1007,12 +1020,13 @@ void UPoseSearchFeatureChannel_Trajectory::InitializeSchema( UE::PoseSearch::FSc
 	{
 		for (int32 SubsampleIdx = 0; SubsampleIdx != SampleOffsets.Num(); ++SubsampleIdx)
 		{
-			Initializer.AddFeatureDesc(FPoseSearchFeatureDesc::Construct(GetChannelIndex(), 0, SubsampleIdx, EPoseSearchFeatureType::ForwardVector, ForwardVectorCardinality, DataOffset));
-			DataOffset += ForwardVectorCardinality;
+			DataOffset += FFeatureVectorHelper::EncodeVectorCardinality;
 		}
 	}
 
-	ChannelCardinality = Initializer.GetCurrentCardinalityFrom(ChannelDataOffset);
+	ChannelCardinality = DataOffset - ChannelDataOffset;
+
+	Initializer.SetCurrentChannelDataOffset(DataOffset);
 }
 
 void UPoseSearchFeatureChannel_Trajectory::FillWeights(TArray<float>& Weights) const
@@ -1035,6 +1049,38 @@ void UPoseSearchFeatureChannel_Trajectory::IndexAsset(const UE::PoseSearch::IAss
 	}
 }
 
+void UPoseSearchFeatureChannel_Trajectory::ComputeMeanDeviations(const Eigen::MatrixXd& CenteredPoseMatrix, Eigen::VectorXd& MeanDeviations) const
+{
+	using namespace UE::PoseSearch;
+
+	int32 DataOffset = ChannelDataOffset;
+	if (bUsePositions)
+	{
+		for (int32 SubsampleIdx = 0; SubsampleIdx != SampleOffsets.Num(); ++SubsampleIdx)
+		{
+			FFeatureVectorHelper::ComputeMeanDeviations(CenteredPoseMatrix, MeanDeviations, DataOffset, FFeatureVectorHelper::EncodeVectorCardinality);
+		}
+	}
+
+	if (bUseLinearVelocities)
+	{
+		for (int32 SubsampleIdx = 0; SubsampleIdx != SampleOffsets.Num(); ++SubsampleIdx)
+		{
+			FFeatureVectorHelper::ComputeMeanDeviations(CenteredPoseMatrix, MeanDeviations, DataOffset, FFeatureVectorHelper::EncodeVectorCardinality);
+		}
+	}
+
+	if (bUseFacingDirections)
+	{
+		for (int32 SubsampleIdx = 0; SubsampleIdx != SampleOffsets.Num(); ++SubsampleIdx)
+		{
+			FFeatureVectorHelper::ComputeMeanDeviations(CenteredPoseMatrix, MeanDeviations, DataOffset, FFeatureVectorHelper::EncodeVectorCardinality);
+		}
+	}
+
+	check(DataOffset == ChannelDataOffset + ChannelCardinality);
+}
+
 float UPoseSearchFeatureChannel_Trajectory::GetSampleTime(const UE::PoseSearch::IAssetIndexer& Indexer, int32 SubsampleIdx, float SampleTime, float RootDistance) const
 {
 	switch (Domain)
@@ -1052,7 +1098,7 @@ float UPoseSearchFeatureChannel_Trajectory::GetSampleTime(const UE::PoseSearch::
 	return 0.0f;
 }
 
-void UPoseSearchFeatureChannel_Trajectory::IndexAssetPrivate(const UE::PoseSearch::IAssetIndexer& Indexer, int32 SampleIdx,  FPoseSearchFeatureVectorBuilder& FeatureVector) const
+void UPoseSearchFeatureChannel_Trajectory::IndexAssetPrivate(const UE::PoseSearch::IAssetIndexer& Indexer, int32 SampleIdx, FPoseSearchFeatureVectorBuilder& FeatureVector) const
 {
 	// This function samples the instantaneous trajectory at time t as well as the trajectory's velocity and acceleration at time t.
 	// Symmetric finite differences are used to approximate derivatives:
@@ -1075,9 +1121,7 @@ void UPoseSearchFeatureChannel_Trajectory::IndexAssetPrivate(const UE::PoseSearc
 		{
 			const float SubsampleTime = GetSampleTime(Indexer, SubsampleIdx, SampleTime, Origin.RootDistance);
 			const FSampleInfo SamplePresent = Indexer.GetSampleInfoRelative(SubsampleTime, Origin);
-			const FPoseSearchFeatureDesc Feature = FPoseSearchFeatureDesc::Construct(GetChannelIndex(), 0, SubsampleIdx, EPoseSearchFeatureType::Position, PositionCardinality, DataOffset);
-			DataOffset += PositionCardinality;
-			FeatureVector.SetVector(Feature, Indexer.MirrorTransform(SamplePresent.RootTransform).GetTranslation());
+			FFeatureVectorHelper::EncodeVector(FeatureVector.EditValues(), DataOffset, Indexer.MirrorTransform(SamplePresent.RootTransform).GetTranslation());
 		}
 	}
 
@@ -1116,9 +1160,7 @@ void UPoseSearchFeatureChannel_Trajectory::IndexAssetPrivate(const UE::PoseSearc
 				LinearVelocity = (MirroredRootFuture.GetTranslation() - MirroredRootPast.GetTranslation()) / (IndexingContext.SamplingContext->FiniteDelta * 2.0f);
 			}
 
-			const FPoseSearchFeatureDesc Feature = FPoseSearchFeatureDesc::Construct(GetChannelIndex(), 0, SubsampleIdx, EPoseSearchFeatureType::LinearVelocity, LinearVelocityCardinality, DataOffset);
-			DataOffset += LinearVelocityCardinality;
-			FeatureVector.SetVector(Feature, LinearVelocity);
+			FFeatureVectorHelper::EncodeVector(FeatureVector.EditValues(), DataOffset, LinearVelocity);
 		}
 	}
 
@@ -1128,9 +1170,7 @@ void UPoseSearchFeatureChannel_Trajectory::IndexAssetPrivate(const UE::PoseSearc
 		{
 			const float SubsampleTime = GetSampleTime(Indexer, SubsampleIdx, SampleTime, Origin.RootDistance);
 			const FSampleInfo SamplePresent = Indexer.GetSampleInfoRelative(SubsampleTime, Origin);
-			const FPoseSearchFeatureDesc Feature = FPoseSearchFeatureDesc::Construct(GetChannelIndex(), 0, SubsampleIdx, EPoseSearchFeatureType::ForwardVector, ForwardVectorCardinality, DataOffset);
-			DataOffset += ForwardVectorCardinality;
-			FeatureVector.SetVector(Feature, Indexer.MirrorTransform(SamplePresent.RootTransform).GetRotation().GetAxisY());
+			FFeatureVectorHelper::EncodeVector(FeatureVector.EditValues(), DataOffset, Indexer.MirrorTransform(SamplePresent.RootTransform).GetRotation().GetAxisY());
 		}
 	}
 }
@@ -1145,7 +1185,6 @@ FFloatRange UPoseSearchFeatureChannel_Trajectory::GetHorizonRange(EPoseSearchFea
 			Extent = FFloatRange::Inclusive(SampleOffsets[0], SampleOffsets.Last());
 		}
 	}
-
 	return Extent;
 }
 
@@ -1163,6 +1202,8 @@ bool UPoseSearchFeatureChannel_Trajectory::BuildQuery(
 	FPoseSearchContext& SearchContext,
 	FPoseSearchFeatureVectorBuilder& InOutQuery) const
 {
+	using namespace UE::PoseSearch;
+
 	if (!SearchContext.Trajectory)
 	{
 		return false;
@@ -1197,9 +1238,7 @@ bool UPoseSearchFeatureChannel_Trajectory::BuildQuery(
 	{
 		for (int32 Idx = 0, Num = SampleOffsets.Num(); Idx < Num; ++Idx)
 		{
-			const FPoseSearchFeatureDesc Feature = FPoseSearchFeatureDesc::Construct(GetChannelIndex(), 0, Idx, EPoseSearchFeatureType::Position, PositionCardinality, DataOffset);
-			DataOffset += PositionCardinality;
-			InOutQuery.SetVector(Feature, Samples[Idx].Transform.GetTranslation());
+			FFeatureVectorHelper::EncodeVector(InOutQuery.EditValues(), DataOffset, Samples[Idx].Transform.GetTranslation());
 		}
 	}
 
@@ -1207,9 +1246,7 @@ bool UPoseSearchFeatureChannel_Trajectory::BuildQuery(
 	{
 		for (int32 Idx = 0, Num = SampleOffsets.Num(); Idx < Num; ++Idx)
 		{
-			const FPoseSearchFeatureDesc Feature = FPoseSearchFeatureDesc::Construct(GetChannelIndex(), 0, Idx, EPoseSearchFeatureType::LinearVelocity, LinearVelocityCardinality, DataOffset);
-			DataOffset += LinearVelocityCardinality;
-			InOutQuery.SetVector(Feature, Samples[Idx].LinearVelocity);
+			FFeatureVectorHelper::EncodeVector(InOutQuery.EditValues(), DataOffset, Samples[Idx].LinearVelocity);
 		}
 	}
 
@@ -1217,16 +1254,14 @@ bool UPoseSearchFeatureChannel_Trajectory::BuildQuery(
 	{ 
 		for (int32 Idx = 0, Num = SampleOffsets.Num(); Idx < Num; ++Idx)
 		{
-			const FPoseSearchFeatureDesc Feature = FPoseSearchFeatureDesc::Construct(GetChannelIndex(), 0, Idx, EPoseSearchFeatureType::ForwardVector, ForwardVectorCardinality, DataOffset);
-			DataOffset += ForwardVectorCardinality;
-			InOutQuery.SetVector(Feature, Samples[Idx].Transform.GetRotation().GetAxisY());
+			FFeatureVectorHelper::EncodeVector(InOutQuery.EditValues(), DataOffset, Samples[Idx].Transform.GetRotation().GetAxisY());
 		}
 	}
 
 	return true;
 }
 
-void UPoseSearchFeatureChannel_Trajectory::DebugDraw(const UE::PoseSearch::FDebugDrawParams& DrawParams, const UE::PoseSearch::FFeatureVectorReader& Reader) const
+void UPoseSearchFeatureChannel_Trajectory::DebugDraw(const UE::PoseSearch::FDebugDrawParams& DrawParams, TArrayView<const float> PoseVector) const
 {
 	using namespace UE::PoseSearch;
 
@@ -1256,11 +1291,11 @@ void UPoseSearchFeatureChannel_Trajectory::DebugDraw(const UE::PoseSearch::FDebu
 	int32 DataOffset = ChannelDataOffset;
 	for (int32 SubsampleIdx = 0; SubsampleIdx != NumSubsamples; ++SubsampleIdx)
 	{
-		const FPoseSearchFeatureDesc Feature = FPoseSearchFeatureDesc::Construct(GetChannelIndex(), 0, SubsampleIdx, EPoseSearchFeatureType::Position, PositionCardinality, DataOffset);
-		DataOffset += PositionCardinality;
-		if (bUsePositions && Reader.GetVector(Feature, &TrajectoryPos[SubsampleIdx]))
+		if (bUsePositions)
 		{
-			FLinearColor LinearColor = DrawParams.Color ? *DrawParams.Color : GetColorForFeature(Feature, Reader.GetLayout());
+			TrajectoryPos[SubsampleIdx] = FFeatureVectorHelper::DecodeVector(PoseVector, DataOffset);
+
+			FLinearColor LinearColor = DrawParams.GetColor(this);
 			FLinearColor GradientColor = GetGradientColor(LinearColor, SubsampleIdx, NumSubsamples, DrawParams.Flags);
 			FColor Color = GradientColor.ToFColor(true);
 
@@ -1284,41 +1319,37 @@ void UPoseSearchFeatureChannel_Trajectory::DebugDraw(const UE::PoseSearch::FDebu
 	{
 		for (int32 SubsampleIdx = 0; SubsampleIdx != NumSubsamples; ++SubsampleIdx)
 		{
-			const FPoseSearchFeatureDesc Feature = FPoseSearchFeatureDesc::Construct(GetChannelIndex(), 0, SubsampleIdx, EPoseSearchFeatureType::LinearVelocity, LinearVelocityCardinality, DataOffset);
-			DataOffset += LinearVelocityCardinality;
-			FVector TrajectoryVel;
-			if (Reader.GetVector(Feature, &TrajectoryVel))
-			{
-				const FLinearColor LinearColor = DrawParams.Color ? *DrawParams.Color : GetColorForFeature(Feature, Reader.GetLayout());
-				const FLinearColor GradientColor = GetGradientColor(LinearColor, SubsampleIdx, NumSubsamples, DrawParams.Flags);
-				const FColor Color = GradientColor.ToFColor(true);
+			FVector TrajectoryVel = FFeatureVectorHelper::DecodeVector(PoseVector, DataOffset);
 
-				TrajectoryVel *= DrawDebugVelocityScale;
-				TrajectoryVel = DrawParams.RootTransform.TransformVector(TrajectoryVel);
-				const FVector TrajectoryVelDirection = TrajectoryVel.GetSafeNormal();
+			const FLinearColor LinearColor = DrawParams.GetColor(this);
+			const FLinearColor GradientColor = GetGradientColor(LinearColor, SubsampleIdx, NumSubsamples, DrawParams.Flags);
+			const FColor Color = GradientColor.ToFColor(true);
 
-				if (EnumHasAnyFlags(DrawParams.Flags, EDebugDrawFlags::DrawSearchIndex))
-				{
-					DrawDebugPoint(DrawParams.World, TrajectoryVel, DrawParams.PointSize, Color, bPersistent, DrawParams.DefaultLifeTime, DepthPriority);
-				}
-				else
-				{
-					const float AdjustedThickness =
-						EnumHasAnyFlags(DrawParams.Flags, EDebugDrawFlags::DrawFast) ?
-						0.0f : DrawDebugLineThickness;
+			TrajectoryVel *= DrawDebugVelocityScale;
+			TrajectoryVel = DrawParams.RootTransform.TransformVector(TrajectoryVel);
+			const FVector TrajectoryVelDirection = TrajectoryVel.GetSafeNormal();
 
-					DrawDebugDirectionalArrow(
-						DrawParams.World,
-						TrajectoryPos[SubsampleIdx] + TrajectoryVelDirection * DrawDebugSphereSize,
-						TrajectoryPos[SubsampleIdx] + TrajectoryVel,
-						DrawDebugArrowSize,
-						Color,
-						bPersistent,
-						LifeTime,
-						DepthPriority,
-						AdjustedThickness
-					);
-				}
+			if (EnumHasAnyFlags(DrawParams.Flags, EDebugDrawFlags::DrawSearchIndex))
+			{
+				DrawDebugPoint(DrawParams.World, TrajectoryVel, DrawParams.PointSize, Color, bPersistent, DrawParams.DefaultLifeTime, DepthPriority);
+			}
+			else
+			{
+				const float AdjustedThickness =
+					EnumHasAnyFlags(DrawParams.Flags, EDebugDrawFlags::DrawFast) ?
+					0.0f : DrawDebugLineThickness;
+
+				DrawDebugDirectionalArrow(
+					DrawParams.World,
+					TrajectoryPos[SubsampleIdx] + TrajectoryVelDirection * DrawDebugSphereSize,
+					TrajectoryPos[SubsampleIdx] + TrajectoryVel,
+					DrawDebugArrowSize,
+					Color,
+					bPersistent,
+					LifeTime,
+					DepthPriority,
+					AdjustedThickness
+				);
 			}
 		}
 	}
@@ -1327,39 +1358,35 @@ void UPoseSearchFeatureChannel_Trajectory::DebugDraw(const UE::PoseSearch::FDebu
 	{
 		for (int32 SubsampleIdx = 0; SubsampleIdx != NumSubsamples; ++SubsampleIdx)
 		{
-			const FPoseSearchFeatureDesc Feature = FPoseSearchFeatureDesc::Construct(GetChannelIndex(), 0, SubsampleIdx, EPoseSearchFeatureType::ForwardVector, ForwardVectorCardinality, DataOffset);
-			DataOffset += ForwardVectorCardinality;
-			FVector TrajectoryForward;
-			if (Reader.GetVector(Feature, &TrajectoryForward))
-			{
-				const FLinearColor LinearColor = DrawParams.Color ? *DrawParams.Color : GetColorForFeature(Feature, Reader.GetLayout());
-				const FLinearColor GradientColor = GetGradientColor(LinearColor, SubsampleIdx, NumSubsamples, DrawParams.Flags);
-				const FColor Color = GradientColor.ToFColor(true);
+			FVector TrajectoryForward = FFeatureVectorHelper::DecodeVector(PoseVector, DataOffset);
 
-				TrajectoryForward = DrawParams.RootTransform.TransformVector(TrajectoryForward);
+			const FLinearColor LinearColor = DrawParams.GetColor(this);
+			const FLinearColor GradientColor = GetGradientColor(LinearColor, SubsampleIdx, NumSubsamples, DrawParams.Flags);
+			const FColor Color = GradientColor.ToFColor(true);
 
-				if (EnumHasAnyFlags(DrawParams.Flags, EDebugDrawFlags::DrawSearchIndex))
-				{
-					DrawDebugPoint(DrawParams.World, TrajectoryForward, DrawParams.PointSize, Color, bPersistent, DrawParams.DefaultLifeTime, DepthPriority);
-				}
-				else
-				{
-					const float AdjustedThickness =
-						EnumHasAnyFlags(DrawParams.Flags, EDebugDrawFlags::DrawFast) ?
-						0.0f : DrawDebugLineThickness;
+			TrajectoryForward = DrawParams.RootTransform.TransformVector(TrajectoryForward);
 
-					DrawDebugDirectionalArrow(
-						DrawParams.World,
-						TrajectoryPos[SubsampleIdx] + TrajectoryForward * DrawDebugSphereSize,
-						TrajectoryPos[SubsampleIdx] + TrajectoryForward * DrawDebugSphereSize * 2.0f,
-						DrawDebugArrowSize,
-						Color,
-						bPersistent,
-						LifeTime,
-						DepthPriority,
-						AdjustedThickness
-					);
-				}
+			if (EnumHasAnyFlags(DrawParams.Flags, EDebugDrawFlags::DrawSearchIndex))
+			{
+				DrawDebugPoint(DrawParams.World, TrajectoryForward, DrawParams.PointSize, Color, bPersistent, DrawParams.DefaultLifeTime, DepthPriority);
+			}
+			else
+			{
+				const float AdjustedThickness =
+					EnumHasAnyFlags(DrawParams.Flags, EDebugDrawFlags::DrawFast) ?
+					0.0f : DrawDebugLineThickness;
+
+				DrawDebugDirectionalArrow(
+					DrawParams.World,
+					TrajectoryPos[SubsampleIdx] + TrajectoryForward * DrawDebugSphereSize,
+					TrajectoryPos[SubsampleIdx] + TrajectoryForward * DrawDebugSphereSize * 2.0f,
+					DrawDebugArrowSize,
+					Color,
+					bPersistent,
+					LifeTime,
+					DepthPriority,
+					AdjustedThickness
+				);
 			}
 		}
 	}
@@ -1368,8 +1395,7 @@ void UPoseSearchFeatureChannel_Trajectory::DebugDraw(const UE::PoseSearch::FDebu
 	{
 		for (int32 SubsampleIdx = 0; SubsampleIdx != NumSubsamples; ++SubsampleIdx)
 		{
-			const FPoseSearchFeatureDesc Feature = FPoseSearchFeatureDesc::Construct(GetChannelIndex(), 0, SubsampleIdx, EPoseSearchFeatureType::Position, PositionCardinality, -1);
-			const FLinearColor LinearColor = DrawParams.Color ? *DrawParams.Color : GetColorForFeature(Feature, Reader.GetLayout());
+			const FLinearColor LinearColor = DrawParams.GetColor(this);
 			const FLinearColor GradientColor = GetGradientColor(LinearColor, SubsampleIdx, NumSubsamples, DrawParams.Flags);
 			const FColor Color = GradientColor.ToFColor(true);
 
@@ -1393,4 +1419,6 @@ void UPoseSearchFeatureChannel_Trajectory::DebugDraw(const UE::PoseSearch::FDebu
 				DrawDebugSampleLabelFontScale);
 		}
 	}
+
+	check(DataOffset == ChannelDataOffset + ChannelCardinality);
 }
\ No newline at end of file
diff --git a/Engine/Plugins/Experimental/Animation/PoseSearch/Source/Runtime/Public/PoseSearch/PoseSearch.h b/Engine/Plugins/Experimental/Animation/PoseSearch/Source/Runtime/Public/PoseSearch/PoseSearch.h
--- a/Engine/Plugins/Experimental/Animation/PoseSearch/Source/Runtime/Public/PoseSearch/PoseSearch.h
+++ b/Engine/Plugins/Experimental/Animation/PoseSearch/Source/Runtime/Public/PoseSearch/PoseSearch.h
@@ -43,31 +43,22 @@ namespace UE::PoseSearch {
 class FPoseHistory;
 struct FPoseSearchDatabaseAsyncCacheTask;
 struct FDebugDrawParams;
-class FFeatureVectorReader;
 struct FSchemaInitializer;
 struct FQueryBuildingContext;
 
 } // namespace UE::PoseSearch
 
-
+namespace Eigen {
+	template<typename _Scalar, int _Rows, int _Cols, int _Options, int _MaxRows, int _MaxCols>
+	class Matrix;
+	using MatrixXd = Matrix<double, -1, -1, 0, -1, -1>;
+	using VectorXd = Matrix<double, -1, 1, 0, -1, 1>;
+}
 
 //////////////////////////////////////////////////////////////////////////
 // Constants
 
-UENUM()
-enum class EPoseSearchFeatureType : int8
-{
-	Position,	
-	Rotation,
-	LinearVelocity,
-	AngularVelocity,
-	ForwardVector,
-	Phase,
-
-	Num UMETA(Hidden),
-	Invalid = Num UMETA(Hidden)
-};
-
+// @todo: move this enum in PoseSearchFeatureChannel.h since only used by UPoseSearchFeatureChannel_Trajectory
 UENUM()
 enum class EPoseSearchFeatureDomain : int32
 {
@@ -206,79 +197,6 @@ struct POSESEARCH_API FPoseSearchBone
 	float Weight = 1.f;
 };
 
-
-
-//////////////////////////////////////////////////////////////////////////
-// Feature descriptions and vector layout
-
-/** Describes each feature of a vector, including data type, sampling options, and buffer offset. */
-USTRUCT()
-struct POSESEARCH_API FPoseSearchFeatureDesc
-{
-	GENERATED_BODY()
-
-	// Index into UPoseSearchSchema::Channels
-	UPROPERTY()
-	int8 ChannelIdx = -1;
-
-	// Optional feature identifier within a channel
-	UPROPERTY()
-	int8 ChannelFeatureId = 0;
-
-	// Index into channel's sample offsets, if any
-	UPROPERTY()
-	int8 SubsampleIdx = 0;
-
-	// Value type of the feature
-	UPROPERTY()
-	EPoseSearchFeatureType Type = EPoseSearchFeatureType::Invalid;
-
-	// Set via FSchemaInitializer::AddFeatureDesc() and ignored by operator== or else Schema->Layout.Features.Find(Feature); called by FPoseSearchFeatureVectorBuilder will fail
-	UPROPERTY()
-	int16 Cardinality = 0;
-
-	// Set via FSchemaInitializer::AddFeatureDesc() and ignored by operator==
-	UPROPERTY()
-	int16 ValueOffset = 0;
-
-	bool operator==(const FPoseSearchFeatureDesc& Other) const;
-
-	bool IsValid() const { return Type != EPoseSearchFeatureType::Invalid; }
-
-	static FPoseSearchFeatureDesc Construct(int8 ChannelIdx, int8 ChannelFeatureId, int8 SubsampleIdx, EPoseSearchFeatureType Type, int16 Cardinality, int16 ValueOffset)
-	{
-		FPoseSearchFeatureDesc PoseSearchFeatureDesc;
-		PoseSearchFeatureDesc.ChannelIdx = ChannelIdx;
-		PoseSearchFeatureDesc.ChannelFeatureId = ChannelFeatureId;
-		PoseSearchFeatureDesc.SubsampleIdx = SubsampleIdx;
-		PoseSearchFeatureDesc.Type = Type;
-		PoseSearchFeatureDesc.Cardinality = Cardinality;
-		PoseSearchFeatureDesc.ValueOffset = ValueOffset;
-		return PoseSearchFeatureDesc;
-	}
-};
-
-
-/**
-* Explicit description of a pose feature vector.
-* Determined by options set in a UPoseSearchSchema and owned by the schema.
-* See UPoseSearchSchema::GenerateLayout().
-*/
-USTRUCT()
-struct POSESEARCH_API FPoseSearchFeatureVectorLayout
-{
-	GENERATED_BODY()
-
-	UPROPERTY()
-	TArray<FPoseSearchFeatureDesc> Features;
-
-	UPROPERTY()
-	int32 NumFloats = 0;
-
-	bool IsValid(int32 ChannelCount) const;
-};
-
-
 //////////////////////////////////////////////////////////////////////////
 // Asset sampling and indexing
 
@@ -420,6 +338,8 @@ public:
 	// Called at database build time to populate pose vectors with this channel's data
 	virtual void IndexAsset(const UE::PoseSearch::IAssetIndexer& Indexer, UE::PoseSearch::FAssetIndexingOutput& IndexingOutput) const PURE_VIRTUAL(UPoseSearchFeatureChannel::IndexAsset, );
 
+	virtual void ComputeMeanDeviations(const Eigen::MatrixXd& CenteredPoseMatrix, Eigen::VectorXd& MeanDeviations) const;
+
 	// Return this channel's range of sampling offsets in the requested sampling domain.
 	// Returns empty range if the channel has no horizon in the requested domain.
 	virtual FFloatRange GetHorizonRange(EPoseSearchFeatureDomain Domain) const PURE_VIRTUAL(UPoseSearchFeatureChannel::GetHorizonRange, return FFloatRange::Empty(); );
@@ -434,7 +354,7 @@ public:
 		PURE_VIRTUAL(UPoseSearchFeatureChannel::BuildQuery, return false;);
 
 	// Draw this channel's data for the given pose vector
-	virtual void DebugDraw(const UE::PoseSearch::FDebugDrawParams& DrawParams, const UE::PoseSearch::FFeatureVectorReader& Reader) const PURE_VIRTUAL(UPoseSearchFeatureChannel::DebugDraw, );
+	virtual void DebugDraw(const UE::PoseSearch::FDebugDrawParams& DrawParams, TArrayView<const float> PoseVector) const PURE_VIRTUAL(UPoseSearchFeatureChannel::DebugDraw, );
 
 private:
 	// IBoneReferenceSkeletonProvider interface
@@ -449,11 +369,9 @@ private:
 	int32 ChannelIdx = -1;
 
 protected:
-	// WIP: mimicking the FPoseSearchFeatureDesc to be able to deprecate it in favor of UPoseSearchFeatureChannel later on 
 	UPROPERTY()
 	int32 ChannelDataOffset = -1;
 
-	// WIP: mimicking the FPoseSearchFeatureDesc to be able to deprecate it in favor of UPoseSearchFeatureChannel later on 
 	UPROPERTY()
 	int32 ChannelCardinality = -1;
 };
@@ -473,9 +391,8 @@ public:
 	// Gets the index into the schema's channel array for the channel currently being initialized
 	int32 GetCurrentChannelIdx() const { return CurrentChannelIdx; }
 
-	int32 GetCurrentCardinalityFrom(int32 ChannelDataOffset) const { check(CurrentChannelDataOffset >= ChannelDataOffset);  return CurrentChannelDataOffset - ChannelDataOffset; }
 	int32 GetCurrentChannelDataOffset() const { return CurrentChannelDataOffset; }
-	int32 AddFeatureDesc(const FPoseSearchFeatureDesc& FeatureDesc);
+	void SetCurrentChannelDataOffset(int32 DataOffset) { CurrentChannelDataOffset = DataOffset; }
 
 private:
 	friend class ::UPoseSearchSchema;
@@ -484,7 +401,6 @@ private:
 	int32 CurrentChannelDataOffset = 0;
 	
 	TArray<FBoneReference> BoneReferences;
-	TArray<FPoseSearchFeatureDesc> Features;
 };
 
 } // namespace UE::PoseSearch
@@ -552,7 +468,7 @@ public:
 	float SamplingInterval = 1.0f / DefaultSampleRate;
 
 	UPROPERTY()
-	FPoseSearchFeatureVectorLayout Layout;
+	int32 SchemaCardinality = 0;
 
 	UPROPERTY()
 	TArray<FBoneReference> BoneReferences;
@@ -940,22 +856,16 @@ public:
 
 	const UPoseSearchSchema* GetSchema() const { return Schema.Get(); }
 
+	TArray<float>& EditValues() { return Values; }
 	TArrayView<const float> GetValues() const { return Values; }
 	TArrayView<const float> GetNormalizedValues() const { return ValuesNormalized; }
 
-	void SetRotation(const FPoseSearchFeatureDesc& Feature, const FQuat& Rotation);
-	void SetVector(const FPoseSearchFeatureDesc& Feature, const FVector& Vector);
-	void SetPhase(const FPoseSearchFeatureDesc& Feature, const FVector2D& Phase);
-
 	void CopyFromSearchIndex(const FPoseSearchIndex& SearchIndex, int32 PoseIdx);
 
 	bool IsInitialized() const;
 	bool IsInitializedForSchema(const UPoseSearchSchema* Schema) const;
-	bool IsComplete() const;
 	bool IsCompatible(const FPoseSearchFeatureVectorBuilder& OtherBuilder) const;
 
-	const TBitArray<>& GetFeaturesAdded() const;
-
 	void Normalize(const FPoseSearchIndex& ForSearchIndex);
 
 private:
@@ -964,8 +874,6 @@ private:
 
 	TArray<float> Values;
 	TArray<float> ValuesNormalized;
-	TBitArray<> FeaturesAdded;
-	int32 NumFeaturesAdded = 0;
 };
 
 
@@ -1042,6 +950,7 @@ struct POSESEARCH_API FDebugDrawParams
 	TWeakObjectPtr<const USkinnedMeshComponent> Mesh = nullptr;
 
 	bool CanDraw() const;
+	FLinearColor GetColor(const UPoseSearchFeatureChannel* channel) const;
 	const FPoseSearchIndex* GetSearchIndex() const;
 	const UPoseSearchSchema* GetSchema() const;
 };
@@ -1317,23 +1226,24 @@ public: // UObject
 
 namespace UE::PoseSearch {
 
-/** Helper object for extracting features from a float buffer according to the feature vector layout. */
-class POSESEARCH_API FFeatureVectorReader
+/** Helper class for extracting and encoding features into a float buffer */
+class POSESEARCH_API FFeatureVectorHelper
 {
 public:
-	void Init(const FPoseSearchFeatureVectorLayout* Layout);
-	void SetValues(TArrayView<const float> Values);
-	bool IsValid() const;
+	enum { EncodeQuatCardinality = 6 };
+	static void EncodeQuat(TArrayView<float> Values, int32& DataOffset, const FQuat& Quat);
+	static FQuat DecodeQuat(TArrayView<const float> Values, int32& DataOffset);
 
-	bool GetRotation(const FPoseSearchFeatureDesc& Feature, FQuat* OutRotation) const;
-	bool GetVector(const FPoseSearchFeatureDesc& Feature, FVector* OutVector) const;
-	bool GetPhase(const FPoseSearchFeatureDesc& Feature, FVector2D* OutPhase) const;
+	enum { EncodeVectorCardinality = 3 };
+	static void EncodeVector(TArrayView<float> Values, int32& DataOffset, const FVector& Vector);
+	static FVector DecodeVector(TArrayView<const float> Values, int32& DataOffset);
 
-	const FPoseSearchFeatureVectorLayout* GetLayout() const { return Layout; }
+	enum { EncodeVector2DCardinality = 2 };
+	static void EncodeVector2D(TArrayView<float> Values, int32& DataOffset, const FVector2D& Vector2D);
+	static FVector2D DecodeVector2D(TArrayView<const float> Values, int32& DataOffset);
 
-private:
-	const FPoseSearchFeatureVectorLayout* Layout = nullptr;
-	TArrayView<const float> Values;
+	// populates MeanDeviations[DataOffset] ... MeanDeviations[DataOffset + Cardinality] with a single value the mean deviation calcualted from a cenetered matrix
+	static void ComputeMeanDeviations(const Eigen::MatrixXd& CenteredPoseMatrix, Eigen::VectorXd& MeanDeviations, int32& DataOffset, int32 Cardinality);
 };
 
 /**
diff --git a/Engine/Plugins/Experimental/Animation/PoseSearch/Source/Runtime/Public/PoseSearch/PoseSearchFeatureChannels.h b/Engine/Plugins/Experimental/Animation/PoseSearch/Source/Runtime/Public/PoseSearch/PoseSearchFeatureChannels.h
--- a/Engine/Plugins/Experimental/Animation/PoseSearch/Source/Runtime/Public/PoseSearch/PoseSearchFeatureChannels.h
+++ b/Engine/Plugins/Experimental/Animation/PoseSearch/Source/Runtime/Public/PoseSearch/PoseSearchFeatureChannels.h
@@ -49,21 +49,17 @@ public:
 	virtual void InitializeSchema(UE::PoseSearch::FSchemaInitializer& Initializer) override;
 	virtual void FillWeights(TArray<float>& Weights) const override;
 	virtual void IndexAsset(const UE::PoseSearch::IAssetIndexer& Indexer, UE::PoseSearch::FAssetIndexingOutput& IndexingOutput) const override;
+	virtual void ComputeMeanDeviations(const Eigen::MatrixXd& CenteredPoseMatrix, Eigen::VectorXd& MeanDeviations) const override;
 	virtual FFloatRange GetHorizonRange(EPoseSearchFeatureDomain Domain) const override;
 	virtual void GenerateDDCKey(FBlake3& InOutKeyHasher) const override;
 	virtual bool BuildQuery(
 		FPoseSearchContext& SearchContext,
 		FPoseSearchFeatureVectorBuilder& InOutQuery) const override;
-	virtual void DebugDraw(const UE::PoseSearch::FDebugDrawParams& DrawParams, const UE::PoseSearch::FFeatureVectorReader& Reader) const override;
+	virtual void DebugDraw(const UE::PoseSearch::FDebugDrawParams& DrawParams, TArrayView<const float> PoseVector) const override;
 
 protected:
 	void AddPoseFeatures(UE::PoseSearch::BoneTransformsCache& BoneTransformsCache, int32 SampleIdx, FPoseSearchFeatureVectorBuilder& FeatureVector, const TArray<TArray<FVector2D>>& Phases) const;
 	void CalculatePhases(UE::PoseSearch::BoneTransformsCache& BoneTransformsCache, UE::PoseSearch::FAssetIndexingOutput& IndexingOutput, TArray<TArray<FVector2D>>& OutPhases) const;
-
-	enum { PositionCardinality = 3 };
-	enum { RotationCardinality = 6 };
-	enum { LinearVelocityCardinality = 3 };
-	enum { PhaseCardinality = 2 };
 };
 
 
@@ -105,18 +101,15 @@ public:
 	virtual void InitializeSchema(UE::PoseSearch::FSchemaInitializer& Initializer) override;
 	virtual void FillWeights(TArray<float>& Weights) const override;
 	virtual void IndexAsset(const UE::PoseSearch::IAssetIndexer& Indexer, UE::PoseSearch::FAssetIndexingOutput& IndexingOutput) const override;
+	virtual void ComputeMeanDeviations(const Eigen::MatrixXd& CenteredPoseMatrix, Eigen::VectorXd& MeanDeviations) const override;
 	virtual FFloatRange GetHorizonRange(EPoseSearchFeatureDomain Domain) const override;
 	virtual void GenerateDDCKey(FBlake3& InOutKeyHasher) const override;
 	virtual bool BuildQuery(
 		FPoseSearchContext& SearchContext,
 		FPoseSearchFeatureVectorBuilder& InOutQuery) const override;
-	virtual void DebugDraw(const UE::PoseSearch::FDebugDrawParams& DrawParams, const UE::PoseSearch::FFeatureVectorReader& Reader) const override;
+	virtual void DebugDraw(const UE::PoseSearch::FDebugDrawParams& DrawParams, TArrayView<const float> PoseVector) const override;
 
 protected:
 	virtual void IndexAssetPrivate(const UE::PoseSearch::IAssetIndexer& Indexer, int32 SampleIdx, FPoseSearchFeatureVectorBuilder& FeatureVector) const;
 	float GetSampleTime(const UE::PoseSearch::IAssetIndexer& Indexer, int32 SubsampleIdx, float SampleTime, float RootDistance) const;
-
-	enum { PositionCardinality = 3 };
-	enum { LinearVelocityCardinality = 3 };
-	enum { ForwardVectorCardinality = 3 };
 };
\ No newline at end of file
