motion matching - PopulateNonSelectableIdx fix: disallow searching the continuity pose from the pose database

diff --git a/Engine/Plugins/Experimental/Animation/PoseSearch/Source/Runtime/Private/PoseSearch.cpp b/Engine/Plugins/Experimental/Animation/PoseSearch/Source/Runtime/Private/PoseSearch.cpp
--- a/Engine/Plugins/Experimental/Animation/PoseSearch/Source/Runtime/Private/PoseSearch.cpp
+++ b/Engine/Plugins/Experimental/Animation/PoseSearch/Source/Runtime/Private/PoseSearch.cpp
@@ -211,7 +211,7 @@ static int32 PopulateNonSelectableIdx(size_t* NonSelectableIdx, int32 NonSelecta
 			}
 		}
 
-		for (int32 i = 1; i <= PoseJumpIndexThreshold; ++i)
+		for (int32 i = 0; i <= PoseJumpIndexThreshold; ++i)
 		{
 			const int32 PoseIdx = SearchContext.CurrentResult.PoseIdx + i;
 			const float DeltaTime = i * Database->Schema->SamplingInterval;
