motion matching - stability improvements when editing motion matching assets while PIE is running

diff --git a/Engine/Plugins/Experimental/Animation/MotionTrajectory/Source/MotionTrajectory/Private/MotionTrajectoryCharacterMovement.cpp b/Engine/Plugins/Experimental/Animation/MotionTrajectory/Source/MotionTrajectory/Private/MotionTrajectoryCharacterMovement.cpp
--- a/Engine/Plugins/Experimental/Animation/MotionTrajectory/Source/MotionTrajectory/Private/MotionTrajectoryCharacterMovement.cpp
+++ b/Engine/Plugins/Experimental/Animation/MotionTrajectory/Source/MotionTrajectory/Private/MotionTrajectoryCharacterMovement.cpp
@@ -25,7 +25,6 @@ FTrajectorySample UCharacterMovementTrajectoryComponent::CalcWorldSpacePresentTr
 		if (MovementComponent->MovementMode == EMovementMode::MOVE_Walking)
 		{
 			FTransform ComponentWorldTransform = Pawn->GetActorTransform();
-			FVector MovementComponentPosition = MovementComponent->GetLastUpdateLocation();
 
 			ReturnValue.Transform = ComponentWorldTransform;
 			ReturnValue.LinearVelocity = MovementComponent->Velocity;
diff --git a/Engine/Plugins/Experimental/Animation/PoseSearch/Source/Editor/Private/PoseSearchDebugger.cpp b/Engine/Plugins/Experimental/Animation/PoseSearch/Source/Editor/Private/PoseSearchDebugger.cpp
--- a/Engine/Plugins/Experimental/Animation/PoseSearch/Source/Editor/Private/PoseSearchDebugger.cpp
+++ b/Engine/Plugins/Experimental/Animation/PoseSearch/Source/Editor/Private/PoseSearchDebugger.cpp
@@ -860,61 +860,68 @@ void SDebuggerDatabaseView::FilterDatabaseRows()
 
 void SDebuggerDatabaseView::CreateRows(const UPoseSearchDatabase& Database)
 {
-	const int32 NumPoses = Database.GetSearchIndex()->NumPoses;
-	UnfilteredDatabaseRows.Reset(NumPoses);
+	const bool bValidIndex = Database.GetSearchIndex() && Database.GetSearchIndex()->IsValid();
+	if (bValidIndex)
+	{
+		const int32 NumPoses = Database.GetSearchIndex()->NumPoses;
+		UnfilteredDatabaseRows.Reset(NumPoses);
 
-	RowsSourceDatabase = &Database;
+		RowsSourceDatabase = &Database;
 
-	// Build database rows
-	for(const FPoseSearchIndexAsset& SearchIndexAsset : Database.GetSearchIndex()->Assets)
-	{
-		if (SearchIndexAsset.Type == ESearchIndexAssetType::Sequence)
+		// Build database rows
+		for (const FPoseSearchIndexAsset& SearchIndexAsset : Database.GetSearchIndex()->Assets)
 		{
-			const FPoseSearchDatabaseSequence& DbSequence = Database.GetSequenceSourceAsset(&SearchIndexAsset);
-			const int32 LastPoseIdx = SearchIndexAsset.FirstPoseIdx + SearchIndexAsset.NumPoses;
-			for (int32 PoseIdx = SearchIndexAsset.FirstPoseIdx; PoseIdx != LastPoseIdx; ++PoseIdx)
+			if (SearchIndexAsset.Type == ESearchIndexAssetType::Sequence)
 			{
-				float Time = Database.GetAssetTime(PoseIdx);
-
-				TSharedRef<FDebuggerDatabaseRowData> Row = UnfilteredDatabaseRows.Add_GetRef(MakeShared<FDebuggerDatabaseRowData>());
-				Row->PoseIdx = PoseIdx;
-				Row->AssetType = ESearchIndexAssetType::Sequence;
-				Row->AssetName = DbSequence.Sequence->GetName();
-				Row->AssetPath = DbSequence.Sequence->GetPathName();
-				Row->DbAssetIdx = SearchIndexAsset.SourceAssetIdx;
-				Row->AssetTime = Time;
-				Row->AnimFrame = DbSequence.Sequence->GetFrameAtTime(Time);
-				Row->bMirrored = SearchIndexAsset.bMirrored;
-				Row->bLooping = DbSequence.Sequence->bLoop;
-				Row->BlendParameters = FVector::Zero();
+				const FPoseSearchDatabaseSequence& DbSequence = Database.GetSequenceSourceAsset(&SearchIndexAsset);
+				const int32 LastPoseIdx = SearchIndexAsset.FirstPoseIdx + SearchIndexAsset.NumPoses;
+				for (int32 PoseIdx = SearchIndexAsset.FirstPoseIdx; PoseIdx != LastPoseIdx; ++PoseIdx)
+				{
+					float Time = Database.GetAssetTime(PoseIdx);
+
+					TSharedRef<FDebuggerDatabaseRowData> Row = UnfilteredDatabaseRows.Add_GetRef(MakeShared<FDebuggerDatabaseRowData>());
+					Row->PoseIdx = PoseIdx;
+					Row->AssetType = ESearchIndexAssetType::Sequence;
+					Row->AssetName = DbSequence.Sequence->GetName();
+					Row->AssetPath = DbSequence.Sequence->GetPathName();
+					Row->DbAssetIdx = SearchIndexAsset.SourceAssetIdx;
+					Row->AssetTime = Time;
+					Row->AnimFrame = DbSequence.Sequence->GetFrameAtTime(Time);
+					Row->bMirrored = SearchIndexAsset.bMirrored;
+					Row->bLooping = DbSequence.Sequence->bLoop;
+					Row->BlendParameters = FVector::Zero();
+				}
 			}
-		}
-		else if (SearchIndexAsset.Type == ESearchIndexAssetType::BlendSpace)
-		{
-			const FPoseSearchDatabaseBlendSpace& DbBlendSpace = Database.GetBlendSpaceSourceAsset(&SearchIndexAsset);
-			const int32 LastPoseIdx = SearchIndexAsset.FirstPoseIdx + SearchIndexAsset.NumPoses;
-			for (int32 PoseIdx = SearchIndexAsset.FirstPoseIdx; PoseIdx != LastPoseIdx; ++PoseIdx)
+			else if (SearchIndexAsset.Type == ESearchIndexAssetType::BlendSpace)
 			{
-				float Time = Database.GetAssetTime(PoseIdx);
-
-				TSharedRef<FDebuggerDatabaseRowData> Row = UnfilteredDatabaseRows.Add_GetRef(MakeShared<FDebuggerDatabaseRowData>());
-				Row->PoseIdx = PoseIdx;
-				Row->AssetType = ESearchIndexAssetType::BlendSpace;
-				Row->AssetName = DbBlendSpace.BlendSpace->GetName();
-				Row->AssetPath = DbBlendSpace.BlendSpace->GetPathName();
-				Row->DbAssetIdx = SearchIndexAsset.SourceAssetIdx;
-				Row->AssetTime = Time;
-				Row->AnimFrame = 0; // There is no frame index associated with a blendspace
-				Row->bMirrored = SearchIndexAsset.bMirrored;
-				Row->bLooping = DbBlendSpace.BlendSpace->bLoop;
-				Row->BlendParameters = SearchIndexAsset.BlendParameters;
+				const FPoseSearchDatabaseBlendSpace& DbBlendSpace = Database.GetBlendSpaceSourceAsset(&SearchIndexAsset);
+				const int32 LastPoseIdx = SearchIndexAsset.FirstPoseIdx + SearchIndexAsset.NumPoses;
+				for (int32 PoseIdx = SearchIndexAsset.FirstPoseIdx; PoseIdx != LastPoseIdx; ++PoseIdx)
+				{
+					float Time = Database.GetAssetTime(PoseIdx);
+
+					TSharedRef<FDebuggerDatabaseRowData> Row = UnfilteredDatabaseRows.Add_GetRef(MakeShared<FDebuggerDatabaseRowData>());
+					Row->PoseIdx = PoseIdx;
+					Row->AssetType = ESearchIndexAssetType::BlendSpace;
+					Row->AssetName = DbBlendSpace.BlendSpace->GetName();
+					Row->AssetPath = DbBlendSpace.BlendSpace->GetPathName();
+					Row->DbAssetIdx = SearchIndexAsset.SourceAssetIdx;
+					Row->AssetTime = Time;
+					Row->AnimFrame = 0; // There is no frame index associated with a blendspace
+					Row->bMirrored = SearchIndexAsset.bMirrored;
+					Row->bLooping = DbBlendSpace.BlendSpace->bLoop;
+					Row->BlendParameters = SearchIndexAsset.BlendParameters;
+				}
+			}
+			else
+			{
+				checkNoEntry();
 			}
 		}
-		else
-		{
-			checkNoEntry();
-		}
-
+	}
+	else
+	{
+		UnfilteredDatabaseRows.Reset();
 	}
 
 	ActiveView.Rows.Reset();
@@ -926,12 +933,13 @@ void SDebuggerDatabaseView::CreateRows(const UPoseSearchDatabase& Database)
 
 void SDebuggerDatabaseView::UpdateRows(const FTraceMotionMatchingStateMessage& State, const UPoseSearchDatabase& Database)
 {
-	const bool bNewDatabase = RowsSourceDatabase != &Database;
-
+	const bool bValidIndex = Database.GetSearchIndex() && Database.GetSearchIndex()->IsValid();
+	const bool bNewDatabase = RowsSourceDatabase != &Database || !bValidIndex;
 	if (bNewDatabase || UnfilteredDatabaseRows.IsEmpty())
 	{
 		CreateRows(Database);
 	}
+
 	check(ActiveView.Rows.Num() == 1);
 	check(ContinuingPoseView.Rows.Num() == 1);
 
@@ -950,23 +958,27 @@ void SDebuggerDatabaseView::UpdateRows(const FTraceMotionMatchingStateMessage& S
 		{
 			const int32 PoseIdx = Row->PoseIdx;
 
-			Database.ComparePoses(
-				SearchContext, 
-				PoseIdx, 
-				State.QueryVectorNormalized, 
-				Row->PoseCostDetails);
-
-			// If we are on the active pose for the frame
-			if (PoseIdx == State.DbPoseIdx)
+			// @todo: invalidate rows in case Database changed while PIE is active with pose debugger. for now we just avoid crashes
+			if (PoseIdx < Database.GetSearchIndex()->NumPoses)
 			{
-				*ActiveView.Rows[0] = Row.Get();
-				IsActiveViewRowInitialized = true;
-			}
+				Database.ComparePoses(
+					SearchContext,
+					PoseIdx,
+					State.QueryVectorNormalized,
+					Row->PoseCostDetails);
+
+				// If we are on the active pose for the frame
+				if (PoseIdx == State.DbPoseIdx)
+				{
+					*ActiveView.Rows[0] = Row.Get();
+					IsActiveViewRowInitialized = true;
+				}
 
-			if (PoseIdx == State.ContinuingPoseIdx)
-			{
-				*ContinuingPoseView.Rows[0] = Row.Get();
-				IsContinuingPoseViewRowInitialized = true;
+				if (PoseIdx == State.ContinuingPoseIdx)
+				{
+					*ContinuingPoseView.Rows[0] = Row.Get();
+					IsContinuingPoseViewRowInitialized = true;
+				}
 			}
 		}
 
@@ -1454,45 +1466,51 @@ void SDebuggerDetailsView::UpdateReflection(const FTraceMotionMatchingStateMessa
 	Reflection->ElapsedPoseJumpTime = State.ElapsedPoseJumpTime;
 	Reflection->bFollowUpAnimation = EnumHasAnyFlags(State.Flags, FTraceMotionMatchingState::EFlags::FollowupAnimation);
 
-	Reflection->AssetPlayerAssetName = FString();
-	if (const FPoseSearchIndexAsset* IndexAsset = Database.GetSearchIndex()->FindAssetForPose(State.DbPoseIdx))
+	const bool bValidIndex = Database.GetSearchIndex() && Database.GetSearchIndex()->IsValid();
+	const bool bValidState = bValidIndex && State.DbPoseIdx < Database.GetSearchIndex()->NumPoses;
+	// @todo: cache the derived data key in the state to check its validity
+	if (bValidState)
 	{
-		Reflection->AssetPlayerAssetName = Database.GetSourceAssetName(IndexAsset);
-	}
+		Reflection->AssetPlayerAssetName = FString();
+		if (const FPoseSearchIndexAsset* IndexAsset = Database.GetSearchIndex()->FindAssetForPose(State.DbPoseIdx))
+		{
+			Reflection->AssetPlayerAssetName = Database.GetSourceAssetName(IndexAsset);
+		}
 
-	Reflection->AssetPlayerTime = State.AssetPlayerTime;
-	Reflection->LastDeltaTime = State.DeltaTime;
-	Reflection->SimLinearVelocity = State.SimLinearVelocity;
-	Reflection->SimAngularVelocity = State.SimAngularVelocity;
-	Reflection->AnimLinearVelocity = State.AnimLinearVelocity;
-	Reflection->AnimAngularVelocity = State.AnimAngularVelocity;
+		Reflection->AssetPlayerTime = State.AssetPlayerTime;
+		Reflection->LastDeltaTime = State.DeltaTime;
+		Reflection->SimLinearVelocity = State.SimLinearVelocity;
+		Reflection->SimAngularVelocity = State.SimAngularVelocity;
+		Reflection->AnimLinearVelocity = State.AnimLinearVelocity;
+		Reflection->AnimAngularVelocity = State.AnimAngularVelocity;
 
-	// Query pose
-	Reflection->QueryPoseVector = State.QueryVector;
+		// Query pose
+		Reflection->QueryPoseVector = State.QueryVector;
 
-	// Active pose
-	Reflection->ActivePoseVector = Database.GetSearchIndex()->GetPoseValues(State.DbPoseIdx);
-	Database.GetSearchIndex()->InverseNormalize(Reflection->ActivePoseVector);
+		// Active pose
+		Reflection->ActivePoseVector = Database.GetSearchIndex()->GetPoseValues(State.DbPoseIdx);
+		Database.GetSearchIndex()->InverseNormalize(Reflection->ActivePoseVector);
 
-	auto DebuggerView = ParentDebuggerViewPtr.Pin();
-	if (DebuggerView.IsValid()) 
-	{
-		TArray<TSharedRef<FDebuggerDatabaseRowData>> SelectedRows = DebuggerView->GetSelectedDatabaseRows();
-		if (!SelectedRows.IsEmpty())
+		auto DebuggerView = ParentDebuggerViewPtr.Pin();
+		if (DebuggerView.IsValid())
 		{
-			const TSharedRef<FDebuggerDatabaseRowData>& Selected = SelectedRows[0];
-			Reflection->SelectedPoseVector = Database.GetSearchIndex()->GetPoseValues(Selected->PoseIdx);
-			Database.GetSearchIndex()->InverseNormalize(Reflection->SelectedPoseVector);
+			TArray<TSharedRef<FDebuggerDatabaseRowData>> SelectedRows = DebuggerView->GetSelectedDatabaseRows();
+			if (!SelectedRows.IsEmpty())
+			{
+				const TSharedRef<FDebuggerDatabaseRowData>& Selected = SelectedRows[0];
+				Reflection->SelectedPoseVector = Database.GetSearchIndex()->GetPoseValues(Selected->PoseIdx);
+				Database.GetSearchIndex()->InverseNormalize(Reflection->SelectedPoseVector);
 
-			Reflection->CostVector = Selected->PoseCostDetails.CostVector;
-			//Database.SearchIndex.InverseNormalize(Reflection->CostVector);
+				Reflection->CostVector = Selected->PoseCostDetails.CostVector;
+				//Database.SearchIndex.InverseNormalize(Reflection->CostVector);
 
-			const TSharedRef<FDebuggerDatabaseRowData>& ActiveRow = DebuggerView->GetPoseIdxDatabaseRow(State.DbPoseIdx);
+				const TSharedRef<FDebuggerDatabaseRowData>& ActiveRow = DebuggerView->GetPoseIdxDatabaseRow(State.DbPoseIdx);
 
-			Reflection->CostVectorDifference = Reflection->CostVector;
-			for (int i = 0; i < Reflection->CostVectorDifference.Num(); ++i)
-			{
-				Reflection->CostVectorDifference[i] -= ActiveRow->PoseCostDetails.CostVector[i];
+				Reflection->CostVectorDifference = Reflection->CostVector;
+				for (int i = 0; i < Reflection->CostVectorDifference.Num(); ++i)
+				{
+					Reflection->CostVectorDifference[i] -= ActiveRow->PoseCostDetails.CostVector[i];
+				}
 			}
 		}
 	}
@@ -2192,10 +2210,15 @@ void FDebuggerViewModel::ShowSelectedSkeleton(int32 PoseIdx, float Time)
 	{
 		return;
 	}
+	const FPoseSearchIndexAsset* IndexAsset = Database->GetSearchIndex()->FindAssetForPose(PoseIdx);
+	if (!IndexAsset)
+	{
+		return;
+	}
 
-	Component->ResetToStart();
+	Component->ResetToStart(); 
 	bSelecting = true;
-	const FPoseSearchIndexAsset* IndexAsset = Database->GetSearchIndex()->FindAssetForPose(PoseIdx);
+	
 	Skeletons[SelectedPose].Type = IndexAsset->Type;
 	Skeletons[SelectedPose].Time = Time;
 	Skeletons[SelectedPose].bMirrored = IndexAsset->bMirrored;
@@ -2275,7 +2298,7 @@ void FDebuggerViewModel::OnUpdateNodeSelection(int32 InNodeId)
 
 	const UPoseSearchDatabase* NewDatabase = GetPoseSearchDatabase();
 
-	if (ActiveMotionMatchingState)
+	if (ActiveMotionMatchingState && NewDatabase)
 	{
 		const FPoseSearchIndexAsset* IndexAsset = NewDatabase->GetSearchIndex()->FindAssetForPose(ActiveMotionMatchingState->DbPoseIdx);
 
@@ -2297,6 +2320,10 @@ void FDebuggerViewModel::OnUpdateNodeSelection(int32 InNodeId)
 
 void FDebuggerViewModel::OnDraw(FSkeletonDrawParams& DrawParams)
 {
+	const UPoseSearchDatabase* PoseSearchDatabase = GetPoseSearchDatabase();
+	if (!PoseSearchDatabase)
+		return;
+
 	// Returns if it is to be drawn this frame
 	auto SetDrawSkeleton = [this](UPoseSearchMeshComponent* InComponent, bool bDraw)
 	{
@@ -2316,7 +2343,8 @@ void FDebuggerViewModel::OnDraw(FSkeletonDrawParams& DrawParams)
 	FillCompactPoseAndComponentRefRotations();
 
 	UPoseSearchMeshComponent::FUpdateContext UpdateContext;
-	UpdateContext.MirrorDataTable = GetPoseSearchDatabase()->Schema->MirrorDataTable;
+
+	UpdateContext.MirrorDataTable = PoseSearchDatabase->Schema->MirrorDataTable;
 	UpdateContext.CompactPoseMirrorBones = &CompactPoseMirrorBones;
 	UpdateContext.ComponentSpaceRefRotations = &ComponentSpaceRefRotations;
 
@@ -2327,32 +2355,39 @@ void FDebuggerViewModel::OnDraw(FSkeletonDrawParams& DrawParams)
 		if (Skeletons[SelectedPose].Type == ESearchIndexAssetType::Sequence)
 		{
 			const FPoseSearchDatabaseSequence* DatabaseSequence = GetAnimSequence(Skeletons[SelectedPose].AssetIdx);
-
-			UpdateContext.Type = ESearchIndexAssetType::Sequence;
-			UpdateContext.Sequence = DatabaseSequence->Sequence;
-			UpdateContext.StartTime = Skeletons[SelectedPose].Time;
-			UpdateContext.Time = Skeletons[SelectedPose].Time;
-			UpdateContext.bMirrored = Skeletons[SelectedPose].bMirrored;
-			UpdateContext.bLoop = DatabaseSequence->Sequence->bLoop;
+			if (DatabaseSequence)
+			{
+				UpdateContext.Type = ESearchIndexAssetType::Sequence;
+				UpdateContext.Sequence = DatabaseSequence->Sequence;
+				UpdateContext.StartTime = Skeletons[SelectedPose].Time;
+				UpdateContext.Time = Skeletons[SelectedPose].Time;
+				UpdateContext.bMirrored = Skeletons[SelectedPose].bMirrored;
+				UpdateContext.bLoop = DatabaseSequence->Sequence->bLoop;
+			}
 		}
 		else if (Skeletons[SelectedPose].Type == ESearchIndexAssetType::BlendSpace)
 		{
 			const FPoseSearchDatabaseBlendSpace* DatabaseBlendSpace = GetBlendSpace(Skeletons[SelectedPose].AssetIdx);
-
-			UpdateContext.Type = ESearchIndexAssetType::BlendSpace;
-			UpdateContext.BlendSpace = DatabaseBlendSpace->BlendSpace;
-			UpdateContext.StartTime = Skeletons[SelectedPose].Time;
-			UpdateContext.Time = Skeletons[SelectedPose].Time;
-			UpdateContext.bMirrored = Skeletons[SelectedPose].bMirrored;
-			UpdateContext.bLoop = DatabaseBlendSpace->BlendSpace->bLoop;
-			UpdateContext.BlendParameters = Skeletons[SelectedPose].BlendParameters;
+			if (DatabaseBlendSpace)
+			{
+				UpdateContext.Type = ESearchIndexAssetType::BlendSpace;
+				UpdateContext.BlendSpace = DatabaseBlendSpace->BlendSpace;
+				UpdateContext.StartTime = Skeletons[SelectedPose].Time;
+				UpdateContext.Time = Skeletons[SelectedPose].Time;
+				UpdateContext.bMirrored = Skeletons[SelectedPose].bMirrored;
+				UpdateContext.bLoop = DatabaseBlendSpace->BlendSpace->bLoop;
+				UpdateContext.BlendParameters = Skeletons[SelectedPose].BlendParameters;
+			}
 		}
 		else
 		{
 			checkNoEntry();
 		}
 
-		Component->UpdatePose(UpdateContext);
+		if (UpdateContext.Type != ESearchIndexAssetType::Invalid)
+		{
+			Component->UpdatePose(UpdateContext);
+		}
 	}
 
 	const bool bDrawAsset = EnumHasAnyFlags(DrawParams.Flags, ESkeletonDrawFlags::Asset);
@@ -2364,32 +2399,39 @@ void FDebuggerViewModel::OnDraw(FSkeletonDrawParams& DrawParams)
 		if (Skeletons[SelectedPose].Type == ESearchIndexAssetType::Sequence)
 		{
 			const FPoseSearchDatabaseSequence* DatabaseSequence = GetAnimSequence(Skeletons[Asset].AssetIdx);
-
-			UpdateContext.Type = ESearchIndexAssetType::Sequence;
-			UpdateContext.Sequence = DatabaseSequence->Sequence;
-			UpdateContext.StartTime = Skeletons[Asset].Time;
-			UpdateContext.Time = Skeletons[Asset].Time;
-			UpdateContext.bMirrored = Skeletons[Asset].bMirrored;
-			UpdateContext.bLoop = DatabaseSequence->Sequence->bLoop;
+			if (DatabaseSequence)
+			{
+				UpdateContext.Type = ESearchIndexAssetType::Sequence;
+				UpdateContext.Sequence = DatabaseSequence->Sequence;
+				UpdateContext.StartTime = Skeletons[Asset].Time;
+				UpdateContext.Time = Skeletons[Asset].Time;
+				UpdateContext.bMirrored = Skeletons[Asset].bMirrored;
+				UpdateContext.bLoop = DatabaseSequence->Sequence->bLoop;
+			}
 		}
 		else if (Skeletons[SelectedPose].Type == ESearchIndexAssetType::BlendSpace)
 		{
 			const FPoseSearchDatabaseBlendSpace* DatabaseBlendSpace = GetBlendSpace(Skeletons[Asset].AssetIdx);
-
-			UpdateContext.Type = ESearchIndexAssetType::BlendSpace;
-			UpdateContext.BlendSpace = DatabaseBlendSpace->BlendSpace;
-			UpdateContext.StartTime = Skeletons[Asset].Time;
-			UpdateContext.Time = Skeletons[Asset].Time;
-			UpdateContext.bMirrored = Skeletons[Asset].bMirrored;
-			UpdateContext.bLoop = DatabaseBlendSpace->BlendSpace->bLoop;
-			UpdateContext.BlendParameters = Skeletons[Asset].BlendParameters;
+			if (DatabaseBlendSpace)
+			{
+				UpdateContext.Type = ESearchIndexAssetType::BlendSpace;
+				UpdateContext.BlendSpace = DatabaseBlendSpace->BlendSpace;
+				UpdateContext.StartTime = Skeletons[Asset].Time;
+				UpdateContext.Time = Skeletons[Asset].Time;
+				UpdateContext.bMirrored = Skeletons[Asset].bMirrored;
+				UpdateContext.bLoop = DatabaseBlendSpace->BlendSpace->bLoop;
+				UpdateContext.BlendParameters = Skeletons[Asset].BlendParameters;
+			}
 		}
 		else
 		{
 			checkNoEntry();
 		}
 
-		Component->UpdatePose(UpdateContext);
+		if (UpdateContext.Type != ESearchIndexAssetType::Invalid)
+		{
+			Component->UpdatePose(UpdateContext);
+		}
 	}
 }
 
@@ -2474,13 +2516,12 @@ void FDebuggerViewModel::UpdateFromTimeline()
 
 void FDebuggerViewModel::UpdateAsset()
 {
-	if (!IsPlayingSelections())
+	const UPoseSearchDatabase* Database = GetPoseSearchDatabase();
+	if (!Database || !IsPlayingSelections())
 	{
 		return;
 	}
 
-	const UPoseSearchDatabase* Database = GetPoseSearchDatabase();
-
 	FSkeleton& AssetSkeleton = Skeletons[Asset];
 	UPoseSearchMeshComponent* Component = AssetSkeleton.Component;
 
@@ -2613,11 +2654,14 @@ void FDebuggerViewModel::PlaySelection(int32 PoseIdx, float Time)
 	{
 		return;
 	}
+	const FPoseSearchIndexAsset* IndexAsset = Database->GetSearchIndex()->FindAssetForPose(PoseIdx);
+	if (!IndexAsset)
+	{
+		return;
+	}
 
 	Component->ResetToStart();
 	
-	const FPoseSearchIndexAsset* IndexAsset = Database->GetSearchIndex()->FindAssetForPose(PoseIdx);
-
 	Skeletons[Asset].Type = IndexAsset->Type;
 	Skeletons[Asset].AssetIdx = IndexAsset->SourceAssetIdx;
 	Skeletons[Asset].Time = Time;
diff --git a/Engine/Plugins/Experimental/Animation/PoseSearch/Source/Runtime/Private/PoseSearch.cpp b/Engine/Plugins/Experimental/Animation/PoseSearch/Source/Runtime/Private/PoseSearch.cpp
--- a/Engine/Plugins/Experimental/Animation/PoseSearch/Source/Runtime/Private/PoseSearch.cpp
+++ b/Engine/Plugins/Experimental/Animation/PoseSearch/Source/Runtime/Private/PoseSearch.cpp
@@ -395,6 +395,15 @@ void UPoseSearchSchema::PostLoad()
 	}
 }
 
+#if WITH_EDITOR
+void UPoseSearchSchema::PostEditChangeProperty(FPropertyChangedEvent& PropertyChangedEvent)
+{
+	SchemaCardinality = -1;
+
+	Super::PostEditChangeProperty(PropertyChangedEvent);
+}
+#endif
+
 bool UPoseSearchSchema::IsValid() const
 {
 	bool bValid = Skeleton != nullptr;
@@ -1706,6 +1715,18 @@ UE::PoseSearch::FSearchResult UPoseSearchDatabase::Search(FPoseSearchContext& Se
 
 	if (!ensure(SearchIndex->IsValid() && !SearchIndex->IsEmpty()))
 	{
+		if (!Schema)
+		{
+			UE_LOG(LogAnimation, Error, TEXT("UPoseSearchDatabase %s failed to index. Reason: no Schema!"), *GetName());
+		}
+		else if(!Schema->IsValid())
+		{
+			UE_LOG(LogAnimation, Error, TEXT("UPoseSearchDatabase %s failed to index. Reason: Schema %s is invalid"), *GetName(), *Schema->GetName());
+		}
+		else
+		{
+			UE_LOG(LogAnimation, Error, TEXT("UPoseSearchDatabase %s failed to index. Reason: is there any unsaved modified asset?"), *GetName());
+		}
 		return Result;
 	}
 
diff --git a/Engine/Plugins/Experimental/Animation/PoseSearch/Source/Runtime/Public/PoseSearch/PoseSearch.h b/Engine/Plugins/Experimental/Animation/PoseSearch/Source/Runtime/Public/PoseSearch/PoseSearch.h
--- a/Engine/Plugins/Experimental/Animation/PoseSearch/Source/Runtime/Public/PoseSearch/PoseSearch.h
+++ b/Engine/Plugins/Experimental/Animation/PoseSearch/Source/Runtime/Public/PoseSearch/PoseSearch.h
@@ -494,6 +494,10 @@ public: // UObject
 public: // IBoneReferenceSkeletonProvider
 	class USkeleton* GetSkeleton(bool& bInvalidSkeletonIsError, const IPropertyHandle* PropertyHandle) override { bInvalidSkeletonIsError = false; return Skeleton; }
 
+#if WITH_EDITOR
+	virtual void PostEditChangeProperty(struct FPropertyChangedEvent& PropertyChangedEvent) override;
+#endif
+
 private:
 	void Finalize();
 	void ResolveBoneReferences();
