motion matching - pose search feature channels debug draw fixes and improvements

diff --git a/Engine/Plugins/Experimental/Animation/PoseSearch/Source/Runtime/Private/PoseSearchFeatureChannels.cpp b/Engine/Plugins/Experimental/Animation/PoseSearch/Source/Runtime/Private/PoseSearchFeatureChannels.cpp
--- a/Engine/Plugins/Experimental/Animation/PoseSearch/Source/Runtime/Private/PoseSearchFeatureChannels.cpp
+++ b/Engine/Plugins/Experimental/Animation/PoseSearch/Source/Runtime/Private/PoseSearchFeatureChannels.cpp
@@ -14,15 +14,14 @@ namespace UE::PoseSearch {
 //////////////////////////////////////////////////////////////////////////
 // Constants
 
-constexpr float DrawDebugLineThickness = 2.0f;
-constexpr float DrawDebugPointSize = 3.0f;
+constexpr float DrawDebugLineThickness = 1.0f;
+constexpr float DrawDebugPointSize = 2.0f;
 constexpr float DrawDebugVelocityScale = 0.08f;
-constexpr float DrawDebugArrowSize = 30.0f;
-constexpr float DrawDebugSphereSize = 3.0f;
-constexpr int32 DrawDebugSphereSegments = 10;
+constexpr float DrawDebugSphereSize = 2.0f;
+constexpr int32 DrawDebugSphereSegments = 8;
 constexpr float DrawDebugGradientStrength = 0.8f;
-constexpr float DrawDebugSampleLabelFontScale = 1.0f;
-static const FVector DrawDebugSampleLabelOffset = FVector(0.0f, 0.0f, -10.0f);
+constexpr float DrawDebugSampleLabelFontScale = 1.5f;
+static const FVector DrawDebugSampleLabelOffset = FVector(0.0f, 0.0f, 5.0f);
 
 constexpr bool UseCharacterSpaceVelocities = true;
 
@@ -933,7 +932,7 @@ void UPoseSearchFeatureChannel_Pose::DebugDraw(const UE::PoseSearch::FDebugDrawP
 
 				if (EnumHasAnyFlags(DrawParams.Flags, EDebugDrawFlags::DrawSearchIndex))
 				{
-					DrawDebugPoint(DrawParams.World, BoneVel, DrawParams.PointSize, Color, bPersistent, DrawParams.DefaultLifeTime, DepthPriority);
+					DrawDebugLine(DrawParams.World, BonePos[SubsampleIdx], BonePos[SubsampleIdx] + BoneVel, Color, bPersistent, DrawParams.DefaultLifeTime, DepthPriority);
 				}
 				else
 				{
@@ -941,11 +940,10 @@ void UPoseSearchFeatureChannel_Pose::DebugDraw(const UE::PoseSearch::FDebugDrawP
 						EnumHasAnyFlags(DrawParams.Flags, EDebugDrawFlags::DrawFast) ?
 						0.0f : DrawDebugLineThickness;
 
-					DrawDebugDirectionalArrow(
+					DrawDebugLine(
 						DrawParams.World,
 						BonePos[SubsampleIdx] + BoneVelDirection * DrawDebugSphereSize,
 						BonePos[SubsampleIdx] + BoneVel,
-						DrawDebugArrowSize,
 						Color,
 						bPersistent,
 						LifeTime,
@@ -1331,7 +1329,7 @@ void UPoseSearchFeatureChannel_Trajectory::DebugDraw(const UE::PoseSearch::FDebu
 
 			if (EnumHasAnyFlags(DrawParams.Flags, EDebugDrawFlags::DrawSearchIndex))
 			{
-				DrawDebugPoint(DrawParams.World, TrajectoryVel, DrawParams.PointSize, Color, bPersistent, DrawParams.DefaultLifeTime, DepthPriority);
+				DrawDebugLine(DrawParams.World, TrajectoryPos[SubsampleIdx], TrajectoryPos[SubsampleIdx] + TrajectoryVel, Color, bPersistent, DrawParams.DefaultLifeTime, DepthPriority);
 			}
 			else
 			{
@@ -1339,11 +1337,10 @@ void UPoseSearchFeatureChannel_Trajectory::DebugDraw(const UE::PoseSearch::FDebu
 					EnumHasAnyFlags(DrawParams.Flags, EDebugDrawFlags::DrawFast) ?
 					0.0f : DrawDebugLineThickness;
 
-				DrawDebugDirectionalArrow(
+				DrawDebugLine(
 					DrawParams.World,
 					TrajectoryPos[SubsampleIdx] + TrajectoryVelDirection * DrawDebugSphereSize,
 					TrajectoryPos[SubsampleIdx] + TrajectoryVel,
-					DrawDebugArrowSize,
 					Color,
 					bPersistent,
 					LifeTime,
@@ -1368,7 +1365,7 @@ void UPoseSearchFeatureChannel_Trajectory::DebugDraw(const UE::PoseSearch::FDebu
 
 			if (EnumHasAnyFlags(DrawParams.Flags, EDebugDrawFlags::DrawSearchIndex))
 			{
-				DrawDebugPoint(DrawParams.World, TrajectoryForward, DrawParams.PointSize, Color, bPersistent, DrawParams.DefaultLifeTime, DepthPriority);
+				DrawDebugLine(DrawParams.World, TrajectoryPos[SubsampleIdx], TrajectoryPos[SubsampleIdx] + TrajectoryForward, Color, bPersistent, DrawParams.DefaultLifeTime, DepthPriority);
 			}
 			else
 			{
@@ -1376,11 +1373,10 @@ void UPoseSearchFeatureChannel_Trajectory::DebugDraw(const UE::PoseSearch::FDebu
 					EnumHasAnyFlags(DrawParams.Flags, EDebugDrawFlags::DrawFast) ?
 					0.0f : DrawDebugLineThickness;
 
-				DrawDebugDirectionalArrow(
+				DrawDebugLine(
 					DrawParams.World,
 					TrajectoryPos[SubsampleIdx] + TrajectoryForward * DrawDebugSphereSize,
 					TrajectoryPos[SubsampleIdx] + TrajectoryForward * DrawDebugSphereSize * 2.0f,
-					DrawDebugArrowSize,
 					Color,
 					bPersistent,
 					LifeTime,
