motion matching - introducing UPoseSearchFeatureChannel_Position as first itemization of UPoseSearchFeatureChannel_Pose

diff --git a/Engine/Plugins/Experimental/Animation/PoseSearch/Source/Editor/Private/PoseSearchDebugger.cpp b/Engine/Plugins/Experimental/Animation/PoseSearch/Source/Editor/Private/PoseSearchDebugger.cpp
--- a/Engine/Plugins/Experimental/Animation/PoseSearch/Source/Editor/Private/PoseSearchDebugger.cpp
+++ b/Engine/Plugins/Experimental/Animation/PoseSearch/Source/Editor/Private/PoseSearchDebugger.cpp
@@ -2516,6 +2516,10 @@ void FDebuggerViewModel::UpdateFromTimeline()
 
 void FDebuggerViewModel::UpdateAsset()
 {
+	// @todo: expose those parameters
+	static float MAX_DISTANCE_RANGE = 200.f;
+	static float MAX_TIME_RANGE = 2.f;
+
 	const UPoseSearchDatabase* Database = GetPoseSearchDatabase();
 	if (!Database || !IsPlayingSelections())
 	{
@@ -2556,11 +2560,8 @@ void FDebuggerViewModel::UpdateAsset()
 	{
 		const float DT = static_cast<float>(FApp::GetDeltaTime()) * AssetPlayRate;
 		const float PlayLength = AnimAsset->GetPlayLength();
-
-		const FFloatRange DistanceRange = Database->Schema->GetHorizonRange(EPoseSearchFeatureDomain::Distance);
-		const bool bExceededDistanceHorizon = !DistanceRange.IsEmpty() && (Component->LastRootMotionDelta.GetTranslation().Size() > DistanceRange.GetUpperBoundValue());
-		const FFloatRange TimeRange = Database->Schema->GetHorizonRange(EPoseSearchFeatureDomain::Time);
-		const bool bExceededTimeHorizon = !TimeRange.IsEmpty() && ((AssetSkeleton.Time - AssetData.StartTime) > TimeRange.GetUpperBoundValue());
+		const bool bExceededDistanceHorizon = Component->LastRootMotionDelta.GetTranslation().Size() > MAX_DISTANCE_RANGE;
+		const bool bExceededTimeHorizon = (AssetSkeleton.Time - AssetData.StartTime) > MAX_TIME_RANGE;
 		const bool bExceededHorizon = bExceededDistanceHorizon && bExceededTimeHorizon;
 		if (bAssetLooping)
 		{
diff --git a/Engine/Plugins/Experimental/Animation/PoseSearch/Source/Runtime/Private/PoseSearch.cpp b/Engine/Plugins/Experimental/Animation/PoseSearch/Source/Runtime/Private/PoseSearch.cpp
--- a/Engine/Plugins/Experimental/Animation/PoseSearch/Source/Runtime/Private/PoseSearch.cpp
+++ b/Engine/Plugins/Experimental/Animation/PoseSearch/Source/Runtime/Private/PoseSearch.cpp
@@ -373,13 +373,22 @@ void UPoseSearchSchema::PostLoad()
 
 	if (!bNeedFinalize)
 	{
-		for (const UPoseSearchFeatureChannel* Channel : Channels)
+		for (UPoseSearchFeatureChannel* Channel : Channels)
 		{
 			if (Channel->ChannelDataOffset == -1 || Channel->ChannelCardinality == -1)
 			{
 				bNeedFinalize = true;
 				break;
 			}
+
+			if (UPoseSearchFeatureChannel_Pose* ChannelPose = Cast<UPoseSearchFeatureChannel_Pose>(Channel))
+			{
+				if (ChannelPose->SchemaBoneIdx.Num() == 0)
+				{
+					bNeedFinalize = true;
+					break;
+				}
+			}
 		}
 	}
 
@@ -402,6 +411,26 @@ void UPoseSearchSchema::PostEditChangeProperty(FPropertyChangedEvent& PropertyCh
 
 	Super::PostEditChangeProperty(PropertyChangedEvent);
 }
+
+void UPoseSearchSchema::GenerateDDCKey(FBlake3& InOutKeyHasher) const
+{
+	for (const TObjectPtr<UPoseSearchFeatureChannel>& Channel : Channels)
+	{
+		if (Channel)
+		{
+			const FBlake3Hash& ChannelClassHash = Channel->GetClass()->GetSchemaHash(false);
+			InOutKeyHasher.Update(MakeMemoryView(ChannelClassHash.GetBytes()));
+			Channel->GenerateDDCKey(InOutKeyHasher);
+		}
+	}
+
+	InOutKeyHasher.Update(&DataPreprocessor, sizeof(DataPreprocessor));
+	InOutKeyHasher.Update(&EffectiveDataPreprocessor, sizeof(EffectiveDataPreprocessor));
+	InOutKeyHasher.Update(&SamplingInterval, sizeof(SamplingInterval));
+	InOutKeyHasher.Update(MakeMemoryView(BoneIndicesWithParents));
+
+	// @todo: add schema version if POSESEARCHDB_DERIVEDDATA_VER is not enought
+}
 #endif
 
 bool UPoseSearchSchema::IsValid() const
@@ -425,21 +454,6 @@ bool UPoseSearchSchema::IsValid() const
 	return bValid;
 }
 
-FFloatRange UPoseSearchSchema::GetHorizonRange(EPoseSearchFeatureDomain Domain) const
-{
-	FFloatRange GlobalRange = FFloatRange::Empty();
-	for (const auto & Channel : Channels)
-	{
-		if (Channel)
-		{
-			FFloatRange ChannelRange = Channel->GetHorizonRange(Domain);
-			GlobalRange = FFloatRange::Hull(GlobalRange, ChannelRange);
-		}
-	}
-
-	return GlobalRange;
-}
-
 void UPoseSearchSchema::ResolveBoneReferences()
 {
 	// Initialize references to obtain bone indices
@@ -931,6 +945,130 @@ int32 UPoseSearchDatabase::GetNumberOfPrincipalComponents() const
 	return FMath::Min<int32>(NumberOfPrincipalComponents, Schema->SchemaCardinality);
 }
 
+#if WITH_EDITOR
+void UPoseSearchDatabase::AddDbSequenceToWriter(const FPoseSearchDatabaseSequence& DbSequence, FBlake3& InOutWriter)
+{
+	// Main Sequence
+	AddRawSequenceToWriter(DbSequence.Sequence, InOutWriter);
+	InOutWriter.Update(&DbSequence.SamplingRange, sizeof(DbSequence.SamplingRange));
+	if (DbSequence.Sequence)
+	{
+		InOutWriter.Update(&DbSequence.Sequence->bLoop, sizeof(DbSequence.Sequence->bLoop));
+	}
+	InOutWriter.Update(&DbSequence.MirrorOption, sizeof(DbSequence.MirrorOption));
+
+	// Lead in sequence
+	AddRawSequenceToWriter(DbSequence.LeadInSequence, InOutWriter);
+	if (DbSequence.LeadInSequence)
+	{
+		InOutWriter.Update(&DbSequence.LeadInSequence->bLoop, sizeof(DbSequence.LeadInSequence->bLoop));
+	}
+
+	// Follow up sequence
+	AddRawSequenceToWriter(DbSequence.FollowUpSequence, InOutWriter);
+	if (DbSequence.FollowUpSequence)
+	{
+		InOutWriter.Update(&DbSequence.FollowUpSequence->bLoop, sizeof(DbSequence.FollowUpSequence->bLoop));
+	}
+
+	// Tags
+	InOutWriter.Update(&DbSequence.GroupTags, sizeof(DbSequence.GroupTags));
+
+	// Notifies
+	AddPoseSearchNotifiesToWriter(DbSequence.Sequence, InOutWriter);
+}
+
+void UPoseSearchDatabase::AddRawSequenceToWriter(const UAnimSequence* Sequence,	FBlake3& InOutWriter)
+{
+	if (Sequence)
+	{
+		InOutWriter.Update(MakeMemoryView(Sequence->GetName()));
+		InOutWriter.Update(MakeMemoryView(Sequence->GetRawDataGuid().ToString()));
+	}
+}
+
+void UPoseSearchDatabase::AddDbBlendSpaceToWriter(const FPoseSearchDatabaseBlendSpace& DbBlendSpace, FBlake3& InOutWriter)
+{
+	if (IsValid(DbBlendSpace.BlendSpace))
+	{
+		const TArray<FBlendSample>& BlendSpaceSamples = DbBlendSpace.BlendSpace->GetBlendSamples();
+		for (const FBlendSample& Sample : BlendSpaceSamples)
+		{
+			AddRawSequenceToWriter(Sample.Animation, InOutWriter);
+			InOutWriter.Update(&Sample.SampleValue, sizeof(Sample.SampleValue));
+			InOutWriter.Update(&Sample.RateScale, sizeof(Sample.RateScale));
+		}
+
+		InOutWriter.Update(&DbBlendSpace.BlendSpace->bLoop, sizeof(DbBlendSpace.BlendSpace->bLoop));
+		InOutWriter.Update(&DbBlendSpace.MirrorOption, sizeof(DbBlendSpace.MirrorOption));
+		InOutWriter.Update(&DbBlendSpace.bUseGridForSampling, sizeof(DbBlendSpace.bUseGridForSampling));
+		InOutWriter.Update(&DbBlendSpace.NumberOfHorizontalSamples, sizeof(DbBlendSpace.NumberOfHorizontalSamples));
+		InOutWriter.Update(&DbBlendSpace.NumberOfVerticalSamples, sizeof(DbBlendSpace.NumberOfVerticalSamples));
+		InOutWriter.Update(&DbBlendSpace.GroupTags, sizeof(DbBlendSpace.GroupTags));
+	}
+}
+
+void UPoseSearchDatabase::AddPoseSearchNotifiesToWriter(const UAnimSequence* Sequence, FBlake3& InOutWriter)
+{
+	if (!Sequence)
+	{
+		return;
+	}
+
+	FAnimNotifyContext NotifyContext;
+	Sequence->GetAnimNotifies(0.0f, Sequence->GetPlayLength(), NotifyContext);
+
+	for (const FAnimNotifyEventReference& EventReference : NotifyContext.ActiveNotifies)
+	{
+		const FAnimNotifyEvent* NotifyEvent = EventReference.GetNotify();
+		if (!NotifyEvent || !NotifyEvent->NotifyStateClass)
+		{
+			continue;
+		}
+
+		if (NotifyEvent->NotifyStateClass->IsA<UAnimNotifyState_PoseSearchBase>())
+		{
+			const float StartTime = NotifyEvent->GetTriggerTime();
+			const float EndTime = NotifyEvent->GetEndTriggerTime();
+			InOutWriter.Update(&StartTime, sizeof(float));
+			InOutWriter.Update(&EndTime, sizeof(float));
+
+			if (const auto ModifyCostNotifyState =
+				Cast<const UAnimNotifyState_PoseSearchModifyCost>(NotifyEvent->NotifyStateClass))
+			{
+				InOutWriter.Update(&ModifyCostNotifyState->CostAddend, sizeof(float));
+			}
+		}
+	}
+}
+
+void UPoseSearchDatabase::GenerateDDCKey(FBlake3& InOutKeyHasher) const
+{
+	if (IsValid(Schema))
+	{
+		Schema->GenerateDDCKey(InOutKeyHasher);
+	}
+
+	for (int32 i = 0; i < Sequences.Num(); ++i)
+	{
+		AddDbSequenceToWriter(Sequences[i], InOutKeyHasher);
+	}
+
+	for (int32 i = 0; i < BlendSpaces.Num(); ++i)
+	{
+		AddDbBlendSpaceToWriter(BlendSpaces[i], InOutKeyHasher);
+	}
+
+	InOutKeyHasher.Update(&NumberOfPrincipalComponents, sizeof(NumberOfPrincipalComponents));
+	InOutKeyHasher.Update(&KDTreeMaxLeafSize, sizeof(KDTreeMaxLeafSize));
+	InOutKeyHasher.Update(&KDTreeQueryNumNeighbors, sizeof(KDTreeQueryNumNeighbors));
+	InOutKeyHasher.Update(&PoseSearchMode, sizeof(PoseSearchMode));
+
+	// @todo: add database version if POSESEARCHDB_DERIVEDDATA_VER is not enought
+}
+
+#endif // WITH_EDITOR
+
 bool UPoseSearchDatabase::IsValidForIndexing() const
 {
 	bool bValid = Schema && Schema->IsValid() && !Sequences.IsEmpty();
@@ -1277,7 +1415,7 @@ bool UPoseSearchDatabase::TryInitSearchIndexAssets(FPoseSearchIndex& OutSearchIn
 		}
 	}
 
-	// @todo: change the above for loops fill OutSearchIndex.Assets already in ascnding group order
+	// @todo: change the above for loops fill OutSearchIndex.Assets already in ascending group order
 	// sorting by ascending SourceGroupIdx
 	OutSearchIndex.Assets.Sort([](const FPoseSearchIndexAsset& InOne, const FPoseSearchIndexAsset& InTwo)
 	{
@@ -3493,6 +3631,7 @@ public: // IAssetIndexer
 	FSampleInfo GetSampleInfoRelative(float SampleTime, const FSampleInfo& Origin) const override;
 	const float GetSampleTimeFromDistance(float Distance) const override;
 	FTransform MirrorTransform(const FTransform& Transform) const override;
+	FTransform GetTransformAndCacheResults(float SampleTime, float OriginTime, int8 SchemaBoneIdx, bool& Clamped) override;
 
 private:
 	FAssetIndexingContext IndexingContext;
@@ -3500,6 +3639,23 @@ private:
 	FPoseSearchPoseMetadata Metadata;
 	
 	void AddMetadata(int32 SampleIdx);
+
+	struct CachedEntry
+	{
+		float SampleTime;
+		float OriginTime;
+		bool Clamped;
+
+		// @todo: minimize the Entry memory footprint
+		FTransform RootTransform;
+		FCompactPose Pose;
+		FCSPose<FCompactPose> ComponentSpacePose;
+		FBlendedCurve UnusedCurve;
+		UE::Anim::FStackAttributeContainer UnusedAtrribute;
+		FAnimationPoseData AnimPoseData = { Pose, UnusedCurve, UnusedAtrribute };
+	};
+
+	TArray<CachedEntry> CachedEntries;
 };
 
 void FAssetIndexer::Reset()
@@ -3560,6 +3716,8 @@ bool FAssetIndexer::Process()
 		Channel->IndexAsset(*this, AssetIndexingOutput);
 	}
 
+	// @todo: this step can be avoided now since Schema.SchemaCardinality is known and the data can be preallocated before Channel->IndexAsset
+	
 	// Merge spans of feature vectors into contiguous buffer
 	for (int32 VectorIdx = 0; VectorIdx != NumSamplesInRange; ++VectorIdx)
 	{
@@ -3860,6 +4018,63 @@ void FAssetIndexer::AddMetadata(int32 SampleIdx)
 	}
 }
 
+FTransform FAssetIndexer::GetTransformAndCacheResults(float SampleTime, float OriginTime, int8 SchemaBoneIdx, bool& Clamped)
+{
+	// @todo: use an hashmap if we end up having too many entries
+	CachedEntry* Entry = CachedEntries.FindByPredicate([SampleTime, OriginTime](const FAssetIndexer::CachedEntry& Entry)
+	{
+		return Entry.SampleTime == SampleTime && Entry.OriginTime == OriginTime;
+	});
+
+	const FAssetSamplingContext* SamplingContext = IndexingContext.SamplingContext;
+
+	if (!Entry)
+	{
+		Entry = &CachedEntries[CachedEntries.AddDefaulted()];
+
+		Entry->SampleTime = SampleTime;
+		Entry->OriginTime = OriginTime;
+
+		Entry->Pose.SetBoneContainer(&SamplingContext->BoneContainer);
+		Entry->UnusedCurve.InitFrom(SamplingContext->BoneContainer);
+
+		IAssetIndexer::FSampleInfo Origin = GetSampleInfo(OriginTime);
+		IAssetIndexer::FSampleInfo Sample = GetSampleInfoRelative(SampleTime, Origin);
+
+		float CurrentTime = Sample.ClipTime;
+		float PreviousTime = CurrentTime - SamplingContext->FiniteDelta;
+
+		FDeltaTimeRecord DeltaTimeRecord;
+		DeltaTimeRecord.Set(PreviousTime, CurrentTime - PreviousTime);
+		FAnimExtractContext ExtractionCtx(CurrentTime, true, DeltaTimeRecord, Sample.Clip->IsLoopable());
+
+		Sample.Clip->ExtractPose(ExtractionCtx, Entry->AnimPoseData);
+
+		if (IndexingContext.bMirrored)
+		{
+			FAnimationRuntime::MirrorPose(
+				Entry->AnimPoseData.GetPose(),
+				IndexingContext.Schema->MirrorDataTable->MirrorAxis,
+				SamplingContext->CompactPoseMirrorBones,
+				SamplingContext->ComponentSpaceRefRotations
+			);
+			// Note curves and attributes are not used during the indexing process and therefore don't need to be mirrored
+		}
+
+		Entry->ComponentSpacePose.InitPose(Entry->Pose);
+		Entry->RootTransform = Sample.RootTransform;
+		Entry->Clamped = Sample.bClamped;
+	}
+
+	const FBoneReference& BoneReference = IndexingContext.Schema->BoneReferences[SchemaBoneIdx];
+	FCompactPoseBoneIndex CompactBoneIndex = SamplingContext->BoneContainer.MakeCompactPoseIndex(FMeshPoseBoneIndex(BoneReference.BoneIndex));
+
+	const FTransform BoneTransform = Entry->ComponentSpacePose.GetComponentSpaceTransform(CompactBoneIndex) * MirrorTransform(Entry->RootTransform);
+	Clamped = Entry->Clamped;
+
+	return BoneTransform;
+}
+
 //////////////////////////////////////////////////////////////////////////
 // PoseSearch API
 
diff --git a/Engine/Plugins/Experimental/Animation/PoseSearch/Source/Runtime/Private/PoseSearchDerivedData.cpp b/Engine/Plugins/Experimental/Animation/PoseSearch/Source/Runtime/Private/PoseSearchDerivedData.cpp
--- a/Engine/Plugins/Experimental/Animation/PoseSearch/Source/Runtime/Private/PoseSearchDerivedData.cpp
+++ b/Engine/Plugins/Experimental/Animation/PoseSearch/Source/Runtime/Private/PoseSearchDerivedData.cpp
@@ -195,153 +195,11 @@ namespace UE::PoseSearch
 		const FGuid VersionGuid =  FDevSystemGuids::GetSystemGuid(FDevSystemGuids::Get().POSESEARCHDB_DERIVEDDATA_VER);
 		Writer.Update(MakeMemoryView(VersionGuid.ToString()));
 
-		AddPoseSearchSchemaToWriter(Database.Schema, Writer);
-
-		for (int32 i = 0; i < Database.Sequences.Num(); ++i)
-		{
-			AddDbSequenceToWriter(Database.Sequences[i], Writer);
-		}
-
-		for (int32 i = 0; i < Database.BlendSpaces.Num(); ++i)
-		{
-			AddDbBlendSpaceToWriter(Database.BlendSpaces[i], Writer);
-		}
-
-		Writer.Update(&Database.NumberOfPrincipalComponents, sizeof(Database.NumberOfPrincipalComponents));
-		Writer.Update(&Database.KDTreeMaxLeafSize, sizeof(Database.KDTreeMaxLeafSize));
-		Writer.Update(&Database.KDTreeQueryNumNeighbors, sizeof(Database.KDTreeQueryNumNeighbors));
-		Writer.Update(&Database.PoseSearchMode, sizeof(Database.PoseSearchMode));
+		Database.GenerateDDCKey(Writer);
 
 		return Writer.Finalize();
 	}
 
-	void FPoseSearchDatabaseAsyncCacheTask::AddPoseSearchSchemaToWriter(
-		const UPoseSearchSchema* Schema, 
-		FBlake3& InOutWriter)
-	{
-		if (IsValid(Schema))
-		{
-			for (const TObjectPtr<UPoseSearchFeatureChannel>& Channel : Schema->Channels)
-			{
-				if (Channel)
-				{
-					const FBlake3Hash& ChannelClassHash = Channel->GetClass()->GetSchemaHash(false);
-					InOutWriter.Update(MakeMemoryView(ChannelClassHash.GetBytes()));
-
-					Channel->GenerateDDCKey(InOutWriter);
-				}
-			}
-			InOutWriter.Update(&Schema->DataPreprocessor, sizeof(Schema->DataPreprocessor));
-			InOutWriter.Update(&Schema->EffectiveDataPreprocessor, sizeof(Schema->EffectiveDataPreprocessor));
-			InOutWriter.Update(&Schema->SamplingInterval, sizeof(Schema->SamplingInterval));
-			InOutWriter.Update(MakeMemoryView(Schema->BoneIndicesWithParents));
-		}
-	}
-
-	void FPoseSearchDatabaseAsyncCacheTask::AddDbSequenceToWriter(
-		const FPoseSearchDatabaseSequence& DbSequence, 
-		FBlake3& InOutWriter)
-	{
-		// Main Sequence
-		AddRawSequenceToWriter(DbSequence.Sequence, InOutWriter);
-		InOutWriter.Update(&DbSequence.SamplingRange, sizeof(DbSequence.SamplingRange));
-		if (DbSequence.Sequence)
-		{
-			InOutWriter.Update(&DbSequence.Sequence->bLoop, sizeof(DbSequence.Sequence->bLoop));
-		}
-		InOutWriter.Update(&DbSequence.MirrorOption, sizeof(DbSequence.MirrorOption));
-
-		// Lead in sequence
-		AddRawSequenceToWriter(DbSequence.LeadInSequence, InOutWriter);
-		if (DbSequence.LeadInSequence)
-		{
-			InOutWriter.Update(&DbSequence.LeadInSequence->bLoop, sizeof(DbSequence.LeadInSequence->bLoop));
-		}
-
-		// Follow up sequence
-		AddRawSequenceToWriter(DbSequence.FollowUpSequence, InOutWriter);
-		if (DbSequence.FollowUpSequence)
-		{
-			InOutWriter.Update(&DbSequence.FollowUpSequence->bLoop, sizeof(DbSequence.FollowUpSequence->bLoop));
-		}
-
-		// Tags
-		InOutWriter.Update(&DbSequence.GroupTags, sizeof(DbSequence.GroupTags));
-
-		// Notifies
-		AddPoseSearchNotifiesToWriter(DbSequence.Sequence, InOutWriter);
-	}
-
-	void FPoseSearchDatabaseAsyncCacheTask::AddRawSequenceToWriter(
-		const UAnimSequence* Sequence, 
-		FBlake3& InOutWriter)
-	{
-		if (Sequence)
-		{
-			InOutWriter.Update(MakeMemoryView(Sequence->GetName()));
-			InOutWriter.Update(MakeMemoryView(Sequence->GetRawDataGuid().ToString()));
-		}
-	}
-
-	void FPoseSearchDatabaseAsyncCacheTask::AddPoseSearchNotifiesToWriter(
-		const UAnimSequence* Sequence, 
-		FBlake3& InOutWriter)
-	{
-		if (!Sequence)
-		{
-			return;
-		}
-
-		FAnimNotifyContext NotifyContext;
-		Sequence->GetAnimNotifies(0.0f, Sequence->GetPlayLength(), NotifyContext);
-
-		for (const FAnimNotifyEventReference& EventReference : NotifyContext.ActiveNotifies)
-		{
-			const FAnimNotifyEvent* NotifyEvent = EventReference.GetNotify();
-			if (!NotifyEvent || !NotifyEvent->NotifyStateClass)
-			{
-				continue;
-			}
-
-			if (NotifyEvent->NotifyStateClass->IsA<UAnimNotifyState_PoseSearchBase>())
-			{
-				const float StartTime = NotifyEvent->GetTriggerTime();
-				const float EndTime = NotifyEvent->GetEndTriggerTime();
-				InOutWriter.Update(&StartTime, sizeof(float));
-				InOutWriter.Update(&EndTime, sizeof(float));
-
-				if (const auto ModifyCostNotifyState =
-					Cast<const UAnimNotifyState_PoseSearchModifyCost>(NotifyEvent->NotifyStateClass))
-				{
-					InOutWriter.Update(&ModifyCostNotifyState->CostAddend, sizeof(float));
-				}
-			}
-		}
-	}
-
-	void FPoseSearchDatabaseAsyncCacheTask::AddDbBlendSpaceToWriter(
-		const FPoseSearchDatabaseBlendSpace& DbBlendSpace, 
-		FBlake3& InOutWriter)
-	{
-		if (IsValid(DbBlendSpace.BlendSpace))
-		{
-			const TArray<FBlendSample>& BlendSpaceSamples = DbBlendSpace.BlendSpace->GetBlendSamples();
-			for (const FBlendSample& Sample : BlendSpaceSamples)
-			{
-				AddRawSequenceToWriter(Sample.Animation, InOutWriter);
-				InOutWriter.Update(&Sample.SampleValue, sizeof(Sample.SampleValue));
-				InOutWriter.Update(&Sample.RateScale, sizeof(Sample.RateScale));
-			}
-
-			InOutWriter.Update(&DbBlendSpace.BlendSpace->bLoop, sizeof(DbBlendSpace.BlendSpace->bLoop));
-			InOutWriter.Update(&DbBlendSpace.MirrorOption, sizeof(DbBlendSpace.MirrorOption));
-			InOutWriter.Update(&DbBlendSpace.bUseGridForSampling, sizeof(DbBlendSpace.bUseGridForSampling));
-			InOutWriter.Update(&DbBlendSpace.NumberOfHorizontalSamples, sizeof(DbBlendSpace.NumberOfHorizontalSamples));
-			InOutWriter.Update(&DbBlendSpace.NumberOfVerticalSamples, sizeof(DbBlendSpace.NumberOfVerticalSamples));
-			InOutWriter.Update(&DbBlendSpace.GroupTags, sizeof(DbBlendSpace.GroupTags));
-		}
-	}
-
 #endif // WITH_EDITOR
 
 	FArchive& operator<<(FArchive& Ar, FPoseSearchIndexPreprocessInfo& Info)
diff --git a/Engine/Plugins/Experimental/Animation/PoseSearch/Source/Runtime/Private/PoseSearchDerivedData.h b/Engine/Plugins/Experimental/Animation/PoseSearch/Source/Runtime/Private/PoseSearchDerivedData.h
--- a/Engine/Plugins/Experimental/Animation/PoseSearch/Source/Runtime/Private/PoseSearchDerivedData.h
+++ b/Engine/Plugins/Experimental/Animation/PoseSearch/Source/Runtime/Private/PoseSearchDerivedData.h
@@ -93,11 +93,6 @@ namespace UE::PoseSearch
 
 	public:
 		static FIoHash CreateKey(const UPoseSearchDatabase& Database);
-		static void AddPoseSearchSchemaToWriter(const UPoseSearchSchema* Schema, FBlake3& InOutWriter);
-		static void AddDbSequenceToWriter(const FPoseSearchDatabaseSequence& DbSequence, FBlake3& InOutWriter);
-		static void AddRawSequenceToWriter(const UAnimSequence* Sequence, FBlake3& InOutWriter);
-		static void AddPoseSearchNotifiesToWriter(const UAnimSequence* Sequence, FBlake3& InOutWriter);
-		static void AddDbBlendSpaceToWriter(const FPoseSearchDatabaseBlendSpace& DbBlendSpace, FBlake3& InOutWriter);
 	};
 #endif // WITH_EDITOR
 
diff --git a/Engine/Plugins/Experimental/Animation/PoseSearch/Source/Runtime/Private/PoseSearchFeatureChannels.cpp b/Engine/Plugins/Experimental/Animation/PoseSearch/Source/Runtime/Private/PoseSearchFeatureChannels.cpp
--- a/Engine/Plugins/Experimental/Animation/PoseSearch/Source/Runtime/Private/PoseSearchFeatureChannels.cpp
+++ b/Engine/Plugins/Experimental/Animation/PoseSearch/Source/Runtime/Private/PoseSearchFeatureChannels.cpp
@@ -55,94 +55,9 @@ T GetValueAtIndex(int32 Sample, const TArray<T>& Values)
 	return (Values[Num - 1] - Values[Num - 2]) * (Sample - (Num - 1)) + Values[Num - 1];
 }
 
-struct BoneTransformsCache
+static void CollectBonePositions(TArray<FVector>& BonePositions, IAssetIndexer& Indexer, int8 SchemaBoneIdx)
 {
-	BoneTransformsCache(const IAssetIndexer& InIndexer)
-		: Indexer(InIndexer)
-	{
-	}
-
-	FTransform Get(float SampleTime, float OriginTime, int8 SchemaBoneIdx, bool& Clamped)
-	{
-		// @todo: use an hashmap if we end up having too many entries
-		Entry* Entry = Entries.FindByPredicate([SampleTime, OriginTime](const BoneTransformsCache::Entry& Entry)
-			{
-				return Entry.SampleTime == SampleTime && Entry.OriginTime == OriginTime;
-			});
-
-		const FAssetIndexingContext& IndexingContext = Indexer.GetIndexingContext();
-		const FAssetSamplingContext* SamplingContext = IndexingContext.SamplingContext;
-
-		if (!Entry)
-		{
-			Entry = &Entries[Entries.AddDefaulted()];
-
-			Entry->SampleTime = SampleTime;
-			Entry->OriginTime = OriginTime;
-
-			Entry->Pose.SetBoneContainer(&SamplingContext->BoneContainer);
-			Entry->UnusedCurve.InitFrom(SamplingContext->BoneContainer);
-
-			IAssetIndexer::FSampleInfo Origin = Indexer.GetSampleInfo(OriginTime);
-			IAssetIndexer::FSampleInfo Sample = Indexer.GetSampleInfoRelative(SampleTime, Origin);
-
-			float CurrentTime = Sample.ClipTime;
-			float PreviousTime = CurrentTime - SamplingContext->FiniteDelta;
-
-			FDeltaTimeRecord DeltaTimeRecord;
-			DeltaTimeRecord.Set(PreviousTime, CurrentTime - PreviousTime);
-			FAnimExtractContext ExtractionCtx(CurrentTime, true, DeltaTimeRecord, Sample.Clip->IsLoopable());
-
-			Sample.Clip->ExtractPose(ExtractionCtx, Entry->AnimPoseData);
-
-			if (IndexingContext.bMirrored)
-			{
-				FAnimationRuntime::MirrorPose(
-					Entry->AnimPoseData.GetPose(),
-					IndexingContext.Schema->MirrorDataTable->MirrorAxis,
-					SamplingContext->CompactPoseMirrorBones,
-					SamplingContext->ComponentSpaceRefRotations
-				);
-				// Note curves and attributes are not used during the indexing process and therefore don't need to be mirrored
-			}
-
-			Entry->ComponentSpacePose.InitPose(Entry->Pose);
-			Entry->RootTransform = Sample.RootTransform;
-			Entry->Clamped = Sample.bClamped;
-		}
-
-		const FBoneReference& BoneReference = IndexingContext.Schema->BoneReferences[SchemaBoneIdx];
-		FCompactPoseBoneIndex CompactBoneIndex = SamplingContext->BoneContainer.MakeCompactPoseIndex(FMeshPoseBoneIndex(BoneReference.BoneIndex));
-
-		const FTransform BoneTransform = Entry->ComponentSpacePose.GetComponentSpaceTransform(CompactBoneIndex) * Indexer.MirrorTransform(Entry->RootTransform);
-		Clamped = Entry->Clamped;
-
-		return BoneTransform;
-	}
-
-	const UE::PoseSearch::IAssetIndexer& Indexer;
-
-	struct Entry
-	{
-		float SampleTime;
-		float OriginTime;
-		bool Clamped;
-
-		// @todo: minimize the Entry memory footprint
-		FTransform RootTransform;
-		FCompactPose Pose;
-		FCSPose<FCompactPose> ComponentSpacePose;
-		FBlendedCurve UnusedCurve;
-		UE::Anim::FStackAttributeContainer UnusedAtrribute;
-		FAnimationPoseData AnimPoseData = { Pose, UnusedCurve, UnusedAtrribute };
-	};
-
-	TArray<Entry> Entries;
-};
-
-static void CollectBonePositions(TArray<FVector>& BonePositions, BoneTransformsCache& BoneTransformsCache, int8 SchemaBoneIdx)
-{
-	const FAssetIndexingContext& IndexingContext = BoneTransformsCache.Indexer.GetIndexingContext();
+	const FAssetIndexingContext& IndexingContext = Indexer.GetIndexingContext();
 	const float FiniteDelta = IndexingContext.Schema->SamplingInterval;
 	const float SampleTimeStart = FMath::Min(IndexingContext.BeginSampleIdx * FiniteDelta, IndexingContext.MainSampler->GetPlayLength());
 	const int32 NumSamples = IndexingContext.EndSampleIdx - IndexingContext.BeginSampleIdx;
@@ -154,7 +69,7 @@ static void CollectBonePositions(TArray<FVector>& BonePositions, BoneTransformsC
 	{
 		const float SampleTime = SampleTimeStart + SampleIdx * FiniteDelta;
 		bool Unused;
-		BonePositions[SampleIdx] = BoneTransformsCache.Get(SampleTime, SampleTimeStart, SchemaBoneIdx, Unused).GetTranslation();
+		BonePositions[SampleIdx] = Indexer.GetTransformAndCacheResults(SampleTime, SampleTimeStart, SchemaBoneIdx, Unused).GetTranslation();
 	}
 }
 
@@ -356,6 +271,119 @@ class USkeleton* UPoseSearchFeatureChannel::GetSkeleton(bool& bInvalidSkeletonIs
 	return Schema ? Schema->Skeleton : nullptr;
 }
 
+//////////////////////////////////////////////////////////////////////////
+// UPoseSearchFeatureChannel_Position
+
+void UPoseSearchFeatureChannel_Position::InitializeSchema(UE::PoseSearch::FSchemaInitializer& Initializer)
+{
+	Super::InitializeSchema(Initializer);
+	ChannelCardinality = UE::PoseSearch::FFeatureVectorHelper::EncodeVectorCardinality;
+	Initializer.SetCurrentChannelDataOffset(ChannelDataOffset + ChannelCardinality);
+	SchemaBoneIdx = Initializer.AddBoneReference(Bone);
+}
+
+void UPoseSearchFeatureChannel_Position::FillWeights(TArray<float>& Weights) const
+{
+	using namespace UE::PoseSearch;
+
+	for (int32 i = 0; i != FFeatureVectorHelper::EncodeVectorCardinality; ++i)
+	{
+		Weights[ChannelDataOffset + i] = Weight;
+	}
+}
+
+void UPoseSearchFeatureChannel_Position::IndexAsset(UE::PoseSearch::IAssetIndexer& Indexer, UE::PoseSearch::FAssetIndexingOutput& IndexingOutput) const
+{
+	using namespace UE::PoseSearch;
+
+	const FAssetIndexingContext& IndexingContext = Indexer.GetIndexingContext();
+	const FAssetSamplingContext* SamplingContext = IndexingContext.SamplingContext;
+
+	for (int32 SampleIdx = IndexingContext.BeginSampleIdx; SampleIdx != IndexingContext.EndSampleIdx; ++SampleIdx)
+	{
+		int32 VectorIdx = SampleIdx - IndexingContext.BeginSampleIdx;
+		FPoseSearchFeatureVectorBuilder& FeatureVector = IndexingOutput.PoseVectors[VectorIdx];
+
+		const float OriginSampleTime = FMath::Min(SampleIdx * IndexingContext.Schema->SamplingInterval, IndexingContext.MainSampler->GetPlayLength());
+		const float SubsampleTime = OriginSampleTime + SampleTimeOffset;
+
+		bool ClampedPresent;
+		const FTransform BoneTransformsPresent = Indexer.GetTransformAndCacheResults(SubsampleTime, UseCharacterSpaceVelocities ? SubsampleTime : OriginSampleTime, SchemaBoneIdx, ClampedPresent);
+		int32 DataOffset = ChannelDataOffset;
+		FFeatureVectorHelper::EncodeVector(FeatureVector.EditValues(), DataOffset, BoneTransformsPresent.GetTranslation());
+	}
+}
+
+void UPoseSearchFeatureChannel_Position::GenerateDDCKey(FBlake3& InOutKeyHasher) const
+{
+	InOutKeyHasher.Update(&Bone, sizeof(Bone));
+	InOutKeyHasher.Update(&Weight, sizeof(Weight));
+	InOutKeyHasher.Update(&SampleTimeOffset, sizeof(SampleTimeOffset));
+}
+
+bool UPoseSearchFeatureChannel_Position::BuildQuery(FPoseSearchContext& SearchContext, FPoseSearchFeatureVectorBuilder& InOutQuery) const
+{
+	using namespace UE::PoseSearch;
+
+	const bool bSkip = SearchContext.CurrentResult.IsValid() && SearchContext.CurrentResult.Database->Schema == InOutQuery.GetSchema();
+
+	if (bSkip)
+	{
+		// todo: instead of skipping because the pose should have been copied already (that's currently happening in
+		// UPoseSearchSchema::BuildIndex()), consider making the copy here, but only copy the right values, not the
+		// whole query
+		return true;
+	}
+
+	// @todo: add a layer of caching to TrySamplePose to improve performance
+	if (!SearchContext.History || !SearchContext.History->TrySamplePose(-SampleTimeOffset, InOutQuery.GetSchema()->Skeleton->GetReferenceSkeleton(), InOutQuery.GetSchema()->BoneIndicesWithParents))
+	{
+		return false;
+	}
+
+	const TArrayView<const FTransform> ComponentPose = SearchContext.History->GetComponentPoseSample();
+	const FTransform RootTransform = SearchContext.History->GetRootTransformSample();
+	const FTransform RootTransformPrev = SearchContext.History->GetPrevRootTransformSample();
+	const int32 SkeletonBoneIndex = InOutQuery.GetSchema()->BoneIndices[SchemaBoneIdx];
+	
+	int32 DataOffset = ChannelDataOffset;
+	FFeatureVectorHelper::EncodeVector(InOutQuery.EditValues(), DataOffset, ComponentPose[SkeletonBoneIndex].GetTranslation());
+	
+	return true;
+}
+
+void UPoseSearchFeatureChannel_Position::DebugDraw(const UE::PoseSearch::FDebugDrawParams& DrawParams, TArrayView<const float> PoseVector) const
+{
+	using namespace UE::PoseSearch;
+
+	const UPoseSearchSchema* Schema = DrawParams.GetSchema();
+	check(Schema && Schema->IsValid());
+
+	const float LifeTime = DrawParams.DefaultLifeTime;
+	const uint8 DepthPriority = ESceneDepthPriorityGroup::SDPG_Foreground + 2;
+	const bool bPersistent = EnumHasAnyFlags(DrawParams.Flags, EDebugDrawFlags::Persistent);
+
+	int32 DataOffset = ChannelDataOffset;
+	const FVector BonePos = DrawParams.RootTransform.TransformPosition(FFeatureVectorHelper::DecodeVector(PoseVector, DataOffset));
+
+	const FLinearColor LinearColor = DrawParams.GetColor(this);
+	const FColor Color = LinearColor.ToFColor(true);
+
+	if (EnumHasAnyFlags(DrawParams.Flags, EDebugDrawFlags::DrawFast | EDebugDrawFlags::DrawSearchIndex))
+	{
+		DrawDebugPoint(DrawParams.World, BonePos, DrawParams.PointSize, Color, bPersistent, DrawParams.DefaultLifeTime, DepthPriority);
+	}
+	else
+	{
+		DrawDebugSphere(DrawParams.World, BonePos, DrawDebugSphereSize, DrawDebugSphereSegments, Color, bPersistent, LifeTime, DepthPriority);
+	}
+
+	if (EnumHasAnyFlags(DrawParams.Flags, EDebugDrawFlags::DrawBoneNames))
+	{
+		DrawDebugString(DrawParams.World, BonePos + FVector(0.0, 0.0, 10.0), Schema->BoneReferences[SchemaBoneIdx].BoneName.ToString(), nullptr, Color, LifeTime, false, 1.0f);
+	}
+}
+
 //////////////////////////////////////////////////////////////////////////
 // UPoseSearchFeatureChannel_Pose
 
@@ -409,12 +437,10 @@ void UPoseSearchFeatureChannel_Pose::InitializeSchema(UE::PoseSearch::FSchemaIni
 
 	Initializer.SetCurrentChannelDataOffset(DataOffset);
 
-	FeatureParams.Reset();
+	SchemaBoneIdx.Reset();
 	for (const FPoseSearchBone& Bone : SampledBones)
 	{
-		FPoseSearchPoseFeatureInfo FeatureInfo;
-		FeatureInfo.SchemaBoneIdx = Initializer.AddBoneReference(Bone.Reference);
-		FeatureParams.Add(FeatureInfo);
+		SchemaBoneIdx.Add(Initializer.AddBoneReference(Bone.Reference));
 	}
 }
 
@@ -478,7 +504,7 @@ void UPoseSearchFeatureChannel_Pose::FillWeights(TArray<float>& Weights) const
 }
 
 // @todo: do we really need to use double(s) in all this math?
-void UPoseSearchFeatureChannel_Pose::CalculatePhases(UE::PoseSearch::BoneTransformsCache& BoneTransformsCache, UE::PoseSearch::FAssetIndexingOutput& IndexingOutput, TArray<TArray<FVector2D>>& OutPhases) const
+void UPoseSearchFeatureChannel_Pose::CalculatePhases(UE::PoseSearch::IAssetIndexer& Indexer, UE::PoseSearch::FAssetIndexingOutput& IndexingOutput, TArray<TArray<FVector2D>>& OutPhases) const
 {
 	// @todo: expose them via UI
 	static float BoneSamplingCentralDifferencesTime = 0.2f; // seconds
@@ -489,7 +515,7 @@ void UPoseSearchFeatureChannel_Pose::CalculatePhases(UE::PoseSearch::BoneTransfo
 	OutPhases.Reset();
 	OutPhases.AddDefaulted(SampledBones.Num());
 	
-	const float FiniteDelta = BoneTransformsCache.Indexer.GetIndexingContext().Schema->SamplingInterval;
+	const float FiniteDelta = Indexer.GetIndexingContext().Schema->SamplingInterval;
 
 	TArray<float> Signal;
 	TArray<float> SmoothedSignal;
@@ -500,7 +526,7 @@ void UPoseSearchFeatureChannel_Pose::CalculatePhases(UE::PoseSearch::BoneTransfo
 		const FPoseSearchBone& SampledBone = SampledBones[ChannelBoneIdx];
 		if (SampledBone.bUsePhase)
 		{
-			CollectBonePositions(BonePositions, BoneTransformsCache, FeatureParams[ChannelBoneIdx].SchemaBoneIdx);
+			CollectBonePositions(BonePositions, Indexer, SchemaBoneIdx[ChannelBoneIdx]);
 
 			// @todo: have different way of calculating signals, for example: height of the bone transform, acceleration, etc?
 			const int32 BoneSamplingCentralDifferencesOffset = FMath::Max(FMath::CeilToInt(BoneSamplingCentralDifferencesTime / FiniteDelta), 1);
@@ -519,24 +545,22 @@ void UPoseSearchFeatureChannel_Pose::CalculatePhases(UE::PoseSearch::BoneTransfo
 	}
 }
 
-void UPoseSearchFeatureChannel_Pose::IndexAsset(const UE::PoseSearch::IAssetIndexer& Indexer,  UE::PoseSearch::FAssetIndexingOutput& IndexingOutput) const
+void UPoseSearchFeatureChannel_Pose::IndexAsset(UE::PoseSearch::IAssetIndexer& Indexer,  UE::PoseSearch::FAssetIndexingOutput& IndexingOutput) const
 {
 	using namespace UE::PoseSearch;
 	 
-	UE::PoseSearch::BoneTransformsCache BoneTransformsCache(Indexer);
-
 	// Phases is an array of array with cardinality SampledBones.Num() times NumSamples (IndexingContext.EndSampleIdx - IndexingContext.BeginSampleIdx)
 	// of 2 dimensional vectors (FVector2D) representing phases in an Eucledean space with phase angle sin/cos as direction and certainty of the signal as magnitude,
 	// where certainty is a function of the amplitude of the signal used as input
 	TArray<TArray<FVector2D>> Phases;
-	CalculatePhases(BoneTransformsCache, IndexingOutput, Phases);
+	CalculatePhases(Indexer, IndexingOutput, Phases);
 
 	const FAssetIndexingContext& IndexingContext = Indexer.GetIndexingContext();
 	for (int32 SampleIdx = IndexingContext.BeginSampleIdx; SampleIdx != IndexingContext.EndSampleIdx; ++SampleIdx)
 	{
 		int32 VectorIdx = SampleIdx - IndexingContext.BeginSampleIdx;
 		FPoseSearchFeatureVectorBuilder& FeatureVector = IndexingOutput.PoseVectors[VectorIdx];
-		AddPoseFeatures(BoneTransformsCache, SampleIdx, FeatureVector, Phases);
+		AddPoseFeatures(Indexer, SampleIdx, FeatureVector, Phases);
 	}
 }
 
@@ -580,7 +604,7 @@ void UPoseSearchFeatureChannel_Pose::ComputeMeanDeviations(const Eigen::MatrixXd
 	check(DataOffset == ChannelDataOffset + ChannelCardinality);
 }
 
-void UPoseSearchFeatureChannel_Pose::AddPoseFeatures(UE::PoseSearch::BoneTransformsCache& BoneTransformsCache, int32 SampleIdx, FPoseSearchFeatureVectorBuilder& FeatureVector, const TArray<TArray<FVector2D>>& Phases) const
+void UPoseSearchFeatureChannel_Pose::AddPoseFeatures(UE::PoseSearch::IAssetIndexer& Indexer, int32 SampleIdx, FPoseSearchFeatureVectorBuilder& FeatureVector, const TArray<TArray<FVector2D>>& Phases) const
 {
 	// This function samples the instantaneous pose at time t as well as the pose's velocity and acceleration at time t.
 	// Symmetric finite differences are used to approximate derivatives:
@@ -597,7 +621,7 @@ void UPoseSearchFeatureChannel_Pose::AddPoseFeatures(UE::PoseSearch::BoneTransfo
 		return;
 	}
 
-	const FAssetIndexingContext& IndexingContext = BoneTransformsCache.Indexer.GetIndexingContext();
+	const FAssetIndexingContext& IndexingContext = Indexer.GetIndexingContext();
 	const FAssetSamplingContext* SamplingContext = IndexingContext.SamplingContext;
 
 	const float SampleTime = FMath::Min(SampleIdx * IndexingContext.Schema->SamplingInterval, IndexingContext.MainSampler->GetPlayLength());
@@ -605,7 +629,6 @@ void UPoseSearchFeatureChannel_Pose::AddPoseFeatures(UE::PoseSearch::BoneTransfo
 	int32 DataOffset = ChannelDataOffset;
 	for (int32 ChannelBoneIdx = 0; ChannelBoneIdx != SampledBones.Num(); ++ChannelBoneIdx)
 	{
-		const FPoseSearchPoseFeatureInfo& PoseFeatureInfo = FeatureParams[ChannelBoneIdx];
 		const FPoseSearchBone& SampledBone = SampledBones[ChannelBoneIdx];
 
 		// Get each bone's component transform, velocity, and acceleration and add accumulated root motion at this time offset
@@ -619,7 +642,7 @@ void UPoseSearchFeatureChannel_Pose::AddPoseFeatures(UE::PoseSearch::BoneTransfo
 				const float SubsampleTime = SampleTime + SampleTimes[SubsampleIdx];
 
 				bool ClampedPresent;
-				const FTransform BoneTransformsPresent = BoneTransformsCache.Get(SubsampleTime, UseCharacterSpaceVelocities ? SubsampleTime : SampleTime, PoseFeatureInfo.SchemaBoneIdx, ClampedPresent);
+				const FTransform BoneTransformsPresent = Indexer.GetTransformAndCacheResults(SubsampleTime, UseCharacterSpaceVelocities ? SubsampleTime : SampleTime, SchemaBoneIdx[ChannelBoneIdx], ClampedPresent);
 				FFeatureVectorHelper::EncodeVector(FeatureVector.EditValues(), DataOffset, BoneTransformsPresent.GetTranslation());
 			}
 		}
@@ -631,7 +654,7 @@ void UPoseSearchFeatureChannel_Pose::AddPoseFeatures(UE::PoseSearch::BoneTransfo
 				const float SubsampleTime = SampleTime + SampleTimes[SubsampleIdx];
 
 				bool ClampedPresent;
-				const FTransform BoneTransformsPresent = BoneTransformsCache.Get(SubsampleTime, UseCharacterSpaceVelocities ? SubsampleTime : SampleTime, PoseFeatureInfo.SchemaBoneIdx, ClampedPresent);
+				const FTransform BoneTransformsPresent = Indexer.GetTransformAndCacheResults(SubsampleTime, UseCharacterSpaceVelocities ? SubsampleTime : SampleTime, SchemaBoneIdx[ChannelBoneIdx], ClampedPresent);
 				FFeatureVectorHelper::EncodeQuat(FeatureVector.EditValues(), DataOffset, BoneTransformsPresent.GetRotation());
 			}
 		}
@@ -643,9 +666,9 @@ void UPoseSearchFeatureChannel_Pose::AddPoseFeatures(UE::PoseSearch::BoneTransfo
 				const float SubsampleTime = SampleTime + SampleTimes[SubsampleIdx];
 
 				bool ClampedPast, ClampedPresent, ClampedFuture;
-				const FTransform BoneTransformsPast = BoneTransformsCache.Get(SubsampleTime - SamplingContext->FiniteDelta, UseCharacterSpaceVelocities ? SubsampleTime - SamplingContext->FiniteDelta : SampleTime, PoseFeatureInfo.SchemaBoneIdx, ClampedPast);
-				const FTransform BoneTransformsPresent = BoneTransformsCache.Get(SubsampleTime, UseCharacterSpaceVelocities ? SubsampleTime : SampleTime, PoseFeatureInfo.SchemaBoneIdx, ClampedPresent);
-				const FTransform BoneTransformsFuture = BoneTransformsCache.Get(SubsampleTime + SamplingContext->FiniteDelta, UseCharacterSpaceVelocities ? SubsampleTime + SamplingContext->FiniteDelta : SampleTime, PoseFeatureInfo.SchemaBoneIdx, ClampedFuture);
+				const FTransform BoneTransformsPast = Indexer.GetTransformAndCacheResults(SubsampleTime - SamplingContext->FiniteDelta, UseCharacterSpaceVelocities ? SubsampleTime - SamplingContext->FiniteDelta : SampleTime, SchemaBoneIdx[ChannelBoneIdx], ClampedPast);
+				const FTransform BoneTransformsPresent = Indexer.GetTransformAndCacheResults(SubsampleTime, UseCharacterSpaceVelocities ? SubsampleTime : SampleTime, SchemaBoneIdx[ChannelBoneIdx], ClampedPresent);
+				const FTransform BoneTransformsFuture = Indexer.GetTransformAndCacheResults(SubsampleTime + SamplingContext->FiniteDelta, UseCharacterSpaceVelocities ? SubsampleTime + SamplingContext->FiniteDelta : SampleTime, SchemaBoneIdx[ChannelBoneIdx], ClampedFuture);
 
 				// We can get a better finite difference if we ignore samples that have
 				// been clamped at either side of the clip. However, if the central sample 
@@ -680,29 +703,13 @@ void UPoseSearchFeatureChannel_Pose::AddPoseFeatures(UE::PoseSearch::BoneTransfo
 	}
 }
 
-FFloatRange UPoseSearchFeatureChannel_Pose::GetHorizonRange(EPoseSearchFeatureDomain Domain) const
-{
-	FFloatRange Extent = FFloatRange::Empty();
-	if (Domain == EPoseSearchFeatureDomain::Time)
-	{
-		if (SampleTimes.Num())
-		{
-			Extent = FFloatRange::Inclusive(SampleTimes[0], SampleTimes.Last());
-		}
-	}
-
-	return Extent;
-}
-
 void UPoseSearchFeatureChannel_Pose::GenerateDDCKey(FBlake3& InOutKeyHasher) const
 {
 	InOutKeyHasher.Update(MakeMemoryView(SampledBones));
 	InOutKeyHasher.Update(MakeMemoryView(SampleTimes));
 }
 
-bool UPoseSearchFeatureChannel_Pose::BuildQuery(
-	FPoseSearchContext& SearchContext,
-	FPoseSearchFeatureVectorBuilder& InOutQuery) const
+bool UPoseSearchFeatureChannel_Pose::BuildQuery(FPoseSearchContext& SearchContext, FPoseSearchFeatureVectorBuilder& InOutQuery) const
 {
 	using namespace UE::PoseSearch;
 
@@ -753,8 +760,7 @@ bool UPoseSearchFeatureChannel_Pose::BuildQuery(
 		const FTransform RootTransformPrev = SearchContext.History->GetPrevRootTransformSample();
 		for (int32 SampledBoneIdx = 0; SampledBoneIdx != SampledBones.Num(); ++SampledBoneIdx)
 		{
-			const int32 SchemaBoneIdx = FeatureParams[SampledBoneIdx].SchemaBoneIdx;
-			const int32 SkeletonBoneIndex = InOutQuery.GetSchema()->BoneIndices[SchemaBoneIdx];
+			const int32 SkeletonBoneIndex = InOutQuery.GetSchema()->BoneIndices[SchemaBoneIdx[SampledBoneIdx]];
 			const int32 CachedTransformsIndex = SubsampleIdx * SampledBones.Num() + SampledBoneIdx;
 			CachedTransforms[CachedTransformsIndex].Current = ComponentPose[SkeletonBoneIndex];
 
@@ -768,6 +774,8 @@ bool UPoseSearchFeatureChannel_Pose::BuildQuery(
 				// animation space velocity
 				CachedTransforms[CachedTransformsIndex].Previous = ComponentPrevPose[SkeletonBoneIndex] * (RootTransformPrev * RootTransform.Inverse());
 			}
+
+			CachedTransforms[CachedTransformsIndex].Valid = true;
 		}
 	}
 
@@ -890,10 +898,9 @@ void UPoseSearchFeatureChannel_Pose::DebugDraw(const UE::PoseSearch::FDebugDrawP
 
 				if (EnumHasAnyFlags(DrawParams.Flags, EDebugDrawFlags::DrawBoneNames))
 				{
-					int32 SchemaBoneIdx = this->FeatureParams[ChannelBoneIdx].SchemaBoneIdx;
 					DrawDebugString(
 						DrawParams.World, BonePos[SubsampleIdx] + FVector(0.0, 0.0, 10.0),
-						Schema->BoneReferences[SchemaBoneIdx].BoneName.ToString(),
+						Schema->BoneReferences[SchemaBoneIdx[ChannelBoneIdx]].BoneName.ToString(),
 						nullptr, Color, LifeTime, false, 1.0f);
 				}
 			}
@@ -1037,7 +1044,7 @@ void UPoseSearchFeatureChannel_Trajectory::FillWeights(TArray<float>& Weights) c
 	}
 }
 
-void UPoseSearchFeatureChannel_Trajectory::IndexAsset(const UE::PoseSearch::IAssetIndexer& Indexer,  UE::PoseSearch::FAssetIndexingOutput& IndexingOutput) const
+void UPoseSearchFeatureChannel_Trajectory::IndexAsset(UE::PoseSearch::IAssetIndexer& Indexer,  UE::PoseSearch::FAssetIndexingOutput& IndexingOutput) const
 {
 	const UE::PoseSearch::FAssetIndexingContext& IndexingContext = Indexer.GetIndexingContext();
 	for (int32 SampleIdx = IndexingContext.BeginSampleIdx; SampleIdx != IndexingContext.EndSampleIdx; ++SampleIdx)
@@ -1173,19 +1180,6 @@ void UPoseSearchFeatureChannel_Trajectory::IndexAssetPrivate(const UE::PoseSearc
 	}
 }
 
-FFloatRange UPoseSearchFeatureChannel_Trajectory::GetHorizonRange(EPoseSearchFeatureDomain InDomain) const
-{
-	FFloatRange Extent = FFloatRange::Empty();
-	if (InDomain == Domain)
-	{
-		if (SampleOffsets.Num())
-		{
-			Extent = FFloatRange::Inclusive(SampleOffsets[0], SampleOffsets.Last());
-		}
-	}
-	return Extent;
-}
-
 void UPoseSearchFeatureChannel_Trajectory::GenerateDDCKey(FBlake3& InOutKeyHasher) const
 {
 	InOutKeyHasher.Update(&bUseLinearVelocities, sizeof(bUseLinearVelocities));
diff --git a/Engine/Plugins/Experimental/Animation/PoseSearch/Source/Runtime/Public/PoseSearch/PoseSearch.h b/Engine/Plugins/Experimental/Animation/PoseSearch/Source/Runtime/Public/PoseSearch/PoseSearch.h
--- a/Engine/Plugins/Experimental/Animation/PoseSearch/Source/Runtime/Public/PoseSearch/PoseSearch.h
+++ b/Engine/Plugins/Experimental/Animation/PoseSearch/Source/Runtime/Public/PoseSearch/PoseSearch.h
@@ -219,7 +219,6 @@ struct POSESEARCH_API FAssetSamplingContext
 	// Only initialized and used when a mirroring table is specified
 	TCustomBoneIndexArray<FQuat, FCompactPoseBoneIndex> ComponentSpaceRefRotations;
 
-
 	void Init(const UPoseSearchSchema* Schema);
 	FTransform MirrorTransform(const FTransform& Transform) const;
 };
@@ -310,6 +309,7 @@ public:
 	virtual FSampleInfo GetSampleInfoRelative(float SampleTime, const FSampleInfo& Origin) const = 0;
 	virtual const float GetSampleTimeFromDistance(float Distance) const = 0;
 	virtual FTransform MirrorTransform(const FTransform& Transform) const = 0;
+	virtual FTransform GetTransformAndCacheResults(float SampleTime, float OriginTime, int8 SchemaBoneIdx, bool& Clamped) = 0;
 };
 
 } // namespace UE::PoseSearch
@@ -324,9 +324,7 @@ class POSESEARCH_API UPoseSearchFeatureChannel : public UObject, public IBoneRef
 	GENERATED_BODY()
 
 public:
-
 	int32 GetChannelIndex() const { checkSlow(ChannelIdx >= 0); return ChannelIdx; }
-
 	int32 GetChannelCardinality() const { checkSlow(ChannelCardinality >= 0); return ChannelCardinality; }
 	int32 GetChannelDataOffset() const { checkSlow(ChannelDataOffset >= 0); return ChannelDataOffset; }
 
@@ -336,22 +334,15 @@ public:
 	virtual void FillWeights(TArray<float>& Weights) const PURE_VIRTUAL(UPoseSearchFeatureChannel::FillWeights, );
 
 	// Called at database build time to populate pose vectors with this channel's data
-	virtual void IndexAsset(const UE::PoseSearch::IAssetIndexer& Indexer, UE::PoseSearch::FAssetIndexingOutput& IndexingOutput) const PURE_VIRTUAL(UPoseSearchFeatureChannel::IndexAsset, );
+	virtual void IndexAsset(UE::PoseSearch::IAssetIndexer& Indexer, UE::PoseSearch::FAssetIndexingOutput& IndexingOutput) const PURE_VIRTUAL(UPoseSearchFeatureChannel::IndexAsset, );
 
 	virtual void ComputeMeanDeviations(const Eigen::MatrixXd& CenteredPoseMatrix, Eigen::VectorXd& MeanDeviations) const;
 
-	// Return this channel's range of sampling offsets in the requested sampling domain.
-	// Returns empty range if the channel has no horizon in the requested domain.
-	virtual FFloatRange GetHorizonRange(EPoseSearchFeatureDomain Domain) const PURE_VIRTUAL(UPoseSearchFeatureChannel::GetHorizonRange, return FFloatRange::Empty(); );
-
 	// Hash channel properties to produce a key for database derived data
 	virtual void GenerateDDCKey(FBlake3& InOutKeyHasher) const PURE_VIRTUAL(UPoseSearchFeatureChannel::GenerateDDCKey, );
 
 	// Called at runtime to add this channel's data to the query pose vector
-	virtual bool BuildQuery(
-		FPoseSearchContext& SearchContext,
-		FPoseSearchFeatureVectorBuilder& InOutQuery) const
-		PURE_VIRTUAL(UPoseSearchFeatureChannel::BuildQuery, return false;);
+	virtual bool BuildQuery(FPoseSearchContext& SearchContext, FPoseSearchFeatureVectorBuilder& InOutQuery) const PURE_VIRTUAL(UPoseSearchFeatureChannel::BuildQuery, return false;);
 
 	// Draw this channel's data for the given pose vector
 	virtual void DebugDraw(const UE::PoseSearch::FDebugDrawParams& DrawParams, TArrayView<const float> PoseVector) const PURE_VIRTUAL(UPoseSearchFeatureChannel::DebugDraw, );
@@ -361,8 +352,6 @@ private:
 	// Note this function is exclusively for FBoneReference details customization
 	class USkeleton* GetSkeleton(bool& bInvalidSkeletonIsError, const IPropertyHandle* PropertyHandle) override;
 
-private:
-
 	friend class ::UPoseSearchSchema;
 
 	UPROPERTY()
@@ -399,7 +388,6 @@ private:
 
 	int32 CurrentChannelIdx = 0;
 	int32 CurrentChannelDataOffset = 0;
-	
 	TArray<FBoneReference> BoneReferences;
 };
 
@@ -483,10 +471,6 @@ public:
 
 	int32 GetNumBones () const { return BoneIndices.Num(); }
 
-	// Returns global range of sampling offsets among all channels in requested sampling domain.
-	// Returns empty range if the channel has no horizon in the requested domain.
-	virtual FFloatRange GetHorizonRange(EPoseSearchFeatureDomain Domain) const;
-
 public: // UObject
 	virtual void PreSave(FObjectPreSaveContext ObjectSaveContext) override;
 	virtual void PostLoad() override;
@@ -496,6 +480,7 @@ public: // IBoneReferenceSkeletonProvider
 
 #if WITH_EDITOR
 	virtual void PostEditChangeProperty(struct FPropertyChangedEvent& PropertyChangedEvent) override;
+	void GenerateDDCKey(FBlake3& InOutKeyHasher) const;
 #endif
 
 private:
@@ -1099,6 +1084,15 @@ public:
 	const FGameplayTagContainer* GetSourceAssetGroupTags(const FPoseSearchIndexAsset* SearchIndexAsset) const;
 	const FString GetSourceAssetName(const FPoseSearchIndexAsset* SearchIndexAsset) const;
 	int32 GetNumberOfPrincipalComponents() const;
+	
+#if WITH_EDITOR
+	void GenerateDDCKey(FBlake3& InOutKeyHasher) const;
+private:
+	static void AddDbSequenceToWriter(const FPoseSearchDatabaseSequence& DbSequence, FBlake3& InOutWriter);
+	static void AddRawSequenceToWriter(const UAnimSequence* Sequence, FBlake3& InOutWriter);
+	static void AddPoseSearchNotifiesToWriter(const UAnimSequence* Sequence, FBlake3& InOutWriter);
+	static void AddDbBlendSpaceToWriter(const FPoseSearchDatabaseBlendSpace& DbBlendSpace, FBlake3& InOutWriter);
+#endif // WITH_EDITOR
 
 public: // UObject
 	virtual void PostLoad() override;
diff --git a/Engine/Plugins/Experimental/Animation/PoseSearch/Source/Runtime/Public/PoseSearch/PoseSearchFeatureChannels.h b/Engine/Plugins/Experimental/Animation/PoseSearch/Source/Runtime/Public/PoseSearch/PoseSearchFeatureChannels.h
--- a/Engine/Plugins/Experimental/Animation/PoseSearch/Source/Runtime/Public/PoseSearch/PoseSearchFeatureChannels.h
+++ b/Engine/Plugins/Experimental/Animation/PoseSearch/Source/Runtime/Public/PoseSearch/PoseSearchFeatureChannels.h
@@ -8,18 +8,33 @@
 
 #include "PoseSearchFeatureChannels.generated.h"
 
-namespace UE::PoseSearch
-{
-	struct BoneTransformsCache;
-}
-
-USTRUCT()
-struct FPoseSearchPoseFeatureInfo
+//////////////////////////////////////////////////////////////////////////
+// UPoseSearchFeatureChannel_Position
+UCLASS(BlueprintType, EditInlineNew)
+class POSESEARCH_API UPoseSearchFeatureChannel_Position : public UPoseSearchFeatureChannel
 {
 	GENERATED_BODY()
 
+public:
+	UPROPERTY(EditAnywhere, Category = "Settings")
+	FBoneReference Bone;
+
+	UPROPERTY(EditAnywhere, Category = "Settings")
+	float Weight = 1.f;
+
+	UPROPERTY(EditAnywhere, Category = "Settings")
+	float SampleTimeOffset = 0.f;
+
 	UPROPERTY()
-	int8 SchemaBoneIdx = 0;
+	int8 SchemaBoneIdx;
+
+	// UPoseSearchFeatureChannel interface
+	virtual void InitializeSchema(UE::PoseSearch::FSchemaInitializer& Initializer) override;
+	virtual void FillWeights(TArray<float>& Weights) const override;
+	virtual void IndexAsset(UE::PoseSearch::IAssetIndexer& Indexer, UE::PoseSearch::FAssetIndexingOutput& IndexingOutput) const override;
+	virtual void GenerateDDCKey(FBlake3& InOutKeyHasher) const override;
+	virtual bool BuildQuery(FPoseSearchContext& SearchContext, FPoseSearchFeatureVectorBuilder& InOutQuery) const override;
+	virtual void DebugDraw(const UE::PoseSearch::FDebugDrawParams& DrawParams, TArrayView<const float> PoseVector) const override;
 };
 
 //////////////////////////////////////////////////////////////////////////
@@ -30,9 +45,6 @@ class POSESEARCH_API UPoseSearchFeatureChannel_Pose : public UPoseSearchFeatureC
 	GENERATED_BODY()
 
 public:
-
-	virtual ~UPoseSearchFeatureChannel_Pose() {}
-
 	UPROPERTY(EditAnywhere, Category = "Settings")
 	TArray<FPoseSearchBone> SampledBones;
 
@@ -40,7 +52,7 @@ public:
 	TArray<float> SampleTimes;
 
 	UPROPERTY()
-	TArray<FPoseSearchPoseFeatureInfo> FeatureParams;
+	TArray<int8> SchemaBoneIdx;
 
 	// UObject interface
 	virtual void PreSave(FObjectPreSaveContext ObjectSaveContext) override;
@@ -48,22 +60,17 @@ public:
 	// UPoseSearchFeatureChannel interface
 	virtual void InitializeSchema(UE::PoseSearch::FSchemaInitializer& Initializer) override;
 	virtual void FillWeights(TArray<float>& Weights) const override;
-	virtual void IndexAsset(const UE::PoseSearch::IAssetIndexer& Indexer, UE::PoseSearch::FAssetIndexingOutput& IndexingOutput) const override;
+	virtual void IndexAsset(UE::PoseSearch::IAssetIndexer& Indexer, UE::PoseSearch::FAssetIndexingOutput& IndexingOutput) const override;
 	virtual void ComputeMeanDeviations(const Eigen::MatrixXd& CenteredPoseMatrix, Eigen::VectorXd& MeanDeviations) const override;
-	virtual FFloatRange GetHorizonRange(EPoseSearchFeatureDomain Domain) const override;
 	virtual void GenerateDDCKey(FBlake3& InOutKeyHasher) const override;
-	virtual bool BuildQuery(
-		FPoseSearchContext& SearchContext,
-		FPoseSearchFeatureVectorBuilder& InOutQuery) const override;
+	virtual bool BuildQuery(FPoseSearchContext& SearchContext, FPoseSearchFeatureVectorBuilder& InOutQuery) const override;
 	virtual void DebugDraw(const UE::PoseSearch::FDebugDrawParams& DrawParams, TArrayView<const float> PoseVector) const override;
 
 protected:
-	void AddPoseFeatures(UE::PoseSearch::BoneTransformsCache& BoneTransformsCache, int32 SampleIdx, FPoseSearchFeatureVectorBuilder& FeatureVector, const TArray<TArray<FVector2D>>& Phases) const;
-	void CalculatePhases(UE::PoseSearch::BoneTransformsCache& BoneTransformsCache, UE::PoseSearch::FAssetIndexingOutput& IndexingOutput, TArray<TArray<FVector2D>>& OutPhases) const;
+	void AddPoseFeatures(UE::PoseSearch::IAssetIndexer& Indexer, int32 SampleIdx, FPoseSearchFeatureVectorBuilder& FeatureVector, const TArray<TArray<FVector2D>>& Phases) const;
+	void CalculatePhases(UE::PoseSearch::IAssetIndexer& Indexer, UE::PoseSearch::FAssetIndexingOutput& IndexingOutput, TArray<TArray<FVector2D>>& OutPhases) const;
 };
 
-
-
 //////////////////////////////////////////////////////////////////////////
 // UPoseSearchFeatureChannel_Trajectory
 UCLASS(BlueprintType, EditInlineNew)
@@ -72,9 +79,6 @@ class POSESEARCH_API UPoseSearchFeatureChannel_Trajectory : public UPoseSearchFe
 	GENERATED_BODY()
 
 public:
-
-	virtual ~UPoseSearchFeatureChannel_Trajectory() {}
-
 	UPROPERTY(EditAnywhere, Category = "Settings")
 	bool bUseLinearVelocities = true;
 
@@ -100,16 +104,13 @@ public:
 	// UPoseSearchFeatureChannel interface
 	virtual void InitializeSchema(UE::PoseSearch::FSchemaInitializer& Initializer) override;
 	virtual void FillWeights(TArray<float>& Weights) const override;
-	virtual void IndexAsset(const UE::PoseSearch::IAssetIndexer& Indexer, UE::PoseSearch::FAssetIndexingOutput& IndexingOutput) const override;
+	virtual void IndexAsset(UE::PoseSearch::IAssetIndexer& Indexer, UE::PoseSearch::FAssetIndexingOutput& IndexingOutput) const override;
 	virtual void ComputeMeanDeviations(const Eigen::MatrixXd& CenteredPoseMatrix, Eigen::VectorXd& MeanDeviations) const override;
-	virtual FFloatRange GetHorizonRange(EPoseSearchFeatureDomain Domain) const override;
 	virtual void GenerateDDCKey(FBlake3& InOutKeyHasher) const override;
-	virtual bool BuildQuery(
-		FPoseSearchContext& SearchContext,
-		FPoseSearchFeatureVectorBuilder& InOutQuery) const override;
+	virtual bool BuildQuery(FPoseSearchContext& SearchContext, FPoseSearchFeatureVectorBuilder& InOutQuery) const override;
 	virtual void DebugDraw(const UE::PoseSearch::FDebugDrawParams& DrawParams, TArrayView<const float> PoseVector) const override;
 
 protected:
-	virtual void IndexAssetPrivate(const UE::PoseSearch::IAssetIndexer& Indexer, int32 SampleIdx, FPoseSearchFeatureVectorBuilder& FeatureVector) const;
+	void IndexAssetPrivate(const UE::PoseSearch::IAssetIndexer& Indexer, int32 SampleIdx, FPoseSearchFeatureVectorBuilder& FeatureVector) const;
 	float GetSampleTime(const UE::PoseSearch::IAssetIndexer& Indexer, int32 SubsampleIdx, float SampleTime, float RootDistance) const;
 };
\ No newline at end of file
