motion matching - caching component space transforms retrieved the from the FPoseHistory via via FPoseSearchContext::TryGetTransformAndCacheResults to speedup UPoseSearchFeatureChannel::BuildQuery implementations

diff --git a/Engine/Plugins/Experimental/Animation/PoseSearch/Source/Editor/Private/PoseSearchDebugger.cpp b/Engine/Plugins/Experimental/Animation/PoseSearch/Source/Editor/Private/PoseSearchDebugger.cpp
--- a/Engine/Plugins/Experimental/Animation/PoseSearch/Source/Editor/Private/PoseSearchDebugger.cpp
+++ b/Engine/Plugins/Experimental/Animation/PoseSearch/Source/Editor/Private/PoseSearchDebugger.cpp
@@ -943,7 +943,7 @@ void SDebuggerDatabaseView::UpdateRows(const FTraceMotionMatchingStateMessage& S
 	check(ActiveView.Rows.Num() == 1);
 	check(ContinuingPoseView.Rows.Num() == 1);
 
-	FPoseSearchContext SearchContext;
+	FSearchContext SearchContext;
 	if (const FPoseSearchIndexAsset* CurrentIndexAsset = Database.GetSearchIndex()->FindAssetForPose(State.DbPoseIdx))
 	{
 		SearchContext.QueryMirrorRequest = CurrentIndexAsset->bMirrored ? 
diff --git a/Engine/Plugins/Experimental/Animation/PoseSearch/Source/Runtime/Private/PoseSearch.cpp b/Engine/Plugins/Experimental/Animation/PoseSearch/Source/Runtime/Private/PoseSearch.cpp
--- a/Engine/Plugins/Experimental/Animation/PoseSearch/Source/Runtime/Private/PoseSearch.cpp
+++ b/Engine/Plugins/Experimental/Animation/PoseSearch/Source/Runtime/Private/PoseSearch.cpp
@@ -174,7 +174,7 @@ static void CalcChannelCosts(const UPoseSearchSchema* Schema, TArrayView<const f
 	}
 }
 
-static int32 PopulateNonSelectableIdx(size_t* NonSelectableIdx, int32 NonSelectableIdxSize, const FPoseSearchContext& SearchContext, int32 IdxOffset = 0)
+static int32 PopulateNonSelectableIdx(size_t* NonSelectableIdx, int32 NonSelectableIdxSize, const FSearchContext& SearchContext, int32 IdxOffset = 0)
 {
 	check(SearchContext.CurrentResult.IsValid());
 
@@ -486,9 +486,7 @@ void UPoseSearchSchema::ResolveBoneReferences()
 }
 
 
-bool UPoseSearchSchema::BuildQuery(
-	FPoseSearchContext& SearchContext,
-	FPoseSearchFeatureVectorBuilder& InOutQuery) const
+bool UPoseSearchSchema::BuildQuery(UE::PoseSearch::FSearchContext& SearchContext, FPoseSearchFeatureVectorBuilder& InOutQuery) const
 {
 	QUICK_SCOPE_CYCLE_COUNTER(STAT_PoseSearch_BuildQuery);
 
@@ -726,7 +724,7 @@ bool UPoseSearchSequenceMetaData::IsValidForSearch() const
 	return IsValidForIndexing() && SearchIndex.IsValid() && !SearchIndex.IsEmpty();
 }
 
-UE::PoseSearch::FSearchResult UPoseSearchSequenceMetaData::Search(FPoseSearchContext& SearchContext) const
+UE::PoseSearch::FSearchResult UPoseSearchSequenceMetaData::Search(UE::PoseSearch::FSearchContext& SearchContext) const
 {
 	using namespace UE::PoseSearch;
 
@@ -1689,11 +1687,7 @@ bool UPoseSearchDatabase::IsCachedCookedPlatformDataLoaded(const ITargetPlatform
 
 #endif // WITH_EDITOR
 
-FPoseSearchCost UPoseSearchDatabase::ComparePoses(
-	FPoseSearchContext& SearchContext,
-	int32 PoseIdx,
-	int32 GroupIdx,
-	const TArrayView<const float>& QueryValues) const
+FPoseSearchCost UPoseSearchDatabase::ComparePoses(UE::PoseSearch::FSearchContext& SearchContext, int32 PoseIdx, int32 GroupIdx, const TArrayView<const float>& QueryValues) const
 {
 	using namespace UE::PoseSearch;
 
@@ -1736,11 +1730,7 @@ FPoseSearchCost UPoseSearchDatabase::ComparePoses(
 	return Result;
 }
 
-FPoseSearchCost UPoseSearchDatabase::ComparePoses(
-	FPoseSearchContext& SearchContext,
-	int32 PoseIdx,
-	const TArrayView<const float>& QueryValues,
-	UE::PoseSearch::FPoseCostDetails& OutPoseCostDetails) const
+FPoseSearchCost UPoseSearchDatabase::ComparePoses(UE::PoseSearch::FSearchContext& SearchContext, int32 PoseIdx, const TArrayView<const float>& QueryValues, UE::PoseSearch::FPoseCostDetails& OutPoseCostDetails) const
 {
 	using namespace Eigen;
 	using namespace UE::PoseSearch;
@@ -1807,11 +1797,7 @@ FPoseSearchCost UPoseSearchDatabase::ComparePoses(
 	return Result;
 }
 
-void UPoseSearchDatabase::ComputePoseCostAddends(
-	int32 PoseIdx,
-	FPoseSearchContext& SearchContext,
-	float& OutNotifyAddend,
-	float& OutMirrorMismatchAddend) const
+void UPoseSearchDatabase::ComputePoseCostAddends(int32 PoseIdx, UE::PoseSearch::FSearchContext& SearchContext, float& OutNotifyAddend, float& OutMirrorMismatchAddend) const
 {
 	OutNotifyAddend = 0.0f;
 	OutMirrorMismatchAddend = 0.0f;
@@ -1832,7 +1818,7 @@ void UPoseSearchDatabase::ComputePoseCostAddends(
 	OutNotifyAddend = PoseMetadata.CostAddend;
 }
 
-UE::PoseSearch::FSearchResult UPoseSearchDatabase::Search(FPoseSearchContext& SearchContext) const
+UE::PoseSearch::FSearchResult UPoseSearchDatabase::Search(UE::PoseSearch::FSearchContext& SearchContext) const
 {
 	using namespace UE::PoseSearch;
 
@@ -1910,7 +1896,7 @@ UE::PoseSearch::FSearchResult UPoseSearchDatabase::Search(FPoseSearchContext& Se
 	return Result;
 }
 
-UE::PoseSearch::FSearchResult UPoseSearchDatabase::SearchPCAKDTree(FPoseSearchContext& SearchContext) const
+UE::PoseSearch::FSearchResult UPoseSearchDatabase::SearchPCAKDTree(UE::PoseSearch::FSearchContext& SearchContext) const
 {
 	QUICK_SCOPE_CYCLE_COUNTER(STAT_PoseSearch_PCA_KNN);
 	SCOPE_CYCLE_COUNTER(STAT_PoseSearchPCAKNN);
@@ -2051,7 +2037,7 @@ UE::PoseSearch::FSearchResult UPoseSearchDatabase::SearchPCAKDTree(FPoseSearchCo
 	return Result;
 }
 
-UE::PoseSearch::FSearchResult UPoseSearchDatabase::SearchBruteForce(FPoseSearchContext& SearchContext) const
+UE::PoseSearch::FSearchResult UPoseSearchDatabase::SearchBruteForce(UE::PoseSearch::FSearchContext& SearchContext) const
 {
 	QUICK_SCOPE_CYCLE_COUNTER(STAT_PoseSearch_Brute_Force);
 	SCOPE_CYCLE_COUNTER(STAT_PoseSearchBruteForce);
@@ -2138,15 +2124,13 @@ UE::PoseSearch::FSearchResult UPoseSearchDatabase::SearchBruteForce(FPoseSearchC
 	return Result;
 }
 
-void UPoseSearchDatabase::BuildQuery(
-	FPoseSearchContext& SearchContext, 
-	FPoseSearchFeatureVectorBuilder& OutQuery) const
+void UPoseSearchDatabase::BuildQuery(UE::PoseSearch::FSearchContext& SearchContext, FPoseSearchFeatureVectorBuilder& OutQuery) const
 {
 	Schema->BuildQuery(SearchContext, OutQuery);
 	OutQuery.Normalize(*GetSearchIndex());
 }
 
-UE::PoseSearch::FSearchResult UPoseSearchDatabaseSet::Search(FPoseSearchContext& SearchContext) const
+UE::PoseSearch::FSearchResult UPoseSearchDatabaseSet::Search(UE::PoseSearch::FSearchContext& SearchContext) const
 {
 	using namespace UE::PoseSearch;
 
@@ -2284,8 +2268,73 @@ void FPoseSearchFeatureVectorBuilder::Normalize(const FPoseSearchIndex& ForSearc
 	ForSearchIndex.Normalize(ValuesNormalized);
 }
 
+namespace UE { namespace PoseSearch
+{
 
-namespace UE { namespace PoseSearch {
+FTransform FSearchContext::TryGetTransformAndCacheResults(float SampleTime, const UPoseSearchSchema* Schema, int8 SchemaBoneIdx, bool& Error)
+{
+	check(History && Schema);
+
+	const FBoneIndexType BoneIndexType = SchemaBoneIdx >= 0 ? Schema->BoneIndices[SchemaBoneIdx] : RoolBoneIdx;
+
+	// @todo: use an hashmap if we end up having too many entries
+	const CachedEntry* Entry = CachedEntries.FindByPredicate([SampleTime, BoneIndexType](const FSearchContext::CachedEntry& Entry)
+	{
+		return Entry.SampleTime == SampleTime && Entry.BoneIndexType == BoneIndexType;
+	});
+
+	if (Entry)
+	{
+		Error = false;
+		return Entry->Transform;
+	}
+
+	if (BoneIndexType >= 0)
+	{
+		TArray<FTransform> SampledLocalPose;
+		if (History->TrySampleLocalPose(-SampleTime, &Schema->BoneIndicesWithParents, &SampledLocalPose, nullptr))
+		{
+			TArray<FTransform> SampledComponentPose;
+			FAnimationRuntime::FillUpComponentSpaceTransforms(Schema->Skeleton->GetReferenceSkeleton(), SampledLocalPose, SampledComponentPose);
+
+			// adding bunch of entries, without caring about adding eventual duplicates
+			for (const FBoneIndexType NewEntryBoneIndexType : Schema->BoneIndicesWithParents)
+			{
+				// @todo: maybe add them with a single allocation with CachedEntries.AddDefaulted(Schema->BoneIndicesWithParents.Num())
+				CachedEntry& NewEntry = CachedEntries[CachedEntries.AddDefaulted()];
+				NewEntry.SampleTime = SampleTime;
+				NewEntry.BoneIndexType = NewEntryBoneIndexType;
+				NewEntry.Transform = SampledComponentPose[NewEntryBoneIndexType];
+			}
+
+			Error = false;
+			return SampledComponentPose[BoneIndexType];
+		}
+
+		Error = true;
+		return FTransform::Identity;
+	}
+	
+	FTransform SampledRootTransform;
+	if (History->TrySampleLocalPose(-SampleTime, nullptr, nullptr, &SampledRootTransform))
+	{
+		CachedEntry& NewEntry = CachedEntries[CachedEntries.AddDefaulted()];
+		NewEntry.SampleTime = SampleTime;
+		NewEntry.BoneIndexType = BoneIndexType;
+		NewEntry.Transform = SampledRootTransform;
+
+		Error = false;
+		return SampledRootTransform;
+	}
+	
+	Error = true;
+	return FTransform::Identity;
+}
+
+void FSearchContext::ClearCachedEntries()
+{
+	CachedEntries.Reset();
+}
 
 //////////////////////////////////////////////////////////////////////////
 // FPoseHistory
@@ -2316,118 +2365,101 @@ static void CopyCompactToSkeletonPose(const FCompactPose& Pose, TArray<FTransfor
 void FPoseHistory::Init(int32 InNumPoses, float InTimeHorizon)
 {
 	Poses.Reserve(InNumPoses);
-	Knots.Reserve(InNumPoses);
 	TimeHorizon = InTimeHorizon;
 }
 
 void FPoseHistory::Init(const FPoseHistory& History)
 {
 	Poses = History.Poses;
-	Knots = History.Knots;
 	TimeHorizon = History.TimeHorizon;
 }
 
-bool FPoseHistory::TrySampleLocalPose(float SecondsAgo, const TArray<FBoneIndexType>& RequiredBones, TArray<FTransform>& LocalPose, FTransform& RootTransform) const
+bool FPoseHistory::TrySampleLocalPose(float SecondsAgo, const TArray<FBoneIndexType>* RequiredBones, TArray<FTransform>* LocalPose, FTransform* RootTransform) const
 {
-	int32 NextIdx = LowerBound(Knots.begin(), Knots.end(), SecondsAgo, TGreater<>());
-	if (NextIdx <= 0 || NextIdx >= Knots.Num())
+	const int32 NextIdx = LowerBound(Poses.begin(), Poses.end(), SecondsAgo, [](const FPose& Pose, float Value)
 	{
+		return Value < Pose.Time;
+	});
+	if (NextIdx <= 0 || NextIdx >= Poses.Num())
+	{
+		// We may not have accumulated enough poses yet
 		return false;
 	}
 
-	int32 PrevIdx = NextIdx - 1;
+	const int32 PrevIdx = NextIdx - 1;
 
 	const FPose& PrevPose = Poses[PrevIdx];
 	const FPose& NextPose = Poses[NextIdx];
 
-	// Compute alpha between previous and next knots
-	float Alpha = FMath::GetMappedRangeValueUnclamped(
-		FVector2f(Knots[PrevIdx], Knots[NextIdx]),
-		FVector2f(0.0f, 1.0f),
-		SecondsAgo);
-
-	// We may not have accumulated enough poses yet
-	if (PrevPose.LocalTransforms.Num() != NextPose.LocalTransforms.Num())
+#if DO_CHECK
+	check(PrevPose.LocalTransforms.Num() == NextPose.LocalTransforms.Num());
+	FBoneIndexType MaxBoneIndexType = 0;
+	if (RequiredBones)
 	{
-		return false;
+		for (FBoneIndexType BoneIndexType : *RequiredBones)
+		{
+			if (BoneIndexType > MaxBoneIndexType)
+			{
+				MaxBoneIndexType = BoneIndexType;
+			}
+		}
+		check(MaxBoneIndexType < PrevPose.LocalTransforms.Num());
 	}
+#endif
+	// Compute alpha between previous and next Poses
+	const float Alpha = FMath::GetMappedRangeValueUnclamped(FVector2f(PrevPose.Time, NextPose.Time), FVector2f(0.0f, 1.0f), SecondsAgo);
 
-	if (RequiredBones.Num() > PrevPose.LocalTransforms.Num())
+	// Lerp between poses by alpha to produce output local pose at requested sample time
+	if (LocalPose)
 	{
-		return false;
+		check(RequiredBones);
+		*LocalPose = PrevPose.LocalTransforms;
+		FAnimationRuntime::LerpBoneTransforms(*LocalPose, NextPose.LocalTransforms, Alpha, *RequiredBones);
 	}
 
-	// Lerp between poses by alpha to produce output local pose at requested sample time
-	LocalPose = PrevPose.LocalTransforms;
-	FAnimationRuntime::LerpBoneTransforms(
-		LocalPose,
-		NextPose.LocalTransforms,
-		Alpha,
-		RequiredBones);
-
-	RootTransform.Blend(PrevPose.RootTransform, NextPose.RootTransform, Alpha);
-
-	return true;
-}
-
-bool FPoseHistory::TrySamplePose(float SecondsAgo, const FReferenceSkeleton& RefSkeleton, const TArray<FBoneIndexType>& RequiredBones)
-{
-	// Compute local space pose at requested time
-	bool bSampled = TrySampleLocalPose(SecondsAgo, RequiredBones, SampledLocalPose, SampledRootTransform);
-
-	// Compute local space pose one sample interval in the past
-	bSampled = bSampled && TrySampleLocalPose(SecondsAgo + GetSampleTimeInterval(), RequiredBones, SampledPrevLocalPose, SampledPrevRootTransform);
-
-	// Convert local to component space
-	if (bSampled)
+	if (RootTransform)
 	{
-		FAnimationRuntime::FillUpComponentSpaceTransforms(RefSkeleton, SampledLocalPose, SampledComponentPose);
-		FAnimationRuntime::FillUpComponentSpaceTransforms(RefSkeleton, SampledPrevLocalPose, SampledPrevComponentPose);
+		RootTransform->Blend(PrevPose.RootTransform, NextPose.RootTransform, Alpha);
 	}
-
-	return bSampled;
+	return true;
 }
 
 bool FPoseHistory::Update(float SecondsElapsed, const FPoseContext& PoseContext, FTransform ComponentTransform, FText* OutError, ERootUpdateMode UpdateMode)
 {
 	// Age our elapsed times
-	for (float& Knot : Knots)
+	for (FPose& Pose : Poses)
 	{
-		Knot += SecondsElapsed;
+		Pose.Time += SecondsElapsed;
 	}
 
-	if (Knots.Num() != Knots.Max())
+	if (Poses.Num() != Poses.Max())
 	{
 		// Consume every pose until the queue is full
-		Knots.AddUninitialized();
 		Poses.Emplace();
 	}
 	else
 	{
-		// Exercise pose retention policy. We must guarantee there is always one additional knot
+		// Exercise pose retention policy. We must guarantee there is always one additional pose
 		// beyond the time horizon so we can compute derivatives at the time horizon. We also
-		// want to evenly distribute knots across the entire history buffer so we only push additional
+		// want to evenly distribute poses across the entire history buffer so we only push additional
 		// poses when enough time has elapsed.
 
 		const float SampleInterval = GetSampleTimeInterval();
 
-		bool bCanEvictOldest = Knots[1] >= TimeHorizon + SampleInterval;
-		bool bShouldPushNewest = Knots[Knots.Num() - 2] >= SampleInterval;
+		bool bCanEvictOldest = Poses[1].Time >= TimeHorizon + SampleInterval;
+		bool bShouldPushNewest = Poses[Poses.Num() - 2].Time >= SampleInterval;
 
 		if (bCanEvictOldest && bShouldPushNewest)
 		{
 			FPose PoseTemp = MoveTemp(Poses.First());
 			Poses.PopFront();
 			Poses.Emplace(MoveTemp(PoseTemp));
-
-			Knots.PopFront();
-			Knots.AddUninitialized();
 		}
 	}
 
 	// Regardless of the retention policy, we always update the most recent pose
-	Knots.Last() = 0.0f;
 	FPose& CurrentPose = Poses.Last();
+	CurrentPose.Time = 0.f;
 	CopyCompactToSkeletonPose(PoseContext.Pose, CurrentPose.LocalTransforms);
 
 	// Initialize with Previous Root Transform or Identity
@@ -2485,8 +2517,8 @@ bool FPoseHistory::Update(float SecondsElapsed, const FPoseContext& PoseContext,
 
 float FPoseHistory::GetSampleTimeInterval() const
 {
-	// Reserve one knot for computing derivatives at the time horizon
-	return TimeHorizon / (Knots.Max() - 1);
+	// Reserve one pose for computing derivatives at the time horizon
+	return TimeHorizon / (Poses.Max() - 1);
 }
 
 //////////////////////////////////////////////////////////////////////////
@@ -5163,7 +5195,7 @@ UE::Anim::IPoseSearchProvider::FSearchResult FModule::Search(const FAnimationBas
 
 	FPoseHistory& PoseHistory = PoseHistoryProvider->GetPoseHistory();
 
-	FPoseSearchContext SearchContext;
+	FSearchContext SearchContext;
 	SearchContext.OwningComponent = GraphContext.AnimInstanceProxy->GetSkelMeshComponent();
 	SearchContext.BoneContainer = &GraphContext.AnimInstanceProxy->GetRequiredBones();
 	SearchContext.History = &PoseHistoryProvider->GetPoseHistory();
diff --git a/Engine/Plugins/Experimental/Animation/PoseSearch/Source/Runtime/Private/PoseSearchFeatureChannels.cpp b/Engine/Plugins/Experimental/Animation/PoseSearch/Source/Runtime/Private/PoseSearchFeatureChannels.cpp
--- a/Engine/Plugins/Experimental/Animation/PoseSearch/Source/Runtime/Private/PoseSearchFeatureChannels.cpp
+++ b/Engine/Plugins/Experimental/Animation/PoseSearch/Source/Runtime/Private/PoseSearchFeatureChannels.cpp
@@ -308,7 +308,7 @@ void UPoseSearchFeatureChannel_Position::IndexAsset(UE::PoseSearch::IAssetIndexe
 		const float SubsampleTime = OriginSampleTime + SampleTimeOffset;
 
 		bool ClampedPresent;
-		const FTransform BoneTransformsPresent = Indexer.GetTransformAndCacheResults(SubsampleTime, UseCharacterSpaceVelocities ? SubsampleTime : OriginSampleTime, SchemaBoneIdx, ClampedPresent);
+		const FTransform BoneTransformsPresent = Indexer.GetTransformAndCacheResults(SubsampleTime, UseSampleTimeOffsetRootBone ? SubsampleTime : OriginSampleTime, SchemaBoneIdx, ClampedPresent);
 		int32 DataOffset = ChannelDataOffset;
 		FFeatureVectorHelper::EncodeVector(FeatureVector.EditValues(), DataOffset, BoneTransformsPresent.GetTranslation());
 	}
@@ -321,35 +321,36 @@ void UPoseSearchFeatureChannel_Position::GenerateDDCKey(FBlake3& InOutKeyHasher)
 	InOutKeyHasher.Update(&SampleTimeOffset, sizeof(SampleTimeOffset));
 }
 
-bool UPoseSearchFeatureChannel_Position::BuildQuery(FPoseSearchContext& SearchContext, FPoseSearchFeatureVectorBuilder& InOutQuery) const
+bool UPoseSearchFeatureChannel_Position::BuildQuery(UE::PoseSearch::FSearchContext& SearchContext, FPoseSearchFeatureVectorBuilder& InOutQuery) const
 {
 	using namespace UE::PoseSearch;
 
-	const bool bSkip = SearchContext.CurrentResult.IsValid() && SearchContext.CurrentResult.Database->Schema == InOutQuery.GetSchema();
+	const bool bSkip = !UseFeaturesFromContinuityPose && SearchContext.CurrentResult.IsValid() && SearchContext.CurrentResult.Database->Schema == InOutQuery.GetSchema();
 
 	if (bSkip)
 	{
-		// todo: instead of skipping because the pose should have been copied already (that's currently happening in
-		// UPoseSearchSchema::BuildIndex()), consider making the copy here, but only copy the right values, not the
-		// whole query
 		return true;
 	}
 
-	// @todo: add a layer of caching to TrySamplePose to improve performance
-	if (!SearchContext.History || !SearchContext.History->TrySamplePose(-SampleTimeOffset, InOutQuery.GetSchema()->Skeleton->GetReferenceSkeleton(), InOutQuery.GetSchema()->BoneIndicesWithParents))
+	if (!SearchContext.History)
 	{
 		return false;
 	}
 
-	const TArrayView<const FTransform> ComponentPose = SearchContext.History->GetComponentPoseSample();
-	const FTransform RootTransform = SearchContext.History->GetRootTransformSample();
-	const FTransform RootTransformPrev = SearchContext.History->GetPrevRootTransformSample();
-	const int32 SkeletonBoneIndex = InOutQuery.GetSchema()->BoneIndices[SchemaBoneIdx];
-	
+	bool AnyError = false;
+	FTransform Transform = SearchContext.TryGetTransformAndCacheResults(SampleTimeOffset, InOutQuery.GetSchema(), SchemaBoneIdx, AnyError);
+
+	if (!UseSampleTimeOffsetRootBone)
+	{
+		const FTransform RootTransform = SearchContext.TryGetTransformAndCacheResults(0.f, InOutQuery.GetSchema(), FSearchContext::RoolBoneIdx, AnyError);
+		const FTransform RootTransformPrev = SearchContext.TryGetTransformAndCacheResults(SampleTimeOffset, InOutQuery.GetSchema(), FSearchContext::RoolBoneIdx, AnyError);
+		Transform = Transform * (RootTransformPrev * RootTransform.Inverse());
+	}
+
 	int32 DataOffset = ChannelDataOffset;
-	FFeatureVectorHelper::EncodeVector(InOutQuery.EditValues(), DataOffset, ComponentPose[SkeletonBoneIndex].GetTranslation());
+	FFeatureVectorHelper::EncodeVector(InOutQuery.EditValues(), DataOffset, Transform.GetTranslation());
 	
-	return true;
+	return !AnyError;
 }
 
 void UPoseSearchFeatureChannel_Position::DebugDraw(const UE::PoseSearch::FDebugDrawParams& DrawParams, TArrayView<const float> PoseVector) const
@@ -709,19 +710,13 @@ void UPoseSearchFeatureChannel_Pose::GenerateDDCKey(FBlake3& InOutKeyHasher) con
 	InOutKeyHasher.Update(MakeMemoryView(SampleTimes));
 }
 
-bool UPoseSearchFeatureChannel_Pose::BuildQuery(FPoseSearchContext& SearchContext, FPoseSearchFeatureVectorBuilder& InOutQuery) const
+bool UPoseSearchFeatureChannel_Pose::BuildQuery(UE::PoseSearch::FSearchContext& SearchContext, FPoseSearchFeatureVectorBuilder& InOutQuery) const
 {
 	using namespace UE::PoseSearch;
 
-	const bool bSkip =
-		SearchContext.CurrentResult.IsValid() &&
-		SearchContext.CurrentResult.Database->Schema == InOutQuery.GetSchema();
-
+	const bool bSkip = SearchContext.CurrentResult.IsValid() && SearchContext.CurrentResult.Database->Schema == InOutQuery.GetSchema();
 	if (bSkip)
 	{
-		// todo: instead of skipping because the pose should have been copied already (that's currently happening in
-		// UPoseSearchSchema::BuildIndex()), consider making the copy here, but only copy the right values, not the
-		// whole query
 		return true;
 	}
 
@@ -737,8 +732,9 @@ bool UPoseSearchFeatureChannel_Pose::BuildQuery(FPoseSearchContext& SearchContex
 		bool Valid = false;
 	};
 	TArray<CachedTransforms> CachedTransforms;
-	CachedTransforms.AddUninitialized(SampleTimes.Num() * SampledBones.Num());
+	CachedTransforms.AddUninitialized(SampleTimes.Num()* SampledBones.Num());
 
+	bool AnyError = false;
 	for (int32 SubsampleIdx = 0; SubsampleIdx != SampleTimes.Num(); ++SubsampleIdx)
 	{
 		// Stop when we've reached future samples
@@ -748,37 +744,37 @@ bool UPoseSearchFeatureChannel_Pose::BuildQuery(FPoseSearchContext& SearchContex
 			break;
 		}
 
-		float SecondsAgo = -SampleTime;
-		if (!SearchContext.History->TrySamplePose(SecondsAgo, InOutQuery.GetSchema()->Skeleton->GetReferenceSkeleton(), InOutQuery.GetSchema()->BoneIndicesWithParents))
-		{
-			return false;
-		}
-
-		const TArrayView<const FTransform> ComponentPose = SearchContext.History->GetComponentPoseSample();
-		const TArrayView<const FTransform> ComponentPrevPose = SearchContext.History->GetPrevComponentPoseSample();
-		const FTransform RootTransform = SearchContext.History->GetRootTransformSample();
-		const FTransform RootTransformPrev = SearchContext.History->GetPrevRootTransformSample();
 		for (int32 SampledBoneIdx = 0; SampledBoneIdx != SampledBones.Num(); ++SampledBoneIdx)
 		{
-			const int32 SkeletonBoneIndex = InOutQuery.GetSchema()->BoneIndices[SchemaBoneIdx[SampledBoneIdx]];
 			const int32 CachedTransformsIndex = SubsampleIdx * SampledBones.Num() + SampledBoneIdx;
-			CachedTransforms[CachedTransformsIndex].Current = ComponentPose[SkeletonBoneIndex];
+			CachedTransforms[CachedTransformsIndex].Current = SearchContext.TryGetTransformAndCacheResults(SampleTime, InOutQuery.GetSchema(), SchemaBoneIdx[SampledBoneIdx], AnyError);
+			const FPoseSearchBone& SampledBone = SampledBones[SampledBoneIdx];
 
-			if (UE::PoseSearch::UseCharacterSpaceVelocities)
+			if (SampledBone.bUseVelocity)
 			{
-				// character space velocity
-				CachedTransforms[CachedTransformsIndex].Previous = ComponentPrevPose[SkeletonBoneIndex];
-			}
-			else
-			{
-				// animation space velocity
-				CachedTransforms[CachedTransformsIndex].Previous = ComponentPrevPose[SkeletonBoneIndex] * (RootTransformPrev * RootTransform.Inverse());
-			}
+				check(SearchContext.History);
+				const float HistorySameplInterval = SearchContext.History->GetSampleTimeInterval();
+
+				CachedTransforms[CachedTransformsIndex].Previous = SearchContext.TryGetTransformAndCacheResults(SampleTime - HistorySameplInterval, InOutQuery.GetSchema(), SchemaBoneIdx[SampledBoneIdx], AnyError);
+
+				if (!UE::PoseSearch::UseCharacterSpaceVelocities)
+				{
+					const FTransform RootTransform = SearchContext.TryGetTransformAndCacheResults(SampleTime, InOutQuery.GetSchema(), FSearchContext::RoolBoneIdx, AnyError);
+					const FTransform RootTransformPrev = SearchContext.TryGetTransformAndCacheResults(SampleTime - HistorySameplInterval, InOutQuery.GetSchema(), FSearchContext::RoolBoneIdx, AnyError);
 
+					// animation space velocity
+					CachedTransforms[CachedTransformsIndex].Previous = CachedTransforms[CachedTransformsIndex].Previous * (RootTransformPrev * RootTransform.Inverse());
+				}
+			}
 			CachedTransforms[CachedTransformsIndex].Valid = true;
 		}
 	}
 
+	if (AnyError)
+	{
+		return false;
+	}
+
 	int32 DataOffset = ChannelDataOffset;
 	for (int32 SampledBoneIdx = 0; SampledBoneIdx != SampledBones.Num(); ++SampledBoneIdx)
 	{
@@ -1190,9 +1186,7 @@ void UPoseSearchFeatureChannel_Trajectory::GenerateDDCKey(FBlake3& InOutKeyHashe
 	InOutKeyHasher.Update(&Weight, sizeof(Weight));
 }
 
-bool UPoseSearchFeatureChannel_Trajectory::BuildQuery(
-	FPoseSearchContext& SearchContext,
-	FPoseSearchFeatureVectorBuilder& InOutQuery) const
+bool UPoseSearchFeatureChannel_Trajectory::BuildQuery(UE::PoseSearch::FSearchContext& SearchContext, FPoseSearchFeatureVectorBuilder& InOutQuery) const
 {
 	using namespace UE::PoseSearch;
 
diff --git a/Engine/Plugins/Experimental/Animation/PoseSearch/Source/Runtime/Private/PoseSearchLibrary.cpp b/Engine/Plugins/Experimental/Animation/PoseSearch/Source/Runtime/Private/PoseSearchLibrary.cpp
--- a/Engine/Plugins/Experimental/Animation/PoseSearch/Source/Runtime/Private/PoseSearchLibrary.cpp
+++ b/Engine/Plugins/Experimental/Animation/PoseSearch/Source/Runtime/Private/PoseSearchLibrary.cpp
@@ -251,7 +251,7 @@ void UpdateMotionMatchingState(
 	if (!bCanAdvance || (InOutMotionMatchingState.ElapsedPoseJumpTime >= Settings.SearchThrottleTime))
 	{
 		// Build the search context
-		FPoseSearchContext SearchContext;
+		FSearchContext SearchContext;
 		SearchContext.DatabaseTagQuery = DatabaseTagQuery;
 		SearchContext.ActiveTagsContainer = ActiveTagsContainer;
 		SearchContext.Trajectory = &Trajectory;
diff --git a/Engine/Plugins/Experimental/Animation/PoseSearch/Source/Runtime/Public/PoseSearch/PoseSearch.h b/Engine/Plugins/Experimental/Animation/PoseSearch/Source/Runtime/Public/PoseSearch/PoseSearch.h
--- a/Engine/Plugins/Experimental/Animation/PoseSearch/Source/Runtime/Public/PoseSearch/PoseSearch.h
+++ b/Engine/Plugins/Experimental/Animation/PoseSearch/Source/Runtime/Public/PoseSearch/PoseSearch.h
@@ -38,13 +38,14 @@ class UAnimNotifyState_PoseSearchBase;
 class UPoseSearchSchema;
 class FBlake3;
 
-namespace UE::PoseSearch {
-
+namespace UE::PoseSearch
+{
 class FPoseHistory;
 struct FPoseSearchDatabaseAsyncCacheTask;
 struct FDebugDrawParams;
 struct FSchemaInitializer;
 struct FQueryBuildingContext;
+struct FSearchContext;
 
 } // namespace UE::PoseSearch
 
@@ -342,7 +343,7 @@ public:
 	virtual void GenerateDDCKey(FBlake3& InOutKeyHasher) const PURE_VIRTUAL(UPoseSearchFeatureChannel::GenerateDDCKey, );
 
 	// Called at runtime to add this channel's data to the query pose vector
-	virtual bool BuildQuery(FPoseSearchContext& SearchContext, FPoseSearchFeatureVectorBuilder& InOutQuery) const PURE_VIRTUAL(UPoseSearchFeatureChannel::BuildQuery, return false;);
+	virtual bool BuildQuery(UE::PoseSearch::FSearchContext& SearchContext, FPoseSearchFeatureVectorBuilder& InOutQuery) const PURE_VIRTUAL(UPoseSearchFeatureChannel::BuildQuery, return false;);
 
 	// Draw this channel's data for the given pose vector
 	virtual void DebugDraw(const UE::PoseSearch::FDebugDrawParams& DrawParams, TArrayView<const float> PoseVector) const PURE_VIRTUAL(UPoseSearchFeatureChannel::DebugDraw, );
@@ -488,7 +489,7 @@ private:
 	void ResolveBoneReferences();
 
 public:
-	bool BuildQuery(FPoseSearchContext& SearchContext, FPoseSearchFeatureVectorBuilder& InOutQuery) const;
+	bool BuildQuery(UE::PoseSearch::FSearchContext& SearchContext, FPoseSearchFeatureVectorBuilder& InOutQuery) const;
 };
 
 
@@ -1014,8 +1015,7 @@ class POSESEARCH_API UPoseSearchSearchableAsset : public UDataAsset
 	GENERATED_BODY()
 public:
 
-	virtual UE::PoseSearch::FSearchResult Search(FPoseSearchContext& SearchContext) const 
-		PURE_VIRTUAL(UPoseSearchSearchableAsset::Search, return UE::PoseSearch::FSearchResult(););
+	virtual UE::PoseSearch::FSearchResult Search(UE::PoseSearch::FSearchContext& SearchContext) const PURE_VIRTUAL(UPoseSearchSearchableAsset::Search, return UE::PoseSearch::FSearchResult(););
 };
 
 
@@ -1155,31 +1155,17 @@ public:
 	bool IsDerivedDataValid();
 
 public:
-	virtual UE::PoseSearch::FSearchResult Search(FPoseSearchContext& SearchContext) const override;
+	virtual UE::PoseSearch::FSearchResult Search(UE::PoseSearch::FSearchContext& SearchContext) const override;
 
-	void BuildQuery(FPoseSearchContext& SearchContext, FPoseSearchFeatureVectorBuilder& OutQuery) const;
+	void BuildQuery(UE::PoseSearch::FSearchContext& SearchContext, FPoseSearchFeatureVectorBuilder& OutQuery) const;
 
-	FPoseSearchCost ComparePoses(
-		FPoseSearchContext& SearchContext,
-		int32 PoseIdx,
-		int32 GroupIdx,
-		const TArrayView<const float>& QueryValues) const;
-
-	FPoseSearchCost ComparePoses(
-		FPoseSearchContext& SearchContext,
-		int32 PoseIdx,
-		const TArrayView<const float>& QueryValues,
-		UE::PoseSearch::FPoseCostDetails& OutPoseCostDetails) const;
-
-	void ComputePoseCostAddends(
-		int32 PoseIdx,
-		FPoseSearchContext& SearchContext,
-		float& OutNotifyAddend,
-		float& OutMirrorMismatchAddend) const;
+	FPoseSearchCost ComparePoses(UE::PoseSearch::FSearchContext& SearchContext, int32 PoseIdx, int32 GroupIdx, const TArrayView<const float>& QueryValues) const;
+	FPoseSearchCost ComparePoses(UE::PoseSearch::FSearchContext& SearchContext, int32 PoseIdx, const TArrayView<const float>& QueryValues, UE::PoseSearch::FPoseCostDetails& OutPoseCostDetails) const;
+	void ComputePoseCostAddends(int32 PoseIdx, UE::PoseSearch::FSearchContext& SearchContext, float& OutNotifyAddend, float& OutMirrorMismatchAddend) const;
 
 protected:
-	UE::PoseSearch::FSearchResult SearchPCAKDTree(FPoseSearchContext& SearchContext) const;
-	UE::PoseSearch::FSearchResult SearchBruteForce(FPoseSearchContext& SearchContext) const;
+	UE::PoseSearch::FSearchResult SearchPCAKDTree(UE::PoseSearch::FSearchContext& SearchContext) const;
+	UE::PoseSearch::FSearchResult SearchBruteForce(UE::PoseSearch::FSearchContext& SearchContext) const;
 };
 
 
@@ -1208,7 +1194,7 @@ public:
 	bool IsValidForIndexing() const;
 	bool IsValidForSearch() const;
 
-	UE::PoseSearch::FSearchResult Search(FPoseSearchContext& SearchContext) const;
+	UE::PoseSearch::FSearchResult Search(UE::PoseSearch::FSearchContext& SearchContext) const;
 
 protected:
 	FPoseSearchCost ComparePoses(int32 PoseIdx, const TArrayView<const float>& QueryValues) const;
@@ -1222,7 +1208,8 @@ public: // UObject
 //////////////////////////////////////////////////////////////////////////
 // Feature vector reader and builder
 
-namespace UE::PoseSearch {
+namespace UE::PoseSearch
+{
 
 /** Helper class for extracting and encoding features into a float buffer */
 class POSESEARCH_API FFeatureVectorHelper
@@ -1260,7 +1247,6 @@ public:
 
 	void Init(int32 InNumPoses, float InTimeHorizon);
 	void Init(const FPoseHistory& History);
-	bool TrySamplePose(float SecondsAgo, const FReferenceSkeleton& RefSkeleton, const TArray<FBoneIndexType>& RequiredBones);
 
 	bool Update(
 		float SecondsElapsed,
@@ -1270,47 +1256,23 @@ public:
 		ERootUpdateMode UpdateMode = ERootUpdateMode::RootMotionDelta);
 
 	float GetSampleTimeInterval() const;
-	TArrayView<const FTransform> GetLocalPoseSample() const { return SampledLocalPose; }
-	TArrayView<const FTransform> GetComponentPoseSample() const { return SampledComponentPose; }
-	TArrayView<const FTransform> GetPrevLocalPoseSample() const { return SampledPrevLocalPose; }
-	TArrayView<const FTransform> GetPrevComponentPoseSample() const { return SampledPrevComponentPose; }
-	const FTransform& GetRootTransformSample() const { return SampledRootTransform; }
-	const FTransform& GetPrevRootTransformSample() const { return SampledPrevRootTransform; }
 	float GetTimeHorizon() const { return TimeHorizon; }
-	FPoseSearchFeatureVectorBuilder& GetQueryBuilder() { return QueryBuilder; }
+	bool TrySampleLocalPose(float Time, const TArray<FBoneIndexType>* RequiredBones, TArray<FTransform>* LocalPose, FTransform* RootTransform) const;
 
 private:
-	bool TrySampleLocalPose(float Time, const TArray<FBoneIndexType>& RequiredBones, TArray<FTransform>& LocalPose, FTransform& RootTransform) const;
 
 	struct FPose
 	{
 		FTransform RootTransform;
 		TArray<FTransform> LocalTransforms;
+		float Time = 0.0f;
 	};
-
 	TRingBuffer<FPose> Poses;
-	TRingBuffer<float> Knots;
-	TArray<FTransform> SampledLocalPose;
-	TArray<FTransform> SampledComponentPose;
-	TArray<FTransform> SampledPrevLocalPose;
-	TArray<FTransform> SampledPrevComponentPose;
-	FTransform SampledRootTransform;
-	FTransform SampledPrevRootTransform;
-
-	FPoseSearchFeatureVectorBuilder QueryBuilder;
-
 	float TimeHorizon = 0.0f;
 };
 
-} // namespace UE::PoseSearch
-
-
-USTRUCT(BlueprintType)
-struct POSESEARCH_API FPoseSearchContext
+struct POSESEARCH_API FSearchContext
 {
-	GENERATED_BODY()
-
-public:
 	EPoseSearchBooleanRequest QueryMirrorRequest = EPoseSearchBooleanRequest::Indifferent;
 	const FGameplayTagQuery* DatabaseTagQuery = nullptr;
 	UE::PoseSearch::FDebugDrawParams DebugDrawParams;
@@ -1321,18 +1283,31 @@ public:
 	const FBoneContainer* BoneContainer = nullptr;
 	const FGameplayTagContainer* ActiveTagsContainer = nullptr;
 	float PoseJumpThresholdTime = 0.f;
-};
 
+	FTransform TryGetTransformAndCacheResults(float SampleTime, const UPoseSearchSchema* Schema, int8 SchemaBoneIdx, bool& Error);
+	void ClearCachedEntries();
 
-//////////////////////////////////////////////////////////////////////////
-// Pose history
+	enum { RoolBoneIdx = -1 };
 
-namespace UE::PoseSearch {
+private:
+	struct CachedEntry
+	{
+		float SampleTime = 0.f;
+
+		// associated transform to BoneIndexType in ComponentSpace (except for the root bone stored in global space)
+		FTransform Transform;
+
+		// if -1 it represents the root bone
+		FBoneIndexType BoneIndexType = -1;
+	};
+
+	// @todo: make it a fixed size array (or hash map if we end up having many CachedEntry) to avoid allocations
+	TArray<CachedEntry> CachedEntries;
+};
 
 //////////////////////////////////////////////////////////////////////////
 // Main PoseSearch API
 
-
 /**
 * Visualize pose search debug information
 *
@@ -1412,5 +1387,5 @@ public:
 	TArray<FPoseSearchDatabaseSetEntry> AssetsToSearch;
 
 public:
-	virtual UE::PoseSearch::FSearchResult Search(FPoseSearchContext& SearchContext) const override;
+	virtual UE::PoseSearch::FSearchResult Search(UE::PoseSearch::FSearchContext& SearchContext) const override;
 };
diff --git a/Engine/Plugins/Experimental/Animation/PoseSearch/Source/Runtime/Public/PoseSearch/PoseSearchFeatureChannels.h b/Engine/Plugins/Experimental/Animation/PoseSearch/Source/Runtime/Public/PoseSearch/PoseSearchFeatureChannels.h
--- a/Engine/Plugins/Experimental/Animation/PoseSearch/Source/Runtime/Public/PoseSearch/PoseSearchFeatureChannels.h
+++ b/Engine/Plugins/Experimental/Animation/PoseSearch/Source/Runtime/Public/PoseSearch/PoseSearchFeatureChannels.h
@@ -28,12 +28,21 @@ public:
 	UPROPERTY()
 	int8 SchemaBoneIdx;
 
+	// reuse continuity pose database feature values if available
+	UPROPERTY(EditAnywhere, Category = "Advanced")
+	bool UseFeaturesFromContinuityPose = true;
+
+	// if UseSampleTimeOffsetRootBone is true, this UPoseSearchFeatureChannel_Position will calculate the position of Bone from the pose SampleTimeOffset seconds away from the current time pose root bone
+	// if false the calculated position of Bone will be in component space from the pose SampleTimeOffset seconds away
+	UPROPERTY(EditAnywhere, Category = "Advanced")
+	bool UseSampleTimeOffsetRootBone = true;
+
 	// UPoseSearchFeatureChannel interface
 	virtual void InitializeSchema(UE::PoseSearch::FSchemaInitializer& Initializer) override;
 	virtual void FillWeights(TArray<float>& Weights) const override;
 	virtual void IndexAsset(UE::PoseSearch::IAssetIndexer& Indexer, UE::PoseSearch::FAssetIndexingOutput& IndexingOutput) const override;
 	virtual void GenerateDDCKey(FBlake3& InOutKeyHasher) const override;
-	virtual bool BuildQuery(FPoseSearchContext& SearchContext, FPoseSearchFeatureVectorBuilder& InOutQuery) const override;
+	virtual bool BuildQuery(UE::PoseSearch::FSearchContext& SearchContext, FPoseSearchFeatureVectorBuilder& InOutQuery) const override;
 	virtual void DebugDraw(const UE::PoseSearch::FDebugDrawParams& DrawParams, TArrayView<const float> PoseVector) const override;
 };
 
@@ -63,7 +72,7 @@ public:
 	virtual void IndexAsset(UE::PoseSearch::IAssetIndexer& Indexer, UE::PoseSearch::FAssetIndexingOutput& IndexingOutput) const override;
 	virtual void ComputeMeanDeviations(const Eigen::MatrixXd& CenteredPoseMatrix, Eigen::VectorXd& MeanDeviations) const override;
 	virtual void GenerateDDCKey(FBlake3& InOutKeyHasher) const override;
-	virtual bool BuildQuery(FPoseSearchContext& SearchContext, FPoseSearchFeatureVectorBuilder& InOutQuery) const override;
+	virtual bool BuildQuery(UE::PoseSearch::FSearchContext& SearchContext, FPoseSearchFeatureVectorBuilder& InOutQuery) const override;
 	virtual void DebugDraw(const UE::PoseSearch::FDebugDrawParams& DrawParams, TArrayView<const float> PoseVector) const override;
 
 protected:
@@ -107,7 +116,7 @@ public:
 	virtual void IndexAsset(UE::PoseSearch::IAssetIndexer& Indexer, UE::PoseSearch::FAssetIndexingOutput& IndexingOutput) const override;
 	virtual void ComputeMeanDeviations(const Eigen::MatrixXd& CenteredPoseMatrix, Eigen::VectorXd& MeanDeviations) const override;
 	virtual void GenerateDDCKey(FBlake3& InOutKeyHasher) const override;
-	virtual bool BuildQuery(FPoseSearchContext& SearchContext, FPoseSearchFeatureVectorBuilder& InOutQuery) const override;
+	virtual bool BuildQuery(UE::PoseSearch::FSearchContext& SearchContext, FPoseSearchFeatureVectorBuilder& InOutQuery) const override;
 	virtual void DebugDraw(const UE::PoseSearch::FDebugDrawParams& DrawParams, TArrayView<const float> PoseVector) const override;
 
 protected:
