motion matching - phase FindLocalMinMax fix

diff --git a/Engine/Plugins/Experimental/Animation/PoseSearch/Source/Runtime/Private/PoseSearchFeatureChannels.cpp b/Engine/Plugins/Experimental/Animation/PoseSearch/Source/Runtime/Private/PoseSearchFeatureChannels.cpp
--- a/Engine/Plugins/Experimental/Animation/PoseSearch/Source/Runtime/Private/PoseSearchFeatureChannels.cpp
+++ b/Engine/Plugins/Experimental/Animation/PoseSearch/Source/Runtime/Private/PoseSearchFeatureChannels.cpp
@@ -105,29 +105,66 @@ static void SmoothSignal(const TArray<float>& Signal, TArray<float>& SmoothedSig
 	}
 }
 
-static void FindLocalMinMax(const TArray<float>& Signal, TArray<LocalMinMax>& MinMax, int32 offset = 1)
+static void FindLocalMinMax(const TArray<float>& Signal, TArray<LocalMinMax>& MinMax)
 {
-	check(offset > 0);
+	enum SignalState
+	{
+		Flat,
+		Ascending,
+		Descending
+	};
+
 	MinMax.Reset();
-	for (int32 i = 0; i < Signal.Num(); ++i)
+	if (Signal.Num() > 1)
 	{
-		const float Previous = GetValueAtIndex(i - offset, Signal);
-		const float Current = GetValueAtIndex(i, Signal);
-		const float Next = GetValueAtIndex(i + offset, Signal);
+		SignalState State = SignalState::Flat;
+		for (int32 SignalIndex = 1; SignalIndex < Signal.Num(); ++SignalIndex)
+		{
+			const int32 PrevSignalIndex = SignalIndex - 1;
+			const float PrevSignalValue = Signal[PrevSignalIndex];
+			const float SignalValue = Signal[SignalIndex];
 
-		const float DeltaSignalValueBackward = Previous - Current;
-		const float DeltaSignalValueForward = Next - Current;
+			if (State == SignalState::Flat)
+			{
+				if (SignalValue > PrevSignalValue)
+				{
+					State = SignalState::Ascending;
+				}
+				else if(SignalValue < PrevSignalValue)
+				{
+					State = SignalState::Descending;
+				}
+			}
+			else if (State == SignalState::Ascending)
+			{
+				if (SignalValue < PrevSignalValue)
+				{
+					State = SignalState::Descending;
+					
+					LocalMinMax LocalMinMax;
+					LocalMinMax.Type = LocalMinMax::Max;
+					LocalMinMax.Index = PrevSignalIndex;
+					LocalMinMax.SignalValue = Signal[LocalMinMax.Index];
+
+					check(MinMax.IsEmpty() || MinMax.Last().Type != LocalMinMax.Type);
+					MinMax.Add(LocalMinMax);
+				}
+			}
+			else // if (State == SignalState::Descending)
+			{
+				if (SignalValue > PrevSignalValue)
+				{
+					State = SignalState::Ascending;
 
-		const float Sign = DeltaSignalValueBackward * DeltaSignalValueForward;
-		if (Sign >= 0.f && DeltaSignalValueBackward != 0.f)
-		{
-			LocalMinMax LocalMinMax;
-			LocalMinMax.Type = DeltaSignalValueForward < 0.f ? LocalMinMax::Max : LocalMinMax::Min;
-			LocalMinMax.Index = i;
-			LocalMinMax.SignalValue = Signal[i];
+					LocalMinMax LocalMinMax;
+					LocalMinMax.Type = LocalMinMax::Min;
+					LocalMinMax.Index = PrevSignalIndex;
+					LocalMinMax.SignalValue = Signal[LocalMinMax.Index];
 
-			check(MinMax.IsEmpty() || MinMax.Last().Type != LocalMinMax.Type);
-			MinMax.Add(LocalMinMax);
+					check(MinMax.IsEmpty() || MinMax.Last().Type != LocalMinMax.Type);
+					MinMax.Add(LocalMinMax);
+				}
+			}
 		}
 	}
 }
