fix for  warning V547: Expression 'BoneIndexType >= 0' is always true. Unsigned type value is always >= 0.

diff --git a/Engine/Plugins/Experimental/Animation/PoseSearch/Source/Runtime/Private/PoseSearch.cpp b/Engine/Plugins/Experimental/Animation/PoseSearch/Source/Runtime/Private/PoseSearch.cpp
--- a/Engine/Plugins/Experimental/Animation/PoseSearch/Source/Runtime/Private/PoseSearch.cpp
+++ b/Engine/Plugins/Experimental/Animation/PoseSearch/Source/Runtime/Private/PoseSearch.cpp
@@ -2275,7 +2275,8 @@ FTransform FSearchContext::TryGetTransformAndCacheResults(float SampleTime, cons
 {
 	check(History && Schema);
 
-	const FBoneIndexType BoneIndexType = SchemaBoneIdx >= 0 ? Schema->BoneIndices[SchemaBoneIdx] : RoolBoneIdx;
+	static constexpr FBoneIndexType RootBoneIdx = 0xFFFF;
+	const FBoneIndexType BoneIndexType = SchemaBoneIdx >= 0 ? Schema->BoneIndices[SchemaBoneIdx] : RootBoneIdx;
 
 	// @todo: use an hashmap if we end up having too many entries
 	const CachedEntry* Entry = CachedEntries.FindByPredicate([SampleTime, BoneIndexType](const FSearchContext::CachedEntry& Entry)
@@ -2289,7 +2290,7 @@ FTransform FSearchContext::TryGetTransformAndCacheResults(float SampleTime, cons
 		return Entry->Transform;
 	}
 
-	if (BoneIndexType >= 0)
+	if (BoneIndexType != RootBoneIdx)
 	{
 		TArray<FTransform> SampledLocalPose;
 		if (History->TrySampleLocalPose(-SampleTime, &Schema->BoneIndicesWithParents, &SampledLocalPose, nullptr))
diff --git a/Engine/Plugins/Experimental/Animation/PoseSearch/Source/Runtime/Private/PoseSearchFeatureChannels.cpp b/Engine/Plugins/Experimental/Animation/PoseSearch/Source/Runtime/Private/PoseSearchFeatureChannels.cpp
--- a/Engine/Plugins/Experimental/Animation/PoseSearch/Source/Runtime/Private/PoseSearchFeatureChannels.cpp
+++ b/Engine/Plugins/Experimental/Animation/PoseSearch/Source/Runtime/Private/PoseSearchFeatureChannels.cpp
@@ -379,8 +379,8 @@ bool UPoseSearchFeatureChannel_Position::BuildQuery(UE::PoseSearch::FSearchConte
 
 	if (!UseSampleTimeOffsetRootBone)
 	{
-		const FTransform RootTransform = SearchContext.TryGetTransformAndCacheResults(0.f, InOutQuery.GetSchema(), FSearchContext::RoolBoneIdx, AnyError);
-		const FTransform RootTransformPrev = SearchContext.TryGetTransformAndCacheResults(SampleTimeOffset, InOutQuery.GetSchema(), FSearchContext::RoolBoneIdx, AnyError);
+		const FTransform RootTransform = SearchContext.TryGetTransformAndCacheResults(0.f, InOutQuery.GetSchema(), FSearchContext::SchemaRootBoneIdx, AnyError);
+		const FTransform RootTransformPrev = SearchContext.TryGetTransformAndCacheResults(SampleTimeOffset, InOutQuery.GetSchema(), FSearchContext::SchemaRootBoneIdx, AnyError);
 		Transform = Transform * (RootTransformPrev * RootTransform.Inverse());
 	}
 
@@ -796,8 +796,8 @@ bool UPoseSearchFeatureChannel_Pose::BuildQuery(UE::PoseSearch::FSearchContext&
 
 				if (!UE::PoseSearch::UseCharacterSpaceVelocities)
 				{
-					const FTransform RootTransform = SearchContext.TryGetTransformAndCacheResults(SampleTime, InOutQuery.GetSchema(), FSearchContext::RoolBoneIdx, AnyError);
-					const FTransform RootTransformPrev = SearchContext.TryGetTransformAndCacheResults(SampleTime - HistorySameplInterval, InOutQuery.GetSchema(), FSearchContext::RoolBoneIdx, AnyError);
+					const FTransform RootTransform = SearchContext.TryGetTransformAndCacheResults(SampleTime, InOutQuery.GetSchema(), FSearchContext::SchemaRootBoneIdx, AnyError);
+					const FTransform RootTransformPrev = SearchContext.TryGetTransformAndCacheResults(SampleTime - HistorySameplInterval, InOutQuery.GetSchema(), FSearchContext::SchemaRootBoneIdx, AnyError);
 
 					// animation space velocity
 					CachedTransforms[CachedTransformsIndex].Previous = CachedTransforms[CachedTransformsIndex].Previous * (RootTransformPrev * RootTransform.Inverse());
diff --git a/Engine/Plugins/Experimental/Animation/PoseSearch/Source/Runtime/Public/PoseSearch/PoseSearch.h b/Engine/Plugins/Experimental/Animation/PoseSearch/Source/Runtime/Public/PoseSearch/PoseSearch.h
--- a/Engine/Plugins/Experimental/Animation/PoseSearch/Source/Runtime/Public/PoseSearch/PoseSearch.h
+++ b/Engine/Plugins/Experimental/Animation/PoseSearch/Source/Runtime/Public/PoseSearch/PoseSearch.h
@@ -1287,7 +1287,7 @@ struct POSESEARCH_API FSearchContext
 	FTransform TryGetTransformAndCacheResults(float SampleTime, const UPoseSearchSchema* Schema, int8 SchemaBoneIdx, bool& Error);
 	void ClearCachedEntries();
 
-	enum { RoolBoneIdx = -1 };
+	static constexpr int8 SchemaRootBoneIdx = -1;
 
 private:
 	struct CachedEntry
