motion matching - implementing FKDTree with nanoflann

diff --git a/Engine/Plugins/Experimental/Animation/PoseSearch/Source/Runtime/PoseSearch.Build.cs b/Engine/Plugins/Experimental/Animation/PoseSearch/Source/Runtime/PoseSearch.Build.cs
--- a/Engine/Plugins/Experimental/Animation/PoseSearch/Source/Runtime/PoseSearch.Build.cs
+++ b/Engine/Plugins/Experimental/Animation/PoseSearch/Source/Runtime/PoseSearch.Build.cs
@@ -9,7 +9,8 @@ public class PoseSearch : ModuleRules
 		PCHUsage = ModuleRules.PCHUsageMode.UseExplicitOrSharedPCHs;
 
 		AddEngineThirdPartyPrivateStaticDependencies(Target,
-			"Eigen"
+			"Eigen",
+			"nanoflann"
 		);
 
 		PublicIncludePaths.AddRange(
diff --git a/Engine/Plugins/Experimental/Animation/PoseSearch/Source/Runtime/Private/KDTree.cpp b/Engine/Plugins/Experimental/Animation/PoseSearch/Source/Runtime/Private/KDTree.cpp
--- a/Engine/Plugins/Experimental/Animation/PoseSearch/Source/Runtime/Private/KDTree.cpp
+++ b/Engine/Plugins/Experimental/Animation/PoseSearch/Source/Runtime/Private/KDTree.cpp
@@ -2,29 +2,46 @@
 
 #include "PoseSearch/KDTree.h"
 
+// @third party code - BEGIN nanoflann
+THIRD_PARTY_INCLUDES_START
+#include "nanoflann/nanoflann.hpp"
+THIRD_PARTY_INCLUDES_END
+// @third party code - END nanoflann
+
 namespace UE { namespace PoseSearch
 {
 
+using FKDTreeImplementationBase = nanoflann::KDTreeSingleIndexAdaptor<nanoflann::L2_Simple_Adaptor<float, FKDTree::DataSource>, FKDTree::DataSource>;
+struct FKDTreeImplementation : FKDTreeImplementationBase
+{
+	using FKDTreeImplementationBase::FKDTreeImplementationBase;
+};
+
 FKDTree::FKDTree(int32 Count, int32 Dim, const float* Data, int32 MaxLeafSize)
 : DataSrc(Count, Dim, Data)
 , Impl(nullptr)
 {
+	Impl = new FKDTreeImplementation(Dim, DataSrc, nanoflann::KDTreeSingleIndexAdaptorParams(MaxLeafSize));
 }
 
 FKDTree::FKDTree()
 : DataSrc(0, 0, nullptr)
 , Impl(nullptr)
 {
+	Impl = new FKDTreeImplementation(0, DataSrc, nanoflann::KDTreeSingleIndexAdaptorParams(0));
 }
 
 FKDTree::~FKDTree()
 {
+	delete Impl;
 }
 
 FKDTree::FKDTree(const FKDTree& r)
 : DataSrc(r.DataSrc.PointCount, r.DataSrc.PointDim, r.DataSrc.Data)
 , Impl(nullptr)
 {
+	check(r.Impl);
+	Impl = new FKDTreeImplementation(r.DataSrc.PointDim, DataSrc, nanoflann::KDTreeSingleIndexAdaptorParams(r.Impl->m_leaf_max_size));
 }
 
 FKDTree& FKDTree::operator=(const FKDTree& r)
@@ -42,12 +59,98 @@ void FKDTree::Construct(int32 Count, int32 Dim, const float* Data, int32 MaxLeaf
 
 bool FKDTree::FindNeighbors(KNNResultSet& Result, const float* Query) const
 {
-	check(false); // unimplemented
-	return false;
+	const nanoflann::SearchParams SearchParams(
+		32,			// Ignored parameter (Kept for compatibility with the FLANN interface).
+		0.f,		// search for eps-approximate neighbours (default: 0)
+		false);		// only for radius search, require neighbours sorted by
+	return Impl->findNeighbors(Result, Query, SearchParams);
+}
+
+FArchive& SerializeSubTree(FArchive& Ar, FKDTree& KDTree, FKDTreeImplementation::NodePtr KDTreeNode)
+{
+	check(KDTree.Impl);
+
+	if (Ar.IsLoading())
+	{
+		KDTreeNode = KDTree.Impl->pool.template allocate<FKDTreeImplementation::Node>();
+	}
+
+	Ar.Serialize(&KDTreeNode->node_type, sizeof(KDTreeNode->node_type));
+
+	bool child1 = KDTreeNode->child1 != nullptr;
+	bool child2 = KDTreeNode->child2 != nullptr;
+	Ar << child1;
+	Ar << child2;
+
+	if (child1)
+	{
+		SerializeSubTree(Ar, KDTree, KDTreeNode->child1);
+	}
+	else if (Ar.IsLoading())
+	{
+		KDTreeNode->child1 = nullptr;
+	}
+
+	if (child2)
+	{
+		SerializeSubTree(Ar, KDTree, KDTreeNode->child2);
+	}
+	else if (Ar.IsLoading())
+	{
+		KDTreeNode->child2 = nullptr;
+	}
+	return Ar;
 }
 
 FArchive& Serialize(FArchive& Ar, FKDTree& KDTree, const float* KDTreeData)
 {
+	check(KDTree.Impl);
+	check(KDTree.Impl->m_size < UINT_MAX);
+	uint32 KDTreeSize = KDTree.Impl->m_size;
+	Ar << KDTreeSize;
+	KDTree.Impl->m_size = KDTreeSize;
+
+	if (KDTree.Impl->m_size  > 0)
+	{
+		Ar << KDTree.Impl->dim;
+
+		uint32 root_bbox_size = KDTree.Impl->root_bbox.size();
+		check(KDTree.Impl->root_bbox.size() < UINT_MAX);
+		Ar << root_bbox_size;
+
+		if (Ar.IsLoading())
+		{
+			KDTree.DataSrc.Data = KDTreeData;
+			KDTree.DataSrc.PointDim = KDTree.Impl->dim;
+			KDTree.DataSrc.PointCount = KDTree.Impl->m_size;
+
+			KDTree.Impl->root_bbox.resize(root_bbox_size);
+		}
+
+		for (FKDTreeImplementation::Interval& el : KDTree.Impl->root_bbox)
+		{
+			Ar.Serialize(&el, sizeof(FKDTreeImplementation::Interval));
+		}
+
+		check(KDTree.Impl->m_leaf_max_size < UINT_MAX);
+		uint32 KDTreeLeafMaxSize = KDTree.Impl->m_leaf_max_size;
+		Ar << KDTreeLeafMaxSize;
+		KDTree.Impl->m_leaf_max_size = KDTreeLeafMaxSize;
+
+		check(KDTree.Impl->vAcc.size() < UINT_MAX);
+		uint32 VAccSize = KDTree.Impl->vAcc.size();
+		Ar << VAccSize;
+		if (Ar.IsLoading())
+		{
+			KDTree.Impl->vAcc.resize(VAccSize);
+		}
+		for (uint32_t& el : KDTree.Impl->vAcc)
+		{
+			Ar << el;
+		}
+		SerializeSubTree(Ar, KDTree, KDTree.Impl->root_node);
+	}
+
 	return Ar;
 }
 
diff --git a/Engine/Plugins/Experimental/Animation/PoseSearch/Source/Runtime/Public/PoseSearch/PoseSearch.h b/Engine/Plugins/Experimental/Animation/PoseSearch/Source/Runtime/Public/PoseSearch/PoseSearch.h
--- a/Engine/Plugins/Experimental/Animation/PoseSearch/Source/Runtime/Public/PoseSearch/PoseSearch.h
+++ b/Engine/Plugins/Experimental/Animation/PoseSearch/Source/Runtime/Public/PoseSearch/PoseSearch.h
@@ -1067,7 +1067,7 @@ public:
 	int32 KDTreeQueryNumNeighbors = 100;
 
 	UPROPERTY(EditAnywhere, Category = "Performance")
-	EPoseSearchMode PoseSearchMode = EPoseSearchMode::BruteForce;
+	EPoseSearchMode PoseSearchMode = EPoseSearchMode::PCAKDTree;
 
 	FPoseSearchIndex* GetSearchIndex();
 	const FPoseSearchIndex* GetSearchIndex() const;
