motion matching - adding UE_POSE_SEARCH_USE_NANOFLANN to optionally strip nanoflann usage

diff --git a/Engine/Plugins/Experimental/Animation/PoseSearch/Source/Runtime/Private/KDTree.cpp b/Engine/Plugins/Experimental/Animation/PoseSearch/Source/Runtime/Private/KDTree.cpp
--- a/Engine/Plugins/Experimental/Animation/PoseSearch/Source/Runtime/Private/KDTree.cpp
+++ b/Engine/Plugins/Experimental/Animation/PoseSearch/Source/Runtime/Private/KDTree.cpp
@@ -2,46 +2,60 @@
 
 #include "PoseSearch/KDTree.h"
 
+#define UE_POSE_SEARCH_USE_NANOFLANN 0
+
 // @third party code - BEGIN nanoflann
+#if UE_POSE_SEARCH_USE_NANOFLANN
 THIRD_PARTY_INCLUDES_START
 #include "nanoflann/nanoflann.hpp"
 THIRD_PARTY_INCLUDES_END
+#endif
 // @third party code - END nanoflann
 
 namespace UE { namespace PoseSearch
 {
 
+#if UE_POSE_SEARCH_USE_NANOFLANN
 using FKDTreeImplementationBase = nanoflann::KDTreeSingleIndexAdaptor<nanoflann::L2_Simple_Adaptor<float, FKDTree::DataSource>, FKDTree::DataSource>;
 struct FKDTreeImplementation : FKDTreeImplementationBase
 {
 	using FKDTreeImplementationBase::FKDTreeImplementationBase;
 };
+#endif
 
 FKDTree::FKDTree(int32 Count, int32 Dim, const float* Data, int32 MaxLeafSize)
 : DataSrc(Count, Dim, Data)
 , Impl(nullptr)
 {
+#if UE_POSE_SEARCH_USE_NANOFLANN
 	Impl = new FKDTreeImplementation(Dim, DataSrc, nanoflann::KDTreeSingleIndexAdaptorParams(MaxLeafSize));
+#endif
 }
 
 FKDTree::FKDTree()
 : DataSrc(0, 0, nullptr)
 , Impl(nullptr)
 {
+#if UE_POSE_SEARCH_USE_NANOFLANN
 	Impl = new FKDTreeImplementation(0, DataSrc, nanoflann::KDTreeSingleIndexAdaptorParams(0));
+#endif
 }
 
 FKDTree::~FKDTree()
 {
+#if UE_POSE_SEARCH_USE_NANOFLANN
 	delete Impl;
+#endif
 }
 
 FKDTree::FKDTree(const FKDTree& r)
 : DataSrc(r.DataSrc.PointCount, r.DataSrc.PointDim, r.DataSrc.Data)
 , Impl(nullptr)
 {
+#if UE_POSE_SEARCH_USE_NANOFLANN
 	check(r.Impl);
 	Impl = new FKDTreeImplementation(r.DataSrc.PointDim, DataSrc, nanoflann::KDTreeSingleIndexAdaptorParams(r.Impl->m_leaf_max_size));
+#endif
 }
 
 FKDTree& FKDTree::operator=(const FKDTree& r)
@@ -59,13 +73,19 @@ void FKDTree::Construct(int32 Count, int32 Dim, const float* Data, int32 MaxLeaf
 
 bool FKDTree::FindNeighbors(KNNResultSet& Result, const float* Query) const
 {
+#if UE_POSE_SEARCH_USE_NANOFLANN
 	const nanoflann::SearchParams SearchParams(
 		32,			// Ignored parameter (Kept for compatibility with the FLANN interface).
 		0.f,		// search for eps-approximate neighbours (default: 0)
 		false);		// only for radius search, require neighbours sorted by
 	return Impl->findNeighbors(Result, Query, SearchParams);
+#else
+	check(false); // unimplemented
+	return false;
+#endif
 }
 
+#if UE_POSE_SEARCH_USE_NANOFLANN
 FArchive& SerializeSubTree(FArchive& Ar, FKDTree& KDTree, FKDTreeImplementation::NodePtr KDTreeNode)
 {
 	check(KDTree.Impl);
@@ -102,8 +122,12 @@ FArchive& SerializeSubTree(FArchive& Ar, FKDTree& KDTree, FKDTreeImplementation:
 	return Ar;
 }
 
+#endif
+
 FArchive& Serialize(FArchive& Ar, FKDTree& KDTree, const float* KDTreeData)
 {
+#if UE_POSE_SEARCH_USE_NANOFLANN
+
 	check(KDTree.Impl);
 	check(KDTree.Impl->m_size < UINT_MAX);
 	uint32 KDTreeSize = KDTree.Impl->m_size;
@@ -150,6 +174,7 @@ FArchive& Serialize(FArchive& Ar, FKDTree& KDTree, const float* KDTreeData)
 		}
 		SerializeSubTree(Ar, KDTree, KDTree.Impl->root_node);
 	}
+#endif
 
 	return Ar;
 }
diff --git a/Engine/Plugins/Experimental/Animation/PoseSearch/Source/Runtime/Public/PoseSearch/PoseSearch.h b/Engine/Plugins/Experimental/Animation/PoseSearch/Source/Runtime/Public/PoseSearch/PoseSearch.h
--- a/Engine/Plugins/Experimental/Animation/PoseSearch/Source/Runtime/Public/PoseSearch/PoseSearch.h
+++ b/Engine/Plugins/Experimental/Animation/PoseSearch/Source/Runtime/Public/PoseSearch/PoseSearch.h
@@ -1067,7 +1067,7 @@ public:
 	int32 KDTreeQueryNumNeighbors = 100;
 
 	UPROPERTY(EditAnywhere, Category = "Performance")
-	EPoseSearchMode PoseSearchMode = EPoseSearchMode::PCAKDTree;
+	EPoseSearchMode PoseSearchMode = EPoseSearchMode::BruteForce;
 
 	FPoseSearchIndex* GetSearchIndex();
 	const FPoseSearchIndex* GetSearchIndex() const;
