motion matching - enable UE_POSE_SEARCH_USE_NANOFLANN

diff --git a/Engine/Plugins/Experimental/Animation/PoseSearch/Source/Runtime/Private/KDTree.cpp b/Engine/Plugins/Experimental/Animation/PoseSearch/Source/Runtime/Private/KDTree.cpp
--- a/Engine/Plugins/Experimental/Animation/PoseSearch/Source/Runtime/Private/KDTree.cpp
+++ b/Engine/Plugins/Experimental/Animation/PoseSearch/Source/Runtime/Private/KDTree.cpp
@@ -2,7 +2,7 @@
 
 #include "PoseSearch/KDTree.h"
 
-#define UE_POSE_SEARCH_USE_NANOFLANN 0
+#define UE_POSE_SEARCH_USE_NANOFLANN 1
 
 // @third party code - BEGIN nanoflann
 #if UE_POSE_SEARCH_USE_NANOFLANN
