motion matching
- add force interrupt to MM node
- adding bitmask enums to FPoseSearchBone to replace bool bUseVel/Pos/Rot/Phase
- introduced UPoseSearchFeatureChannel_Heading as additional itemization of UPoseSearchFeatureChannel_Pose
- added UPoseSearchFeatureChannel_Pose::bUseFeaturesFromContinuityPose
- added UPoseSearchFeatureChannel_Trajectory::Samples with flags to support non homogeneous pos/vel/dir per sample time
- added data upgrade within UPoseSearchSchema::PostLoad

diff --git a/Engine/Plugins/Experimental/Animation/PoseSearch/Source/Runtime/Private/AnimNode_MotionMatching.cpp b/Engine/Plugins/Experimental/Animation/PoseSearch/Source/Runtime/Private/AnimNode_MotionMatching.cpp
--- a/Engine/Plugins/Experimental/Animation/PoseSearch/Source/Runtime/Private/AnimNode_MotionMatching.cpp
+++ b/Engine/Plugins/Experimental/Animation/PoseSearch/Source/Runtime/Private/AnimNode_MotionMatching.cpp
@@ -92,7 +92,8 @@ void FAnimNode_MotionMatching::UpdateAssetPlayer(const FAnimationUpdateContext&
 		&ActiveTagsContainer,
 		Trajectory,
 		Settings,
-		MotionMatchingState
+		MotionMatchingState,
+		bForceInterrupt
 	);
 
 
diff --git a/Engine/Plugins/Experimental/Animation/PoseSearch/Source/Runtime/Private/PoseSearch.cpp b/Engine/Plugins/Experimental/Animation/PoseSearch/Source/Runtime/Private/PoseSearch.cpp
--- a/Engine/Plugins/Experimental/Animation/PoseSearch/Source/Runtime/Private/PoseSearch.cpp
+++ b/Engine/Plugins/Experimental/Animation/PoseSearch/Source/Runtime/Private/PoseSearch.cpp
@@ -334,10 +334,10 @@ void UPoseSearchSchema::PostLoad()
 
 		UPoseSearchFeatureChannel_Trajectory* Channel = NewObject<UPoseSearchFeatureChannel_Trajectory>(this);
 		Channel->Domain = EPoseSearchFeatureDomain::Time;
-		Channel->SampleOffsets = TrajectorySampleTimes_DEPRECATED;
-		Channel->bUseFacingDirections = bUseTrajectoryForwardVectors_DEPRECATED;
-		Channel->bUseLinearVelocities = bUseTrajectoryVelocities_DEPRECATED;
-		Channel->bUsePositions = bUseTrajectoryPositions_DEPRECATED;
+		Channel->SampleOffsets_DEPRECATED = TrajectorySampleTimes_DEPRECATED;
+		Channel->bUseFacingDirections_DEPRECATED = bUseTrajectoryForwardVectors_DEPRECATED;
+		Channel->bUseLinearVelocities_DEPRECATED = bUseTrajectoryVelocities_DEPRECATED;
+		Channel->bUsePositions_DEPRECATED = bUseTrajectoryPositions_DEPRECATED;
 		Channels.Add(Channel);
 
 		TrajectorySampleTimes_DEPRECATED.Empty();
@@ -349,10 +349,10 @@ void UPoseSearchSchema::PostLoad()
 
 		UPoseSearchFeatureChannel_Trajectory* Channel = NewObject<UPoseSearchFeatureChannel_Trajectory>(this);
 		Channel->Domain = EPoseSearchFeatureDomain::Distance;
-		Channel->SampleOffsets = TrajectorySampleDistances_DEPRECATED;
-		Channel->bUseFacingDirections = bUseTrajectoryForwardVectors_DEPRECATED;
-		Channel->bUseLinearVelocities = bUseTrajectoryVelocities_DEPRECATED;
-		Channel->bUsePositions = bUseTrajectoryPositions_DEPRECATED;
+		Channel->SampleOffsets_DEPRECATED = TrajectorySampleDistances_DEPRECATED;
+		Channel->bUseFacingDirections_DEPRECATED = bUseTrajectoryForwardVectors_DEPRECATED;
+		Channel->bUseLinearVelocities_DEPRECATED = bUseTrajectoryVelocities_DEPRECATED;
+		Channel->bUsePositions_DEPRECATED = bUseTrajectoryPositions_DEPRECATED;
 		Channels.Add(Channel);
 
 		TrajectorySampleDistances_DEPRECATED.Empty();
@@ -380,16 +380,81 @@ void UPoseSearchSchema::PostLoad()
 				bNeedFinalize = true;
 				break;
 			}
+		}
+	}
 
-			if (UPoseSearchFeatureChannel_Pose* ChannelPose = Cast<UPoseSearchFeatureChannel_Pose>(Channel))
+	for (UPoseSearchFeatureChannel* Channel : Channels)
+	{
+		if (UPoseSearchFeatureChannel_Pose* ChannelPose = Cast<UPoseSearchFeatureChannel_Pose>(Channel))
+		{
+			for (FPoseSearchBone& SampledBone : ChannelPose->SampledBones)
+			if (SampledBone.bUseVelocity_DEPRECATED ||
+				SampledBone.bUsePosition_DEPRECATED ||
+				SampledBone.bUseRotation_DEPRECATED ||
+				SampledBone.bUsePhase_DEPRECATED)
 			{
-				if (ChannelPose->SchemaBoneIdx.Num() == 0)
+				SampledBone.Flags = 0;
+				if (SampledBone.bUseVelocity_DEPRECATED)
 				{
-					bNeedFinalize = true;
-					break;
+					SampledBone.Flags |= EPoseSearchBoneFlags::Velocity;
+					SampledBone.bUseVelocity_DEPRECATED = false;
+				}
+				
+				if (SampledBone.bUsePosition_DEPRECATED)
+				{
+					SampledBone.Flags |= EPoseSearchBoneFlags::Position;
+					SampledBone.bUsePosition_DEPRECATED = false;
 				}
+
+				if (SampledBone.bUseRotation_DEPRECATED)
+				{
+					SampledBone.Flags |= EPoseSearchBoneFlags::Rotation;
+					SampledBone.bUseRotation_DEPRECATED = false;
+				}
+
+				if (SampledBone.bUsePhase_DEPRECATED)
+				{
+					SampledBone.Flags |= EPoseSearchBoneFlags::Phase;
+					SampledBone.bUsePhase_DEPRECATED = false;
+				}
+
+				bNeedFinalize = true;
 			}
 		}
+
+		if (UPoseSearchFeatureChannel_Trajectory* ChannelTrajectory = Cast<UPoseSearchFeatureChannel_Trajectory>(Channel))
+		{
+			if (ChannelTrajectory->SampleOffsets_DEPRECATED.Num() > 0)
+			{
+				ChannelTrajectory->Samples.Reset();
+
+				for (float Offset : ChannelTrajectory->SampleOffsets_DEPRECATED)
+				{
+					FPoseSearchTrajectorySample Sample;
+					Sample.Flags = 0;
+					if (ChannelTrajectory->bUsePositions_DEPRECATED)
+					{
+						Sample.Flags |= EPoseSearchTrajectoryFlags::Position;
+					}
+					if (ChannelTrajectory->bUseLinearVelocities_DEPRECATED)
+					{
+						Sample.Flags |= EPoseSearchTrajectoryFlags::Velocity;
+					}
+					if (ChannelTrajectory->bUseFacingDirections_DEPRECATED)
+					{
+						Sample.Flags |= EPoseSearchTrajectoryFlags::FacingDirection;
+					}
+					Sample.Offset = Offset;
+					Sample.Weight = 1.f;
+					ChannelTrajectory->Samples.Add(Sample);
+				}
+
+				ChannelTrajectory->SampleOffsets_DEPRECATED.Reset();
+				ChannelTrajectory->bUsePositions_DEPRECATED = false;
+				ChannelTrajectory->bUseLinearVelocities_DEPRECATED = false;
+				ChannelTrajectory->bUseFacingDirections_DEPRECATED = false;
+			}
+		}		
 	}
 
 #endif // WITH_EDITORONLY_DATA
@@ -1968,7 +2033,7 @@ UE::PoseSearch::FSearchResult UPoseSearchDatabase::SearchPCAKDTree(UE::PoseSearc
 
 	TArrayView<const float> NormalizedQueryValues = Result.ComposedQuery.GetNormalizedValues();
 
-	const bool IsCurrentResultFromThisDatabase = SearchContext.CurrentResult.IsValid() && SearchContext.CurrentResult.Database == this;
+	const bool IsCurrentResultFromThisDatabase = !SearchContext.bForceInterrupt && SearchContext.CurrentResult.IsValid() && SearchContext.CurrentResult.Database == this;
 	if (IsCurrentResultFromThisDatabase)
 	{
 		Result.ContinuityPoseCost = ComparePoses(
@@ -2055,7 +2120,7 @@ UE::PoseSearch::FSearchResult UPoseSearchDatabase::SearchBruteForce(UE::PoseSear
 	constexpr int NonSelectableIdxDataSize = 128;
 	size_t* NonSelectableIdxData((size_t*)FMemory_Alloca(NonSelectableIdxDataSize * sizeof(size_t)));
 	int NonSelectableIdxUsedSize = 0;
-	const bool IsCurrentResultFromThisDatabase = SearchContext.CurrentResult.IsValid() && SearchContext.CurrentResult.Database == this;
+	const bool IsCurrentResultFromThisDatabase = !SearchContext.bForceInterrupt && SearchContext.CurrentResult.IsValid() && SearchContext.CurrentResult.Database == this;
 	if (IsCurrentResultFromThisDatabase)
 	{
 		Result.ContinuityPoseCost = ComparePoses(
diff --git a/Engine/Plugins/Experimental/Animation/PoseSearch/Source/Runtime/Private/PoseSearchFeatureChannels.cpp b/Engine/Plugins/Experimental/Animation/PoseSearch/Source/Runtime/Private/PoseSearchFeatureChannels.cpp
--- a/Engine/Plugins/Experimental/Animation/PoseSearch/Source/Runtime/Private/PoseSearchFeatureChannels.cpp
+++ b/Engine/Plugins/Experimental/Animation/PoseSearch/Source/Runtime/Private/PoseSearchFeatureChannels.cpp
@@ -345,7 +345,7 @@ void UPoseSearchFeatureChannel_Position::IndexAsset(UE::PoseSearch::IAssetIndexe
 		const float SubsampleTime = OriginSampleTime + SampleTimeOffset;
 
 		bool ClampedPresent;
-		const FTransform BoneTransformsPresent = Indexer.GetTransformAndCacheResults(SubsampleTime, UseSampleTimeOffsetRootBone ? SubsampleTime : OriginSampleTime, SchemaBoneIdx, ClampedPresent);
+		const FTransform BoneTransformsPresent = Indexer.GetTransformAndCacheResults(SubsampleTime, bUseSampleTimeOffsetRootBone ? SubsampleTime : OriginSampleTime, SchemaBoneIdx, ClampedPresent);
 		int32 DataOffset = ChannelDataOffset;
 		FFeatureVectorHelper::EncodeVector(FeatureVector.EditValues(), DataOffset, BoneTransformsPresent.GetTranslation());
 	}
@@ -356,14 +356,15 @@ void UPoseSearchFeatureChannel_Position::GenerateDDCKey(FBlake3& InOutKeyHasher)
 	InOutKeyHasher.Update(&Bone, sizeof(Bone));
 	InOutKeyHasher.Update(&Weight, sizeof(Weight));
 	InOutKeyHasher.Update(&SampleTimeOffset, sizeof(SampleTimeOffset));
+	InOutKeyHasher.Update(&bUseFeaturesFromContinuityPose, sizeof(bUseFeaturesFromContinuityPose));
+	InOutKeyHasher.Update(&bUseSampleTimeOffsetRootBone, sizeof(bUseSampleTimeOffsetRootBone));
 }
 
 bool UPoseSearchFeatureChannel_Position::BuildQuery(UE::PoseSearch::FSearchContext& SearchContext, FPoseSearchFeatureVectorBuilder& InOutQuery) const
 {
 	using namespace UE::PoseSearch;
 
-	const bool bSkip = !UseFeaturesFromContinuityPose && SearchContext.CurrentResult.IsValid() && SearchContext.CurrentResult.Database->Schema == InOutQuery.GetSchema();
-
+	const bool bSkip = !bUseFeaturesFromContinuityPose && SearchContext.CurrentResult.IsValid() && SearchContext.CurrentResult.Database->Schema == InOutQuery.GetSchema();
 	if (bSkip)
 	{
 		return true;
@@ -377,7 +378,7 @@ bool UPoseSearchFeatureChannel_Position::BuildQuery(UE::PoseSearch::FSearchConte
 	bool AnyError = false;
 	FTransform Transform = SearchContext.TryGetTransformAndCacheResults(SampleTimeOffset, InOutQuery.GetSchema(), SchemaBoneIdx, AnyError);
 
-	if (!UseSampleTimeOffsetRootBone)
+	if (!bUseSampleTimeOffsetRootBone)
 	{
 		const FTransform RootTransform = SearchContext.TryGetTransformAndCacheResults(0.f, InOutQuery.GetSchema(), FSearchContext::SchemaRootBoneIdx, AnyError);
 		const FTransform RootTransformPrev = SearchContext.TryGetTransformAndCacheResults(SampleTimeOffset, InOutQuery.GetSchema(), FSearchContext::SchemaRootBoneIdx, AnyError);
@@ -409,7 +410,7 @@ void UPoseSearchFeatureChannel_Position::DebugDraw(const UE::PoseSearch::FDebugD
 
 	if (EnumHasAnyFlags(DrawParams.Flags, EDebugDrawFlags::DrawFast | EDebugDrawFlags::DrawSearchIndex))
 	{
-		DrawDebugPoint(DrawParams.World, BonePos, DrawParams.PointSize, Color, bPersistent, DrawParams.DefaultLifeTime, DepthPriority);
+		DrawDebugPoint(DrawParams.World, BonePos, DrawParams.PointSize, Color, bPersistent, LifeTime, DepthPriority);
 	}
 	else
 	{
@@ -422,6 +423,139 @@ void UPoseSearchFeatureChannel_Position::DebugDraw(const UE::PoseSearch::FDebugD
 	}
 }
 
+//////////////////////////////////////////////////////////////////////////
+// UPoseSearchFeatureChannel_Heading
+
+void UPoseSearchFeatureChannel_Heading::InitializeSchema(UE::PoseSearch::FSchemaInitializer& Initializer)
+{
+	Super::InitializeSchema(Initializer);
+	ChannelCardinality = UE::PoseSearch::FFeatureVectorHelper::EncodeVectorCardinality;
+	Initializer.SetCurrentChannelDataOffset(ChannelDataOffset + ChannelCardinality);
+	SchemaBoneIdx = Initializer.AddBoneReference(Bone);
+}
+
+void UPoseSearchFeatureChannel_Heading::FillWeights(TArray<float>& Weights) const
+{
+	using namespace UE::PoseSearch;
+
+	for (int32 i = 0; i != FFeatureVectorHelper::EncodeVectorCardinality; ++i)
+	{
+		Weights[ChannelDataOffset + i] = Weight;
+	}
+}
+
+FVector UPoseSearchFeatureChannel_Heading::GetAxis(const FQuat& Rotation) const
+{
+	switch (HeadingAxis)
+	{
+	case HeadingAxis::X:
+		return Rotation.GetAxisX();
+	case HeadingAxis::Y:
+		return Rotation.GetAxisY();
+	case HeadingAxis::Z:
+		return Rotation.GetAxisZ();
+	}
+
+	checkNoEntry();
+	return FVector(1.f, 0.f, 0.f);
+}
+
+void UPoseSearchFeatureChannel_Heading::IndexAsset(UE::PoseSearch::IAssetIndexer& Indexer, UE::PoseSearch::FAssetIndexingOutput& IndexingOutput) const
+{
+	using namespace UE::PoseSearch;
+
+	const FAssetIndexingContext& IndexingContext = Indexer.GetIndexingContext();
+	const FAssetSamplingContext* SamplingContext = IndexingContext.SamplingContext;
+
+	for (int32 SampleIdx = IndexingContext.BeginSampleIdx; SampleIdx != IndexingContext.EndSampleIdx; ++SampleIdx)
+	{
+		int32 VectorIdx = SampleIdx - IndexingContext.BeginSampleIdx;
+		FPoseSearchFeatureVectorBuilder& FeatureVector = IndexingOutput.PoseVectors[VectorIdx];
+
+		const float OriginSampleTime = FMath::Min(SampleIdx * IndexingContext.Schema->SamplingInterval, IndexingContext.MainSampler->GetPlayLength());
+		const float SubsampleTime = OriginSampleTime + SampleTimeOffset;
+
+		bool ClampedPresent;
+		const FTransform BoneTransformsPresent = Indexer.GetTransformAndCacheResults(SubsampleTime, bUseSampleTimeOffsetRootBone ? SubsampleTime : OriginSampleTime, SchemaBoneIdx, ClampedPresent);
+		int32 DataOffset = ChannelDataOffset;
+		
+		FFeatureVectorHelper::EncodeVector(FeatureVector.EditValues(), DataOffset, GetAxis(BoneTransformsPresent.GetRotation()));
+	}
+}
+
+void UPoseSearchFeatureChannel_Heading::GenerateDDCKey(FBlake3& InOutKeyHasher) const
+{
+	InOutKeyHasher.Update(&Bone, sizeof(Bone));
+	InOutKeyHasher.Update(&Weight, sizeof(Weight));
+	InOutKeyHasher.Update(&SampleTimeOffset, sizeof(SampleTimeOffset));
+	InOutKeyHasher.Update(&HeadingAxis, sizeof(HeadingAxis));
+	InOutKeyHasher.Update(&bUseFeaturesFromContinuityPose, sizeof(bUseFeaturesFromContinuityPose));
+	InOutKeyHasher.Update(&bUseSampleTimeOffsetRootBone, sizeof(bUseSampleTimeOffsetRootBone));
+}
+
+bool UPoseSearchFeatureChannel_Heading::BuildQuery(UE::PoseSearch::FSearchContext& SearchContext, FPoseSearchFeatureVectorBuilder& InOutQuery) const
+{
+	using namespace UE::PoseSearch;
+
+	const bool bSkip = !bUseFeaturesFromContinuityPose && SearchContext.CurrentResult.IsValid() && SearchContext.CurrentResult.Database->Schema == InOutQuery.GetSchema();
+	if (bSkip)
+	{
+		return true;
+	}
+
+	if (!SearchContext.History)
+	{
+		return false;
+	}
+
+	bool AnyError = false;
+	FTransform Transform = SearchContext.TryGetTransformAndCacheResults(SampleTimeOffset, InOutQuery.GetSchema(), SchemaBoneIdx, AnyError);
+
+	if (!bUseSampleTimeOffsetRootBone)
+	{
+		const FTransform RootTransform = SearchContext.TryGetTransformAndCacheResults(0.f, InOutQuery.GetSchema(), FSearchContext::SchemaRootBoneIdx, AnyError);
+		const FTransform RootTransformPrev = SearchContext.TryGetTransformAndCacheResults(SampleTimeOffset, InOutQuery.GetSchema(), FSearchContext::SchemaRootBoneIdx, AnyError);
+		Transform = Transform * (RootTransformPrev * RootTransform.Inverse());
+	}
+
+	int32 DataOffset = ChannelDataOffset;
+	FFeatureVectorHelper::EncodeVector(InOutQuery.EditValues(), DataOffset, GetAxis(Transform.GetRotation()));
+
+	return !AnyError;
+}
+
+void UPoseSearchFeatureChannel_Heading::DebugDraw(const UE::PoseSearch::FDebugDrawParams& DrawParams, TArrayView<const float> PoseVector) const
+{
+	using namespace UE::PoseSearch;
+
+	const UPoseSearchSchema* Schema = DrawParams.GetSchema();
+	check(Schema && Schema->IsValid());
+
+	const float LifeTime = DrawParams.DefaultLifeTime;
+	const uint8 DepthPriority = ESceneDepthPriorityGroup::SDPG_Foreground + 2;
+	const bool bPersistent = EnumHasAnyFlags(DrawParams.Flags, EDebugDrawFlags::Persistent);
+
+	// @todo: search in the schema for a UPoseSearchFeatureChannel_Position with the same bone as SchemaBoneIdx and decode it's data as position
+	// @todo: initialize with the character position instead of FVector::Zero?
+	const FVector BonePos = DrawParams.Mesh != nullptr ? DrawParams.Mesh->GetSocketTransform(Bone.BoneName).GetLocation() : FVector::Zero();
+
+	int32 DataOffset = ChannelDataOffset;
+	const FVector BoneHeading = DrawParams.RootTransform.TransformPosition(FFeatureVectorHelper::DecodeVector(PoseVector, DataOffset));
+
+	const FLinearColor LinearColor = DrawParams.GetColor(this);
+	const FColor Color = LinearColor.ToFColor(true);
+
+	if (EnumHasAnyFlags(DrawParams.Flags, EDebugDrawFlags::DrawSearchIndex))
+	{
+		DrawDebugLine(DrawParams.World, BonePos, BonePos + BoneHeading, Color, bPersistent, LifeTime, DepthPriority);
+	}
+	else
+	{
+		const float AdjustedThickness = EnumHasAnyFlags(DrawParams.Flags, EDebugDrawFlags::DrawFast) ? 0.0f : DrawDebugLineThickness;
+		DrawDebugLine(DrawParams.World, BonePos, BonePos + BoneHeading, Color, bPersistent, LifeTime, DepthPriority, AdjustedThickness);
+	}
+}
+
 //////////////////////////////////////////////////////////////////////////
 // UPoseSearchFeatureChannel_Pose
 
@@ -441,28 +575,28 @@ void UPoseSearchFeatureChannel_Pose::InitializeSchema(UE::PoseSearch::FSchemaIni
 	for (int32 ChannelBoneIdx = 0; ChannelBoneIdx != SampledBones.Num(); ++ChannelBoneIdx)
 	{
 		const FPoseSearchBone& SampledBone = SampledBones[ChannelBoneIdx];
-		if (SampledBone.bUsePosition)
+		if (EnumHasAnyFlags(SampledBone.Flags, EPoseSearchBoneFlags::Position))
 		{
 			for (int32 SubsampleIdx = 0; SubsampleIdx != SampleTimes.Num(); ++SubsampleIdx)
 			{
 				DataOffset += FFeatureVectorHelper::EncodeVectorCardinality;
 			}
 		}
-		if (SampledBone.bUseRotation)
+		if (EnumHasAnyFlags(SampledBone.Flags, EPoseSearchBoneFlags::Rotation))
 		{
 			for (int32 SubsampleIdx = 0; SubsampleIdx != SampleTimes.Num(); ++SubsampleIdx)
 			{
 				DataOffset += FFeatureVectorHelper::EncodeQuatCardinality;
 			}
 		}
-		if (SampledBone.bUseVelocity)
+		if (EnumHasAnyFlags(SampledBone.Flags, EPoseSearchBoneFlags::Velocity))
 		{
 			for (int32 SubsampleIdx = 0; SubsampleIdx != SampleTimes.Num(); ++SubsampleIdx)
 			{
 				DataOffset += FFeatureVectorHelper::EncodeVectorCardinality;
 			}
 		}
-		if (SampledBone.bUsePhase)
+		if (EnumHasAnyFlags(SampledBone.Flags, EPoseSearchBoneFlags::Phase))
 		{
 			for (int32 SubsampleIdx = 0; SubsampleIdx != SampleTimes.Num(); ++SubsampleIdx)
 			{
@@ -492,46 +626,46 @@ void UPoseSearchFeatureChannel_Pose::FillWeights(TArray<float>& Weights) const
 	for (int32 ChannelBoneIdx = 0; ChannelBoneIdx != NumBones; ++ChannelBoneIdx)
 	{
 		const FPoseSearchBone& SampledBone = SampledBones[ChannelBoneIdx];
-		if (SampledBone.bUsePosition)
+		if (EnumHasAnyFlags(SampledBone.Flags, EPoseSearchBoneFlags::Position))
 		{
 			for (int32 SubsampleIdx = 0; SubsampleIdx != SampleTimes.Num(); ++SubsampleIdx)
 			{
 				for (int32 i = 0; i != FFeatureVectorHelper::EncodeVectorCardinality; ++i)
 				{
-					Weights[DataOffset + i] = SampledBone.Weight;
+					Weights[DataOffset + i] = Weight * SampledBone.Weight;
 				}
 				DataOffset += FFeatureVectorHelper::EncodeVectorCardinality;
 			}
 		}
-		if (SampledBone.bUseRotation)
+		if (EnumHasAnyFlags(SampledBone.Flags, EPoseSearchBoneFlags::Rotation))
 		{
 			for (int32 SubsampleIdx = 0; SubsampleIdx != SampleTimes.Num(); ++SubsampleIdx)
 			{
 				for (int32 i = 0; i != FFeatureVectorHelper::EncodeQuatCardinality; ++i)
 				{
-					Weights[DataOffset + i] = SampledBone.Weight;
+					Weights[DataOffset + i] = Weight * SampledBone.Weight;
 				}
 				DataOffset += FFeatureVectorHelper::EncodeQuatCardinality;
 			}
 		}
-		if (SampledBone.bUseVelocity)
+		if (EnumHasAnyFlags(SampledBone.Flags, EPoseSearchBoneFlags::Velocity))
 		{
 			for (int32 SubsampleIdx = 0; SubsampleIdx != SampleTimes.Num(); ++SubsampleIdx)
 			{
 				for (int32 i = 0; i != FFeatureVectorHelper::EncodeVectorCardinality; ++i)
 				{
-					Weights[DataOffset + i] = SampledBone.Weight;
+					Weights[DataOffset + i] = Weight * SampledBone.Weight;
 				}
 				DataOffset += FFeatureVectorHelper::EncodeVectorCardinality;
 			}
 		}
-		if (SampledBone.bUsePhase)
+		if (EnumHasAnyFlags(SampledBone.Flags, EPoseSearchBoneFlags::Phase))
 		{
 			for (int32 SubsampleIdx = 0; SubsampleIdx != SampleTimes.Num(); ++SubsampleIdx)
 			{
 				for (int32 i = 0; i != FFeatureVectorHelper::EncodeVector2DCardinality; ++i)
 				{
-					Weights[DataOffset + i] = SampledBone.Weight;
+					Weights[DataOffset + i] = Weight * SampledBone.Weight;
 				}
 				DataOffset += FFeatureVectorHelper::EncodeVector2DCardinality;
 			}
@@ -562,7 +696,7 @@ void UPoseSearchFeatureChannel_Pose::CalculatePhases(UE::PoseSearch::IAssetIndex
 	for (int32 ChannelBoneIdx = 0; ChannelBoneIdx != SampledBones.Num(); ++ChannelBoneIdx)
 	{
 		const FPoseSearchBone& SampledBone = SampledBones[ChannelBoneIdx];
-		if (SampledBone.bUsePhase)
+		if (EnumHasAnyFlags(SampledBone.Flags, EPoseSearchBoneFlags::Phase))
 		{
 			CollectBonePositions(BonePositions, Indexer, SchemaBoneIdx[ChannelBoneIdx]);
 
@@ -612,28 +746,28 @@ void UPoseSearchFeatureChannel_Pose::ComputeMeanDeviations(const Eigen::MatrixXd
 	for (int32 ChannelBoneIdx = 0; ChannelBoneIdx != NumBones; ++ChannelBoneIdx)
 	{
 		const FPoseSearchBone& SampledBone = SampledBones[ChannelBoneIdx];
-		if (SampledBone.bUsePosition)
+		if (EnumHasAnyFlags(SampledBone.Flags, EPoseSearchBoneFlags::Position))
 		{
 			for (int32 SubsampleIdx = 0; SubsampleIdx != SampleTimes.Num(); ++SubsampleIdx)
 			{
 				FFeatureVectorHelper::ComputeMeanDeviations(CenteredPoseMatrix, MeanDeviations, DataOffset, FFeatureVectorHelper::EncodeVectorCardinality);
 			}
 		}
-		if (SampledBone.bUseRotation)
+		if (EnumHasAnyFlags(SampledBone.Flags, EPoseSearchBoneFlags::Rotation))
 		{
 			for (int32 SubsampleIdx = 0; SubsampleIdx != SampleTimes.Num(); ++SubsampleIdx)
 			{
 				FFeatureVectorHelper::ComputeMeanDeviations(CenteredPoseMatrix, MeanDeviations, DataOffset, FFeatureVectorHelper::EncodeQuatCardinality);
 			}
 		}
-		if (SampledBone.bUseVelocity)
+		if (EnumHasAnyFlags(SampledBone.Flags, EPoseSearchBoneFlags::Velocity))
 		{
 			for (int32 SubsampleIdx = 0; SubsampleIdx != SampleTimes.Num(); ++SubsampleIdx)
 			{
 				FFeatureVectorHelper::ComputeMeanDeviations(CenteredPoseMatrix, MeanDeviations, DataOffset, FFeatureVectorHelper::EncodeVectorCardinality);
 			}
 		}
-		if (SampledBone.bUsePhase)
+		if (EnumHasAnyFlags(SampledBone.Flags, EPoseSearchBoneFlags::Phase))
 		{
 			FFeatureVectorHelper::ComputeMeanDeviations(CenteredPoseMatrix, MeanDeviations, DataOffset, FFeatureVectorHelper::EncodeVector2DCardinality);
 		}
@@ -673,7 +807,7 @@ void UPoseSearchFeatureChannel_Pose::AddPoseFeatures(UE::PoseSearch::IAssetIndex
 		// Think of this process as freezing the character in place (at SampleTime) and then tracing the paths of their joints
 		// as they move through space from past to present to future (at times indicated by PoseSampleTimes).
 
-		if (SampledBone.bUsePosition)
+		if (EnumHasAnyFlags(SampledBone.Flags, EPoseSearchBoneFlags::Position))
 		{
 			for (int32 SubsampleIdx = 0; SubsampleIdx != SampleTimes.Num(); ++SubsampleIdx)
 			{
@@ -685,7 +819,7 @@ void UPoseSearchFeatureChannel_Pose::AddPoseFeatures(UE::PoseSearch::IAssetIndex
 			}
 		}
 
-		if (SampledBone.bUseRotation)
+		if (EnumHasAnyFlags(SampledBone.Flags, EPoseSearchBoneFlags::Rotation))
 		{
 			for (int32 SubsampleIdx = 0; SubsampleIdx != SampleTimes.Num(); ++SubsampleIdx)
 			{
@@ -697,7 +831,7 @@ void UPoseSearchFeatureChannel_Pose::AddPoseFeatures(UE::PoseSearch::IAssetIndex
 			}
 		}
 
-		if (SampledBone.bUseVelocity)
+		if (EnumHasAnyFlags(SampledBone.Flags, EPoseSearchBoneFlags::Velocity))
 		{
 			for (int32 SubsampleIdx = 0; SubsampleIdx != SampleTimes.Num(); ++SubsampleIdx)
 			{
@@ -730,7 +864,7 @@ void UPoseSearchFeatureChannel_Pose::AddPoseFeatures(UE::PoseSearch::IAssetIndex
 			}
 		}
 
-		if (SampledBone.bUsePhase)
+		if (EnumHasAnyFlags(SampledBone.Flags, EPoseSearchBoneFlags::Phase))
 		{
 			for (int32 SubsampleIdx = 0; SubsampleIdx != SampleTimes.Num(); ++SubsampleIdx)
 			{
@@ -743,15 +877,17 @@ void UPoseSearchFeatureChannel_Pose::AddPoseFeatures(UE::PoseSearch::IAssetIndex
 
 void UPoseSearchFeatureChannel_Pose::GenerateDDCKey(FBlake3& InOutKeyHasher) const
 {
+	InOutKeyHasher.Update(&Weight, sizeof(Weight));
 	InOutKeyHasher.Update(MakeMemoryView(SampledBones));
 	InOutKeyHasher.Update(MakeMemoryView(SampleTimes));
+	InOutKeyHasher.Update(&bUseFeaturesFromContinuityPose, sizeof(bUseFeaturesFromContinuityPose));
 }
 
 bool UPoseSearchFeatureChannel_Pose::BuildQuery(UE::PoseSearch::FSearchContext& SearchContext, FPoseSearchFeatureVectorBuilder& InOutQuery) const
 {
 	using namespace UE::PoseSearch;
 
-	const bool bSkip = SearchContext.CurrentResult.IsValid() && SearchContext.CurrentResult.Database->Schema == InOutQuery.GetSchema();
+	const bool bSkip = !bUseFeaturesFromContinuityPose && SearchContext.CurrentResult.IsValid() && SearchContext.CurrentResult.Database->Schema == InOutQuery.GetSchema();
 	if (bSkip)
 	{
 		return true;
@@ -787,7 +923,7 @@ bool UPoseSearchFeatureChannel_Pose::BuildQuery(UE::PoseSearch::FSearchContext&
 			CachedTransforms[CachedTransformsIndex].Current = SearchContext.TryGetTransformAndCacheResults(SampleTime, InOutQuery.GetSchema(), SchemaBoneIdx[SampledBoneIdx], AnyError);
 			const FPoseSearchBone& SampledBone = SampledBones[SampledBoneIdx];
 
-			if (SampledBone.bUseVelocity)
+			if (EnumHasAnyFlags(SampledBone.Flags, EPoseSearchBoneFlags::Velocity))
 			{
 				check(SearchContext.History);
 				const float HistorySameplInterval = SearchContext.History->GetSampleTimeInterval();
@@ -816,7 +952,7 @@ bool UPoseSearchFeatureChannel_Pose::BuildQuery(UE::PoseSearch::FSearchContext&
 	for (int32 SampledBoneIdx = 0; SampledBoneIdx != SampledBones.Num(); ++SampledBoneIdx)
 	{
 		const FPoseSearchBone& SampledBone = SampledBones[SampledBoneIdx];
-		if (SampledBone.bUsePosition)
+		if (EnumHasAnyFlags(SampledBone.Flags, EPoseSearchBoneFlags::Position))
 		{
 			for (int32 SubsampleIdx = 0; SubsampleIdx != SampleTimes.Num(); ++SubsampleIdx)
 			{
@@ -833,7 +969,7 @@ bool UPoseSearchFeatureChannel_Pose::BuildQuery(UE::PoseSearch::FSearchContext&
 			}
 		}
 
-		if (SampledBone.bUseRotation)
+		if (EnumHasAnyFlags(SampledBone.Flags, EPoseSearchBoneFlags::Rotation))
 		{
 			for (int32 SubsampleIdx = 0; SubsampleIdx != SampleTimes.Num(); ++SubsampleIdx)
 			{
@@ -850,7 +986,7 @@ bool UPoseSearchFeatureChannel_Pose::BuildQuery(UE::PoseSearch::FSearchContext&
 			}
 		}
 
-		if (SampledBone.bUseVelocity)
+		if (EnumHasAnyFlags(SampledBone.Flags, EPoseSearchBoneFlags::Velocity))
 		{
 			for (int32 SubsampleIdx = 0; SubsampleIdx != SampleTimes.Num(); ++SubsampleIdx)
 			{
@@ -868,7 +1004,7 @@ bool UPoseSearchFeatureChannel_Pose::BuildQuery(UE::PoseSearch::FSearchContext&
 			}
 		}
 
-		if (SampledBone.bUsePhase)
+		if (EnumHasAnyFlags(SampledBone.Flags, EPoseSearchBoneFlags::Phase))
 		{
 			for (int32 SubsampleIdx = 0; SubsampleIdx != SampleTimes.Num(); ++SubsampleIdx)
 			{
@@ -910,7 +1046,7 @@ void UPoseSearchFeatureChannel_Pose::DebugDraw(const UE::PoseSearch::FDebugDrawP
 
 		TArray<FVector> BonePos;
 		BonePos.AddUninitialized(NumSubsamples);
-		if (SampledBone.bUsePosition)
+		if (EnumHasAnyFlags(SampledBone.Flags, EPoseSearchBoneFlags::Position))
 		{
 			for (int32 SubsampleIdx = 0; SubsampleIdx != NumSubsamples; ++SubsampleIdx)
 			{
@@ -922,7 +1058,7 @@ void UPoseSearchFeatureChannel_Pose::DebugDraw(const UE::PoseSearch::FDebugDrawP
 				BonePos[SubsampleIdx] = DrawParams.RootTransform.TransformPosition(BonePos[SubsampleIdx]);
 				if (EnumHasAnyFlags(DrawParams.Flags, EDebugDrawFlags::DrawFast | EDebugDrawFlags::DrawSearchIndex))
 				{
-					DrawDebugPoint(DrawParams.World, BonePos[SubsampleIdx], DrawParams.PointSize, Color, bPersistent, DrawParams.DefaultLifeTime, DepthPriority);
+					DrawDebugPoint(DrawParams.World, BonePos[SubsampleIdx], DrawParams.PointSize, Color, bPersistent, LifeTime, DepthPriority);
 				}
 				else
 				{
@@ -947,7 +1083,7 @@ void UPoseSearchFeatureChannel_Pose::DebugDraw(const UE::PoseSearch::FDebugDrawP
 			}
 		}
 
-		if (SampledBone.bUseRotation)
+		if (EnumHasAnyFlags(SampledBone.Flags, EPoseSearchBoneFlags::Rotation))
 		{
 			for (int32 SubsampleIdx = 0; SubsampleIdx != NumSubsamples; ++SubsampleIdx)
 			{
@@ -957,7 +1093,7 @@ void UPoseSearchFeatureChannel_Pose::DebugDraw(const UE::PoseSearch::FDebugDrawP
 			}
 		}
 
-		if (SampledBone.bUseVelocity)
+		if (EnumHasAnyFlags(SampledBone.Flags, EPoseSearchBoneFlags::Velocity))
 		{
 			for (int32 SubsampleIdx = 0; SubsampleIdx != NumSubsamples; ++SubsampleIdx)
 			{
@@ -972,13 +1108,11 @@ void UPoseSearchFeatureChannel_Pose::DebugDraw(const UE::PoseSearch::FDebugDrawP
 
 				if (EnumHasAnyFlags(DrawParams.Flags, EDebugDrawFlags::DrawSearchIndex))
 				{
-					DrawDebugLine(DrawParams.World, BonePos[SubsampleIdx], BonePos[SubsampleIdx] + BoneVel, Color, bPersistent, DrawParams.DefaultLifeTime, DepthPriority);
+					DrawDebugLine(DrawParams.World, BonePos[SubsampleIdx], BonePos[SubsampleIdx] + BoneVel, Color, bPersistent, LifeTime, DepthPriority);
 				}
 				else
 				{
-					const float AdjustedThickness =
-						EnumHasAnyFlags(DrawParams.Flags, EDebugDrawFlags::DrawFast) ?
-						0.0f : DrawDebugLineThickness;
+					const float AdjustedThickness = EnumHasAnyFlags(DrawParams.Flags, EDebugDrawFlags::DrawFast) ? 0.0f : DrawDebugLineThickness;
 
 					DrawDebugLine(
 						DrawParams.World,
@@ -993,7 +1127,7 @@ void UPoseSearchFeatureChannel_Pose::DebugDraw(const UE::PoseSearch::FDebugDrawP
 			}
 		}
 
-		if (SampledBone.bUsePhase)
+		if (EnumHasAnyFlags(SampledBone.Flags, EPoseSearchBoneFlags::Phase))
 		{
 			for (int32 SubsampleIdx = 0; SubsampleIdx != NumSubsamples; ++SubsampleIdx)
 			{
@@ -1026,7 +1160,10 @@ void UPoseSearchFeatureChannel_Pose::DebugDraw(const UE::PoseSearch::FDebugDrawP
 // UPoseSearchFeatureChannel_Trajectory
 void UPoseSearchFeatureChannel_Trajectory::PreSave(FObjectPreSaveContext ObjectSaveContext)
 {
-	SampleOffsets.Sort(TLess<>());
+	Samples.Sort([](const FPoseSearchTrajectorySample& a, const FPoseSearchTrajectorySample& b)
+	{
+		return a.Offset < b.Offset;
+	});
 
 	Super::PreSave(ObjectSaveContext);
 }
@@ -1038,25 +1175,18 @@ void UPoseSearchFeatureChannel_Trajectory::InitializeSchema(UE::PoseSearch::FSch
 	Super::InitializeSchema(Initializer);
 
 	int32 DataOffset = ChannelDataOffset;
-	if (bUsePositions)
+
+	for (const FPoseSearchTrajectorySample& Sample : Samples)
 	{
-		for (int32 SubsampleIdx = 0; SubsampleIdx != SampleOffsets.Num(); ++SubsampleIdx)
+		if (EnumHasAnyFlags(Sample.Flags, EPoseSearchTrajectoryFlags::Position))
 		{
 			DataOffset += FFeatureVectorHelper::EncodeVectorCardinality;
 		}
-	}
-
-	if (bUseLinearVelocities)
-	{
-		for (int32 SubsampleIdx = 0; SubsampleIdx != SampleOffsets.Num(); ++SubsampleIdx)
+		if (EnumHasAnyFlags(Sample.Flags, EPoseSearchTrajectoryFlags::Velocity))
 		{
 			DataOffset += FFeatureVectorHelper::EncodeVectorCardinality;
 		}
-	}
-
-	if (bUseFacingDirections)
-	{
-		for (int32 SubsampleIdx = 0; SubsampleIdx != SampleOffsets.Num(); ++SubsampleIdx)
+		if (EnumHasAnyFlags(Sample.Flags, EPoseSearchTrajectoryFlags::FacingDirection))
 		{
 			DataOffset += FFeatureVectorHelper::EncodeVectorCardinality;
 		}
@@ -1069,12 +1199,38 @@ void UPoseSearchFeatureChannel_Trajectory::InitializeSchema(UE::PoseSearch::FSch
 
 void UPoseSearchFeatureChannel_Trajectory::FillWeights(TArray<float>& Weights) const
 {
-	const int32 Begin = ChannelDataOffset;
-	const int32 End = ChannelDataOffset + ChannelCardinality;
-	for (int i = Begin; i < End; ++i)
+	using namespace UE::PoseSearch;
+
+	int32 DataOffset = ChannelDataOffset;
+	for (const FPoseSearchTrajectorySample& Sample : Samples)
 	{
-		Weights[i] = Weight;
+		if (EnumHasAnyFlags(Sample.Flags, EPoseSearchTrajectoryFlags::Position))
+		{
+			for (int32 i = 0; i != FFeatureVectorHelper::EncodeVectorCardinality; ++i)
+			{
+				Weights[DataOffset + i] = Weight * Sample.Weight;
+			}
+			DataOffset += FFeatureVectorHelper::EncodeVectorCardinality;
+		}
+		if (EnumHasAnyFlags(Sample.Flags, EPoseSearchTrajectoryFlags::Velocity))
+		{
+			for (int32 i = 0; i != FFeatureVectorHelper::EncodeVectorCardinality; ++i)
+			{
+				Weights[DataOffset + i] = Weight * Sample.Weight;
+			}
+			DataOffset += FFeatureVectorHelper::EncodeVectorCardinality;
+		}
+		if (EnumHasAnyFlags(Sample.Flags, EPoseSearchTrajectoryFlags::FacingDirection))
+		{
+			for (int32 i = 0; i != FFeatureVectorHelper::EncodeVectorCardinality; ++i)
+			{
+				Weights[DataOffset + i] = Weight * Sample.Weight;
+			}
+			DataOffset += FFeatureVectorHelper::EncodeVectorCardinality;
+		}
 	}
+
+	check(DataOffset == ChannelDataOffset + ChannelCardinality);
 }
 
 void UPoseSearchFeatureChannel_Trajectory::IndexAsset(UE::PoseSearch::IAssetIndexer& Indexer,  UE::PoseSearch::FAssetIndexingOutput& IndexingOutput) const
@@ -1092,25 +1248,18 @@ void UPoseSearchFeatureChannel_Trajectory::ComputeMeanDeviations(const Eigen::Ma
 	using namespace UE::PoseSearch;
 
 	int32 DataOffset = ChannelDataOffset;
-	if (bUsePositions)
+
+	for (const FPoseSearchTrajectorySample& Sample : Samples)
 	{
-		for (int32 SubsampleIdx = 0; SubsampleIdx != SampleOffsets.Num(); ++SubsampleIdx)
+		if (EnumHasAnyFlags(Sample.Flags, EPoseSearchTrajectoryFlags::Position))
 		{
 			FFeatureVectorHelper::ComputeMeanDeviations(CenteredPoseMatrix, MeanDeviations, DataOffset, FFeatureVectorHelper::EncodeVectorCardinality);
 		}
-	}
-
-	if (bUseLinearVelocities)
-	{
-		for (int32 SubsampleIdx = 0; SubsampleIdx != SampleOffsets.Num(); ++SubsampleIdx)
+		if (EnumHasAnyFlags(Sample.Flags, EPoseSearchTrajectoryFlags::Velocity))
 		{
 			FFeatureVectorHelper::ComputeMeanDeviations(CenteredPoseMatrix, MeanDeviations, DataOffset, FFeatureVectorHelper::EncodeVectorCardinality);
 		}
-	}
-
-	if (bUseFacingDirections)
-	{
-		for (int32 SubsampleIdx = 0; SubsampleIdx != SampleOffsets.Num(); ++SubsampleIdx)
+		if (EnumHasAnyFlags(Sample.Flags, EPoseSearchTrajectoryFlags::FacingDirection))
 		{
 			FFeatureVectorHelper::ComputeMeanDeviations(CenteredPoseMatrix, MeanDeviations, DataOffset, FFeatureVectorHelper::EncodeVectorCardinality);
 		}
@@ -1119,15 +1268,15 @@ void UPoseSearchFeatureChannel_Trajectory::ComputeMeanDeviations(const Eigen::Ma
 	check(DataOffset == ChannelDataOffset + ChannelCardinality);
 }
 
-float UPoseSearchFeatureChannel_Trajectory::GetSampleTime(const UE::PoseSearch::IAssetIndexer& Indexer, int32 SubsampleIdx, float SampleTime, float RootDistance) const
+float UPoseSearchFeatureChannel_Trajectory::GetSampleTime(const UE::PoseSearch::IAssetIndexer& Indexer, float Offset, float SampleTime, float RootDistance) const
 {
 	switch (Domain)
 	{
 	case EPoseSearchFeatureDomain::Time:
-		return SampleTime + SampleOffsets[SubsampleIdx];
+		return SampleTime + Offset;
 
 	case EPoseSearchFeatureDomain::Distance:
-		return Indexer.GetSampleTimeFromDistance(RootDistance + SampleOffsets[SubsampleIdx]);
+		return Indexer.GetSampleTimeFromDistance(RootDistance + Offset);
 
 	default:
 		checkNoEntry();
@@ -1153,22 +1302,18 @@ void UPoseSearchFeatureChannel_Trajectory::IndexAssetPrivate(const UE::PoseSearc
 	FSampleInfo Origin = Indexer.GetSampleInfo(SampleTime);
 
 	int32 DataOffset = ChannelDataOffset;
-	if (bUsePositions)
+	for (const FPoseSearchTrajectorySample& Sample : Samples)
 	{
-		for (int32 SubsampleIdx = 0; SubsampleIdx != SampleOffsets.Num(); ++SubsampleIdx)
+		const float SubsampleTime = GetSampleTime(Indexer, Sample.Offset, SampleTime, Origin.RootDistance);
+
+		if (EnumHasAnyFlags(Sample.Flags, EPoseSearchTrajectoryFlags::Position))
 		{
-			const float SubsampleTime = GetSampleTime(Indexer, SubsampleIdx, SampleTime, Origin.RootDistance);
 			const FSampleInfo SamplePresent = Indexer.GetSampleInfoRelative(SubsampleTime, Origin);
 			FFeatureVectorHelper::EncodeVector(FeatureVector.EditValues(), DataOffset, Indexer.MirrorTransform(SamplePresent.RootTransform).GetTranslation());
 		}
-	}
-
-	if (bUseLinearVelocities)
-	{
-		for (int32 SubsampleIdx = 0; SubsampleIdx != SampleOffsets.Num(); ++SubsampleIdx)
+	
+		if (EnumHasAnyFlags(Sample.Flags, EPoseSearchTrajectoryFlags::Velocity))
 		{
-			const float SubsampleTime = GetSampleTime(Indexer, SubsampleIdx, SampleTime, Origin.RootDistance);
-
 			// For each pose subsample term, get the corresponding clip, accumulated root motion,
 			// and wrap the time parameter based on the clip's length.
 			const FSampleInfo SamplePast = Indexer.GetSampleInfoRelative(SubsampleTime - IndexingContext.SamplingContext->FiniteDelta, Origin);
@@ -1200,13 +1345,9 @@ void UPoseSearchFeatureChannel_Trajectory::IndexAssetPrivate(const UE::PoseSearc
 
 			FFeatureVectorHelper::EncodeVector(FeatureVector.EditValues(), DataOffset, LinearVelocity);
 		}
-	}
 
-	if (bUseFacingDirections)
-	{
-		for (int32 SubsampleIdx = 0; SubsampleIdx != SampleOffsets.Num(); ++SubsampleIdx)
+		if (EnumHasAnyFlags(Sample.Flags, EPoseSearchTrajectoryFlags::FacingDirection))
 		{
-			const float SubsampleTime = GetSampleTime(Indexer, SubsampleIdx, SampleTime, Origin.RootDistance);
 			const FSampleInfo SamplePresent = Indexer.GetSampleInfoRelative(SubsampleTime, Origin);
 			FFeatureVectorHelper::EncodeVector(FeatureVector.EditValues(), DataOffset, Indexer.MirrorTransform(SamplePresent.RootTransform).GetRotation().GetAxisY());
 		}
@@ -1215,11 +1356,8 @@ void UPoseSearchFeatureChannel_Trajectory::IndexAssetPrivate(const UE::PoseSearc
 
 void UPoseSearchFeatureChannel_Trajectory::GenerateDDCKey(FBlake3& InOutKeyHasher) const
 {
-	InOutKeyHasher.Update(&bUseLinearVelocities, sizeof(bUseLinearVelocities));
-	InOutKeyHasher.Update(&bUsePositions, sizeof(bUsePositions));
-	InOutKeyHasher.Update(&bUseFacingDirections, sizeof(bUseFacingDirections));
 	InOutKeyHasher.Update(&Domain, sizeof(Domain));
-	InOutKeyHasher.Update(MakeMemoryView(SampleOffsets));
+	InOutKeyHasher.Update(MakeMemoryView(Samples));
 	InOutKeyHasher.Update(&Weight, sizeof(Weight));
 }
 
@@ -1248,36 +1386,28 @@ bool UPoseSearchFeatureChannel_Trajectory::BuildQuery(UE::PoseSearch::FSearchCon
 			return false;
 	}
 
-	// @todo: optimize dynamic allocation
-	TArray<FTrajectorySample> Samples;
-	Samples.Reserve(SampleOffsets.Num());
-	for (int32 Idx = 0, NextIterStartIdx = 0, Num = SampleOffsets.Num(); Idx < Num; ++Idx)
-	{
-		Samples.Add(FTrajectorySampleRange::IterSampleTrajectory(SearchContext.Trajectory->Samples, SampleDomain, SampleOffsets[Idx], NextIterStartIdx));
-	}
-
+	int32 NextIterStartIdx = 0;
 	int32 DataOffset = ChannelDataOffset;
-	if (bUsePositions)
+	float PreviousOffset = -FLT_MAX;
+	for (const FPoseSearchTrajectorySample& Sample : Samples)
 	{
-		for (int32 Idx = 0, Num = SampleOffsets.Num(); Idx < Num; ++Idx)
+		// making sure Samples are sorted
+		check(Sample.Offset >= PreviousOffset);
+		const FTrajectorySample TrajectorySample = FTrajectorySampleRange::IterSampleTrajectory(SearchContext.Trajectory->Samples, SampleDomain, Sample.Offset, NextIterStartIdx);
+
+		if (EnumHasAnyFlags(Sample.Flags, EPoseSearchTrajectoryFlags::Position))
 		{
-			FFeatureVectorHelper::EncodeVector(InOutQuery.EditValues(), DataOffset, Samples[Idx].Transform.GetTranslation());
+			FFeatureVectorHelper::EncodeVector(InOutQuery.EditValues(), DataOffset, TrajectorySample.Transform.GetTranslation());
 		}
-	}
 
-	if (bUseLinearVelocities)
-	{
-		for (int32 Idx = 0, Num = SampleOffsets.Num(); Idx < Num; ++Idx)
+		if (EnumHasAnyFlags(Sample.Flags, EPoseSearchTrajectoryFlags::Velocity))
 		{
-			FFeatureVectorHelper::EncodeVector(InOutQuery.EditValues(), DataOffset, Samples[Idx].LinearVelocity);
+			FFeatureVectorHelper::EncodeVector(InOutQuery.EditValues(), DataOffset, TrajectorySample.LinearVelocity);
 		}
-	}
 
-	if (bUseFacingDirections)
-	{ 
-		for (int32 Idx = 0, Num = SampleOffsets.Num(); Idx < Num; ++Idx)
+		if (EnumHasAnyFlags(Sample.Flags, EPoseSearchTrajectoryFlags::FacingDirection))
 		{
-			FFeatureVectorHelper::EncodeVector(InOutQuery.EditValues(), DataOffset, Samples[Idx].Transform.GetRotation().GetAxisY());
+			FFeatureVectorHelper::EncodeVector(InOutQuery.EditValues(), DataOffset, TrajectorySample.Transform.GetRotation().GetAxisY());
 		}
 	}
 
@@ -1292,8 +1422,8 @@ void UPoseSearchFeatureChannel_Trajectory::DebugDraw(const UE::PoseSearch::FDebu
 	const uint8 DepthPriority = ESceneDepthPriorityGroup::SDPG_Foreground + 2;
 	const bool bPersistent = EnumHasAnyFlags(DrawParams.Flags, EDebugDrawFlags::Persistent);
 
-	const int32 NumSubsamples = SampleOffsets.Num();
-	if (NumSubsamples == 0)
+	const int32 NumSamples = Samples.Num();
+	if (NumSamples == 0)
 	{
 		return;
 	}
@@ -1309,43 +1439,40 @@ void UPoseSearchFeatureChannel_Trajectory::DebugDraw(const UE::PoseSearch::FDebu
 		return OriginalColor * (1.0f - DrawDebugGradientStrength * (SampleIdx / (float)Denominator));
 	};
 
-	TArray<FVector> TrajectoryPos;
-	TrajectoryPos.AddUninitialized(NumSubsamples);
 	int32 DataOffset = ChannelDataOffset;
-	for (int32 SubsampleIdx = 0; SubsampleIdx != NumSubsamples; ++SubsampleIdx)
+	int32 SampleIdx = 0;
+	for (const FPoseSearchTrajectorySample& Sample : Samples)
 	{
-		if (bUsePositions)
+		FVector TrajectoryPos;
+		if (EnumHasAnyFlags(Sample.Flags, EPoseSearchTrajectoryFlags::Position))
 		{
-			TrajectoryPos[SubsampleIdx] = FFeatureVectorHelper::DecodeVector(PoseVector, DataOffset);
+			TrajectoryPos = FFeatureVectorHelper::DecodeVector(PoseVector, DataOffset);
 
 			FLinearColor LinearColor = DrawParams.GetColor(this);
-			FLinearColor GradientColor = GetGradientColor(LinearColor, SubsampleIdx, NumSubsamples, DrawParams.Flags);
+			FLinearColor GradientColor = GetGradientColor(LinearColor, SampleIdx, NumSamples, DrawParams.Flags);
 			FColor Color = GradientColor.ToFColor(true);
 
-			TrajectoryPos[SubsampleIdx] = DrawParams.RootTransform.TransformPosition(TrajectoryPos[SubsampleIdx]);
+			TrajectoryPos = DrawParams.RootTransform.TransformPosition(TrajectoryPos);
 			if (EnumHasAnyFlags(DrawParams.Flags, EDebugDrawFlags::DrawFast | EDebugDrawFlags::DrawSearchIndex))
 			{
-				DrawDebugPoint(DrawParams.World, TrajectoryPos[SubsampleIdx], DrawParams.PointSize, Color, bPersistent, DrawParams.DefaultLifeTime, DepthPriority);
+				DrawDebugPoint(DrawParams.World, TrajectoryPos, DrawParams.PointSize, Color, bPersistent, LifeTime, DepthPriority);
 			}
 			else
 			{
-				DrawDebugSphere(DrawParams.World, TrajectoryPos[SubsampleIdx], DrawDebugSphereSize, DrawDebugSphereSegments, Color, bPersistent, LifeTime, DepthPriority);
+				DrawDebugSphere(DrawParams.World, TrajectoryPos, DrawDebugSphereSize, DrawDebugSphereSegments, Color, bPersistent, LifeTime, DepthPriority);
 			}
 		}
 		else
 		{
-			TrajectoryPos[SubsampleIdx] = DrawParams.RootTransform.GetTranslation();
+			TrajectoryPos = DrawParams.RootTransform.GetTranslation();
 		}
-	}
 
-	if (bUseLinearVelocities)
-	{
-		for (int32 SubsampleIdx = 0; SubsampleIdx != NumSubsamples; ++SubsampleIdx)
+		if (EnumHasAnyFlags(Sample.Flags, EPoseSearchTrajectoryFlags::Velocity))
 		{
 			FVector TrajectoryVel = FFeatureVectorHelper::DecodeVector(PoseVector, DataOffset);
 
 			const FLinearColor LinearColor = DrawParams.GetColor(this);
-			const FLinearColor GradientColor = GetGradientColor(LinearColor, SubsampleIdx, NumSubsamples, DrawParams.Flags);
+			const FLinearColor GradientColor = GetGradientColor(LinearColor, SampleIdx, NumSamples, DrawParams.Flags);
 			const FColor Color = GradientColor.ToFColor(true);
 
 			TrajectoryVel *= DrawDebugVelocityScale;
@@ -1354,18 +1481,16 @@ void UPoseSearchFeatureChannel_Trajectory::DebugDraw(const UE::PoseSearch::FDebu
 
 			if (EnumHasAnyFlags(DrawParams.Flags, EDebugDrawFlags::DrawSearchIndex))
 			{
-				DrawDebugLine(DrawParams.World, TrajectoryPos[SubsampleIdx], TrajectoryPos[SubsampleIdx] + TrajectoryVel, Color, bPersistent, DrawParams.DefaultLifeTime, DepthPriority);
+				DrawDebugLine(DrawParams.World, TrajectoryPos, TrajectoryPos + TrajectoryVel, Color, bPersistent, LifeTime, DepthPriority);
 			}
 			else
 			{
-				const float AdjustedThickness =
-					EnumHasAnyFlags(DrawParams.Flags, EDebugDrawFlags::DrawFast) ?
-					0.0f : DrawDebugLineThickness;
+				const float AdjustedThickness = EnumHasAnyFlags(DrawParams.Flags, EDebugDrawFlags::DrawFast) ? 0.0f : DrawDebugLineThickness;
 
 				DrawDebugLine(
 					DrawParams.World,
-					TrajectoryPos[SubsampleIdx] + TrajectoryVelDirection * DrawDebugSphereSize,
-					TrajectoryPos[SubsampleIdx] + TrajectoryVel,
+					TrajectoryPos + TrajectoryVelDirection * DrawDebugSphereSize,
+					TrajectoryPos + TrajectoryVel,
 					Color,
 					bPersistent,
 					LifeTime,
@@ -1374,34 +1499,29 @@ void UPoseSearchFeatureChannel_Trajectory::DebugDraw(const UE::PoseSearch::FDebu
 				);
 			}
 		}
-	}
 
-	if (bUseFacingDirections)
-	{
-		for (int32 SubsampleIdx = 0; SubsampleIdx != NumSubsamples; ++SubsampleIdx)
+		if (EnumHasAnyFlags(Sample.Flags, EPoseSearchTrajectoryFlags::FacingDirection))
 		{
 			FVector TrajectoryForward = FFeatureVectorHelper::DecodeVector(PoseVector, DataOffset);
 
 			const FLinearColor LinearColor = DrawParams.GetColor(this);
-			const FLinearColor GradientColor = GetGradientColor(LinearColor, SubsampleIdx, NumSubsamples, DrawParams.Flags);
+			const FLinearColor GradientColor = GetGradientColor(LinearColor, SampleIdx, NumSamples, DrawParams.Flags);
 			const FColor Color = GradientColor.ToFColor(true);
 
 			TrajectoryForward = DrawParams.RootTransform.TransformVector(TrajectoryForward);
 
 			if (EnumHasAnyFlags(DrawParams.Flags, EDebugDrawFlags::DrawSearchIndex))
 			{
-				DrawDebugLine(DrawParams.World, TrajectoryPos[SubsampleIdx], TrajectoryPos[SubsampleIdx] + TrajectoryForward, Color, bPersistent, DrawParams.DefaultLifeTime, DepthPriority);
+				DrawDebugLine(DrawParams.World, TrajectoryPos, TrajectoryPos + TrajectoryForward, Color, bPersistent, LifeTime, DepthPriority);
 			}
 			else
 			{
-				const float AdjustedThickness =
-					EnumHasAnyFlags(DrawParams.Flags, EDebugDrawFlags::DrawFast) ?
-					0.0f : DrawDebugLineThickness;
+				const float AdjustedThickness = EnumHasAnyFlags(DrawParams.Flags, EDebugDrawFlags::DrawFast) ? 0.0f : DrawDebugLineThickness;
 
 				DrawDebugLine(
 					DrawParams.World,
-					TrajectoryPos[SubsampleIdx] + TrajectoryForward * DrawDebugSphereSize,
-					TrajectoryPos[SubsampleIdx] + TrajectoryForward * DrawDebugSphereSize * 2.0f,
+					TrajectoryPos + TrajectoryForward * DrawDebugSphereSize,
+					TrajectoryPos + TrajectoryForward * DrawDebugSphereSize * 2.0f,
 					Color,
 					bPersistent,
 					LifeTime,
@@ -1410,28 +1530,25 @@ void UPoseSearchFeatureChannel_Trajectory::DebugDraw(const UE::PoseSearch::FDebu
 				);
 			}
 		}
-	}
 
-	if (EnumHasAnyFlags(DrawParams.Flags, EDebugDrawFlags::DrawSampleLabels))
-	{
-		for (int32 SubsampleIdx = 0; SubsampleIdx != NumSubsamples; ++SubsampleIdx)
+		if (EnumHasAnyFlags(DrawParams.Flags, EDebugDrawFlags::DrawSampleLabels))
 		{
 			const FLinearColor LinearColor = DrawParams.GetColor(this);
-			const FLinearColor GradientColor = GetGradientColor(LinearColor, SubsampleIdx, NumSubsamples, DrawParams.Flags);
+			const FLinearColor GradientColor = GetGradientColor(LinearColor, SampleIdx, NumSamples, DrawParams.Flags);
 			const FColor Color = GradientColor.ToFColor(true);
 
 			FString SampleLabel;
 			if (DrawParams.LabelPrefix.IsEmpty())
 			{
-				SampleLabel = FString::Format(TEXT("{0}"), { SubsampleIdx });
+				SampleLabel = FString::Format(TEXT("{0}"), { SampleIdx });
 			}
 			else
 			{
-				SampleLabel = FString::Format(TEXT("{1}[{0}]"), { SubsampleIdx, DrawParams.LabelPrefix.GetData() });
+				SampleLabel = FString::Format(TEXT("{1}[{0}]"), { SampleIdx, DrawParams.LabelPrefix.GetData() });
 			}
 			DrawDebugString(
 				DrawParams.World,
-				TrajectoryPos[SubsampleIdx] + DrawDebugSampleLabelOffset,
+				TrajectoryPos + DrawDebugSampleLabelOffset,
 				SampleLabel,
 				nullptr,
 				Color,
@@ -1439,6 +1556,8 @@ void UPoseSearchFeatureChannel_Trajectory::DebugDraw(const UE::PoseSearch::FDebu
 				false,
 				DrawDebugSampleLabelFontScale);
 		}
+
+		++SampleIdx;
 	}
 
 	check(DataOffset == ChannelDataOffset + ChannelCardinality);
diff --git a/Engine/Plugins/Experimental/Animation/PoseSearch/Source/Runtime/Private/PoseSearchLibrary.cpp b/Engine/Plugins/Experimental/Animation/PoseSearch/Source/Runtime/Private/PoseSearchLibrary.cpp
--- a/Engine/Plugins/Experimental/Animation/PoseSearch/Source/Runtime/Private/PoseSearchLibrary.cpp
+++ b/Engine/Plugins/Experimental/Animation/PoseSearch/Source/Runtime/Private/PoseSearchLibrary.cpp
@@ -219,7 +219,8 @@ void UpdateMotionMatchingState(
 	const FGameplayTagContainer* ActiveTagsContainer,
 	const FTrajectorySampleRange& Trajectory,
 	const FMotionMatchingSettings& Settings,
-	FMotionMatchingState& InOutMotionMatchingState
+	FMotionMatchingState& InOutMotionMatchingState,
+	bool bForceInterrupt
 )
 {
 	QUICK_SCOPE_CYCLE_COUNTER(STAT_PoseSearch_Update);
@@ -245,7 +246,7 @@ void UpdateMotionMatchingState(
 	// Check if we can advance. Includes the case where we can advance but only by switching to a follow up asset.
 	bool bAdvanceToFollowUpAsset = false;
 	FSearchResult FollowUpAsset;
-	bool bCanAdvance = InOutMotionMatchingState.CanAdvance(Context.GetDeltaTime(), bAdvanceToFollowUpAsset, FollowUpAsset);
+	bool bCanAdvance = InOutMotionMatchingState.CanAdvance(DeltaTime, bAdvanceToFollowUpAsset, FollowUpAsset);
 
 	// If we can't advance or enough time has elapsed since the last pose jump then search
 	if (!bCanAdvance || (InOutMotionMatchingState.ElapsedPoseJumpTime >= Settings.SearchThrottleTime))
@@ -257,6 +258,7 @@ void UpdateMotionMatchingState(
 		SearchContext.Trajectory = &Trajectory;
 		SearchContext.OwningComponent = Context.AnimInstanceProxy->GetSkelMeshComponent();
 		SearchContext.BoneContainer = &Context.AnimInstanceProxy->GetRequiredBones();
+		SearchContext.bForceInterrupt = bForceInterrupt;
 
 #if WITH_EDITORONLY_DATA
 		SearchContext.DebugDrawParams.SearchCostHistoryBruteForce = &InOutMotionMatchingState.SearchCostHistoryBruteForce;
diff --git a/Engine/Plugins/Experimental/Animation/PoseSearch/Source/Runtime/Public/PoseSearch/AnimNode_MotionMatching.h b/Engine/Plugins/Experimental/Animation/PoseSearch/Source/Runtime/Public/PoseSearch/AnimNode_MotionMatching.h
--- a/Engine/Plugins/Experimental/Animation/PoseSearch/Source/Runtime/Public/PoseSearch/AnimNode_MotionMatching.h
+++ b/Engine/Plugins/Experimental/Animation/PoseSearch/Source/Runtime/Public/PoseSearch/AnimNode_MotionMatching.h
@@ -26,7 +26,6 @@ public:
 	UPROPERTY(EditAnywhere, BlueprintReadWrite, Category=Settings, meta=(PinShownByDefault))
 	TObjectPtr<const UPoseSearchSearchableAsset> Searchable = nullptr;
 
-
 	UPROPERTY(EditAnywhere, BlueprintReadWrite, Category = Settings, meta = (PinHiddenByDefault))
 	bool bUseDatabaseTagQuery = false;
 
@@ -54,6 +53,10 @@ public:
 	UPROPERTY(EditAnywhere, BlueprintReadWrite, Category = Settings, meta = (PinHiddenByDefault))
 	bool bResetOnBecomingRelevant = true;
 
+	// if true the continuity pose will be invalidated
+	UPROPERTY(EditAnywhere, BlueprintReadWrite, Category = Settings, meta = (PinHiddenByDefault))
+	bool bForceInterrupt = false;
+
 #if WITH_EDITORONLY_DATA
 	UPROPERTY(EditAnywhere, Category=Debug, meta = (PinShownByDefault))
 	bool bDebugDraw = false;
diff --git a/Engine/Plugins/Experimental/Animation/PoseSearch/Source/Runtime/Public/PoseSearch/PoseSearch.h b/Engine/Plugins/Experimental/Animation/PoseSearch/Source/Runtime/Public/PoseSearch/PoseSearch.h
--- a/Engine/Plugins/Experimental/Animation/PoseSearch/Source/Runtime/Public/PoseSearch/PoseSearch.h
+++ b/Engine/Plugins/Experimental/Animation/PoseSearch/Source/Runtime/Public/PoseSearch/PoseSearch.h
@@ -172,6 +172,19 @@ struct FPoseSearchBlockTransitionParameters
 	float SequenceEndInterval = 0.0f;
 };
 
+// @todo: move it into PoseSearchFeatureChannels after removing SampledBones_DEPRECATED 
+UENUM(meta = (Bitflags, UseEnumValuesAsMaskValuesInEditor = "true"))
+enum class EPoseSearchBoneFlags : uint32
+{
+	Velocity = 1 << 0,
+	Position = 1 << 1,
+	Rotation = 1 << 2,
+	Phase = 1 << 3,
+};
+ENUM_CLASS_FLAGS(EPoseSearchBoneFlags);
+constexpr bool EnumHasAnyFlags(int32 Flags, EPoseSearchBoneFlags Contains) { return (Flags & int32(Contains)) != 0; }
+inline int32& operator|=(int32& Lhs, EPoseSearchBoneFlags Rhs) { return Lhs |= int32(Rhs); }
+
 // @todo: move it into PoseSearchFeatureChannels after removing SampledBones_DEPRECATED 
 USTRUCT()
 struct POSESEARCH_API FPoseSearchBone
@@ -181,17 +194,22 @@ struct POSESEARCH_API FPoseSearchBone
 	UPROPERTY(EditAnywhere, Category = Config)
 	FBoneReference Reference;
 
-	UPROPERTY(EditAnywhere, Category = Config)
-	bool bUseVelocity = false;
+#if WITH_EDITORONLY_DATA
+	UPROPERTY()
+	bool bUseVelocity_DEPRECATED = false;
 
-	UPROPERTY(EditAnywhere, Category = Config)
-	bool bUsePosition = false;
+	UPROPERTY()
+	bool bUsePosition_DEPRECATED = false;
 
-	UPROPERTY(EditAnywhere, Category = Config)
-	bool bUseRotation = false;
+	UPROPERTY()
+	bool bUseRotation_DEPRECATED = false;
 
-	UPROPERTY(EditAnywhere, Category = Config)
-	bool bUsePhase = false;
+	UPROPERTY()
+	bool bUsePhase_DEPRECATED = false;
+#endif
+
+	UPROPERTY(EditAnywhere, meta = (Bitmask, BitmaskEnum = "/Script/PoseSearch.EPoseSearchBoneFlags"), Category = Config)
+	int32 Flags = int32(EPoseSearchBoneFlags::Position);
 
 	// @todo: temporary location for the channel bone weight to help the weights refactoring
 	UPROPERTY(EditAnywhere, Category = Config)
@@ -1283,6 +1301,7 @@ struct POSESEARCH_API FSearchContext
 	const FBoneContainer* BoneContainer = nullptr;
 	const FGameplayTagContainer* ActiveTagsContainer = nullptr;
 	float PoseJumpThresholdTime = 0.f;
+	bool bForceInterrupt = false;
 
 	FTransform TryGetTransformAndCacheResults(float SampleTime, const UPoseSearchSchema* Schema, int8 SchemaBoneIdx, bool& Error);
 	void ClearCachedEntries();
diff --git a/Engine/Plugins/Experimental/Animation/PoseSearch/Source/Runtime/Public/PoseSearch/PoseSearchFeatureChannels.h b/Engine/Plugins/Experimental/Animation/PoseSearch/Source/Runtime/Public/PoseSearch/PoseSearchFeatureChannels.h
--- a/Engine/Plugins/Experimental/Animation/PoseSearch/Source/Runtime/Public/PoseSearch/PoseSearchFeatureChannels.h
+++ b/Engine/Plugins/Experimental/Animation/PoseSearch/Source/Runtime/Public/PoseSearch/PoseSearchFeatureChannels.h
@@ -29,13 +29,13 @@ public:
 	int8 SchemaBoneIdx;
 
 	// reuse continuity pose database feature values if available
-	UPROPERTY(EditAnywhere, Category = "Advanced")
-	bool UseFeaturesFromContinuityPose = true;
+	UPROPERTY(EditAnywhere, Category = "Settings")
+	bool bUseFeaturesFromContinuityPose = true;
 
 	// if UseSampleTimeOffsetRootBone is true, this UPoseSearchFeatureChannel_Position will calculate the position of Bone from the pose SampleTimeOffset seconds away from the current time pose root bone
 	// if false the calculated position of Bone will be in component space from the pose SampleTimeOffset seconds away
-	UPROPERTY(EditAnywhere, Category = "Advanced")
-	bool UseSampleTimeOffsetRootBone = true;
+	UPROPERTY(EditAnywhere, Category = "Settings")
+	bool bUseSampleTimeOffsetRootBone = true;
 
 	// UPoseSearchFeatureChannel interface
 	virtual void InitializeSchema(UE::PoseSearch::FSchemaInitializer& Initializer) override;
@@ -46,6 +46,61 @@ public:
 	virtual void DebugDraw(const UE::PoseSearch::FDebugDrawParams& DrawParams, TArrayView<const float> PoseVector) const override;
 };
 
+//////////////////////////////////////////////////////////////////////////
+// UPoseSearchFeatureChannel_Heading
+
+UENUM(BlueprintType)
+enum class HeadingAxis : uint8
+{
+	X,
+	Y,
+	Z,
+
+	Num UMETA(Hidden),
+	Invalid = Num UMETA(Hidden)
+};
+
+UCLASS(BlueprintType, EditInlineNew)
+class POSESEARCH_API UPoseSearchFeatureChannel_Heading : public UPoseSearchFeatureChannel
+{
+	GENERATED_BODY()
+
+public:
+	UPROPERTY(EditAnywhere, Category = "Settings")
+	FBoneReference Bone;
+
+	UPROPERTY(EditAnywhere, Category = "Settings")
+	float Weight = 1.f;
+
+	UPROPERTY(EditAnywhere, Category = "Settings")
+	float SampleTimeOffset = 0.f;
+
+	UPROPERTY(EditAnywhere, Category = "Settings")
+	HeadingAxis HeadingAxis = HeadingAxis::X;	
+
+	UPROPERTY()
+	int8 SchemaBoneIdx;
+
+	// reuse continuity pose database feature values if available
+	UPROPERTY(EditAnywhere, Category = "Settings")
+	bool bUseFeaturesFromContinuityPose = true;
+
+	// if UseSampleTimeOffsetRootBone is true, this UPoseSearchFeatureChannel_Position will calculate the position of Bone from the pose SampleTimeOffset seconds away from the current time pose root bone
+	// if false the calculated position of Bone will be in component space from the pose SampleTimeOffset seconds away
+	UPROPERTY(EditAnywhere, Category = "Settings")
+	bool bUseSampleTimeOffsetRootBone = true;
+
+	// UPoseSearchFeatureChannel interface
+	virtual void InitializeSchema(UE::PoseSearch::FSchemaInitializer& Initializer) override;
+	virtual void FillWeights(TArray<float>& Weights) const override;
+	virtual void IndexAsset(UE::PoseSearch::IAssetIndexer& Indexer, UE::PoseSearch::FAssetIndexingOutput& IndexingOutput) const override;
+	virtual void GenerateDDCKey(FBlake3& InOutKeyHasher) const override;
+	virtual bool BuildQuery(UE::PoseSearch::FSearchContext& SearchContext, FPoseSearchFeatureVectorBuilder& InOutQuery) const override;
+	virtual void DebugDraw(const UE::PoseSearch::FDebugDrawParams& DrawParams, TArrayView<const float> PoseVector) const override;
+
+	FVector GetAxis(const FQuat& Rotation) const;
+};
+
 //////////////////////////////////////////////////////////////////////////
 // UPoseSearchFeatureChannel_Pose
 UCLASS(BlueprintType, EditInlineNew)
@@ -54,6 +109,9 @@ class POSESEARCH_API UPoseSearchFeatureChannel_Pose : public UPoseSearchFeatureC
 	GENERATED_BODY()
 
 public:
+	UPROPERTY(EditAnywhere, Category = "Settings")
+	float Weight = 1.f;
+
 	UPROPERTY(EditAnywhere, Category = "Settings")
 	TArray<FPoseSearchBone> SampledBones;
 
@@ -63,6 +121,10 @@ public:
 	UPROPERTY()
 	TArray<int8> SchemaBoneIdx;
 
+	// reuse continuity pose database feature values if available
+	UPROPERTY(EditAnywhere, Category = "Settings")
+	bool bUseFeaturesFromContinuityPose = true;
+
 	// UObject interface
 	virtual void PreSave(FObjectPreSaveContext ObjectSaveContext) override;
 
@@ -82,30 +144,62 @@ protected:
 
 //////////////////////////////////////////////////////////////////////////
 // UPoseSearchFeatureChannel_Trajectory
+
+UENUM(meta = (Bitflags, UseEnumValuesAsMaskValuesInEditor = "true"))
+enum class EPoseSearchTrajectoryFlags : uint32
+{
+	Velocity = 1 << 0,
+	Position = 1 << 1,
+	FacingDirection = 1 << 3,
+};
+ENUM_CLASS_FLAGS(EPoseSearchTrajectoryFlags);
+constexpr bool EnumHasAnyFlags(int32 Flags, EPoseSearchTrajectoryFlags Contains) { return (Flags & int32(Contains)) != 0; }
+inline int32& operator|=(int32& Lhs, EPoseSearchTrajectoryFlags Rhs) { return Lhs |= int32(Rhs); }
+
+USTRUCT()
+struct POSESEARCH_API FPoseSearchTrajectorySample
+{
+	GENERATED_BODY()
+
+	UPROPERTY(EditAnywhere, Category = Config)
+	float Offset = 0.f; // offset in time or distance depending on UPoseSearchFeatureChannel_Trajectory.Domain
+
+	UPROPERTY(EditAnywhere, meta = (Bitmask, BitmaskEnum = "/Script/PoseSearch.EPoseSearchTrajectoryFlags"), Category = Config)
+	int32 Flags = int32(EPoseSearchTrajectoryFlags::Position);
+
+	UPROPERTY(EditAnywhere, Category = Config)
+	float Weight = 1.f;
+};
+
 UCLASS(BlueprintType, EditInlineNew)
 class POSESEARCH_API UPoseSearchFeatureChannel_Trajectory : public UPoseSearchFeatureChannel
 {
 	GENERATED_BODY()
 
 public:
-	UPROPERTY(EditAnywhere, Category = "Settings")
-	bool bUseLinearVelocities = true;
+#if WITH_EDITORONLY_DATA
+	UPROPERTY()
+	bool bUseLinearVelocities_DEPRECATED = true;
 
-	UPROPERTY(EditAnywhere, Category = "Settings")
-	bool bUsePositions = true;
+	UPROPERTY()
+	bool bUsePositions_DEPRECATED = true;
 
-	UPROPERTY(EditAnywhere, Category = "Settings")
-	bool bUseFacingDirections = false;
+	UPROPERTY()
+	bool bUseFacingDirections_DEPRECATED = false;
+#endif
 
 	UPROPERTY(EditAnywhere, Category="Settings")
 	EPoseSearchFeatureDomain Domain = EPoseSearchFeatureDomain::Time;
 
-	// @todo: temporary location for the channel bone weight to help the weights refactoring (later on it makes sense to have a weight per SampleOffsets, but right now for semplicity the traj will have just one weight)
 	UPROPERTY(EditAnywhere, Category = "Settings")
 	float Weight = 1.f;
 
+#if WITH_EDITORONLY_DATA
+	UPROPERTY()
+	TArray<float> SampleOffsets_DEPRECATED;
+#endif
 	UPROPERTY(EditAnywhere, Category = "Settings")
-	TArray<float> SampleOffsets;
+	TArray<FPoseSearchTrajectorySample> Samples;
 
 	// UObject interface
 	virtual void PreSave(FObjectPreSaveContext ObjectSaveContext) override;
@@ -121,5 +215,5 @@ public:
 
 protected:
 	void IndexAssetPrivate(const UE::PoseSearch::IAssetIndexer& Indexer, int32 SampleIdx, FPoseSearchFeatureVectorBuilder& FeatureVector) const;
-	float GetSampleTime(const UE::PoseSearch::IAssetIndexer& Indexer, int32 SubsampleIdx, float SampleTime, float RootDistance) const;
+	float GetSampleTime(const UE::PoseSearch::IAssetIndexer& Indexer, float Offset, float SampleTime, float RootDistance) const;
 };
\ No newline at end of file
diff --git a/Engine/Plugins/Experimental/Animation/PoseSearch/Source/Runtime/Public/PoseSearch/PoseSearchLibrary.h b/Engine/Plugins/Experimental/Animation/PoseSearch/Source/Runtime/Public/PoseSearch/PoseSearchLibrary.h
--- a/Engine/Plugins/Experimental/Animation/PoseSearch/Source/Runtime/Public/PoseSearch/PoseSearchLibrary.h
+++ b/Engine/Plugins/Experimental/Animation/PoseSearch/Source/Runtime/Public/PoseSearch/PoseSearchLibrary.h
@@ -114,7 +114,8 @@ POSESEARCH_API void UpdateMotionMatchingState(
 	const FGameplayTagContainer* ActiveTagsContainer,
 	const FTrajectorySampleRange& Trajectory,
 	const FMotionMatchingSettings& Settings,
-	FMotionMatchingState& InOutMotionMatchingState
+	FMotionMatchingState& InOutMotionMatchingState,
+	bool bForceInterrupt
 );
 
 UCLASS(DisplayName = "Bias")
