motion matching - defaulting FDatabaseViewModel::PoseFeaturesDrawMode to EFeaturesDrawMode::All

diff --git a/Engine/Plugins/Experimental/Animation/PoseSearch/Source/Editor/Public/PoseSearchDatabaseViewModel.h b/Engine/Plugins/Experimental/Animation/PoseSearch/Source/Editor/Public/PoseSearchDatabaseViewModel.h
--- a/Engine/Plugins/Experimental/Animation/PoseSearch/Source/Editor/Public/PoseSearchDatabaseViewModel.h
+++ b/Engine/Plugins/Experimental/Animation/PoseSearch/Source/Editor/Public/PoseSearchDatabaseViewModel.h
@@ -129,7 +129,7 @@ namespace UE::PoseSearch
 		float MaxPreviewPlayLength = 0.0f;
 
 		/** What features to show in the viewport */
-		EFeaturesDrawMode PoseFeaturesDrawMode = EFeaturesDrawMode::None;
+		EFeaturesDrawMode PoseFeaturesDrawMode = EFeaturesDrawMode::All;
 
 		/** What animations to show in the viewport */
 		EAnimationPreviewMode AnimationPreviewMode = EAnimationPreviewMode::OriginalOnly;
