PoseSearch: Fix exception when sorting debugger rows by cost columns

diff --git a/Engine/Plugins/Experimental/Animation/PoseSearch/Source/Editor/Private/PoseSearchDebugger.cpp b/Engine/Plugins/Experimental/Animation/PoseSearch/Source/Editor/Private/PoseSearchDebugger.cpp
--- a/Engine/Plugins/Experimental/Animation/PoseSearch/Source/Editor/Private/PoseSearchDebugger.cpp
+++ b/Engine/Plugins/Experimental/Animation/PoseSearch/Source/Editor/Private/PoseSearchDebugger.cpp
@@ -281,7 +281,7 @@ namespace DebuggerDatabaseColumns
 		virtual FText GetLabel() const = 0;
 
 		using FRowDataRef = TSharedRef<FDebuggerDatabaseRowData>;
-		using FSortPredicate = TFunctionRef<bool(const FRowDataRef&, const FRowDataRef&)>;
+		using FSortPredicate = TFunction<bool(const FRowDataRef&, const FRowDataRef&)>;
 
 		virtual FSortPredicate GetSortPredicate() const = 0;
 
