motion matching
- best search pose candidates collection in priority queue (not hooked up to the PoseSearchDebugger yet)
- FTraceMotionMatchingStateMessage serialization improvements using FArchive

diff --git a/Engine/Plugins/Experimental/Animation/PoseSearch/Source/Editor/Private/PoseSearchDebugger.cpp b/Engine/Plugins/Experimental/Animation/PoseSearch/Source/Editor/Private/PoseSearchDebugger.cpp
--- a/Engine/Plugins/Experimental/Animation/PoseSearch/Source/Editor/Private/PoseSearchDebugger.cpp
+++ b/Engine/Plugins/Experimental/Animation/PoseSearch/Source/Editor/Private/PoseSearchDebugger.cpp
@@ -2142,32 +2142,17 @@ const FTraceMotionMatchingStateMessage* FDebuggerViewModel::GetMotionMatchingSta
 
 const UPoseSearchDatabase* FDebuggerViewModel::GetPoseSearchDatabase() const
 {
-	if (!ActiveMotionMatchingState)
+	if (ActiveMotionMatchingState)
 	{
-		return nullptr;
-	}
-
-	const uint64 DatabaseId = ActiveMotionMatchingState->DatabaseId;
-	if (DatabaseId == 0)
-	{
-	    return nullptr;
-	}
-	
-	UObject* DatabaseObject = FObjectTrace::GetObjectFromId(DatabaseId);
-	// @TODO: Load the object if unloaded
-	if (DatabaseObject == nullptr)
-	{
-		return nullptr;
-	}
-	check(DatabaseObject->IsA<UPoseSearchDatabase>());
-
-	UPoseSearchDatabase* Database = Cast<UPoseSearchDatabase>(DatabaseObject);
-	const UPoseSearchSchema* Schema = Database->Schema;
-	if (Schema == nullptr || !Schema->IsValid())
-	{
-		return nullptr;
+		if (const UPoseSearchDatabase* Database = ActiveMotionMatchingState->GetPoseSearchDatabase())
+		{
+			if (Database->Schema && Database->Schema->IsValid())
+			{
+				return Database;
+			}
+		}
 	}
-	return Database;
+	return nullptr;
 }
 
 const FPoseSearchDatabaseSequence* FDebuggerViewModel::GetAnimSequence(int32 SequenceIdx) const
@@ -2175,8 +2160,7 @@ const FPoseSearchDatabaseSequence* FDebuggerViewModel::GetAnimSequence(int32 Seq
 	const UPoseSearchDatabase* Database = GetPoseSearchDatabase();
 	if (Database && Database->Sequences.IsValidIndex(SequenceIdx))
 	{
-		const FPoseSearchDatabaseSequence* DatabaseSequence = &Database->Sequences[SequenceIdx];
-		if (DatabaseSequence)
+		if (const FPoseSearchDatabaseSequence* DatabaseSequence = &Database->Sequences[SequenceIdx])
 		{
 			return DatabaseSequence;
 		}
@@ -2189,8 +2173,7 @@ const FPoseSearchDatabaseBlendSpace* FDebuggerViewModel::GetBlendSpace(int32 Ble
 	const UPoseSearchDatabase* Database = GetPoseSearchDatabase();
 	if (Database && Database->BlendSpaces.IsValidIndex(BlendSpaceIdx))
 	{
-		const FPoseSearchDatabaseBlendSpace* DatabaseBlendSpace = &Database->BlendSpaces[BlendSpaceIdx];
-		if (DatabaseBlendSpace)
+		if (const FPoseSearchDatabaseBlendSpace* DatabaseBlendSpace = &Database->BlendSpaces[BlendSpaceIdx])
 		{
 			return DatabaseBlendSpace;
 		}
diff --git a/Engine/Plugins/Experimental/Animation/PoseSearch/Source/Editor/Private/Trace/PoseSearchTraceAnalyzer.cpp b/Engine/Plugins/Experimental/Animation/PoseSearch/Source/Editor/Private/Trace/PoseSearchTraceAnalyzer.cpp
--- a/Engine/Plugins/Experimental/Animation/PoseSearch/Source/Editor/Private/Trace/PoseSearchTraceAnalyzer.cpp
+++ b/Engine/Plugins/Experimental/Animation/PoseSearch/Source/Editor/Private/Trace/PoseSearchTraceAnalyzer.cpp
@@ -9,7 +9,8 @@
 
 namespace UE { namespace PoseSearch {
 
-FTraceAnalyzer::FTraceAnalyzer(TraceServices::IAnalysisSession& InSession, FTraceProvider& InTraceProvider) : Session(InSession), TraceProvider(InTraceProvider)
+FTraceAnalyzer::FTraceAnalyzer(TraceServices::IAnalysisSession& InSession, FTraceProvider& InTraceProvider)
+: Session(InSession), TraceProvider(InTraceProvider)
 {
 }
 
@@ -31,51 +32,18 @@ bool FTraceAnalyzer::OnEvent(uint16 RouteId, EStyle Style, const FOnEventContext
 	LLM_SCOPE_BYNAME(TEXT("Insights/PoseSearch::FTraceAnalyzer"));
 
 	TraceServices::FAnalysisSessionEditScope Scope(Session);
-	const FEventData& EventData = Context.EventData;
 
-	// Gather event data values
-	const double Time = Context.EventTime.AsSeconds(EventData.GetValue<uint64>("Cycle"));
-	const uint16 FrameCounter = EventData.GetValue<uint16>("FrameCounter");
-	const uint64 AnimInstanceId = EventData.GetValue<uint64>("AnimInstanceId");
-	const uint64 SkeletalMeshComponentId = EventData.GetValue<uint64>("SkeletalMeshComponentId");
-	const int32 NodeId = EventData.GetValue<int32>("NodeId");
-
-	switch (RouteId)
+	if (RouteId == RouteId_MotionMatchingState)
 	{
-		case RouteId_MotionMatchingState:
-			{
-				FTraceMotionMatchingStateMessage Message;
-				Message.DatabaseId = EventData.GetValue<uint64>("DatabaseId");
-				Message.Flags = static_cast<FTraceMotionMatchingState::EFlags>(EventData.GetValue<uint32>("Flags"));
-				Message.DbPoseIdx = EventData.GetValue<int32>("DbPoseIdx");
-				Message.ContinuingPoseIdx = EventData.GetValue<int32>("ContinuingPoseIdx");
-				Message.ElapsedPoseJumpTime = EventData.GetValue<float>("ElapsedPoseJumpTime");
-				Message.QueryVector = EventData.GetArrayView<float>("QueryVector");
-                Message.QueryVectorNormalized = EventData.GetArrayView<float>("QueryVectorNormalized");
-
-				Message.AssetPlayerTime = EventData.GetValue<float>("AssetPlayerTime");
-				Message.DeltaTime = EventData.GetValue<float>("DeltaTime");
-				Message.SimLinearVelocity = EventData.GetValue<float>("SimLinearVelocity");
-				Message.SimAngularVelocity = EventData.GetValue<float>("SimAngularVelocity");
-				Message.AnimLinearVelocity = EventData.GetValue<float>("AnimLinearVelocity");
-				Message.AnimAngularVelocity = EventData.GetValue<float>("AnimAngularVelocity");
-				Message.DatabaseSequenceFilter = EventData.GetArrayView<bool>("DatabaseSequenceFilter");
-				Message.DatabaseBlendSpaceFilter = EventData.GetArrayView<bool>("DatabaseBlendSpaceFilter");
-
-				// Common data
-				Message.NodeId = NodeId;
-				Message.AnimInstanceId = AnimInstanceId;
-				Message.SkeletalMeshComponentId = SkeletalMeshComponentId;
-				Message.FrameCounter = FrameCounter;
-
-				TraceProvider.AppendMotionMatchingState(Message, Time);
-				break;
-			}
-		default:
-			{
-				// Should not happen
-				checkNoEntry();
-			}
+		FTraceMotionMatchingStateMessage Message;
+		FMemoryReaderView Archive(Context.EventData.GetArrayView<uint8>("Data"));
+		Archive << Message;
+		TraceProvider.AppendMotionMatchingState(Message, Context.EventTime.AsSeconds(Message.Cycle));
+	}
+	else
+	{
+		// Should not happen
+		checkNoEntry();
 	}
 
 	return true;
diff --git a/Engine/Plugins/Experimental/Animation/PoseSearch/Source/Editor/Private/Trace/PoseSearchTraceProvider.cpp b/Engine/Plugins/Experimental/Animation/PoseSearch/Source/Editor/Private/Trace/PoseSearchTraceProvider.cpp
--- a/Engine/Plugins/Experimental/Animation/PoseSearch/Source/Editor/Private/Trace/PoseSearchTraceProvider.cpp
+++ b/Engine/Plugins/Experimental/Animation/PoseSearch/Source/Editor/Private/Trace/PoseSearchTraceProvider.cpp
@@ -38,5 +38,11 @@ void FTraceProvider::AppendMotionMatchingState(const FTraceMotionMatchingStateMe
 	Session.UpdateDurationSeconds(InTime);
 }
 
+FArchive& operator<<(FArchive& Ar, FTraceMotionMatchingStateMessage& State)
+{
+	Ar << static_cast<FTraceMessage&>(State);
+	Ar << static_cast<FTraceMotionMatchingState&>(State);
+	return Ar;
+}
 
 }}
diff --git a/Engine/Plugins/Experimental/Animation/PoseSearch/Source/Editor/Private/Trace/PoseSearchTraceProvider.h b/Engine/Plugins/Experimental/Animation/PoseSearch/Source/Editor/Private/Trace/PoseSearchTraceProvider.h
--- a/Engine/Plugins/Experimental/Animation/PoseSearch/Source/Editor/Private/Trace/PoseSearchTraceProvider.h
+++ b/Engine/Plugins/Experimental/Animation/PoseSearch/Source/Editor/Private/Trace/PoseSearchTraceProvider.h
@@ -10,36 +10,12 @@
 
 namespace UE { namespace PoseSearch {
 
-// Message types for appending / reading to the timeline
-/** Base message type for common data */
-struct FTraceMessage
-{
-	uint64 AnimInstanceId;
-	uint64 SkeletalMeshComponentId;
-	int32 NodeId;
-	uint16 FrameCounter;
-};
 
 /** Motion matching state message container */
-struct FTraceMotionMatchingStateMessage : FTraceMessage
+struct FTraceMotionMatchingStateMessage : FTraceMotionMatchingState, FTraceMessage
 {
-	FTraceMotionMatchingState::EFlags Flags;
-	float ElapsedPoseJumpTime = -1.0f;
-	TArray<float> QueryVector;
-	TArray<float> QueryVectorNormalized;
-	int32 DbPoseIdx = INDEX_NONE;
-	uint64 DatabaseId = INDEX_NONE;
-	int32 ContinuingPoseIdx = INDEX_NONE;
-	float AssetPlayerTime = 0.0f;
-	float DeltaTime = 0.0f;
-	float SimLinearVelocity = 0.0f;
-	float SimAngularVelocity = 0.0f;
-	float AnimLinearVelocity = 0.0f;
-	float AnimAngularVelocity = 0.0f;
-	TArray<bool> DatabaseSequenceFilter;
-	TArray<bool> DatabaseBlendSpaceFilter;
 };
-
+FArchive& operator<<(FArchive& Ar, FTraceMotionMatchingStateMessage& State);
 
 /**
  * Provider to the widgets for pose search functionality, largely mimicking FAnimationProvider
diff --git a/Engine/Plugins/Experimental/Animation/PoseSearch/Source/Runtime/Private/PoseSearch.cpp b/Engine/Plugins/Experimental/Animation/PoseSearch/Source/Runtime/Private/PoseSearch.cpp
--- a/Engine/Plugins/Experimental/Animation/PoseSearch/Source/Runtime/Private/PoseSearch.cpp
+++ b/Engine/Plugins/Experimental/Animation/PoseSearch/Source/Runtime/Private/PoseSearch.cpp
@@ -2086,6 +2086,10 @@ UE::PoseSearch::FSearchResult UPoseSearchDatabase::SearchPCAKDTree(UE::PoseSearc
 				BestPoseCost = PoseCost;
 				BestPoseIdx = PoseIdx;
 			}
+
+			#if UE_POSE_SEARCH_TRACE_ENABLED
+			SearchContext.BestCandidates.Add(PoseCost.GetTotalCost(), PoseIdx, this);
+			#endif
 		}
 	}
 
@@ -2173,6 +2177,10 @@ UE::PoseSearch::FSearchResult UPoseSearchDatabase::SearchBruteForce(UE::PoseSear
 				BestPoseCost = PoseCost;
 				BestPoseIdx = PoseIdx;
 			}
+
+			#if UE_POSE_SEARCH_TRACE_ENABLED
+			SearchContext.BestCandidates.Add(PoseCost.GetTotalCost(), PoseIdx, this);
+			#endif
 		}
 	}
 
diff --git a/Engine/Plugins/Experimental/Animation/PoseSearch/Source/Runtime/Private/PoseSearchLibrary.cpp b/Engine/Plugins/Experimental/Animation/PoseSearch/Source/Runtime/Private/PoseSearchLibrary.cpp
--- a/Engine/Plugins/Experimental/Animation/PoseSearch/Source/Runtime/Private/PoseSearchLibrary.cpp
+++ b/Engine/Plugins/Experimental/Animation/PoseSearch/Source/Runtime/Private/PoseSearchLibrary.cpp
@@ -249,10 +249,10 @@ void UpdateMotionMatchingState(
 	bool bCanAdvance = InOutMotionMatchingState.CanAdvance(DeltaTime, bAdvanceToFollowUpAsset, FollowUpAsset);
 
 	// If we can't advance or enough time has elapsed since the last pose jump then search
+	FSearchContext SearchContext;
 	if (!bCanAdvance || (InOutMotionMatchingState.ElapsedPoseJumpTime >= Settings.SearchThrottleTime))
 	{
 		// Build the search context
-		FSearchContext SearchContext;
 		SearchContext.DatabaseTagQuery = DatabaseTagQuery;
 		SearchContext.ActiveTagsContainer = ActiveTagsContainer;
 		SearchContext.Trajectory = &Trajectory;
@@ -324,6 +324,9 @@ void UpdateMotionMatchingState(
 				}
 			}
 		}
+
+		// todo: cache the queries into TraceState, one per TraceState.DatabaseEntries or
+		// InOutMotionMatchingState.CurrentSearchResult.ComposedQuery = SearchResult.ComposedQuery;
 	}
 
 	// If we didn't search or it didn't find a pose to jump to, and we can 
@@ -345,6 +348,29 @@ void UpdateMotionMatchingState(
 
 	// Record debugger details
 #if UE_POSE_SEARCH_TRACE_ENABLED
+	FTraceMotionMatchingState TraceState;
+	TArray<const UPoseSearchDatabase*> Databases;
+	while (!SearchContext.BestCandidates.IsEmpty())
+	{
+		FSearchContext::PoseCandidate PoseCandidate;
+		SearchContext.BestCandidates.Pop(PoseCandidate);
+
+		int32 DatabaseEntryIndex = 0;
+		if (!Databases.Find(PoseCandidate.Database, DatabaseEntryIndex))
+		{
+			Databases.Push(PoseCandidate.Database);
+			DatabaseEntryIndex = Databases.Num() - 1;
+
+			TraceState.DatabaseEntries.AddDefaulted();
+			TraceState.DatabaseEntries[DatabaseEntryIndex].DatabaseId = FTraceMotionMatchingState::GetIdFromDatabase(PoseCandidate.Database);
+		}
+
+		FTraceMotionMatchingStatePoseEntry PoseEntry;
+		PoseEntry.DbPoseIdx = PoseCandidate.PoseIdx;
+		PoseEntry.Cost = PoseCandidate.Cost;
+		TraceState.DatabaseEntries[DatabaseEntryIndex].PoseEntries.Add(PoseEntry);
+	}
+
 	if (InOutMotionMatchingState.CurrentSearchResult.IsValid())
 	{
 		float SimLinearVelocity, SimAngularVelocity, AnimLinearVelocity, AnimAngularVelocity;
@@ -396,7 +422,6 @@ void UpdateMotionMatchingState(
 			DatabaseTagQuery, 
 			DatabaseBlendSpaceFilter);
 
-		FTraceMotionMatchingState TraceState;
 		if (EnumHasAnyFlags(InOutMotionMatchingState.Flags, EMotionMatchingFlags::JumpedToFollowUp))
 		{
 			TraceState.Flags |= FTraceMotionMatchingState::EFlags::FollowupAnimation;
@@ -407,7 +432,7 @@ void UpdateMotionMatchingState(
 		TraceState.QueryVector = InOutMotionMatchingState.CurrentSearchResult.ComposedQuery.GetValues();
 		TraceState.QueryVectorNormalized = InOutMotionMatchingState.CurrentSearchResult.ComposedQuery.GetNormalizedValues();
 		TraceState.DbPoseIdx = InOutMotionMatchingState.CurrentSearchResult.PoseIdx;
-		TraceState.DatabaseId = FObjectTrace::GetObjectId(InOutMotionMatchingState.CurrentSearchResult.Database.Get());
+		TraceState.SetDatabase(InOutMotionMatchingState.CurrentSearchResult.Database.Get());
 		TraceState.ContinuingPoseIdx = LastResult.PoseIdx;
 
 		TraceState.AssetPlayerTime = InOutMotionMatchingState.AssetPlayerTime;
@@ -418,7 +443,8 @@ void UpdateMotionMatchingState(
 		TraceState.AnimAngularVelocity = AnimAngularVelocity;
 		TraceState.DatabaseSequenceFilter = DatabaseSequenceFilter;
 		TraceState.DatabaseBlendSpaceFilter = DatabaseBlendSpaceFilter;
-		UE_TRACE_POSE_SEARCH_MOTION_MATCHING_STATE(Context, TraceState)
+
+		TraceState.Output(Context);
 	}
 #endif
 }
diff --git a/Engine/Plugins/Experimental/Animation/PoseSearch/Source/Runtime/Private/Trace/PoseSearchTraceLogger.cpp b/Engine/Plugins/Experimental/Animation/PoseSearch/Source/Runtime/Private/Trace/PoseSearchTraceLogger.cpp
--- a/Engine/Plugins/Experimental/Animation/PoseSearch/Source/Runtime/Private/Trace/PoseSearchTraceLogger.cpp
+++ b/Engine/Plugins/Experimental/Animation/PoseSearch/Source/Runtime/Private/Trace/PoseSearchTraceLogger.cpp
@@ -10,27 +10,8 @@
 UE_TRACE_CHANNEL_DEFINE(PoseSearchChannel);
 
 UE_TRACE_EVENT_BEGIN(PoseSearch, MotionMatchingState)
-	UE_TRACE_EVENT_FIELD(uint64, Cycle)
-	UE_TRACE_EVENT_FIELD(uint16, FrameCounter)
-	UE_TRACE_EVENT_FIELD(uint64, AnimInstanceId)
-	UE_TRACE_EVENT_FIELD(uint64, SkeletalMeshComponentId)
-	UE_TRACE_EVENT_FIELD(int32, NodeId)
-	UE_TRACE_EVENT_FIELD(float, ElapsedPoseJumpTime)
-	UE_TRACE_EVENT_FIELD(uint32, Flags)
-	UE_TRACE_EVENT_FIELD(uint64, DatabaseId)
-	UE_TRACE_EVENT_FIELD(int32, DbPoseIdx)
-	UE_TRACE_EVENT_FIELD(int32, ContinuingPoseIdx)
-	UE_TRACE_EVENT_FIELD(float[], QueryVector)
-	UE_TRACE_EVENT_FIELD(float[], QueryVectorNormalized)
-	UE_TRACE_EVENT_FIELD(float, AssetPlayerTime)
-	UE_TRACE_EVENT_FIELD(float, DeltaTime)
-	UE_TRACE_EVENT_FIELD(float, SimLinearVelocity)
-	UE_TRACE_EVENT_FIELD(float, SimAngularVelocity)
-	UE_TRACE_EVENT_FIELD(float, AnimLinearVelocity)
-	UE_TRACE_EVENT_FIELD(float, AnimAngularVelocity)
-	UE_TRACE_EVENT_FIELD(bool[], DatabaseSequenceFilter)
-	UE_TRACE_EVENT_FIELD(bool[], DatabaseBlendSpaceFilter)
-	UE_TRACE_EVENT_END()
+	UE_TRACE_EVENT_FIELD(uint8[], Data)
+UE_TRACE_EVENT_END()
 
 namespace UE { namespace PoseSearch {
 
@@ -54,38 +35,77 @@ static bool ShouldTrace(const FAnimationBaseContext& InContext)
 	return !CANNOT_TRACE_OBJECT(InContext.AnimInstanceProxy->GetSkelMeshComponent());
 }
 
-void FTraceMotionMatchingState::Output(const FAnimationBaseContext& InContext, const FTraceMotionMatchingState& State)
+FArchive& operator<<(FArchive& Ar, FTraceMessage& State)
+{
+	Ar << State.Cycle;
+	Ar << State.AnimInstanceId;
+	Ar << State.SkeletalMeshComponentId;
+	Ar << State.NodeId;
+	Ar << State.FrameCounter;
+	return Ar;
+}
+
+FArchive& operator<<(FArchive& Ar, FTraceMotionMatchingStatePoseEntry& Entry)
+{
+	Ar << Entry.DbPoseIdx;
+	Ar << Entry.Cost;
+	return Ar;
+}
+
+FArchive& operator<<(FArchive& Ar, FTraceMotionMatchingStateDatabaseEntry& Entry)
+{
+	Ar << Entry.DatabaseId;
+	Ar << Entry.PoseEntries;
+	return Ar;
+}
+
+FArchive& operator<<(FArchive& Ar, FTraceMotionMatchingState& State)
+{
+	Ar << State.ElapsedPoseJumpTime;
+	Ar << State.Flags;
+	Ar << State.QueryVector;
+	Ar << State.QueryVectorNormalized;
+	Ar << State.DatabaseSequenceFilter;
+	Ar << State.DatabaseBlendSpaceFilter;
+	Ar << State.DbPoseIdx;
+	Ar << State.DatabaseId;
+	Ar << State.ContinuingPoseIdx;
+	Ar << State.AssetPlayerTime;
+	Ar << State.DeltaTime;
+	Ar << State.SimLinearVelocity;
+	Ar << State.SimAngularVelocity;
+	Ar << State.AnimLinearVelocity;
+	Ar << State.AnimAngularVelocity;
+	Ar << State.DatabaseEntries;
+	return Ar;
+}
+
+void FTraceMotionMatchingState::Output(const FAnimationBaseContext& InContext)
 {
 	if (!ShouldTrace(InContext))
 	{
 		return;
 	}
 
+	TArray<uint8> ArchiveData;
+	FMemoryWriter Archive(ArchiveData);
+
 	UObject* AnimInstance = InContext.AnimInstanceProxy->GetAnimInstanceObject();
 	TRACE_OBJECT(AnimInstance);
 	UObject* SkeletalMeshComponent = AnimInstance->GetOuter();
 
+	FTraceMessage TraceMessage;
+	TraceMessage.Cycle = FPlatformTime::Cycles64();
+	TraceMessage.AnimInstanceId = FObjectTrace::GetObjectId(AnimInstance);
+	TraceMessage.SkeletalMeshComponentId = FObjectTrace::GetObjectId(SkeletalMeshComponent);
+	TraceMessage.NodeId = InContext.GetCurrentNodeId();
+	TraceMessage.FrameCounter = FObjectTrace::GetObjectWorldTickCounter(AnimInstance);
+
+	Archive << TraceMessage;
+	Archive << *this;
+
 	UE_TRACE_LOG(PoseSearch, MotionMatchingState, PoseSearchChannel)
-		<< MotionMatchingState.Cycle(FPlatformTime::Cycles64())
-		<< MotionMatchingState.FrameCounter(FObjectTrace::GetObjectWorldTickCounter(AnimInstance))
-		<< MotionMatchingState.AnimInstanceId(FObjectTrace::GetObjectId(AnimInstance))
-		<< MotionMatchingState.SkeletalMeshComponentId(FObjectTrace::GetObjectId(SkeletalMeshComponent))
-		<< MotionMatchingState.NodeId(InContext.GetCurrentNodeId())
-		<< MotionMatchingState.ElapsedPoseJumpTime(State.ElapsedPoseJumpTime)
-		<< MotionMatchingState.Flags(static_cast<uint32>(State.Flags))
-		<< MotionMatchingState.DatabaseId(State.DatabaseId)
-		<< MotionMatchingState.DbPoseIdx(State.DbPoseIdx)
-		<< MotionMatchingState.ContinuingPoseIdx(State.ContinuingPoseIdx)
-		<< MotionMatchingState.QueryVector(State.QueryVector.GetData(), State.QueryVector.Num())
-		<< MotionMatchingState.QueryVectorNormalized(State.QueryVectorNormalized.GetData(), State.QueryVectorNormalized.Num())
-		<< MotionMatchingState.AssetPlayerTime(State.AssetPlayerTime)
-		<< MotionMatchingState.DeltaTime(State.DeltaTime)
-		<< MotionMatchingState.SimLinearVelocity(State.SimLinearVelocity)
-		<< MotionMatchingState.SimAngularVelocity(State.SimAngularVelocity)
-		<< MotionMatchingState.AnimLinearVelocity(State.AnimLinearVelocity)
-		<< MotionMatchingState.AnimAngularVelocity(State.AnimAngularVelocity)
-		<< MotionMatchingState.DatabaseSequenceFilter(State.DatabaseSequenceFilter.GetData(), State.DatabaseSequenceFilter.Num())
-		<< MotionMatchingState.DatabaseBlendSpaceFilter(State.DatabaseBlendSpaceFilter.GetData(), State.DatabaseBlendSpaceFilter.Num());
+		<< MotionMatchingState.Data(ArchiveData.GetData(), ArchiveData.Num());
 }
 
 }}
diff --git a/Engine/Plugins/Experimental/Animation/PoseSearch/Source/Runtime/Private/Trace/PoseSearchTraceLogger.h b/Engine/Plugins/Experimental/Animation/PoseSearch/Source/Runtime/Private/Trace/PoseSearchTraceLogger.h
--- a/Engine/Plugins/Experimental/Animation/PoseSearch/Source/Runtime/Private/Trace/PoseSearchTraceLogger.h
+++ b/Engine/Plugins/Experimental/Animation/PoseSearch/Source/Runtime/Private/Trace/PoseSearchTraceLogger.h
@@ -3,10 +3,7 @@
 #pragma once
 
 #include "PoseSearch/PoseSearch.h"
-#include "ObjectTrace.h"
 
-// Enable this if object tracing is enabled, mimics animation tracing
-#define UE_POSE_SEARCH_TRACE_ENABLED OBJECT_TRACE_ENABLED
 
 #if UE_POSE_SEARCH_TRACE_ENABLED
 
@@ -20,6 +17,42 @@ struct POSESEARCH_API FTraceLogger
     static const FName Name;
 };
 
+// Message types for appending / reading to the timeline
+/** Base message type for common data */
+struct POSESEARCH_API FTraceMessage
+{
+	uint64 Cycle = 0;
+
+	uint64 AnimInstanceId = 0;
+
+	/** Skeletal Mesh Component Id, outer class of the AnimInstance.
+	 *  Used for retrieval of traced root transforms from the animation provider.
+	 */
+	uint64 SkeletalMeshComponentId = 0;
+
+	/** Node Id of the motion matching node associated with this message */
+	int32 NodeId = 0;
+
+	uint16 FrameCounter = 0;
+};
+POSESEARCH_API FArchive& operator<<(FArchive& Ar, FTraceMessage& State);
+
+struct POSESEARCH_API FTraceMotionMatchingStatePoseEntry
+{
+	int32 DbPoseIdx = 0;
+	float Cost = 0.f;
+};
+POSESEARCH_API FArchive& operator<<(FArchive& Ar, FTraceMotionMatchingStatePoseEntry& Entry);
+
+struct POSESEARCH_API FTraceMotionMatchingStateDatabaseEntry
+{
+	// @todo: can we use UPoseSearchDatabase* instead of DatabaseId?
+	//UPoseSearchDatabase const* Database = nullptr;
+	uint64 DatabaseId = 0;
+	TArray<FTraceMotionMatchingStatePoseEntry> PoseEntries;
+};
+POSESEARCH_API FArchive& operator<<(FArchive& Ar, FTraceMotionMatchingStateDatabaseEntry& Entry);
+
 /**
  * Used to trace motion matching state data via the logger, which is then placed into a timeline
  */
@@ -41,14 +74,14 @@ struct POSESEARCH_API FTraceMotionMatchingState
 	EFlags Flags = EFlags::None;
 
 	/** Search vectors in normalized and unnormalized forms */
-	TArrayView<const float> QueryVector;
-	TArrayView<const float> QueryVectorNormalized;
+	TArray<float> QueryVector;
+	TArray<float> QueryVectorNormalized;
 
 	/** If true, groups are being filtered by DatabaseGroupQuery */
-	TArrayView<const bool> DatabaseSequenceFilter;
+	TArray<bool> DatabaseSequenceFilter;
 
 	/** If true, groups are being filtered by DatabaseGroupQuery */
-	TArrayView<const bool> DatabaseBlendSpaceFilter;
+	TArray<bool> DatabaseBlendSpaceFilter;
 
 	/** Index of the pose in our database */
 	int32 DbPoseIdx = 0;
@@ -59,39 +92,59 @@ struct POSESEARCH_API FTraceMotionMatchingState
 	/** Index of the continuing pose in our database */
 	int32 ContinuingPoseIdx = 0;
 
-	/** Node Id of the motion matching node associated with this message */
-	int32 NodeId = 0;
-
-	/** Skeletal Mesh Component Id, outer class of the AnimInstance.
-	 *  Used for retrieval of traced root transforms from the animation provider.
-	 */
-	uint64 SkeletalMeshComponentId = 0;
-
 	float AssetPlayerTime = 0.0f;
 	float DeltaTime = 0.0f;
 	float SimLinearVelocity = 0.0f;
 	float SimAngularVelocity = 0.0f;
 	float AnimLinearVelocity = 0.0f;
 	float AnimAngularVelocity = 0.0f;
-
+	
+	TArray<FTraceMotionMatchingStateDatabaseEntry> DatabaseEntries;
 
 	/** Output the current state info to the logger */
-	static void Output(const FAnimationBaseContext& InContext, const FTraceMotionMatchingState& State);
+	void Output(const FAnimationBaseContext& InContext);
+
+	static const UPoseSearchDatabase* GetDatabaseFromId(uint64 DatabaseId)
+	{
+		if (DatabaseId == 0)
+		{
+			return nullptr;
+		}
+
+		UObject* DatabaseObject = FObjectTrace::GetObjectFromId(DatabaseId);
+		// @TODO: Load the object if unloaded
+		if (DatabaseObject == nullptr)
+		{
+			return nullptr;
+		}
+		check(DatabaseObject->IsA<UPoseSearchDatabase>());
+
+		const UPoseSearchDatabase* Database = Cast<UPoseSearchDatabase>(DatabaseObject);
+		return Database;
+	}
+
+	static uint64 GetIdFromDatabase(const UPoseSearchDatabase* Database)
+	{
+		return FObjectTrace::GetObjectId(Database);
+	}
+
+	void SetDatabase(const UPoseSearchDatabase* Database)
+	{
+		DatabaseId = GetIdFromDatabase(Database);
+	}
+
+	const UPoseSearchDatabase* GetPoseSearchDatabase() const
+	{
+		return GetDatabaseFromId(DatabaseId);
+	}
 	
 	static const FName Name;
 };
 ENUM_CLASS_FLAGS(FTraceMotionMatchingState::EFlags)
 
+POSESEARCH_API FArchive& operator<<(FArchive& Ar, FTraceMotionMatchingState& State);
 
 }}
 
-#define UE_TRACE_POSE_SEARCH_MOTION_MATCHING_STATE(Context, MotionMatchingState) \
-	UE::PoseSearch::FTraceMotionMatchingState::Output(Context, MotionMatchingState);
-
-#else // UE_POSE_SEARCH_TRACE_ENABLED
-
-// Empty declarations if not enabled
-#define UE_TRACE_POSE_SEARCH_MOTION_MATCHING_STATE(Context, MotionMatchingState)
-
 #endif // UE_POSE_SEARCH_TRACE_ENABLED
 
diff --git a/Engine/Plugins/Experimental/Animation/PoseSearch/Source/Runtime/Public/PoseSearch/PoseSearch.h b/Engine/Plugins/Experimental/Animation/PoseSearch/Source/Runtime/Public/PoseSearch/PoseSearch.h
--- a/Engine/Plugins/Experimental/Animation/PoseSearch/Source/Runtime/Public/PoseSearch/PoseSearch.h
+++ b/Engine/Plugins/Experimental/Animation/PoseSearch/Source/Runtime/Public/PoseSearch/PoseSearch.h
@@ -19,11 +19,13 @@
 #include "GameplayTagContainer.h"
 #include "Interfaces/Interface_BoneReferenceSkeletonProvider.h"
 #include "PoseSearch/KDTree.h"
-
+#include "ObjectTrace.h"
 #include "PoseSearch.generated.h"
 
 DECLARE_LOG_CATEGORY_EXTERN(LogPoseSearch, Log, All);
 
+// Enable this if object tracing is enabled, mimics animation tracing
+#define UE_POSE_SEARCH_TRACE_ENABLED OBJECT_TRACE_ENABLED
 
 //////////////////////////////////////////////////////////////////////////
 // Forward declarations
@@ -1322,6 +1324,52 @@ private:
 
 	// @todo: make it a fixed size array (or hash map if we end up having many CachedEntry) to avoid allocations
 	TArray<CachedEntry> CachedEntries;
+
+#if UE_POSE_SEARCH_TRACE_ENABLED
+
+public:
+	struct PoseCandidate
+	{
+		float Cost = 0.f;
+		int32 PoseIdx = 0;
+		const UPoseSearchDatabase* Database = nullptr;
+
+		bool operator<(const PoseCandidate& Other) const { return Cost > Other.Cost; }
+	};
+
+	struct BestPoseCandidates : private TArray<PoseCandidate>
+	{
+		typedef TArray<PoseCandidate> Super;
+		using Super::IsEmpty;
+
+		int32 MaxPoseCandidates = 100;
+
+		void Add(float Cost, int32 PoseIdx, const UPoseSearchDatabase* Database)
+		{
+			if (Num() < MaxPoseCandidates || Cost < HeapTop().Cost)
+			{
+				while (Num() >= MaxPoseCandidates)
+				{
+					ElementType Unused;
+					Pop(Unused);
+				}
+
+				FSearchContext::PoseCandidate PoseCandidate;
+				PoseCandidate.Cost = Cost;
+				PoseCandidate.PoseIdx = PoseIdx;
+				PoseCandidate.Database = Database;
+				HeapPush(PoseCandidate);
+			}
+		}
+
+		void Pop(PoseCandidate& OutItem)
+		{
+			HeapPop(OutItem, false);
+		}
+	};
+	
+	BestPoseCandidates BestCandidates;
+#endif
 };
 
 //////////////////////////////////////////////////////////////////////////
