motion matching
- moving biases to schema
- making consistent naming for continuing pose
- removing MinPercentImprovement

diff --git a/Engine/Plugins/Experimental/Animation/PoseSearch/Source/Editor/Private/PoseSearchDebugger.cpp b/Engine/Plugins/Experimental/Animation/PoseSearch/Source/Editor/Private/PoseSearchDebugger.cpp
--- a/Engine/Plugins/Experimental/Animation/PoseSearch/Source/Editor/Private/PoseSearchDebugger.cpp
+++ b/Engine/Plugins/Experimental/Animation/PoseSearch/Source/Editor/Private/PoseSearchDebugger.cpp
@@ -961,9 +961,12 @@ void SDebuggerDatabaseView::UpdateRows(const FTraceMotionMatchingStateMessage& S
 			// @todo: invalidate rows in case Database changed while PIE is active with pose debugger. for now we just avoid crashes
 			if (PoseIdx < Database.GetSearchIndex()->NumPoses)
 			{
+				const EPoseComparisonFlags PoseComparisonFlags = (PoseIdx == State.ContinuingPoseIdx) ? EPoseComparisonFlags::ContinuingPose : EPoseComparisonFlags::None;
+
 				Database.ComparePoses(
 					SearchContext,
 					PoseIdx,
+					PoseComparisonFlags,
 					State.QueryVectorNormalized,
 					Row->PoseCostDetails);
 
diff --git a/Engine/Plugins/Experimental/Animation/PoseSearch/Source/Runtime/Private/PoseSearch.cpp b/Engine/Plugins/Experimental/Animation/PoseSearch/Source/Runtime/Private/PoseSearch.cpp
--- a/Engine/Plugins/Experimental/Animation/PoseSearch/Source/Runtime/Private/PoseSearch.cpp
+++ b/Engine/Plugins/Experimental/Animation/PoseSearch/Source/Runtime/Private/PoseSearch.cpp
@@ -493,6 +493,9 @@ void UPoseSearchSchema::GenerateDDCKey(FBlake3& InOutKeyHasher) const
 	InOutKeyHasher.Update(&EffectiveDataPreprocessor, sizeof(EffectiveDataPreprocessor));
 	InOutKeyHasher.Update(&SamplingInterval, sizeof(SamplingInterval));
 	InOutKeyHasher.Update(MakeMemoryView(BoneIndicesWithParents));
+	InOutKeyHasher.Update(&ContinuingPoseCostBias, sizeof(ContinuingPoseCostBias));
+	InOutKeyHasher.Update(&BaseCostBias, sizeof(BaseCostBias));
+	InOutKeyHasher.Update(&MirrorMismatchCostBias, sizeof(MirrorMismatchCostBias));
 
 	// @todo: add schema version if POSESEARCHDB_DERIVEDDATA_VER is not enought
 }
@@ -823,7 +826,7 @@ UE::PoseSearch::FSearchResult UPoseSearchSequenceMetaData::Search(UE::PoseSearch
 				continue;
 			}
 
-			FPoseSearchCost PoseCost = ComparePoses(PoseIdx, NormalizedQueryValues);
+			FPoseSearchCost PoseCost = ComparePoses(PoseIdx, EPoseComparisonFlags::ContinuingPose, NormalizedQueryValues);
 
 			if (PoseCost < BestPoseCost)
 			{
@@ -848,7 +851,8 @@ UE::PoseSearch::FSearchResult UPoseSearchSequenceMetaData::Search(UE::PoseSearch
 }
 
 FPoseSearchCost UPoseSearchSequenceMetaData::ComparePoses(
-	int32 PoseIdx, 
+	int32 PoseIdx,
+	UE::PoseSearch::EPoseComparisonFlags PoseComparisonFlags,
 	const TArrayView<const float>& QueryValues) const
 {
 	using namespace UE::PoseSearch;
@@ -856,15 +860,13 @@ FPoseSearchCost UPoseSearchSequenceMetaData::ComparePoses(
 	FPoseSearchCost Result;
 
 	TArrayView<const float> PoseValues = SearchIndex.GetPoseValues(PoseIdx);
-	if (!ensure(PoseValues.Num() == QueryValues.Num()))
-	{
-		return Result;
-	}
+	check(PoseValues.Num() == QueryValues.Num());
 
 	Result.SetDissimilarity(CompareFeatureVectors(PoseValues.Num(), PoseValues.GetData(), QueryValues.GetData()));
 
 	const FPoseSearchPoseMetadata& PoseMetadata = SearchIndex.PoseMetadata[PoseIdx];
-	Result.SetCostAddend(PoseMetadata.CostAddend);
+	const float ContinuingPoseCostAddend = EnumHasAnyFlags(PoseComparisonFlags, EPoseComparisonFlags::ContinuingPose) ? PoseMetadata.ContinuingPoseCostAddend : 0.f;
+	Result.SetCostAddend(PoseMetadata.CostAddend + ContinuingPoseCostAddend);
 
 	return Result;
 }
@@ -1752,32 +1754,22 @@ bool UPoseSearchDatabase::IsCachedCookedPlatformDataLoaded(const ITargetPlatform
 
 #endif // WITH_EDITOR
 
-FPoseSearchCost UPoseSearchDatabase::ComparePoses(UE::PoseSearch::FSearchContext& SearchContext, int32 PoseIdx, int32 GroupIdx, const TArrayView<const float>& QueryValues) const
+FPoseSearchCost UPoseSearchDatabase::ComparePoses(UE::PoseSearch::FSearchContext& SearchContext, int32 PoseIdx, UE::PoseSearch::EPoseComparisonFlags PoseComparisonFlags, int32 GroupIdx, const TArrayView<const float>& QueryValues) const
 {
 	using namespace UE::PoseSearch;
 
 	FPoseSearchCost Result;
 
 	const FPoseSearchIndex* SearchIndex = GetSearchIndex();
-	if (!ensure(SearchIndex))
-	{
-		return Result;
-	}
+	check(SearchIndex);
 
 	TArrayView<const float> PoseValues = SearchIndex->GetPoseValues(PoseIdx);
-	if (!ensure(PoseValues.Num() == QueryValues.Num()))
-	{
-		return Result;
-	}
+	check(PoseValues.Num() == QueryValues.Num());
 
 	if (GroupIdx == INDEX_NONE)
 	{
 		const FPoseSearchIndexAsset* SearchIndexAsset = SearchIndex->FindAssetForPose(PoseIdx);
-		if (!ensure(SearchIndexAsset))
-		{
-			return Result;
-		}
-
+		check(SearchIndexAsset)
 		GroupIdx = SearchIndexAsset->SourceGroupIdx;
 	}
 
@@ -1789,13 +1781,13 @@ FPoseSearchCost UPoseSearchDatabase::ComparePoses(UE::PoseSearch::FSearchContext
 
 	float NotifyAddend = 0.0f;
 	float MirrorMismatchAddend = 0.0f;
-	ComputePoseCostAddends(PoseIdx, SearchContext, NotifyAddend, MirrorMismatchAddend);
+	ComputePoseCostAddends(PoseIdx, PoseComparisonFlags, SearchContext, NotifyAddend, MirrorMismatchAddend);
 	Result.SetCostAddend(NotifyAddend + MirrorMismatchAddend);
 
 	return Result;
 }
 
-FPoseSearchCost UPoseSearchDatabase::ComparePoses(UE::PoseSearch::FSearchContext& SearchContext, int32 PoseIdx, const TArrayView<const float>& QueryValues, UE::PoseSearch::FPoseCostDetails& OutPoseCostDetails) const
+FPoseSearchCost UPoseSearchDatabase::ComparePoses(UE::PoseSearch::FSearchContext& SearchContext, int32 PoseIdx, UE::PoseSearch::EPoseComparisonFlags PoseComparisonFlags, const TArrayView<const float>& QueryValues, UE::PoseSearch::FPoseCostDetails& OutPoseCostDetails) const
 {
 	using namespace Eigen;
 	using namespace UE::PoseSearch;
@@ -1804,10 +1796,7 @@ FPoseSearchCost UPoseSearchDatabase::ComparePoses(UE::PoseSearch::FSearchContext
 
 	TArrayView<const float> PoseValues = GetSearchIndex()->GetPoseValues(PoseIdx);
 	const int32 Dims = PoseValues.Num();
-	if (!ensure(Dims == QueryValues.Num()))
-	{
-		return Result;
-	}
+	check(Dims == QueryValues.Num());
 
 	OutPoseCostDetails.CostVector.SetNum(Dims);
 
@@ -1829,7 +1818,7 @@ FPoseSearchCost UPoseSearchDatabase::ComparePoses(UE::PoseSearch::FSearchContext
 	// Output result
 	float NotifyAddend = 0.0f;
 	float MirrorMismatchAddend = 0.0f;
-	ComputePoseCostAddends(PoseIdx, SearchContext, NotifyAddend, MirrorMismatchAddend);
+	ComputePoseCostAddends(PoseIdx, PoseComparisonFlags, SearchContext, NotifyAddend, MirrorMismatchAddend);
 	Result.SetCostAddend(NotifyAddend + MirrorMismatchAddend);
 
 	// Output cost details
@@ -1847,6 +1836,7 @@ FPoseSearchCost UPoseSearchDatabase::ComparePoses(UE::PoseSearch::FSearchContext
 		FPoseSearchCost RuntimeComparatorCost = ComparePoses(
 			SearchContext, 
 			PoseIdx, 
+			PoseComparisonFlags,
 			SearchIndexAsset->SourceGroupIdx,
 			QueryValues);
 		checkSlow(FMath::IsNearlyEqual(Result.GetTotalCost(), RuntimeComparatorCost.GetTotalCost(), 1e-3f));
@@ -1862,25 +1852,29 @@ FPoseSearchCost UPoseSearchDatabase::ComparePoses(UE::PoseSearch::FSearchContext
 	return Result;
 }
 
-void UPoseSearchDatabase::ComputePoseCostAddends(int32 PoseIdx, UE::PoseSearch::FSearchContext& SearchContext, float& OutNotifyAddend, float& OutMirrorMismatchAddend) const
+void UPoseSearchDatabase::ComputePoseCostAddends(int32 PoseIdx, UE::PoseSearch::EPoseComparisonFlags PoseComparisonFlags, UE::PoseSearch::FSearchContext& SearchContext, float& OutNotifyAddend, float& OutMirrorMismatchAddend) const
 {
 	OutNotifyAddend = 0.0f;
 	OutMirrorMismatchAddend = 0.0f;
 
+	const FPoseSearchIndex* SearchIndex = GetSearchIndex();
+	check(SearchIndex);
 	if (SearchContext.QueryMirrorRequest != EPoseSearchBooleanRequest::Indifferent)
 	{
-		const FPoseSearchIndexAsset* IndexAsset = GetSearchIndex()->FindAssetForPose(PoseIdx);
+		const FPoseSearchIndexAsset* IndexAsset = SearchIndex->FindAssetForPose(PoseIdx);
 		const bool bMirroringMismatch =
 			(IndexAsset->bMirrored && SearchContext.QueryMirrorRequest == EPoseSearchBooleanRequest::FalseValue) ||
 			(!IndexAsset->bMirrored && SearchContext.QueryMirrorRequest == EPoseSearchBooleanRequest::TrueValue);
 		if (bMirroringMismatch)
 		{
-			OutMirrorMismatchAddend = MirroringMismatchCost;
+			check(Schema);
+			OutMirrorMismatchAddend = Schema->MirrorMismatchCostBias;
 		}
 	}
 
-	const FPoseSearchPoseMetadata& PoseMetadata = GetSearchIndex()->PoseMetadata[PoseIdx];
-	OutNotifyAddend = PoseMetadata.CostAddend;
+	const FPoseSearchPoseMetadata& PoseMetadata = SearchIndex->PoseMetadata[PoseIdx];
+	const float ContinuingPoseCostAddend = EnumHasAnyFlags(PoseComparisonFlags, UE::PoseSearch::EPoseComparisonFlags::ContinuingPose) ? PoseMetadata.ContinuingPoseCostAddend : 0.f;
+	OutNotifyAddend = PoseMetadata.CostAddend + ContinuingPoseCostAddend;
 }
 
 UE::PoseSearch::FSearchResult UPoseSearchDatabase::Search(UE::PoseSearch::FSearchContext& SearchContext) const
@@ -2036,9 +2030,10 @@ UE::PoseSearch::FSearchResult UPoseSearchDatabase::SearchPCAKDTree(UE::PoseSearc
 	const bool IsCurrentResultFromThisDatabase = !SearchContext.bForceInterrupt && SearchContext.CurrentResult.IsValid() && SearchContext.CurrentResult.Database == this;
 	if (IsCurrentResultFromThisDatabase)
 	{
-		Result.ContinuityPoseCost = ComparePoses(
+		Result.ContinuingPoseCost = ComparePoses(
 			SearchContext,
 			SearchContext.CurrentResult.PoseIdx,
+			EPoseComparisonFlags::ContinuingPose,
 			SearchIndex->FindAssetForPose(SearchContext.CurrentResult.PoseIdx)->SourceGroupIdx,
 			NormalizedQueryValues);
 	}
@@ -2078,6 +2073,7 @@ UE::PoseSearch::FSearchResult UPoseSearchDatabase::SearchPCAKDTree(UE::PoseSearc
 			FPoseSearchCost PoseCost = ComparePoses(
 				SearchContext, 
 				PoseIdx, 
+				EPoseComparisonFlags::None,
 				GroupSearchIndex.GroupIndex,
 				NormalizedQueryValues);
 
@@ -2127,9 +2123,10 @@ UE::PoseSearch::FSearchResult UPoseSearchDatabase::SearchBruteForce(UE::PoseSear
 	const bool IsCurrentResultFromThisDatabase = !SearchContext.bForceInterrupt && SearchContext.CurrentResult.IsValid() && SearchContext.CurrentResult.Database == this;
 	if (IsCurrentResultFromThisDatabase)
 	{
-		Result.ContinuityPoseCost = ComparePoses(
+		Result.ContinuingPoseCost = ComparePoses(
 			SearchContext,
 			SearchContext.CurrentResult.PoseIdx,
+			EPoseComparisonFlags::ContinuingPose,
 			SearchIndex->FindAssetForPose(SearchContext.CurrentResult.PoseIdx)->SourceGroupIdx,
 			NormalizedQueryValues);
 
@@ -2169,6 +2166,7 @@ UE::PoseSearch::FSearchResult UPoseSearchDatabase::SearchBruteForce(UE::PoseSear
 			FPoseSearchCost PoseCost = ComparePoses(
 				SearchContext,
 				PoseIdx, 
+				EPoseComparisonFlags::None,
 				Asset.SourceGroupIdx,
 				NormalizedQueryValues);
 
@@ -2208,42 +2206,7 @@ UE::PoseSearch::FSearchResult UPoseSearchDatabaseSet::Search(UE::PoseSearch::FSe
 	using namespace UE::PoseSearch;
 
 	FSearchResult Result;
-	FPoseSearchCost ContinuityCost;
-
-	auto ProcessActiveEntry = 
-		[&SearchContext, &Result, &ContinuityCost](const FPoseSearchDatabaseSetEntry& Entry) 
-		-> EPoseSearchPostSearchStatus
-	{
-		EPoseSearchPostSearchStatus PostSearchStatus = EPoseSearchPostSearchStatus::Continue;
-		FSearchResult EntryResult = Entry.Searchable->Search(SearchContext);
-		if (EntryResult.IsValid())
-		{
-			if (IsValid(Entry.PostProcessor))
-			{
-				PostSearchStatus = Entry.PostProcessor->PostProcess(EntryResult.PoseCost);
-				if (EntryResult.ContinuityPoseCost.IsValid())
-				{
-					Entry.PostProcessor->PostProcess(EntryResult.ContinuityPoseCost);
-				}
-			}
-
-			if (!Result.IsValid() || EntryResult.PoseCost.GetTotalCost() < Result.PoseCost.GetTotalCost())
-			{
-				Result = EntryResult;
-			}
-
-			if (EntryResult.ContinuityPoseCost.IsValid())
-			{
-				if (!ContinuityCost.IsValid() ||
-					EntryResult.ContinuityPoseCost.GetTotalCost() < ContinuityCost.GetTotalCost())
-				{
-					ContinuityCost = EntryResult.ContinuityPoseCost;
-				}
-			}
-		}
-
-		return PostSearchStatus;
-	};
+	FPoseSearchCost ContinuingCost;
 
 	for (const FPoseSearchDatabaseSetEntry& Entry : AssetsToSearch)
 	{
@@ -2261,17 +2224,33 @@ UE::PoseSearch::FSearchResult UPoseSearchDatabaseSet::Search(UE::PoseSearch::FSe
 
 		if (bSearchEntry)
 		{
-			EPoseSearchPostSearchStatus PostSearchStatus = ProcessActiveEntry(Entry);
-			if (PostSearchStatus == EPoseSearchPostSearchStatus::Stop)
+			FSearchResult EntryResult = Entry.Searchable->Search(SearchContext);
+			if (EntryResult.IsValid())
 			{
-				break;
+				if (!Result.IsValid() || EntryResult.PoseCost.GetTotalCost() < Result.PoseCost.GetTotalCost())
+				{
+					Result = EntryResult;
+				}
+
+				if (EntryResult.ContinuingPoseCost.IsValid())
+				{
+					if (!ContinuingCost.IsValid() || EntryResult.ContinuingPoseCost.GetTotalCost() < ContinuingCost.GetTotalCost())
+					{
+						ContinuingCost = EntryResult.ContinuingPoseCost;
+					}
+				}
+
+				if (Entry.PostSearchStatus == EPoseSearchPostSearchStatus::Stop)
+				{
+					break;
+				}
 			}
 		}
 	}
 
 	if (Result.IsValid())
 	{
-		Result.ContinuityPoseCost = ContinuityCost;
+		Result.ContinuingPoseCost = ContinuingCost;
 	}
 	else
 	{
@@ -4107,6 +4086,9 @@ void FAssetIndexer::AddMetadata(int32 SampleIdx)
 		EnumAddFlags(Metadata.Flags, EPoseSearchPoseFlags::BlockTransition);
 	}
 
+	Metadata.CostAddend = IndexingContext.Schema->BaseCostBias;
+	Metadata.ContinuingPoseCostAddend = IndexingContext.Schema->ContinuingPoseCostBias;
+	
 	TArray<UAnimNotifyState_PoseSearchBase*> NotifyStates;
 	IndexingContext.MainSampler->ExtractPoseSearchNotifyStates(SampleTime, NotifyStates);
 	for (const UAnimNotifyState_PoseSearchBase* PoseSearchNotify : NotifyStates)
@@ -4121,6 +4103,12 @@ void FAssetIndexer::AddMetadata(int32 SampleIdx)
 				Cast<const UAnimNotifyState_PoseSearchModifyCost>(PoseSearchNotify);
 			Metadata.CostAddend = ModifyCostNotify->CostAddend;
 		}
+		else if (PoseSearchNotify->GetClass()->IsChildOf<UAnimNotifyState_PoseSearchOverrideContinuingPoseCostBias>())
+		{
+			const UAnimNotifyState_PoseSearchOverrideContinuingPoseCostBias* ContinuingPoseCostBias =
+				Cast<const UAnimNotifyState_PoseSearchOverrideContinuingPoseCostBias>(PoseSearchNotify);
+			Metadata.ContinuingPoseCostAddend = ContinuingPoseCostBias->CostAddend;
+		}
 	}
 }
 
@@ -5373,13 +5361,6 @@ void FModule::OnObjectSaved(UObject* SavedObject, FObjectPreSaveContext SaveCont
 
 }} // namespace UE::PoseSearch
 
-EPoseSearchPostSearchStatus UPoseSearchPostProcessor::PostProcess_Implementation(
-	FPoseSearchCost& InOutCost) const
-{
-	return EPoseSearchPostSearchStatus::Continue;
-}
-
-
 #undef LOCTEXT_NAMESPACE
 
 IMPLEMENT_MODULE(UE::PoseSearch::FModule, PoseSearch)
\ No newline at end of file
diff --git a/Engine/Plugins/Experimental/Animation/PoseSearch/Source/Runtime/Private/PoseSearchFeatureChannels.cpp b/Engine/Plugins/Experimental/Animation/PoseSearch/Source/Runtime/Private/PoseSearchFeatureChannels.cpp
--- a/Engine/Plugins/Experimental/Animation/PoseSearch/Source/Runtime/Private/PoseSearchFeatureChannels.cpp
+++ b/Engine/Plugins/Experimental/Animation/PoseSearch/Source/Runtime/Private/PoseSearchFeatureChannels.cpp
@@ -356,7 +356,7 @@ void UPoseSearchFeatureChannel_Position::GenerateDDCKey(FBlake3& InOutKeyHasher)
 	InOutKeyHasher.Update(&Bone, sizeof(Bone));
 	InOutKeyHasher.Update(&Weight, sizeof(Weight));
 	InOutKeyHasher.Update(&SampleTimeOffset, sizeof(SampleTimeOffset));
-	InOutKeyHasher.Update(&bUseFeaturesFromContinuityPose, sizeof(bUseFeaturesFromContinuityPose));
+	InOutKeyHasher.Update(&bUseFeaturesFromContinuingPose, sizeof(bUseFeaturesFromContinuingPose));
 	InOutKeyHasher.Update(&bUseSampleTimeOffsetRootBone, sizeof(bUseSampleTimeOffsetRootBone));
 }
 
@@ -364,7 +364,7 @@ bool UPoseSearchFeatureChannel_Position::BuildQuery(UE::PoseSearch::FSearchConte
 {
 	using namespace UE::PoseSearch;
 
-	const bool bSkip = bUseFeaturesFromContinuityPose && SearchContext.CurrentResult.IsValid() && SearchContext.CurrentResult.Database->Schema == InOutQuery.GetSchema();
+	const bool bSkip = bUseFeaturesFromContinuingPose && SearchContext.CurrentResult.IsValid() && SearchContext.CurrentResult.Database->Schema == InOutQuery.GetSchema();
 	if (bSkip)
 	{
 		return true;
@@ -489,7 +489,7 @@ void UPoseSearchFeatureChannel_Heading::GenerateDDCKey(FBlake3& InOutKeyHasher)
 	InOutKeyHasher.Update(&Weight, sizeof(Weight));
 	InOutKeyHasher.Update(&SampleTimeOffset, sizeof(SampleTimeOffset));
 	InOutKeyHasher.Update(&HeadingAxis, sizeof(HeadingAxis));
-	InOutKeyHasher.Update(&bUseFeaturesFromContinuityPose, sizeof(bUseFeaturesFromContinuityPose));
+	InOutKeyHasher.Update(&bUseFeaturesFromContinuingPose, sizeof(bUseFeaturesFromContinuingPose));
 	InOutKeyHasher.Update(&bUseSampleTimeOffsetRootBone, sizeof(bUseSampleTimeOffsetRootBone));
 }
 
@@ -497,7 +497,7 @@ bool UPoseSearchFeatureChannel_Heading::BuildQuery(UE::PoseSearch::FSearchContex
 {
 	using namespace UE::PoseSearch;
 
-	const bool bSkip = bUseFeaturesFromContinuityPose && SearchContext.CurrentResult.IsValid() && SearchContext.CurrentResult.Database->Schema == InOutQuery.GetSchema();
+	const bool bSkip = bUseFeaturesFromContinuingPose && SearchContext.CurrentResult.IsValid() && SearchContext.CurrentResult.Database->Schema == InOutQuery.GetSchema();
 	if (bSkip)
 	{
 		return true;
@@ -880,14 +880,14 @@ void UPoseSearchFeatureChannel_Pose::GenerateDDCKey(FBlake3& InOutKeyHasher) con
 	InOutKeyHasher.Update(&Weight, sizeof(Weight));
 	InOutKeyHasher.Update(MakeMemoryView(SampledBones));
 	InOutKeyHasher.Update(MakeMemoryView(SampleTimes));
-	InOutKeyHasher.Update(&bUseFeaturesFromContinuityPose, sizeof(bUseFeaturesFromContinuityPose));
+	InOutKeyHasher.Update(&bUseFeaturesFromContinuingPose, sizeof(bUseFeaturesFromContinuingPose));
 }
 
 bool UPoseSearchFeatureChannel_Pose::BuildQuery(UE::PoseSearch::FSearchContext& SearchContext, FPoseSearchFeatureVectorBuilder& InOutQuery) const
 {
 	using namespace UE::PoseSearch;
 
-	const bool bSkip = bUseFeaturesFromContinuityPose && SearchContext.CurrentResult.IsValid() && SearchContext.CurrentResult.Database->Schema == InOutQuery.GetSchema();
+	const bool bSkip = bUseFeaturesFromContinuingPose && SearchContext.CurrentResult.IsValid() && SearchContext.CurrentResult.Database->Schema == InOutQuery.GetSchema();
 	if (bSkip)
 	{
 		return true;
diff --git a/Engine/Plugins/Experimental/Animation/PoseSearch/Source/Runtime/Private/PoseSearchLibrary.cpp b/Engine/Plugins/Experimental/Animation/PoseSearch/Source/Runtime/Private/PoseSearchLibrary.cpp
--- a/Engine/Plugins/Experimental/Animation/PoseSearch/Source/Runtime/Private/PoseSearchLibrary.cpp
+++ b/Engine/Plugins/Experimental/Animation/PoseSearch/Source/Runtime/Private/PoseSearchLibrary.cpp
@@ -288,40 +288,28 @@ void UpdateMotionMatchingState(
 		if (SearchResult.IsValid())
 		{
 			// If the result is valid and we couldn't advance we should always jump to the search result
-			if (!bCanAdvance)
-			{
-				InOutMotionMatchingState.JumpToPose(Context, Settings, SearchResult);
-			}
+			bool bJumpToPose = true;
+
 			// Otherwise we need to check if the result is a good improvement over the current pose
-			else
+			if (bCanAdvance)
 			{
 				// Consider the search result better if it is more similar to the query than the current pose we're playing back from the database
 				check(SearchResult.PoseCost.GetDissimilarity() >= 0.0f);
-				bool bBetterPose = true;
-				if (SearchResult.ContinuityPoseCost.IsValid())
+				
+				if (SearchResult.ContinuingPoseCost.IsValid())
 				{
-					if ((SearchResult.ContinuityPoseCost.GetTotalCost() <= SearchResult.PoseCost.GetTotalCost()) || 
-						(SearchResult.ContinuityPoseCost.GetDissimilarity() <= SearchResult.PoseCost.GetDissimilarity()))
+					if ((SearchResult.ContinuingPoseCost.GetTotalCost() <= SearchResult.PoseCost.GetTotalCost()) ||
+						(SearchResult.ContinuingPoseCost.GetDissimilarity() <= SearchResult.PoseCost.GetDissimilarity()))
 					{
-						bBetterPose = false;
-					}
-					else
-					{
-						checkSlow(
-							SearchResult.ContinuityPoseCost.GetDissimilarity() > 0.0f && 
-							 SearchResult.ContinuityPoseCost.GetDissimilarity() > SearchResult.PoseCost.GetDissimilarity());
-						const float RelativeSimilarityGain = -1.0f * 
-							(SearchResult.PoseCost.GetDissimilarity() - SearchResult.ContinuityPoseCost.GetDissimilarity()) / 
-							SearchResult.ContinuityPoseCost.GetDissimilarity();
-						bBetterPose = RelativeSimilarityGain >= Settings.MinPercentImprovement / 100.0f;
+						bJumpToPose = false;
 					}
 				}
+			}
 
-				// Jump to candidate pose if there was a better option
-				if (bBetterPose)
-				{
-					InOutMotionMatchingState.JumpToPose(Context, Settings, SearchResult);
-				}
+			// Jump to candidate pose if there was a better option
+			if (bJumpToPose)
+			{
+				InOutMotionMatchingState.JumpToPose(Context, Settings, SearchResult);
 			}
 		}
 
@@ -479,12 +467,4 @@ float FMotionMatchingState::ComputeJumpBlendTime(
 	return JumpBlendTime;
 }
 
-EPoseSearchPostSearchStatus UPoseSearchPostProcessor_Bias::PostProcess_Implementation(FPoseSearchCost& InOutCost) const
-{
-	InOutCost.SetDissimilarity(Multiplier * InOutCost.GetDissimilarity());
-	InOutCost.SetCostAddend(Addend + InOutCost.GetCostAddend());
-
-	return EPoseSearchPostSearchStatus::Continue;
-}
-
 #undef LOCTEXT_NAMESPACE
\ No newline at end of file
diff --git a/Engine/Plugins/Experimental/Animation/PoseSearch/Source/Runtime/Public/PoseSearch/AnimNode_MotionMatching.h b/Engine/Plugins/Experimental/Animation/PoseSearch/Source/Runtime/Public/PoseSearch/AnimNode_MotionMatching.h
--- a/Engine/Plugins/Experimental/Animation/PoseSearch/Source/Runtime/Public/PoseSearch/AnimNode_MotionMatching.h
+++ b/Engine/Plugins/Experimental/Animation/PoseSearch/Source/Runtime/Public/PoseSearch/AnimNode_MotionMatching.h
@@ -53,7 +53,7 @@ public:
 	UPROPERTY(EditAnywhere, BlueprintReadWrite, Category = Settings, meta = (PinHiddenByDefault))
 	bool bResetOnBecomingRelevant = true;
 
-	// if true the continuity pose will be invalidated
+	// if true the continuing pose will be invalidated
 	UPROPERTY(EditAnywhere, BlueprintReadWrite, Category = Settings, meta = (PinHiddenByDefault))
 	bool bForceInterrupt = false;
 
diff --git a/Engine/Plugins/Experimental/Animation/PoseSearch/Source/Runtime/Public/PoseSearch/PoseSearch.h b/Engine/Plugins/Experimental/Animation/PoseSearch/Source/Runtime/Public/PoseSearch/PoseSearch.h
--- a/Engine/Plugins/Experimental/Animation/PoseSearch/Source/Runtime/Public/PoseSearch/PoseSearch.h
+++ b/Engine/Plugins/Experimental/Animation/PoseSearch/Source/Runtime/Public/PoseSearch/PoseSearch.h
@@ -72,7 +72,7 @@ enum class EPoseSearchFeatureDomain : int32
 	Invalid = Num UMETA(Hidden)
 };
 
-UENUM(BlueprintType)
+UENUM()
 enum class EPoseSearchBooleanRequest : uint8
 {
 	FalseValue,
@@ -125,6 +125,17 @@ enum class ESearchIndexAssetType : int32
 	BlendSpace,
 };
 
+namespace UE::PoseSearch {
+	
+enum class EPoseComparisonFlags : int32
+{
+	None = 0,
+	ContinuingPose = 1 << 0,
+};
+ENUM_CLASS_FLAGS(EPoseComparisonFlags);
+
+} // namespace UE::PoseSearch
+
 UENUM()
 enum class EPoseSearchMirrorOption : int32
 {
@@ -488,6 +499,19 @@ public:
 	UPROPERTY(Transient)
 	TArray<uint16> BoneIndicesWithParents;
 
+	// cost added to the continuing pose from databases that uses this schema
+	UPROPERTY(EditAnywhere, Category = "Schema")
+	float ContinuingPoseCostBias = 0.f;
+
+	// base cost added to all poses from databases that uses this schema. it can be overridden by UAnimNotifyState_PoseSearchModifyCost
+	UPROPERTY(EditAnywhere, Category = "Schema")
+	float BaseCostBias = 0.f;
+
+	// If there's a mirroring mismatch between the currently playing asset and a search candidate, this cost will be 
+	// added to the candidate, making it less likely to be selected
+	UPROPERTY(EditAnywhere, Category = "Schema")
+	float MirrorMismatchCostBias = 0.f;
+
 	bool IsValid () const;
 
 	int32 GetNumBones () const { return BoneIndices.Num(); }
@@ -560,8 +584,13 @@ struct POSESEARCH_API FPoseSearchPoseMetadata
 	UPROPERTY()
 	EPoseSearchPoseFlags Flags = EPoseSearchPoseFlags::None;
 
+	// @todo: consider float16
 	UPROPERTY()
 	float CostAddend = 0.0f;
+
+	// @todo: consider float16
+	UPROPERTY()
+	float ContinuingPoseCostAddend = 0.0f;
 };
 
 
@@ -986,7 +1015,7 @@ struct FSearchResult
 	FPoseSearchFeatureVectorBuilder ComposedQuery;
 
 	// cost of the current pose with the query from database in the result, if possible
-	FPoseSearchCost ContinuityPoseCost; 
+	FPoseSearchCost ContinuingPoseCost; 
 
 	float AssetTime = 0.0f;
 
@@ -1049,11 +1078,6 @@ public:
 	UPROPERTY(EditAnywhere, BlueprintReadOnly, Category="Database", DisplayName="Config")
 	TObjectPtr<const UPoseSearchSchema> Schema = nullptr;
 
-	// If there's a mirroring mismatch between the currently playing sequence and a search candidate, this cost will be 
-	// added to the candidate, making it less likely to be selected
-	UPROPERTY(EditAnywhere, BlueprintReadWrite, Category="Database")
-	float MirroringMismatchCost = 0.0f;
-
 	UPROPERTY(EditAnywhere, Category = "Database")
 	FPoseSearchExtrapolationParameters ExtrapolationParameters;
 
@@ -1179,9 +1203,9 @@ public:
 
 	void BuildQuery(UE::PoseSearch::FSearchContext& SearchContext, FPoseSearchFeatureVectorBuilder& OutQuery) const;
 
-	FPoseSearchCost ComparePoses(UE::PoseSearch::FSearchContext& SearchContext, int32 PoseIdx, int32 GroupIdx, const TArrayView<const float>& QueryValues) const;
-	FPoseSearchCost ComparePoses(UE::PoseSearch::FSearchContext& SearchContext, int32 PoseIdx, const TArrayView<const float>& QueryValues, UE::PoseSearch::FPoseCostDetails& OutPoseCostDetails) const;
-	void ComputePoseCostAddends(int32 PoseIdx, UE::PoseSearch::FSearchContext& SearchContext, float& OutNotifyAddend, float& OutMirrorMismatchAddend) const;
+	FPoseSearchCost ComparePoses(UE::PoseSearch::FSearchContext& SearchContext, int32 PoseIdx, UE::PoseSearch::EPoseComparisonFlags PoseComparisonFlags, int32 GroupIdx, const TArrayView<const float>& QueryValues) const;
+	FPoseSearchCost ComparePoses(UE::PoseSearch::FSearchContext& SearchContext, int32 PoseIdx, UE::PoseSearch::EPoseComparisonFlags PoseComparisonFlags, const TArrayView<const float>& QueryValues, UE::PoseSearch::FPoseCostDetails& OutPoseCostDetails) const;
+	void ComputePoseCostAddends(int32 PoseIdx, UE::PoseSearch::EPoseComparisonFlags PoseComparisonFlags, UE::PoseSearch::FSearchContext& SearchContext, float& OutNotifyAddend, float& OutMirrorMismatchAddend) const;
 
 protected:
 	UE::PoseSearch::FSearchResult SearchPCAKDTree(UE::PoseSearch::FSearchContext& SearchContext) const;
@@ -1217,7 +1241,7 @@ public:
 	UE::PoseSearch::FSearchResult Search(UE::PoseSearch::FSearchContext& SearchContext) const;
 
 protected:
-	FPoseSearchCost ComparePoses(int32 PoseIdx, const TArrayView<const float>& QueryValues) const;
+	FPoseSearchCost ComparePoses(int32 PoseIdx, UE::PoseSearch::EPoseComparisonFlags PoseComparisonFlags, const TArrayView<const float>& QueryValues) const;
 
 public: // UObject
 	virtual void PreSave(FObjectPreSaveContext ObjectSaveContext) override;
@@ -1406,7 +1430,7 @@ POSESEARCH_API bool BuildIndex(UPoseSearchDatabase* Database, FPoseSearchIndex&
 } // namespace UE::PoseSearch
 
 
-UENUM(BlueprintType)
+UENUM()
 enum class EPoseSearchPostSearchStatus : uint8
 {
 	// Continue looking for results 
@@ -1416,17 +1440,6 @@ enum class EPoseSearchPostSearchStatus : uint8
 	Stop
 };
 
-
-UCLASS(Abstract, EditInlineNew, Blueprintable)
-class POSESEARCH_API UPoseSearchPostProcessor : public UDataAsset
-{
-	GENERATED_BODY()
-
-public:
-	UFUNCTION(BlueprintNativeEvent, BlueprintPure, Category=Settings)
-	EPoseSearchPostSearchStatus PostProcess(FPoseSearchCost& InOutCost) const;
-};
-
 USTRUCT(BlueprintType, Category = "Animation|Pose Search")
 struct POSESEARCH_API FPoseSearchDatabaseSetEntry
 {
@@ -1439,8 +1452,8 @@ public:
 	UPROPERTY(EditAnywhere, Category = Settings)
 	FGameplayTag Tag;
 
-	UPROPERTY(EditAnywhere, Instanced, Category = Settings)
-	TObjectPtr<UPoseSearchPostProcessor> PostProcessor = nullptr;
+	UPROPERTY(EditAnywhere, Category = Settings)
+	EPoseSearchPostSearchStatus PostSearchStatus = EPoseSearchPostSearchStatus::Continue;
 };
 
 /** A data asset which holds a collection searchable assets. */
diff --git a/Engine/Plugins/Experimental/Animation/PoseSearch/Source/Runtime/Public/PoseSearch/PoseSearchAnimNotifies.h b/Engine/Plugins/Experimental/Animation/PoseSearch/Source/Runtime/Public/PoseSearch/PoseSearchAnimNotifies.h
--- a/Engine/Plugins/Experimental/Animation/PoseSearch/Source/Runtime/Public/PoseSearch/PoseSearchAnimNotifies.h
+++ b/Engine/Plugins/Experimental/Animation/PoseSearch/Source/Runtime/Public/PoseSearch/PoseSearchAnimNotifies.h
@@ -31,7 +31,8 @@ class POSESEARCH_API UAnimNotifyState_PoseSearchBlockTransition : public UAnimNo
 
 // Pose matching cost will be affected by this, making the animation segment more or less likely to be selected based
 // on the notify parameters
-UCLASS(Blueprintable, meta = (DisplayName = "Pose Matching: Modify Cost"))
+UCLASS(Blueprintable, meta = (DisplayName = "Pose Matching: Override Base Cost Bias"))
+// @todo: rename into UAnimNotifyState_PoseSearchOverrideBaseCostBias
 class POSESEARCH_API UAnimNotifyState_PoseSearchModifyCost : public UAnimNotifyState_PoseSearchBase
 {
 	GENERATED_BODY()
@@ -43,3 +44,18 @@ public:
 	UPROPERTY(EditAnywhere, Category = Config, meta = (DisplayName = "Modifier"))
 	float CostAddend = -1.0f;
 };
+
+// Pose matching cost for the continuing pose will be affected by this, making the animation segment more or less 
+// likely to be continuing playing based on the notify parameters
+UCLASS(Blueprintable, meta = (DisplayName = "Pose Matching: Override Continuing Pose Cost Bias"))
+class POSESEARCH_API UAnimNotifyState_PoseSearchOverrideContinuingPoseCostBias : public UAnimNotifyState_PoseSearchBase
+{
+	GENERATED_BODY()
+
+public:
+
+	// A negative value reduces the cost and makes the segment more likely to continuing playing. A positive value, conversely,
+	// makes the segment less likely to continuing playing
+	UPROPERTY(EditAnywhere, Category = Config, meta = (DisplayName = "Modifier"))
+	float CostAddend = -1.0f;
+};
\ No newline at end of file
diff --git a/Engine/Plugins/Experimental/Animation/PoseSearch/Source/Runtime/Public/PoseSearch/PoseSearchFeatureChannels.h b/Engine/Plugins/Experimental/Animation/PoseSearch/Source/Runtime/Public/PoseSearch/PoseSearchFeatureChannels.h
--- a/Engine/Plugins/Experimental/Animation/PoseSearch/Source/Runtime/Public/PoseSearch/PoseSearchFeatureChannels.h
+++ b/Engine/Plugins/Experimental/Animation/PoseSearch/Source/Runtime/Public/PoseSearch/PoseSearchFeatureChannels.h
@@ -28,9 +28,9 @@ public:
 	UPROPERTY()
 	int8 SchemaBoneIdx;
 
-	// reuse continuity pose database feature values if available
+	// reuse continuing pose database feature values if available
 	UPROPERTY(EditAnywhere, Category = "Settings")
-	bool bUseFeaturesFromContinuityPose = true;
+	bool bUseFeaturesFromContinuingPose = true;
 
 	// if UseSampleTimeOffsetRootBone is true, this UPoseSearchFeatureChannel_Position will calculate the position of Bone from the pose SampleTimeOffset seconds away from the current time pose root bone
 	// if false the calculated position of Bone will be in component space from the pose SampleTimeOffset seconds away
@@ -81,9 +81,9 @@ public:
 	UPROPERTY()
 	int8 SchemaBoneIdx;
 
-	// reuse continuity pose database feature values if available
+	// reuse continuing pose database feature values if available
 	UPROPERTY(EditAnywhere, Category = "Settings")
-	bool bUseFeaturesFromContinuityPose = true;
+	bool bUseFeaturesFromContinuingPose = true;
 
 	// if UseSampleTimeOffsetRootBone is true, this UPoseSearchFeatureChannel_Position will calculate the position of Bone from the pose SampleTimeOffset seconds away from the current time pose root bone
 	// if false the calculated position of Bone will be in component space from the pose SampleTimeOffset seconds away
@@ -121,9 +121,9 @@ public:
 	UPROPERTY()
 	TArray<int8> SchemaBoneIdx;
 
-	// reuse continuity pose database feature values if available
+	// reuse continuing pose database feature values if available
 	UPROPERTY(EditAnywhere, Category = "Settings")
-	bool bUseFeaturesFromContinuityPose = true;
+	bool bUseFeaturesFromContinuingPose = true;
 
 	// UObject interface
 	virtual void PreSave(FObjectPreSaveContext ObjectSaveContext) override;
diff --git a/Engine/Plugins/Experimental/Animation/PoseSearch/Source/Runtime/Public/PoseSearch/PoseSearchLibrary.h b/Engine/Plugins/Experimental/Animation/PoseSearch/Source/Runtime/Public/PoseSearch/PoseSearchLibrary.h
--- a/Engine/Plugins/Experimental/Animation/PoseSearch/Source/Runtime/Public/PoseSearch/PoseSearchLibrary.h
+++ b/Engine/Plugins/Experimental/Animation/PoseSearch/Source/Runtime/Public/PoseSearch/PoseSearchLibrary.h
@@ -42,10 +42,6 @@ struct POSESEARCH_API FMotionMatchingSettings
 	// Minimum amount of time to wait between pose search queries
 	UPROPERTY(EditAnywhere, BlueprintReadWrite, Category=Settings, meta=(ClampMin="0"))
 	float SearchThrottleTime = 0.1f;
-
-	// How much better the search result must be compared to the current pose in order to jump to it
-	UPROPERTY(EditAnywhere, BlueprintReadWrite, Category=Settings, meta=(ClampMin="0", ClampMax="100"))
-	float MinPercentImprovement = 40.0f;
 };
 
 USTRUCT(BlueprintType, Category="Animation|Pose Search")
@@ -118,18 +114,3 @@ POSESEARCH_API void UpdateMotionMatchingState(
 	bool bForceInterrupt
 );
 
-UCLASS(DisplayName = "Bias")
-class POSESEARCH_API UPoseSearchPostProcessor_Bias : public UPoseSearchPostProcessor
-{
-	GENERATED_BODY()
-
-public:
-	virtual EPoseSearchPostSearchStatus PostProcess_Implementation(FPoseSearchCost& InOutCost) const override;
-
-	UPROPERTY(EditAnywhere, Category = Parameters)
-	float Multiplier = 1.0f;
-
-	UPROPERTY(EditAnywhere, Category = Parameters)
-	float Addend = 0.0f;
-};
-
