motion matching
- query caching (within seach context)
- performing continuing pose cost evaluation first for database sets
- introduced bSkipSearchIfPossible to avoid perfroming searches in case they cannot improve search scoring
- jumping to pose based on ly on search result total cost, not accounting for dissimilarity cost

diff --git a/Engine/Plugins/Experimental/Animation/PoseSearch/Source/Runtime/Private/PoseSearch.cpp b/Engine/Plugins/Experimental/Animation/PoseSearch/Source/Runtime/Private/PoseSearch.cpp
--- a/Engine/Plugins/Experimental/Animation/PoseSearch/Source/Runtime/Private/PoseSearch.cpp
+++ b/Engine/Plugins/Experimental/Animation/PoseSearch/Source/Runtime/Private/PoseSearch.cpp
@@ -2023,12 +2023,14 @@ UE::PoseSearch::FSearchResult UPoseSearchDatabase::SearchPCAKDTree(UE::PoseSearc
 	FPoseSearchCost BestPoseCost;
 	int32 BestPoseIdx = INDEX_NONE;
 
-	BuildQuery(SearchContext, Result.ComposedQuery);
+	SearchContext.GetOrBuildQuery(this, Result.ComposedQuery);
 
 	TArrayView<const float> NormalizedQueryValues = Result.ComposedQuery.GetNormalizedValues();
 
 	const bool IsCurrentResultFromThisDatabase = !SearchContext.bForceInterrupt && SearchContext.CurrentResult.IsValid() && SearchContext.CurrentResult.Database == this;
-	if (IsCurrentResultFromThisDatabase)
+	
+	// evaluating the continuing pose only if it hasn't already being evaluated and the related animation can advance
+	if (IsCurrentResultFromThisDatabase && SearchContext.bCanAdvance && !Result.ContinuingPoseCost.IsValid())
 	{
 		Result.ContinuingPoseCost = ComparePoses(
 			SearchContext,
@@ -2036,56 +2038,71 @@ UE::PoseSearch::FSearchResult UPoseSearchDatabase::SearchPCAKDTree(UE::PoseSearc
 			EPoseComparisonFlags::ContinuingPose,
 			SearchIndex->FindAssetForPose(SearchContext.CurrentResult.PoseIdx)->SourceGroupIdx,
 			NormalizedQueryValues);
+
+		if (bSkipSearchIfPossible)
+		{
+			SearchContext.UpdateCurrentBestCost(Result.ContinuingPoseCost);
+		}
 	}
 
-	for (const FGroupSearchIndex& GroupSearchIndex : SearchIndex->Groups)
+	// since any PoseCost calculated here is at least SearchIndex->MinCostAddend,
+	// there's no point in performing the search if CurrentBestTotalCost is already better than that
+	if (SearchContext.GetCurrentBestTotalCost() > SearchIndex->MinCostAddend)
 	{
-		// we're offsetting the NonSelectableIdx by "-GroupSearchIndex.StartPoseIndex" to have the indexes in kdtree space rather than in PoseIdx database space (kdtree index 0 is PoseIdx of GroupSearchIndex.StartPoseIndex)
-		const int NonSelectableIdxUsedSize = IsCurrentResultFromThisDatabase ? PopulateNonSelectableIdx(NonSelectableIdxData, NonSelectableIdxDataSize, SearchContext, -GroupSearchIndex.StartPoseIndex) : 0;
-		TArrayView<size_t> NonSelectableIdx(NonSelectableIdxData, NonSelectableIdxUsedSize);
-		const RowMajorVectorMapConst MapWeights(GroupSearchIndex.Weights.GetData(), 1, NumDimensions);
-		FKDTree::KNNResultSet ResultSet(ClampedKDTreeQueryNumNeighbors, ResultIndexes, ResultDistanceSqr, NonSelectableIdx);
+		for (const FGroupSearchIndex& GroupSearchIndex : SearchIndex->Groups)
+		{
+			// we're offsetting the NonSelectableIdx by "-GroupSearchIndex.StartPoseIndex" to have the indexes in kdtree space rather than in PoseIdx database space (kdtree index 0 is PoseIdx of GroupSearchIndex.StartPoseIndex)
+			const int NonSelectableIdxUsedSize = IsCurrentResultFromThisDatabase ? PopulateNonSelectableIdx(NonSelectableIdxData, NonSelectableIdxDataSize, SearchContext, -GroupSearchIndex.StartPoseIndex) : 0;
+			TArrayView<size_t> NonSelectableIdx(NonSelectableIdxData, NonSelectableIdxUsedSize);
+			const RowMajorVectorMapConst MapWeights(GroupSearchIndex.Weights.GetData(), 1, NumDimensions);
+			FKDTree::KNNResultSet ResultSet(ClampedKDTreeQueryNumNeighbors, ResultIndexes, ResultDistanceSqr, NonSelectableIdx);
 
-		check(NormalizedQueryValues.Num() == NumDimensions);
+			check(NormalizedQueryValues.Num() == NumDimensions);
 
-		const RowMajorVectorMapConst Mean(GroupSearchIndex.Mean.GetData(), 1, NumDimensions);
-		const ColMajorMatrixMapConst PCAProjectionMatrix(GroupSearchIndex.PCAProjectionMatrix.GetData(), NumDimensions, ClampedNumberOfPrincipalComponents);
+			const RowMajorVectorMapConst Mean(GroupSearchIndex.Mean.GetData(), 1, NumDimensions);
+			const ColMajorMatrixMapConst PCAProjectionMatrix(GroupSearchIndex.PCAProjectionMatrix.GetData(), NumDimensions, ClampedNumberOfPrincipalComponents);
 
-		// transforming query values into PCA space to query the KDTree
-		const RowMajorVectorMapConst QueryValues(NormalizedQueryValues.GetData(), 1, NumDimensions);
-		WeightedQueryValues = QueryValues.array() * MapWeights.array();
-		CenteredQueryValues.noalias() = WeightedQueryValues - Mean;
-		ProjectedQueryValues.noalias() = CenteredQueryValues * PCAProjectionMatrix;
+			// transforming query values into PCA space to query the KDTree
+			const RowMajorVectorMapConst QueryValues(NormalizedQueryValues.GetData(), 1, NumDimensions);
+			WeightedQueryValues = QueryValues.array() * MapWeights.array();
+			CenteredQueryValues.noalias() = WeightedQueryValues - Mean;
+			ProjectedQueryValues.noalias() = CenteredQueryValues * PCAProjectionMatrix;
 
-		GroupSearchIndex.KDTree.FindNeighbors(ResultSet, ProjectedQueryValues.data());
+			GroupSearchIndex.KDTree.FindNeighbors(ResultSet, ProjectedQueryValues.data());
 
-		for (size_t ResultIndex = 0; ResultIndex < ResultSet.Num(); ++ResultIndex)
-		{
-			const int32 PoseIdx = ResultIndexes[ResultIndex] + GroupSearchIndex.StartPoseIndex;
+			for (size_t ResultIndex = 0; ResultIndex < ResultSet.Num(); ++ResultIndex)
+			{
+				const int32 PoseIdx = ResultIndexes[ResultIndex] + GroupSearchIndex.StartPoseIndex;
 
-			const FPoseSearchPoseMetadata& Metadata = SearchIndex->PoseMetadata[PoseIdx];
+				const FPoseSearchPoseMetadata& Metadata = SearchIndex->PoseMetadata[PoseIdx];
 
-			if (EnumHasAnyFlags(Metadata.Flags, EPoseSearchPoseFlags::BlockTransition))
-			{
-				continue;
-			}
+				if (EnumHasAnyFlags(Metadata.Flags, EPoseSearchPoseFlags::BlockTransition))
+				{
+					continue;
+				}
 
-			FPoseSearchCost PoseCost = ComparePoses(
-				SearchContext, 
-				PoseIdx, 
-				EPoseComparisonFlags::None,
-				GroupSearchIndex.GroupIndex,
-				NormalizedQueryValues);
+				FPoseSearchCost PoseCost = ComparePoses(
+					SearchContext,
+					PoseIdx,
+					EPoseComparisonFlags::None,
+					GroupSearchIndex.GroupIndex,
+					NormalizedQueryValues);
 
-			if (PoseCost < BestPoseCost)
-			{
-				BestPoseCost = PoseCost;
-				BestPoseIdx = PoseIdx;
+				if (PoseCost < BestPoseCost)
+				{
+					BestPoseCost = PoseCost;
+					BestPoseIdx = PoseIdx;
+				}
+
+#if UE_POSE_SEARCH_TRACE_ENABLED
+				SearchContext.BestCandidates.Add(PoseCost.GetTotalCost(), PoseIdx, this);
+#endif
 			}
+		}
 
-			#if UE_POSE_SEARCH_TRACE_ENABLED
-			SearchContext.BestCandidates.Add(PoseCost.GetTotalCost(), PoseIdx, this);
-			#endif
+		if (bSkipSearchIfPossible && BestPoseCost.IsValid())
+		{
+			SearchContext.UpdateCurrentBestCost(BestPoseCost);
 		}
 	}
 
@@ -2114,7 +2131,7 @@ UE::PoseSearch::FSearchResult UPoseSearchDatabase::SearchBruteForce(UE::PoseSear
 	const FPoseSearchIndex* SearchIndex = GetSearchIndex();
 	check(SearchIndex);
 
-	BuildQuery(SearchContext, Result.ComposedQuery);
+	SearchContext.GetOrBuildQuery(this, Result.ComposedQuery);
 	TArrayView<const float> NormalizedQueryValues = Result.ComposedQuery.GetNormalizedValues();
 
 	constexpr int NonSelectableIdxDataSize = 128;
@@ -2123,12 +2140,21 @@ UE::PoseSearch::FSearchResult UPoseSearchDatabase::SearchBruteForce(UE::PoseSear
 	const bool IsCurrentResultFromThisDatabase = !SearchContext.bForceInterrupt && SearchContext.CurrentResult.IsValid() && SearchContext.CurrentResult.Database == this;
 	if (IsCurrentResultFromThisDatabase)
 	{
-		Result.ContinuingPoseCost = ComparePoses(
-			SearchContext,
-			SearchContext.CurrentResult.PoseIdx,
-			EPoseComparisonFlags::ContinuingPose,
-			SearchIndex->FindAssetForPose(SearchContext.CurrentResult.PoseIdx)->SourceGroupIdx,
-			NormalizedQueryValues);
+		// evaluating the continuing pose only if it hasn't already being evaluated and the related animation can advance
+		if (SearchContext.bCanAdvance && !Result.ContinuingPoseCost.IsValid())
+		{
+			Result.ContinuingPoseCost = ComparePoses(
+				SearchContext,
+				SearchContext.CurrentResult.PoseIdx,
+				EPoseComparisonFlags::ContinuingPose,
+				SearchIndex->FindAssetForPose(SearchContext.CurrentResult.PoseIdx)->SourceGroupIdx,
+				NormalizedQueryValues);
+
+			if (bSkipSearchIfPossible)
+			{
+				SearchContext.UpdateCurrentBestCost(Result.ContinuingPoseCost);
+			}
+		}
 
 		NonSelectableIdxUsedSize = PopulateNonSelectableIdx(NonSelectableIdxData, NonSelectableIdxDataSize, SearchContext);
 	}
@@ -2137,48 +2163,59 @@ UE::PoseSearch::FSearchResult UPoseSearchDatabase::SearchBruteForce(UE::PoseSear
 
 	FPoseSearchCost BestPoseCost;
 	int32 BestPoseIdx = INDEX_NONE;
-	for (const FPoseSearchIndexAsset& Asset : SearchIndex->Assets)
+
+	// since any PoseCost calculated here is at least SearchIndex->MinCostAddend,
+	// there's no point in performing the search if CurrentBestTotalCost is already better than that
+	if (SearchContext.GetCurrentBestTotalCost() > SearchIndex->MinCostAddend)
 	{
-		if (SearchContext.DatabaseTagQuery)
+		for (const FPoseSearchIndexAsset& Asset : SearchIndex->Assets)
 		{
-			if (!SearchContext.DatabaseTagQuery->Matches(*GetSourceAssetGroupTags(&Asset)))
+			if (SearchContext.DatabaseTagQuery)
 			{
-				continue;
+				if (!SearchContext.DatabaseTagQuery->Matches(*GetSourceAssetGroupTags(&Asset)))
+				{
+					continue;
+				}
 			}
-		}
-
-		const bool CheckForNonSelectableIdx = IsCurrentResultFromThisDatabase && (&Asset == SearchContext.CurrentResult.SearchIndexAsset);
-		const int32 EndIndex = Asset.FirstPoseIdx + Asset.NumPoses;
-		for (int32 PoseIdx = Asset.FirstPoseIdx; PoseIdx < EndIndex; ++PoseIdx)
-		{
-			const FPoseSearchPoseMetadata& Metadata = SearchIndex->PoseMetadata[PoseIdx];
 
-			if (EnumHasAnyFlags(Metadata.Flags, EPoseSearchPoseFlags::BlockTransition))
+			const bool CheckForNonSelectableIdx = IsCurrentResultFromThisDatabase && (&Asset == SearchContext.CurrentResult.SearchIndexAsset);
+			const int32 EndIndex = Asset.FirstPoseIdx + Asset.NumPoses;
+			for (int32 PoseIdx = Asset.FirstPoseIdx; PoseIdx < EndIndex; ++PoseIdx)
 			{
-				continue;
-			}
+				const FPoseSearchPoseMetadata& Metadata = SearchIndex->PoseMetadata[PoseIdx];
 
-			if (CheckForNonSelectableIdx && NonSelectableIdx.Contains(PoseIdx))
-			{
-				continue;
-			}
+				if (EnumHasAnyFlags(Metadata.Flags, EPoseSearchPoseFlags::BlockTransition))
+				{
+					continue;
+				}
 
-			FPoseSearchCost PoseCost = ComparePoses(
-				SearchContext,
-				PoseIdx, 
-				EPoseComparisonFlags::None,
-				Asset.SourceGroupIdx,
-				NormalizedQueryValues);
+				if (CheckForNonSelectableIdx && NonSelectableIdx.Contains(PoseIdx))
+				{
+					continue;
+				}
 
-			if (PoseCost < BestPoseCost)
-			{
-				BestPoseCost = PoseCost;
-				BestPoseIdx = PoseIdx;
+				FPoseSearchCost PoseCost = ComparePoses(
+					SearchContext,
+					PoseIdx, 
+					EPoseComparisonFlags::None,
+					Asset.SourceGroupIdx,
+					NormalizedQueryValues);
+
+				if (PoseCost < BestPoseCost)
+				{
+					BestPoseCost = PoseCost;
+					BestPoseIdx = PoseIdx;
+				}
+
+				#if UE_POSE_SEARCH_TRACE_ENABLED
+				SearchContext.BestCandidates.Add(PoseCost.GetTotalCost(), PoseIdx, this);
+				#endif
 			}
+		}
 
-			#if UE_POSE_SEARCH_TRACE_ENABLED
-			SearchContext.BestCandidates.Add(PoseCost.GetTotalCost(), PoseIdx, this);
-			#endif
+		if (bSkipSearchIfPossible && BestPoseCost.IsValid())
+		{
+			SearchContext.UpdateCurrentBestCost(BestPoseCost);
 		}
 	}
 
@@ -2208,6 +2245,37 @@ UE::PoseSearch::FSearchResult UPoseSearchDatabaseSet::Search(UE::PoseSearch::FSe
 	FSearchResult Result;
 	FPoseSearchCost ContinuingCost;
 
+	// evaluating the continuing pose before all the active entries
+	if (bEvaluateContinuingPoseFirst && 
+		SearchContext.bCanAdvance &&
+		!SearchContext.bForceInterrupt &&
+		SearchContext.CurrentResult.IsValid())
+	{
+		const UPoseSearchDatabase* Database = SearchContext.CurrentResult.Database.Get();
+		check(Database);
+		const FPoseSearchIndex* PoseSearchIndex = Database->GetSearchIndex();
+		check(PoseSearchIndex);
+
+		SearchContext.GetOrBuildQuery(Database, Result.ComposedQuery);
+
+		TArrayView<const float> NormalizedQueryValues = Result.ComposedQuery.GetNormalizedValues();
+
+		const FPoseSearchIndexAsset* PoseSearchIndexAsset = PoseSearchIndex->FindAssetForPose(SearchContext.CurrentResult.PoseIdx);
+		check(PoseSearchIndexAsset);
+
+		Result.ContinuingPoseCost = Database->ComparePoses(
+				SearchContext,
+				SearchContext.CurrentResult.PoseIdx,
+				EPoseComparisonFlags::ContinuingPose,
+				PoseSearchIndexAsset->SourceGroupIdx,
+				NormalizedQueryValues);
+
+		if (Database->bSkipSearchIfPossible)
+		{
+			SearchContext.UpdateCurrentBestCost(Result.ContinuingPoseCost);
+		}
+	}
+
 	for (const FPoseSearchDatabaseSetEntry& Entry : AssetsToSearch)
 	{
 		if (!IsValid(Entry.Searchable))
@@ -2331,7 +2399,7 @@ FTransform FSearchContext::TryGetTransformAndCacheResults(float SampleTime, cons
 	const FBoneIndexType BoneIndexType = SchemaBoneIdx >= 0 ? Schema->BoneIndices[SchemaBoneIdx] : RootBoneIdx;
 
 	// @todo: use an hashmap if we end up having too many entries
-	const CachedEntry* Entry = CachedEntries.FindByPredicate([SampleTime, BoneIndexType](const FSearchContext::CachedEntry& Entry)
+	const FCachedEntry* Entry = CachedEntries.FindByPredicate([SampleTime, BoneIndexType](const FSearchContext::FCachedEntry& Entry)
 	{
 		return Entry.SampleTime == SampleTime && Entry.BoneIndexType == BoneIndexType;
 	});
@@ -2354,7 +2422,7 @@ FTransform FSearchContext::TryGetTransformAndCacheResults(float SampleTime, cons
 			for (const FBoneIndexType NewEntryBoneIndexType : Schema->BoneIndicesWithParents)
 			{
 				// @todo: maybe add them with a single allocation with CachedEntries.AddDefaulted(Schema->BoneIndicesWithParents.Num())
-				CachedEntry& NewEntry = CachedEntries[CachedEntries.AddDefaulted()];
+				FCachedEntry& NewEntry = CachedEntries[CachedEntries.AddDefaulted()];
 				NewEntry.SampleTime = SampleTime;
 				NewEntry.BoneIndexType = NewEntryBoneIndexType;
 				NewEntry.Transform = SampledComponentPose[NewEntryBoneIndexType];
@@ -2371,7 +2439,7 @@ FTransform FSearchContext::TryGetTransformAndCacheResults(float SampleTime, cons
 	FTransform SampledRootTransform;
 	if (History->TrySampleLocalPose(-SampleTime, nullptr, nullptr, &SampledRootTransform))
 	{
-		CachedEntry& NewEntry = CachedEntries[CachedEntries.AddDefaulted()];
+		FCachedEntry& NewEntry = CachedEntries[CachedEntries.AddDefaulted()];
 		NewEntry.SampleTime = SampleTime;
 		NewEntry.BoneIndexType = BoneIndexType;
 		NewEntry.Transform = SampledRootTransform;
@@ -2389,6 +2457,41 @@ void FSearchContext::ClearCachedEntries()
 	CachedEntries.Reset();
 }
 
+void FSearchContext::ResetCurrentBestCost()
+{
+	CurrentBestTotalCost = MAX_flt;
+}
+
+void FSearchContext::UpdateCurrentBestCost(const FPoseSearchCost& PoseSearchCost)
+{
+	check(PoseSearchCost.IsValid());
+
+	if (PoseSearchCost.GetTotalCost() < CurrentBestTotalCost)
+	{
+		CurrentBestTotalCost = PoseSearchCost.GetTotalCost();
+	};
+}
+
+bool FSearchContext::GetOrBuildQuery(const UPoseSearchDatabase* Database, FPoseSearchFeatureVectorBuilder& FeatureVectorBuilder)
+{
+	const FSearchContext::FCachedQuery* CachedQuery = CachedQueries.FindByPredicate([Database](const FSearchContext::FCachedQuery& CachedQuery)
+	{
+		return CachedQuery.Database == Database;
+	});
+
+	if (CachedQuery)
+	{
+		FeatureVectorBuilder = CachedQuery->FeatureVectorBuilder;
+		return true;
+	}
+
+	FSearchContext::FCachedQuery& NewCachedQuery = CachedQueries[CachedQueries.AddDefaulted()];
+	NewCachedQuery.Database = Database;
+	Database->BuildQuery(*this, NewCachedQuery.FeatureVectorBuilder);
+	FeatureVectorBuilder = NewCachedQuery.FeatureVectorBuilder;
+	return false;
+}
+
 //////////////////////////////////////////////////////////////////////////
 // FPoseHistory
 
@@ -4940,6 +5043,7 @@ struct FDatabaseIndexingContext
 	void PrepareIndexers();
 	bool IndexAssets();
 	void JoinIndex();
+	float CalculateMinCostAddend() const;
 };
 
 void FDatabaseIndexingContext::PrepareSamplers()
@@ -5105,12 +5209,32 @@ bool FDatabaseIndexingContext::IndexAssets()
 	return true;
 }
 
+float FDatabaseIndexingContext::CalculateMinCostAddend() const
+{
+	float MinCostAddend = 0.f;
+
+	check(SearchIndex);
+	if (!SearchIndex->PoseMetadata.IsEmpty())
+	{
+		MinCostAddend = MAX_FLT;
+		for (const FPoseSearchPoseMetadata& PoseMetadata : SearchIndex->PoseMetadata)
+		{
+			if (PoseMetadata.CostAddend < MinCostAddend)
+			{
+				MinCostAddend = PoseMetadata.CostAddend;
+			}
+		}
+	}
+	return MinCostAddend;
+}
+
 void FDatabaseIndexingContext::JoinIndex()
 {
 	// Write index info to asset and count up total poses and storage required
 	int32 TotalPoses = 0;
 	int32 TotalFloats = 0;
-	
+
+	check(SearchIndex);
 	SearchIndex->Groups.Reset();
 
 	if (SearchIndex->Assets.Num() > 0)
@@ -5163,6 +5287,7 @@ void FDatabaseIndexingContext::JoinIndex()
 
 	SearchIndex->NumPoses = TotalPoses;
 	SearchIndex->Schema = Database->Schema;
+	SearchIndex->MinCostAddend = CalculateMinCostAddend();
 }
 
 bool BuildIndex(UPoseSearchDatabase* Database, FPoseSearchIndex& OutSearchIndex)
diff --git a/Engine/Plugins/Experimental/Animation/PoseSearch/Source/Runtime/Private/PoseSearchLibrary.cpp b/Engine/Plugins/Experimental/Animation/PoseSearch/Source/Runtime/Private/PoseSearchLibrary.cpp
--- a/Engine/Plugins/Experimental/Animation/PoseSearch/Source/Runtime/Private/PoseSearchLibrary.cpp
+++ b/Engine/Plugins/Experimental/Animation/PoseSearch/Source/Runtime/Private/PoseSearchLibrary.cpp
@@ -246,7 +246,7 @@ void UpdateMotionMatchingState(
 	// Check if we can advance. Includes the case where we can advance but only by switching to a follow up asset.
 	bool bAdvanceToFollowUpAsset = false;
 	FSearchResult FollowUpAsset;
-	bool bCanAdvance = InOutMotionMatchingState.CanAdvance(DeltaTime, bAdvanceToFollowUpAsset, FollowUpAsset);
+	const bool bCanAdvance = InOutMotionMatchingState.CanAdvance(DeltaTime, bAdvanceToFollowUpAsset, FollowUpAsset);
 
 	// If we can't advance or enough time has elapsed since the last pose jump then search
 	FSearchContext SearchContext;
@@ -259,6 +259,7 @@ void UpdateMotionMatchingState(
 		SearchContext.OwningComponent = Context.AnimInstanceProxy->GetSkelMeshComponent();
 		SearchContext.BoneContainer = &Context.AnimInstanceProxy->GetRequiredBones();
 		SearchContext.bForceInterrupt = bForceInterrupt;
+		SearchContext.bCanAdvance = bCanAdvance;
 
 #if WITH_EDITORONLY_DATA
 		SearchContext.DebugDrawParams.SearchCostHistoryBruteForce = &InOutMotionMatchingState.SearchCostHistoryBruteForce;
@@ -284,33 +285,12 @@ void UpdateMotionMatchingState(
 
 		// Search the database for the nearest match to the updated query vector
 		FSearchResult SearchResult = Searchable->Search(SearchContext);
+		const float ContinuingPoseCost = bCanAdvance && SearchResult.ContinuingPoseCost.IsValid() ? SearchResult.ContinuingPoseCost.GetTotalCost() : MAX_flt;
+		const float PoseCost = SearchResult.PoseCost.IsValid()  ? SearchResult.PoseCost.GetTotalCost() : MAX_flt;
 
-		if (SearchResult.IsValid())
+		if (PoseCost < ContinuingPoseCost)
 		{
-			// If the result is valid and we couldn't advance we should always jump to the search result
-			bool bJumpToPose = true;
-
-			// Otherwise we need to check if the result is a good improvement over the current pose
-			if (bCanAdvance)
-			{
-				// Consider the search result better if it is more similar to the query than the current pose we're playing back from the database
-				check(SearchResult.PoseCost.GetDissimilarity() >= 0.0f);
-				
-				if (SearchResult.ContinuingPoseCost.IsValid())
-				{
-					if ((SearchResult.ContinuingPoseCost.GetTotalCost() <= SearchResult.PoseCost.GetTotalCost()) ||
-						(SearchResult.ContinuingPoseCost.GetDissimilarity() <= SearchResult.PoseCost.GetDissimilarity()))
-					{
-						bJumpToPose = false;
-					}
-				}
-			}
-
-			// Jump to candidate pose if there was a better option
-			if (bJumpToPose)
-			{
-				InOutMotionMatchingState.JumpToPose(Context, Settings, SearchResult);
-			}
+			InOutMotionMatchingState.JumpToPose(Context, Settings, SearchResult);
 		}
 
 		// todo: cache the queries into TraceState, one per TraceState.DatabaseEntries or
@@ -340,7 +320,7 @@ void UpdateMotionMatchingState(
 	TArray<const UPoseSearchDatabase*> Databases;
 	while (!SearchContext.BestCandidates.IsEmpty())
 	{
-		FSearchContext::PoseCandidate PoseCandidate;
+		FSearchContext::FPoseCandidate PoseCandidate;
 		SearchContext.BestCandidates.Pop(PoseCandidate);
 
 		int32 DatabaseEntryIndex = 0;
diff --git a/Engine/Plugins/Experimental/Animation/PoseSearch/Source/Runtime/Public/PoseSearch/PoseSearch.h b/Engine/Plugins/Experimental/Animation/PoseSearch/Source/Runtime/Public/PoseSearch/PoseSearch.h
--- a/Engine/Plugins/Experimental/Animation/PoseSearch/Source/Runtime/Public/PoseSearch/PoseSearch.h
+++ b/Engine/Plugins/Experimental/Animation/PoseSearch/Source/Runtime/Public/PoseSearch/PoseSearch.h
@@ -713,6 +713,10 @@ struct POSESEARCH_API FPoseSearchIndex
 	UPROPERTY()
 	TArray<FPoseSearchIndexAsset> Assets;
 
+	// minimum of the database metadata CostAddend: it represents the minimum cost of any search for the associated database (we'll skip the search in case the search result total cost is already less than MinCostAddend)
+	UPROPERTY()
+	float MinCostAddend = -MAX_FLT;
+
 	bool IsValid() const;
 	bool IsEmpty() const;
 
@@ -1113,6 +1117,10 @@ public:
 	UPROPERTY(EditAnywhere, Category = "Performance")
 	EPoseSearchMode PoseSearchMode = EPoseSearchMode::BruteForce;
 
+	// if true, this database search will be skipped if cannot decrease the pose cost, and poses will not be listed into the PoseSearchDebugger
+	UPROPERTY(EditAnywhere, Category = "Performance")
+	bool bSkipSearchIfPossible = false;
+
 	FPoseSearchIndex* GetSearchIndex();
 	const FPoseSearchIndex* GetSearchIndex() const;
 
@@ -1328,14 +1336,22 @@ struct POSESEARCH_API FSearchContext
 	const FGameplayTagContainer* ActiveTagsContainer = nullptr;
 	float PoseJumpThresholdTime = 0.f;
 	bool bForceInterrupt = false;
+	// can the continuing pose advance? (if not we skip evaluating it)
+	bool bCanAdvance = true;
 
 	FTransform TryGetTransformAndCacheResults(float SampleTime, const UPoseSearchSchema* Schema, int8 SchemaBoneIdx, bool& Error);
 	void ClearCachedEntries();
 
+	void ResetCurrentBestCost();
+	void UpdateCurrentBestCost(const FPoseSearchCost& PoseSearchCost);
+	float GetCurrentBestTotalCost() const { return CurrentBestTotalCost; }
+
+	bool GetOrBuildQuery(const UPoseSearchDatabase* Database, FPoseSearchFeatureVectorBuilder& FeatureVectorBuilder);
+
 	static constexpr int8 SchemaRootBoneIdx = -1;
 
 private:
-	struct CachedEntry
+	struct FCachedEntry
 	{
 		float SampleTime = 0.f;
 
@@ -1347,23 +1363,33 @@ private:
 	};
 
 	// @todo: make it a fixed size array (or hash map if we end up having many CachedEntry) to avoid allocations
-	TArray<CachedEntry> CachedEntries;
+	TArray<FCachedEntry> CachedEntries;
+	
+	struct FCachedQuery
+	{
+		const UPoseSearchDatabase* Database = nullptr;
+		FPoseSearchFeatureVectorBuilder FeatureVectorBuilder;
+	};
+
+	TArray<FCachedQuery> CachedQueries;
+
+	float CurrentBestTotalCost = MAX_flt;
 
 #if UE_POSE_SEARCH_TRACE_ENABLED
 
 public:
-	struct PoseCandidate
+	struct FPoseCandidate
 	{
 		float Cost = 0.f;
 		int32 PoseIdx = 0;
 		const UPoseSearchDatabase* Database = nullptr;
 
-		bool operator<(const PoseCandidate& Other) const { return Cost > Other.Cost; }
+		bool operator<(const FPoseCandidate& Other) const { return Cost > Other.Cost; }
 	};
 
-	struct BestPoseCandidates : private TArray<PoseCandidate>
+	struct FBestPoseCandidates : private TArray<FPoseCandidate>
 	{
-		typedef TArray<PoseCandidate> Super;
+		typedef TArray<FPoseCandidate> Super;
 		using Super::IsEmpty;
 
 		int32 MaxPoseCandidates = 100;
@@ -1378,7 +1404,7 @@ public:
 					Pop(Unused);
 				}
 
-				FSearchContext::PoseCandidate PoseCandidate;
+				FSearchContext::FPoseCandidate PoseCandidate;
 				PoseCandidate.Cost = Cost;
 				PoseCandidate.PoseIdx = PoseIdx;
 				PoseCandidate.Database = Database;
@@ -1386,13 +1412,13 @@ public:
 			}
 		}
 
-		void Pop(PoseCandidate& OutItem)
+		void Pop(FPoseCandidate& OutItem)
 		{
 			HeapPop(OutItem, false);
 		}
 	};
 	
-	BestPoseCandidates BestCandidates;
+	FBestPoseCandidates BestCandidates;
 #endif
 };
 
@@ -1466,6 +1492,11 @@ public:
 	UPROPERTY(EditAnywhere, Category = Settings)
 	TArray<FPoseSearchDatabaseSetEntry> AssetsToSearch;
 
+	// if there's a valid continuing pose and bEvaluateContinuingPoseFirst is true, the continuing pose will be evaluated as first search,
+	// otherwise it'll be evaluated with the related database: if the database is not active the continuing pose evaluation will be skipped
+	UPROPERTY(EditAnywhere, Category = Settings)
+	bool bEvaluateContinuingPoseFirst = true;
+
 public:
 	virtual UE::PoseSearch::FSearchResult Search(UE::PoseSearch::FSearchContext& SearchContext) const override;
 };
