motion matching - phase calculation fix

diff --git a/Engine/Plugins/Experimental/Animation/PoseSearch/Source/Runtime/Private/PoseSearchFeatureChannels.cpp b/Engine/Plugins/Experimental/Animation/PoseSearch/Source/Runtime/Private/PoseSearchFeatureChannels.cpp
--- a/Engine/Plugins/Experimental/Animation/PoseSearch/Source/Runtime/Private/PoseSearchFeatureChannels.cpp
+++ b/Engine/Plugins/Experimental/Animation/PoseSearch/Source/Runtime/Private/PoseSearchFeatureChannels.cpp
@@ -243,11 +243,11 @@ static void ValidateLocalMinMax(const TArray<LocalMinMax>& MinMax)
 		check(MinMax[i].Index > MinMax[i - 1].Index);
 		if (MinMax[i].Type == LocalMinMax::Min)
 		{
-			check(MinMax[i].SignalValue < MinMax[i - 1].SignalValue);
+			check(MinMax[i].SignalValue <= MinMax[i - 1].SignalValue);
 		}
 		else
 		{
-			check(MinMax[i].SignalValue > MinMax[i - 1].SignalValue);
+			check(MinMax[i].SignalValue >= MinMax[i - 1].SignalValue);
 		}
 	}
 }
@@ -730,7 +730,7 @@ void UPoseSearchFeatureChannel_Pose::IndexAsset(UE::PoseSearch::IAssetIndexer& I
 	const FAssetIndexingContext& IndexingContext = Indexer.GetIndexingContext();
 	for (int32 SampleIdx = IndexingContext.BeginSampleIdx; SampleIdx != IndexingContext.EndSampleIdx; ++SampleIdx)
 	{
-		int32 VectorIdx = SampleIdx - IndexingContext.BeginSampleIdx;
+		const int32 VectorIdx = SampleIdx - IndexingContext.BeginSampleIdx;
 		FPoseSearchFeatureVectorBuilder& FeatureVector = IndexingOutput.PoseVectors[VectorIdx];
 		AddPoseFeatures(Indexer, SampleIdx, FeatureVector, Phases);
 	}
@@ -869,7 +869,8 @@ void UPoseSearchFeatureChannel_Pose::AddPoseFeatures(UE::PoseSearch::IAssetIndex
 			for (int32 SubsampleIdx = 0; SubsampleIdx != SampleTimes.Num(); ++SubsampleIdx)
 			{
 				// @todo: support for SubsampleIdx
-				FFeatureVectorHelper::EncodeVector2D(FeatureVector.EditValues(), DataOffset, Phases[ChannelBoneIdx][SampleIdx]);
+				const int32 VectorIdx = SampleIdx - IndexingContext.BeginSampleIdx;
+				FFeatureVectorHelper::EncodeVector2D(FeatureVector.EditValues(), DataOffset, Phases[ChannelBoneIdx][VectorIdx]);
 			}
 		}
 	}
