motion matching - WIP embedding a blend stack into FAnimNode_MotionMatching

diff --git a/Engine/Plugins/Experimental/Animation/PoseSearch/Source/Runtime/Private/AnimNode_BlendStack.cpp b/Engine/Plugins/Experimental/Animation/PoseSearch/Source/Runtime/Private/AnimNode_BlendStack.cpp
new file mode 100644
--- /dev/null
+++ b/Engine/Plugins/Experimental/Animation/PoseSearch/Source/Runtime/Private/AnimNode_BlendStack.cpp
@@ -0,0 +1,232 @@
+// Copyright Epic Games, Inc. All Rights Reserved.
+
+#include "PoseSearch/AnimNode_BlendStack.h"
+
+#include "Animation/AnimSequence.h"
+#include "Animation/BlendSpace.h"
+
+#define LOCTEXT_NAMESPACE "AnimNode_BlendStack"
+
+/////////////////////////////////////////////////////
+// FPoseSearchAnimPlayer
+void FPoseSearchAnimPlayer::Initialize(ESearchIndexAssetType InAssetType, UAnimationAsset* AnimationAsset, float AccumulatedTime, bool bLoop, bool bMirrored, UMirrorDataTable* MirrorDataTable, float BlendTime, FVector BlendParameters)
+{
+	check(AnimationAsset);
+	check(MirrorDataTable);
+
+	AssetType = InAssetType;
+	TotalBlendInTime = BlendTime;
+	CurrentBlendInTime = 0.f;
+	BlendWeight = 0.f;
+
+	MirrorNode.SetMirrorDataTable(MirrorDataTable);
+	MirrorNode.SetMirror(bMirrored);
+	
+	if (AssetType == ESearchIndexAssetType::Sequence)
+	{
+		UAnimSequenceBase* Sequence = Cast<UAnimSequenceBase>(AnimationAsset);
+		check(Sequence);
+
+		SequencePlayerNode.SetAccumulatedTime(AccumulatedTime);
+		SequencePlayerNode.SetSequence(Sequence);
+		SequencePlayerNode.SetLoopAnimation(bLoop);
+		SequencePlayerNode.SetPlayRate(1.0f);
+	}
+	else if (AssetType == ESearchIndexAssetType::BlendSpace)
+	{
+		UBlendSpace* BlendSpace = Cast<UBlendSpace>(AnimationAsset);
+		check(BlendSpace);
+		BlendSpacePlayerNode.SetResetPlayTimeWhenBlendSpaceChanges(false /*!bReset*/);
+		BlendSpacePlayerNode.SetAccumulatedTime(AccumulatedTime);
+		BlendSpacePlayerNode.SetBlendSpace(BlendSpace);
+		BlendSpacePlayerNode.SetLoop(bLoop);
+		BlendSpacePlayerNode.SetPlayRate(1.0f);
+		BlendSpacePlayerNode.SetPosition(BlendParameters);
+	}
+	else
+	{
+		checkNoEntry();
+	}
+
+	UpdateSourceLinkNode();
+}
+
+// @todo: maybe implement copy/move constructors and assignement operator do so (or use a list instead of an array)
+// since we're making copies and moving this object in memory, we're using this method to set the MirrorNode SourceLinkNode when necessary
+void FPoseSearchAnimPlayer::UpdateSourceLinkNode()
+{
+	if (AssetType == ESearchIndexAssetType::Sequence)
+	{
+		MirrorNode.SetSourceLinkNode(&SequencePlayerNode);
+	}
+	else if (AssetType == ESearchIndexAssetType::BlendSpace)
+	{
+		MirrorNode.SetSourceLinkNode(&BlendSpacePlayerNode);
+	}
+	else
+	{
+		checkNoEntry();
+	}
+}
+
+void FPoseSearchAnimPlayer::Evaluate_AnyThread(FPoseContext& Output)
+{
+	UpdateSourceLinkNode();
+	MirrorNode.Evaluate_AnyThread(Output);
+}
+
+void FPoseSearchAnimPlayer::Update_AnyThread(const FAnimationUpdateContext& Context)
+{
+	UpdateSourceLinkNode();
+	MirrorNode.Update_AnyThread(Context);
+	CurrentBlendInTime += Context.GetDeltaTime();
+}
+
+float FPoseSearchAnimPlayer::GetAccumulatedTime() const
+{
+	if (AssetType == ESearchIndexAssetType::Sequence)
+	{
+		return SequencePlayerNode.GetAccumulatedTime();
+	}
+	
+	if (AssetType == ESearchIndexAssetType::BlendSpace)
+	{
+		return BlendSpacePlayerNode.GetAccumulatedTime();
+	}
+	
+	checkNoEntry();
+	return 0.f;
+}
+
+float FPoseSearchAnimPlayer::GetBlendInPercentage() const
+{
+	if (FMath::IsNearlyZero(TotalBlendInTime))
+	{
+		return 1.f;
+	}
+
+	return FMath::Clamp(CurrentBlendInTime / TotalBlendInTime, 0.f, 1.f);
+}
+
+void FPoseSearchAnimPlayer::SetBlendWeight(float InBlendWeight)
+{
+	BlendWeight = InBlendWeight;
+}
+
+/////////////////////////////////////////////////////
+// FAnimNode_BlendStack
+void FAnimNode_BlendStack::Evaluate_AnyThread(FPoseContext& Output)
+{
+	Super::Evaluate_AnyThread(Output);
+
+	DECLARE_SCOPE_HIERARCHICAL_COUNTER_ANIMNODE(Evaluate_AnyThread);
+
+	const int32 BlendStackSize = AnimPlayers.Num();
+	if (BlendStackSize <= 0)
+	{
+		// nothing to do
+	}
+	else if (BlendStackSize == 1)
+	{
+		AnimPlayers.First().Evaluate_AnyThread(Output);
+	}
+	else
+	{
+		TArray<FPoseContext, TMemStackAllocator<>> PoseContexts;
+		PoseContexts.Reserve(BlendStackSize);
+		for (int32 i = 0; i < BlendStackSize; ++i)
+		{
+			PoseContexts.Add(Output);
+			AnimPlayers[i].Evaluate_AnyThread(PoseContexts[i]);
+		}
+
+		// @todo: optimize copies and allocations: is there a FAnimationRuntime::BlendPosesTogether(PoseContexts, Weights, AnimationPoseData)?
+		TArray<FCompactPose, TMemStackAllocator<>> Poses;
+		TArray<FBlendedCurve, TMemStackAllocator<>> Curves;
+		TArray<UE::Anim::FStackAttributeContainer, TMemStackAllocator<>> Attributes;
+		TArray<float, TMemStackAllocator<>> Weights;
+
+		Poses.AddDefaulted(BlendStackSize);
+		Curves.AddDefaulted(BlendStackSize);
+		Attributes.AddDefaulted(BlendStackSize);
+		Weights.AddDefaulted(BlendStackSize);
+
+		for (int32 i = 0; i < BlendStackSize; ++i)
+		{
+			Poses[i] = PoseContexts[i].Pose;
+			Curves[i] = PoseContexts[i].Curve;
+			Attributes[i] = PoseContexts[i].CustomAttributes;
+			Weights[i] = AnimPlayers[i].GetBlendWeight();
+		}
+
+		FAnimationPoseData AnimationPoseData(Output);
+		FAnimationRuntime::BlendPosesTogether(Poses, Curves, Attributes, Weights, AnimationPoseData);
+	}
+}
+
+void FAnimNode_BlendStack::UpdateAssetPlayer(const FAnimationUpdateContext& Context)
+{
+	Super::UpdateAssetPlayer(Context);
+
+	const int32 BlendStackSize = AnimPlayers.Num();
+	for (FPoseSearchAnimPlayer& AnimPlayer : AnimPlayers)
+	{
+		AnimPlayer.Update_AnyThread(Context);
+	}
+	
+	CalculateWeights();
+	PruneBlendStack(BlendStackSize);
+}
+
+float FAnimNode_BlendStack::GetAccumulatedTime() const
+{
+	return AnimPlayers.IsEmpty() ? 0.f : AnimPlayers.First().GetAccumulatedTime();
+}
+
+void FAnimNode_BlendStack::BlendTo(ESearchIndexAssetType AssetType, UAnimationAsset* AnimationAsset, float AccumulatedTime, bool bLoop, bool bMirrored, UMirrorDataTable* MirrorDataTable, int32 MaxActiveBlends, float BlendTime, FVector BlendParameters)
+{
+	AnimPlayers.PushFirst(FPoseSearchAnimPlayer());
+	AnimPlayers.First().Initialize(AssetType, AnimationAsset, AccumulatedTime, bLoop, bMirrored, MirrorDataTable, BlendTime, BlendParameters);
+
+	CalculateWeights();
+	PruneBlendStack(MaxActiveBlends);
+}
+
+void FAnimNode_BlendStack::CalculateWeights()
+{
+	// AnimPlayers[0] is the most newly inserted AnimPlayer, AnimPlayers[AnimPlayers.Num()-1] is the oldest, so to calculate the weights
+	// we ask AnimPlayers[0] its BlendInPercentage and then distribute the left over (CurrentWeightMultiplier) to the rest of the AnimPlayers
+	float CurrentWeightMultiplier = 1.f;
+	const int32 BlendStackSize = AnimPlayers.Num();
+	for (int32 i = 0; i < BlendStackSize; ++i)
+	{
+		const bool bIsLastAnimPlayers = i == BlendStackSize - 1;
+		const float BlendInPercentage = bIsLastAnimPlayers ? 1.f : AnimPlayers[i].GetBlendInPercentage();
+		AnimPlayers[i].SetBlendWeight(CurrentWeightMultiplier * BlendInPercentage);
+
+		CurrentWeightMultiplier *= (1.f - BlendInPercentage);
+	}
+}
+
+void FAnimNode_BlendStack::PruneBlendStack(int32 MaxActiveBlends)
+{
+	int32 FirstZeroWeightAnimPlayerIndex = 1;
+	for (; FirstZeroWeightAnimPlayerIndex < AnimPlayers.Num(); ++FirstZeroWeightAnimPlayerIndex)
+	{
+		if (FMath::IsNearlyZero(AnimPlayers[FirstZeroWeightAnimPlayerIndex].GetBlendWeight()))
+		{
+			break;
+		}
+	}
+	// if the weight of the FirstZeroWeightAnimPlayerIndex-th AnimPlayer is zero, all the subsequent AnimPlayers contribution will be zero, so we delete from FirstZeroWeightAnimPlayerIndex forward
+
+	const int32 WantedAnimPlayersNum = FMath::Max(1, FMath::Min(FirstZeroWeightAnimPlayerIndex, MaxActiveBlends)); // we save at least one FPoseSearchAnimPlayer
+
+	while (AnimPlayers.Num() > WantedAnimPlayersNum)
+	{
+		// @todo: instead of just popping the unwanted AnimPlayers we should store the blended pose + weight in a FPoseContext
+		AnimPlayers.PopLast();
+	}
+}
+
+#undef LOCTEXT_NAMESPACE
\ No newline at end of file
diff --git a/Engine/Plugins/Experimental/Animation/PoseSearch/Source/Runtime/Private/AnimNode_MotionMatching.cpp b/Engine/Plugins/Experimental/Animation/PoseSearch/Source/Runtime/Private/AnimNode_MotionMatching.cpp
--- a/Engine/Plugins/Experimental/Animation/PoseSearch/Source/Runtime/Private/AnimNode_MotionMatching.cpp
+++ b/Engine/Plugins/Experimental/Animation/PoseSearch/Source/Runtime/Private/AnimNode_MotionMatching.cpp
@@ -10,7 +10,6 @@
 #include "Animation/AnimRootMotionProvider.h"
 #include "DynamicPlayRate/DynamicPlayRateLibrary.h"
 #include "PoseSearch/AnimNode_PoseSearchHistoryCollector.h"
-#include "PoseSearch/PoseSearch.h"
 #include "Trace/PoseSearchTraceLogger.h"
 
 #define LOCTEXT_NAMESPACE "AnimNode_MotionMatching"
@@ -24,13 +23,9 @@ void FAnimNode_MotionMatching::Initialize_AnyThread(const FAnimationInitializeCo
 
 	GetEvaluateGraphExposedInputs().Execute(Context);
 
-	CurrentAssetPlayerNode = &SequencePlayerNode;
+	BlendStackNode.Initialize_AnyThread(Context);
 
-	MirrorNode.SetSourceLinkNode(CurrentAssetPlayerNode);
-
-	BlendSpacePlayerNode.SetResetPlayTimeWhenBlendSpaceChanges(false /*!bReset*/);
-
-	Source.SetLinkNode(&MirrorNode);
+	Source.SetLinkNode(&BlendStackNode);
 	Source.Initialize(Context);
 }
 
@@ -80,7 +75,7 @@ void FAnimNode_MotionMatching::UpdateAssetPlayer(const FAnimationUpdateContext&
 		// We adjust the motion matching state asset time to the current player node's asset time. This is done 
 		// because the player node may have ticked more or less time than we expected due to variable dt or the 
 		// dynamic playback rate adjustment and as such the motion matching state does not update by itself
-		MotionMatchingState.AdjustAssetTime(CurrentAssetPlayerNode->GetAccumulatedTime());
+		MotionMatchingState.AdjustAssetTime(BlendStackNode.GetAccumulatedTime());
 	}
 	UpdateCounter.SynchronizeWith(Context.AnimInstanceProxy->GetUpdateCounter());
 
@@ -96,64 +91,34 @@ void FAnimNode_MotionMatching::UpdateAssetPlayer(const FAnimationUpdateContext&
 		bForceInterrupt
 	);
 
-
-	if (MotionMatchingState.CurrentSearchResult.Database.IsValid() && 
-		MotionMatchingState.CurrentSearchResult.Database->Schema)
-	{
-		MirrorNode.SetMirrorDataTable(MotionMatchingState.CurrentSearchResult.Database->Schema->MirrorDataTable.Get());
-	}
-
+	// If a new pose is requested, blend into the new asset via BlendStackNode
+	const UPoseSearchDatabase* Database = MotionMatchingState.CurrentSearchResult.Database.Get();
 	const FPoseSearchIndexAsset* SearchIndexAsset = MotionMatchingState.GetCurrentSearchIndexAsset();
 
-	// If a new pose is requested, jump to the pose by updating the embedded sequence player node
-	if (SearchIndexAsset && (MotionMatchingState.Flags & EMotionMatchingFlags::JumpedToPose) == EMotionMatchingFlags::JumpedToPose)
+	if (Database && Database->Schema && SearchIndexAsset && (MotionMatchingState.Flags & EMotionMatchingFlags::JumpedToPose) == EMotionMatchingFlags::JumpedToPose)
 	{
-		const UPoseSearchDatabase* Database = MotionMatchingState.CurrentSearchResult.Database.Get();
-
-		if (SearchIndexAsset->Type == ESearchIndexAssetType::Sequence)
-		{
-			CurrentAssetPlayerNode = &SequencePlayerNode;
-
-			const FPoseSearchDatabaseSequence& ResultDbSequence = Database->GetSequenceSourceAsset(SearchIndexAsset);
-			SequencePlayerNode.SetAccumulatedTime(MotionMatchingState.AssetPlayerTime);
-			SequencePlayerNode.SetSequence(ResultDbSequence.Sequence);
-			SequencePlayerNode.SetLoopAnimation(ResultDbSequence.Sequence->bLoop);
-			SequencePlayerNode.SetPlayRate(1.0f);
-		}
-		else if (SearchIndexAsset->Type == ESearchIndexAssetType::BlendSpace)
-		{
-			CurrentAssetPlayerNode = &BlendSpacePlayerNode;
-
-			const FPoseSearchDatabaseBlendSpace& ResultDbBlendSpace = Database->GetBlendSpaceSourceAsset(SearchIndexAsset);
-			BlendSpacePlayerNode.SetAccumulatedTime(MotionMatchingState.AssetPlayerTime);
-			BlendSpacePlayerNode.SetBlendSpace(ResultDbBlendSpace.BlendSpace);
-			BlendSpacePlayerNode.SetLoop(ResultDbBlendSpace.BlendSpace->bLoop);
-			BlendSpacePlayerNode.SetPlayRate(1.0f);
-			BlendSpacePlayerNode.SetPosition(SearchIndexAsset->BlendParameters);
-		}
-		else
-		{
-			checkNoEntry();
-		}
-
-		MirrorNode.SetSourceLinkNode(CurrentAssetPlayerNode);
-		MirrorNode.SetMirror(SearchIndexAsset->bMirrored);
+		const FPoseSearchDatabaseAnimationAssetBase& AnimationAssetBase = Database->GetAnimationSourceAsset(SearchIndexAsset);
+		BlendStackNode.BlendTo(	SearchIndexAsset->Type, AnimationAssetBase.GetAnimationAsset(), MotionMatchingState.AssetPlayerTime,
+								AnimationAssetBase.IsLooping(), SearchIndexAsset->bMirrored, Database->Schema->MirrorDataTable.Get(),
+								Settings.MaxActiveBlends, Settings.BlendTime, SearchIndexAsset->BlendParameters);
 	}
 
-	if (SearchIndexAsset && SearchIndexAsset->Type == ESearchIndexAssetType::Sequence)
+	if (!BlendStackNode.AnimPlayers.IsEmpty() && BlendStackNode.AnimPlayers[0].GetAssetType() == ESearchIndexAssetType::Sequence)
 	{
+		// @todo: integrate this functionality with the BlendStackNode
+		
 		// Optionally applying dynamic play rate adjustment to chosen sequences based on predictive motion analysis
 		const float PlayRate = DynamicPlayRateAdjustment(
 			Context,
 			Trajectory,
 			DynamicPlayRateSettings,
-			SequencePlayerNode.GetSequence(),
-			SequencePlayerNode.GetAccumulatedTime(),
-			SequencePlayerNode.GetPlayRate(),
-			SequencePlayerNode.GetLoopAnimation()
+			BlendStackNode.AnimPlayers[0].GetSequencePlayerNode().GetSequence(),
+			GetAccumulatedTime(),
+			BlendStackNode.AnimPlayers[0].GetSequencePlayerNode().GetPlayRate(),
+			BlendStackNode.AnimPlayers[0].GetSequencePlayerNode().GetLoopAnimation()
 		);
 
-		SequencePlayerNode.SetPlayRate(PlayRate);
+		BlendStackNode.AnimPlayers[0].GetSequencePlayerNode().SetPlayRate(PlayRate);
 	}
 
 	Source.Update(Context);
@@ -210,27 +175,27 @@ void FAnimNode_MotionMatching::GatherDebugData(FNodeDebugData& DebugData)
 // FAnimNode_AssetPlayerBase interface
 float FAnimNode_MotionMatching::GetAccumulatedTime() const
 {
-	return CurrentAssetPlayerNode->GetAccumulatedTime();
+	return BlendStackNode.GetAccumulatedTime();
 }
 
 UAnimationAsset* FAnimNode_MotionMatching::GetAnimAsset() const
 {
-	return CurrentAssetPlayerNode->GetAnimAsset();
+	return BlendStackNode.GetAnimAsset();
 }
 
 float FAnimNode_MotionMatching::GetCurrentAssetLength() const
 {
-	return CurrentAssetPlayerNode->GetCurrentAssetLength();
+	return BlendStackNode.GetCurrentAssetLength();
 }
 
 float FAnimNode_MotionMatching::GetCurrentAssetTime() const
 {
-	return CurrentAssetPlayerNode->GetCurrentAssetLength();
+	return BlendStackNode.GetCurrentAssetLength();
 }
 
 float FAnimNode_MotionMatching::GetCurrentAssetTimePlayRateAdjusted() const
 {
-	return CurrentAssetPlayerNode->GetCurrentAssetTimePlayRateAdjusted();
+	return BlendStackNode.GetCurrentAssetTimePlayRateAdjusted();
 }
 
 bool FAnimNode_MotionMatching::GetIgnoreForRelevancyTest() const
diff --git a/Engine/Plugins/Experimental/Animation/PoseSearch/Source/Runtime/Private/PoseSearch.cpp b/Engine/Plugins/Experimental/Animation/PoseSearch/Source/Runtime/Private/PoseSearch.cpp
--- a/Engine/Plugins/Experimental/Animation/PoseSearch/Source/Runtime/Private/PoseSearch.cpp
+++ b/Engine/Plugins/Experimental/Animation/PoseSearch/Source/Runtime/Private/PoseSearch.cpp
@@ -879,6 +879,16 @@ FFloatInterval FPoseSearchDatabaseSequence::GetEffectiveSamplingRange() const
 	return UE::PoseSearch::GetEffectiveSamplingRange(Sequence, SamplingRange);
 }
 
+UAnimationAsset* FPoseSearchDatabaseSequence::GetAnimationAsset() const
+{
+	return Sequence;
+}
+
+bool FPoseSearchDatabaseSequence::IsLooping() const
+{
+	return Sequence->bLoop;
+}
+
 //////////////////////////////////////////////////////////////////////////
 // UPoseSearchDatabase
 
@@ -942,6 +952,22 @@ float UPoseSearchDatabase::GetAssetTime(int32 PoseIdx, const FPoseSearchIndexAss
 	return AssetTime;
 }
 
+const FPoseSearchDatabaseAnimationAssetBase& UPoseSearchDatabase::GetAnimationSourceAsset(const FPoseSearchIndexAsset* SearchIndexAsset) const
+{
+	if (SearchIndexAsset->Type == ESearchIndexAssetType::Sequence)
+	{
+		return Sequences[SearchIndexAsset->SourceAssetIdx];
+	}
+
+	if (SearchIndexAsset->Type == ESearchIndexAssetType::BlendSpace)
+	{
+		return BlendSpaces[SearchIndexAsset->SourceAssetIdx];
+	}
+
+	checkNoEntry();
+	return Sequences[SearchIndexAsset->SourceAssetIdx];
+}
+
 const FPoseSearchDatabaseSequence& UPoseSearchDatabase::GetSequenceSourceAsset(const FPoseSearchIndexAsset* SearchIndexAsset) const
 {
 	check(SearchIndexAsset->Type == ESearchIndexAssetType::Sequence);
@@ -1291,6 +1317,16 @@ static inline void CollectGroupIndices(
 	}
 }
 
+UAnimationAsset* FPoseSearchDatabaseBlendSpace::GetAnimationAsset() const
+{
+	return BlendSpace.Get();
+}
+
+bool FPoseSearchDatabaseBlendSpace::IsLooping() const
+{
+	return BlendSpace->bLoop;
+}
+
 void FPoseSearchDatabaseBlendSpace::GetBlendSpaceParameterSampleRanges(
 	int32& HorizontalBlendNum,
 	int32& VerticalBlendNum,
@@ -2234,6 +2270,7 @@ UE::PoseSearch::FSearchResult UPoseSearchDatabase::SearchBruteForce(UE::PoseSear
 
 void UPoseSearchDatabase::BuildQuery(UE::PoseSearch::FSearchContext& SearchContext, FPoseSearchFeatureVectorBuilder& OutQuery) const
 {
+	check(Schema && Schema->IsValid());
 	Schema->BuildQuery(SearchContext, OutQuery);
 	OutQuery.Normalize(*GetSearchIndex());
 }
diff --git a/Engine/Plugins/Experimental/Animation/PoseSearch/Source/Runtime/Public/PoseSearch/AnimNode_BlendStack.h b/Engine/Plugins/Experimental/Animation/PoseSearch/Source/Runtime/Public/PoseSearch/AnimNode_BlendStack.h
new file mode 100644
--- /dev/null
+++ b/Engine/Plugins/Experimental/Animation/PoseSearch/Source/Runtime/Public/PoseSearch/AnimNode_BlendStack.h
@@ -0,0 +1,71 @@
+// Copyright Epic Games, Inc. All Rights Reserved.
+
+#pragma once
+
+#include "Animation/AnimNode_AssetPlayerBase.h"
+#include "Animation/AnimNode_SequencePlayer.h"
+#include "AnimNodes/AnimNode_BlendSpacePlayer.h"
+#include "AnimNodes/AnimNode_Mirror.h"
+#include "Containers/Deque.h"
+#include "PoseSearch/PoseSearch.h"
+
+#include "AnimNode_BlendStack.generated.h"
+
+USTRUCT()
+struct FPoseSearchAnimPlayer
+{
+	GENERATED_BODY()
+	
+	void Initialize(ESearchIndexAssetType InAssetType, UAnimationAsset* AnimationAsset, float AccumulatedTime, bool bLoop, bool bMirrored, UMirrorDataTable* MirrorDataTable, float BlendTime = 0.f, FVector BlendParameters = FVector::Zero());
+	void Evaluate_AnyThread(FPoseContext& Output);
+	void Update_AnyThread(const FAnimationUpdateContext& Context);
+	float GetAccumulatedTime() const;
+	float GetBlendWeight() const { return BlendWeight;  }
+	float GetBlendInPercentage() const;
+	void SetBlendWeight(float InBlendWeight);
+	ESearchIndexAssetType GetAssetType() const { return AssetType; }
+
+	// @todo: used only for DynamicPlayRateAdjustment. Remove once the functionality is integrated with the BlendStackNode
+	FAnimNode_SequencePlayer_Standalone GetSequencePlayerNode() { return SequencePlayerNode; }
+
+protected:
+	void UpdateSourceLinkNode();
+
+	// @todo: maybe use an union between the standalone players?
+	
+	// Embedded sequence player node for playing animations from the motion matching database
+	FAnimNode_SequencePlayer_Standalone SequencePlayerNode;
+
+	// Embedded blendspace player node for playing blendspaces from the motion matching database
+	FAnimNode_BlendSpacePlayer_Standalone BlendSpacePlayerNode;
+
+	// Embedded mirror node to handle mirroring if the pose search results in a mirrored sequence
+	FAnimNode_Mirror_Standalone MirrorNode;
+
+	ESearchIndexAssetType AssetType = ESearchIndexAssetType::Sequence;
+
+	float TotalBlendInTime = 0.f;
+	float CurrentBlendInTime = 0.f;
+	float BlendWeight = 1.f;
+};
+
+USTRUCT(BlueprintInternalUseOnly)
+struct POSESEARCH_API FAnimNode_BlendStack : public FAnimNode_AssetPlayerBase
+{
+	GENERATED_BODY()
+
+	TDeque<FPoseSearchAnimPlayer> AnimPlayers;
+
+	// FAnimNode_Base interface
+	virtual void Evaluate_AnyThread(FPoseContext& Output) override;
+	// End of FAnimNode_Base interface
+
+	void BlendTo(ESearchIndexAssetType AssetType, UAnimationAsset* AnimationAsset, float AccumulatedTime, bool bLoop, bool bMirrored, UMirrorDataTable* MirrorDataTable, int32 MaxActiveBlends = 3, float BlendTime = 0.3f, FVector BlendParameters = FVector::Zero());
+	void CalculateWeights();
+	void PruneBlendStack(int32 MaxActiveBlends);
+
+	// FAnimNode_AssetPlayerBase interface
+	virtual float GetAccumulatedTime() const override;
+	virtual void UpdateAssetPlayer(const FAnimationUpdateContext& Context) override;
+	// End of FAnimNode_AssetPlayerBase interface
+};
diff --git a/Engine/Plugins/Experimental/Animation/PoseSearch/Source/Runtime/Public/PoseSearch/AnimNode_MotionMatching.h b/Engine/Plugins/Experimental/Animation/PoseSearch/Source/Runtime/Public/PoseSearch/AnimNode_MotionMatching.h
--- a/Engine/Plugins/Experimental/Animation/PoseSearch/Source/Runtime/Public/PoseSearch/AnimNode_MotionMatching.h
+++ b/Engine/Plugins/Experimental/Animation/PoseSearch/Source/Runtime/Public/PoseSearch/AnimNode_MotionMatching.h
@@ -3,11 +3,9 @@
 #pragma once
 
 #include "Animation/AnimNode_AssetPlayerBase.h"
-#include "Animation/AnimNode_SequencePlayer.h"
-#include "AnimNodes/AnimNode_BlendSpacePlayer.h"
-#include "AnimNodes/AnimNode_Mirror.h"
 #include "Animation/MotionTrajectoryTypes.h"
 #include "DynamicPlayRate/DynamicPlayRateLibrary.h"
+#include "PoseSearch/AnimNode_BlendStack.h"
 #include "PoseSearch/PoseSearch.h"
 #include "PoseSearch/PoseSearchLibrary.h"
 
@@ -78,28 +76,18 @@ public:
 
 private:
 
-	// Embedded sequence player node for playing animations from the motion matching database
-	FAnimNode_SequencePlayer_Standalone SequencePlayerNode;
-
-	// Embedded blendspace player node for playing blendspaces from the motion matching database
-	FAnimNode_BlendSpacePlayer_Standalone BlendSpacePlayerNode;
-
-	// Embedded mirror node to handle mirroring if the pose search results in a mirrored sequence
-	FAnimNode_Mirror_Standalone MirrorNode;
+	FAnimNode_BlendStack BlendStackNode;
 
 	// Encapsulated motion matching algorithm and internal state
 	FMotionMatchingState MotionMatchingState;
 
-	// Current Asset Player Node
-	FAnimNode_AssetPlayerBase* CurrentAssetPlayerNode = &SequencePlayerNode;
-
 	// Update Counter for detecting being relevant
 	FGraphTraversalCounter UpdateCounter;
 
 	// FAnimNode_AssetPlayerBase
 protected:
 	// FAnimNode_AssetPlayerBase interface
-	virtual float GetAccumulatedTime() const;
+	virtual float GetAccumulatedTime() const override;
 	virtual UAnimationAsset* GetAnimAsset() const override;
 	virtual void UpdateAssetPlayer(const FAnimationUpdateContext& Context) override;
 	virtual float GetCurrentAssetLength() const override;
diff --git a/Engine/Plugins/Experimental/Animation/PoseSearch/Source/Runtime/Public/PoseSearch/PoseSearch.h b/Engine/Plugins/Experimental/Animation/PoseSearch/Source/Runtime/Public/PoseSearch/PoseSearch.h
--- a/Engine/Plugins/Experimental/Animation/PoseSearch/Source/Runtime/Public/PoseSearch/PoseSearch.h
+++ b/Engine/Plugins/Experimental/Animation/PoseSearch/Source/Runtime/Public/PoseSearch/PoseSearch.h
@@ -739,9 +739,19 @@ struct POSESEARCH_API FPoseSearchIndex
 //////////////////////////////////////////////////////////////////////////
 // Database
 
+USTRUCT()
+struct POSESEARCH_API FPoseSearchDatabaseAnimationAssetBase
+{
+	GENERATED_BODY()
+
+	virtual ~FPoseSearchDatabaseAnimationAssetBase() {}
+	virtual UAnimationAsset* GetAnimationAsset() const { return nullptr; }
+	virtual bool IsLooping() const { return false; }
+};
+
 /** An entry in a UPoseSearchDatabase. */
 USTRUCT(BlueprintType, Category = "Animation|Pose Search")
-struct POSESEARCH_API FPoseSearchDatabaseSequence
+struct POSESEARCH_API FPoseSearchDatabaseSequence : public FPoseSearchDatabaseAnimationAssetBase
 {
 	GENERATED_BODY()
 
@@ -772,11 +782,14 @@ struct POSESEARCH_API FPoseSearchDatabaseSequence
 	FGameplayTagContainer GroupTags;
 
 	FFloatInterval GetEffectiveSamplingRange() const;
+
+	virtual UAnimationAsset* GetAnimationAsset() const override;
+	virtual bool IsLooping() const override;
 };
 
 /** An blend space entry in a UPoseSearchDatabase. */
 USTRUCT(BlueprintType, Category = "Animation|Pose Search")
-struct POSESEARCH_API FPoseSearchDatabaseBlendSpace
+struct POSESEARCH_API FPoseSearchDatabaseBlendSpace : public FPoseSearchDatabaseAnimationAssetBase
 {
 	GENERATED_BODY()
 
@@ -800,6 +813,9 @@ struct POSESEARCH_API FPoseSearchDatabaseBlendSpace
 	UPROPERTY(EditAnywhere, Category = "Group")
 	FGameplayTagContainer GroupTags;
 
+	virtual UAnimationAsset* GetAnimationAsset() const override;
+	virtual bool IsLooping() const override;
+
 public:
 
 	void GetBlendSpaceParameterSampleRanges(
@@ -1130,6 +1146,7 @@ public:
 	int32 GetPoseIndexFromTime(float AssetTime, const FPoseSearchIndexAsset* SearchIndexAsset) const;
 	float GetAssetTime(int32 PoseIdx, const FPoseSearchIndexAsset* SearchIndexAsset = nullptr) const;
 
+	const FPoseSearchDatabaseAnimationAssetBase& GetAnimationSourceAsset(const FPoseSearchIndexAsset* SearchIndexAsset) const;
 	const FPoseSearchDatabaseSequence& GetSequenceSourceAsset(const FPoseSearchIndexAsset* SearchIndexAsset) const;
 	const FPoseSearchDatabaseBlendSpace& GetBlendSpaceSourceAsset(const FPoseSearchIndexAsset* SearchIndexAsset) const;
 	const bool IsSourceAssetLooping(const FPoseSearchIndexAsset* SearchIndexAsset) const;
diff --git a/Engine/Plugins/Experimental/Animation/PoseSearch/Source/Runtime/Public/PoseSearch/PoseSearchLibrary.h b/Engine/Plugins/Experimental/Animation/PoseSearch/Source/Runtime/Public/PoseSearch/PoseSearchLibrary.h
--- a/Engine/Plugins/Experimental/Animation/PoseSearch/Source/Runtime/Public/PoseSearch/PoseSearchLibrary.h
+++ b/Engine/Plugins/Experimental/Animation/PoseSearch/Source/Runtime/Public/PoseSearch/PoseSearchLibrary.h
@@ -27,12 +27,16 @@ struct POSESEARCH_API FMotionMatchingSettings
 {
 	GENERATED_BODY()
 
-	// Time in seconds to blend out to the new pose. Uses inertial blending and requires an Inertialization node after this node.
+	// Time in seconds to blend out to the new pose. Uses either inertial blending, requiring an Inertialization node after this node, or the internal blend stack, if MaxActiveBlends is greter than zero
 	UPROPERTY(EditAnywhere, BlueprintReadWrite, Category=Settings, meta=(ClampMin="0"))
 	float BlendTime = 0.2f;
 
+	// Number of max active blendin animation in the blend stack. If MaxActiveBlends is zero then blend stack is disabled
+	UPROPERTY(EditAnywhere, BlueprintReadWrite, Category=Settings, meta=(ClampMin="0"))
+	int32 MaxActiveBlends = 0;
+
 	// If the pose jump requires a mirroring change and this value is greater than 0, it will be used instead of BlendTime
-	UPROPERTY(EditAnywhere, BlueprintReadWrite, Category = Settings, meta = (ClampMin = "0", DislayAfter = "BlendTime"))
+	UPROPERTY(EditAnywhere, BlueprintReadWrite, Category=Settings, meta=(ClampMin="0", DislayAfter="BlendTime"))
 	float MirrorChangeBlendTime = 0.0f;
 	
 	// Don't jump to poses that are less than this many seconds away
