Changed errors and asserts to log warnings in PoseSearch so we don't fail the cook if a database asset isn't properly configured (e.g. missing schema). This might result in broken motion matching (i.e. tposes) but we won't block other teams from getting builds.
Also following the mandate that we never ensure on data errors (example discussion: https://epic.slack.com/archives/C02ER7UE95Y/p1639520721078300).

diff --git a/Engine/Plugins/Experimental/Animation/PoseSearch/Source/Runtime/Private/AnimNode_PoseSearchHistoryCollector.cpp b/Engine/Plugins/Experimental/Animation/PoseSearch/Source/Runtime/Private/AnimNode_PoseSearchHistoryCollector.cpp
--- a/Engine/Plugins/Experimental/Animation/PoseSearch/Source/Runtime/Private/AnimNode_PoseSearchHistoryCollector.cpp
+++ b/Engine/Plugins/Experimental/Animation/PoseSearch/Source/Runtime/Private/AnimNode_PoseSearchHistoryCollector.cpp
@@ -76,7 +76,7 @@ void FAnimNode_PoseSearchHistoryCollector::Evaluate_AnyThread(FPoseContext& Outp
 		FPoseHistory::ERootUpdateMode::RootMotionDelta :
 		FPoseHistory::ERootUpdateMode::ComponentTransformDelta))
 	{
-		Output.LogMessage(EMessageSeverity::Error, ErrorText);
+		Output.LogMessage(EMessageSeverity::Warning, ErrorText);
 	}
 }
 
diff --git a/Engine/Plugins/Experimental/Animation/PoseSearch/Source/Runtime/Private/PoseSearch.cpp b/Engine/Plugins/Experimental/Animation/PoseSearch/Source/Runtime/Private/PoseSearch.cpp
--- a/Engine/Plugins/Experimental/Animation/PoseSearch/Source/Runtime/Private/PoseSearch.cpp
+++ b/Engine/Plugins/Experimental/Animation/PoseSearch/Source/Runtime/Private/PoseSearch.cpp
@@ -1702,9 +1702,21 @@ void UPoseSearchDatabase::Serialize(FArchive& Ar)
 				PrivateDerivedData->SearchIndex.Schema = Schema;
 			}
 		}
-		check(Ar.IsLoading() || (Ar.IsCooking() && IsDerivedDataValid()));
-		FPoseSearchIndex* SearchIndex = GetSearchIndex();
-		Ar << *SearchIndex;
+		
+		if (Ar.IsLoading() || (Ar.IsCooking() && IsDerivedDataValid()))
+		{
+			FPoseSearchIndex* SearchIndex = GetSearchIndex();
+			Ar << *SearchIndex;
+		}
+		else
+		{
+			UE_LOG(
+				LogPoseSearch,
+				Warning,
+				TEXT("Invalid condition in UPoseSearchDatabase::Serialize for asset %s. It might be that the asset has a missing or invalid schema. IsLoading: %d, IsCooking: %d, IsDerivedDataValid: %d"),
+				*GetNameSafe(this),
+				Ar.IsLoading(), Ar.IsCooking(), IsDerivedDataValid());
+		}
 	}
 }
 
@@ -1932,19 +1944,19 @@ UE::PoseSearch::FSearchResult UPoseSearchDatabase::Search(UE::PoseSearch::FSearc
 		return Result;
 	}
 
-	if (!ensure(SearchIndex->IsValid() && !SearchIndex->IsEmpty()))
+	if (!SearchIndex->IsValid() || SearchIndex->IsEmpty())
 	{
 		if (!Schema)
 		{
-			UE_LOG(LogAnimation, Error, TEXT("UPoseSearchDatabase %s failed to index. Reason: no Schema!"), *GetName());
+			UE_LOG(LogAnimation, Warning, TEXT("UPoseSearchDatabase %s failed to index. Reason: no Schema!"), *GetName());
 		}
 		else if(!Schema->IsValid())
 		{
-			UE_LOG(LogAnimation, Error, TEXT("UPoseSearchDatabase %s failed to index. Reason: Schema %s is invalid"), *GetName(), *Schema->GetName());
+			UE_LOG(LogAnimation, Warning, TEXT("UPoseSearchDatabase %s failed to index. Reason: Schema %s is invalid"), *GetName(), *Schema->GetName());
 		}
 		else
 		{
-			UE_LOG(LogAnimation, Error, TEXT("UPoseSearchDatabase %s failed to index. Reason: is there any unsaved modified asset?"), *GetName());
+			UE_LOG(LogAnimation, Warning, TEXT("UPoseSearchDatabase %s failed to index. Reason: is there any unsaved modified asset?"), *GetName());
 		}
 		return Result;
 	}
@@ -2361,7 +2373,7 @@ UE::PoseSearch::FSearchResult UPoseSearchDatabaseSet::Search(UE::PoseSearch::FSe
 	{
 		UE_LOG(
 			LogPoseSearch, 
-			Error, 
+			Warning, 
 			TEXT("Invalid result searching %s"), *GetName());
 	}
 
@@ -5335,7 +5347,7 @@ bool BuildIndex(UPoseSearchDatabase* Database, FPoseSearchIndex& OutSearchIndex)
 
 	if (!Database->IsValidForIndexing())
 	{
-		UE_LOG(LogPoseSearch, Error, TEXT("Database '%f' is invalid for indexing"), *Database->GetName());
+		UE_LOG(LogPoseSearch, Warning, TEXT("Database '%f' is invalid for indexing"), *Database->GetName());
 		return false;
 	}
 
