Diagnostic information for invalid Bone Container

diff --git a/Engine/Plugins/Experimental/Animation/PoseSearch/Source/Runtime/Private/PoseSearch.cpp b/Engine/Plugins/Experimental/Animation/PoseSearch/Source/Runtime/Private/PoseSearch.cpp
--- a/Engine/Plugins/Experimental/Animation/PoseSearch/Source/Runtime/Private/PoseSearch.cpp
+++ b/Engine/Plugins/Experimental/Animation/PoseSearch/Source/Runtime/Private/PoseSearch.cpp
@@ -3165,6 +3165,7 @@ public:
 	virtual float ExtractRootDistance(float Time) const override;
 	virtual FTransform ExtractRootTransform(float Time) const override;
 	virtual void ExtractPoseSearchNotifyStates(float Time, TArray<class UAnimNotifyState_PoseSearchBase*>& NotifyStates) const override;
+	virtual const UAnimationAsset* GetAsset() const override { return Input.Sequence; }
 
 private:
 	float TotalRootDistance = 0.0f;
@@ -3384,6 +3385,8 @@ public:
 	virtual FTransform ExtractRootTransform(float Time) const override;
 	virtual void ExtractPoseSearchNotifyStates(float Time, TArray<class UAnimNotifyState_PoseSearchBase*>& NotifyStates) const override;
 
+	virtual const UAnimationAsset* GetAsset() const override { return Input.BlendSpace; }
+
 private:
 	float PlayLength = 0.0f;
 	float TotalRootDistance = 0.0f;
@@ -4281,6 +4284,17 @@ FTransform FAssetIndexer::GetTransformAndCacheResults(float SampleTime, float Or
 		Entry->SampleTime = SampleTime;
 		Entry->OriginTime = OriginTime;
 
+		if (!SamplingContext->BoneContainer.IsValid())
+		{
+			UE_LOG(LogPoseSearch,
+				Warning, 
+				TEXT("Invalid BoneContainer encountered in FAssetIndexer::GetTransformAndCacheResults. Asset: %s. Schema: %s. BoneContainerAsset: %s. NumBoneIndices: %d"),
+				*GetNameSafe(IndexingContext.MainSampler->GetAsset()),
+				*GetNameSafe(IndexingContext.Schema),
+				*GetNameSafe(SamplingContext->BoneContainer.GetAsset()),
+				SamplingContext->BoneContainer.GetCompactPoseNumBones());
+		}
+
 		Entry->Pose.SetBoneContainer(&SamplingContext->BoneContainer);
 		Entry->UnusedCurve.InitFrom(SamplingContext->BoneContainer);
 
@@ -5351,6 +5365,8 @@ bool BuildIndex(UPoseSearchDatabase* Database, FPoseSearchIndex& OutSearchIndex)
 		return false;
 	}
 
+	UE_LOG(LogPoseSearch, Log, TEXT("Building PoseSearch index for database %s"), *GetNameSafe(Database));
+
 	OutSearchIndex.Schema = Database->Schema;
 
 	if (!Database->TryInitSearchIndexAssets(OutSearchIndex))
diff --git a/Engine/Plugins/Experimental/Animation/PoseSearch/Source/Runtime/Public/PoseSearch/PoseSearch.h b/Engine/Plugins/Experimental/Animation/PoseSearch/Source/Runtime/Public/PoseSearch/PoseSearch.h
--- a/Engine/Plugins/Experimental/Animation/PoseSearch/Source/Runtime/Public/PoseSearch/PoseSearch.h
+++ b/Engine/Plugins/Experimental/Animation/PoseSearch/Source/Runtime/Public/PoseSearch/PoseSearch.h
@@ -289,6 +289,8 @@ public:
 	// Extracts notify states inheriting from UAnimNotifyState_PoseSearchBase present in the sequence at Time.
 	// The function does not empty NotifyStates before adding new notifies!
 	virtual void ExtractPoseSearchNotifyStates(float Time, TArray<UAnimNotifyState_PoseSearchBase*>& NotifyStates) const = 0;
+
+	virtual const UAnimationAsset* GetAsset() const = 0;
 };
 
 /**
