motion matching - query pose feature data blend

diff --git a/Engine/Plugins/Experimental/Animation/PoseSearch/Source/Runtime/Private/AnimNode_MotionMatching.cpp b/Engine/Plugins/Experimental/Animation/PoseSearch/Source/Runtime/Private/AnimNode_MotionMatching.cpp
--- a/Engine/Plugins/Experimental/Animation/PoseSearch/Source/Runtime/Private/AnimNode_MotionMatching.cpp
+++ b/Engine/Plugins/Experimental/Animation/PoseSearch/Source/Runtime/Private/AnimNode_MotionMatching.cpp
@@ -98,7 +98,7 @@ void FAnimNode_MotionMatching::UpdateAssetPlayer(const FAnimationUpdateContext&
 	if (Database && Database->Schema && SearchIndexAsset && (MotionMatchingState.Flags & EMotionMatchingFlags::JumpedToPose) == EMotionMatchingFlags::JumpedToPose)
 	{
 		const FPoseSearchDatabaseAnimationAssetBase& AnimationAssetBase = Database->GetAnimationSourceAsset(SearchIndexAsset);
-		BlendStackNode.BlendTo(	SearchIndexAsset->Type, AnimationAssetBase.GetAnimationAsset(), MotionMatchingState.AssetPlayerTime,
+		BlendStackNode.BlendTo(	SearchIndexAsset->Type, AnimationAssetBase.GetAnimationAsset(), MotionMatchingState.CurrentSearchResult.AssetTime,
 								AnimationAssetBase.IsLooping(), SearchIndexAsset->bMirrored, Database->Schema->MirrorDataTable.Get(),
 								Settings.MaxActiveBlends, Settings.BlendTime, SearchIndexAsset->BlendParameters);
 	}
diff --git a/Engine/Plugins/Experimental/Animation/PoseSearch/Source/Runtime/Private/PoseSearch.cpp b/Engine/Plugins/Experimental/Animation/PoseSearch/Source/Runtime/Private/PoseSearch.cpp
--- a/Engine/Plugins/Experimental/Animation/PoseSearch/Source/Runtime/Private/PoseSearch.cpp
+++ b/Engine/Plugins/Experimental/Animation/PoseSearch/Source/Runtime/Private/PoseSearch.cpp
@@ -37,6 +37,7 @@
 #include "Animation/AnimRootMotionProvider.h"
 #include "Animation/BlendSpace1D.h"
 
+
 #if WITH_EDITOR
 #include "AssetRegistry/AssetRegistryModule.h"
 #include "AssetRegistry/ARFilter.h"
@@ -560,18 +561,6 @@ bool UPoseSearchSchema::BuildQuery(UE::PoseSearch::FSearchContext& SearchContext
 
 	InOutQuery.Init(this);
 
-	// Copy search query directly from the database if we have an active pose
-	const bool bCopySearchResult = 
-		SearchContext.CurrentResult.IsValid() &&
-		SearchContext.CurrentResult.Database->Schema == this;
-
-	if (bCopySearchResult)
-	{
-		InOutQuery.CopyFromSearchIndex(
-			*SearchContext.CurrentResult.Database->GetSearchIndex(),
-			SearchContext.CurrentResult.PoseIdx);
-	}
-
 	bool bSuccess = true;
 	for (const TObjectPtr<UPoseSearchFeatureChannel>& Channel : Channels)
 	{
@@ -652,7 +641,8 @@ float FPoseSearchIndex::GetAssetTime(int32 PoseIdx, const FPoseSearchIndexAsset*
 			SamplingRange.Max);
 		return AssetTime;
 	}
-	else if (Asset->Type == ESearchIndexAssetType::BlendSpace)
+	
+	if (Asset->Type == ESearchIndexAssetType::BlendSpace)
 	{
 		const FFloatInterval SamplingRange = Asset->SamplingInterval;
 
@@ -663,11 +653,9 @@ float FPoseSearchIndex::GetAssetTime(int32 PoseIdx, const FPoseSearchIndexAsset*
 			SamplingRange.Max) / (Asset->NumPoses * Schema->SamplingInterval);
 		return AssetTime;
 	}
-	else
-	{
-		checkNoEntry();
-		return -1.0f;
-	}
+	
+	checkNoEntry();
+	return -1.0f;
 }
 
 
@@ -757,6 +745,36 @@ void FPoseSearchIndex::InverseNormalize(TArrayView<float> InOutNormalizedPoseVec
 	NormalizedPoseVector = (InverseTransformationMtx * NormalizedPoseVector) + SampleMean;
 }
 
+float FPoseSearchIndex::ComputeMirrorMismatchAddend(int32 PoseIdx, UE::PoseSearch::FSearchContext& SearchContext) const
+{
+	if (SearchContext.QueryMirrorRequest != EPoseSearchBooleanRequest::Indifferent)
+	{
+		const FPoseSearchIndexAsset* IndexAsset = FindAssetForPose(PoseIdx);
+		const bool bMirroringMismatch =
+			(IndexAsset->bMirrored && SearchContext.QueryMirrorRequest == EPoseSearchBooleanRequest::FalseValue) ||
+			(!IndexAsset->bMirrored && SearchContext.QueryMirrorRequest == EPoseSearchBooleanRequest::TrueValue);
+		if (bMirroringMismatch)
+		{
+			check(Schema);
+			return Schema->MirrorMismatchCostBias;
+		}
+	}
+	return 0.f;
+}
+
+float FPoseSearchIndex::ComputeNotifyAddend(int32 PoseIdx) const
+{
+	return PoseMetadata[PoseIdx].CostAddend;
+}
+
+float FPoseSearchIndex::ComputeContinuingPoseCostAddend(int32 PoseIdx, UE::PoseSearch::EPoseComparisonFlags PoseComparisonFlags) const
+{
+	if (EnumHasAnyFlags(PoseComparisonFlags, UE::PoseSearch::EPoseComparisonFlags::ContinuingPose))
+	{
+		return PoseMetadata[PoseIdx].ContinuingPoseCostAddend;
+	}
+	return 0.f;
+}
 
 //////////////////////////////////////////////////////////////////////////
 // UPoseSearchSequenceMetaData
@@ -841,7 +859,6 @@ UE::PoseSearch::FSearchResult UPoseSearchSequenceMetaData::Search(UE::PoseSearch
 	Result.SearchIndexAsset = SearchIndex.FindAssetForPose(BestPoseIdx);
 	Result.AssetTime = SearchIndex.GetAssetTime(BestPoseIdx, Result.SearchIndexAsset);
 	Result.Database = nullptr;
-	Result.Sequence = nullptr;
 
 	SearchContext.DebugDrawParams.PoseVector = NormalizedQueryValues;
 	SearchContext.DebugDrawParams.PoseIdx = Result.PoseIdx;
@@ -864,9 +881,12 @@ FPoseSearchCost UPoseSearchSequenceMetaData::ComparePoses(
 
 	Result.SetDissimilarity(CompareFeatureVectors(PoseValues.Num(), PoseValues.GetData(), QueryValues.GetData()));
 
-	const FPoseSearchPoseMetadata& PoseMetadata = SearchIndex.PoseMetadata[PoseIdx];
-	const float ContinuingPoseCostAddend = EnumHasAnyFlags(PoseComparisonFlags, EPoseComparisonFlags::ContinuingPose) ? PoseMetadata.ContinuingPoseCostAddend : 0.f;
-	Result.SetCostAddend(PoseMetadata.CostAddend + ContinuingPoseCostAddend);
+	// @todo: shouldn't we include MirrorMismatchAddend as well?
+	//const float MirrorMismatchAddend = SearchIndex.ComputeMirrorMismatchAddend(PoseIdx, SearchContext);
+	const float NotifyAddend = SearchIndex.ComputeNotifyAddend(PoseIdx);
+	const float ContinuingPoseCostAddend = SearchIndex.ComputeContinuingPoseCostAddend(PoseIdx, PoseComparisonFlags);
+
+	Result.SetCostAddend(NotifyAddend + ContinuingPoseCostAddend);
 
 	return Result;
 }
@@ -912,26 +932,37 @@ const FPoseSearchIndex* UPoseSearchDatabase::GetSearchIndex() const
 	return &PrivateDerivedData->SearchIndex;
 }
 
-int32 UPoseSearchDatabase::GetPoseIndexFromTime(float Time, const FPoseSearchIndexAsset* SearchIndexAsset) const
+float UPoseSearchDatabase::GetAssetTime(int32 PoseIdx, const FPoseSearchIndexAsset* Asset) const
 {
-	FFloatInterval Range = SearchIndexAsset->SamplingInterval;
-
-	const bool bHasPoseIndex =
-		SearchIndexAsset->FirstPoseIdx != INDEX_NONE &&
-		SearchIndexAsset->NumPoses > 0 &&
-		Range.Contains(Time);
+	const float AssetTime = GetSearchIndex()->GetAssetTime(PoseIdx, Asset);
+	return AssetTime;
+}
 
+int32 UPoseSearchDatabase::GetPoseIndexFromTime(float Time, const FPoseSearchIndexAsset* SearchIndexAsset) const
+{
+	const bool bIsLooping = IsSourceAssetLooping(SearchIndexAsset);
+	const FFloatInterval& Range = SearchIndexAsset->SamplingInterval;
+	const bool bHasPoseIndex = SearchIndexAsset->FirstPoseIdx != INDEX_NONE && SearchIndexAsset->NumPoses > 0 && (bIsLooping || Range.Contains(Time));
 	if (bHasPoseIndex)
 	{
 		int32 PoseOffset = FMath::RoundToInt(Schema->SampleRate * (Time - Range.Min));
 		
-		check(PoseOffset >= 0);
-
-		if (PoseOffset >= SearchIndexAsset->NumPoses)
+		if (PoseOffset < 0)
+		{
+			if (bIsLooping)
+			{
+				PoseOffset = (PoseOffset % SearchIndexAsset->NumPoses) + SearchIndexAsset->NumPoses;
+			}
+			else
+			{
+				PoseOffset = 0;
+			}
+		}
+		else if (PoseOffset >= SearchIndexAsset->NumPoses)
 		{
-			if (IsSourceAssetLooping(SearchIndexAsset))
+			if (bIsLooping)
 			{
-				PoseOffset -= SearchIndexAsset->NumPoses;
+				PoseOffset = PoseOffset % SearchIndexAsset->NumPoses;
 			}
 			else
 			{
@@ -946,10 +977,40 @@ int32 UPoseSearchDatabase::GetPoseIndexFromTime(float Time, const FPoseSearchInd
 	return INDEX_NONE;
 }
 
-float UPoseSearchDatabase::GetAssetTime(int32 PoseIdx, const FPoseSearchIndexAsset* Asset) const
+bool UPoseSearchDatabase::GetPoseIndicesAndLerpValueFromTime(float Time, const FPoseSearchIndexAsset* SearchIndexAsset, int32& PrevPoseIdx, int32& PoseIdx, int32& NextPoseIdx, float& LerpValue) const
 {
-	float AssetTime = GetSearchIndex()->GetAssetTime(PoseIdx, Asset);
-	return AssetTime;
+	PoseIdx = GetPoseIndexFromTime(Time, SearchIndexAsset);
+	if (PoseIdx == INDEX_NONE)
+	{
+		PrevPoseIdx = INDEX_NONE;
+		NextPoseIdx = INDEX_NONE;
+		LerpValue = 0.f;
+		return false;
+	}
+
+	const FFloatInterval& Range = SearchIndexAsset->SamplingInterval;
+	const float FloatPoseOffset = Schema->SampleRate * (Time - Range.Min);
+	const int32 PoseOffset = FMath::RoundToInt(FloatPoseOffset);
+	LerpValue = FloatPoseOffset - float(PoseOffset);
+
+	const float PrevTime = Time - 1.f / Schema->SampleRate;
+	const float NextTime = Time + 1.f / Schema->SampleRate;
+
+	PrevPoseIdx = GetPoseIndexFromTime(PrevTime, SearchIndexAsset);
+	if (PrevPoseIdx == INDEX_NONE)
+	{
+		PrevPoseIdx = PoseIdx;
+	}
+
+	NextPoseIdx = GetPoseIndexFromTime(NextTime, SearchIndexAsset);
+	if (NextPoseIdx == INDEX_NONE)
+	{
+		NextPoseIdx = PoseIdx;
+	}
+
+	check(LerpValue >= -0.5f && LerpValue <= 0.5f);
+
+	return true;
 }
 
 const FPoseSearchDatabaseAnimationAssetBase& UPoseSearchDatabase::GetAnimationSourceAsset(const FPoseSearchIndexAsset* SearchIndexAsset) const
@@ -1827,10 +1888,11 @@ FPoseSearchCost UPoseSearchDatabase::ComparePoses(UE::PoseSearch::FSearchContext
 		QueryValues.GetData(),
 		SearchIndex->FindGroup(GroupIdx)->Weights.GetData()));
 
-	float NotifyAddend = 0.0f;
-	float MirrorMismatchAddend = 0.0f;
-	ComputePoseCostAddends(PoseIdx, PoseComparisonFlags, SearchContext, NotifyAddend, MirrorMismatchAddend);
-	Result.SetCostAddend(NotifyAddend + MirrorMismatchAddend);
+	const float MirrorMismatchAddend = SearchIndex->ComputeMirrorMismatchAddend(PoseIdx, SearchContext);
+	const float NotifyAddend = SearchIndex->ComputeNotifyAddend(PoseIdx);
+	const float ContinuingPoseCostAddend = SearchIndex->ComputeContinuingPoseCostAddend(PoseIdx, PoseComparisonFlags);
+
+	Result.SetCostAddend(NotifyAddend + MirrorMismatchAddend + ContinuingPoseCostAddend);
 
 	return Result;
 }
@@ -1842,7 +1904,10 @@ FPoseSearchCost UPoseSearchDatabase::ComparePoses(UE::PoseSearch::FSearchContext
 
 	FPoseSearchCost Result;
 
-	TArrayView<const float> PoseValues = GetSearchIndex()->GetPoseValues(PoseIdx);
+	const FPoseSearchIndex* SearchIndex = GetSearchIndex();
+	check(SearchIndex);
+
+	TArrayView<const float> PoseValues = SearchIndex->GetPoseValues(PoseIdx);
 	const int32 Dims = PoseValues.Num();
 	check(Dims == QueryValues.Num());
 
@@ -1854,8 +1919,8 @@ FPoseSearchCost UPoseSearchDatabase::ComparePoses(UE::PoseSearch::FSearchContext
 	auto QueryVector = Map<const ArrayXf>(QueryValues.GetData(), Dims);
 
 	// Compute weighted squared difference vector
-	const FPoseSearchIndexAsset* SearchIndexAsset = GetSearchIndex()->FindAssetForPose(PoseIdx);
-	const TArray<float>& Weights = GetSearchIndex()->FindGroup(SearchIndexAsset->SourceGroupIdx)->Weights;
+	const FPoseSearchIndexAsset* SearchIndexAsset = SearchIndex->FindAssetForPose(PoseIdx);
+	const TArray<float>& Weights = SearchIndex->FindGroup(SearchIndexAsset->SourceGroupIdx)->Weights;
 
 	check(Weights.Num() == Dims);
 	auto WeightsVector = Map<const ArrayXf>(Weights.GetData(), Dims);
@@ -1863,18 +1928,18 @@ FPoseSearchCost UPoseSearchDatabase::ComparePoses(UE::PoseSearch::FSearchContext
 	OutCostVector = WeightsVector * (PoseVector - QueryVector).square();
 	Result.SetDissimilarity(OutCostVector.sum());
 
-	// Output result
-	float NotifyAddend = 0.0f;
-	float MirrorMismatchAddend = 0.0f;
-	ComputePoseCostAddends(PoseIdx, PoseComparisonFlags, SearchContext, NotifyAddend, MirrorMismatchAddend);
-	Result.SetCostAddend(NotifyAddend + MirrorMismatchAddend);
+	const float MirrorMismatchAddend = SearchIndex->ComputeMirrorMismatchAddend(PoseIdx, SearchContext);
+	const float NotifyAddend = SearchIndex->ComputeNotifyAddend(PoseIdx);
+	const float ContinuingPoseCostAddend = SearchIndex->ComputeContinuingPoseCostAddend(PoseIdx, PoseComparisonFlags);
+
+	Result.SetCostAddend(NotifyAddend + MirrorMismatchAddend + ContinuingPoseCostAddend);
 
 	// Output cost details
-	OutPoseCostDetails.NotifyCostAddend = NotifyAddend;
+	OutPoseCostDetails.NotifyCostAddend = NotifyAddend + ContinuingPoseCostAddend;
 	OutPoseCostDetails.MirrorMismatchCostAddend = MirrorMismatchAddend;
 	OutPoseCostDetails.PoseCost = Result;
 	CalcChannelCosts(
-		GetSearchIndex()->Schema,
+		SearchIndex->Schema,
 		OutPoseCostDetails.CostVector,
 		OutPoseCostDetails.ChannelCosts);
 
@@ -1900,31 +1965,6 @@ FPoseSearchCost UPoseSearchDatabase::ComparePoses(UE::PoseSearch::FSearchContext
 	return Result;
 }
 
-void UPoseSearchDatabase::ComputePoseCostAddends(int32 PoseIdx, UE::PoseSearch::EPoseComparisonFlags PoseComparisonFlags, UE::PoseSearch::FSearchContext& SearchContext, float& OutNotifyAddend, float& OutMirrorMismatchAddend) const
-{
-	OutNotifyAddend = 0.0f;
-	OutMirrorMismatchAddend = 0.0f;
-
-	const FPoseSearchIndex* SearchIndex = GetSearchIndex();
-	check(SearchIndex);
-	if (SearchContext.QueryMirrorRequest != EPoseSearchBooleanRequest::Indifferent)
-	{
-		const FPoseSearchIndexAsset* IndexAsset = SearchIndex->FindAssetForPose(PoseIdx);
-		const bool bMirroringMismatch =
-			(IndexAsset->bMirrored && SearchContext.QueryMirrorRequest == EPoseSearchBooleanRequest::FalseValue) ||
-			(!IndexAsset->bMirrored && SearchContext.QueryMirrorRequest == EPoseSearchBooleanRequest::TrueValue);
-		if (bMirroringMismatch)
-		{
-			check(Schema);
-			OutMirrorMismatchAddend = Schema->MirrorMismatchCostBias;
-		}
-	}
-
-	const FPoseSearchPoseMetadata& PoseMetadata = SearchIndex->PoseMetadata[PoseIdx];
-	const float ContinuingPoseCostAddend = EnumHasAnyFlags(PoseComparisonFlags, UE::PoseSearch::EPoseComparisonFlags::ContinuingPose) ? PoseMetadata.ContinuingPoseCostAddend : 0.f;
-	OutNotifyAddend = PoseMetadata.CostAddend + ContinuingPoseCostAddend;
-}
-
 UE::PoseSearch::FSearchResult UPoseSearchDatabase::Search(UE::PoseSearch::FSearchContext& SearchContext) const
 {
 	using namespace UE::PoseSearch;
@@ -2153,16 +2193,25 @@ UE::PoseSearch::FSearchResult UPoseSearchDatabase::SearchPCAKDTree(UE::PoseSearc
 			SearchContext.UpdateCurrentBestCost(BestPoseCost);
 		}
 	}
+	else if (Result.ContinuingPoseCost.IsValid())
+	{
+		// if the search were skipped we'll reuse the continuing pose search result and subtract the ContinuingPoseCostAddend from the cost
+		BestPoseCost = Result.ContinuingPoseCost;
+		BestPoseIdx = SearchContext.CurrentResult.PoseIdx;
+		BestPoseCost.SetCostAddend(BestPoseCost.GetCostAddend() - SearchIndex->ComputeContinuingPoseCostAddend(SearchContext.CurrentResult.PoseIdx, EPoseComparisonFlags::ContinuingPose));
+	}
 
-	Result.PoseCost = BestPoseCost;
-	Result.PoseIdx = BestPoseIdx;
-	Result.SearchIndexAsset = SearchIndex->FindAssetForPose(BestPoseIdx);
-	Result.AssetTime = SearchIndex->GetAssetTime(BestPoseIdx, Result.SearchIndexAsset);
-	Result.Database = this;
-	Result.Sequence = nullptr;
+	if (BestPoseIdx != INDEX_NONE)
+	{
+		Result.PoseCost = BestPoseCost;
+		Result.PoseIdx = BestPoseIdx;
+		Result.SearchIndexAsset = SearchIndex->FindAssetForPose(BestPoseIdx);
+		Result.AssetTime = SearchIndex->GetAssetTime(BestPoseIdx, Result.SearchIndexAsset);
+		Result.Database = this;
 #if WITH_EDITOR
-	Result.SearchIndexHash = GetSearchIndexHash();
+		Result.SearchIndexHash = GetSearchIndexHash();
 #endif // WITH_EDITOR
+	}
 
 	return Result;
 }
@@ -2266,16 +2315,25 @@ UE::PoseSearch::FSearchResult UPoseSearchDatabase::SearchBruteForce(UE::PoseSear
 			SearchContext.UpdateCurrentBestCost(BestPoseCost);
 		}
 	}
+	else if (Result.ContinuingPoseCost.IsValid())
+	{
+		// if the search were skipped we'll reuse the continuing pose search result and subtract the ContinuingPoseCostAddend from the cost
+		BestPoseCost = Result.ContinuingPoseCost;
+		BestPoseIdx = SearchContext.CurrentResult.PoseIdx;
+		BestPoseCost.SetCostAddend(BestPoseCost.GetCostAddend() - SearchIndex->ComputeContinuingPoseCostAddend(SearchContext.CurrentResult.PoseIdx, EPoseComparisonFlags::ContinuingPose));
+	}
 
-	Result.PoseCost = BestPoseCost;
-	Result.PoseIdx = BestPoseIdx;
-	Result.SearchIndexAsset = SearchIndex->FindAssetForPose(BestPoseIdx);
-	Result.AssetTime = SearchIndex->GetAssetTime(BestPoseIdx, Result.SearchIndexAsset);
-	Result.Database = this;
-	Result.Sequence = nullptr;
+	if (BestPoseIdx != INDEX_NONE)
+	{
+		Result.PoseCost = BestPoseCost;
+		Result.PoseIdx = BestPoseIdx;
+		Result.SearchIndexAsset = SearchIndex->FindAssetForPose(BestPoseIdx);
+		Result.AssetTime = SearchIndex->GetAssetTime(BestPoseIdx, Result.SearchIndexAsset);
+		Result.Database = this;
 #if WITH_EDITOR
-	Result.SearchIndexHash = GetSearchIndexHash();
+		Result.SearchIndexHash = GetSearchIndexHash();
 #endif // WITH_EDITOR
+	}
 
 	return Result;
 }
@@ -2541,6 +2599,24 @@ bool FSearchContext::GetOrBuildQuery(const UPoseSearchDatabase* Database, FPoseS
 	return false;
 }
 
+void FSearchContext::CacheCurrentResultFeatureVectors()
+{
+	if (!CurrentResultPoseVector.IsInitialized())
+	{
+		check(CurrentResult.IsValid());
+		const FPoseSearchIndex* SearchIndex = CurrentResult.Database->GetSearchIndex();
+		check(SearchIndex);
+
+		CurrentResultPoseVector.Init(CurrentResult.Database->Schema);
+		CurrentResultPrevPoseVector.Init(CurrentResult.Database->Schema);
+		CurrentResultNextPoseVector.Init(CurrentResult.Database->Schema);
+
+		CurrentResultPoseVector.CopyFromSearchIndex(*SearchIndex, CurrentResult.PoseIdx);
+		CurrentResultPrevPoseVector.CopyFromSearchIndex(*SearchIndex, CurrentResult.PrevPoseIdx);
+		CurrentResultNextPoseVector.CopyFromSearchIndex(*SearchIndex, CurrentResult.NextPoseIdx);
+	}
+}
+
 //////////////////////////////////////////////////////////////////////////
 // FPoseHistory
 
@@ -2743,7 +2819,35 @@ void FFeatureVectorHelper::EncodeQuat(TArrayView<float> Values, int32& DataOffse
 	DataOffset += EncodeQuatCardinality;
 }
 
+void FFeatureVectorHelper::EncodeQuat(TArrayView<float> Values, int32& DataOffset, TArrayView<const float> PrevValues, TArrayView<const float> CurPrevValues, TArrayView<const float> NextPrevValues, float LerpValue)
+{
+	FQuat Quat = DecodeQuatInternal(CurPrevValues, DataOffset);
+
+	// linear interpolation
+	if (!FMath::IsNearlyZero(LerpValue))
+	{
+		if (LerpValue < 0.f)
+		{
+			Quat = FQuat::Slerp(Quat, DecodeQuatInternal(PrevValues, DataOffset), -LerpValue);
+		}
+		else
+		{
+			Quat = FQuat::Slerp(Quat, DecodeQuatInternal(NextPrevValues, DataOffset), LerpValue);
+		}
+	}
+	// @todo: do we need to add options for cubic interpolation?
+
+	EncodeQuat(Values, DataOffset, Quat);
+}
+
 FQuat FFeatureVectorHelper::DecodeQuat(TArrayView<const float> Values, int32& DataOffset)
+{
+	const FQuat Quat = DecodeQuatInternal(Values, DataOffset);
+	DataOffset += EncodeQuatCardinality;
+	return Quat;
+}
+
+FQuat FFeatureVectorHelper::DecodeQuatInternal(TArrayView<const float> Values, int32 DataOffset)
 {
 	const FVector X(Values[DataOffset + 0], Values[DataOffset + 1], Values[DataOffset + 2]);
 	const FVector Y(Values[DataOffset + 3], Values[DataOffset + 4], Values[DataOffset + 5]);
@@ -2754,7 +2858,6 @@ FQuat FFeatureVectorHelper::DecodeQuat(TArrayView<const float> Values, int32& Da
 	M.SetColumn(1, Y);
 	M.SetColumn(2, Z);
 
-	DataOffset += EncodeQuatCardinality;
 	return FQuat(M);
 }
 
@@ -2766,13 +2869,44 @@ void FFeatureVectorHelper::EncodeVector(TArrayView<float> Values, int32& DataOff
 	DataOffset += EncodeVectorCardinality;
 }
 
+void FFeatureVectorHelper::EncodeVector(TArrayView<float> Values, int32& DataOffset, TArrayView<const float> PrevValues, TArrayView<const float> CurPrevValues, TArrayView<const float> NextPrevValues, float LerpValue, bool bNormalize)
+{
+	FVector Vector = DecodeVectorInternal(CurPrevValues, DataOffset);
+
+	// linear interpolation
+	if (!FMath::IsNearlyZero(LerpValue))
+	{
+		if (LerpValue < 0.f)
+		{
+			Vector = FMath::Lerp(Vector, DecodeVectorInternal(PrevValues, DataOffset), -LerpValue);
+		}
+		else
+		{
+			Vector = FMath::Lerp(Vector, DecodeVectorInternal(NextPrevValues, DataOffset), LerpValue);
+		}
+	}
+	// @todo: do we need to add options for cubic interpolation?
+
+	if (bNormalize)
+	{
+		Vector = Vector.GetSafeNormal(UE_SMALL_NUMBER, FVector::XAxisVector);
+	}
+
+	EncodeVector(Values, DataOffset, Vector);
+}
+
 FVector FFeatureVectorHelper::DecodeVector(TArrayView<const float> Values, int32& DataOffset)
 {
-	const FVector Vector(Values[DataOffset + 0], Values[DataOffset + 1], Values[DataOffset + 2]);
+	const FVector Vector = DecodeVectorInternal(Values, DataOffset);
 	DataOffset += EncodeVectorCardinality;
 	return Vector;
 }
 
+FVector FFeatureVectorHelper::DecodeVectorInternal(TArrayView<const float> Values, int32 DataOffset)
+{
+	return FVector(Values[DataOffset + 0], Values[DataOffset + 1], Values[DataOffset + 2]);
+}
+
 void FFeatureVectorHelper::EncodeVector2D(TArrayView<float> Values, int32& DataOffset, const FVector2D& Vector2D)
 {
 	Values[DataOffset + 0] = Vector2D.X;
@@ -2780,13 +2914,39 @@ void FFeatureVectorHelper::EncodeVector2D(TArrayView<float> Values, int32& DataO
 	DataOffset += EncodeVector2DCardinality;
 }
 
+void FFeatureVectorHelper::EncodeVector2D(TArrayView<float> Values, int32& DataOffset, TArrayView<const float> PrevValues, TArrayView<const float> CurPrevValues, TArrayView<const float> NextPrevValues, float LerpValue)
+{
+	FVector2D Vector2D = DecodeVector2DInternal(CurPrevValues, DataOffset);
+
+	// linear interpolation
+	if (!FMath::IsNearlyZero(LerpValue))
+	{
+		if (LerpValue < 0.f)
+		{
+			Vector2D = FMath::Lerp(Vector2D, DecodeVector2DInternal(PrevValues, DataOffset), -LerpValue);
+		}
+		else
+		{
+			Vector2D = FMath::Lerp(Vector2D, DecodeVector2DInternal(NextPrevValues, DataOffset), LerpValue);
+		}
+	}
+	// @todo: do we need to add options for cubic interpolation?
+
+	EncodeVector2D(Values, DataOffset, Vector2D);
+}
+
 FVector2D FFeatureVectorHelper::DecodeVector2D(TArrayView<const float> Values, int32& DataOffset)
 {
-	const FVector2D Vector2D(Values[DataOffset + 0], Values[DataOffset + 1]);
+	const FVector2D Vector2D = DecodeVector2DInternal(Values, DataOffset);
 	DataOffset += EncodeVector2DCardinality;
 	return Vector2D;
 }
 
+FVector2D FFeatureVectorHelper::DecodeVector2DInternal(TArrayView<const float> Values, int32 DataOffset)
+{
+	return FVector2D(Values[DataOffset + 0], Values[DataOffset + 1]);
+}
+
 void FFeatureVectorHelper::ComputeMeanDeviations(const Eigen::MatrixXd& CenteredPoseMatrix, Eigen::VectorXd& MeanDeviations, int32& DataOffset, int32 Cardinality)
 {
 	const int32 NumPoses = MeanDeviations.size();
@@ -2891,9 +3051,8 @@ void FSearchResult::Update(float NewAssetTime)
 	{
 		const FPoseSearchDatabaseSequence& DbSequence = Database->GetSequenceSourceAsset(SearchIndexAsset);
 
-		if (SearchIndexAsset->SamplingInterval.Contains(NewAssetTime))
+		if (Database->GetPoseIndicesAndLerpValueFromTime(NewAssetTime, SearchIndexAsset, PrevPoseIdx, PoseIdx, NextPoseIdx, LerpValue))
 		{
-			PoseIdx = Database->GetPoseIndexFromTime(NewAssetTime, SearchIndexAsset);
 			AssetTime = NewAssetTime;
 		}
 		else
@@ -2909,15 +3068,13 @@ void FSearchResult::Update(float NewAssetTime)
 		int32 TriangulationIndex = 0;
 		DbBlendSpace.BlendSpace->GetSamplesFromBlendInput(SearchIndexAsset->BlendParameters, BlendSamples, TriangulationIndex, true);
 
-		float PlayLength = DbBlendSpace.BlendSpace->GetAnimationLengthFromSampleData(BlendSamples);
+		const float PlayLength = DbBlendSpace.BlendSpace->GetAnimationLengthFromSampleData(BlendSamples);
 
 		// Asset player time for blendspaces is normalized [0, 1] so we need to convert 
 		// to a real time before we advance it
-		float RealTime = NewAssetTime * PlayLength;
-
-		if (SearchIndexAsset->SamplingInterval.Contains(RealTime))
+		const float RealTime = NewAssetTime * PlayLength;
+		if (Database->GetPoseIndicesAndLerpValueFromTime(RealTime, SearchIndexAsset, PrevPoseIdx, PoseIdx, NextPoseIdx, LerpValue))
 		{
-			PoseIdx = Database->GetPoseIndexFromTime(RealTime, SearchIndexAsset);
 			AssetTime = NewAssetTime;
 		}
 		else
@@ -2951,7 +3108,6 @@ void FSearchResult::Reset()
 	PoseIdx = INDEX_NONE;
 	SearchIndexAsset = nullptr;
 	Database = nullptr;
-	Sequence = nullptr;
 	ComposedQuery.Reset();
 	AssetTime = 0.0f;
 
diff --git a/Engine/Plugins/Experimental/Animation/PoseSearch/Source/Runtime/Private/PoseSearchFeatureChannels.cpp b/Engine/Plugins/Experimental/Animation/PoseSearch/Source/Runtime/Private/PoseSearchFeatureChannels.cpp
--- a/Engine/Plugins/Experimental/Animation/PoseSearch/Source/Runtime/Private/PoseSearchFeatureChannels.cpp
+++ b/Engine/Plugins/Experimental/Animation/PoseSearch/Source/Runtime/Private/PoseSearchFeatureChannels.cpp
@@ -364,15 +364,17 @@ bool UPoseSearchFeatureChannel_Position::BuildQuery(UE::PoseSearch::FSearchConte
 {
 	using namespace UE::PoseSearch;
 
-	const bool bSkip = bUseFeaturesFromContinuingPose && SearchContext.CurrentResult.IsValid() && SearchContext.CurrentResult.Database->Schema == InOutQuery.GetSchema();
-	if (bSkip)
+	const bool bIsCurrentResultValid = SearchContext.CurrentResult.IsValid();
+	const bool bSkip = bUseFeaturesFromContinuingPose && bIsCurrentResultValid && SearchContext.CurrentResult.Database->Schema == InOutQuery.GetSchema();
+	if (bSkip || !SearchContext.History)
 	{
-		return true;
-	}
-
-	if (!SearchContext.History)
-	{
-		return false;
+		if (bIsCurrentResultValid)
+		{
+			SearchContext.CacheCurrentResultFeatureVectors();
+			int32 DataOffset = ChannelDataOffset;
+			FFeatureVectorHelper::EncodeVector(InOutQuery.EditValues(), DataOffset, SearchContext.CurrentResultPrevPoseVector.GetValues(), SearchContext.CurrentResultPoseVector.GetValues(), SearchContext.CurrentResultNextPoseVector.GetValues(), SearchContext.CurrentResult.LerpValue);
+		}
+		return bSkip;
 	}
 
 	bool AnyError = false;
@@ -497,15 +499,17 @@ bool UPoseSearchFeatureChannel_Heading::BuildQuery(UE::PoseSearch::FSearchContex
 {
 	using namespace UE::PoseSearch;
 
-	const bool bSkip = bUseFeaturesFromContinuingPose && SearchContext.CurrentResult.IsValid() && SearchContext.CurrentResult.Database->Schema == InOutQuery.GetSchema();
-	if (bSkip)
+	const bool bIsCurrentResultValid = SearchContext.CurrentResult.IsValid();
+	const bool bSkip = bUseFeaturesFromContinuingPose && bIsCurrentResultValid && SearchContext.CurrentResult.Database->Schema == InOutQuery.GetSchema();
+	if (bSkip || !SearchContext.History)
 	{
-		return true;
-	}
-
-	if (!SearchContext.History)
-	{
-		return false;
+		if (bIsCurrentResultValid)
+		{
+			SearchContext.CacheCurrentResultFeatureVectors();
+			int32 DataOffset = ChannelDataOffset;
+			FFeatureVectorHelper::EncodeVector(InOutQuery.EditValues(), DataOffset, SearchContext.CurrentResultPrevPoseVector.GetValues(), SearchContext.CurrentResultPoseVector.GetValues(), SearchContext.CurrentResultNextPoseVector.GetValues(), SearchContext.CurrentResult.LerpValue, true);
+		}
+		return bSkip;
 	}
 
 	bool AnyError = false;
@@ -888,15 +892,52 @@ bool UPoseSearchFeatureChannel_Pose::BuildQuery(UE::PoseSearch::FSearchContext&
 {
 	using namespace UE::PoseSearch;
 
-	const bool bSkip = bUseFeaturesFromContinuingPose && SearchContext.CurrentResult.IsValid() && SearchContext.CurrentResult.Database->Schema == InOutQuery.GetSchema();
-	if (bSkip)
+	const bool bIsCurrentResultValid = SearchContext.CurrentResult.IsValid();
+	const bool bSkip = bUseFeaturesFromContinuingPose && bIsCurrentResultValid && SearchContext.CurrentResult.Database->Schema == InOutQuery.GetSchema();
+	if (bSkip || !SearchContext.History)
 	{
-		return true;
-	}
+		if (bIsCurrentResultValid)
+		{
+			SearchContext.CacheCurrentResultFeatureVectors();
 
-	if (!SearchContext.History)
-	{
-		return false;
+			int32 DataOffset = ChannelDataOffset;
+			for (int32 SampledBoneIdx = 0; SampledBoneIdx != SampledBones.Num(); ++SampledBoneIdx)
+			{
+				const FPoseSearchBone& SampledBone = SampledBones[SampledBoneIdx];
+				if (EnumHasAnyFlags(SampledBone.Flags, EPoseSearchBoneFlags::Position))
+				{
+					for (int32 SubsampleIdx = 0; SubsampleIdx != SampleTimes.Num(); ++SubsampleIdx)
+					{
+						FFeatureVectorHelper::EncodeVector(InOutQuery.EditValues(), DataOffset, SearchContext.CurrentResultPrevPoseVector.GetValues(), SearchContext.CurrentResultPoseVector.GetValues(), SearchContext.CurrentResultNextPoseVector.GetValues(), SearchContext.CurrentResult.LerpValue);
+					}
+				}
+
+				if (EnumHasAnyFlags(SampledBone.Flags, EPoseSearchBoneFlags::Rotation))
+				{
+					for (int32 SubsampleIdx = 0; SubsampleIdx != SampleTimes.Num(); ++SubsampleIdx)
+					{
+						FFeatureVectorHelper::EncodeQuat(InOutQuery.EditValues(), DataOffset, SearchContext.CurrentResultPrevPoseVector.GetValues(), SearchContext.CurrentResultPoseVector.GetValues(), SearchContext.CurrentResultNextPoseVector.GetValues(), SearchContext.CurrentResult.LerpValue);
+					}
+				}
+
+				if (EnumHasAnyFlags(SampledBone.Flags, EPoseSearchBoneFlags::Velocity))
+				{
+					for (int32 SubsampleIdx = 0; SubsampleIdx != SampleTimes.Num(); ++SubsampleIdx)
+					{
+						FFeatureVectorHelper::EncodeVector(InOutQuery.EditValues(), DataOffset, SearchContext.CurrentResultPrevPoseVector.GetValues(), SearchContext.CurrentResultPoseVector.GetValues(), SearchContext.CurrentResultNextPoseVector.GetValues(), SearchContext.CurrentResult.LerpValue);
+					}
+				}
+
+				if (EnumHasAnyFlags(SampledBone.Flags, EPoseSearchBoneFlags::Phase))
+				{
+					for (int32 SubsampleIdx = 0; SubsampleIdx != SampleTimes.Num(); ++SubsampleIdx)
+					{
+						FFeatureVectorHelper::EncodeVector2D(InOutQuery.EditValues(), DataOffset, SearchContext.CurrentResultPrevPoseVector.GetValues(), SearchContext.CurrentResultPoseVector.GetValues(), SearchContext.CurrentResultNextPoseVector.GetValues(), SearchContext.CurrentResult.LerpValue);
+					}
+				}
+			}
+		}
+		return bSkip;
 	}
 
 	struct CachedTransforms
@@ -906,7 +947,7 @@ bool UPoseSearchFeatureChannel_Pose::BuildQuery(UE::PoseSearch::FSearchContext&
 		bool Valid = false;
 	};
 	TArray<CachedTransforms> CachedTransforms;
-	CachedTransforms.AddUninitialized(SampleTimes.Num()* SampledBones.Num());
+	CachedTransforms.AddUninitialized(SampleTimes.Num() * SampledBones.Num());
 
 	bool AnyError = false;
 	for (int32 SubsampleIdx = 0; SubsampleIdx != SampleTimes.Num(); ++SubsampleIdx)
@@ -1368,6 +1409,7 @@ bool UPoseSearchFeatureChannel_Trajectory::BuildQuery(UE::PoseSearch::FSearchCon
 
 	if (!SearchContext.Trajectory)
 	{
+		// @todo: do we want to reuse the SearchContext.CurrentResult data if valid?
 		return false;
 	}
 
diff --git a/Engine/Plugins/Experimental/Animation/PoseSearch/Source/Runtime/Private/PoseSearchLibrary.cpp b/Engine/Plugins/Experimental/Animation/PoseSearch/Source/Runtime/Private/PoseSearchLibrary.cpp
--- a/Engine/Plugins/Experimental/Animation/PoseSearch/Source/Runtime/Private/PoseSearchLibrary.cpp
+++ b/Engine/Plugins/Experimental/Animation/PoseSearch/Source/Runtime/Private/PoseSearchLibrary.cpp
@@ -70,7 +70,6 @@ static void ComputeDatabaseBlendSpaceFilter(
 void FMotionMatchingState::Reset()
 {
 	CurrentSearchResult.Reset();
-	AssetPlayerTime = 0.0f;
 	// Set the elapsed time to INFINITY to trigger a search right away
 	ElapsedPoseJumpTime = INFINITY;
 }
@@ -78,7 +77,6 @@ void FMotionMatchingState::Reset()
 void FMotionMatchingState::AdjustAssetTime(float AssetTime)
 {
 	CurrentSearchResult.Update(AssetTime);
-	AssetPlayerTime = CurrentSearchResult.AssetTime;
 }
 
 bool FMotionMatchingState::CanAdvance(float DeltaTime, bool& bOutAdvanceToFollowUpAsset, UE::PoseSearch::FSearchResult& OutFollowUpAsset) const
@@ -99,7 +97,7 @@ bool FMotionMatchingState::CanAdvance(float DeltaTime, bool& bOutAdvanceToFollow
 			CurrentSearchResult.Database->GetSequenceSourceAsset(SearchIndexAsset);
 		const float AssetLength = DbSequence.Sequence->GetPlayLength();
 
-		float SteppedTime = AssetPlayerTime;
+		float SteppedTime = CurrentSearchResult.AssetTime;
 		ETypeAdvanceAnim AdvanceType = FAnimationRuntime::AdvanceTime(
 			DbSequence.Sequence->bLoop,
 			DeltaTime,
@@ -140,7 +138,7 @@ bool FMotionMatchingState::CanAdvance(float DeltaTime, bool& bOutAdvanceToFollow
 				// this is essentially what the matching time in the corresponding main sequence is.
 				// Here we are assuming that the tick will advance the asset player timer into the 
 				// valid region
-				const float FollowUpAssetTime = AssetPlayerTime - AssetLength;
+				const float FollowUpAssetTime = CurrentSearchResult.AssetTime - AssetLength;
 
 				// There is no correspoding pose index when we switch due to what is mentioned above
 				// so for now we just take whatever pose index is associated with the first frame.
@@ -164,7 +162,7 @@ bool FMotionMatchingState::CanAdvance(float DeltaTime, bool& bOutAdvanceToFollow
 
 		// Asset player time for blendspaces is normalized [0, 1] so we need to convert 
 		// to a real time before we advance it
-		float RealTime = AssetPlayerTime * PlayLength;
+		float RealTime = CurrentSearchResult.AssetTime * PlayLength;
 		float SteppedTime = RealTime;
 		ETypeAdvanceAnim AdvanceType = FAnimationRuntime::AdvanceTime(
 			DbBlendSpace.BlendSpace->bLoop,
@@ -205,7 +203,6 @@ void FMotionMatchingState::JumpToPose(const FAnimationUpdateContext& Context, co
 	CurrentSearchResult = Result;
 
 	ElapsedPoseJumpTime = 0.0f;
-	AssetPlayerTime = Result.AssetTime;
 
 	const float JumpBlendTime = ComputeJumpBlendTime(Result, Settings);
 	RequestInertialBlend(Context, JumpBlendTime);
@@ -403,7 +400,7 @@ void UpdateMotionMatchingState(
 		TraceState.SetDatabase(InOutMotionMatchingState.CurrentSearchResult.Database.Get());
 		TraceState.ContinuingPoseIdx = LastResult.PoseIdx;
 
-		TraceState.AssetPlayerTime = InOutMotionMatchingState.AssetPlayerTime;
+		TraceState.AssetPlayerTime = InOutMotionMatchingState.CurrentSearchResult.AssetTime;
 		TraceState.DeltaTime = DeltaTime;
 		TraceState.SimLinearVelocity = SimLinearVelocity;
 		TraceState.SimAngularVelocity = SimAngularVelocity;
diff --git a/Engine/Plugins/Experimental/Animation/PoseSearch/Source/Runtime/Public/PoseSearch/PoseSearch.h b/Engine/Plugins/Experimental/Animation/PoseSearch/Source/Runtime/Public/PoseSearch/PoseSearch.h
--- a/Engine/Plugins/Experimental/Animation/PoseSearch/Source/Runtime/Public/PoseSearch/PoseSearch.h
+++ b/Engine/Plugins/Experimental/Animation/PoseSearch/Source/Runtime/Public/PoseSearch/PoseSearch.h
@@ -734,9 +734,12 @@ struct POSESEARCH_API FPoseSearchIndex
 
 	void Normalize (TArrayView<float> PoseVector) const;
 	void InverseNormalize (TArrayView<float> PoseVector) const;
-};
-
 
+	// individual cost addends calculation methods
+	float ComputeMirrorMismatchAddend(int32 PoseIdx, UE::PoseSearch::FSearchContext& SearchContext) const;
+	float ComputeNotifyAddend(int32 PoseIdx) const;
+	float ComputeContinuingPoseCostAddend(int32 PoseIdx, UE::PoseSearch::EPoseComparisonFlags PoseComparisonFlags) const;
+};
 
 //////////////////////////////////////////////////////////////////////////
 // Database
@@ -1031,9 +1034,15 @@ struct FSearchResult
 {
 	FPoseSearchCost PoseCost;
 	int32 PoseIdx = INDEX_NONE;
+
+	int32 PrevPoseIdx = INDEX_NONE;
+	int32 NextPoseIdx = INDEX_NONE;
+
+	// lerp value to find AssetTime from PrevPoseIdx -> AssetTime -> NextPoseIdx, within range [-0.5, 0.5]
+	float LerpValue = 0.f;
+
 	const FPoseSearchIndexAsset* SearchIndexAsset = nullptr;
 	TWeakObjectPtr<const UPoseSearchDatabase> Database = nullptr;
-	TWeakObjectPtr<const UAnimSequenceBase> Sequence = nullptr;
 	FPoseSearchFeatureVectorBuilder ComposedQuery;
 
 	// cost of the current pose with the query from database in the result, if possible
@@ -1145,8 +1154,9 @@ public:
 	bool IsValidForIndexing() const;
 	bool IsValidForSearch() const;
 
-	int32 GetPoseIndexFromTime(float AssetTime, const FPoseSearchIndexAsset* SearchIndexAsset) const;
 	float GetAssetTime(int32 PoseIdx, const FPoseSearchIndexAsset* SearchIndexAsset = nullptr) const;
+	int32 GetPoseIndexFromTime(float AssetTime, const FPoseSearchIndexAsset* SearchIndexAsset) const;
+	bool GetPoseIndicesAndLerpValueFromTime(float Time, const FPoseSearchIndexAsset* SearchIndexAsset, int32& PrevPoseIdx, int32& PoseIdx, int32& NextPoseIdx, float& LerpValue) const;
 
 	const FPoseSearchDatabaseAnimationAssetBase& GetAnimationSourceAsset(const FPoseSearchIndexAsset* SearchIndexAsset) const;
 	const FPoseSearchDatabaseSequence& GetSequenceSourceAsset(const FPoseSearchIndexAsset* SearchIndexAsset) const;
@@ -1232,7 +1242,6 @@ public:
 
 	FPoseSearchCost ComparePoses(UE::PoseSearch::FSearchContext& SearchContext, int32 PoseIdx, UE::PoseSearch::EPoseComparisonFlags PoseComparisonFlags, int32 GroupIdx, const TArrayView<const float>& QueryValues) const;
 	FPoseSearchCost ComparePoses(UE::PoseSearch::FSearchContext& SearchContext, int32 PoseIdx, UE::PoseSearch::EPoseComparisonFlags PoseComparisonFlags, const TArrayView<const float>& QueryValues, UE::PoseSearch::FPoseCostDetails& OutPoseCostDetails) const;
-	void ComputePoseCostAddends(int32 PoseIdx, UE::PoseSearch::EPoseComparisonFlags PoseComparisonFlags, UE::PoseSearch::FSearchContext& SearchContext, float& OutNotifyAddend, float& OutMirrorMismatchAddend) const;
 
 protected:
 	UE::PoseSearch::FSearchResult SearchPCAKDTree(UE::PoseSearch::FSearchContext& SearchContext) const;
@@ -1288,18 +1297,26 @@ class POSESEARCH_API FFeatureVectorHelper
 public:
 	enum { EncodeQuatCardinality = 6 };
 	static void EncodeQuat(TArrayView<float> Values, int32& DataOffset, const FQuat& Quat);
+	static void EncodeQuat(TArrayView<float> Values, int32& DataOffset, TArrayView<const float> PrevValues, TArrayView<const float> CurPrevValues, TArrayView<const float> NextPrevValues, float LerpValue);
 	static FQuat DecodeQuat(TArrayView<const float> Values, int32& DataOffset);
 
 	enum { EncodeVectorCardinality = 3 };
 	static void EncodeVector(TArrayView<float> Values, int32& DataOffset, const FVector& Vector);
+	static void EncodeVector(TArrayView<float> Values, int32& DataOffset, TArrayView<const float> PrevValues, TArrayView<const float> CurPrevValues, TArrayView<const float> NextPrevValues, float LerpValue, bool bNormalize = false);
 	static FVector DecodeVector(TArrayView<const float> Values, int32& DataOffset);
 
 	enum { EncodeVector2DCardinality = 2 };
 	static void EncodeVector2D(TArrayView<float> Values, int32& DataOffset, const FVector2D& Vector2D);
+	static void EncodeVector2D(TArrayView<float> Values, int32& DataOffset, TArrayView<const float> PrevValues, TArrayView<const float> CurPrevValues, TArrayView<const float> NextPrevValues, float LerpValue);
 	static FVector2D DecodeVector2D(TArrayView<const float> Values, int32& DataOffset);
 
 	// populates MeanDeviations[DataOffset] ... MeanDeviations[DataOffset + Cardinality] with a single value the mean deviation calcualted from a cenetered matrix
 	static void ComputeMeanDeviations(const Eigen::MatrixXd& CenteredPoseMatrix, Eigen::VectorXd& MeanDeviations, int32& DataOffset, int32 Cardinality);
+
+private:
+	static FQuat DecodeQuatInternal(TArrayView<const float> Values, int32 DataOffset);
+	static FVector DecodeVectorInternal(TArrayView<const float> Values, int32 DataOffset);
+	static FVector2D DecodeVector2DInternal(TArrayView<const float> Values, int32 DataOffset);
 };
 
 /**
@@ -1366,9 +1383,14 @@ struct POSESEARCH_API FSearchContext
 	float GetCurrentBestTotalCost() const { return CurrentBestTotalCost; }
 
 	bool GetOrBuildQuery(const UPoseSearchDatabase* Database, FPoseSearchFeatureVectorBuilder& FeatureVectorBuilder);
+	void CacheCurrentResultFeatureVectors();
 
 	static constexpr int8 SchemaRootBoneIdx = -1;
 
+	FPoseSearchFeatureVectorBuilder CurrentResultPoseVector;
+	FPoseSearchFeatureVectorBuilder CurrentResultPrevPoseVector;
+	FPoseSearchFeatureVectorBuilder CurrentResultNextPoseVector;
+
 private:
 	struct FCachedEntry
 	{
diff --git a/Engine/Plugins/Experimental/Animation/PoseSearch/Source/Runtime/Public/PoseSearch/PoseSearchLibrary.h b/Engine/Plugins/Experimental/Animation/PoseSearch/Source/Runtime/Public/PoseSearch/PoseSearchLibrary.h
--- a/Engine/Plugins/Experimental/Animation/PoseSearch/Source/Runtime/Public/PoseSearch/PoseSearchLibrary.h
+++ b/Engine/Plugins/Experimental/Animation/PoseSearch/Source/Runtime/Public/PoseSearch/PoseSearchLibrary.h
@@ -79,10 +79,6 @@ struct POSESEARCH_API FMotionMatchingState
 	UPROPERTY(Transient)
 	float ElapsedPoseJumpTime = 0.f;
 
-	// Current time within the asset player node
-	UPROPERTY(Transient)
-	float AssetPlayerTime = 0.f;
-
 	// Evaluation flags relevant to the state of motion matching
 	UPROPERTY(VisibleAnywhere, BlueprintReadOnly, Category=State)
 	EMotionMatchingFlags Flags = EMotionMatchingFlags::None;
