motion matching - blend stack fix for root motion from everything

diff --git a/Engine/Plugins/Experimental/Animation/PoseSearch/Source/Runtime/Private/AnimNode_BlendStack.cpp b/Engine/Plugins/Experimental/Animation/PoseSearch/Source/Runtime/Private/AnimNode_BlendStack.cpp
--- a/Engine/Plugins/Experimental/Animation/PoseSearch/Source/Runtime/Private/AnimNode_BlendStack.cpp
+++ b/Engine/Plugins/Experimental/Animation/PoseSearch/Source/Runtime/Private/AnimNode_BlendStack.cpp
@@ -124,7 +124,7 @@ void FAnimNode_BlendStack::Evaluate_AnyThread(FPoseContext& Output)
 	const int32 BlendStackSize = AnimPlayers.Num();
 	if (BlendStackSize <= 0)
 	{
-		// nothing to do
+		Output.ResetToRefPose();
 	}
 	else if (BlendStackSize == 1)
 	{
@@ -169,13 +169,17 @@ void FAnimNode_BlendStack::UpdateAssetPlayer(const FAnimationUpdateContext& Cont
 	Super::UpdateAssetPlayer(Context);
 
 	const int32 BlendStackSize = AnimPlayers.Num();
+
+	CalculateWeights();
+	PruneBlendStack(BlendStackSize);
+
 	for (FPoseSearchAnimPlayer& AnimPlayer : AnimPlayers)
 	{
-		AnimPlayer.Update_AnyThread(Context);
+		const FAnimationUpdateContext AnimPlayerContext = Context.FractionalWeightAndRootMotion(AnimPlayer.GetBlendWeight(), AnimPlayer.GetBlendWeight());
+		AnimPlayer.Update_AnyThread(AnimPlayerContext);
+		check(AnimPlayer.GetBlendWeight() == AnimPlayerContext.GetFinalBlendWeight());
+		check(AnimPlayer.GetBlendWeight() == AnimPlayerContext.GetRootMotionWeightModifier());
 	}
-	
-	CalculateWeights();
-	PruneBlendStack(BlendStackSize);
 }
 
 float FAnimNode_BlendStack::GetAccumulatedTime() const
