motion matching - ~UPoseSearchDatabase implementation

diff --git a/Engine/Plugins/Experimental/Animation/PoseSearch/Source/Runtime/Private/PoseSearch.cpp b/Engine/Plugins/Experimental/Animation/PoseSearch/Source/Runtime/Private/PoseSearch.cpp
--- a/Engine/Plugins/Experimental/Animation/PoseSearch/Source/Runtime/Private/PoseSearch.cpp
+++ b/Engine/Plugins/Experimental/Animation/PoseSearch/Source/Runtime/Private/PoseSearch.cpp
@@ -912,6 +912,17 @@ bool FPoseSearchDatabaseSequence::IsLooping() const
 //////////////////////////////////////////////////////////////////////////
 // UPoseSearchDatabase
 
+UPoseSearchDatabase::~UPoseSearchDatabase()
+{
+	if (PrivateDerivedData)
+	{
+#if WITH_EDITOR
+		PrivateDerivedData->CancelCache();
+#endif
+		delete PrivateDerivedData;
+	}
+}
+
 FPoseSearchIndex* UPoseSearchDatabase::GetSearchIndex()
 {
 	if (PrivateDerivedData == nullptr)
diff --git a/Engine/Plugins/Experimental/Animation/PoseSearch/Source/Runtime/Public/PoseSearch/PoseSearch.h b/Engine/Plugins/Experimental/Animation/PoseSearch/Source/Runtime/Public/PoseSearch/PoseSearch.h
--- a/Engine/Plugins/Experimental/Animation/PoseSearch/Source/Runtime/Public/PoseSearch/PoseSearch.h
+++ b/Engine/Plugins/Experimental/Animation/PoseSearch/Source/Runtime/Public/PoseSearch/PoseSearch.h
@@ -1148,6 +1148,8 @@ public:
 	UPROPERTY(EditAnywhere, Category = "Performance")
 	bool bSkipSearchIfPossible = false;
 
+	virtual ~UPoseSearchDatabase();
+
 	FPoseSearchIndex* GetSearchIndex();
 	const FPoseSearchIndex* GetSearchIndex() const;
 
