motion matching - removing invalid checks when motion matching node is blended in or out

diff --git a/Engine/Plugins/Experimental/Animation/PoseSearch/Source/Runtime/Private/AnimNode_BlendStack.cpp b/Engine/Plugins/Experimental/Animation/PoseSearch/Source/Runtime/Private/AnimNode_BlendStack.cpp
--- a/Engine/Plugins/Experimental/Animation/PoseSearch/Source/Runtime/Private/AnimNode_BlendStack.cpp
+++ b/Engine/Plugins/Experimental/Animation/PoseSearch/Source/Runtime/Private/AnimNode_BlendStack.cpp
@@ -177,8 +177,6 @@ void FAnimNode_BlendStack::UpdateAssetPlayer(const FAnimationUpdateContext& Cont
 	{
 		const FAnimationUpdateContext AnimPlayerContext = Context.FractionalWeightAndRootMotion(AnimPlayer.GetBlendWeight(), AnimPlayer.GetBlendWeight());
 		AnimPlayer.Update_AnyThread(AnimPlayerContext);
-		check(AnimPlayer.GetBlendWeight() == AnimPlayerContext.GetFinalBlendWeight());
-		check(AnimPlayer.GetBlendWeight() == AnimPlayerContext.GetRootMotionWeightModifier());
 	}
 }
 
