motion matching: reconstructing missing traj positions in UPoseSearchFeatureChannel_Trajectory::DebugDraw

diff --git a/Engine/Plugins/Experimental/Animation/PoseSearch/Source/Runtime/Private/PoseSearchFeatureChannels.cpp b/Engine/Plugins/Experimental/Animation/PoseSearch/Source/Runtime/Private/PoseSearchFeatureChannels.cpp
--- a/Engine/Plugins/Experimental/Animation/PoseSearch/Source/Runtime/Private/PoseSearchFeatureChannels.cpp
+++ b/Engine/Plugins/Experimental/Animation/PoseSearch/Source/Runtime/Private/PoseSearchFeatureChannels.cpp
@@ -1,6 +1,7 @@
 // Copyright Epic Games, Inc. All Rights Reserved.
 
 #include "PoseSearch/PoseSearchFeatureChannels.h"
+#include "Algo/BinarySearch.h"
 #include "AnimationRuntime.h"
 #include "Animation/AnimationPoseData.h"
 #include "Animation/AttributesRuntime.h"
@@ -1457,6 +1458,103 @@ bool UPoseSearchFeatureChannel_Trajectory::BuildQuery(UE::PoseSearch::FSearchCon
 	return true;
 }
 
+// lazy initialized helper to interpolate or extrapolate (linearly) UPoseSearchFeatureChannel_Trajectory trajectory positions from FPoseSearchTrajectorySample containing EPoseSearchTrajectoryFlags::Position for samples without it
+struct TrajectoryPositionReconstructor
+{
+	struct PositionAndOffsetSample
+	{
+		FVector Position;
+		float Offset;
+	};
+
+	TArray<PositionAndOffsetSample> PositionAndOffsetSamples;
+	bool bInitialized = false;
+
+	void Init(const UPoseSearchFeatureChannel_Trajectory& TrajectoryChannel, TArrayView<const float> PoseVector, const FTransform& RootTransform)
+	{
+		PositionAndOffsetSamples.Reserve(TrajectoryChannel.Samples.Num() + 1);
+		bool bAddZeroOffsetSample = true;
+		int32 DataOffset = TrajectoryChannel.GetChannelDataOffset();
+		for (const FPoseSearchTrajectorySample& Sample : TrajectoryChannel.Samples)
+		{
+			if (EnumHasAnyFlags(Sample.Flags, EPoseSearchTrajectoryFlags::Position))
+			{
+				PositionAndOffsetSample PositionAndOffsetSample;
+				PositionAndOffsetSample.Position = UE::PoseSearch::FFeatureVectorHelper::DecodeVector(PoseVector, DataOffset);
+				PositionAndOffsetSample.Position = RootTransform.TransformPosition(PositionAndOffsetSample.Position);
+				PositionAndOffsetSample.Offset = Sample.Offset;
+				PositionAndOffsetSamples.Add(PositionAndOffsetSample);
+
+				if (FMath::IsNearlyZero(Sample.Offset))
+				{
+					bAddZeroOffsetSample = false;
+				}
+			}
+			if (EnumHasAnyFlags(Sample.Flags, EPoseSearchTrajectoryFlags::Velocity))
+			{
+				DataOffset += UE::PoseSearch::FFeatureVectorHelper::EncodeVectorCardinality;
+			}
+			if (EnumHasAnyFlags(Sample.Flags, EPoseSearchTrajectoryFlags::FacingDirection))
+			{
+				DataOffset += UE::PoseSearch::FFeatureVectorHelper::EncodeVectorCardinality;
+			}
+		}
+
+		if (bAddZeroOffsetSample)
+		{
+			PositionAndOffsetSample PositionAndOffsetSample;
+			PositionAndOffsetSample.Position = RootTransform.GetTranslation();
+			PositionAndOffsetSample.Offset = 0.f;
+			PositionAndOffsetSamples.Add(PositionAndOffsetSample);
+		}
+
+		PositionAndOffsetSamples.Sort([](const PositionAndOffsetSample& a, const PositionAndOffsetSample& b)
+			{
+				return a.Offset < b.Offset;
+			});
+
+		bInitialized = true;
+	}
+
+	FVector GetReconstructedTrajectoryPos(const UPoseSearchFeatureChannel_Trajectory& TrajectoryChannel, TArrayView<const float> PoseVector, const FTransform& RootTransform, float SampleOffset)
+	{
+		if (!bInitialized)
+		{
+			Init(TrajectoryChannel, PoseVector, RootTransform);
+		}
+
+		return GetReconstructedTrajectoryPos(SampleOffset);
+	}
+
+	FVector GetReconstructedTrajectoryPos(float SampleOffset) const
+	{
+		check(bInitialized);
+		check(PositionAndOffsetSamples.Num() > 0);
+		if (PositionAndOffsetSamples.Num() >= 2)
+		{
+			const int32 LowerBoundIdx = Algo::LowerBound(PositionAndOffsetSamples, SampleOffset, [](const PositionAndOffsetSample& PositionAndOffsetSample, float Value)
+				{
+					return Value > PositionAndOffsetSample.Offset;
+				});
+
+			const int32 PrevIdx = FMath::Clamp(LowerBoundIdx, 0, PositionAndOffsetSamples.Num() - 2);
+			const int32 NextIdx = PrevIdx + 1;
+
+			const float Denominator = PositionAndOffsetSamples[NextIdx].Offset - PositionAndOffsetSamples[PrevIdx].Offset;
+			if (FMath::IsNearlyZero(Denominator))
+			{
+				return PositionAndOffsetSamples[PrevIdx].Position;
+			}
+
+			const float Numerator = SampleOffset - PositionAndOffsetSamples[PrevIdx].Offset;
+			const float LerpValue = Numerator / Denominator;
+			return FMath::Lerp(PositionAndOffsetSamples[PrevIdx].Position, PositionAndOffsetSamples[NextIdx].Position, LerpValue);
+		}
+
+		return PositionAndOffsetSamples[0].Position;
+	}
+};
+
 void UPoseSearchFeatureChannel_Trajectory::DebugDraw(const UE::PoseSearch::FDebugDrawParams& DrawParams, TArrayView<const float> PoseVector) const
 {
 	using namespace UE::PoseSearch;
@@ -1484,18 +1582,23 @@ void UPoseSearchFeatureChannel_Trajectory::DebugDraw(const UE::PoseSearch::FDebu
 
 	int32 DataOffset = ChannelDataOffset;
 	int32 SampleIdx = 0;
+	TrajectoryPositionReconstructor TrajectoryPositionReconstructor;
 	for (const FPoseSearchTrajectorySample& Sample : Samples)
 	{
-		FVector TrajectoryPos;
+		FVector TrajectoryPos = FVector::Zero();
+		
 		if (EnumHasAnyFlags(Sample.Flags, EPoseSearchTrajectoryFlags::Position))
 		{
 			TrajectoryPos = FFeatureVectorHelper::DecodeVector(PoseVector, DataOffset);
+			TrajectoryPos = DrawParams.RootTransform.TransformPosition(TrajectoryPos);
+			
+			// validating TrajectoryPositionReconstructor
+			check((TrajectoryPositionReconstructor.GetReconstructedTrajectoryPos(*this, PoseVector, DrawParams.RootTransform, Sample.Offset) - TrajectoryPos).IsNearlyZero());
 
 			FLinearColor LinearColor = DrawParams.GetColor(this);
 			FLinearColor GradientColor = GetGradientColor(LinearColor, SampleIdx, NumSamples, DrawParams.Flags);
 			FColor Color = GradientColor.ToFColor(true);
 
-			TrajectoryPos = DrawParams.RootTransform.TransformPosition(TrajectoryPos);
 			if (EnumHasAnyFlags(DrawParams.Flags, EDebugDrawFlags::DrawFast | EDebugDrawFlags::DrawSearchIndex))
 			{
 				DrawDebugPoint(DrawParams.World, TrajectoryPos, DrawParams.PointSize, Color, bPersistent, LifeTime, DepthPriority);
@@ -1507,7 +1610,7 @@ void UPoseSearchFeatureChannel_Trajectory::DebugDraw(const UE::PoseSearch::FDebu
 		}
 		else
 		{
-			TrajectoryPos = DrawParams.RootTransform.GetTranslation();
+			TrajectoryPos = TrajectoryPositionReconstructor.GetReconstructedTrajectoryPos(*this, PoseVector, DrawParams.RootTransform, Sample.Offset);
 		}
 
 		if (EnumHasAnyFlags(Sample.Flags, EPoseSearchTrajectoryFlags::Velocity))
@@ -1564,7 +1667,7 @@ void UPoseSearchFeatureChannel_Trajectory::DebugDraw(const UE::PoseSearch::FDebu
 				DrawDebugLine(
 					DrawParams.World,
 					TrajectoryPos + TrajectoryForward * DrawDebugSphereSize,
-					TrajectoryPos + TrajectoryForward * DrawDebugSphereSize * 2.0f,
+					TrajectoryPos + TrajectoryForward * DrawDebugSphereSize * 10.0f,
 					Color,
 					bPersistent,
 					LifeTime,
