motion matching - FacingDirection fix and more checks

diff --git a/Engine/Plugins/Experimental/Animation/PoseSearch/Source/Runtime/Private/PoseSearchFeatureChannels.cpp b/Engine/Plugins/Experimental/Animation/PoseSearch/Source/Runtime/Private/PoseSearchFeatureChannels.cpp
--- a/Engine/Plugins/Experimental/Animation/PoseSearch/Source/Runtime/Private/PoseSearchFeatureChannels.cpp
+++ b/Engine/Plugins/Experimental/Animation/PoseSearch/Source/Runtime/Private/PoseSearchFeatureChannels.cpp
@@ -390,7 +390,7 @@ bool UPoseSearchFeatureChannel_Position::BuildQuery(UE::PoseSearch::FSearchConte
 
 	int32 DataOffset = ChannelDataOffset;
 	FFeatureVectorHelper::EncodeVector(InOutQuery.EditValues(), DataOffset, Transform.GetTranslation());
-	
+	check(DataOffset == ChannelDataOffset + ChannelCardinality);
 	return !AnyError;
 }
 
@@ -481,8 +481,8 @@ void UPoseSearchFeatureChannel_Heading::IndexAsset(UE::PoseSearch::IAssetIndexer
 		bool ClampedPresent;
 		const FTransform BoneTransformsPresent = Indexer.GetTransformAndCacheResults(SubsampleTime, bUseSampleTimeOffsetRootBone ? SubsampleTime : OriginSampleTime, SchemaBoneIdx, ClampedPresent);
 		int32 DataOffset = ChannelDataOffset;
-		
 		FFeatureVectorHelper::EncodeVector(FeatureVector.EditValues(), DataOffset, GetAxis(BoneTransformsPresent.GetRotation()));
+		check(DataOffset == ChannelDataOffset + ChannelCardinality);
 	}
 }
 
@@ -509,6 +509,7 @@ bool UPoseSearchFeatureChannel_Heading::BuildQuery(UE::PoseSearch::FSearchContex
 			SearchContext.CacheCurrentResultFeatureVectors();
 			int32 DataOffset = ChannelDataOffset;
 			FFeatureVectorHelper::EncodeVector(InOutQuery.EditValues(), DataOffset, SearchContext.CurrentResultPrevPoseVector.GetValues(), SearchContext.CurrentResultPoseVector.GetValues(), SearchContext.CurrentResultNextPoseVector.GetValues(), SearchContext.CurrentResult.LerpValue, true);
+			check(DataOffset == ChannelDataOffset + ChannelCardinality);
 		}
 		return bSkip;
 	}
@@ -525,6 +526,7 @@ bool UPoseSearchFeatureChannel_Heading::BuildQuery(UE::PoseSearch::FSearchContex
 
 	int32 DataOffset = ChannelDataOffset;
 	FFeatureVectorHelper::EncodeVector(InOutQuery.EditValues(), DataOffset, GetAxis(Transform.GetRotation()));
+	check(DataOffset == ChannelDataOffset + ChannelCardinality);
 
 	return !AnyError;
 }
@@ -546,6 +548,7 @@ void UPoseSearchFeatureChannel_Heading::DebugDraw(const UE::PoseSearch::FDebugDr
 
 	int32 DataOffset = ChannelDataOffset;
 	const FVector BoneHeading = DrawParams.RootTransform.TransformPosition(FFeatureVectorHelper::DecodeVector(PoseVector, DataOffset));
+	check(DataOffset == ChannelDataOffset + ChannelCardinality);
 
 	const FLinearColor LinearColor = DrawParams.GetColor(this);
 	const FColor Color = LinearColor.ToFColor(true);
@@ -879,6 +882,8 @@ void UPoseSearchFeatureChannel_Pose::AddPoseFeatures(UE::PoseSearch::IAssetIndex
 			}
 		}
 	}
+
+	check(DataOffset == ChannelDataOffset + ChannelCardinality);
 }
 
 void UPoseSearchFeatureChannel_Pose::GenerateDDCKey(FBlake3& InOutKeyHasher) const
@@ -1060,6 +1065,8 @@ bool UPoseSearchFeatureChannel_Pose::BuildQuery(UE::PoseSearch::FSearchContext&
 		}
 	}
 
+	check(DataOffset == ChannelDataOffset + ChannelCardinality);
+
 	return true;
 }
 
@@ -1392,9 +1399,10 @@ void UPoseSearchFeatureChannel_Trajectory::IndexAssetPrivate(const UE::PoseSearc
 		if (EnumHasAnyFlags(Sample.Flags, EPoseSearchTrajectoryFlags::FacingDirection))
 		{
 			const FSampleInfo SamplePresent = Indexer.GetSampleInfoRelative(SubsampleTime, Origin);
-			FFeatureVectorHelper::EncodeVector(FeatureVector.EditValues(), DataOffset, Indexer.MirrorTransform(SamplePresent.RootTransform).GetRotation().GetAxisY());
+			FFeatureVectorHelper::EncodeVector(FeatureVector.EditValues(), DataOffset, Indexer.MirrorTransform(SamplePresent.RootTransform).GetRotation().GetForwardVector());
 		}
 	}
+	check(DataOffset == ChannelDataOffset + ChannelCardinality);
 }
 
 void UPoseSearchFeatureChannel_Trajectory::GenerateDDCKey(FBlake3& InOutKeyHasher) const
@@ -1451,9 +1459,10 @@ bool UPoseSearchFeatureChannel_Trajectory::BuildQuery(UE::PoseSearch::FSearchCon
 
 		if (EnumHasAnyFlags(Sample.Flags, EPoseSearchTrajectoryFlags::FacingDirection))
 		{
-			FFeatureVectorHelper::EncodeVector(InOutQuery.EditValues(), DataOffset, TrajectorySample.Transform.GetRotation().GetAxisY());
+			FFeatureVectorHelper::EncodeVector(InOutQuery.EditValues(), DataOffset, TrajectorySample.Transform.GetRotation().GetForwardVector());
 		}
 	}
+	check(DataOffset == ChannelDataOffset + ChannelCardinality);
 
 	return true;
 }
@@ -1509,11 +1518,12 @@ struct TrajectoryPositionReconstructor
 		}
 
 		PositionAndOffsetSamples.Sort([](const PositionAndOffsetSample& a, const PositionAndOffsetSample& b)
-			{
-				return a.Offset < b.Offset;
-			});
+		{
+			return a.Offset < b.Offset;
+		});
 
 		bInitialized = true;
+		check(DataOffset == TrajectoryChannel.GetChannelDataOffset() + TrajectoryChannel.GetChannelCardinality());
 	}
 
 	FVector GetReconstructedTrajectoryPos(const UPoseSearchFeatureChannel_Trajectory& TrajectoryChannel, TArrayView<const float> PoseVector, const FTransform& RootTransform, float SampleOffset)
