motion matching - keeping track of previously selected pose indixes and databases to avoid reselecting them within a tunable time threshold (Settings.PoseJumpHistoryThresholdTime)

diff --git a/Engine/Plugins/Experimental/Animation/PoseSearch/Source/Runtime/Private/PoseSearch.cpp b/Engine/Plugins/Experimental/Animation/PoseSearch/Source/Runtime/Private/PoseSearch.cpp
--- a/Engine/Plugins/Experimental/Animation/PoseSearch/Source/Runtime/Private/PoseSearch.cpp
+++ b/Engine/Plugins/Experimental/Animation/PoseSearch/Source/Runtime/Private/PoseSearch.cpp
@@ -175,22 +175,22 @@ static void CalcChannelCosts(const UPoseSearchSchema* Schema, TArrayView<const f
 	}
 }
 
-static int32 PopulateNonSelectableIdx(size_t* NonSelectableIdx, int32 NonSelectableIdxSize, const FSearchContext& SearchContext, int32 IdxOffset = 0)
+static int32 PopulateNonSelectableIdx(size_t* NonSelectableIdx, int32 NonSelectableIdxSize, const FSearchContext& SearchContext, const UPoseSearchDatabase* Database, int32 IdxOffset = 0)
 {
-	check(SearchContext.CurrentResult.IsValid());
+	check(Database);
 
 	int32 NonSelectableIdxUsedSize = 0;
-	if (SearchContext.PoseJumpThresholdTime > 0.f)
+	if (SearchContext.IsCurrentResultFromDatabase(Database) && SearchContext.CurrentResult.IsValid() && SearchContext.PoseJumpThresholdTime > 0.f)
 	{
-		const UPoseSearchDatabase* Database = SearchContext.CurrentResult.Database.Get();
+		const UPoseSearchDatabase* CurrentResultDatabase = SearchContext.CurrentResult.Database.Get();
 		const FPoseSearchIndexAsset* CurrentIndexAsset = SearchContext.CurrentResult.SearchIndexAsset;
-		const int32 PoseJumpIndexThreshold = FMath::FloorToInt(SearchContext.PoseJumpThresholdTime / Database->Schema->SamplingInterval);
-		const bool IsLooping = Database->IsSourceAssetLooping(CurrentIndexAsset);
+		const int32 PoseJumpIndexThreshold = FMath::FloorToInt(SearchContext.PoseJumpThresholdTime / CurrentResultDatabase->Schema->SamplingInterval);
+		const bool IsLooping = CurrentResultDatabase->IsSourceAssetLooping(CurrentIndexAsset);
 
 		for (int32 i = -PoseJumpIndexThreshold; i <= -1; ++i)
 		{
 			const int32 PoseIdx = SearchContext.CurrentResult.PoseIdx + i;
-			const float DeltaTime = i * Database->Schema->SamplingInterval;
+			const float DeltaTime = i * CurrentResultDatabase->Schema->SamplingInterval;
 
 			// @todo: should we use the quantized time associated to InOutMotionMatchingState.DbPoseIdx instead of InOutMotionMatchingState.AssetPlayerTime?
 			float PoseAssetPlayerTime = SearchContext.CurrentResult.AssetTime + DeltaTime;
@@ -215,7 +215,7 @@ static int32 PopulateNonSelectableIdx(size_t* NonSelectableIdx, int32 NonSelecta
 		for (int32 i = 0; i <= PoseJumpIndexThreshold; ++i)
 		{
 			const int32 PoseIdx = SearchContext.CurrentResult.PoseIdx + i;
-			const float DeltaTime = i * Database->Schema->SamplingInterval;
+			const float DeltaTime = i * CurrentResultDatabase->Schema->SamplingInterval;
 
 			// @todo: should we use the quantized time associated to InOutMotionMatchingState.DbPoseIdx instead of InOutMotionMatchingState.AssetPlayerTime?
 			float PoseAssetPlayerTime = SearchContext.CurrentResult.AssetTime + DeltaTime;
@@ -233,6 +233,28 @@ static int32 PopulateNonSelectableIdx(size_t* NonSelectableIdx, int32 NonSelecta
 			}
 			else
 			{
+				UE_LOG(LogPoseSearch, Warning, TEXT("PopulateNonSelectableIdx couldn't add all the NonSelectableIdx"));
+				break;
+			}
+		}
+	}
+
+	if (SearchContext.PoseIndicesHistory)
+	{
+		const FObjectKey DatabaseKey(Database);
+		for (auto It = SearchContext.PoseIndicesHistory->IndexToTime.CreateConstIterator(); It; ++It)
+		{
+			if (NonSelectableIdxUsedSize < NonSelectableIdxSize)
+			{
+				const FHistoricalPoseIndex& HistoricalPoseIndex = It.Key();
+				if (HistoricalPoseIndex.DatabaseKey == DatabaseKey)
+				{
+					NonSelectableIdx[NonSelectableIdxUsedSize++] = HistoricalPoseIndex.PoseIndex + IdxOffset;
+				}
+			}
+			else
+			{
+				UE_LOG(LogPoseSearch, Warning, TEXT("PopulateNonSelectableIdx couldn't add all the NonSelectableIdx"));
 				break;
 			}
 		}
@@ -1058,15 +1080,14 @@ const bool UPoseSearchDatabase::IsSourceAssetLooping(const FPoseSearchIndexAsset
 	{
 		return GetSequenceSourceAsset(SearchIndexAsset).Sequence->bLoop;
 	}
-	else if (SearchIndexAsset->Type == ESearchIndexAssetType::BlendSpace)
+
+	if (SearchIndexAsset->Type == ESearchIndexAssetType::BlendSpace)
 	{
 		return GetBlendSpaceSourceAsset(SearchIndexAsset).BlendSpace->bLoop;
 	}
-	else
-	{
-		checkNoEntry();
-		return false;
-	}
+	
+	checkNoEntry();
+	return false;
 }
 
 const FGameplayTagContainer* UPoseSearchDatabase::GetSourceAssetGroupTags(const FPoseSearchIndexAsset* SearchIndexAsset) const
@@ -1075,15 +1096,14 @@ const FGameplayTagContainer* UPoseSearchDatabase::GetSourceAssetGroupTags(const
 	{
 		return &GetSequenceSourceAsset(SearchIndexAsset).GroupTags;
 	}
-	else if (SearchIndexAsset->Type == ESearchIndexAssetType::BlendSpace)
+	
+	if (SearchIndexAsset->Type == ESearchIndexAssetType::BlendSpace)
 	{
 		return &GetBlendSpaceSourceAsset(SearchIndexAsset).GroupTags;
 	}
-	else
-	{
-		checkNoEntry();
-		return nullptr;
-	}
+	
+	checkNoEntry();
+	return nullptr;
 }
 
 const FString UPoseSearchDatabase::GetSourceAssetName(const FPoseSearchIndexAsset* SearchIndexAsset) const
@@ -1092,15 +1112,14 @@ const FString UPoseSearchDatabase::GetSourceAssetName(const FPoseSearchIndexAsse
 	{
 		return GetSequenceSourceAsset(SearchIndexAsset).Sequence->GetName();
 	}
-	else if (SearchIndexAsset->Type == ESearchIndexAssetType::BlendSpace)
+	
+	if (SearchIndexAsset->Type == ESearchIndexAssetType::BlendSpace)
 	{
 		return GetBlendSpaceSourceAsset(SearchIndexAsset).BlendSpace->GetName();
 	}
-	else
-	{
-		checkNoEntry();
-		return FString();
-	}
+	
+	checkNoEntry();
+	return FString();
 }
 
 int32 UPoseSearchDatabase::GetNumberOfPrincipalComponents() const
@@ -2126,8 +2145,8 @@ UE::PoseSearch::FSearchResult UPoseSearchDatabase::SearchPCAKDTree(UE::PoseSearc
 
 	TArrayView<const float> NormalizedQueryValues = Result.ComposedQuery.GetNormalizedValues();
 
-	const bool IsCurrentResultFromThisDatabase = !SearchContext.bForceInterrupt && SearchContext.CurrentResult.IsValid() && SearchContext.CurrentResult.Database == this;
-	
+	const bool IsCurrentResultFromThisDatabase = SearchContext.IsCurrentResultFromDatabase(this);
+
 	// evaluating the continuing pose only if it hasn't already being evaluated and the related animation can advance
 	if (IsCurrentResultFromThisDatabase && SearchContext.bCanAdvance && !Result.ContinuingPoseCost.IsValid())
 	{
@@ -2151,7 +2170,7 @@ UE::PoseSearch::FSearchResult UPoseSearchDatabase::SearchPCAKDTree(UE::PoseSearc
 		for (const FGroupSearchIndex& GroupSearchIndex : SearchIndex->Groups)
 		{
 			// we're offsetting the NonSelectableIdx by "-GroupSearchIndex.StartPoseIndex" to have the indexes in kdtree space rather than in PoseIdx database space (kdtree index 0 is PoseIdx of GroupSearchIndex.StartPoseIndex)
-			const int NonSelectableIdxUsedSize = IsCurrentResultFromThisDatabase ? PopulateNonSelectableIdx(NonSelectableIdxData, NonSelectableIdxDataSize, SearchContext, -GroupSearchIndex.StartPoseIndex) : 0;
+			const int NonSelectableIdxUsedSize = PopulateNonSelectableIdx(NonSelectableIdxData, NonSelectableIdxDataSize, SearchContext, this, -GroupSearchIndex.StartPoseIndex);
 			TArrayView<size_t> NonSelectableIdx(NonSelectableIdxData, NonSelectableIdxUsedSize);
 			const RowMajorVectorMapConst MapWeights(GroupSearchIndex.Weights.GetData(), 1, NumDimensions);
 			FKDTree::KNNResultSet ResultSet(ClampedKDTreeQueryNumNeighbors, ResultIndexes, ResultDistanceSqr, NonSelectableIdx);
@@ -2242,10 +2261,7 @@ UE::PoseSearch::FSearchResult UPoseSearchDatabase::SearchBruteForce(UE::PoseSear
 	SearchContext.GetOrBuildQuery(this, Result.ComposedQuery);
 	TArrayView<const float> NormalizedQueryValues = Result.ComposedQuery.GetNormalizedValues();
 
-	constexpr int NonSelectableIdxDataSize = 128;
-	size_t* NonSelectableIdxData((size_t*)FMemory_Alloca(NonSelectableIdxDataSize * sizeof(size_t)));
-	int NonSelectableIdxUsedSize = 0;
-	const bool IsCurrentResultFromThisDatabase = !SearchContext.bForceInterrupt && SearchContext.CurrentResult.IsValid() && SearchContext.CurrentResult.Database == this;
+	const bool IsCurrentResultFromThisDatabase = SearchContext.IsCurrentResultFromDatabase(this);
 	if (IsCurrentResultFromThisDatabase)
 	{
 		// evaluating the continuing pose only if it hasn't already being evaluated and the related animation can advance
@@ -2263,10 +2279,11 @@ UE::PoseSearch::FSearchResult UPoseSearchDatabase::SearchBruteForce(UE::PoseSear
 				SearchContext.UpdateCurrentBestCost(Result.ContinuingPoseCost);
 			}
 		}
-
-		NonSelectableIdxUsedSize = PopulateNonSelectableIdx(NonSelectableIdxData, NonSelectableIdxDataSize, SearchContext);
 	}
 
+	constexpr int NonSelectableIdxDataSize = 128;
+	size_t* NonSelectableIdxData((size_t*)FMemory_Alloca(NonSelectableIdxDataSize * sizeof(size_t)));
+	const int NonSelectableIdxUsedSize = PopulateNonSelectableIdx(NonSelectableIdxData, NonSelectableIdxDataSize, SearchContext, this);
 	TArrayView<size_t> NonSelectableIdx(NonSelectableIdxData, NonSelectableIdxUsedSize);
 
 	FPoseSearchCost BestPoseCost;
@@ -2509,6 +2526,33 @@ void FPoseSearchFeatureVectorBuilder::Normalize(const FPoseSearchIndex& ForSearc
 namespace UE { namespace PoseSearch
 {
 
+void FPoseIndicesHistory::Update(const FSearchResult& SearchResult, float DeltaTime, float MaxTime)
+{
+	if (MaxTime > 0.f)
+	{
+		for (auto It = IndexToTime.CreateIterator(); It; ++It)
+		{
+			It.Value() += DeltaTime;
+			if (It.Value() > MaxTime)
+			{
+				It.RemoveCurrent();
+			}
+		}
+
+		if (SearchResult.IsValid())
+		{
+			FHistoricalPoseIndex HistoricalPoseIndex;
+			HistoricalPoseIndex.PoseIndex = SearchResult.PoseIdx;
+			HistoricalPoseIndex.DatabaseKey = FObjectKey(SearchResult.Database.Get());
+			IndexToTime.Add(HistoricalPoseIndex, 0.f);
+		}
+	}
+	else
+	{
+		IndexToTime.Reset();
+	}
+}
+
 FTransform FSearchContext::TryGetTransformAndCacheResults(float SampleTime, const UPoseSearchSchema* Schema, int8 SchemaBoneIdx, bool& Error)
 {
 	check(History && Schema);
@@ -2628,6 +2672,11 @@ void FSearchContext::CacheCurrentResultFeatureVectors()
 	}
 }
 
+bool FSearchContext::IsCurrentResultFromDatabase(const UPoseSearchDatabase* Database) const
+{
+	return !bForceInterrupt && CurrentResult.IsValid() && CurrentResult.Database == Database;
+}
+
 //////////////////////////////////////////////////////////////////////////
 // FPoseHistory
 
diff --git a/Engine/Plugins/Experimental/Animation/PoseSearch/Source/Runtime/Public/PoseSearch/PoseSearch.h b/Engine/Plugins/Experimental/Animation/PoseSearch/Source/Runtime/Public/PoseSearch/PoseSearch.h
--- a/Engine/Plugins/Experimental/Animation/PoseSearch/Source/Runtime/Public/PoseSearch/PoseSearch.h
+++ b/Engine/Plugins/Experimental/Animation/PoseSearch/Source/Runtime/Public/PoseSearch/PoseSearch.h
@@ -1361,6 +1361,29 @@ private:
 	float TimeHorizon = 0.0f;
 };
 
+struct FHistoricalPoseIndex
+{
+	bool operator==(const FHistoricalPoseIndex& Index) const
+	{
+		return PoseIndex == Index.PoseIndex && DatabaseKey == Index.DatabaseKey;
+	}
+
+	friend FORCEINLINE uint32 GetTypeHash(const FHistoricalPoseIndex& Index)
+	{
+		return GetTypeHash(Index.PoseIndex) ^ GetTypeHash(Index.DatabaseKey);
+	}
+
+	int32 PoseIndex = -1;
+	FObjectKey DatabaseKey;
+};
+
+struct FPoseIndicesHistory
+{
+	void Update(const FSearchResult& SearchResult, float DeltaTime, float MaxTime);
+	void Reset() { IndexToTime.Reset(); }
+	TMap<FHistoricalPoseIndex, float> IndexToTime;
+};
+
 struct POSESEARCH_API FSearchContext
 {
 	EPoseSearchBooleanRequest QueryMirrorRequest = EPoseSearchBooleanRequest::Indifferent;
@@ -1387,11 +1410,14 @@ struct POSESEARCH_API FSearchContext
 	bool GetOrBuildQuery(const UPoseSearchDatabase* Database, FPoseSearchFeatureVectorBuilder& FeatureVectorBuilder);
 	void CacheCurrentResultFeatureVectors();
 
+	bool IsCurrentResultFromDatabase(const UPoseSearchDatabase* Database) const;
+
 	static constexpr int8 SchemaRootBoneIdx = -1;
 
 	FPoseSearchFeatureVectorBuilder CurrentResultPoseVector;
 	FPoseSearchFeatureVectorBuilder CurrentResultPrevPoseVector;
 	FPoseSearchFeatureVectorBuilder CurrentResultNextPoseVector;
+	const FPoseIndicesHistory* PoseIndicesHistory = nullptr;
 
 private:
 	struct FCachedEntry
diff --git a/Engine/Plugins/Experimental/Animation/PoseSearch/Source/Runtime/Public/PoseSearch/PoseSearchLibrary.h b/Engine/Plugins/Experimental/Animation/PoseSearch/Source/Runtime/Public/PoseSearch/PoseSearchLibrary.h
--- a/Engine/Plugins/Experimental/Animation/PoseSearch/Source/Runtime/Public/PoseSearch/PoseSearchLibrary.h
+++ b/Engine/Plugins/Experimental/Animation/PoseSearch/Source/Runtime/Public/PoseSearch/PoseSearchLibrary.h
@@ -43,6 +43,10 @@ struct POSESEARCH_API FMotionMatchingSettings
 	UPROPERTY(EditAnywhere, BlueprintReadWrite, Category=Settings, meta=(ClampMin="0"))
 	float PoseJumpThresholdTime = 1.f;
 
+	// Don't jump to poses that has been selected previously within this many seconds in the past
+	UPROPERTY(EditAnywhere, BlueprintReadWrite, Category = Settings, meta = (ClampMin = "0"))
+	float PoseReselectHistory = 0.f;
+
 	// Minimum amount of time to wait between pose search queries
 	UPROPERTY(EditAnywhere, BlueprintReadWrite, Category=Settings, meta=(ClampMin="0"))
 	float SearchThrottleTime = 0.1f;
@@ -92,6 +96,8 @@ struct POSESEARCH_API FMotionMatchingState
 	FDebugFloatHistory SearchCostHistoryBruteForce;
 	FDebugFloatHistory SearchCostHistoryKDTree;
 #endif
+
+	UE::PoseSearch::FPoseIndicesHistory PoseIndicesHistory;
 };
 
 /**
