motion matching - missing file for "keeping track of previously selected pose indixes and databases to avoid reselecting them within a tunable time threshold"

diff --git a/Engine/Plugins/Experimental/Animation/PoseSearch/Source/Runtime/Private/PoseSearchLibrary.cpp b/Engine/Plugins/Experimental/Animation/PoseSearch/Source/Runtime/Private/PoseSearchLibrary.cpp
--- a/Engine/Plugins/Experimental/Animation/PoseSearch/Source/Runtime/Private/PoseSearchLibrary.cpp
+++ b/Engine/Plugins/Experimental/Animation/PoseSearch/Source/Runtime/Private/PoseSearchLibrary.cpp
@@ -13,8 +13,8 @@
 
 #define LOCTEXT_NAMESPACE "PoseSearchLibrary"
 
-namespace UE::PoseSearch {
-
+namespace UE::PoseSearch
+{
 static void ComputeDatabaseSequenceFilter(
 	const UPoseSearchDatabase* Database, 
 	const FGameplayTagQuery* Query, 
@@ -63,7 +63,6 @@ static void ComputeDatabaseBlendSpaceFilter(
 
 } // namespace UE::PoseSearch
 
-
 //////////////////////////////////////////////////////////////////////////
 // FMotionMatchingState
 
@@ -72,6 +71,7 @@ void FMotionMatchingState::Reset()
 	CurrentSearchResult.Reset();
 	// Set the elapsed time to INFINITY to trigger a search right away
 	ElapsedPoseJumpTime = INFINITY;
+	PoseIndicesHistory.Reset();
 }
 
 void FMotionMatchingState::AdjustAssetTime(float AssetTime)
@@ -265,6 +265,7 @@ void UpdateMotionMatchingState(
 
 		SearchContext.CurrentResult = InOutMotionMatchingState.CurrentSearchResult;
 		SearchContext.PoseJumpThresholdTime = Settings.PoseJumpThresholdTime;
+		SearchContext.PoseIndicesHistory = &InOutMotionMatchingState.PoseIndicesHistory;
 
 		IPoseHistoryProvider* PoseHistoryProvider = Context.GetMessage<IPoseHistoryProvider>();
 		if (PoseHistoryProvider)
@@ -291,7 +292,7 @@ void UpdateMotionMatchingState(
 		}
 
 		// todo: cache the queries into TraceState, one per TraceState.DatabaseEntries or
-		// InOutMotionMatchingState.CurrentSearchResult.ComposedQuery = SearchResult.ComposedQuery;
+		InOutMotionMatchingState.CurrentSearchResult.ComposedQuery = SearchResult.ComposedQuery;
 	}
 
 	// If we didn't search or it didn't find a pose to jump to, and we can 
@@ -311,6 +312,8 @@ void UpdateMotionMatchingState(
 		InOutMotionMatchingState.ElapsedPoseJumpTime += DeltaTime;
 	}
 
+	InOutMotionMatchingState.PoseIndicesHistory.Update(InOutMotionMatchingState.CurrentSearchResult, DeltaTime, Settings.PoseReselectHistory);
+	
 	// Record debugger details
 #if UE_POSE_SEARCH_TRACE_ENABLED
 	FTraceMotionMatchingState TraceState;
