CIS fix

#rnx

diff --git a/Engine/Plugins/Experimental/Animation/PoseSearch/Source/Runtime/Public/PoseSearch/PoseSearch.h b/Engine/Plugins/Experimental/Animation/PoseSearch/Source/Runtime/Public/PoseSearch/PoseSearch.h
--- a/Engine/Plugins/Experimental/Animation/PoseSearch/Source/Runtime/Public/PoseSearch/PoseSearch.h
+++ b/Engine/Plugins/Experimental/Animation/PoseSearch/Source/Runtime/Public/PoseSearch/PoseSearch.h
@@ -9,6 +9,7 @@
 #include "Misc/EnumClassFlags.h"
 #include "UObject/Object.h"
 #include "UObject/ObjectMacros.h"
+#include "Templates/TypeHash.h"
 #include "Animation/AnimMetaData.h"
 #include "Animation/AnimNodeMessages.h"
 #include "Animation/AnimationAsset.h"
@@ -1370,7 +1371,7 @@ struct FHistoricalPoseIndex
 
 	friend FORCEINLINE uint32 GetTypeHash(const FHistoricalPoseIndex& Index)
 	{
-		return GetTypeHash(Index.PoseIndex) ^ GetTypeHash(Index.DatabaseKey);
+		return ::GetTypeHash(Index.PoseIndex) ^ GetTypeHash(Index.DatabaseKey);
 	}
 
 	int32 PoseIndex = -1;
