motion matching - using HashCombineFast in FHistoricalPoseIndex::GetTypeHash

diff --git a/Engine/Plugins/Experimental/Animation/PoseSearch/Source/Runtime/Public/PoseSearch/PoseSearch.h b/Engine/Plugins/Experimental/Animation/PoseSearch/Source/Runtime/Public/PoseSearch/PoseSearch.h
--- a/Engine/Plugins/Experimental/Animation/PoseSearch/Source/Runtime/Public/PoseSearch/PoseSearch.h
+++ b/Engine/Plugins/Experimental/Animation/PoseSearch/Source/Runtime/Public/PoseSearch/PoseSearch.h
@@ -1371,7 +1371,7 @@ struct FHistoricalPoseIndex
 
 	friend FORCEINLINE uint32 GetTypeHash(const FHistoricalPoseIndex& Index)
 	{
-		return ::GetTypeHash(Index.PoseIndex) ^ GetTypeHash(Index.DatabaseKey);
+		return HashCombineFast(::GetTypeHash(Index.PoseIndex), GetTypeHash(Index.DatabaseKey));
 	}
 
 	int32 PoseIndex = -1;
