motion matching - WIP tunable pose feature channels debug coloring

diff --git a/Engine/Plugins/Experimental/Animation/PoseSearch/Source/Editor/Private/PoseSearchDatabasePreviewScene.cpp b/Engine/Plugins/Experimental/Animation/PoseSearch/Source/Editor/Private/PoseSearchDatabasePreviewScene.cpp
--- a/Engine/Plugins/Experimental/Animation/PoseSearch/Source/Editor/Private/PoseSearchDatabasePreviewScene.cpp
+++ b/Engine/Plugins/Experimental/Animation/PoseSearch/Source/Editor/Private/PoseSearchDatabasePreviewScene.cpp
@@ -69,14 +69,12 @@ namespace UE::PoseSearch
 					DrawParams.RootTransform = PreviewActor.Mesh->GetComponentTransform();
 					DrawParams.Database = ViewModel->GetPoseSearchDatabase();
 					DrawParams.World = GetWorld();
-					DrawParams.DefaultLifeTime = 0.0f;
-					DrawParams.PoseIdx = PreviewActor.CurrentPoseIndex;
-					DrawParams.PointSize = 5.0f;
+					DrawParams.DefaultLifeTime = 0.f;
+					DrawParams.PointSize = 5.f;
 					DrawParams.Mesh = PreviewActor.Mesh;
 
 					EnumAddFlags(DrawParams.Flags, UE::PoseSearch::EDebugDrawFlags::DrawFast);
-
-					UE::PoseSearch::Draw(DrawParams);
+					DrawFeatureVector(DrawParams, PreviewActor.CurrentPoseIndex);
 				}
 			}
 		}
diff --git a/Engine/Plugins/Experimental/Animation/PoseSearch/Source/Editor/Private/PoseSearchDebugger.cpp b/Engine/Plugins/Experimental/Animation/PoseSearch/Source/Editor/Private/PoseSearchDebugger.cpp
--- a/Engine/Plugins/Experimental/Animation/PoseSearch/Source/Editor/Private/PoseSearchDebugger.cpp
+++ b/Engine/Plugins/Experimental/Animation/PoseSearch/Source/Editor/Private/PoseSearchDebugger.cpp
@@ -1742,7 +1742,6 @@ void SDebuggerView::DrawFeatures(
 	// Set shared state
 	FDebugDrawParams DrawParams;
 	DrawParams.Database = &Database;
-	DrawParams.PoseVector = State.QueryVector;
 	DrawParams.World = &DebuggerWorld;
 	DrawParams.RootTransform = Transform;
 	// Single frame render
@@ -1753,47 +1752,41 @@ void SDebuggerView::DrawFeatures(
 	auto SetDrawFlags = [](FDebugDrawParams& InDrawParams, const FPoseSearchDebuggerFeatureDrawOptions& Options)
 	{
 		InDrawParams.Flags = EDebugDrawFlags::None;
-		if (Options.bDisable)
+		if (!Options.bDisable)
 		{
-			return;
-		}
-
-		if (Options.bDrawBoneNames)
-		{
-			InDrawParams.Flags |= EDebugDrawFlags::DrawBoneNames;
-		}
-
-		if (Options.bDrawSampleLabels)
-		{
-			InDrawParams.Flags |= EDebugDrawFlags::DrawSampleLabels;
-		}
+			if (Options.bDrawBoneNames)
+			{
+				EnumAddFlags(InDrawParams.Flags, EDebugDrawFlags::DrawBoneNames);
+			}
+			
+			if (Options.bDrawSampleLabels)
+			{
+				EnumAddFlags(InDrawParams.Flags, EDebugDrawFlags::DrawSampleLabels);
+			}
 
-		if (Options.bDrawSamplesWithColorGradient)
-		{
-			InDrawParams.Flags |= EDebugDrawFlags::DrawSamplesWithColorGradient;
+			if (Options.bDrawSamplesWithColorGradient)
+			{
+				EnumAddFlags(InDrawParams.Flags, EDebugDrawFlags::DrawSamplesWithColorGradient);
+			}
 		}
 	};
 
 	// Draw query vector
-	DrawParams.Color = &FLinearColor::Blue;
 	SetDrawFlags(DrawParams, Reflection->QueryDrawOptions);
-	DrawParams.LabelPrefix = TEXT("Q");
-	Draw(DrawParams);
-	DrawParams.PoseVector = {};
+	EnumAddFlags(DrawParams.Flags, EDebugDrawFlags::DrawQuery);
+	DrawFeatureVector(DrawParams, State.QueryVector);
+	EnumRemoveFlags(DrawParams.Flags, EDebugDrawFlags::DrawQuery);
 
 	const TSharedPtr<SListView<TSharedRef<FDebuggerDatabaseRowData>>>& DatabaseRows = DatabaseView->GetDatabaseRows();
 	TArray<TSharedRef<FDebuggerDatabaseRowData>> Selected = DatabaseRows->GetSelectedItems();
-
+	
 	// Red for non-active database view
-	DrawParams.Color = &FLinearColor::Red;
-	DrawParams.LabelPrefix = TEXT("S");
 	SetDrawFlags(DrawParams, Reflection->SelectedPoseDrawOptions);
 
 	// Draw any selected database vectors
 	for (const TSharedRef<FDebuggerDatabaseRowData>& Row : Selected)
 	{
-		DrawParams.PoseIdx = Row->PoseIdx;
-		Draw(DrawParams);
+		DrawFeatureVector(DrawParams, Row->PoseIdx);
 	}
 
 	Selected = DatabaseView->GetActiveRow()->GetSelectedItems();
@@ -1803,13 +1796,8 @@ void SDebuggerView::DrawFeatures(
 
 	if (!Selected.IsEmpty())
 	{
-		// Green for the active view
-		DrawParams.Color = &FLinearColor::Green;
-		
 		// Use the motion-matching state's pose idx, as the active row may be update-throttled at this point
-		DrawParams.PoseIdx = State.DbPoseIdx;
-		DrawParams.LabelPrefix = TEXT("A");
-		Draw(DrawParams);
+		DrawFeatureVector(DrawParams, State.DbPoseIdx);
 	}
 
 	Selected = DatabaseView->GetContinuingPoseRow()->GetSelectedItems();
@@ -1819,13 +1807,8 @@ void SDebuggerView::DrawFeatures(
 
 	if (!Selected.IsEmpty())
 	{
-		// Green for the ContinuingPose view
-		DrawParams.Color = &FLinearColor::Gray;
-
 		// Use the motion-matching state's pose idx, as the ContinuingPose row may be update-throttled at this point
-		DrawParams.PoseIdx = (*Selected.begin())->PoseIdx;
-		DrawParams.LabelPrefix = TEXT("C");
-		Draw(DrawParams);
+		DrawFeatureVector(DrawParams, (*Selected.begin())->PoseIdx);
 	}
 
 	FSkeletonDrawParams SkeletonDrawParams;
diff --git a/Engine/Plugins/Experimental/Animation/PoseSearch/Source/Editor/Private/PoseSearchEditor.cpp b/Engine/Plugins/Experimental/Animation/PoseSearch/Source/Editor/Private/PoseSearchEditor.cpp
--- a/Engine/Plugins/Experimental/Animation/PoseSearch/Source/Editor/Private/PoseSearchEditor.cpp
+++ b/Engine/Plugins/Experimental/Animation/PoseSearch/Source/Editor/Private/PoseSearchEditor.cpp
@@ -26,18 +26,17 @@ DEFINE_LOG_CATEGORY(LogPoseSearchEditor);
 //////////////////////////////////////////////////////////////////////////
 // FEditorCommands
 
-namespace UE::PoseSearch {
+namespace UE::PoseSearch
+{
 
-struct FEditorCommands
+namespace FEditorCommands
 {
-	static void DrawSearchIndex();
-};
 
-void FEditorCommands::DrawSearchIndex()
+void DrawSearchIndex()
 {
 	FDebugDrawParams DrawParams;
 	DrawParams.DefaultLifeTime = 60.0f;
-	DrawParams.Flags = EDebugDrawFlags::DrawSearchIndex;
+	EnumAddFlags(DrawParams.Flags, EDebugDrawFlags::DrawSearchIndex);
 
 	TArray<UObject*> EditedAssets = GEditor->GetEditorSubsystem<UAssetEditorSubsystem>()->GetAllEditedAssets();
 	for (UObject* EditedAsset : EditedAssets)
@@ -56,13 +55,14 @@ void FEditorCommands::DrawSearchIndex()
 					TSharedRef<IPersonaPreviewScene> Scene = Toolkit->GetPreviewScene();
 
 					DrawParams.World = Scene->GetWorld();
-					Draw(DrawParams);
+					DrawSearchIndex(DrawParams);
 				}
 			}
 		}
 	}
 }
 
+} // namespace FEditorCommands
 
 //////////////////////////////////////////////////////////////////////////
 // FPoseSearchEditorModule
diff --git a/Engine/Plugins/Experimental/Animation/PoseSearch/Source/Runtime/Private/AnimNode_MotionMatching.cpp b/Engine/Plugins/Experimental/Animation/PoseSearch/Source/Runtime/Private/AnimNode_MotionMatching.cpp
--- a/Engine/Plugins/Experimental/Animation/PoseSearch/Source/Runtime/Private/AnimNode_MotionMatching.cpp
+++ b/Engine/Plugins/Experimental/Animation/PoseSearch/Source/Runtime/Private/AnimNode_MotionMatching.cpp
@@ -136,6 +136,8 @@ bool FAnimNode_MotionMatching::HasPreUpdate() const
 void FAnimNode_MotionMatching::PreUpdate(const UAnimInstance* InAnimInstance)
 {
 #if WITH_EDITORONLY_DATA
+	using namespace UE::PoseSearch;
+
 	if (bWasEvaluated && bDebugDraw)
 	{
 		USkeletalMeshComponent* SkeletalMeshComponent = InAnimInstance->GetSkelMeshComponent();
@@ -147,20 +149,40 @@ void FAnimNode_MotionMatching::PreUpdate(const UAnimInstance* InAnimInstance)
 		DrawParams.World = SkeletalMeshComponent->GetWorld();
 		DrawParams.DefaultLifeTime = 0.0f;
 
-		if (bDebugDrawMatch)
-		{
-			DrawParams.PoseIdx = MotionMatchingState.CurrentSearchResult.PoseIdx;
-		}
-
-		if (bDebugDrawQuery)
+		if (DrawParams.CanDraw())
 		{
-			DrawParams.PoseVector = MotionMatchingState.CurrentSearchResult.ComposedQuery.GetValues();
+			if (bDebugDrawMatch)
+			{
+				DrawFeatureVector(DrawParams, MotionMatchingState.CurrentSearchResult.PoseIdx);
+			}
+
+			if (bDebugDrawQuery)
+			{
+				EnumAddFlags(DrawParams.Flags, EDebugDrawFlags::DrawQuery);
+				DrawFeatureVector(DrawParams, MotionMatchingState.CurrentSearchResult.ComposedQuery.GetValues());
+			}
+
+			if (DrawParams.Database->PoseSearchMode == EPoseSearchMode::PCAKDTree_Compare)
+			{
+				const float TotalCostKDTree = MotionMatchingState.CurrentSearchResult.IsValid() ? MotionMatchingState.CurrentSearchResult.PoseCost.GetTotalCost() : 0.f;
+				const float TotalCostBruteForce = MotionMatchingState.CurrentSearchResult.IsValid() ? MotionMatchingState.CurrentSearchResult.BruteForcePoseCost.GetTotalCost() : 0.f;
+
+				MotionMatchingState.SearchCostHistoryKDTree.AddSample(TotalCostKDTree);
+				MotionMatchingState.SearchCostHistoryBruteForce.AddSample(TotalCostBruteForce);
+
+				// making SearchCostHistoryKDTree and SearchCostHistoryBruteForce min max consistent
+				MotionMatchingState.SearchCostHistoryKDTree.MinValue = FMath::Min(MotionMatchingState.SearchCostHistoryKDTree.MinValue, MotionMatchingState.SearchCostHistoryBruteForce.MinValue);
+				MotionMatchingState.SearchCostHistoryKDTree.MaxValue = FMath::Max(MotionMatchingState.SearchCostHistoryKDTree.MaxValue, MotionMatchingState.SearchCostHistoryBruteForce.MaxValue);
+
+				MotionMatchingState.SearchCostHistoryBruteForce.MinValue = MotionMatchingState.SearchCostHistoryKDTree.MinValue;
+				MotionMatchingState.SearchCostHistoryBruteForce.MaxValue = MotionMatchingState.SearchCostHistoryKDTree.MaxValue;
+
+				const FTransform OffsetTransform(FRotator(0.f, 0.f, 0.f), FVector(-50.f, -75.f, 100.f));
+				FVector2D DrawSize(150.f, 50.f);
+				DrawDebugFloatHistory(*DrawParams.World, MotionMatchingState.SearchCostHistoryKDTree, OffsetTransform * DrawParams.RootTransform, DrawSize, FColor(255,192,203,150)); // pink
+				DrawDebugFloatHistory(*DrawParams.World, MotionMatchingState.SearchCostHistoryBruteForce, OffsetTransform * DrawParams.RootTransform, DrawSize, FColor(0,0,255,150)); // blue
+			}
 		}
-
-		DrawParams.SearchCostHistoryKDTree = &MotionMatchingState.SearchCostHistoryKDTree;
-		DrawParams.SearchCostHistoryBruteForce = &MotionMatchingState.SearchCostHistoryBruteForce;
-
-		UE::PoseSearch::Draw(DrawParams);
 	}
 
 	bWasEvaluated = false;
diff --git a/Engine/Plugins/Experimental/Animation/PoseSearch/Source/Runtime/Private/PoseSearch.cpp b/Engine/Plugins/Experimental/Animation/PoseSearch/Source/Runtime/Private/PoseSearch.cpp
--- a/Engine/Plugins/Experimental/Animation/PoseSearch/Source/Runtime/Private/PoseSearch.cpp
+++ b/Engine/Plugins/Experimental/Animation/PoseSearch/Source/Runtime/Private/PoseSearch.cpp
@@ -882,9 +882,10 @@ UE::PoseSearch::FSearchResult UPoseSearchSequenceMetaData::Search(UE::PoseSearch
 	Result.AssetTime = SearchIndex.GetAssetTime(BestPoseIdx, Result.SearchIndexAsset);
 	Result.Database = nullptr;
 
-	SearchContext.DebugDrawParams.PoseVector = NormalizedQueryValues;
-	SearchContext.DebugDrawParams.PoseIdx = Result.PoseIdx;
-	Draw(SearchContext.DebugDrawParams);
+	DrawFeatureVector(SearchContext.DebugDrawParams, Result.PoseIdx);
+	
+	EnumAddFlags(SearchContext.DebugDrawParams.Flags, EDebugDrawFlags::DrawQuery);
+	DrawFeatureVector(SearchContext.DebugDrawParams, NormalizedQueryValues);
 
 	return Result;
 }
@@ -2031,45 +2032,24 @@ UE::PoseSearch::FSearchResult UPoseSearchDatabase::Search(UE::PoseSearch::FSearc
 		return Result;
 	}
 
+	if (PoseSearchMode == EPoseSearchMode::BruteForce || PoseSearchMode == EPoseSearchMode::PCAKDTree_Compare)
+	{
+		Result = SearchBruteForce(SearchContext);
+	}
 
 	if (PoseSearchMode != EPoseSearchMode::BruteForce)
 	{
-		Result = SearchPCAKDTree(SearchContext);
-
 #if WITH_EDITORONLY_DATA
-		if (PoseSearchMode == EPoseSearchMode::PCAKDTree_Compare && SearchContext.DebugDrawParams.SearchCostHistoryKDTree && Result.IsValid())
-		{
-			SearchContext.DebugDrawParams.SearchCostHistoryKDTree->AddSample(Result.PoseCost.GetTotalCost());
-		}
+		FPoseSearchCost BruteForcePoseCost = Result.BruteForcePoseCost;
 #endif
-	}
 
-	if (PoseSearchMode == EPoseSearchMode::BruteForce || PoseSearchMode == EPoseSearchMode::PCAKDTree_Compare)
-	{
-		Result = SearchBruteForce(SearchContext);
+		Result = SearchPCAKDTree(SearchContext);
 
 #if WITH_EDITORONLY_DATA
-		if (PoseSearchMode == EPoseSearchMode::PCAKDTree_Compare && SearchContext.DebugDrawParams.SearchCostHistoryBruteForce && Result.IsValid())
-		{
-			SearchContext.DebugDrawParams.SearchCostHistoryBruteForce->AddSample(Result.PoseCost.GetTotalCost());
-
-			// making SearchCostHistoryKDTree and SearchCostHistoryBruteForce min max consistent
-			if (SearchContext.DebugDrawParams.SearchCostHistoryKDTree)
-			{
-				SearchContext.DebugDrawParams.SearchCostHistoryKDTree->MinValue = FMath::Min(SearchContext.DebugDrawParams.SearchCostHistoryKDTree->MinValue, SearchContext.DebugDrawParams.SearchCostHistoryBruteForce->MinValue);
-				SearchContext.DebugDrawParams.SearchCostHistoryKDTree->MaxValue = FMath::Max(SearchContext.DebugDrawParams.SearchCostHistoryKDTree->MaxValue, SearchContext.DebugDrawParams.SearchCostHistoryBruteForce->MaxValue);
-
-				SearchContext.DebugDrawParams.SearchCostHistoryBruteForce->MinValue = SearchContext.DebugDrawParams.SearchCostHistoryKDTree->MinValue;
-				SearchContext.DebugDrawParams.SearchCostHistoryBruteForce->MaxValue = SearchContext.DebugDrawParams.SearchCostHistoryKDTree->MaxValue;
-			}
-		}
+		Result.BruteForcePoseCost = BruteForcePoseCost;
 #endif
 	}
-
-	SearchContext.DebugDrawParams.PoseVector = Result.ComposedQuery.GetNormalizedValues();
-	SearchContext.DebugDrawParams.PoseIdx = Result.PoseIdx;
-	Draw(SearchContext.DebugDrawParams);
-
+	
 	return Result;
 }
 
@@ -2363,6 +2343,10 @@ UE::PoseSearch::FSearchResult UPoseSearchDatabase::SearchBruteForce(UE::PoseSear
 #endif // WITH_EDITOR
 	}
 
+#if WITH_EDITORONLY_DATA
+	Result.BruteForcePoseCost = Result.PoseCost; 
+#endif
+
 	return Result;
 }
 
@@ -2379,6 +2363,9 @@ UE::PoseSearch::FSearchResult UPoseSearchDatabaseSet::Search(UE::PoseSearch::FSe
 
 	FSearchResult Result;
 	FPoseSearchCost ContinuingCost;
+#if WITH_EDITOR
+	FPoseSearchCost BruteForcePoseCost;
+#endif
 
 	// evaluating the continuing pose before all the active entries
 	if (bEvaluateContinuingPoseFirst && 
@@ -2428,39 +2415,38 @@ UE::PoseSearch::FSearchResult UPoseSearchDatabaseSet::Search(UE::PoseSearch::FSe
 		if (bSearchEntry)
 		{
 			FSearchResult EntryResult = Entry.Searchable->Search(SearchContext);
-			if (EntryResult.IsValid())
-			{
-				if (!Result.IsValid() || EntryResult.PoseCost.GetTotalCost() < Result.PoseCost.GetTotalCost())
-				{
-					Result = EntryResult;
-				}
 
-				if (EntryResult.ContinuingPoseCost.IsValid())
-				{
-					if (!ContinuingCost.IsValid() || EntryResult.ContinuingPoseCost.GetTotalCost() < ContinuingCost.GetTotalCost())
-					{
-						ContinuingCost = EntryResult.ContinuingPoseCost;
-					}
-				}
+			if (EntryResult.PoseCost.GetTotalCost() < Result.PoseCost.GetTotalCost())
+			{
+				Result = EntryResult;
+			}
 
-				if (Entry.PostSearchStatus == EPoseSearchPostSearchStatus::Stop)
-				{
-					break;
-				}
+			if (EntryResult.ContinuingPoseCost.GetTotalCost() < ContinuingCost.GetTotalCost())
+			{
+				ContinuingCost = EntryResult.ContinuingPoseCost;
+			}
+#if WITH_EDITOR
+			if (EntryResult.BruteForcePoseCost.GetTotalCost() < BruteForcePoseCost.GetTotalCost())
+			{
+				BruteForcePoseCost = EntryResult.BruteForcePoseCost;
+			}
+#endif
+			if (Entry.PostSearchStatus == EPoseSearchPostSearchStatus::Stop)
+			{
+				break;
 			}
 		}
 	}
 
-	if (Result.IsValid())
-	{
-		Result.ContinuingPoseCost = ContinuingCost;
-	}
-	else
+	Result.ContinuingPoseCost = ContinuingCost;
+
+#if WITH_EDITOR
+	Result.BruteForcePoseCost = BruteForcePoseCost;
+#endif
+
+	if (!Result.IsValid())
 	{
-		UE_LOG(
-			LogPoseSearch, 
-			Warning, 
-			TEXT("Invalid result searching %s"), *GetName());
+		UE_LOG(LogPoseSearch, Warning, TEXT("Invalid result searching %s"), *GetName());
 	}
 
 	return Result;
@@ -3040,29 +3026,45 @@ bool FDebugDrawParams::CanDraw() const
 	return SearchIndex->IsValid() && !SearchIndex->IsEmpty();
 }
 
-FLinearColor FDebugDrawParams::GetColor(const UPoseSearchFeatureChannel* channel) const
+FColor FDebugDrawParams::GetColor(int32 ColorPreset, float GradientPercentage) const
 {
-	if (Color)
-	{
-		return *Color;
-	}
+	FLinearColor Color = FLinearColor::Red;
 
 	const UPoseSearchSchema* Schema = GetSchema();
 	if (!Schema || Schema->SchemaCardinality <= 0)
 	{
-		return FLinearColor::Red;
+		Color = FLinearColor::Red;
+	}
+	else if (ColorPreset < 0 || ColorPreset >= Schema->ColorPresets.Num())
+	{
+		if (EnumHasAnyFlags(Flags, EDebugDrawFlags::DrawQuery))
+		{
+			Color = FLinearColor::Blue;
+		}
+		else
+		{
+			Color = FLinearColor::Green;
+		}
+	}
+	else
+	{
+		if (EnumHasAnyFlags(Flags, EDebugDrawFlags::DrawQuery))
+		{
+			Color = Schema->ColorPresets[ColorPreset].Query;
+		}
+		else
+		{
+			Color = Schema->ColorPresets[ColorPreset].Result;
+		}
 	}
 
-	const float TotalData = Schema->SchemaCardinality;
-	const float ChannelData = channel->GetChannelDataOffset();
-	const float HalfData = TotalData * 0.5f;
-
-	const float Hue = ChannelData < HalfData
-		? FMath::GetMappedRangeValueUnclamped({ 0.f, HalfData }, FVector2f(60.f, 0.f), ChannelData)
-		: FMath::GetMappedRangeValueUnclamped({ HalfData, TotalData }, FVector2f(280.f, 220.f), ChannelData);
+	if (EnumHasAnyFlags(Flags, EDebugDrawFlags::DrawSamplesWithColorGradient))
+	{
+		constexpr float DrawDebugGradientStrength = 0.8f;
+		Color = Color * (1.0f - DrawDebugGradientStrength * GradientPercentage);
+	}
 
-	const FLinearColor ColorHSV(Hue, 1.f, 1.f);
-	return ColorHSV.HSVToLinearRGB();
+	return Color.ToFColor(true);
 }
 
 const FPoseSearchIndex* FDebugDrawParams::GetSearchIndex() const
@@ -4554,80 +4556,53 @@ FTransform FAssetIndexer::GetTransformAndCacheResults(float SampleTime, float Or
 //////////////////////////////////////////////////////////////////////////
 // PoseSearch API
 
-static void DrawFeatureVector(const FDebugDrawParams& DrawParams, TArrayView<const float> PoseVector)
+void DrawFeatureVector(const FDebugDrawParams& DrawParams, TArrayView<const float> PoseVector)
 {
-	const UPoseSearchSchema* Schema = DrawParams.GetSchema();
-	check(Schema)
-
-	if (PoseVector.Num() != Schema->SchemaCardinality)
+	if (DrawParams.CanDraw())
 	{
-		return;
-	}
+		const UPoseSearchSchema* Schema = DrawParams.GetSchema();
+		check(Schema);
 
-	for (int32 ChannelIdx = 0; ChannelIdx != Schema->Channels.Num(); ++ChannelIdx)
-	{
-		if (DrawParams.ChannelMask & (1 << ChannelIdx))
+		if (PoseVector.Num() == Schema->SchemaCardinality)
 		{
-			Schema->Channels[ChannelIdx]->DebugDraw(DrawParams, PoseVector);
+			for (int32 ChannelIdx = 0; ChannelIdx != Schema->Channels.Num(); ++ChannelIdx)
+			{
+				if (DrawParams.ChannelMask & (1 << ChannelIdx))
+				{
+					Schema->Channels[ChannelIdx]->DebugDraw(DrawParams, PoseVector);
+				}
+			}
 		}
 	}
 }
 
-static void DrawSearchIndex(const FDebugDrawParams& DrawParams)
+void DrawFeatureVector(const FDebugDrawParams& DrawParams, int32 PoseIdx)
 {
-	const FPoseSearchIndex* SearchIndex = DrawParams.GetSearchIndex();
-	check(SearchIndex);
-
-	const int32 LastPoseIdx = SearchIndex->NumPoses;
-
-	TArray<float> PoseVector;
-	for (int32 PoseIdx = 0; PoseIdx != LastPoseIdx; ++PoseIdx)
+	if (PoseIdx != INDEX_NONE && DrawParams.CanDraw())
 	{
-		PoseVector = SearchIndex->GetPoseValues(PoseIdx);
-		SearchIndex->InverseNormalize(PoseVector);
-		DrawFeatureVector(DrawParams, PoseVector);
+		const FPoseSearchIndex* SearchIndex = DrawParams.GetSearchIndex();
+		// PreprocessInfo happens to be invalid when updating the database
+		if (SearchIndex->PreprocessInfo.IsValid())
+		{
+			// @todo: embed normalization into the weights to remove the InverseNormalize step and avoid the copy here, using the TArrayView instead of the TArray for the PoseVector
+			TArray<float> PoseVector;
+			PoseVector = SearchIndex->GetPoseValues(PoseIdx);
+			SearchIndex->InverseNormalize(PoseVector);
+			DrawFeatureVector(DrawParams, PoseVector);
+		}
 	}
 }
 
-void Draw(const FDebugDrawParams& DebugDrawParams)
+void DrawSearchIndex(const FDebugDrawParams& DrawParams)
 {
-	if (DebugDrawParams.CanDraw())
+	if (DrawParams.CanDraw())
 	{
-		if (DebugDrawParams.PoseIdx != INDEX_NONE)
-		{
-			const FPoseSearchIndex* SearchIndex = DebugDrawParams.GetSearchIndex();
-			check(SearchIndex);
-
-			// PreprocessInfo happens to be invalid when updating the database
-			if (SearchIndex->PreprocessInfo.IsValid())
-			{
-				TArray<float> PoseVector;
-				PoseVector = SearchIndex->GetPoseValues(DebugDrawParams.PoseIdx);
-				SearchIndex->InverseNormalize(PoseVector);
-				DrawFeatureVector(DebugDrawParams, PoseVector);
-			}
-		}
-		if (!DebugDrawParams.PoseVector.IsEmpty())
+		const FPoseSearchIndex* SearchIndex = DrawParams.GetSearchIndex();
+		const int32 LastPoseIdx = SearchIndex->NumPoses;
+		for (int32 PoseIdx = 0; PoseIdx != LastPoseIdx; ++PoseIdx)
 		{
-			DrawFeatureVector(DebugDrawParams, DebugDrawParams.PoseVector);
+			DrawFeatureVector(DrawParams, PoseIdx);
 		}
-		if (EnumHasAnyFlags(DebugDrawParams.Flags, EDebugDrawFlags::DrawSearchIndex))
-		{
-			DrawSearchIndex(DebugDrawParams);
-		}
-
-#if WITH_EDITORONLY_DATA
-		if (DebugDrawParams.Database->PoseSearchMode == EPoseSearchMode::PCAKDTree_Compare)
-		{
-			if (DebugDrawParams.SearchCostHistoryKDTree && DebugDrawParams.SearchCostHistoryBruteForce)
-			{
-				const FTransform OffsetTransform(FRotator(0.f, 90.f, 0.f), FVector(-50.f, 0.f, 100.f));
-				FVector2D DrawSize(150.f, 50.f);
-				DrawDebugFloatHistory(*DebugDrawParams.World, *DebugDrawParams.SearchCostHistoryKDTree, OffsetTransform * DebugDrawParams.RootTransform, DrawSize, FColor::Red);
-				DrawDebugFloatHistory(*DebugDrawParams.World, *DebugDrawParams.SearchCostHistoryBruteForce, OffsetTransform * DebugDrawParams.RootTransform, DrawSize, FColor::Blue);
-			}
-		}
-#endif
 	}
 }
 
diff --git a/Engine/Plugins/Experimental/Animation/PoseSearch/Source/Runtime/Private/PoseSearchFeatureChannels.cpp b/Engine/Plugins/Experimental/Animation/PoseSearch/Source/Runtime/Private/PoseSearchFeatureChannels.cpp
--- a/Engine/Plugins/Experimental/Animation/PoseSearch/Source/Runtime/Private/PoseSearchFeatureChannels.cpp
+++ b/Engine/Plugins/Experimental/Animation/PoseSearch/Source/Runtime/Private/PoseSearchFeatureChannels.cpp
@@ -20,7 +20,6 @@ constexpr float DrawDebugPointSize = 2.0f;
 constexpr float DrawDebugVelocityScale = 0.08f;
 constexpr float DrawDebugSphereSize = 2.0f;
 constexpr int32 DrawDebugSphereSegments = 8;
-constexpr float DrawDebugGradientStrength = 0.8f;
 constexpr float DrawDebugSampleLabelFontScale = 1.5f;
 static const FVector DrawDebugSampleLabelOffset = FVector(0.0f, 0.0f, 5.0f);
 
@@ -408,8 +407,7 @@ void UPoseSearchFeatureChannel_Position::DebugDraw(const UE::PoseSearch::FDebugD
 	int32 DataOffset = ChannelDataOffset;
 	const FVector BonePos = DrawParams.RootTransform.TransformPosition(FFeatureVectorHelper::DecodeVector(PoseVector, DataOffset));
 
-	const FLinearColor LinearColor = DrawParams.GetColor(this);
-	const FColor Color = LinearColor.ToFColor(true);
+	const FColor Color = DrawParams.GetColor(ColorPresetIndex);
 
 	if (EnumHasAnyFlags(DrawParams.Flags, EDebugDrawFlags::DrawFast | EDebugDrawFlags::DrawSearchIndex))
 	{
@@ -550,8 +548,7 @@ void UPoseSearchFeatureChannel_Heading::DebugDraw(const UE::PoseSearch::FDebugDr
 	const FVector BoneHeading = DrawParams.RootTransform.TransformPosition(FFeatureVectorHelper::DecodeVector(PoseVector, DataOffset));
 	check(DataOffset == ChannelDataOffset + ChannelCardinality);
 
-	const FLinearColor LinearColor = DrawParams.GetColor(this);
-	const FColor Color = LinearColor.ToFColor(true);
+	const FColor Color = DrawParams.GetColor(ColorPresetIndex);
 
 	if (EnumHasAnyFlags(DrawParams.Flags, EDebugDrawFlags::DrawSearchIndex))
 	{
@@ -1102,8 +1099,8 @@ void UPoseSearchFeatureChannel_Pose::DebugDraw(const UE::PoseSearch::FDebugDrawP
 			{
 				BonePos[SubsampleIdx] = FFeatureVectorHelper::DecodeVector(PoseVector, DataOffset);
 
-				const FLinearColor LinearColor = DrawParams.GetColor(this);
-				const FColor Color = LinearColor.ToFColor(true);
+				const float GradientPercentage = NumSubsamples > 1 ? float(SubsampleIdx) / float(NumSubsamples - 1) : 0.f;
+				const FColor Color = DrawParams.GetColor(SampledBone.ColorPresetIndex, GradientPercentage);
 
 				BonePos[SubsampleIdx] = DrawParams.RootTransform.TransformPosition(BonePos[SubsampleIdx]);
 				if (EnumHasAnyFlags(DrawParams.Flags, EDebugDrawFlags::DrawFast | EDebugDrawFlags::DrawSearchIndex))
@@ -1149,8 +1146,7 @@ void UPoseSearchFeatureChannel_Pose::DebugDraw(const UE::PoseSearch::FDebugDrawP
 			{
 				FVector BoneVel = FFeatureVectorHelper::DecodeVector(PoseVector, DataOffset);
 
-				const FLinearColor LinearColor = DrawParams.GetColor(this);
-				const FColor Color = LinearColor.ToFColor(true);
+				const FColor Color = DrawParams.GetColor(SampledBone.ColorPresetIndex);
 
 				BoneVel *= DrawDebugVelocityScale;
 				BoneVel = DrawParams.RootTransform.TransformVector(BoneVel);
@@ -1183,8 +1179,8 @@ void UPoseSearchFeatureChannel_Pose::DebugDraw(const UE::PoseSearch::FDebugDrawP
 			{
 				const FVector2D Phase = FFeatureVectorHelper::DecodeVector2D(PoseVector, DataOffset);
 
-				const FLinearColor LinearColor = DrawParams.GetColor(this);
-				const FColor Color = LinearColor.ToFColor(true);
+				const float GradientPercentage = NumSubsamples > 1 ? float(SubsampleIdx) / float(NumSubsamples - 1) : 0.f;
+				const FColor Color = DrawParams.GetColor(SampledBone.ColorPresetIndex, GradientPercentage);
 
 				static float ScaleFactor = 1.f;
 
@@ -1579,17 +1575,6 @@ void UPoseSearchFeatureChannel_Trajectory::DebugDraw(const UE::PoseSearch::FDebu
 		return;
 	}
 
-	auto GetGradientColor = [](const FLinearColor& OriginalColor, int SampleIdx, int NumSamples, EDebugDrawFlags Flags) -> FLinearColor
-	{
-		int Denominator = NumSamples - 1;
-		if (Denominator <= 0 || !EnumHasAnyFlags(Flags, EDebugDrawFlags::DrawSamplesWithColorGradient))
-		{
-			return OriginalColor;
-		}
-
-		return OriginalColor * (1.0f - DrawDebugGradientStrength * (SampleIdx / (float)Denominator));
-	};
-
 	int32 DataOffset = ChannelDataOffset;
 	int32 SampleIdx = 0;
 	TrajectoryPositionReconstructor TrajectoryPositionReconstructor;
@@ -1605,9 +1590,8 @@ void UPoseSearchFeatureChannel_Trajectory::DebugDraw(const UE::PoseSearch::FDebu
 			// validating TrajectoryPositionReconstructor
 			check((TrajectoryPositionReconstructor.GetReconstructedTrajectoryPos(*this, PoseVector, DrawParams.RootTransform, Sample.Offset) - TrajectoryPos).IsNearlyZero());
 
-			FLinearColor LinearColor = DrawParams.GetColor(this);
-			FLinearColor GradientColor = GetGradientColor(LinearColor, SampleIdx, NumSamples, DrawParams.Flags);
-			FColor Color = GradientColor.ToFColor(true);
+			const float GradientPercentage = NumSamples > 1 ? float(SampleIdx) / float(NumSamples - 1) : 0.f;
+			const FColor Color = DrawParams.GetColor(Sample.ColorPresetIndex, GradientPercentage);
 
 			if (EnumHasAnyFlags(DrawParams.Flags, EDebugDrawFlags::DrawFast | EDebugDrawFlags::DrawSearchIndex))
 			{
@@ -1627,9 +1611,8 @@ void UPoseSearchFeatureChannel_Trajectory::DebugDraw(const UE::PoseSearch::FDebu
 		{
 			FVector TrajectoryVel = FFeatureVectorHelper::DecodeVector(PoseVector, DataOffset);
 
-			const FLinearColor LinearColor = DrawParams.GetColor(this);
-			const FLinearColor GradientColor = GetGradientColor(LinearColor, SampleIdx, NumSamples, DrawParams.Flags);
-			const FColor Color = GradientColor.ToFColor(true);
+			const float GradientPercentage = NumSamples > 1 ? float(SampleIdx) / float(NumSamples - 1) : 0.f;
+			const FColor Color = DrawParams.GetColor(Sample.ColorPresetIndex, GradientPercentage);
 
 			TrajectoryVel *= DrawDebugVelocityScale;
 			TrajectoryVel = DrawParams.RootTransform.TransformVector(TrajectoryVel);
@@ -1660,9 +1643,8 @@ void UPoseSearchFeatureChannel_Trajectory::DebugDraw(const UE::PoseSearch::FDebu
 		{
 			FVector TrajectoryForward = FFeatureVectorHelper::DecodeVector(PoseVector, DataOffset);
 
-			const FLinearColor LinearColor = DrawParams.GetColor(this);
-			const FLinearColor GradientColor = GetGradientColor(LinearColor, SampleIdx, NumSamples, DrawParams.Flags);
-			const FColor Color = GradientColor.ToFColor(true);
+			const float GradientPercentage = NumSamples > 1 ? float(SampleIdx) / float(NumSamples - 1) : 0.f;
+			const FColor Color = DrawParams.GetColor(Sample.ColorPresetIndex, GradientPercentage);
 
 			TrajectoryForward = DrawParams.RootTransform.TransformVector(TrajectoryForward);
 
@@ -1689,19 +1671,11 @@ void UPoseSearchFeatureChannel_Trajectory::DebugDraw(const UE::PoseSearch::FDebu
 
 		if (EnumHasAnyFlags(DrawParams.Flags, EDebugDrawFlags::DrawSampleLabels))
 		{
-			const FLinearColor LinearColor = DrawParams.GetColor(this);
-			const FLinearColor GradientColor = GetGradientColor(LinearColor, SampleIdx, NumSamples, DrawParams.Flags);
-			const FColor Color = GradientColor.ToFColor(true);
+			const float GradientPercentage = NumSamples > 1 ? float(SampleIdx) / float(NumSamples - 1) : 0.f;
+			const FColor Color = DrawParams.GetColor(Sample.ColorPresetIndex, GradientPercentage);
+
+			const FString SampleLabel = FString::Format(TEXT("{0}"), { SampleIdx });
 
-			FString SampleLabel;
-			if (DrawParams.LabelPrefix.IsEmpty())
-			{
-				SampleLabel = FString::Format(TEXT("{0}"), { SampleIdx });
-			}
-			else
-			{
-				SampleLabel = FString::Format(TEXT("{1}[{0}]"), { SampleIdx, DrawParams.LabelPrefix.GetData() });
-			}
 			DrawDebugString(
 				DrawParams.World,
 				TrajectoryPos + DrawDebugSampleLabelOffset,
diff --git a/Engine/Plugins/Experimental/Animation/PoseSearch/Source/Runtime/Private/PoseSearchLibrary.cpp b/Engine/Plugins/Experimental/Animation/PoseSearch/Source/Runtime/Private/PoseSearchLibrary.cpp
--- a/Engine/Plugins/Experimental/Animation/PoseSearch/Source/Runtime/Private/PoseSearchLibrary.cpp
+++ b/Engine/Plugins/Experimental/Animation/PoseSearch/Source/Runtime/Private/PoseSearchLibrary.cpp
@@ -257,12 +257,6 @@ void UpdateMotionMatchingState(
 		SearchContext.BoneContainer = &Context.AnimInstanceProxy->GetRequiredBones();
 		SearchContext.bForceInterrupt = bForceInterrupt;
 		SearchContext.bCanAdvance = bCanAdvance;
-
-#if WITH_EDITORONLY_DATA
-		SearchContext.DebugDrawParams.SearchCostHistoryBruteForce = &InOutMotionMatchingState.SearchCostHistoryBruteForce;
-		SearchContext.DebugDrawParams.SearchCostHistoryKDTree = &InOutMotionMatchingState.SearchCostHistoryKDTree;
-#endif
-
 		SearchContext.CurrentResult = InOutMotionMatchingState.CurrentSearchResult;
 		SearchContext.PoseJumpThresholdTime = Settings.PoseJumpThresholdTime;
 		SearchContext.PoseIndicesHistory = &InOutMotionMatchingState.PoseIndicesHistory;
@@ -284,13 +278,17 @@ void UpdateMotionMatchingState(
 		// Search the database for the nearest match to the updated query vector
 		FSearchResult SearchResult = Searchable->Search(SearchContext);
 		const float ContinuingPoseCost = bCanAdvance && SearchResult.ContinuingPoseCost.IsValid() ? SearchResult.ContinuingPoseCost.GetTotalCost() : MAX_flt;
-		const float PoseCost = SearchResult.PoseCost.IsValid()  ? SearchResult.PoseCost.GetTotalCost() : MAX_flt;
-
+		const float PoseCost = SearchResult.PoseCost.IsValid() ? SearchResult.PoseCost.GetTotalCost() : MAX_flt;
+		
 		if (PoseCost < ContinuingPoseCost)
 		{
 			InOutMotionMatchingState.JumpToPose(Context, Settings, SearchResult);
 		}
 
+#if WITH_EDITOR
+		InOutMotionMatchingState.CurrentSearchResult.BruteForcePoseCost = SearchResult.BruteForcePoseCost;
+#endif
+
 		// todo: cache the queries into TraceState, one per TraceState.DatabaseEntries or
 		InOutMotionMatchingState.CurrentSearchResult.ComposedQuery = SearchResult.ComposedQuery;
 	}
diff --git a/Engine/Plugins/Experimental/Animation/PoseSearch/Source/Runtime/Public/PoseSearch/PoseSearch.h b/Engine/Plugins/Experimental/Animation/PoseSearch/Source/Runtime/Public/PoseSearch/PoseSearch.h
--- a/Engine/Plugins/Experimental/Animation/PoseSearch/Source/Runtime/Public/PoseSearch/PoseSearch.h
+++ b/Engine/Plugins/Experimental/Animation/PoseSearch/Source/Runtime/Public/PoseSearch/PoseSearch.h
@@ -228,6 +228,9 @@ struct POSESEARCH_API FPoseSearchBone
 	// @todo: temporary location for the channel bone weight to help the weights refactoring
 	UPROPERTY(EditAnywhere, Category = Config)
 	float Weight = 1.f;
+
+	UPROPERTY(EditAnywhere, Category = Config)
+	int32 ColorPresetIndex = 0;
 };
 
 //////////////////////////////////////////////////////////////////////////
@@ -426,8 +429,21 @@ private:
 	TArray<FBoneReference> BoneReferences;
 };
 
+
+
 } // namespace UE::PoseSearch
 
+USTRUCT()
+struct FPoseSearchSchemaColorPreset
+{
+	GENERATED_BODY()
+
+	UPROPERTY(EditAnywhere, Category = "Colors")
+	FLinearColor Query = FLinearColor::Blue;
+
+	UPROPERTY(EditAnywhere, Category = "Colors")
+	FLinearColor Result = FLinearColor::Yellow;
+};
 
 /**
 * Specifies the format of a pose search index. At runtime, queries are built according to the schema for searching.
@@ -476,7 +492,6 @@ public:
 	TArray<float> TrajectorySampleDistances_DEPRECATED;
 #endif // WITH_EDITOR
 
-
 	// If set, this schema will support mirroring pose search databases
 	UPROPERTY(EditAnywhere, Category = "Schema")
 	TObjectPtr<UMirrorDataTable> MirrorDataTable;
@@ -515,6 +530,9 @@ public:
 	UPROPERTY(EditAnywhere, Category = "Schema")
 	float MirrorMismatchCostBias = 0.f;
 
+	UPROPERTY(EditAnywhere, Category = "Schema")
+	TArray<FPoseSearchSchemaColorPreset> ColorPresets;
+	
 	bool IsValid () const;
 
 	int32 GetNumBones () const { return BoneIndices.Num(); }
@@ -952,32 +970,26 @@ enum class EDebugDrawFlags : uint32
 	// Draw the entire search index as a point cloud
 	DrawSearchIndex     = 1 << 0,
 
-	// Draw pose features for each pose vector
-	IncludePose         = 1 << 1,
-
-	// Draw trajectory features for each pose vector
-	IncludeTrajectory   = 1 << 2,
-
-	// Draw all pose vector features
-	IncludeAllFeatures  = IncludePose | IncludeTrajectory,
+	// Draw using Query colors form the schema / config
+	DrawQuery = 1 << 1,
 
 	/**
 	 * Keep rendered data until the next call to FlushPersistentDebugLines().
 	 * Combine with DrawSearchIndex to draw the search index only once.
 	 */
-	Persistent = 1 << 3,
+	Persistent = 1 << 2,
 	
 	// Label samples with their indices
-	DrawSampleLabels = 1 << 4,
+	DrawSampleLabels = 1 << 3,
 
 	// Fade colors
-	DrawSamplesWithColorGradient = 1 << 5,
+	DrawSamplesWithColorGradient = 1 << 4,
 
 	// Draw Bone Names
-	DrawBoneNames = 1 << 6,
+	DrawBoneNames = 1 << 5,
 
 	// Draws simpler shapes to improve performance
-	DrawFast = 1 << 7,
+	DrawFast = 1 << 6,
 };
 ENUM_CLASS_FLAGS(EDebugDrawFlags);
 
@@ -986,36 +998,19 @@ struct POSESEARCH_API FDebugDrawParams
 	const UWorld* World = nullptr;
 	const UPoseSearchDatabase* Database = nullptr;
 	const UPoseSearchSequenceMetaData* SequenceMetaData = nullptr;
-	EDebugDrawFlags Flags = EDebugDrawFlags::DrawBoneNames;
+	EDebugDrawFlags Flags = EDebugDrawFlags::None;
 	uint32 ChannelMask = (uint32)-1;
 
-	float DefaultLifeTime = 5.0f;
-	float PointSize = 1.0f;
+	float DefaultLifeTime = 5.f;
+	float PointSize = 1.f;
 
 	FTransform RootTransform = FTransform::Identity;
 
-	// If set, draw the corresponding pose from the search index
-	int32 PoseIdx = INDEX_NONE;
-
-	// If set, draw using this uniform color instead of feature-based coloring
-	const FLinearColor* Color = nullptr;
-
-	// If set, interpret the buffer as a pose vector and draw it
-	TArrayView<const float> PoseVector;
-
-	// Optional prefix for sample labels
-	FStringView LabelPrefix;
-
-#if WITH_EDITORONLY_DATA
-	FDebugFloatHistory* SearchCostHistoryBruteForce = nullptr;
-	FDebugFloatHistory* SearchCostHistoryKDTree = nullptr;
-#endif
-
 	// Optional Mesh for gathering SocketTransform(s)
 	TWeakObjectPtr<const USkinnedMeshComponent> Mesh = nullptr;
 
 	bool CanDraw() const;
-	FLinearColor GetColor(const UPoseSearchFeatureChannel* channel) const;
+	FColor GetColor(int32 ColorPreset, float GradientPercentage = 0.f) const;
 	const FPoseSearchIndex* GetSearchIndex() const;
 	const UPoseSearchSchema* GetSchema() const;
 };
@@ -1054,6 +1049,9 @@ struct FSearchResult
 #if WITH_EDITOR
 	FIoHash SearchIndexHash = FIoHash::Zero;
 #endif // WITH_EDITOR
+#if WITH_EDITORONLY_DATA
+	FPoseSearchCost BruteForcePoseCost;
+#endif // WITH_EDITORONLY_DATA
 
 	// Attempts to set the internal state to match the provided asset time including updating the internal DbPoseIdx. 
 	// If the provided asset time is out of bounds for the currently playing asset then this function will reset the 
@@ -1495,13 +1493,9 @@ public:
 //////////////////////////////////////////////////////////////////////////
 // Main PoseSearch API
 
-/**
-* Visualize pose search debug information
-*
-* @param DrawParams		Visualization options
-*/
-POSESEARCH_API void Draw(const FDebugDrawParams& DrawParams);
-
+POSESEARCH_API void DrawFeatureVector(const FDebugDrawParams& DrawParams, TArrayView<const float> PoseVector);
+POSESEARCH_API void DrawFeatureVector(const FDebugDrawParams& DrawParams, int32 PoseIdx);
+POSESEARCH_API void DrawSearchIndex(const FDebugDrawParams& DrawParams);
 
 /**
 * Creates a pose search index for an animation sequence
diff --git a/Engine/Plugins/Experimental/Animation/PoseSearch/Source/Runtime/Public/PoseSearch/PoseSearchFeatureChannels.h b/Engine/Plugins/Experimental/Animation/PoseSearch/Source/Runtime/Public/PoseSearch/PoseSearchFeatureChannels.h
--- a/Engine/Plugins/Experimental/Animation/PoseSearch/Source/Runtime/Public/PoseSearch/PoseSearchFeatureChannels.h
+++ b/Engine/Plugins/Experimental/Animation/PoseSearch/Source/Runtime/Public/PoseSearch/PoseSearchFeatureChannels.h
@@ -28,6 +28,9 @@ public:
 	UPROPERTY()
 	int8 SchemaBoneIdx;
 
+	UPROPERTY(EditAnywhere, Category = "Settings")
+	int32 ColorPresetIndex = 0;
+
 	// reuse continuing pose database feature values if available
 	UPROPERTY(EditAnywhere, Category = "Settings")
 	bool bUseFeaturesFromContinuingPose = true;
@@ -81,6 +84,9 @@ public:
 	UPROPERTY()
 	int8 SchemaBoneIdx;
 
+	UPROPERTY(EditAnywhere, Category = "Settings")
+	int32 ColorPresetIndex = 0;
+
 	// reuse continuing pose database feature values if available
 	UPROPERTY(EditAnywhere, Category = "Settings")
 	bool bUseFeaturesFromContinuingPose = true;
@@ -169,6 +175,9 @@ struct POSESEARCH_API FPoseSearchTrajectorySample
 
 	UPROPERTY(EditAnywhere, Category = Config)
 	float Weight = 1.f;
+
+	UPROPERTY(EditAnywhere, Category = Config)
+	int32 ColorPresetIndex = 0;
 };
 
 UCLASS(BlueprintType, EditInlineNew)
