PoseSearch: Updates to debugger to make it aware of searchable assets

diff --git a/Engine/Plugins/Experimental/Animation/PoseSearch/Source/Editor/Private/PoseSearchDebugger.cpp b/Engine/Plugins/Experimental/Animation/PoseSearch/Source/Editor/Private/PoseSearchDebugger.cpp
--- a/Engine/Plugins/Experimental/Animation/PoseSearch/Source/Editor/Private/PoseSearchDebugger.cpp
+++ b/Engine/Plugins/Experimental/Animation/PoseSearch/Source/Editor/Private/PoseSearchDebugger.cpp
@@ -30,11 +30,15 @@
 
 #define LOCTEXT_NAMESPACE "PoseSearchDebugger"
 
+namespace UE::PoseSearch {
+
 static FLinearColor LinearColorBlend(FLinearColor LinearColorA, FLinearColor LinearColorB, float BlendParam)
 {
 	return LinearColorA + (LinearColorB - LinearColorA) * BlendParam;
 }
 
+} // namespace UE::PoseSearch
+
 void UPoseSearchMeshComponent::Initialize(const FTransform& InComponentToWorld)
 {
 	SetComponentToWorld(InComponentToWorld);
@@ -224,14 +228,14 @@ public:
 	void CalculateColors(FChannelCostRange TotalCostRange, TArrayView<const FChannelCostRange> ChannelCostRanges)
 	{
 		const float CostColorBlend = (PoseCostDetails.PoseCost.GetTotalCost() - TotalCostRange.Min) / TotalCostRange.Delta;
-		CostColor = LinearColorBlend(FLinearColor::Green, FLinearColor::Red, CostColorBlend);
+		CostColor = UE::PoseSearch::LinearColorBlend(FLinearColor::Green, FLinearColor::Red, CostColorBlend);
 
 		ChannelCostColors.SetNum(ChannelCostRanges.Num());
 		for (int32 ChannelIdx = 0; ChannelIdx != ChannelCostRanges.Num(); ++ChannelIdx)
 		{
 			const FChannelCostRange& CostRange = ChannelCostRanges[ChannelIdx];
 			const float ColorBlend = (GetChannelCost(ChannelIdx) - CostRange.Min) / CostRange.Delta;
-			ChannelCostColors[ChannelIdx] = LinearColorBlend(FLinearColor::Green, FLinearColor::Red, ColorBlend);
+			ChannelCostColors[ChannelIdx] = UE::PoseSearch::LinearColorBlend(FLinearColor::Green, FLinearColor::Red, ColorBlend);
 		}
 	}
 
@@ -242,6 +246,11 @@ public:
 
 	ESearchIndexAssetType AssetType = ESearchIndexAssetType::Invalid;
 	int32 PoseIdx = 0;
+	TWeakObjectPtr<const UPoseSearchDatabase> SourceDatabase = nullptr;
+	FTraceMotionMatchingStatePoseEntry::EFlags PoseEntryFlags = FTraceMotionMatchingStatePoseEntry::EFlags::None;
+	FTraceMotionMatchingStateDatabaseEntry::EFlags DatabaseEntryFlags = FTraceMotionMatchingStateDatabaseEntry::EFlags::None;
+	FString DatabaseName = "";
+	FString DatabasePath = "";
 	FString AssetName = "";
 	FString AssetPath = "";
 	int32 DbAssetIdx = 0;
@@ -334,6 +343,35 @@ namespace DebuggerDatabaseColumns
 		}
 	};
 
+	struct FDatabaseName : IColumn
+	{
+		using IColumn::IColumn;
+
+		virtual FText GetLabel() const override { return LOCTEXT("ColumnLabelDatabaseName", "Database"); }
+
+		virtual FSortPredicate GetSortPredicate() const override
+		{
+			return [](const FRowDataRef& Row0, const FRowDataRef& Row1) -> bool { return Row0->DatabaseName < Row1->DatabaseName; };
+		}
+
+		virtual TSharedRef<SWidget> GenerateWidget(const FRowDataRef& RowData) const override
+		{
+			return SNew(SHyperlink)
+				.Text_Lambda([RowData]() -> FText { return FText::FromString(RowData->DatabaseName); })
+				.TextStyle(&FCoreStyle::Get().GetWidgetStyle<FTextBlockStyle>("SmallText"))
+				.ToolTipText_Lambda([RowData]() -> FText
+					{
+						return FText::Format(
+							LOCTEXT("DatabaseHyperlinkTooltipFormat", "Open database '{0}'"),
+							FText::FromString(RowData->DatabasePath));
+					})
+				.OnNavigate_Lambda([RowData]()
+					{
+						GEditor->GetEditorSubsystem<UAssetEditorSubsystem>()->OpenEditorForAsset(RowData->DatabasePath);
+					});
+		}
+	};
+
 	struct FAssetName : IColumn
 	{
 		using IColumn::IColumn;
@@ -386,8 +424,7 @@ namespace DebuggerDatabaseColumns
 			}
 			else
 			{
-				checkNoEntry();
-				return FText::FromString(TEXT(""));
+				return FText::FromString(TEXT("None"));
 			}
 		}
 	};
@@ -425,8 +462,7 @@ namespace DebuggerDatabaseColumns
 			}
 			else
 			{
-				checkNoEntry();
-				return FText();
+				return FText::FromString(TEXT("-"));
 			}
 		}
 	};
@@ -443,8 +479,15 @@ namespace DebuggerDatabaseColumns
 		}
 		
 		virtual FText GetRowText(const FRowDataRef& Row) const override
-        {
-        	return FText::AsNumber(Row->PoseCostDetails.PoseCost.GetTotalCost());
+		{
+			if (Row->PoseCostDetails.PoseCost.IsValid())
+			{
+				return FText::AsNumber(Row->PoseCostDetails.PoseCost.GetTotalCost());
+			}
+			else
+			{
+				return FText::FromString(TEXT("-"));
+			}
         }
 
 		virtual FSlateColor GetColorAndOpacity(const FRowDataRef& Row) const override
@@ -476,7 +519,14 @@ namespace DebuggerDatabaseColumns
 
 		virtual FText GetRowText(const FRowDataRef& Row) const override
 		{
-			return FText::AsNumber(Row->GetChannelCost(ChannelIdx));
+			if (Row->PoseCostDetails.ChannelCosts.IsValidIndex(ChannelIdx))
+			{
+				return FText::AsNumber(Row->GetChannelCost(ChannelIdx));
+			}
+			else
+			{
+				return FText::FromString(TEXT("-"));
+			}
 		}
 
 		virtual FSlateColor GetColorAndOpacity(const FRowDataRef& Row) const override
@@ -556,11 +606,7 @@ namespace DebuggerDatabaseColumns
 
 		virtual FText GetRowText(const FRowDataRef& Row) const override
 		{
-			if (Row->AssetType == ESearchIndexAssetType::Sequence)
-			{
-				return FText::FromString(TEXT("-"));
-			}
-			else if (Row->AssetType == ESearchIndexAssetType::BlendSpace)
+			if (Row->AssetType == ESearchIndexAssetType::BlendSpace)
 			{
 				return FText::Format(
 					LOCTEXT("Blend Parameters", "({0}, {1})"),
@@ -569,8 +615,7 @@ namespace DebuggerDatabaseColumns
 			}
 			else
 			{
-				checkNoEntry();
-				return FText();
+				return FText::FromString(TEXT("-"));
 			}
 		}
 	};
@@ -668,25 +713,28 @@ class SDebuggerMessageBox : public SCompoundWidget
 	}
 };
 
-void SDebuggerDatabaseView::Update(const FTraceMotionMatchingStateMessage& State, const UPoseSearchDatabase& Database)
+void SDebuggerDatabaseView::Update(const FTraceMotionMatchingStateMessage& State)
 {
-	UpdateRows(State, Database);
-}
+	CreateRows(State);
 
-const TSharedRef<FDebuggerDatabaseRowData>& SDebuggerDatabaseView::GetPoseIdxDatabaseRow(int32 PoseIdx) const
-{
-	const TSharedRef<FDebuggerDatabaseRowData>* RowPtr = UnfilteredDatabaseRows.FindByPredicate(
-		[=](TSharedRef<FDebuggerDatabaseRowData>& Row)
-	{
-		return Row->PoseIdx == PoseIdx;
-	});
+	check(ActiveView.Rows.Num() == 1);
+	check(ContinuingPoseView.Rows.Num() == 1);
 
-	check(RowPtr != nullptr);
+	DatabaseSequenceFilter = State.DatabaseSequenceFilter;
+	DatabaseBlendSpaceFilter = State.DatabaseBlendSpaceFilter;
 
-	return *RowPtr;
-}
+	SortDatabaseRows();
+	FilterDatabaseRows();
 
 
+	//	if (bNewDatabase)
+	// 	{
+	// 		FilteredDatabaseView.ListView->ClearSelection();
+	// 	}
+
+	ComputeFilteredDatabaseRowsColors();
+}
+
 void SDebuggerDatabaseView::RefreshColumns()
 {
 	using namespace DebuggerDatabaseColumns;
@@ -780,7 +828,7 @@ void SDebuggerDatabaseView::OnDatabaseRowSelectionChanged(
 {
 	if (Row.IsValid())
 	{
-		OnPoseSelectionChanged.ExecuteIfBound(Row->PoseIdx, Row->AssetTime);
+		OnPoseSelectionChanged.ExecuteIfBound(Row->SourceDatabase.Get(), Row->PoseIdx, Row->AssetTime);
 	}
 }
 
@@ -819,38 +867,29 @@ void SDebuggerDatabaseView::FilterDatabaseRows()
 	FString FilterString = FilterText.ToString();
 	TArray<FString> Tokens;
 	FilterString.ParseIntoArrayWS(Tokens);
+	const bool bHasNameFilter = !Tokens.IsEmpty();
 
-	if (Tokens.Num() == 0)
-	{
-		for (const auto& UnfilteredRow : UnfilteredDatabaseRows)
-		{
-			if (!bAssetFilterEnabled || 
-				(UnfilteredRow->AssetType == ESearchIndexAssetType::Sequence && DatabaseSequenceFilter[UnfilteredRow->DbAssetIdx]) ||
-				(UnfilteredRow->AssetType == ESearchIndexAssetType::BlendSpace && DatabaseBlendSpaceFilter[UnfilteredRow->DbAssetIdx]))
-			{
-				FilteredDatabaseView.Rows.Add(UnfilteredRow);
-			}
-		}
-	}
-	else
+	for (const auto& UnfilteredRow : UnfilteredDatabaseRows)
 	{
-		for (const auto& UnfilteredRow : UnfilteredDatabaseRows)
+		bool bPoseIsFromCurrentDb = EnumHasAnyFlags(UnfilteredRow->DatabaseEntryFlags, FTraceMotionMatchingStateDatabaseEntry::EFlags::CurrentDatabase);
+		bool bPassesSequenceFilter = !bPoseIsFromCurrentDb || (UnfilteredRow->AssetType == ESearchIndexAssetType::Sequence && DatabaseSequenceFilter[UnfilteredRow->DbAssetIdx]);
+		bool bPassesBlendSpaceFilter = !bPoseIsFromCurrentDb || (UnfilteredRow->AssetType == ESearchIndexAssetType::BlendSpace && DatabaseBlendSpaceFilter[UnfilteredRow->DbAssetIdx]);
+		bool bPassesAssetFilter = !bAssetFilterEnabled || bPassesSequenceFilter || bPassesBlendSpaceFilter;
+
+		if (bPassesAssetFilter)
 		{
-			if (!bAssetFilterEnabled ||
-				(UnfilteredRow->AssetType == ESearchIndexAssetType::Sequence && DatabaseSequenceFilter[UnfilteredRow->DbAssetIdx]) ||
-				(UnfilteredRow->AssetType == ESearchIndexAssetType::BlendSpace && DatabaseBlendSpaceFilter[UnfilteredRow->DbAssetIdx]))
+			bool bPassesNameFilter = true;
+			if (bHasNameFilter)
 			{
-				bool bMatchesAllTokens = Algo::AllOf(
-					Tokens,
-					[&](FString Token)
+				bPassesNameFilter = Algo::AllOf(Tokens,[&](FString Token)
 				{
 					return UnfilteredRow->AssetName.Contains(Token);
 				});
+			}
 
-				if (bMatchesAllTokens)
-				{
-					FilteredDatabaseView.Rows.Add(UnfilteredRow);
-				}
+			if (bPassesNameFilter)
+			{
+				FilteredDatabaseView.Rows.Add(UnfilteredRow);
 			}
 		}
 	}
@@ -858,156 +897,84 @@ void SDebuggerDatabaseView::FilterDatabaseRows()
 	FilteredDatabaseView.ListView->RequestListRefresh();
 }
 
-void SDebuggerDatabaseView::CreateRows(const UPoseSearchDatabase& Database)
+void SDebuggerDatabaseView::CreateRows(const FTraceMotionMatchingStateMessage& State)
 {
-	const bool bValidIndex = Database.GetSearchIndex() && Database.GetSearchIndex()->IsValid();
-	if (bValidIndex)
-	{
-		const int32 NumPoses = Database.GetSearchIndex()->NumPoses;
-		UnfilteredDatabaseRows.Reset(NumPoses);
-
-		RowsSourceDatabase = &Database;
-
-		// Build database rows
-		for (const FPoseSearchIndexAsset& SearchIndexAsset : Database.GetSearchIndex()->Assets)
-		{
-			if (SearchIndexAsset.Type == ESearchIndexAssetType::Sequence)
-			{
-				const FPoseSearchDatabaseSequence& DbSequence = Database.GetSequenceSourceAsset(&SearchIndexAsset);
-				const int32 LastPoseIdx = SearchIndexAsset.FirstPoseIdx + SearchIndexAsset.NumPoses;
-				for (int32 PoseIdx = SearchIndexAsset.FirstPoseIdx; PoseIdx != LastPoseIdx; ++PoseIdx)
-				{
-					float Time = Database.GetAssetTime(PoseIdx);
-
-					TSharedRef<FDebuggerDatabaseRowData> Row = UnfilteredDatabaseRows.Add_GetRef(MakeShared<FDebuggerDatabaseRowData>());
-					Row->PoseIdx = PoseIdx;
-					Row->AssetType = ESearchIndexAssetType::Sequence;
-					Row->AssetName = DbSequence.Sequence->GetName();
-					Row->AssetPath = DbSequence.Sequence->GetPathName();
-					Row->DbAssetIdx = SearchIndexAsset.SourceAssetIdx;
-					Row->AssetTime = Time;
-					Row->AnimFrame = DbSequence.Sequence->GetFrameAtTime(Time);
-					Row->bMirrored = SearchIndexAsset.bMirrored;
-					Row->bLooping = DbSequence.Sequence->bLoop;
-					Row->BlendParameters = FVector::Zero();
-				}
-			}
-			else if (SearchIndexAsset.Type == ESearchIndexAssetType::BlendSpace)
-			{
-				const FPoseSearchDatabaseBlendSpace& DbBlendSpace = Database.GetBlendSpaceSourceAsset(&SearchIndexAsset);
-				const int32 LastPoseIdx = SearchIndexAsset.FirstPoseIdx + SearchIndexAsset.NumPoses;
-				for (int32 PoseIdx = SearchIndexAsset.FirstPoseIdx; PoseIdx != LastPoseIdx; ++PoseIdx)
-				{
-					float Time = Database.GetAssetTime(PoseIdx);
-
-					TSharedRef<FDebuggerDatabaseRowData> Row = UnfilteredDatabaseRows.Add_GetRef(MakeShared<FDebuggerDatabaseRowData>());
-					Row->PoseIdx = PoseIdx;
-					Row->AssetType = ESearchIndexAssetType::BlendSpace;
-					Row->AssetName = DbBlendSpace.BlendSpace->GetName();
-					Row->AssetPath = DbBlendSpace.BlendSpace->GetPathName();
-					Row->DbAssetIdx = SearchIndexAsset.SourceAssetIdx;
-					Row->AssetTime = Time;
-					Row->AnimFrame = 0; // There is no frame index associated with a blendspace
-					Row->bMirrored = SearchIndexAsset.bMirrored;
-					Row->bLooping = DbBlendSpace.BlendSpace->bLoop;
-					Row->BlendParameters = SearchIndexAsset.BlendParameters;
-				}
-			}
-			else
-			{
-				checkNoEntry();
-			}
-		}
-	}
-	else
-	{
-		UnfilteredDatabaseRows.Reset();
-	}
+	UnfilteredDatabaseRows.Reset();
 
 	ActiveView.Rows.Reset();
 	ActiveView.Rows.Add(MakeShared<FDebuggerDatabaseRowData>());
 
 	ContinuingPoseView.Rows.Reset();
 	ContinuingPoseView.Rows.Add(MakeShared<FDebuggerDatabaseRowData>());
-}
 
-void SDebuggerDatabaseView::UpdateRows(const FTraceMotionMatchingStateMessage& State, const UPoseSearchDatabase& Database)
-{
-	const bool bValidIndex = Database.GetSearchIndex() && Database.GetSearchIndex()->IsValid();
-	const bool bNewDatabase = RowsSourceDatabase != &Database || !bValidIndex;
-	if (bNewDatabase || UnfilteredDatabaseRows.IsEmpty())
+	for (const FTraceMotionMatchingStateDatabaseEntry& DbEntry : State.DatabaseEntries)
 	{
-		CreateRows(Database);
-	}
-
-	check(ActiveView.Rows.Num() == 1);
-	check(ContinuingPoseView.Rows.Num() == 1);
+		const auto* Database = FTraceMotionMatchingState::GetObjectFromId<UPoseSearchDatabase>(DbEntry.DatabaseId);
 
-	FSearchContext SearchContext;
-	if (const FPoseSearchIndexAsset* CurrentIndexAsset = Database.GetSearchIndex()->FindAssetForPose(State.DbPoseIdx))
-	{
-		SearchContext.QueryMirrorRequest = CurrentIndexAsset->bMirrored ? 
-			EPoseSearchBooleanRequest::TrueValue : EPoseSearchBooleanRequest::FalseValue;
-	}
-	
-	if (!UnfilteredDatabaseRows.IsEmpty())
-	{
-		bool IsActiveViewRowInitialized = false;
-		bool IsContinuingPoseViewRowInitialized = false;
-		for (const TSharedRef<FDebuggerDatabaseRowData>& Row : UnfilteredDatabaseRows)
+		const FPoseSearchIndex* SearchIndex = Database->GetSearchIndex();
+		if (SearchIndex && SearchIndex->IsValid())
 		{
-			const int32 PoseIdx = Row->PoseIdx;
-
-			// @todo: invalidate rows in case Database changed while PIE is active with pose debugger. for now we just avoid crashes
-			if (PoseIdx < Database.GetSearchIndex()->NumPoses)
+			for (const FTraceMotionMatchingStatePoseEntry& PoseEntry : DbEntry.PoseEntries)
 			{
-				const EPoseComparisonFlags PoseComparisonFlags = (PoseIdx == State.ContinuingPoseIdx) ? EPoseComparisonFlags::ContinuingPose : EPoseComparisonFlags::None;
-
-				Database.ComparePoses(
-					SearchContext,
-					PoseIdx,
-					PoseComparisonFlags,
-					State.QueryVectorNormalized,
-					Row->PoseCostDetails);
-
-				// If we are on the active pose for the frame
-				if (PoseIdx == State.DbPoseIdx)
+				const FPoseSearchIndexAsset* SearchIndexAsset = SearchIndex->FindAssetForPose(PoseEntry.DbPoseIdx);
+				if (SearchIndexAsset)
 				{
-					*ActiveView.Rows[0] = Row.Get();
-					IsActiveViewRowInitialized = true;
-				}
+					TSharedRef<FDebuggerDatabaseRowData>& Row = UnfilteredDatabaseRows.Add_GetRef(MakeShared<FDebuggerDatabaseRowData>());
 
-				if (PoseIdx == State.ContinuingPoseIdx)
-				{
-					*ContinuingPoseView.Rows[0] = Row.Get();
-					IsContinuingPoseViewRowInitialized = true;
-				}
-			}
-		}
+					const float Time = SearchIndex->GetAssetTime(PoseEntry.DbPoseIdx, SearchIndexAsset);
 
-		if (!IsActiveViewRowInitialized)
-		{
-			*ActiveView.Rows[0] = UnfilteredDatabaseRows[0].Get();
-		}
-
-		if (!IsContinuingPoseViewRowInitialized)
-		{
-			*ContinuingPoseView.Rows[0] = *ActiveView.Rows[0];
-		}
-	}
+					Row->PoseIdx = PoseEntry.DbPoseIdx;
+					Row->SourceDatabase = Database;
+					Row->DatabaseName = Database->GetName();
+					Row->DatabasePath = Database->GetPathName();
+					Row->DatabaseEntryFlags = DbEntry.Flags;
+					Row->PoseEntryFlags = PoseEntry.Flags;
+					Row->DbAssetIdx = SearchIndexAsset->SourceAssetIdx;
+					Row->AssetTime = Time;
+					Row->bMirrored = SearchIndexAsset->bMirrored;
+					Row->PoseCostDetails.PoseCost.Set(PoseEntry.Cost, 0.0f);
 
-	DatabaseSequenceFilter = State.DatabaseSequenceFilter;
-	DatabaseBlendSpaceFilter = State.DatabaseBlendSpaceFilter;
+					if (SearchIndexAsset->Type == ESearchIndexAssetType::Sequence)
+					{
+						const FPoseSearchDatabaseSequence& DbSequence = Database->GetSequenceSourceAsset(SearchIndexAsset);
+						Row->AssetType = ESearchIndexAssetType::Sequence;
+						Row->AssetName = DbSequence.Sequence->GetName();
+						Row->AssetPath = DbSequence.Sequence->GetPathName();
+						Row->AnimFrame = DbSequence.Sequence->GetFrameAtTime(Time);
+						Row->bLooping = DbSequence.Sequence->bLoop;
+						Row->BlendParameters = FVector::Zero();
+					}
+					else if (SearchIndexAsset->Type == ESearchIndexAssetType::BlendSpace)
+					{
+						const FPoseSearchDatabaseBlendSpace& DbBlendSpace = Database->GetBlendSpaceSourceAsset(SearchIndexAsset);
+						Row->AssetType = ESearchIndexAssetType::BlendSpace;
+						Row->AssetName = DbBlendSpace.BlendSpace->GetName();
+						Row->AssetPath = DbBlendSpace.BlendSpace->GetPathName();
+						Row->AnimFrame = 0; // There is no frame index associated with a blendspace
+						Row->bLooping = DbBlendSpace.BlendSpace->bLoop;
+						Row->BlendParameters = SearchIndexAsset->BlendParameters;
+					}
+					else
+					{
+						checkNoEntry();
+					}
 
-	SortDatabaseRows();
-	FilterDatabaseRows();
+					if (EnumHasAnyFlags(PoseEntry.Flags, FTraceMotionMatchingStatePoseEntry::EFlags::ContinuingPose))
+					{
+						*ContinuingPoseView.Rows[0] = Row.Get();
+					}
 
-	if (bNewDatabase)
-	{
-		FilteredDatabaseView.ListView->ClearSelection();
+					if (EnumHasAnyFlags(PoseEntry.Flags, FTraceMotionMatchingStatePoseEntry::EFlags::CurrentPose))
+					{
+						*ActiveView.Rows[0] = Row.Get();
+					}
+				}
+			}
+		}
 	}
 
-	ComputeFilteredDatabaseRowsColors();
+	ActiveView.ListView->RequestListRefresh();
+	ContinuingPoseView.ListView->RequestListRefresh();
 }
 
 void SDebuggerDatabaseView::ComputeFilteredDatabaseRowsColors()
@@ -1098,29 +1065,30 @@ void SDebuggerDatabaseView::Construct(const FArguments& InArgs)
 	// @TODO: Support runtime reordering of these indices
 	// Construct all column types
 	int32 ColumnIdx = 0;
+	AddColumn(MakeShared<FDatabaseName>(ColumnIdx++));
 	AddColumn(MakeShared<FAssetName>(ColumnIdx++));
 	AddColumn(MakeShared<FAssetType>(ColumnIdx++));
 
 	auto CostColumn = MakeShared<FCost>(ColumnIdx++);
 	AddColumn(CostColumn);
 
-	int32 ChannelIdx = 0;
-	FChannelCostColumn::FParams ChannelCostParams;
-	ChannelCostParams.ChannelIdx = ChannelIdx++;
-	ChannelCostParams.SortIndex = ColumnIdx++;
-	AddColumn(MakeShared<FChannelCostColumn>(ChannelCostParams));
-
-	ChannelCostParams.ChannelIdx = ChannelIdx++;
-	ChannelCostParams.SortIndex = ColumnIdx++;
-	AddColumn(MakeShared<FChannelCostColumn>(ChannelCostParams));
-
-	ChannelCostParams.ChannelIdx = ChannelIdx++;
-	ChannelCostParams.SortIndex = ColumnIdx++;
-	AddColumn(MakeShared<FChannelCostColumn>(ChannelCostParams));
-
-	ChannelCostParams.ChannelIdx = ChannelIdx++;
-	ChannelCostParams.SortIndex = ColumnIdx++;
-	AddColumn(MakeShared<FChannelCostColumn>(ChannelCostParams));
+// 	int32 ChannelIdx = 0;
+// 	FChannelCostColumn::FParams ChannelCostParams;
+// 	ChannelCostParams.ChannelIdx = ChannelIdx++;
+// 	ChannelCostParams.SortIndex = ColumnIdx++;
+// 	AddColumn(MakeShared<FChannelCostColumn>(ChannelCostParams));
+// 
+// 	ChannelCostParams.ChannelIdx = ChannelIdx++;
+// 	ChannelCostParams.SortIndex = ColumnIdx++;
+// 	AddColumn(MakeShared<FChannelCostColumn>(ChannelCostParams));
+// 
+// 	ChannelCostParams.ChannelIdx = ChannelIdx++;
+// 	ChannelCostParams.SortIndex = ColumnIdx++;
+// 	AddColumn(MakeShared<FChannelCostColumn>(ChannelCostParams));
+// 
+// 	ChannelCostParams.ChannelIdx = ChannelIdx++;
+// 	ChannelCostParams.SortIndex = ColumnIdx++;
+// 	AddColumn(MakeShared<FChannelCostColumn>(ChannelCostParams));
 
 
 	AddColumn(MakeShared<FCostModifier>(ColumnIdx++));
@@ -1338,7 +1306,7 @@ void SDebuggerDatabaseView::Construct(const FArguments& InArgs)
 					.VAlign(VAlign_Fill)
 					[
 						SNew(STextBlock)
-						.Text(FText::FromString("Pose Database"))	
+						.Text(FText::FromString("Pose Candidates"))
 					]
 				]
 				.AutoWidth()
@@ -1446,9 +1414,9 @@ void SDebuggerDetailsView::Construct(const FArguments& InArgs)
 	];
 }
 
-void SDebuggerDetailsView::Update(const FTraceMotionMatchingStateMessage& State, const UPoseSearchDatabase& Database) const
+void SDebuggerDetailsView::Update(const FTraceMotionMatchingStateMessage& State) const
 {
-	UpdateReflection(State, Database);
+	UpdateReflection(State);
 }
 
 SDebuggerDetailsView::~SDebuggerDetailsView()
@@ -1460,61 +1428,77 @@ SDebuggerDetailsView::~SDebuggerDetailsView()
 	}
 }
 
-void SDebuggerDetailsView::UpdateReflection(const FTraceMotionMatchingStateMessage& State, const UPoseSearchDatabase& Database) const
+void SDebuggerDetailsView::UpdateReflection(const FTraceMotionMatchingStateMessage& State) const
 {
 	check(Reflection);
-	const UPoseSearchSchema* Schema = Database.Schema;
 
-	Reflection->CurrentDatabaseName = Database.GetName();
+	const UPoseSearchDatabase* CurrentDatabase = State.GetCurrentDatabase();
+	const FPoseSearchIndex* CurrentSearchIndex = CurrentDatabase ? CurrentDatabase->GetSearchIndex() : nullptr;
+	int32 CurrentDbPoseIdx = State.GetCurrentDatabasePoseIndex();
+
+	Reflection->CurrentDatabaseName = CurrentDatabase ? CurrentDatabase->GetName() : "None";
 	Reflection->ElapsedPoseJumpTime = State.ElapsedPoseJumpTime;
 	Reflection->bFollowUpAnimation = EnumHasAnyFlags(State.Flags, FTraceMotionMatchingState::EFlags::FollowupAnimation);
 
-	const bool bValidIndex = Database.GetSearchIndex() && Database.GetSearchIndex()->IsValid();
-	const bool bValidState = bValidIndex && State.DbPoseIdx < Database.GetSearchIndex()->NumPoses;
-	// @todo: cache the derived data key in the state to check its validity
-	if (bValidState)
+
+	Reflection->AssetPlayerAssetName = "None";
+	if (CurrentSearchIndex)
 	{
-		Reflection->AssetPlayerAssetName = FString();
-		if (const FPoseSearchIndexAsset* IndexAsset = Database.GetSearchIndex()->FindAssetForPose(State.DbPoseIdx))
+		if (const FPoseSearchIndexAsset* IndexAsset = CurrentSearchIndex->FindAssetForPose(CurrentDbPoseIdx))
 		{
-			Reflection->AssetPlayerAssetName = Database.GetSourceAssetName(IndexAsset);
+			Reflection->AssetPlayerAssetName = CurrentDatabase->GetSourceAssetName(IndexAsset);
 		}
+	}
 
-		Reflection->AssetPlayerTime = State.AssetPlayerTime;
-		Reflection->LastDeltaTime = State.DeltaTime;
-		Reflection->SimLinearVelocity = State.SimLinearVelocity;
-		Reflection->SimAngularVelocity = State.SimAngularVelocity;
-		Reflection->AnimLinearVelocity = State.AnimLinearVelocity;
-		Reflection->AnimAngularVelocity = State.AnimAngularVelocity;
 
-		// Query pose
-		Reflection->QueryPoseVector = State.QueryVector;
+	Reflection->AssetPlayerTime = State.AssetPlayerTime;
+	Reflection->LastDeltaTime = State.DeltaTime;
+	Reflection->SimLinearVelocity = State.SimLinearVelocity;
+	Reflection->SimAngularVelocity = State.SimAngularVelocity;
+	Reflection->AnimLinearVelocity = State.AnimLinearVelocity;
+	Reflection->AnimAngularVelocity = State.AnimAngularVelocity;
 
-		// Active pose
-		Reflection->ActivePoseVector = Database.GetSearchIndex()->GetPoseValues(State.DbPoseIdx);
-		Database.GetSearchIndex()->InverseNormalize(Reflection->ActivePoseVector);
+	// Query pose
+	Reflection->QueryPoseVector = State.QueryVector;
 
-		auto DebuggerView = ParentDebuggerViewPtr.Pin();
-		if (DebuggerView.IsValid())
+	// Active pose
+	if (CurrentSearchIndex)
+	{
+		Reflection->ActivePoseVector = CurrentSearchIndex->GetPoseValues(CurrentDbPoseIdx);
+		CurrentSearchIndex->InverseNormalize(Reflection->ActivePoseVector);
+	}
+
+	auto DebuggerView = ParentDebuggerViewPtr.Pin();
+	if (DebuggerView.IsValid())
+	{
+		TArray<TSharedRef<FDebuggerDatabaseRowData>> SelectedRows = DebuggerView->GetSelectedDatabaseRows();
+		if (!SelectedRows.IsEmpty())
 		{
-			TArray<TSharedRef<FDebuggerDatabaseRowData>> SelectedRows = DebuggerView->GetSelectedDatabaseRows();
-			if (!SelectedRows.IsEmpty())
+			const TSharedRef<FDebuggerDatabaseRowData>& Selected = SelectedRows[0];
+			if (Selected->SourceDatabase.IsValid())
 			{
-				const TSharedRef<FDebuggerDatabaseRowData>& Selected = SelectedRows[0];
-				Reflection->SelectedPoseVector = Database.GetSearchIndex()->GetPoseValues(Selected->PoseIdx);
-				Database.GetSearchIndex()->InverseNormalize(Reflection->SelectedPoseVector);
-
-				Reflection->CostVector = Selected->PoseCostDetails.CostVector;
-				//Database.SearchIndex.InverseNormalize(Reflection->CostVector);
-
-				const TSharedRef<FDebuggerDatabaseRowData>& ActiveRow = DebuggerView->GetPoseIdxDatabaseRow(State.DbPoseIdx);
-
-				Reflection->CostVectorDifference = Reflection->CostVector;
-				for (int i = 0; i < Reflection->CostVectorDifference.Num(); ++i)
+				const FPoseSearchIndex* SelectedSearchIndex = Selected->SourceDatabase->GetSearchIndex();
+				if (SelectedSearchIndex)
 				{
-					Reflection->CostVectorDifference[i] -= ActiveRow->PoseCostDetails.CostVector[i];
+					Reflection->SelectedPoseVector = SelectedSearchIndex->GetPoseValues(Selected->PoseIdx);
+					SelectedSearchIndex->InverseNormalize(Reflection->SelectedPoseVector);
 				}
 			}
+
+			Reflection->CostVector = Selected->PoseCostDetails.CostVector;
+// 
+// 			TArray<TSharedRef<FDebuggerDatabaseRowData>> SelectedRows = DebuggerView->GetSelectedDatabaseRows();
+// 
+// 
+// 			State.DatabaseEntries[State.CurrentDbEntryIdx]
+// 
+// 			const TSharedRef<FDebuggerDatabaseRowData>& ActiveRow = DebuggerView->GetPoseIdxDatabaseRow(CurrentDbPoseIdx);
+// 
+// 			Reflection->CostVectorDifference = Reflection->CostVector;
+// 			for (int i = 0; i < Reflection->CostVectorDifference.Num(); ++i)
+// 			{
+// 				Reflection->CostVectorDifference[i] -= ActiveRow->PoseCostDetails.CostVector[i];
+// 			}
 		}
 	}
 }
@@ -1612,7 +1596,7 @@ void SDebuggerView::Tick(const FGeometry& AllottedGeometry, const double InCurre
 
 	TSharedPtr<FDebuggerViewModel> Model = ViewModel.Get();
 
-	bool bNeedUpdate = Model->NeedsUpdate();
+	bool bNeedUpdate = Model->HasSearchableAssetChanged();
 
 	// We haven't reached the update point yet
 	if (CurrentConsecutiveFrames < ConsecutiveFramesUpdateThreshold)
@@ -1641,7 +1625,7 @@ void SDebuggerView::Tick(const FGeometry& AllottedGeometry, const double InCurre
 	if (bNeedUpdate)
 	{
 		Model->OnUpdate();
-		if (UpdateSelection())
+		if (UpdateNodeSelection())
 		{
 			Model->OnUpdateNodeSelection(SelectedNodeId);
 			UpdateViews();
@@ -1655,7 +1639,7 @@ void SDebuggerView::Tick(const FGeometry& AllottedGeometry, const double InCurre
 	DrawVisualization();
 }
 
-bool SDebuggerView::UpdateSelection()
+bool SDebuggerView::UpdateNodeSelection()
 {
 	TSharedPtr<FDebuggerViewModel> Model = ViewModel.Get();
 
@@ -1683,7 +1667,7 @@ bool SDebuggerView::UpdateSelection()
 				.Padding(10.0f)
 				[
 					SNew(SButton)
-					.Text(FText::FromString(Model->GetPoseSearchDatabase()->GetName()))
+					.Text(FText::FromString(Model->GetSearchableAsset()->GetName()))
 					.HAlign(HAlign_Center)
 					.VAlign(VAlign_Center)
 					.ContentPadding(10.0f)
@@ -1699,11 +1683,10 @@ bool SDebuggerView::UpdateSelection()
 void SDebuggerView::UpdateViews() const
 {
 	const FTraceMotionMatchingStateMessage* State = ViewModel.Get()->GetMotionMatchingState();
-	const UPoseSearchDatabase* Database = ViewModel.Get()->GetPoseSearchDatabase();
-	if (State && Database)
+	if (State)
 	{
-		DatabaseView->Update(*State, *Database);
-		DetailsView->Update(*State, *Database);
+		DatabaseView->Update(*State);
+		DetailsView->Update(*State);
 	}
 }
 
@@ -1713,11 +1696,10 @@ void SDebuggerView::DrawVisualization() const
 	check(DebuggerWorld);
 
 	const FTraceMotionMatchingStateMessage* State = ViewModel.Get()->GetMotionMatchingState();
-	const UPoseSearchDatabase* Database = ViewModel.Get()->GetPoseSearchDatabase();
 	const FTransform* Transform = ViewModel.Get()->GetRootTransform();
-	if (State && Database && Transform)
+	if (State && Transform)
 	{
-		DrawFeatures(*DebuggerWorld, *State, *Database, *Transform, ViewModel.Get()->GetMeshComponent());
+		DrawFeatures(*DebuggerWorld, *State, *Transform, ViewModel.Get()->GetMeshComponent());
 	}
 }
 
@@ -1726,27 +1708,20 @@ TArray<TSharedRef<FDebuggerDatabaseRowData>> SDebuggerView::GetSelectedDatabaseR
 	return DatabaseView->GetDatabaseRows()->GetSelectedItems();
 }
 
-const TSharedRef<FDebuggerDatabaseRowData>& SDebuggerView::GetPoseIdxDatabaseRow(int32 PoseIdx) const
-{
-	return DatabaseView->GetPoseIdxDatabaseRow(PoseIdx);
-}
-
 void SDebuggerView::DrawFeatures(
 	const UWorld& DebuggerWorld,
 	const FTraceMotionMatchingStateMessage& State,
-	const UPoseSearchDatabase& Database,
 	const FTransform& Transform,
 	const USkinnedMeshComponent* Mesh
 ) const
 {
 	// Set shared state
 	FDebugDrawParams DrawParams;
-	DrawParams.Database = &Database;
 	DrawParams.World = &DebuggerWorld;
 	DrawParams.RootTransform = Transform;
-	// Single frame render
-	DrawParams.DefaultLifeTime = 0.0f;
+	DrawParams.DefaultLifeTime = 0.0f; // Single frame render
 	DrawParams.Mesh = Mesh;
+
 	const TObjectPtr<UPoseSearchDebuggerReflection>& Reflection = DetailsView->GetReflection();
 
 	auto SetDrawFlags = [](FDebugDrawParams& InDrawParams, const FPoseSearchDebuggerFeatureDrawOptions& Options)
@@ -1771,59 +1746,81 @@ void SDebuggerView::DrawFeatures(
 		}
 	};
 
-	// Draw query vector
-	SetDrawFlags(DrawParams, Reflection->QueryDrawOptions);
-	EnumAddFlags(DrawParams.Flags, EDebugDrawFlags::DrawQuery);
-	DrawFeatureVector(DrawParams, State.QueryVector);
-	EnumRemoveFlags(DrawParams.Flags, EDebugDrawFlags::DrawQuery);
 
-	const TSharedPtr<SListView<TSharedRef<FDebuggerDatabaseRowData>>>& DatabaseRows = DatabaseView->GetDatabaseRows();
-	TArray<TSharedRef<FDebuggerDatabaseRowData>> Selected = DatabaseRows->GetSelectedItems();
-	
-	// Red for non-active database view
-	SetDrawFlags(DrawParams, Reflection->SelectedPoseDrawOptions);
-
-	// Draw any selected database vectors
-	for (const TSharedRef<FDebuggerDatabaseRowData>& Row : Selected)
+	// Draw query vector
 	{
-		DrawFeatureVector(DrawParams, Row->PoseIdx);
+		DrawParams.Database = ViewModel.Get()->GetCurrentDatabase();
+		SetDrawFlags(DrawParams, Reflection->QueryDrawOptions);
+		EnumAddFlags(DrawParams.Flags, EDebugDrawFlags::DrawQuery);
+		DrawFeatureVector(DrawParams, State.QueryVector);
+		EnumRemoveFlags(DrawParams.Flags, EDebugDrawFlags::DrawQuery);
 	}
 
-	Selected = DatabaseView->GetActiveRow()->GetSelectedItems();
-	
-	// Active row should only have 0 or 1
-	check(Selected.Num() < 2);
 
-	if (!Selected.IsEmpty())
+	// Draw selected poses
 	{
-		// Use the motion-matching state's pose idx, as the active row may be update-throttled at this point
-		DrawFeatureVector(DrawParams, State.DbPoseIdx);
-	}
+		const TSharedPtr<SListView<TSharedRef<FDebuggerDatabaseRowData>>>& DatabaseRows = DatabaseView->GetDatabaseRows();
+		TArray<TSharedRef<FDebuggerDatabaseRowData>> SelectedRows = DatabaseRows->GetSelectedItems();
+	
+		// Red for non-active database view
+		SetDrawFlags(DrawParams, Reflection->SelectedPoseDrawOptions);
 
-	Selected = DatabaseView->GetContinuingPoseRow()->GetSelectedItems();
+		// Draw any selected database vectors
+		for (const TSharedRef<FDebuggerDatabaseRowData>& Row : SelectedRows)
+		{
+			DrawParams.Database = Row->SourceDatabase.Get();
+			DrawFeatureVector(DrawParams, Row->PoseIdx);
+		}
+	}
 
-	// ContinuingPose row should only have 0 or 1
-	check(Selected.Num() < 2);
 
-	if (!Selected.IsEmpty())
+	// Draw active pose
 	{
-		// Use the motion-matching state's pose idx, as the ContinuingPose row may be update-throttled at this point
-		DrawFeatureVector(DrawParams, (*Selected.begin())->PoseIdx);
+		TArray<TSharedRef<FDebuggerDatabaseRowData>> ActiveRows = DatabaseView->GetActiveRow()->GetSelectedItems();
+
+		// Active row should only have 0 or 1
+		check(ActiveRows.Num() < 2);
+
+		if (!ActiveRows.IsEmpty())
+		{		
+			// Use the motion-matching state's pose idx, as the active row may be update-throttled at this point
+			DrawParams.Database = ActiveRows[0]->SourceDatabase.Get();
+			DrawFeatureVector(DrawParams, ActiveRows[0]->PoseIdx);
+		}
 	}
 
-	FSkeletonDrawParams SkeletonDrawParams;
-	if (Reflection->bDrawSelectedSkeleton)
+
+	// Draw continuing pose
 	{
-		SkeletonDrawParams.Flags |= ESkeletonDrawFlags::SelectedPose;
+		TArray<TSharedRef<FDebuggerDatabaseRowData>> ContinuingRows = DatabaseView->GetContinuingPoseRow()->GetSelectedItems();
+
+		// ContinuingPose row should only have 0 or 1
+		check(ContinuingRows.Num() < 2);
+
+		if (!ContinuingRows.IsEmpty())
+		{
+			DrawParams.Database = ContinuingRows[0]->SourceDatabase.Get();
+			DrawFeatureVector(DrawParams, ContinuingRows[0]->PoseIdx);
+		}
 	}
-	if (Reflection->bDrawActiveSkeleton)
+
+
+	// Draw skeleton
 	{
-		SkeletonDrawParams.Flags |= ESkeletonDrawFlags::ActivePose;
-	}
+		FSkeletonDrawParams SkeletonDrawParams;
+		if (Reflection->bDrawSelectedSkeleton)
+		{
+			SkeletonDrawParams.Flags |= ESkeletonDrawFlags::SelectedPose;
+		}
+		if (Reflection->bDrawActiveSkeleton)
+		{
+			SkeletonDrawParams.Flags |= ESkeletonDrawFlags::ActivePose;
+		}
 
-	SkeletonDrawParams.Flags |= ESkeletonDrawFlags::Asset;
+		SkeletonDrawParams.Flags |= ESkeletonDrawFlags::Asset;
 
-	ViewModel.Get()->OnDraw(SkeletonDrawParams);
+		ViewModel.Get()->OnDraw(SkeletonDrawParams);
+	}
 }
 
 int32 SDebuggerView::SelectView() const
@@ -1861,24 +1858,24 @@ int32 SDebuggerView::SelectView() const
 	return Debugger;
 }
 
-void SDebuggerView::OnPoseSelectionChanged(int32 PoseIdx, float Time)
+void SDebuggerView::OnPoseSelectionChanged(const UPoseSearchDatabase* Database, int32 DbPoseIdx, float Time)
 {
 	const TSharedPtr<FDebuggerViewModel> Model = ViewModel.Get();
 	const FTraceMotionMatchingStateMessage* State = Model->GetMotionMatchingState();
-	const UPoseSearchDatabase* Database = Model->GetPoseSearchDatabase();
 
-	if (State && Database)
+	if (State)
 	{
-		DetailsView->Update(*State, *Database);
+		DetailsView->Update(*State);
 	}
 	
-	if (PoseIdx == INDEX_NONE)
+	if (DbPoseIdx == INDEX_NONE)
 	{
 		Model->ClearSelectedSkeleton();
 	}
 	else
 	{
-		Model->ShowSelectedSkeleton(PoseIdx, Time);
+		Model->ShowSelectedSkeleton(Database, DbPoseIdx, Time);
+
 		// Stop asset player when switching selections
 		Model->StopSelection();
 	}
@@ -2126,60 +2123,31 @@ const FTraceMotionMatchingStateMessage* FDebuggerViewModel::GetMotionMatchingSta
 	return ActiveMotionMatchingState;
 }
 
-const UPoseSearchDatabase* FDebuggerViewModel::GetPoseSearchDatabase() const
+const UPoseSearchDatabase* FDebuggerViewModel::GetCurrentDatabase() const
 {
-	if (ActiveMotionMatchingState)
-	{
-		if (const UPoseSearchDatabase* Database = ActiveMotionMatchingState->GetPoseSearchDatabase())
-		{
-			if (Database->Schema && Database->Schema->IsValid())
-			{
-				return Database;
-			}
-		}
-	}
-	return nullptr;
+	return ActiveMotionMatchingState ? ActiveMotionMatchingState->GetCurrentDatabase() : nullptr;
 }
 
-const FPoseSearchDatabaseSequence* FDebuggerViewModel::GetAnimSequence(int32 SequenceIdx) const
+const UPoseSearchSearchableAsset* FDebuggerViewModel::GetSearchableAsset() const
 {
-	const UPoseSearchDatabase* Database = GetPoseSearchDatabase();
-	if (Database && Database->Sequences.IsValidIndex(SequenceIdx))
+	if (ActiveMotionMatchingState)
 	{
-		if (const FPoseSearchDatabaseSequence* DatabaseSequence = &Database->Sequences[SequenceIdx])
-		{
-			return DatabaseSequence;
-		}
+		return FTraceMotionMatchingState::GetObjectFromId<UPoseSearchSearchableAsset>(ActiveMotionMatchingState->SearchableAssetId);
 	}
-	return nullptr;
-}
 
-const FPoseSearchDatabaseBlendSpace* FDebuggerViewModel::GetBlendSpace(int32 BlendSpaceIdx) const
-{
-	const UPoseSearchDatabase* Database = GetPoseSearchDatabase();
-	if (Database && Database->BlendSpaces.IsValidIndex(BlendSpaceIdx))
-	{
-		if (const FPoseSearchDatabaseBlendSpace* DatabaseBlendSpace = &Database->BlendSpaces[BlendSpaceIdx])
-		{
-			return DatabaseBlendSpace;
-		}
-	}
 	return nullptr;
 }
 
-void FDebuggerViewModel::ShowSelectedSkeleton(int32 PoseIdx, float Time)
+void FDebuggerViewModel::ShowSelectedSkeleton(const UPoseSearchDatabase* Database, int32 DbPoseIdx, float Time)
 {
-	const UPoseSearchDatabase* Database = GetPoseSearchDatabase();
-	if (!Database)
-	{
-		return;
-	}
 	UPoseSearchMeshComponent* Component = Skeletons[SelectedPose].Component;
 	if (!Component)
 	{
 		return;
 	}
-	const FPoseSearchIndexAsset* IndexAsset = Database->GetSearchIndex()->FindAssetForPose(PoseIdx);
+
+	const FPoseSearchIndex* SearchIndex = Database->GetSearchIndex();
+	const FPoseSearchIndexAsset* IndexAsset = SearchIndex ? SearchIndex->FindAssetForPose(DbPoseIdx) : nullptr;
 	if (!IndexAsset)
 	{
 		return;
@@ -2191,6 +2159,7 @@ void FDebuggerViewModel::ShowSelectedSkeleton(int32 PoseIdx, float Time)
 	Skeletons[SelectedPose].Type = IndexAsset->Type;
 	Skeletons[SelectedPose].Time = Time;
 	Skeletons[SelectedPose].bMirrored = IndexAsset->bMirrored;
+	Skeletons[SelectedPose].SourceDatabase = Database;
 	Skeletons[SelectedPose].AssetIdx = IndexAsset->SourceAssetIdx;
 	Skeletons[SelectedPose].BlendParameters = IndexAsset->BlendParameters;
 }
@@ -2215,11 +2184,10 @@ const FTransform* FDebuggerViewModel::GetRootTransform() const
 	return RootTransform;
 }
 
-bool FDebuggerViewModel::NeedsUpdate() const
+bool FDebuggerViewModel::HasSearchableAssetChanged() const
 {
-	const UPoseSearchDatabase* NewDatabase = GetPoseSearchDatabase();
-	const bool bDatabaseChanged = NewDatabase != CurrentDatabase;
-	return bDatabaseChanged;
+	uint64 NewSearchableAssetId = ActiveMotionMatchingState ? ActiveMotionMatchingState->SearchableAssetId : 0;
+	return NewSearchableAssetId != SearchableAssetId;
 }
 
 void FDebuggerViewModel::OnUpdate()
@@ -2265,32 +2233,36 @@ void FDebuggerViewModel::OnUpdateNodeSelection(int32 InNodeId)
 		}
 	}
 
-	const UPoseSearchDatabase* NewDatabase = GetPoseSearchDatabase();
-
-	if (ActiveMotionMatchingState && NewDatabase)
+	if (ActiveMotionMatchingState)
 	{
-		const FPoseSearchIndexAsset* IndexAsset = NewDatabase->GetSearchIndex()->FindAssetForPose(ActiveMotionMatchingState->DbPoseIdx);
+		const UPoseSearchDatabase* CurrentDatabase = ActiveMotionMatchingState->GetCurrentDatabase();
+		const FPoseSearchIndex* CurrentSearchIndex = CurrentDatabase ? CurrentDatabase->GetSearchIndex() : nullptr;
+		int32 CurrentPoseIdx = ActiveMotionMatchingState->GetCurrentDatabasePoseIndex();
+
+		const FPoseSearchIndexAsset* IndexAsset = CurrentSearchIndex ? CurrentSearchIndex->FindAssetForPose(CurrentPoseIdx) : nullptr;
 
 		if (IndexAsset)
 		{
 			Skeletons[Asset].Type = IndexAsset->Type;
 			Skeletons[Asset].bMirrored = IndexAsset->bMirrored;
+			Skeletons[Asset].SourceDatabase = CurrentDatabase;
 			Skeletons[Asset].AssetIdx = IndexAsset->SourceAssetIdx;
 			Skeletons[Asset].BlendParameters = IndexAsset->BlendParameters;
 		}
 	}
 
-	if (NewDatabase != CurrentDatabase)
+	uint64 NewSearchableAssetId = ActiveMotionMatchingState ? ActiveMotionMatchingState->SearchableAssetId : 0;
+	if (NewSearchableAssetId != SearchableAssetId)
 	{
 		ClearSelectedSkeleton();
-		CurrentDatabase = NewDatabase;
+		SearchableAssetId = NewSearchableAssetId;
 	}
 }
 
 void FDebuggerViewModel::OnDraw(FSkeletonDrawParams& DrawParams)
 {
-	const UPoseSearchDatabase* PoseSearchDatabase = GetPoseSearchDatabase();
-	if (!PoseSearchDatabase)
+	const UPoseSearchDatabase* CurrentDatabase = GetCurrentDatabase();
+	if (!CurrentDatabase)
 		return;
 
 	// Returns if it is to be drawn this frame
@@ -2313,7 +2285,7 @@ void FDebuggerViewModel::OnDraw(FSkeletonDrawParams& DrawParams)
 
 	UPoseSearchMeshComponent::FUpdateContext UpdateContext;
 
-	UpdateContext.MirrorDataTable = PoseSearchDatabase->Schema->MirrorDataTable;
+	UpdateContext.MirrorDataTable = CurrentDatabase->Schema->MirrorDataTable;
 	UpdateContext.CompactPoseMirrorBones = &CompactPoseMirrorBones;
 	UpdateContext.ComponentSpaceRefRotations = &ComponentSpaceRefRotations;
 
@@ -2323,7 +2295,7 @@ void FDebuggerViewModel::OnDraw(FSkeletonDrawParams& DrawParams)
 
 		if (Skeletons[SelectedPose].Type == ESearchIndexAssetType::Sequence)
 		{
-			const FPoseSearchDatabaseSequence* DatabaseSequence = GetAnimSequence(Skeletons[SelectedPose].AssetIdx);
+			const FPoseSearchDatabaseSequence* DatabaseSequence = Skeletons[SelectedPose].GetAnimSequence();
 			if (DatabaseSequence)
 			{
 				UpdateContext.Type = ESearchIndexAssetType::Sequence;
@@ -2336,7 +2308,7 @@ void FDebuggerViewModel::OnDraw(FSkeletonDrawParams& DrawParams)
 		}
 		else if (Skeletons[SelectedPose].Type == ESearchIndexAssetType::BlendSpace)
 		{
-			const FPoseSearchDatabaseBlendSpace* DatabaseBlendSpace = GetBlendSpace(Skeletons[SelectedPose].AssetIdx);
+			const FPoseSearchDatabaseBlendSpace* DatabaseBlendSpace = Skeletons[SelectedPose].GetBlendSpace();
 			if (DatabaseBlendSpace)
 			{
 				UpdateContext.Type = ESearchIndexAssetType::BlendSpace;
@@ -2367,7 +2339,7 @@ void FDebuggerViewModel::OnDraw(FSkeletonDrawParams& DrawParams)
 
 		if (Skeletons[SelectedPose].Type == ESearchIndexAssetType::Sequence)
 		{
-			const FPoseSearchDatabaseSequence* DatabaseSequence = GetAnimSequence(Skeletons[Asset].AssetIdx);
+			const FPoseSearchDatabaseSequence* DatabaseSequence = Skeletons[Asset].GetAnimSequence();
 			if (DatabaseSequence)
 			{
 				UpdateContext.Type = ESearchIndexAssetType::Sequence;
@@ -2380,7 +2352,7 @@ void FDebuggerViewModel::OnDraw(FSkeletonDrawParams& DrawParams)
 		}
 		else if (Skeletons[SelectedPose].Type == ESearchIndexAssetType::BlendSpace)
 		{
-			const FPoseSearchDatabaseBlendSpace* DatabaseBlendSpace = GetBlendSpace(Skeletons[Asset].AssetIdx);
+			const FPoseSearchDatabaseBlendSpace* DatabaseBlendSpace = Skeletons[Asset].GetBlendSpace();
 			if (DatabaseBlendSpace)
 			{
 				UpdateContext.Type = ESearchIndexAssetType::BlendSpace;
@@ -2489,7 +2461,7 @@ void FDebuggerViewModel::UpdateAsset()
 	static float MAX_DISTANCE_RANGE = 200.f;
 	static float MAX_TIME_RANGE = 2.f;
 
-	const UPoseSearchDatabase* Database = GetPoseSearchDatabase();
+	const UPoseSearchDatabase* Database = GetCurrentDatabase();
 	if (!Database || !IsPlayingSelections())
 	{
 		return;
@@ -2510,13 +2482,13 @@ void FDebuggerViewModel::UpdateAsset()
 
 	if (AssetSkeleton.Type == ESearchIndexAssetType::Sequence)
 	{
-		const FPoseSearchDatabaseSequence* DatabaseSequence = GetAnimSequence(AssetSkeleton.AssetIdx);
+		const FPoseSearchDatabaseSequence* DatabaseSequence = AssetSkeleton.GetAnimSequence();
 		AnimAsset = DatabaseSequence->Sequence;
 		bAssetLooping = DatabaseSequence->Sequence->bLoop;
 	}
 	else if (AssetSkeleton.Type == ESearchIndexAssetType::BlendSpace)
 	{
-		const FPoseSearchDatabaseBlendSpace* DatabaseBlendSpace = GetBlendSpace(AssetSkeleton.AssetIdx);
+		const FPoseSearchDatabaseBlendSpace* DatabaseBlendSpace = AssetSkeleton.GetBlendSpace();
 		AnimAsset = DatabaseBlendSpace->BlendSpace;
 		bAssetLooping = DatabaseBlendSpace->BlendSpace->bLoop;
 	}
@@ -2590,7 +2562,7 @@ const USkinnedMeshComponent* FDebuggerViewModel::GetMeshComponent() const
 
 void FDebuggerViewModel::FillCompactPoseAndComponentRefRotations()
 {
-	const UPoseSearchDatabase* Database = GetPoseSearchDatabase();
+	const UPoseSearchDatabase* Database = GetCurrentDatabase();
 	if (Database != nullptr)
 	{
 		UMirrorDataTable* MirrorDataTable = Database->Schema->MirrorDataTable;
@@ -2614,7 +2586,7 @@ void FDebuggerViewModel::FillCompactPoseAndComponentRefRotations()
 
 void FDebuggerViewModel::PlaySelection(int32 PoseIdx, float Time)
 {
-	const UPoseSearchDatabase* Database = GetPoseSearchDatabase();
+	const UPoseSearchDatabase* Database = GetCurrentDatabase();
 	if (!Database)
 	{
 		return;
@@ -2659,6 +2631,24 @@ void FDebuggerViewModel::OnWorldCleanup(UWorld* InWorld, bool bSessionEnded, boo
 	bSkeletonsInitialized = false;
 }
 
+const FPoseSearchDatabaseSequence* FDebuggerViewModel::FSkeleton::GetAnimSequence() const
+{
+	if ((Type == ESearchIndexAssetType::Sequence) && SourceDatabase.IsValid() && SourceDatabase->Sequences.IsValidIndex(AssetIdx))
+	{
+		return &SourceDatabase->Sequences[AssetIdx];
+	}
+	return nullptr;
+}
+
+const FPoseSearchDatabaseBlendSpace* FDebuggerViewModel::FSkeleton::GetBlendSpace() const
+{
+	if ((Type == ESearchIndexAssetType::BlendSpace) && SourceDatabase.IsValid() && SourceDatabase->BlendSpaces.IsValidIndex(AssetIdx))
+	{
+		return &SourceDatabase->BlendSpaces[AssetIdx];
+	}
+	return nullptr;
+}
+
 FDebugger* FDebugger::Debugger;
 void FDebugger::Initialize()
 {
diff --git a/Engine/Plugins/Experimental/Animation/PoseSearch/Source/Editor/Private/PoseSearchDebugger.h b/Engine/Plugins/Experimental/Animation/PoseSearch/Source/Editor/Private/PoseSearchDebugger.h
--- a/Engine/Plugins/Experimental/Animation/PoseSearch/Source/Editor/Private/PoseSearchDebugger.h
+++ b/Engine/Plugins/Experimental/Animation/PoseSearch/Source/Editor/Private/PoseSearchDebugger.h
@@ -181,7 +181,7 @@ struct FSkeletonDrawParams
 };
 
 /** Sets model selection data on row selection */
-DECLARE_DELEGATE_TwoParams(FOnPoseSelectionChanged, int32, float)
+DECLARE_DELEGATE_ThreeParams(FOnPoseSelectionChanged, const UPoseSearchDatabase*, int32, float)
 class FDebuggerViewModel;
 
 /**
@@ -196,7 +196,7 @@ class SDebuggerDatabaseView : public SCompoundWidget
 	SLATE_END_ARGS()
 	
 	void Construct(const FArguments& InArgs);
-	void Update(const FTraceMotionMatchingStateMessage& State, const UPoseSearchDatabase& Database);
+	void Update(const FTraceMotionMatchingStateMessage& State);
 
 	const TSharedPtr<SListView<TSharedRef<FDebuggerDatabaseRowData>>>& GetActiveRow() const { return ActiveView.ListView; }
 	const TSharedPtr<SListView<TSharedRef<FDebuggerDatabaseRowData>>>& GetContinuingPoseRow() const { return ContinuingPoseView.ListView; }
@@ -217,7 +217,7 @@ private:
 	const FColumnMap* GetColumnMap() const { return &Columns; }
 
 	/** Creates widgets for every pose in the database, initializing the static data in the process */
-	void CreateRows(const UPoseSearchDatabase& Database);
+	void CreateRows(const FTraceMotionMatchingStateMessage& State);
 
 	/** Sorts the database by the current sort predicate, updating the view order */
 	void SortDatabaseRows();
@@ -226,9 +226,6 @@ private:
 
 	void ComputeFilteredDatabaseRowsColors();
 
-	/** Sets dynamic data for each row, such as score at the current time */
-	void UpdateRows(const FTraceMotionMatchingStateMessage& State, const UPoseSearchDatabase& Database);
-
 	/** Acquires sort predicate for the given column */
 	EColumnSortMode::Type GetColumnSortMode(const FName ColumnId) const;
 
@@ -307,7 +304,6 @@ private:
 	/** All database poses */
 	TArray<TSharedRef<FDebuggerDatabaseRowData>> UnfilteredDatabaseRows;
 
-	TWeakObjectPtr<const UPoseSearchDatabase> RowsSourceDatabase = nullptr;
 	TArrayView<const bool> DatabaseSequenceFilter;
 	TArrayView<const bool> DatabaseBlendSpaceFilter;
 
@@ -337,14 +333,14 @@ class SDebuggerDetailsView : public SCompoundWidget
 	virtual ~SDebuggerDetailsView() override;
 	
 	void Construct(const FArguments& InArgs);
-	void Update(const FTraceMotionMatchingStateMessage& State, const UPoseSearchDatabase& Database) const;
+	void Update(const FTraceMotionMatchingStateMessage& State) const;
 
 	/** Get a const version of our reflection object */
 	const TObjectPtr<UPoseSearchDebuggerReflection>& GetReflection() const { return Reflection; }
 	
 private:
 	/** Update our details view object with new state information */
-	void UpdateReflection(const FTraceMotionMatchingStateMessage& State, const UPoseSearchDatabase& Database) const;
+	void UpdateReflection(const FTraceMotionMatchingStateMessage& State) const;
 	
 	TWeakPtr<SDebuggerView> ParentDebuggerViewPtr;
 
@@ -376,15 +372,16 @@ public:
 	virtual void Tick(const FGeometry& AllottedGeometry, const double InCurrentTime, const float InDeltaTime) override;
 	virtual FName GetName() const override;
 	virtual uint64 GetObjectId() const override;
+
+	TSharedPtr<FDebuggerViewModel> GetViewModel() const { return ViewModel.Get(); }
 	TArray<TSharedRef<FDebuggerDatabaseRowData>> GetSelectedDatabaseRows() const;
-	const TSharedRef<FDebuggerDatabaseRowData>& GetPoseIdxDatabaseRow(int32 PoseIdx) const;
 
 private:
 	/** Called each frame to draw features of the query vector & database selections */
-	void DrawFeatures(const UWorld& DebuggerWorld, const FTraceMotionMatchingStateMessage& State, const UPoseSearchDatabase& Database, const FTransform& Transform, const USkinnedMeshComponent* Mesh) const;
+	void DrawFeatures(const UWorld& DebuggerWorld, const FTraceMotionMatchingStateMessage& State, const FTransform& Transform, const USkinnedMeshComponent* Mesh) const;
 	
 	/** Check if a node selection was made, true if a node is selected */
-	bool UpdateSelection();
+	bool UpdateNodeSelection();
 
 	/** Update the database and details views */
 	void UpdateViews() const;
@@ -397,7 +394,7 @@ private:
 	/** Callback when a button in the selection view is clicked */
 	FReply OnUpdateNodeSelection(int32 InSelectedNodeId);
 
-	void OnPoseSelectionChanged(int32 PoseIdx, float Time);
+	void OnPoseSelectionChanged(const UPoseSearchDatabase* Database, int32 PoseIdx, float Time);
 
 	/** Button interaction to toggle play / stop of the asset */
 	FReply TogglePlaySelectedAssets() const;
@@ -478,13 +475,14 @@ public:
 
 	// Used for view callbacks
     const FTraceMotionMatchingStateMessage* GetMotionMatchingState() const;
-    const UPoseSearchDatabase* GetPoseSearchDatabase() const;
+	const UPoseSearchDatabase* GetCurrentDatabase() const;
+	const UPoseSearchSearchableAsset* GetSearchableAsset() const;
 	const TArray<int32>* GetNodeIds() const;
 	int32 GetNodesNum() const;
 	const FTransform* GetRootTransform() const;
 
 	/** Checks if Update must be called */
-	bool NeedsUpdate() const;
+	bool HasSearchableAssetChanged() const;
 
 	/** Update motion matching states for frame */
 	void OnUpdate();
@@ -502,7 +500,7 @@ public:
 	const FPoseSearchDatabaseBlendSpace* GetBlendSpace(int32 BlendSpaceIdx) const;
 
 	/** Sets the selected pose skeleton*/
-	void ShowSelectedSkeleton(int32 PoseIdx, float Time);
+	void ShowSelectedSkeleton(const UPoseSearchDatabase* Database, int32 DbPoseIdx, float Time);
 	
 	/** Clears the selected pose skeleton */
 	void ClearSelectedSkeleton();
@@ -546,7 +544,8 @@ private:
 	/** Currently active motion matching state based on node selection in the view */
 	const FTraceMotionMatchingStateMessage* ActiveMotionMatchingState = nullptr;
 
-	TWeakObjectPtr<const UPoseSearchDatabase> CurrentDatabase = nullptr;
+	/** Active motion matching state's searchable asset */
+	uint64 SearchableAssetId = 0;
 
 	/** Current Skeletal Mesh Component Id for the AnimInstance */
 	uint64 SkeletalMeshComponentId = 0;
@@ -578,7 +577,10 @@ private:
 		/** Type of the asset being played */
 		ESearchIndexAssetType Type = ESearchIndexAssetType::Invalid;
 
-		/** Active asset index being used for this setting this skeleton */
+		/** Source database for this skeleton  */
+		TWeakObjectPtr<const UPoseSearchDatabase> SourceDatabase = nullptr;
+
+		/** Source asset for this skeleton */
 		int32 AssetIdx = 0;
 		
 		/** Time in the sequence this skeleton is accessing */
@@ -589,6 +591,9 @@ private:
 
 		/** Blend Parameters if asset is a BlendSpace */
 		FVector BlendParameters = FVector::Zero();
+
+		const FPoseSearchDatabaseSequence* GetAnimSequence() const;
+		const FPoseSearchDatabaseBlendSpace* GetBlendSpace() const;
 	};
 	
 	/** Index for each type of skeleton we store for debug visualization */
diff --git a/Engine/Plugins/Experimental/Animation/PoseSearch/Source/Runtime/Private/PoseSearchLibrary.cpp b/Engine/Plugins/Experimental/Animation/PoseSearch/Source/Runtime/Private/PoseSearchLibrary.cpp
--- a/Engine/Plugins/Experimental/Animation/PoseSearch/Source/Runtime/Private/PoseSearchLibrary.cpp
+++ b/Engine/Plugins/Experimental/Animation/PoseSearch/Source/Runtime/Private/PoseSearchLibrary.cpp
@@ -209,6 +209,180 @@ void FMotionMatchingState::JumpToPose(const FAnimationUpdateContext& Context, co
 	Flags |= EMotionMatchingFlags::JumpedToPose;
 }
 
+static void TraceMotionMatchingState(
+	const FAnimationUpdateContext& UpdateContext,
+	const UPoseSearchSearchableAsset* Searchable,
+	UE::PoseSearch::FSearchContext& SearchContext,
+	const FMotionMatchingState& MotionMatchingState,
+	const FGameplayTagQuery* DatabaseTagQuery,
+	const FTrajectorySampleRange& Trajectory,
+	const UE::PoseSearch::FSearchResult& LastResult
+)
+{
+#if UE_POSE_SEARCH_TRACE_ENABLED
+	using namespace UE::PoseSearch;
+
+	if (!MotionMatchingState.CurrentSearchResult.IsValid())
+	{
+		return;
+	}
+
+	const float DeltaTime = UpdateContext.GetDeltaTime();
+
+	FTraceMotionMatchingState TraceState;
+	bool bContinuingPoseTraced = false;
+	bool bCurrentPoseTraced = false;
+	while (!SearchContext.BestCandidates.IsEmpty())
+	{
+		FSearchContext::FPoseCandidate PoseCandidate;
+		SearchContext.BestCandidates.Pop(PoseCandidate);
+
+		const uint64 DatabaseId = FTraceMotionMatchingState::GetIdFromObject(PoseCandidate.Database);
+		const int32 DbEntryIdx = TraceState.DatabaseEntries.AddUnique({DatabaseId});
+		FTraceMotionMatchingStateDatabaseEntry& DbEntry = TraceState.DatabaseEntries[DbEntryIdx];
+
+		FTraceMotionMatchingStatePoseEntry PoseEntry;
+		PoseEntry.DbPoseIdx = PoseCandidate.PoseIdx;
+		PoseEntry.Cost = PoseCandidate.Cost;
+
+		if (PoseCandidate.Database == MotionMatchingState.CurrentSearchResult.Database)
+		{
+			DbEntry.Flags |= FTraceMotionMatchingStateDatabaseEntry::EFlags::CurrentDatabase;
+		}
+		
+		if (PoseCandidate == MotionMatchingState.CurrentSearchResult)
+		{
+			PoseEntry.Flags |= FTraceMotionMatchingStatePoseEntry::EFlags::CurrentPose;
+			bCurrentPoseTraced = true;
+		}
+
+		if (PoseCandidate == LastResult)
+		{
+			PoseEntry.Flags |= FTraceMotionMatchingStatePoseEntry::EFlags::ContinuingPose;
+			bContinuingPoseTraced = true;
+		}
+
+		DbEntry.PoseEntries.Add(PoseEntry);
+	}
+
+	if (!bContinuingPoseTraced && LastResult.IsValid())
+	{
+		const uint64 LastResultDbId = FTraceMotionMatchingState::GetIdFromObject(LastResult.Database.Get());
+
+		int32 LastResultDbIdx = TraceState.DatabaseEntries.AddUnique({LastResultDbId});
+		FTraceMotionMatchingStateDatabaseEntry& DbEntry = TraceState.DatabaseEntries[LastResultDbIdx];
+
+		int32 LastResultPoseEntryIdx = DbEntry.PoseEntries.AddUnique({LastResult.PoseIdx});
+		FTraceMotionMatchingStatePoseEntry& PoseEntry = DbEntry.PoseEntries[LastResultPoseEntryIdx];
+
+		PoseEntry.Cost = LastResult.PoseCost.GetTotalCost();
+		PoseEntry.Flags |= FTraceMotionMatchingStatePoseEntry::EFlags::ContinuingPose;
+		bContinuingPoseTraced = true;
+	}
+
+	if (!bCurrentPoseTraced && MotionMatchingState.CurrentSearchResult.IsValid())
+	{
+		const uint64 CurrentResultDbId = FTraceMotionMatchingState::GetIdFromObject(MotionMatchingState.CurrentSearchResult.Database.Get());
+
+		int32 DbEntryIdx = TraceState.DatabaseEntries.AddUnique({CurrentResultDbId});
+		FTraceMotionMatchingStateDatabaseEntry& DbEntry = TraceState.DatabaseEntries[DbEntryIdx];
+		DbEntry.Flags |= FTraceMotionMatchingStateDatabaseEntry::EFlags::CurrentDatabase;
+
+		int32 PoseEntryIdx = DbEntry.PoseEntries.AddUnique({MotionMatchingState.CurrentSearchResult.PoseIdx});
+		FTraceMotionMatchingStatePoseEntry& PoseEntry = DbEntry.PoseEntries[PoseEntryIdx];
+
+		PoseEntry.Cost = MotionMatchingState.CurrentSearchResult.PoseCost.GetTotalCost();
+		PoseEntry.Flags |= FTraceMotionMatchingStatePoseEntry::EFlags::CurrentPose;
+		bCurrentPoseTraced = true;
+	}
+
+	TraceState.CurrentDbEntryIdx = TraceState.DatabaseEntries.IndexOfByPredicate(
+		[](const auto& DbEntry)
+		{
+			return EnumHasAnyFlags(DbEntry.Flags, FTraceMotionMatchingStateDatabaseEntry::EFlags::CurrentDatabase);
+		});
+
+	if (TraceState.CurrentDbEntryIdx != INDEX_NONE)
+	{
+		TraceState.CurrentPoseEntryIdx = TraceState.DatabaseEntries[TraceState.CurrentDbEntryIdx].PoseEntries.IndexOfByPredicate(
+			[](const auto& PoseEntry)
+			{
+				return EnumHasAnyFlags(PoseEntry.Flags, FTraceMotionMatchingStatePoseEntry::EFlags::CurrentPose);
+			});
+	}
+
+
+	float SimLinearVelocity, SimAngularVelocity, AnimLinearVelocity, AnimAngularVelocity;
+
+	if (DeltaTime > SMALL_NUMBER)
+	{
+		// simulation
+
+		int32 FirstIdx = 0;
+		const FTrajectorySample PrevSample = FTrajectorySampleRange::IterSampleTrajectory(
+			Trajectory.Samples,
+			ETrajectorySampleDomain::Time,
+			-DeltaTime, FirstIdx);
+
+		const FTrajectorySample CurrSample = FTrajectorySampleRange::IterSampleTrajectory(
+			Trajectory.Samples,
+			ETrajectorySampleDomain::Time,
+			0.0f, FirstIdx);
+
+		const FTransform SimDelta = CurrSample.Transform.GetRelativeTransform(PrevSample.Transform);
+
+		SimLinearVelocity = SimDelta.GetTranslation().Size() / DeltaTime;
+		SimAngularVelocity = FMath::RadiansToDegrees(SimDelta.GetRotation().GetAngle()) / DeltaTime;
+
+		// animation
+
+		const FTransform AnimDelta = MotionMatchingState.RootMotionTransformDelta;
+
+		AnimLinearVelocity = AnimDelta.GetTranslation().Size() / DeltaTime;
+		AnimAngularVelocity = FMath::RadiansToDegrees(AnimDelta.GetRotation().GetAngle()) / DeltaTime;
+	}
+	else
+	{
+		SimLinearVelocity = 0.0f;
+		SimAngularVelocity = 0.0f;
+		AnimLinearVelocity = 0.0f;
+		AnimAngularVelocity = 0.0f;
+	}
+
+	TArray<bool> DatabaseSequenceFilter;
+	ComputeDatabaseSequenceFilter(
+		MotionMatchingState.CurrentSearchResult.Database.Get(),
+		DatabaseTagQuery, 
+		DatabaseSequenceFilter);
+
+	TArray<bool> DatabaseBlendSpaceFilter;
+	ComputeDatabaseBlendSpaceFilter(
+		MotionMatchingState.CurrentSearchResult.Database.Get(),
+		DatabaseTagQuery, 
+		DatabaseBlendSpaceFilter);
+
+	if (EnumHasAnyFlags(MotionMatchingState.Flags, EMotionMatchingFlags::JumpedToFollowUp))
+	{
+		TraceState.Flags |= FTraceMotionMatchingState::EFlags::FollowupAnimation;
+	}
+
+	TraceState.SearchableAssetId = FTraceMotionMatchingState::GetIdFromObject(Searchable);
+	TraceState.ElapsedPoseJumpTime = MotionMatchingState.ElapsedPoseJumpTime;
+	TraceState.QueryVector = MotionMatchingState.CurrentSearchResult.ComposedQuery.GetValues();
+	TraceState.QueryVectorNormalized = MotionMatchingState.CurrentSearchResult.ComposedQuery.GetNormalizedValues();
+	TraceState.AssetPlayerTime = MotionMatchingState.CurrentSearchResult.AssetTime;
+	TraceState.DeltaTime = DeltaTime;
+	TraceState.SimLinearVelocity = SimLinearVelocity;
+	TraceState.SimAngularVelocity = SimAngularVelocity;
+	TraceState.AnimLinearVelocity = AnimLinearVelocity;
+	TraceState.AnimAngularVelocity = AnimAngularVelocity;
+	TraceState.DatabaseSequenceFilter = DatabaseSequenceFilter;
+	TraceState.DatabaseBlendSpaceFilter = DatabaseBlendSpaceFilter;
+
+	TraceState.Output(UpdateContext);
+#endif
+}
+
 void UpdateMotionMatchingState(
 	const FAnimationUpdateContext& Context,
 	const UPoseSearchSearchableAsset* Searchable,
@@ -242,8 +416,8 @@ void UpdateMotionMatchingState(
 
 	// Check if we can advance. Includes the case where we can advance but only by switching to a follow up asset.
 	bool bAdvanceToFollowUpAsset = false;
-	FSearchResult FollowUpAsset;
-	const bool bCanAdvance = InOutMotionMatchingState.CanAdvance(DeltaTime, bAdvanceToFollowUpAsset, FollowUpAsset);
+	FSearchResult FollowUpAssetResult;
+	const bool bCanAdvance = InOutMotionMatchingState.CanAdvance(DeltaTime, bAdvanceToFollowUpAsset, FollowUpAssetResult);
 
 	// If we can't advance or enough time has elapsed since the last pose jump then search
 	FSearchContext SearchContext;
@@ -300,7 +474,7 @@ void UpdateMotionMatchingState(
 		&& bCanAdvance
 		&& bAdvanceToFollowUpAsset)
 	{
-		InOutMotionMatchingState.JumpToPose(Context, Settings, FollowUpAsset);
+		InOutMotionMatchingState.JumpToPose(Context, Settings, FollowUpAssetResult);
 		InOutMotionMatchingState.Flags |= EMotionMatchingFlags::JumpedToFollowUp;
 	}
 
@@ -313,106 +487,15 @@ void UpdateMotionMatchingState(
 	InOutMotionMatchingState.PoseIndicesHistory.Update(InOutMotionMatchingState.CurrentSearchResult, DeltaTime, Settings.PoseReselectHistory);
 	
 	// Record debugger details
-#if UE_POSE_SEARCH_TRACE_ENABLED
-	FTraceMotionMatchingState TraceState;
-	TArray<const UPoseSearchDatabase*> Databases;
-	while (!SearchContext.BestCandidates.IsEmpty())
-	{
-		FSearchContext::FPoseCandidate PoseCandidate;
-		SearchContext.BestCandidates.Pop(PoseCandidate);
-
-		int32 DatabaseEntryIndex = 0;
-		if (!Databases.Find(PoseCandidate.Database, DatabaseEntryIndex))
-		{
-			Databases.Push(PoseCandidate.Database);
-			DatabaseEntryIndex = Databases.Num() - 1;
-
-			TraceState.DatabaseEntries.AddDefaulted();
-			TraceState.DatabaseEntries[DatabaseEntryIndex].DatabaseId = FTraceMotionMatchingState::GetIdFromDatabase(PoseCandidate.Database);
-		}
-
-		FTraceMotionMatchingStatePoseEntry PoseEntry;
-		PoseEntry.DbPoseIdx = PoseCandidate.PoseIdx;
-		PoseEntry.Cost = PoseCandidate.Cost;
-		TraceState.DatabaseEntries[DatabaseEntryIndex].PoseEntries.Add(PoseEntry);
-	}
-
-	if (InOutMotionMatchingState.CurrentSearchResult.IsValid())
-	{
-		float SimLinearVelocity, SimAngularVelocity, AnimLinearVelocity, AnimAngularVelocity;
-
-		if (DeltaTime > SMALL_NUMBER)
-		{
-			// simulation
-
-			int32 FirstIdx = 0;
-			const FTrajectorySample PrevSample = FTrajectorySampleRange::IterSampleTrajectory(
-				Trajectory.Samples,
-				ETrajectorySampleDomain::Time,
-				-DeltaTime, FirstIdx);
-
-			const FTrajectorySample CurrSample = FTrajectorySampleRange::IterSampleTrajectory(
-				Trajectory.Samples,
-				ETrajectorySampleDomain::Time,
-				0.0f, FirstIdx);
-
-			const FTransform SimDelta = CurrSample.Transform.GetRelativeTransform(PrevSample.Transform);
-
-			SimLinearVelocity = SimDelta.GetTranslation().Size() / DeltaTime;
-			SimAngularVelocity = FMath::RadiansToDegrees(SimDelta.GetRotation().GetAngle()) / DeltaTime;
-
-			// animation
-
-			const FTransform AnimDelta = InOutMotionMatchingState.RootMotionTransformDelta;
-
-			AnimLinearVelocity = AnimDelta.GetTranslation().Size() / DeltaTime;
-			AnimAngularVelocity = FMath::RadiansToDegrees(AnimDelta.GetRotation().GetAngle()) / DeltaTime;
-		}
-		else
-		{
-			SimLinearVelocity = 0.0f;
-			SimAngularVelocity = 0.0f;
-			AnimLinearVelocity = 0.0f;
-			AnimAngularVelocity = 0.0f;
-		}
-
-		TArray<bool> DatabaseSequenceFilter;
-		ComputeDatabaseSequenceFilter(
-			InOutMotionMatchingState.CurrentSearchResult.Database.Get(), 
-			DatabaseTagQuery, 
-			DatabaseSequenceFilter);
-
-		TArray<bool> DatabaseBlendSpaceFilter;
-		ComputeDatabaseBlendSpaceFilter(
-			InOutMotionMatchingState.CurrentSearchResult.Database.Get(),
-			DatabaseTagQuery, 
-			DatabaseBlendSpaceFilter);
-
-		if (EnumHasAnyFlags(InOutMotionMatchingState.Flags, EMotionMatchingFlags::JumpedToFollowUp))
-		{
-			TraceState.Flags |= FTraceMotionMatchingState::EFlags::FollowupAnimation;
-		}
-
-		TraceState.ElapsedPoseJumpTime = InOutMotionMatchingState.ElapsedPoseJumpTime;
-		// @TODO: Change this to only be the previous query, not persistently updated (i.e. if throttled)?
-		TraceState.QueryVector = InOutMotionMatchingState.CurrentSearchResult.ComposedQuery.GetValues();
-		TraceState.QueryVectorNormalized = InOutMotionMatchingState.CurrentSearchResult.ComposedQuery.GetNormalizedValues();
-		TraceState.DbPoseIdx = InOutMotionMatchingState.CurrentSearchResult.PoseIdx;
-		TraceState.SetDatabase(InOutMotionMatchingState.CurrentSearchResult.Database.Get());
-		TraceState.ContinuingPoseIdx = LastResult.PoseIdx;
-
-		TraceState.AssetPlayerTime = InOutMotionMatchingState.CurrentSearchResult.AssetTime;
-		TraceState.DeltaTime = DeltaTime;
-		TraceState.SimLinearVelocity = SimLinearVelocity;
-		TraceState.SimAngularVelocity = SimAngularVelocity;
-		TraceState.AnimLinearVelocity = AnimLinearVelocity;
-		TraceState.AnimAngularVelocity = AnimAngularVelocity;
-		TraceState.DatabaseSequenceFilter = DatabaseSequenceFilter;
-		TraceState.DatabaseBlendSpaceFilter = DatabaseBlendSpaceFilter;
-
-		TraceState.Output(Context);
-	}
-#endif
+	TraceMotionMatchingState(
+		Context,
+		Searchable,
+		SearchContext,
+		InOutMotionMatchingState,
+		DatabaseTagQuery,
+		Trajectory,
+		LastResult
+	);
 }
 
 const FPoseSearchIndexAsset* FMotionMatchingState::GetCurrentSearchIndexAsset() const
diff --git a/Engine/Plugins/Experimental/Animation/PoseSearch/Source/Runtime/Private/Trace/PoseSearchTraceLogger.cpp b/Engine/Plugins/Experimental/Animation/PoseSearch/Source/Runtime/Private/Trace/PoseSearchTraceLogger.cpp
--- a/Engine/Plugins/Experimental/Animation/PoseSearch/Source/Runtime/Private/Trace/PoseSearchTraceLogger.cpp
+++ b/Engine/Plugins/Experimental/Animation/PoseSearch/Source/Runtime/Private/Trace/PoseSearchTraceLogger.cpp
@@ -49,27 +49,27 @@ FArchive& operator<<(FArchive& Ar, FTraceMotionMatchingStatePoseEntry& Entry)
 {
 	Ar << Entry.DbPoseIdx;
 	Ar << Entry.Cost;
+	Ar << Entry.Flags;
 	return Ar;
 }
 
 FArchive& operator<<(FArchive& Ar, FTraceMotionMatchingStateDatabaseEntry& Entry)
 {
 	Ar << Entry.DatabaseId;
+	Ar << Entry.Flags;
 	Ar << Entry.PoseEntries;
 	return Ar;
 }
 
 FArchive& operator<<(FArchive& Ar, FTraceMotionMatchingState& State)
 {
+	Ar << State.SearchableAssetId;
 	Ar << State.ElapsedPoseJumpTime;
 	Ar << State.Flags;
 	Ar << State.QueryVector;
 	Ar << State.QueryVectorNormalized;
 	Ar << State.DatabaseSequenceFilter;
 	Ar << State.DatabaseBlendSpaceFilter;
-	Ar << State.DbPoseIdx;
-	Ar << State.DatabaseId;
-	Ar << State.ContinuingPoseIdx;
 	Ar << State.AssetPlayerTime;
 	Ar << State.DeltaTime;
 	Ar << State.SimLinearVelocity;
@@ -77,6 +77,8 @@ FArchive& operator<<(FArchive& Ar, FTraceMotionMatchingState& State)
 	Ar << State.AnimLinearVelocity;
 	Ar << State.AnimAngularVelocity;
 	Ar << State.DatabaseEntries;
+	Ar << State.CurrentDbEntryIdx;
+	Ar << State.CurrentPoseEntryIdx;
 	return Ar;
 }
 
@@ -108,5 +110,29 @@ void FTraceMotionMatchingState::Output(const FAnimationBaseContext& InContext)
 		<< MotionMatchingState.Data(ArchiveData.GetData(), ArchiveData.Num());
 }
 
+const UPoseSearchDatabase* FTraceMotionMatchingState::GetCurrentDatabase() const
+{
+	if (CurrentDbEntryIdx == INDEX_NONE)
+	{
+		return nullptr;
+	}
+
+	const FTraceMotionMatchingStateDatabaseEntry& DbEntry = DatabaseEntries[CurrentDbEntryIdx];
+	const UPoseSearchDatabase* Database = GetObjectFromId<UPoseSearchDatabase>(DbEntry.DatabaseId);
+	return Database;
+}
+
+int32 FTraceMotionMatchingState::GetCurrentDatabasePoseIndex() const
+{
+	if ((CurrentDbEntryIdx == INDEX_NONE) || (CurrentPoseEntryIdx == INDEX_NONE))
+	{
+		return INDEX_NONE;
+	}
+
+	const FTraceMotionMatchingStateDatabaseEntry& DbEntry = DatabaseEntries[CurrentDbEntryIdx];
+	const FTraceMotionMatchingStatePoseEntry& PoseEntry = DbEntry.PoseEntries[CurrentPoseEntryIdx];
+	return PoseEntry.DbPoseIdx;
+}
+
 }}
 #endif // UE_POSE_SEARCH_TRACE_ENABLED
\ No newline at end of file
diff --git a/Engine/Plugins/Experimental/Animation/PoseSearch/Source/Runtime/Private/Trace/PoseSearchTraceLogger.h b/Engine/Plugins/Experimental/Animation/PoseSearch/Source/Runtime/Private/Trace/PoseSearchTraceLogger.h
--- a/Engine/Plugins/Experimental/Animation/PoseSearch/Source/Runtime/Private/Trace/PoseSearchTraceLogger.h
+++ b/Engine/Plugins/Experimental/Animation/PoseSearch/Source/Runtime/Private/Trace/PoseSearchTraceLogger.h
@@ -39,18 +39,39 @@ POSESEARCH_API FArchive& operator<<(FArchive& Ar, FTraceMessage& State);
 
 struct POSESEARCH_API FTraceMotionMatchingStatePoseEntry
 {
-	int32 DbPoseIdx = 0;
+	enum class EFlags : uint8
+	{
+		None			= 0,
+		ContinuingPose	= 1 << 0,
+		CurrentPose		= 1 << 1,
+	};
+
+	int32 DbPoseIdx = INDEX_NONE;
 	float Cost = 0.f;
+	EFlags Flags = EFlags::None;
+
+	bool operator==(const FTraceMotionMatchingStatePoseEntry& Other) const { return DbPoseIdx == Other.DbPoseIdx; }
 };
+ENUM_CLASS_FLAGS(FTraceMotionMatchingStatePoseEntry::EFlags);
 POSESEARCH_API FArchive& operator<<(FArchive& Ar, FTraceMotionMatchingStatePoseEntry& Entry);
 
 struct POSESEARCH_API FTraceMotionMatchingStateDatabaseEntry
 {
+	enum class EFlags : uint8
+	{
+		None			= 0,
+		CurrentDatabase	= 1 << 0,
+	};
+
 	// @todo: can we use UPoseSearchDatabase* instead of DatabaseId?
 	//UPoseSearchDatabase const* Database = nullptr;
 	uint64 DatabaseId = 0;
+	EFlags Flags = EFlags::None;
 	TArray<FTraceMotionMatchingStatePoseEntry> PoseEntries;
+
+	bool operator==(const FTraceMotionMatchingStateDatabaseEntry& Other) const { return DatabaseId == Other.DatabaseId; }
 };
+ENUM_CLASS_FLAGS(FTraceMotionMatchingStateDatabaseEntry::EFlags);
 POSESEARCH_API FArchive& operator<<(FArchive& Ar, FTraceMotionMatchingStateDatabaseEntry& Entry);
 
 /**
@@ -66,6 +87,9 @@ struct POSESEARCH_API FTraceMotionMatchingState
 		/** Whether the last animation was a forced follow-up animation due to expended animation runway */
 		FollowupAnimation = 1u << 0
 	};
+
+	/** ObjectId of active searchable asset */
+	uint64 SearchableAssetId = 0;
 	
 	/** Amount of time since the last pose switch */
 	float ElapsedPoseJumpTime = 0.0f;
@@ -83,15 +107,6 @@ struct POSESEARCH_API FTraceMotionMatchingState
 	/** If true, groups are being filtered by DatabaseGroupQuery */
 	TArray<bool> DatabaseBlendSpaceFilter;
 
-	/** Index of the pose in our database */
-	int32 DbPoseIdx = 0;
-
-	/** Object Id of the database asset */
-	uint64 DatabaseId = 0;
-
-	/** Index of the continuing pose in our database */
-	int32 ContinuingPoseIdx = 0;
-
 	float AssetPlayerTime = 0.0f;
 	float DeltaTime = 0.0f;
 	float SimLinearVelocity = 0.0f;
@@ -101,41 +116,36 @@ struct POSESEARCH_API FTraceMotionMatchingState
 	
 	TArray<FTraceMotionMatchingStateDatabaseEntry> DatabaseEntries;
 
+	/** Index of the current database in DatabaseEntries */
+	int32 CurrentDbEntryIdx = INDEX_NONE;
+
+	/** Index of the current pose in DatabaseEntries[CurrentDbEntryIdx].PoseEntries */
+	int32 CurrentPoseEntryIdx = INDEX_NONE;
+
 	/** Output the current state info to the logger */
 	void Output(const FAnimationBaseContext& InContext);
 
-	static const UPoseSearchDatabase* GetDatabaseFromId(uint64 DatabaseId)
-	{
-		if (DatabaseId == 0)
-		{
-			return nullptr;
-		}
+	const UPoseSearchDatabase* GetCurrentDatabase() const;
+	int32 GetCurrentDatabasePoseIndex() const;
 
-		UObject* DatabaseObject = FObjectTrace::GetObjectFromId(DatabaseId);
-		// @TODO: Load the object if unloaded
-		if (DatabaseObject == nullptr)
+	template<typename T>
+	static const T* GetObjectFromId(uint64 ObjectId)
+	{
+		if (ObjectId)
 		{
-			return nullptr;
+			UObject* Object = FObjectTrace::GetObjectFromId(ObjectId);
+			if (Object)
+			{
+				return CastChecked<T>(Object);
+			}
 		}
-		check(DatabaseObject->IsA<UPoseSearchDatabase>());
 
-		const UPoseSearchDatabase* Database = Cast<UPoseSearchDatabase>(DatabaseObject);
-		return Database;
-	}
-
-	static uint64 GetIdFromDatabase(const UPoseSearchDatabase* Database)
-	{
-		return FObjectTrace::GetObjectId(Database);
-	}
-
-	void SetDatabase(const UPoseSearchDatabase* Database)
-	{
-		DatabaseId = GetIdFromDatabase(Database);
+		return nullptr;
 	}
 
-	const UPoseSearchDatabase* GetPoseSearchDatabase() const
+	static uint64 GetIdFromObject(const UObject* Object)
 	{
-		return GetDatabaseFromId(DatabaseId);
+		return FObjectTrace::GetObjectId(Object);
 	}
 	
 	static const FName Name;
diff --git a/Engine/Plugins/Experimental/Animation/PoseSearch/Source/Runtime/Public/PoseSearch/PoseSearch.h b/Engine/Plugins/Experimental/Animation/PoseSearch/Source/Runtime/Public/PoseSearch/PoseSearch.h
--- a/Engine/Plugins/Experimental/Animation/PoseSearch/Source/Runtime/Public/PoseSearch/PoseSearch.h
+++ b/Engine/Plugins/Experimental/Animation/PoseSearch/Source/Runtime/Public/PoseSearch/PoseSearch.h
@@ -965,11 +965,11 @@ namespace UE::PoseSearch
 
 enum class EDebugDrawFlags : uint32
 {
-	None			    = 0,
+	None = 0,
 
 	// Draw the entire search index as a point cloud
-	DrawSearchIndex     = 1 << 0,
-
+	DrawSearchIndex = 1 << 0,
+	
 	// Draw using Query colors form the schema / config
 	DrawQuery = 1 << 1,
 
@@ -1453,6 +1453,7 @@ public:
 		const UPoseSearchDatabase* Database = nullptr;
 
 		bool operator<(const FPoseCandidate& Other) const { return Cost > Other.Cost; }
+		bool operator==(const FSearchResult& SearchResult) const { return (PoseIdx == SearchResult.PoseIdx) && (Database == SearchResult.Database.Get()); }
 	};
 
 	struct FBestPoseCandidates : private TArray<FPoseCandidate>
