motion matching - debug lines improvements

diff --git a/Engine/Plugins/Experimental/Animation/PoseSearch/Source/Runtime/Private/AnimNode_MotionMatching.cpp b/Engine/Plugins/Experimental/Animation/PoseSearch/Source/Runtime/Private/AnimNode_MotionMatching.cpp
--- a/Engine/Plugins/Experimental/Animation/PoseSearch/Source/Runtime/Private/AnimNode_MotionMatching.cpp
+++ b/Engine/Plugins/Experimental/Animation/PoseSearch/Source/Runtime/Private/AnimNode_MotionMatching.cpp
@@ -143,9 +143,10 @@ void FAnimNode_MotionMatching::PreUpdate(const UAnimInstance* InAnimInstance)
 		USkeletalMeshComponent* SkeletalMeshComponent = InAnimInstance->GetSkelMeshComponent();
 		check(SkeletalMeshComponent);
 
+		const UE::PoseSearch::FSearchResult& CurResult = MotionMatchingState.CurrentSearchResult;
 		UE::PoseSearch::FDebugDrawParams DrawParams;
 		DrawParams.RootTransform = SkeletalMeshComponent->GetComponentTransform();
-		DrawParams.Database = MotionMatchingState.CurrentSearchResult.Database.Get();
+		DrawParams.Database = CurResult.Database.Get();
 		DrawParams.World = SkeletalMeshComponent->GetWorld();
 		DrawParams.DefaultLifeTime = 0.0f;
 
@@ -153,34 +154,45 @@ void FAnimNode_MotionMatching::PreUpdate(const UAnimInstance* InAnimInstance)
 		{
 			if (bDebugDrawMatch)
 			{
-				DrawFeatureVector(DrawParams, MotionMatchingState.CurrentSearchResult.PoseIdx);
+				DrawFeatureVector(DrawParams, CurResult.PoseIdx);
 			}
 
 			if (bDebugDrawQuery)
 			{
 				EnumAddFlags(DrawParams.Flags, EDebugDrawFlags::DrawQuery);
-				DrawFeatureVector(DrawParams, MotionMatchingState.CurrentSearchResult.ComposedQuery.GetValues());
+				DrawFeatureVector(DrawParams, CurResult.ComposedQuery.GetValues());
 			}
 
 			if (DrawParams.Database->PoseSearchMode == EPoseSearchMode::PCAKDTree_Compare)
 			{
-				const float TotalCostKDTree = MotionMatchingState.CurrentSearchResult.IsValid() ? MotionMatchingState.CurrentSearchResult.PoseCost.GetTotalCost() : 0.f;
-				const float TotalCostBruteForce = MotionMatchingState.CurrentSearchResult.IsValid() ? MotionMatchingState.CurrentSearchResult.BruteForcePoseCost.GetTotalCost() : 0.f;
+				FDebugFloatHistory& C = MotionMatchingState.SearchCostHistoryContinuing;
+				FDebugFloatHistory& B = MotionMatchingState.SearchCostHistoryBruteForce;
+				FDebugFloatHistory& K = MotionMatchingState.SearchCostHistoryKDTree;
 
-				MotionMatchingState.SearchCostHistoryKDTree.AddSample(TotalCostKDTree);
-				MotionMatchingState.SearchCostHistoryBruteForce.AddSample(TotalCostBruteForce);
+				C.AddSample(CurResult.ContinuingPoseCost.IsValid() ? CurResult.ContinuingPoseCost.GetTotalCost() : C.MaxValue);
+				B.AddSample(CurResult.BruteForcePoseCost.IsValid() ? CurResult.BruteForcePoseCost.GetTotalCost() : B.MaxValue);
+				K.AddSample(CurResult.PoseCost.IsValid() ? CurResult.PoseCost.GetTotalCost() : K.MaxValue);
 
 				// making SearchCostHistoryKDTree and SearchCostHistoryBruteForce min max consistent
-				MotionMatchingState.SearchCostHistoryKDTree.MinValue = FMath::Min(MotionMatchingState.SearchCostHistoryKDTree.MinValue, MotionMatchingState.SearchCostHistoryBruteForce.MinValue);
-				MotionMatchingState.SearchCostHistoryKDTree.MaxValue = FMath::Max(MotionMatchingState.SearchCostHistoryKDTree.MaxValue, MotionMatchingState.SearchCostHistoryBruteForce.MaxValue);
+				const float MinValue = FMath::Min(C.MinValue, FMath::Min(B.MinValue, K.MinValue));
+				const float MaxValue = FMath::Max(C.MaxValue, FMath::Max(B.MaxValue, K.MaxValue));
+				
+				C.MinValue = MinValue;
+				C.MaxValue = MaxValue;
 
-				MotionMatchingState.SearchCostHistoryBruteForce.MinValue = MotionMatchingState.SearchCostHistoryKDTree.MinValue;
-				MotionMatchingState.SearchCostHistoryBruteForce.MaxValue = MotionMatchingState.SearchCostHistoryKDTree.MaxValue;
+				B.MinValue = MinValue;
+				B.MaxValue = MaxValue;
 
+				K.MinValue = MinValue;
+				K.MaxValue = MaxValue;
+
+				const FVector2D DrawSize(150.f, 100.f);
 				const FTransform OffsetTransform(FRotator(0.f, 0.f, 0.f), FVector(-50.f, -75.f, 100.f));
-				FVector2D DrawSize(150.f, 50.f);
-				DrawDebugFloatHistory(*DrawParams.World, MotionMatchingState.SearchCostHistoryKDTree, OffsetTransform * DrawParams.RootTransform, DrawSize, FColor(255,192,203,150)); // pink
-				DrawDebugFloatHistory(*DrawParams.World, MotionMatchingState.SearchCostHistoryBruteForce, OffsetTransform * DrawParams.RootTransform, DrawSize, FColor(0,0,255,150)); // blue
+				const FTransform DrawTransform = OffsetTransform * DrawParams.RootTransform;
+
+				DrawDebugFloatHistory(*DrawParams.World, K, OffsetTransform * DrawParams.RootTransform, DrawSize, FColor(255, 192, 203, 160)); // pink
+				DrawDebugFloatHistory(*DrawParams.World, B, OffsetTransform * DrawParams.RootTransform, DrawSize, FColor(0, 0, 255, 160)); // blue
+				DrawDebugFloatHistory(*DrawParams.World, C, OffsetTransform * DrawParams.RootTransform, DrawSize, FColor(160, 160, 160, 160)); // gray
 			}
 		}
 	}
diff --git a/Engine/Plugins/Experimental/Animation/PoseSearch/Source/Runtime/Private/PoseSearch.cpp b/Engine/Plugins/Experimental/Animation/PoseSearch/Source/Runtime/Private/PoseSearch.cpp
--- a/Engine/Plugins/Experimental/Animation/PoseSearch/Source/Runtime/Private/PoseSearch.cpp
+++ b/Engine/Plugins/Experimental/Animation/PoseSearch/Source/Runtime/Private/PoseSearch.cpp
@@ -1252,6 +1252,16 @@ void UPoseSearchDatabase::GenerateDDCKey(FBlake3& InOutKeyHasher) const
 
 #endif // WITH_EDITOR
 
+bool UPoseSearchDatabase::GetSkipSearchIfPossible() const
+{
+	if (PoseSearchMode == EPoseSearchMode::PCAKDTree_Validate || PoseSearchMode == EPoseSearchMode::PCAKDTree_Compare)
+	{
+		return false;
+	}
+
+	return bSkipSearchIfPossible;
+}
+
 bool UPoseSearchDatabase::IsValidForIndexing() const
 {
 	bool bValid = Schema && Schema->IsValid() && !Sequences.IsEmpty();
@@ -2047,6 +2057,10 @@ UE::PoseSearch::FSearchResult UPoseSearchDatabase::Search(UE::PoseSearch::FSearc
 
 #if WITH_EDITORONLY_DATA
 		Result.BruteForcePoseCost = BruteForcePoseCost;
+		if (PoseSearchMode == EPoseSearchMode::PCAKDTree_Compare)
+		{
+			check(Result.BruteForcePoseCost.GetTotalCost() <= Result.PoseCost.GetTotalCost());
+		}
 #endif
 	}
 	
@@ -2137,7 +2151,7 @@ UE::PoseSearch::FSearchResult UPoseSearchDatabase::SearchPCAKDTree(UE::PoseSearc
 			SearchIndex->FindAssetForPose(SearchContext.CurrentResult.PoseIdx)->SourceGroupIdx,
 			NormalizedQueryValues);
 
-		if (bSkipSearchIfPossible)
+		if (GetSkipSearchIfPossible())
 		{
 			SearchContext.UpdateCurrentBestCost(Result.ContinuingPoseCost);
 		}
@@ -2198,7 +2212,7 @@ UE::PoseSearch::FSearchResult UPoseSearchDatabase::SearchPCAKDTree(UE::PoseSearc
 			}
 		}
 
-		if (bSkipSearchIfPossible && BestPoseCost.IsValid())
+		if (GetSkipSearchIfPossible() && BestPoseCost.IsValid())
 		{
 			SearchContext.UpdateCurrentBestCost(BestPoseCost);
 		}
@@ -2254,7 +2268,7 @@ UE::PoseSearch::FSearchResult UPoseSearchDatabase::SearchBruteForce(UE::PoseSear
 				SearchIndex->FindAssetForPose(SearchContext.CurrentResult.PoseIdx)->SourceGroupIdx,
 				NormalizedQueryValues);
 
-			if (bSkipSearchIfPossible)
+			if (GetSkipSearchIfPossible())
 			{
 				SearchContext.UpdateCurrentBestCost(Result.ContinuingPoseCost);
 			}
@@ -2318,7 +2332,7 @@ UE::PoseSearch::FSearchResult UPoseSearchDatabase::SearchBruteForce(UE::PoseSear
 			}
 		}
 
-		if (bSkipSearchIfPossible && BestPoseCost.IsValid())
+		if (GetSkipSearchIfPossible() && BestPoseCost.IsValid())
 		{
 			SearchContext.UpdateCurrentBestCost(BestPoseCost);
 		}
@@ -2392,7 +2406,7 @@ UE::PoseSearch::FSearchResult UPoseSearchDatabaseSet::Search(UE::PoseSearch::FSe
 				PoseSearchIndexAsset->SourceGroupIdx,
 				NormalizedQueryValues);
 
-		if (Database->bSkipSearchIfPossible)
+		if (Database->GetSkipSearchIfPossible())
 		{
 			SearchContext.UpdateCurrentBestCost(Result.ContinuingPoseCost);
 		}
diff --git a/Engine/Plugins/Experimental/Animation/PoseSearch/Source/Runtime/Private/PoseSearchLibrary.cpp b/Engine/Plugins/Experimental/Animation/PoseSearch/Source/Runtime/Private/PoseSearchLibrary.cpp
--- a/Engine/Plugins/Experimental/Animation/PoseSearch/Source/Runtime/Private/PoseSearchLibrary.cpp
+++ b/Engine/Plugins/Experimental/Animation/PoseSearch/Source/Runtime/Private/PoseSearchLibrary.cpp
@@ -451,20 +451,24 @@ void UpdateMotionMatchingState(
 
 		// Search the database for the nearest match to the updated query vector
 		FSearchResult SearchResult = Searchable->Search(SearchContext);
-		const float ContinuingPoseCost = bCanAdvance && SearchResult.ContinuingPoseCost.IsValid() ? SearchResult.ContinuingPoseCost.GetTotalCost() : MAX_flt;
-		const float PoseCost = SearchResult.PoseCost.IsValid() ? SearchResult.PoseCost.GetTotalCost() : MAX_flt;
+
+		// making sure we haven't calculated ContinuingPoseCost if we !bCanAdvance 
+		check(bCanAdvance || !SearchResult.ContinuingPoseCost.IsValid());
 		
-		if (PoseCost < ContinuingPoseCost)
+		if (SearchResult.PoseCost.GetTotalCost() < SearchResult.ContinuingPoseCost.GetTotalCost())
 		{
 			InOutMotionMatchingState.JumpToPose(Context, Settings, SearchResult);
 		}
-
+		else
+		{
+			// copying few properties of SearchResult into CurrentSearchResult to facilitate debug drawing
 #if WITH_EDITOR
-		InOutMotionMatchingState.CurrentSearchResult.BruteForcePoseCost = SearchResult.BruteForcePoseCost;
+			InOutMotionMatchingState.CurrentSearchResult.BruteForcePoseCost = SearchResult.BruteForcePoseCost;
 #endif
-
-		// todo: cache the queries into TraceState, one per TraceState.DatabaseEntries or
-		InOutMotionMatchingState.CurrentSearchResult.ComposedQuery = SearchResult.ComposedQuery;
+			InOutMotionMatchingState.CurrentSearchResult.PoseCost = SearchResult.PoseCost;
+			InOutMotionMatchingState.CurrentSearchResult.ContinuingPoseCost = SearchResult.ContinuingPoseCost;
+			InOutMotionMatchingState.CurrentSearchResult.ComposedQuery = SearchResult.ComposedQuery;
+		}
 	}
 
 	// If we didn't search or it didn't find a pose to jump to, and we can 
diff --git a/Engine/Plugins/Experimental/Animation/PoseSearch/Source/Runtime/Public/PoseSearch/PoseSearch.h b/Engine/Plugins/Experimental/Animation/PoseSearch/Source/Runtime/Public/PoseSearch/PoseSearch.h
--- a/Engine/Plugins/Experimental/Animation/PoseSearch/Source/Runtime/Public/PoseSearch/PoseSearch.h
+++ b/Engine/Plugins/Experimental/Animation/PoseSearch/Source/Runtime/Public/PoseSearch/PoseSearch.h
@@ -1151,7 +1151,8 @@ public:
 
 	FPoseSearchIndex* GetSearchIndex();
 	const FPoseSearchIndex* GetSearchIndex() const;
-
+	
+	bool GetSkipSearchIfPossible() const;
 	bool IsValidForIndexing() const;
 	bool IsValidForSearch() const;
 
diff --git a/Engine/Plugins/Experimental/Animation/PoseSearch/Source/Runtime/Public/PoseSearch/PoseSearchLibrary.h b/Engine/Plugins/Experimental/Animation/PoseSearch/Source/Runtime/Public/PoseSearch/PoseSearchLibrary.h
--- a/Engine/Plugins/Experimental/Animation/PoseSearch/Source/Runtime/Public/PoseSearch/PoseSearchLibrary.h
+++ b/Engine/Plugins/Experimental/Animation/PoseSearch/Source/Runtime/Public/PoseSearch/PoseSearchLibrary.h
@@ -93,8 +93,10 @@ struct POSESEARCH_API FMotionMatchingState
 	FTransform RootMotionTransformDelta = FTransform::Identity;
 
 #if WITH_EDITORONLY_DATA
-	FDebugFloatHistory SearchCostHistoryBruteForce;
-	FDebugFloatHistory SearchCostHistoryKDTree;
+	enum { SearchCostHistoryNumSamples = 200 };
+	FDebugFloatHistory SearchCostHistoryContinuing = FDebugFloatHistory(SearchCostHistoryNumSamples, 0, 0, true);
+	FDebugFloatHistory SearchCostHistoryBruteForce = FDebugFloatHistory(SearchCostHistoryNumSamples, 0, 0, true);
+	FDebugFloatHistory SearchCostHistoryKDTree = FDebugFloatHistory(SearchCostHistoryNumSamples, 0, 0, true);
 #endif
 
 	UE::PoseSearch::FPoseIndicesHistory PoseIndicesHistory;
