PoseSearch:
- Fixed extremely high pose scores caused by numeric instability when some features have near zero deviation. This problem was being triggered by databases with forward facing vectors enabled in their schema, but no actual rotation in any of the animation sequences. A hard coded check was not sufficient to catch the problem, so it is now exposed as a "Minimum Mean Deviation" setting on feature channels with a higher default value (1.0 vs 0.0001).
- Fixed brute force search tracing pose candidates in PCA kd tree compare mode
- Also removed sphering normalization since it is unfinished and unused

diff --git a/Engine/Plugins/Experimental/Animation/PoseSearch/Source/Runtime/Private/PoseSearch.cpp b/Engine/Plugins/Experimental/Animation/PoseSearch/Source/Runtime/Private/PoseSearch.cpp
--- a/Engine/Plugins/Experimental/Animation/PoseSearch/Source/Runtime/Private/PoseSearch.cpp
+++ b/Engine/Plugins/Experimental/Animation/PoseSearch/Source/Runtime/Private/PoseSearch.cpp
@@ -280,7 +280,7 @@ void UPoseSearchFeatureChannel::ComputeMeanDeviations(const Eigen::MatrixXd& Cen
 	using namespace UE::PoseSearch;
 
 	int32 DataOffset = ChannelDataOffset;
-	FFeatureVectorHelper::ComputeMeanDeviations(CenteredPoseMatrix, MeanDeviations, DataOffset, ChannelCardinality);
+	FFeatureVectorHelper::ComputeMeanDeviations(MinimumMeanDeviation, CenteredPoseMatrix, MeanDeviations, DataOffset, ChannelCardinality);
 
 	check(DataOffset == ChannelDataOffset + ChannelCardinality);
 }
@@ -2326,9 +2326,12 @@ UE::PoseSearch::FSearchResult UPoseSearchDatabase::SearchBruteForce(UE::PoseSear
 					BestPoseIdx = PoseIdx;
 				}
 
-				#if UE_POSE_SEARCH_TRACE_ENABLED
-				SearchContext.BestCandidates.Add(PoseCost.GetTotalCost(), PoseIdx, this);
-				#endif
+#if UE_POSE_SEARCH_TRACE_ENABLED
+				if (PoseSearchMode == EPoseSearchMode::BruteForce)
+				{
+					SearchContext.BestCandidates.Add(PoseCost.GetTotalCost(), PoseIdx, this);
+				}
+#endif
 			}
 		}
 
@@ -3007,7 +3010,7 @@ FVector2D FFeatureVectorHelper::DecodeVector2DInternal(TArrayView<const float> V
 	return FVector2D(Values[DataOffset + 0], Values[DataOffset + 1]);
 }
 
-void FFeatureVectorHelper::ComputeMeanDeviations(const Eigen::MatrixXd& CenteredPoseMatrix, Eigen::VectorXd& MeanDeviations, int32& DataOffset, int32 Cardinality)
+void FFeatureVectorHelper::ComputeMeanDeviations(float MinMeanDeviation, const Eigen::MatrixXd& CenteredPoseMatrix, Eigen::VectorXd& MeanDeviations, int32& DataOffset, int32 Cardinality)
 {
 	const int32 NumPoses = MeanDeviations.size();
 
@@ -3017,7 +3020,7 @@ void FFeatureVectorHelper::ComputeMeanDeviations(const Eigen::MatrixXd& Centered
 
 	// Fill the feature's corresponding scaling axes with the average distance
 	// Avoid scaling by zero by leaving near-zero deviations as 1.0
-	MeanDeviations.segment(DataOffset, Cardinality).setConstant(FeatureMeanDeviation > KINDA_SMALL_NUMBER ? FeatureMeanDeviation : 1.f);
+	MeanDeviations.segment(DataOffset, Cardinality).setConstant(FeatureMeanDeviation > MinMeanDeviation ? FeatureMeanDeviation : 1.f);
 
 	DataOffset += Cardinality;
 }
@@ -4805,212 +4808,6 @@ static void PreprocessSearchIndexNormalize(FPoseSearchIndex* SearchIndex)
 #endif // UE_POSE_SEARCH_EIGEN_DEBUG
 }
 
-static void PreprocessSearchIndexSphere(FPoseSearchIndex* SearchIndex)
-{
-	// This function performs correlation based zero-phase component analysis sphering (ZCA-cor sphering)
-	// The pose matrix is transformed in place and the transformation matrix, its inverse,
-	// and data mean vector are computed and stored along with it.
-	//
-	// N:	number of dimensions for input column vectors
-	// P:	number of input column vectors
-	// X:	NxP input matrix
-	// x_p:	pth column vector of input matrix
-	// u:   mean column vector of X
-	//
-	// Eigendecomposition of correlation matrix of X:
-	// cor(X) = (1/P) * X * X^T = V * D * V^T
-	//
-	// V:	eigenvectors of cor(X), stacked as columns in an orthogonal NxN matrix
-	// D:	eigenvalues of cor(X), as diagonal NxN matrix
-	// d_n:	nth eigenvalue
-	// s_n: nth standard deviation
-	// s_n^2 = d_n, the variance along the nth eigenvector
-	// s_n   = d_n^(1/2)
-	//
-	// ZCA sphering algorithm:
-	//
-	// 1) mean-center X
-	//    x_p := x_p - u
-	// 2) align largest orthogonal directions of variance in X to coordinate axes (PCA rotate)
-	//    x_p := V^T * x_p
-	// 3) rescale X by inverse standard deviation
-	//    x_p := x_p * d_n^(-1/2)
-	// 4) return now rescaled X back to original rotation (inverse PCA rotate)
-	//    x_p := V * x_p
-	// 
-	// Let D^(-1/2) be the inverse square root of D where the nth diagonal element is d_n^(-1/2)
-	// then steps 2-4 can be expressed as a series of matrix multiplications:
-	// Z = V * D^(-1/2) * V^T
-	// X := Z * X
-	//
-	// By persisting the mean vector u and linear transform Z, we can bring an input vector q
-	// into the same space as the sphered data matrix X:
-	// q := Z * (q - u)
-	//
-	// This operation is invertible, a sphere standardized data vector x can be unscaled via:
-	// Z^(-1) = V * D^(1/2) * V^T
-	// x := (Z^(-1) * x) + u
-	//
-	// The sphering process allows nearest neighbor queries to use the Mahalanobis metric
-	// which is unitless, scale-invariant, and accounts for feature correlation.
-	// The Mahalanobis distance between two random vectors x and y in data matrix X is:
-	// d(x,y) = ((x-y)^T * cov(X)^(-1) * (x-y))^(1/2)
-	//
-	// Because sphering transforms X into a new matrix with identity covariance, the Mahalanobis
-	// distance equation above reduces to Euclidean distance since cov(X)^(-1) = I:
-	// d(x,y) = ((x-y)^T * (x-y))^(1/2)
-	// 
-	// References:
-	// Watt, Jeremy, et al. Machine Learning Refined: Foundations, Algorithms, and Applications.
-	// 2nd ed., Cambridge University Press, 2020.
-	// 
-	// Kessy, Agnan, Alex Lewin, and Korbinian Strimmer. "Optimal whitening and decorrelation."
-	// The American Statistician 72.4 (2018): 309-314.
-	// 
-	// https://en.wikipedia.org/wiki/Whitening_transformation
-	// 
-	// https://en.wikipedia.org/wiki/Mahalanobis_distance
-	//
-	// Note this sphering preprocessor needs more work and isn't yet exposed in the editor as an option.
-	// Todo:
-	// - Figure out apparent flipping behavior
-	// - Try singular value decomposition in place of eigendecomposition
-	// - Remove zero variance feature axes from data and search queries
-	// - Support weighted Mahalanobis metric. User supplied weights need to be transformed to data's new basis.
-
-#if UE_POSE_SEARCH_EIGEN_DEBUG
-	double StartTime = FPlatformTime::Seconds();
-#endif
-
-	using namespace Eigen;
-
-	check(SearchIndex->IsValid() && !SearchIndex->IsEmpty());
-
-	FPoseSearchIndexPreprocessInfo& Info = SearchIndex->PreprocessInfo;
-	Info.Reset();
-
-	const int32 NumPoses = SearchIndex->NumPoses;
-	const int32 NumDimensions = SearchIndex->Schema->SchemaCardinality;
-
-	// Map input buffer
-	auto PoseMatrixSourceMap = Map<Matrix<float, Dynamic, Dynamic, RowMajor>>(
-		SearchIndex->Values.GetData(),
-		NumPoses,		// rows
-		NumDimensions	// cols
-	);
-
-	// Copy row major float matrix to column major double matrix
-	MatrixXd PoseMatrix = PoseMatrixSourceMap.transpose().cast<double>();
-	checkSlow(PoseMatrix.rows() == NumDimensions);
-	checkSlow(PoseMatrix.cols() == NumPoses);
-
-#if UE_POSE_SEARCH_EIGEN_DEBUG
-	MatrixXd PoseMatrixOriginal = PoseMatrix;
-#endif
-
-	// Mean center
-	VectorXd SampleMean = PoseMatrix.rowwise().mean();
-	PoseMatrix = PoseMatrix.colwise() - SampleMean;
-
-	// Compute per-feature average distances
-	VectorXd MeanDeviations = ComputeChannelsMeanDeviations(PoseMatrix, SearchIndex->Schema);
-
-	// Rescale data by transforming it with the scaling matrix
-	// Now each feature has an average Euclidean length = 1.
-	MatrixXd PoseMatrixNormalized = MeanDeviations.cwiseInverse().asDiagonal() * PoseMatrix;
-
-	// Compute sample covariance
-	MatrixXd Covariance = ((1.0 / NumPoses) * (PoseMatrixNormalized * PoseMatrixNormalized.transpose())) + 1e-7 * MatrixXd::Identity(NumDimensions, NumDimensions);
-
-	VectorXd StdDev = Covariance.diagonal().cwiseSqrt();
-	VectorXd InvStdDev = StdDev.cwiseInverse();
-	MatrixXd Correlation = InvStdDev.asDiagonal() * Covariance * InvStdDev.asDiagonal();
-
-	// Compute eigenvalues and eigenvectors of correlation matrix
-	SelfAdjointEigenSolver<MatrixXd> EigenDecomposition(Correlation, ComputeEigenvectors);
-
-	VectorXd EigenValues = EigenDecomposition.eigenvalues();
-	MatrixXd EigenVectors = EigenDecomposition.eigenvectors();
-
-	// Sort eigenpairs by descending eigenvalue
-	{
-		const Eigen::Index n = EigenValues.size();
-		for (Eigen::Index i = 0; i < n-1; ++i)
-		{
-			Index k;
-			EigenValues.segment(i,n-i).cwiseAbs().maxCoeff(&k);
-			if (k > 0)
-			{
-				std::swap(EigenValues[i], EigenValues[k+i]);
-				EigenVectors.col(i).swap(EigenVectors.col(k+i));
-			}
-		}
-	}
-
-	// Regularize eigenvalues
-	EigenValues = EigenValues.array() + 1e-7;
-
-	// Compute ZCA-cor and ZCA-cor^(-1)
-	MatrixXd ZCA = EigenVectors * EigenValues.cwiseInverse().cwiseSqrt().asDiagonal() * EigenVectors.transpose() * MeanDeviations.cwiseInverse().asDiagonal();
-	MatrixXd ZCAInverse = MeanDeviations.asDiagonal() * EigenVectors * EigenValues.cwiseSqrt().asDiagonal() * EigenVectors.transpose();
-
-	// Apply sphering transform to the data matrix
-	PoseMatrix = ZCA * PoseMatrix;
-	checkSlow(PoseMatrix.rows() == NumDimensions);
-	checkSlow(PoseMatrix.cols() == NumPoses);
-
-	// Write data back to source buffer, converting from column data back to row data
-	PoseMatrixSourceMap = PoseMatrix.transpose().cast<float>();
-
-	// Output preprocessing info
-	Info.NumDimensions = NumDimensions;
-	Info.TransformationMatrix.SetNumZeroed(ZCA.size());
-	Info.InverseTransformationMatrix.SetNumZeroed(ZCAInverse.size());
-	Info.SampleMean.SetNumZeroed(SampleMean.size());
-
-	auto TransformMap = Map<Matrix<float, Dynamic, Dynamic, ColMajor>>(
-		Info.TransformationMatrix.GetData(),
-		ZCA.rows(), ZCA.cols()
-	);
-
-	auto InverseTransformMap = Map<Matrix<float, Dynamic, Dynamic, ColMajor>>(
-		Info.InverseTransformationMatrix.GetData(),
-		ZCAInverse.rows(), ZCAInverse.cols()
-	);
-
-	auto SampleMeanMap = Map<VectorXf>(Info.SampleMean.GetData(), SampleMean.size());
-
-	// Output sphering matrix, inverse sphering matrix, and mean vector
-	TransformMap = ZCA.cast<float>();
-	InverseTransformMap = ZCAInverse.cast<float>();
-	SampleMeanMap = SampleMean.cast<float>();
-
-#if UE_POSE_SEARCH_EIGEN_DEBUG
-	double ElapsedTime = FPlatformTime::Seconds() - StartTime;
-
-	FString EigenValuesStr = EigenMatrixToString(EigenValues);
-	FString EigenVectorsStr = EigenMatrixToString(EigenVectors);
-
-	FString CovarianceStr = EigenMatrixToString(Covariance);
-	FString CorrelationStr = EigenMatrixToString(Correlation);
-
-	FString ZCAStr = EigenMatrixToString(ZCA);
-	FString ZCAInverseStr = EigenMatrixToString(ZCAInverse);
-
-	FString PoseMatrixSphereStr = EigenMatrixToString(PoseMatrix);
-	MatrixXd PoseMatrixUnsphered = ZCAInverse * PoseMatrix;
-	PoseMatrixUnsphered = PoseMatrixUnsphered.colwise() + SampleMean;
-	FString PoseMatrixUnspheredStr = EigenMatrixToString(PoseMatrixUnsphered);
-	FString PoseMatrixOriginalStr = EigenMatrixToString(PoseMatrixOriginal);
-
-	FString OutputPoseMatrixStr = EigenMatrixToString(PoseMatrixSourceMap);
-
-	FString TransformStr = EigenMatrixToString(TransformMap);
-	FString InverseTransformStr = EigenMatrixToString(InverseTransformMap);
-	FString SampleMeanStr = EigenMatrixToString(SampleMeanMap);
-#endif // UE_POSE_SEARCH_EIGEN_DEBUG
-}
-
 static void PreprocessSearchIndex(FPoseSearchIndex* SearchIndex)
 {
 	switch (SearchIndex->Schema->EffectiveDataPreprocessor)
@@ -5019,10 +4816,6 @@ static void PreprocessSearchIndex(FPoseSearchIndex* SearchIndex)
 			PreprocessSearchIndexNormalize(SearchIndex);
 		break;
 
-		case EPoseSearchDataPreprocessor::Sphere:
-			PreprocessSearchIndexSphere(SearchIndex);
-		break;
-
 		case EPoseSearchDataPreprocessor::None:
 			PreprocessSearchIndexNone(SearchIndex);
 		break;
diff --git a/Engine/Plugins/Experimental/Animation/PoseSearch/Source/Runtime/Private/PoseSearchFeatureChannels.cpp b/Engine/Plugins/Experimental/Animation/PoseSearch/Source/Runtime/Private/PoseSearchFeatureChannels.cpp
--- a/Engine/Plugins/Experimental/Animation/PoseSearch/Source/Runtime/Private/PoseSearchFeatureChannels.cpp
+++ b/Engine/Plugins/Experimental/Animation/PoseSearch/Source/Runtime/Private/PoseSearchFeatureChannels.cpp
@@ -755,26 +755,29 @@ void UPoseSearchFeatureChannel_Pose::ComputeMeanDeviations(const Eigen::MatrixXd
 		{
 			for (int32 SubsampleIdx = 0; SubsampleIdx != SampleTimes.Num(); ++SubsampleIdx)
 			{
-				FFeatureVectorHelper::ComputeMeanDeviations(CenteredPoseMatrix, MeanDeviations, DataOffset, FFeatureVectorHelper::EncodeVectorCardinality);
+				FFeatureVectorHelper::ComputeMeanDeviations(MinimumMeanDeviation, CenteredPoseMatrix, MeanDeviations, DataOffset, FFeatureVectorHelper::EncodeVectorCardinality);
 			}
 		}
 		if (EnumHasAnyFlags(SampledBone.Flags, EPoseSearchBoneFlags::Rotation))
 		{
 			for (int32 SubsampleIdx = 0; SubsampleIdx != SampleTimes.Num(); ++SubsampleIdx)
 			{
-				FFeatureVectorHelper::ComputeMeanDeviations(CenteredPoseMatrix, MeanDeviations, DataOffset, FFeatureVectorHelper::EncodeQuatCardinality);
+				FFeatureVectorHelper::ComputeMeanDeviations(MinimumMeanDeviation, CenteredPoseMatrix, MeanDeviations, DataOffset, FFeatureVectorHelper::EncodeQuatCardinality);
 			}
 		}
 		if (EnumHasAnyFlags(SampledBone.Flags, EPoseSearchBoneFlags::Velocity))
 		{
 			for (int32 SubsampleIdx = 0; SubsampleIdx != SampleTimes.Num(); ++SubsampleIdx)
 			{
-				FFeatureVectorHelper::ComputeMeanDeviations(CenteredPoseMatrix, MeanDeviations, DataOffset, FFeatureVectorHelper::EncodeVectorCardinality);
+				FFeatureVectorHelper::ComputeMeanDeviations(MinimumMeanDeviation, CenteredPoseMatrix, MeanDeviations, DataOffset, FFeatureVectorHelper::EncodeVectorCardinality);
 			}
 		}
 		if (EnumHasAnyFlags(SampledBone.Flags, EPoseSearchBoneFlags::Phase))
 		{
-			FFeatureVectorHelper::ComputeMeanDeviations(CenteredPoseMatrix, MeanDeviations, DataOffset, FFeatureVectorHelper::EncodeVector2DCardinality);
+			for (int32 SubsampleIdx = 0; SubsampleIdx != SampleTimes.Num(); ++SubsampleIdx)
+			{
+				FFeatureVectorHelper::ComputeMeanDeviations(MinimumMeanDeviation, CenteredPoseMatrix, MeanDeviations, DataOffset, FFeatureVectorHelper::EncodeVector2DCardinality);
+			}
 		}
 	}
 
@@ -1299,15 +1302,15 @@ void UPoseSearchFeatureChannel_Trajectory::ComputeMeanDeviations(const Eigen::Ma
 	{
 		if (EnumHasAnyFlags(Sample.Flags, EPoseSearchTrajectoryFlags::Position))
 		{
-			FFeatureVectorHelper::ComputeMeanDeviations(CenteredPoseMatrix, MeanDeviations, DataOffset, FFeatureVectorHelper::EncodeVectorCardinality);
+			FFeatureVectorHelper::ComputeMeanDeviations(MinimumMeanDeviation, CenteredPoseMatrix, MeanDeviations, DataOffset, FFeatureVectorHelper::EncodeVectorCardinality);
 		}
 		if (EnumHasAnyFlags(Sample.Flags, EPoseSearchTrajectoryFlags::Velocity))
 		{
-			FFeatureVectorHelper::ComputeMeanDeviations(CenteredPoseMatrix, MeanDeviations, DataOffset, FFeatureVectorHelper::EncodeVectorCardinality);
+			FFeatureVectorHelper::ComputeMeanDeviations(MinimumMeanDeviation, CenteredPoseMatrix, MeanDeviations, DataOffset, FFeatureVectorHelper::EncodeVectorCardinality);
 		}
 		if (EnumHasAnyFlags(Sample.Flags, EPoseSearchTrajectoryFlags::FacingDirection))
 		{
-			FFeatureVectorHelper::ComputeMeanDeviations(CenteredPoseMatrix, MeanDeviations, DataOffset, FFeatureVectorHelper::EncodeVectorCardinality);
+			FFeatureVectorHelper::ComputeMeanDeviations(MinimumMeanDeviation, CenteredPoseMatrix, MeanDeviations, DataOffset, FFeatureVectorHelper::EncodeVectorCardinality);
 		}
 	}
 
diff --git a/Engine/Plugins/Experimental/Animation/PoseSearch/Source/Runtime/Public/PoseSearch/PoseSearch.h b/Engine/Plugins/Experimental/Animation/PoseSearch/Source/Runtime/Public/PoseSearch/PoseSearch.h
--- a/Engine/Plugins/Experimental/Animation/PoseSearch/Source/Runtime/Public/PoseSearch/PoseSearch.h
+++ b/Engine/Plugins/Experimental/Animation/PoseSearch/Source/Runtime/Public/PoseSearch/PoseSearch.h
@@ -102,7 +102,6 @@ enum class EPoseSearchDataPreprocessor : int32
 	None,
 	Automatic,
 	Normalize,
-	Sphere UMETA(Hidden),
 
 	Num UMETA(Hidden),
 	Invalid = Num UMETA(Hidden)
@@ -369,11 +368,15 @@ public:
 	// Called during UPoseSearchSchema::Finalize to prepare the schema for this channel
 	virtual void InitializeSchema(UE::PoseSearch::FSchemaInitializer& Initializer);
 	
+	// Called at database build time to collect feature weights.
+	// Weights is sized to the cardinality of the schema and the feature channel should write
+	// its weights at the channel's data offset. Channels should provide a weight for each dimension.
 	virtual void FillWeights(TArray<float>& Weights) const PURE_VIRTUAL(UPoseSearchFeatureChannel::FillWeights, );
 
 	// Called at database build time to populate pose vectors with this channel's data
 	virtual void IndexAsset(UE::PoseSearch::IAssetIndexer& Indexer, UE::PoseSearch::FAssetIndexingOutput& IndexingOutput) const PURE_VIRTUAL(UPoseSearchFeatureChannel::IndexAsset, );
 
+	// Called at database build time to calculate normalization values
 	virtual void ComputeMeanDeviations(const Eigen::MatrixXd& CenteredPoseMatrix, Eigen::VectorXd& MeanDeviations) const;
 
 	// Hash channel properties to produce a key for database derived data
@@ -385,6 +388,10 @@ public:
 	// Draw this channel's data for the given pose vector
 	virtual void DebugDraw(const UE::PoseSearch::FDebugDrawParams& DrawParams, TArrayView<const float> PoseVector) const PURE_VIRTUAL(UPoseSearchFeatureChannel::DebugDraw, );
 
+	// Used during data normalization. If a feature has less than this amount of deviation from the mean across all poses in a database, then it will not be normalized.
+	UPROPERTY(EditAnywhere, Category = "Settings", meta = (ClampMin = "0.0001"))
+	float MinimumMeanDeviation = 1.0f;
+
 private:
 	// IBoneReferenceSkeletonProvider interface
 	// Note this function is exclusively for FBoneReference details customization
@@ -1313,7 +1320,7 @@ public:
 	static FVector2D DecodeVector2D(TArrayView<const float> Values, int32& DataOffset);
 
 	// populates MeanDeviations[DataOffset] ... MeanDeviations[DataOffset + Cardinality] with a single value the mean deviation calcualted from a cenetered matrix
-	static void ComputeMeanDeviations(const Eigen::MatrixXd& CenteredPoseMatrix, Eigen::VectorXd& MeanDeviations, int32& DataOffset, int32 Cardinality);
+	static void ComputeMeanDeviations(float MinMeanDeviation, const Eigen::MatrixXd& CenteredPoseMatrix, Eigen::VectorXd& MeanDeviations, int32& DataOffset, int32 Cardinality);
 
 private:
 	static FQuat DecodeQuatInternal(TArrayView<const float> Values, int32 DataOffset);
