motion matching - adding MinimumMeanDeviation to channels GenerateDDCKey

diff --git a/Engine/Plugins/Experimental/Animation/PoseSearch/Source/Runtime/Private/PoseSearchFeatureChannels.cpp b/Engine/Plugins/Experimental/Animation/PoseSearch/Source/Runtime/Private/PoseSearchFeatureChannels.cpp
--- a/Engine/Plugins/Experimental/Animation/PoseSearch/Source/Runtime/Private/PoseSearchFeatureChannels.cpp
+++ b/Engine/Plugins/Experimental/Animation/PoseSearch/Source/Runtime/Private/PoseSearchFeatureChannels.cpp
@@ -358,6 +358,7 @@ void UPoseSearchFeatureChannel_Position::GenerateDDCKey(FBlake3& InOutKeyHasher)
 	InOutKeyHasher.Update(&SampleTimeOffset, sizeof(SampleTimeOffset));
 	InOutKeyHasher.Update(&bUseFeaturesFromContinuingPose, sizeof(bUseFeaturesFromContinuingPose));
 	InOutKeyHasher.Update(&bUseSampleTimeOffsetRootBone, sizeof(bUseSampleTimeOffsetRootBone));
+	InOutKeyHasher.Update(&MinimumMeanDeviation, sizeof(MinimumMeanDeviation));
 }
 
 bool UPoseSearchFeatureChannel_Position::BuildQuery(UE::PoseSearch::FSearchContext& SearchContext, FPoseSearchFeatureVectorBuilder& InOutQuery) const
@@ -492,6 +493,7 @@ void UPoseSearchFeatureChannel_Heading::GenerateDDCKey(FBlake3& InOutKeyHasher)
 	InOutKeyHasher.Update(&HeadingAxis, sizeof(HeadingAxis));
 	InOutKeyHasher.Update(&bUseFeaturesFromContinuingPose, sizeof(bUseFeaturesFromContinuingPose));
 	InOutKeyHasher.Update(&bUseSampleTimeOffsetRootBone, sizeof(bUseSampleTimeOffsetRootBone));
+	InOutKeyHasher.Update(&MinimumMeanDeviation, sizeof(MinimumMeanDeviation));
 }
 
 bool UPoseSearchFeatureChannel_Heading::BuildQuery(UE::PoseSearch::FSearchContext& SearchContext, FPoseSearchFeatureVectorBuilder& InOutQuery) const
@@ -892,6 +894,7 @@ void UPoseSearchFeatureChannel_Pose::GenerateDDCKey(FBlake3& InOutKeyHasher) con
 	InOutKeyHasher.Update(MakeMemoryView(SampledBones));
 	InOutKeyHasher.Update(MakeMemoryView(SampleTimes));
 	InOutKeyHasher.Update(&bUseFeaturesFromContinuingPose, sizeof(bUseFeaturesFromContinuingPose));
+	InOutKeyHasher.Update(&MinimumMeanDeviation, sizeof(MinimumMeanDeviation));
 }
 
 bool UPoseSearchFeatureChannel_Pose::BuildQuery(UE::PoseSearch::FSearchContext& SearchContext, FPoseSearchFeatureVectorBuilder& InOutQuery) const
@@ -1409,6 +1412,7 @@ void UPoseSearchFeatureChannel_Trajectory::GenerateDDCKey(FBlake3& InOutKeyHashe
 	InOutKeyHasher.Update(&Domain, sizeof(Domain));
 	InOutKeyHasher.Update(MakeMemoryView(Samples));
 	InOutKeyHasher.Update(&Weight, sizeof(Weight));
+	InOutKeyHasher.Update(&MinimumMeanDeviation, sizeof(MinimumMeanDeviation));
 }
 
 bool UPoseSearchFeatureChannel_Trajectory::BuildQuery(UE::PoseSearch::FSearchContext& SearchContext, FPoseSearchFeatureVectorBuilder& InOutQuery) const
