motion matching - keeping an AssetIndex copy into the PreviewActor rather than a pointer, since while editing the database, the index is invalidated

diff --git a/Engine/Plugins/Experimental/Animation/PoseSearch/Source/Editor/Private/PoseSearchDatabaseViewModel.cpp b/Engine/Plugins/Experimental/Animation/PoseSearch/Source/Editor/Private/PoseSearchDatabaseViewModel.cpp
--- a/Engine/Plugins/Experimental/Animation/PoseSearch/Source/Editor/Private/PoseSearchDatabaseViewModel.cpp
+++ b/Engine/Plugins/Experimental/Animation/PoseSearch/Source/Editor/Private/PoseSearchDatabaseViewModel.cpp
@@ -174,7 +174,7 @@ namespace UE::PoseSearch
 			return PreviewActor;
 		}
 
-		PreviewActor.IndexAsset = &IndexAsset;
+		PreviewActor.IndexAsset = IndexAsset;
 		PreviewActor.CurrentPoseIndex = INDEX_NONE;
 
 		// todo: use preview when possible
@@ -259,8 +259,7 @@ namespace UE::PoseSearch
 				PreviewActor.Actor->SetActorTransform(RootMotion);
 				PreviewActor.AnimInstance->SetPosition(CurrentTime);
 
-				PreviewActor.CurrentPoseIndex =
-					PoseSearchDatabase->GetPoseIndexFromTime(CurrentTime, PreviewActor.IndexAsset);
+				PreviewActor.CurrentPoseIndex = PoseSearchDatabase->GetPoseIndexFromTime(CurrentTime, &PreviewActor.IndexAsset);
 			}
 		}
 	}
@@ -417,7 +416,7 @@ namespace UE::PoseSearch
 
 		if (SelectedPreviewActor)
 		{
-			SelectedActorIndexAsset = SelectedPreviewActor->IndexAsset;
+			SelectedActorIndexAsset = &SelectedPreviewActor->IndexAsset;
 		}
 	}
 
diff --git a/Engine/Plugins/Experimental/Animation/PoseSearch/Source/Editor/Public/PoseSearchDatabaseViewModel.h b/Engine/Plugins/Experimental/Animation/PoseSearch/Source/Editor/Public/PoseSearchDatabaseViewModel.h
--- a/Engine/Plugins/Experimental/Animation/PoseSearch/Source/Editor/Public/PoseSearchDatabaseViewModel.h
+++ b/Engine/Plugins/Experimental/Animation/PoseSearch/Source/Editor/Public/PoseSearchDatabaseViewModel.h
@@ -7,6 +7,7 @@
 #include "MovieSceneFwd.h"
 
 #include "PoseSearchDatabasePreviewScene.h"
+#include "PoseSearch/PoseSearch.h"
 
 class UWorld;
 class UPoseSearchDatabase;
@@ -41,7 +42,7 @@ namespace UE::PoseSearch
 		TWeakObjectPtr<AActor> Actor = nullptr;
 		TWeakObjectPtr<UDebugSkelMeshComponent> Mesh = nullptr;
 		TWeakObjectPtr<UAnimPreviewInstance> AnimInstance = nullptr;
-		const FPoseSearchIndexAsset* IndexAsset = nullptr;
+		FPoseSearchIndexAsset IndexAsset; // keeping a copy since database index can be invalidated
 		int32 CurrentPoseIndex = INDEX_NONE;
 
 		bool IsValid()
