PoseSearch: Fixed compilation errors when OBJECT_TRACE_ENABLED is 0

diff --git a/Engine/Plugins/Experimental/Animation/PoseSearch/Source/Runtime/Private/Trace/PoseSearchTraceLogger.cpp b/Engine/Plugins/Experimental/Animation/PoseSearch/Source/Runtime/Private/Trace/PoseSearchTraceLogger.cpp
--- a/Engine/Plugins/Experimental/Animation/PoseSearch/Source/Runtime/Private/Trace/PoseSearchTraceLogger.cpp
+++ b/Engine/Plugins/Experimental/Animation/PoseSearch/Source/Runtime/Private/Trace/PoseSearchTraceLogger.cpp
@@ -5,7 +5,6 @@
 #include "Trace/Trace.inl"
 #include "Animation/AnimNodeBase.h"
 
-#if UE_POSE_SEARCH_TRACE_ENABLED
 
 UE_TRACE_CHANNEL_DEFINE(PoseSearchChannel);
 
@@ -84,6 +83,7 @@ FArchive& operator<<(FArchive& Ar, FTraceMotionMatchingState& State)
 
 void FTraceMotionMatchingState::Output(const FAnimationBaseContext& InContext)
 {
+#if OBJECT_TRACE_ENABLED
 	if (!ShouldTrace(InContext))
 	{
 		return;
@@ -108,6 +108,7 @@ void FTraceMotionMatchingState::Output(const FAnimationBaseContext& InContext)
 
 	UE_TRACE_LOG(PoseSearch, MotionMatchingState, PoseSearchChannel)
 		<< MotionMatchingState.Data(ArchiveData.GetData(), ArchiveData.Num());
+#endif
 }
 
 const UPoseSearchDatabase* FTraceMotionMatchingState::GetCurrentDatabase() const
@@ -134,5 +135,4 @@ int32 FTraceMotionMatchingState::GetCurrentDatabasePoseIndex() const
 	return PoseEntry.DbPoseIdx;
 }
 
-}}
-#endif // UE_POSE_SEARCH_TRACE_ENABLED
\ No newline at end of file
+}}
\ No newline at end of file
diff --git a/Engine/Plugins/Experimental/Animation/PoseSearch/Source/Runtime/Private/Trace/PoseSearchTraceLogger.h b/Engine/Plugins/Experimental/Animation/PoseSearch/Source/Runtime/Private/Trace/PoseSearchTraceLogger.h
--- a/Engine/Plugins/Experimental/Animation/PoseSearch/Source/Runtime/Private/Trace/PoseSearchTraceLogger.h
+++ b/Engine/Plugins/Experimental/Animation/PoseSearch/Source/Runtime/Private/Trace/PoseSearchTraceLogger.h
@@ -4,9 +4,6 @@
 
 #include "PoseSearch/PoseSearch.h"
 
-
-#if UE_POSE_SEARCH_TRACE_ENABLED
-
 UE_TRACE_CHANNEL_EXTERN(PoseSearchChannel, POSESEARCH_API);
 
 namespace UE { namespace PoseSearch {
@@ -131,6 +128,7 @@ struct POSESEARCH_API FTraceMotionMatchingState
 	template<typename T>
 	static const T* GetObjectFromId(uint64 ObjectId)
 	{
+#if OBJECT_TRACE_ENABLED
 		if (ObjectId)
 		{
 			UObject* Object = FObjectTrace::GetObjectFromId(ObjectId);
@@ -139,13 +137,18 @@ struct POSESEARCH_API FTraceMotionMatchingState
 				return CastChecked<T>(Object);
 			}
 		}
+#endif
 
 		return nullptr;
 	}
 
 	static uint64 GetIdFromObject(const UObject* Object)
 	{
+#if OBJECT_TRACE_ENABLED
 		return FObjectTrace::GetObjectId(Object);
+#else
+		return 0;
+#endif
 	}
 	
 	static const FName Name;
@@ -155,6 +158,3 @@ ENUM_CLASS_FLAGS(FTraceMotionMatchingState::EFlags)
 POSESEARCH_API FArchive& operator<<(FArchive& Ar, FTraceMotionMatchingState& State);
 
 }}
-
-#endif // UE_POSE_SEARCH_TRACE_ENABLED
-
