PoseSearch: Fixed check when attempting to draw debug viz for the database preview actor during search index building

diff --git a/Engine/Plugins/Experimental/Animation/PoseSearch/Source/Editor/Private/PoseSearchDatabasePreviewScene.cpp b/Engine/Plugins/Experimental/Animation/PoseSearch/Source/Editor/Private/PoseSearchDatabasePreviewScene.cpp
--- a/Engine/Plugins/Experimental/Animation/PoseSearch/Source/Editor/Private/PoseSearchDatabasePreviewScene.cpp
+++ b/Engine/Plugins/Experimental/Animation/PoseSearch/Source/Editor/Private/PoseSearchDatabasePreviewScene.cpp
@@ -58,16 +58,18 @@ namespace UE::PoseSearch
 		}
 
 		const FDatabaseViewModel* ViewModel = GetEditor()->GetViewModel();
-		if (ViewModel->GetPoseSearchDatabase()->IsValidForSearch() &&
+		const UPoseSearchDatabase* Database = ViewModel->GetPoseSearchDatabase();
+
+		if (Database && Database->IsValidForSearch() &&
 			ViewModel->IsPoseFeaturesDrawMode(EFeaturesDrawMode::All))
 		{
 			for (const FDatabasePreviewActor& PreviewActor : ViewModel->GetPreviewActors())
 			{
-				if (PreviewActor.CurrentPoseIndex != INDEX_NONE)
+				if (Database->IsValidPoseIndex(PreviewActor.CurrentPoseIndex))
 				{
 					UE::PoseSearch::FDebugDrawParams DrawParams;
 					DrawParams.RootTransform = PreviewActor.Mesh->GetComponentTransform();
-					DrawParams.Database = ViewModel->GetPoseSearchDatabase();
+					DrawParams.Database = Database;
 					DrawParams.World = GetWorld();
 					DrawParams.DefaultLifeTime = 0.f;
 					DrawParams.PointSize = 5.f;
diff --git a/Engine/Plugins/Experimental/Animation/PoseSearch/Source/Runtime/Private/PoseSearch.cpp b/Engine/Plugins/Experimental/Animation/PoseSearch/Source/Runtime/Private/PoseSearch.cpp
--- a/Engine/Plugins/Experimental/Animation/PoseSearch/Source/Runtime/Private/PoseSearch.cpp
+++ b/Engine/Plugins/Experimental/Animation/PoseSearch/Source/Runtime/Private/PoseSearch.cpp
@@ -1247,7 +1247,7 @@ void UPoseSearchDatabase::GenerateDDCKey(FBlake3& InOutKeyHasher) const
 	InOutKeyHasher.Update(&KDTreeQueryNumNeighbors, sizeof(KDTreeQueryNumNeighbors));
 	InOutKeyHasher.Update(&PoseSearchMode, sizeof(PoseSearchMode));
 
-	// @todo: add database version if POSESEARCHDB_DERIVEDDATA_VER is not enought
+	// @todo: add database version if POSESEARCHDB_DERIVEDDATA_VER is not enough
 }
 
 #endif // WITH_EDITOR
@@ -1306,6 +1306,12 @@ bool UPoseSearchDatabase::IsValidForSearch() const
 	return bIsValid;
 }
 
+bool UPoseSearchDatabase::IsValidPoseIndex(int32 PoseIdx) const
+{
+	const FPoseSearchIndex* SearchIndex = GetSearchIndex();
+	return SearchIndex ? SearchIndex->IsValidPoseIndex(PoseIdx) : false;
+}
+
 void UPoseSearchDatabase::CollectSimpleSequences()
 {
 	for (auto& SimpleSequence: SimpleSequences)
@@ -4598,6 +4604,7 @@ void DrawFeatureVector(const FDebugDrawParams& DrawParams, int32 PoseIdx)
 	if (PoseIdx != INDEX_NONE && DrawParams.CanDraw())
 	{
 		const FPoseSearchIndex* SearchIndex = DrawParams.GetSearchIndex();
+
 		// PreprocessInfo happens to be invalid when updating the database
 		if (SearchIndex->PreprocessInfo.IsValid())
 		{
diff --git a/Engine/Plugins/Experimental/Animation/PoseSearch/Source/Runtime/Public/PoseSearch/PoseSearch.h b/Engine/Plugins/Experimental/Animation/PoseSearch/Source/Runtime/Public/PoseSearch/PoseSearch.h
--- a/Engine/Plugins/Experimental/Animation/PoseSearch/Source/Runtime/Public/PoseSearch/PoseSearch.h
+++ b/Engine/Plugins/Experimental/Animation/PoseSearch/Source/Runtime/Public/PoseSearch/PoseSearch.h
@@ -746,6 +746,7 @@ struct POSESEARCH_API FPoseSearchIndex
 	float MinCostAddend = -MAX_FLT;
 
 	bool IsValid() const;
+	bool IsValidPoseIndex(int32 PoseIdx) const { return PoseIdx < NumPoses; }
 	bool IsEmpty() const;
 
 	TArrayView<const float> GetPoseValues(int32 PoseIdx) const;
@@ -1162,6 +1163,7 @@ public:
 	bool GetSkipSearchIfPossible() const;
 	bool IsValidForIndexing() const;
 	bool IsValidForSearch() const;
+	bool IsValidPoseIndex(int32 PoseIdx) const;
 
 	float GetAssetTime(int32 PoseIdx, const FPoseSearchIndexAsset* SearchIndexAsset = nullptr) const;
 	int32 GetPoseIndexFromTime(float AssetTime, const FPoseSearchIndexAsset* SearchIndexAsset) const;
