motion matching - adding VelocityXY, PositionXY, for UPoseSearchFeatureChannel_Trajectory

diff --git a/Engine/Plugins/Experimental/Animation/PoseSearch/Source/Runtime/Private/PoseSearchFeatureChannels.cpp b/Engine/Plugins/Experimental/Animation/PoseSearch/Source/Runtime/Private/PoseSearchFeatureChannels.cpp
--- a/Engine/Plugins/Experimental/Animation/PoseSearch/Source/Runtime/Private/PoseSearchFeatureChannels.cpp
+++ b/Engine/Plugins/Experimental/Animation/PoseSearch/Source/Runtime/Private/PoseSearchFeatureChannels.cpp
@@ -543,8 +543,7 @@ void UPoseSearchFeatureChannel_Heading::DebugDraw(const UE::PoseSearch::FDebugDr
 	const bool bPersistent = EnumHasAnyFlags(DrawParams.Flags, EDebugDrawFlags::Persistent);
 
 	// @todo: search in the schema for a UPoseSearchFeatureChannel_Position with the same bone as SchemaBoneIdx and decode it's data as position
-	// @todo: initialize with the character position instead of FVector::Zero?
-	const FVector BonePos = DrawParams.Mesh != nullptr ? DrawParams.Mesh->GetSocketTransform(Bone.BoneName).GetLocation() : FVector::Zero();
+	const FVector BonePos = DrawParams.Mesh != nullptr ? DrawParams.Mesh->GetSocketTransform(Bone.BoneName).GetLocation() : DrawParams.RootTransform.GetTranslation();
 
 	int32 DataOffset = ChannelDataOffset;
 	const FVector BoneHeading = DrawParams.RootTransform.TransformPosition(FFeatureVectorHelper::DecodeVector(PoseVector, DataOffset));
@@ -1131,8 +1130,7 @@ void UPoseSearchFeatureChannel_Pose::DebugDraw(const UE::PoseSearch::FDebugDrawP
 		{
 			for (int32 SubsampleIdx = 0; SubsampleIdx != NumSubsamples; ++SubsampleIdx)
 			{
-				// @todo: initialize with the character position instead of FVector::Zero?
-				BonePos[SubsampleIdx] = DrawParams.Mesh != nullptr ? DrawParams.Mesh->GetSocketTransform(SampledBones[ChannelBoneIdx].Reference.BoneName).GetLocation() : FVector::Zero();
+				BonePos[SubsampleIdx] = DrawParams.Mesh != nullptr ? DrawParams.Mesh->GetSocketTransform(SampledBones[ChannelBoneIdx].Reference.BoneName).GetLocation() : DrawParams.RootTransform.GetTranslation();
 			}
 		}
 
@@ -1234,10 +1232,20 @@ void UPoseSearchFeatureChannel_Trajectory::InitializeSchema(UE::PoseSearch::FSch
 		{
 			DataOffset += FFeatureVectorHelper::EncodeVectorCardinality;
 		}
+		if (EnumHasAnyFlags(Sample.Flags, EPoseSearchTrajectoryFlags::PositionXY))
+		{
+			DataOffset += FFeatureVectorHelper::EncodeVector2DCardinality;
+		}
+
 		if (EnumHasAnyFlags(Sample.Flags, EPoseSearchTrajectoryFlags::Velocity))
 		{
 			DataOffset += FFeatureVectorHelper::EncodeVectorCardinality;
 		}
+		if (EnumHasAnyFlags(Sample.Flags, EPoseSearchTrajectoryFlags::VelocityXY))
+		{
+			DataOffset += FFeatureVectorHelper::EncodeVector2DCardinality;
+		}
+
 		if (EnumHasAnyFlags(Sample.Flags, EPoseSearchTrajectoryFlags::FacingDirection))
 		{
 			DataOffset += FFeatureVectorHelper::EncodeVectorCardinality;
@@ -1264,6 +1272,15 @@ void UPoseSearchFeatureChannel_Trajectory::FillWeights(TArray<float>& Weights) c
 			}
 			DataOffset += FFeatureVectorHelper::EncodeVectorCardinality;
 		}
+		if (EnumHasAnyFlags(Sample.Flags, EPoseSearchTrajectoryFlags::PositionXY))
+		{
+			for (int32 i = 0; i != FFeatureVectorHelper::EncodeVector2DCardinality; ++i)
+			{
+				Weights[DataOffset + i] = Weight * Sample.Weight;
+			}
+			DataOffset += FFeatureVectorHelper::EncodeVector2DCardinality;
+		}
+
 		if (EnumHasAnyFlags(Sample.Flags, EPoseSearchTrajectoryFlags::Velocity))
 		{
 			for (int32 i = 0; i != FFeatureVectorHelper::EncodeVectorCardinality; ++i)
@@ -1272,6 +1289,15 @@ void UPoseSearchFeatureChannel_Trajectory::FillWeights(TArray<float>& Weights) c
 			}
 			DataOffset += FFeatureVectorHelper::EncodeVectorCardinality;
 		}
+		if (EnumHasAnyFlags(Sample.Flags, EPoseSearchTrajectoryFlags::VelocityXY))
+		{
+			for (int32 i = 0; i != FFeatureVectorHelper::EncodeVector2DCardinality; ++i)
+			{
+				Weights[DataOffset + i] = Weight * Sample.Weight;
+			}
+			DataOffset += FFeatureVectorHelper::EncodeVector2DCardinality;
+		}
+
 		if (EnumHasAnyFlags(Sample.Flags, EPoseSearchTrajectoryFlags::FacingDirection))
 		{
 			for (int32 i = 0; i != FFeatureVectorHelper::EncodeVectorCardinality; ++i)
@@ -1307,10 +1333,20 @@ void UPoseSearchFeatureChannel_Trajectory::ComputeMeanDeviations(const Eigen::Ma
 		{
 			FFeatureVectorHelper::ComputeMeanDeviations(MinimumMeanDeviation, CenteredPoseMatrix, MeanDeviations, DataOffset, FFeatureVectorHelper::EncodeVectorCardinality);
 		}
+		if (EnumHasAnyFlags(Sample.Flags, EPoseSearchTrajectoryFlags::PositionXY))
+		{
+			FFeatureVectorHelper::ComputeMeanDeviations(MinimumMeanDeviation, CenteredPoseMatrix, MeanDeviations, DataOffset, FFeatureVectorHelper::EncodeVector2DCardinality);
+		}
+
 		if (EnumHasAnyFlags(Sample.Flags, EPoseSearchTrajectoryFlags::Velocity))
 		{
 			FFeatureVectorHelper::ComputeMeanDeviations(MinimumMeanDeviation, CenteredPoseMatrix, MeanDeviations, DataOffset, FFeatureVectorHelper::EncodeVectorCardinality);
 		}
+		if (EnumHasAnyFlags(Sample.Flags, EPoseSearchTrajectoryFlags::VelocityXY))
+		{
+			FFeatureVectorHelper::ComputeMeanDeviations(MinimumMeanDeviation, CenteredPoseMatrix, MeanDeviations, DataOffset, FFeatureVectorHelper::EncodeVector2DCardinality);
+		}
+
 		if (EnumHasAnyFlags(Sample.Flags, EPoseSearchTrajectoryFlags::FacingDirection))
 		{
 			FFeatureVectorHelper::ComputeMeanDeviations(MinimumMeanDeviation, CenteredPoseMatrix, MeanDeviations, DataOffset, FFeatureVectorHelper::EncodeVectorCardinality);
@@ -1358,50 +1394,57 @@ void UPoseSearchFeatureChannel_Trajectory::IndexAssetPrivate(const UE::PoseSearc
 	{
 		const float SubsampleTime = GetSampleTime(Indexer, Sample.Offset, SampleTime, Origin.RootDistance);
 
+		// For each pose subsample term, get the corresponding clip, accumulated root motion,
+			// and wrap the time parameter based on the clip's length.
+		const FSampleInfo SamplePast = Indexer.GetSampleInfoRelative(SubsampleTime - IndexingContext.SamplingContext->FiniteDelta, Origin);
+		const FSampleInfo SamplePresent = Indexer.GetSampleInfoRelative(SubsampleTime, Origin);
+		const FSampleInfo SampleFuture = Indexer.GetSampleInfoRelative(SubsampleTime + IndexingContext.SamplingContext->FiniteDelta, Origin);
+
+		// Mirror transforms if requested
+		const FTransform MirroredRootPast = Indexer.MirrorTransform(SamplePast.RootTransform);
+		const FTransform MirroredRootPresent = Indexer.MirrorTransform(SamplePresent.RootTransform);
+		const FTransform MirroredRootFuture = Indexer.MirrorTransform(SampleFuture.RootTransform);
+
+		// We can get a better finite difference if we ignore samples that have
+		// been clamped at either side of the clip. However, if the central sample 
+		// itself is clamped, or there are no samples that are clamped, we can just 
+		// use the central difference as normal.
+		FVector LinearVelocity;
+		if (SamplePast.bClamped && !SamplePresent.bClamped && !SampleFuture.bClamped)
+		{
+			LinearVelocity = (MirroredRootFuture.GetTranslation() - MirroredRootPresent.GetTranslation()) / IndexingContext.SamplingContext->FiniteDelta;
+		}
+		else if (SampleFuture.bClamped && !SamplePresent.bClamped && !SamplePast.bClamped)
+		{
+			LinearVelocity = (MirroredRootPresent.GetTranslation() - MirroredRootPast.GetTranslation()) / IndexingContext.SamplingContext->FiniteDelta;
+		}
+		else
+		{
+			LinearVelocity = (MirroredRootFuture.GetTranslation() - MirroredRootPast.GetTranslation()) / (IndexingContext.SamplingContext->FiniteDelta * 2.0f);
+		}
+
+
 		if (EnumHasAnyFlags(Sample.Flags, EPoseSearchTrajectoryFlags::Position))
 		{
-			const FSampleInfo SamplePresent = Indexer.GetSampleInfoRelative(SubsampleTime, Origin);
-			FFeatureVectorHelper::EncodeVector(FeatureVector.EditValues(), DataOffset, Indexer.MirrorTransform(SamplePresent.RootTransform).GetTranslation());
+			FFeatureVectorHelper::EncodeVector(FeatureVector.EditValues(), DataOffset, MirroredRootPresent.GetTranslation());
 		}
-	
-		if (EnumHasAnyFlags(Sample.Flags, EPoseSearchTrajectoryFlags::Velocity))
+		if (EnumHasAnyFlags(Sample.Flags, EPoseSearchTrajectoryFlags::PositionXY))
 		{
-			// For each pose subsample term, get the corresponding clip, accumulated root motion,
-			// and wrap the time parameter based on the clip's length.
-			const FSampleInfo SamplePast = Indexer.GetSampleInfoRelative(SubsampleTime - IndexingContext.SamplingContext->FiniteDelta, Origin);
-			const FSampleInfo SamplePresent = Indexer.GetSampleInfoRelative(SubsampleTime, Origin);
-			const FSampleInfo SampleFuture = Indexer.GetSampleInfoRelative(SubsampleTime + IndexingContext.SamplingContext->FiniteDelta, Origin);
-
-			// Mirror transforms if requested
-			const FTransform MirroredRootPast = Indexer.MirrorTransform(SamplePast.RootTransform);
-			const FTransform MirroredRootPresent = Indexer.MirrorTransform(SamplePresent.RootTransform);
-			const FTransform MirroredRootFuture = Indexer.MirrorTransform(SampleFuture.RootTransform);
-
-			// We can get a better finite difference if we ignore samples that have
-			// been clamped at either side of the clip. However, if the central sample 
-			// itself is clamped, or there are no samples that are clamped, we can just 
-			// use the central difference as normal.
-			FVector LinearVelocity;
-			if (SamplePast.bClamped && !SamplePresent.bClamped && !SampleFuture.bClamped)
-			{
-				LinearVelocity = (MirroredRootFuture.GetTranslation() - MirroredRootPresent.GetTranslation()) / IndexingContext.SamplingContext->FiniteDelta;
-			}
-			else if (SampleFuture.bClamped && !SamplePresent.bClamped && !SamplePast.bClamped)
-			{
-				LinearVelocity = (MirroredRootPresent.GetTranslation() - MirroredRootPast.GetTranslation()) / IndexingContext.SamplingContext->FiniteDelta;
-			}
-			else
-			{
-				LinearVelocity = (MirroredRootFuture.GetTranslation() - MirroredRootPast.GetTranslation()) / (IndexingContext.SamplingContext->FiniteDelta * 2.0f);
-			}
+			FFeatureVectorHelper::EncodeVector2D(FeatureVector.EditValues(), DataOffset, FVector2D(MirroredRootPresent.GetTranslation().X, MirroredRootPresent.GetTranslation().Y));
+		}
 
+		if (EnumHasAnyFlags(Sample.Flags, EPoseSearchTrajectoryFlags::Velocity))
+		{
 			FFeatureVectorHelper::EncodeVector(FeatureVector.EditValues(), DataOffset, LinearVelocity);
 		}
+		if (EnumHasAnyFlags(Sample.Flags, EPoseSearchTrajectoryFlags::VelocityXY))
+		{
+			FFeatureVectorHelper::EncodeVector2D(FeatureVector.EditValues(), DataOffset, FVector2D(LinearVelocity.X, LinearVelocity.Y));
+		}
 
 		if (EnumHasAnyFlags(Sample.Flags, EPoseSearchTrajectoryFlags::FacingDirection))
 		{
-			const FSampleInfo SamplePresent = Indexer.GetSampleInfoRelative(SubsampleTime, Origin);
-			FFeatureVectorHelper::EncodeVector(FeatureVector.EditValues(), DataOffset, Indexer.MirrorTransform(SamplePresent.RootTransform).GetRotation().GetForwardVector());
+			FFeatureVectorHelper::EncodeVector(FeatureVector.EditValues(), DataOffset, MirroredRootPresent.GetRotation().GetForwardVector());
 		}
 	}
 	check(DataOffset == ChannelDataOffset + ChannelCardinality);
@@ -1454,11 +1497,19 @@ bool UPoseSearchFeatureChannel_Trajectory::BuildQuery(UE::PoseSearch::FSearchCon
 		{
 			FFeatureVectorHelper::EncodeVector(InOutQuery.EditValues(), DataOffset, TrajectorySample.Transform.GetTranslation());
 		}
+		if (EnumHasAnyFlags(Sample.Flags, EPoseSearchTrajectoryFlags::PositionXY))
+		{
+			FFeatureVectorHelper::EncodeVector2D(InOutQuery.EditValues(), DataOffset, FVector2D(TrajectorySample.Transform.GetTranslation().X, TrajectorySample.Transform.GetTranslation().Y));
+		}
 
 		if (EnumHasAnyFlags(Sample.Flags, EPoseSearchTrajectoryFlags::Velocity))
 		{
 			FFeatureVectorHelper::EncodeVector(InOutQuery.EditValues(), DataOffset, TrajectorySample.LinearVelocity);
 		}
+		if (EnumHasAnyFlags(Sample.Flags, EPoseSearchTrajectoryFlags::VelocityXY))
+		{
+			FFeatureVectorHelper::EncodeVector2D(InOutQuery.EditValues(), DataOffset, FVector2D(TrajectorySample.LinearVelocity.X, TrajectorySample.LinearVelocity.Y));
+		}
 
 		if (EnumHasAnyFlags(Sample.Flags, EPoseSearchTrajectoryFlags::FacingDirection))
 		{
@@ -1502,10 +1553,34 @@ struct TrajectoryPositionReconstructor
 					bAddZeroOffsetSample = false;
 				}
 			}
+			if (EnumHasAnyFlags(Sample.Flags, EPoseSearchTrajectoryFlags::PositionXY))
+			{
+				const FVector2D Position2D = UE::PoseSearch::FFeatureVectorHelper::DecodeVector2D(PoseVector, DataOffset);
+
+				if (!EnumHasAnyFlags(Sample.Flags, EPoseSearchTrajectoryFlags::Position))
+				{
+					PositionAndOffsetSample PositionAndOffsetSample;
+					PositionAndOffsetSample.Position = FVector(Position2D.X, Position2D.Y, 0);
+					PositionAndOffsetSample.Position = RootTransform.TransformPosition(PositionAndOffsetSample.Position);
+					PositionAndOffsetSample.Offset = Sample.Offset;
+					PositionAndOffsetSamples.Add(PositionAndOffsetSample);
+
+					if (FMath::IsNearlyZero(Sample.Offset))
+					{
+						bAddZeroOffsetSample = false;
+					}
+				}
+			}
+
 			if (EnumHasAnyFlags(Sample.Flags, EPoseSearchTrajectoryFlags::Velocity))
 			{
 				DataOffset += UE::PoseSearch::FFeatureVectorHelper::EncodeVectorCardinality;
 			}
+			if (EnumHasAnyFlags(Sample.Flags, EPoseSearchTrajectoryFlags::VelocityXY))
+			{
+				DataOffset += UE::PoseSearch::FFeatureVectorHelper::EncodeVector2DCardinality;
+			}
+
 			if (EnumHasAnyFlags(Sample.Flags, EPoseSearchTrajectoryFlags::FacingDirection))
 			{
 				DataOffset += UE::PoseSearch::FFeatureVectorHelper::EncodeVectorCardinality;
@@ -1587,13 +1662,15 @@ void UPoseSearchFeatureChannel_Trajectory::DebugDraw(const UE::PoseSearch::FDebu
 	TrajectoryPositionReconstructor TrajectoryPositionReconstructor;
 	for (const FPoseSearchTrajectorySample& Sample : Samples)
 	{
+		bool bIsTrajectoryPosValid = false;
 		FVector TrajectoryPos = FVector::Zero();
 		
 		if (EnumHasAnyFlags(Sample.Flags, EPoseSearchTrajectoryFlags::Position))
 		{
 			TrajectoryPos = FFeatureVectorHelper::DecodeVector(PoseVector, DataOffset);
 			TrajectoryPos = DrawParams.RootTransform.TransformPosition(TrajectoryPos);
-			
+			bIsTrajectoryPosValid = true;
+
 			// validating TrajectoryPositionReconstructor
 			check((TrajectoryPositionReconstructor.GetReconstructedTrajectoryPos(*this, PoseVector, DrawParams.RootTransform, Sample.Offset) - TrajectoryPos).IsNearlyZero());
 
@@ -1609,7 +1686,30 @@ void UPoseSearchFeatureChannel_Trajectory::DebugDraw(const UE::PoseSearch::FDebu
 				DrawDebugSphere(DrawParams.World, TrajectoryPos, DrawDebugSphereSize, DrawDebugSphereSegments, Color, bPersistent, LifeTime, DepthPriority);
 			}
 		}
-		else
+		if (EnumHasAnyFlags(Sample.Flags, EPoseSearchTrajectoryFlags::PositionXY))
+		{
+			FVector2D TrajectoryPos2D = FFeatureVectorHelper::DecodeVector2D(PoseVector, DataOffset);
+			if (!bIsTrajectoryPosValid)
+			{
+				TrajectoryPos = FVector(TrajectoryPos2D.X, TrajectoryPos2D.Y, 0);
+				TrajectoryPos = DrawParams.RootTransform.TransformPosition(TrajectoryPos);
+				bIsTrajectoryPosValid = true;
+
+				const float GradientPercentage = NumSamples > 1 ? float(SampleIdx) / float(NumSamples - 1) : 0.f;
+				const FColor Color = DrawParams.GetColor(Sample.ColorPresetIndex, GradientPercentage);
+
+				if (EnumHasAnyFlags(DrawParams.Flags, EDebugDrawFlags::DrawFast | EDebugDrawFlags::DrawSearchIndex))
+				{
+					DrawDebugPoint(DrawParams.World, TrajectoryPos, DrawParams.PointSize, Color, bPersistent, LifeTime, DepthPriority);
+				}
+				else
+				{
+					DrawDebugSphere(DrawParams.World, TrajectoryPos, DrawDebugSphereSize, DrawDebugSphereSegments, Color, bPersistent, LifeTime, DepthPriority);
+				}
+			}
+		}
+		
+		if (!bIsTrajectoryPosValid)
 		{
 			TrajectoryPos = TrajectoryPositionReconstructor.GetReconstructedTrajectoryPos(*this, PoseVector, DrawParams.RootTransform, Sample.Offset);
 		}
@@ -1645,6 +1745,38 @@ void UPoseSearchFeatureChannel_Trajectory::DebugDraw(const UE::PoseSearch::FDebu
 				);
 			}
 		}
+		if (EnumHasAnyFlags(Sample.Flags, EPoseSearchTrajectoryFlags::VelocityXY))
+		{
+			const FVector2D TrajectoryVel2D = FFeatureVectorHelper::DecodeVector2D(PoseVector, DataOffset);
+			FVector TrajectoryVel(TrajectoryVel2D.X, TrajectoryVel2D.Y, 0.f);
+
+			const float GradientPercentage = NumSamples > 1 ? float(SampleIdx) / float(NumSamples - 1) : 0.f;
+			const FColor Color = DrawParams.GetColor(Sample.ColorPresetIndex, GradientPercentage);
+
+			TrajectoryVel *= DrawDebugVelocityScale;
+			TrajectoryVel = DrawParams.RootTransform.TransformVector(TrajectoryVel);
+			const FVector TrajectoryVelDirection = TrajectoryVel.GetSafeNormal();
+
+			if (EnumHasAnyFlags(DrawParams.Flags, EDebugDrawFlags::DrawSearchIndex))
+			{
+				DrawDebugLine(DrawParams.World, TrajectoryPos, TrajectoryPos + TrajectoryVel, Color, bPersistent, LifeTime, DepthPriority);
+			}
+			else
+			{
+				const float AdjustedThickness = EnumHasAnyFlags(DrawParams.Flags, EDebugDrawFlags::DrawFast) ? 0.0f : DrawDebugLineThickness;
+
+				DrawDebugLine(
+					DrawParams.World,
+					TrajectoryPos + TrajectoryVelDirection * DrawDebugSphereSize,
+					TrajectoryPos + TrajectoryVel,
+					Color,
+					bPersistent,
+					LifeTime,
+					DepthPriority,
+					AdjustedThickness
+				);
+			}
+		}
 
 		if (EnumHasAnyFlags(Sample.Flags, EPoseSearchTrajectoryFlags::FacingDirection))
 		{
diff --git a/Engine/Plugins/Experimental/Animation/PoseSearch/Source/Runtime/Public/PoseSearch/PoseSearchFeatureChannels.h b/Engine/Plugins/Experimental/Animation/PoseSearch/Source/Runtime/Public/PoseSearch/PoseSearchFeatureChannels.h
--- a/Engine/Plugins/Experimental/Animation/PoseSearch/Source/Runtime/Public/PoseSearch/PoseSearchFeatureChannels.h
+++ b/Engine/Plugins/Experimental/Animation/PoseSearch/Source/Runtime/Public/PoseSearch/PoseSearchFeatureChannels.h
@@ -157,6 +157,8 @@ enum class EPoseSearchTrajectoryFlags : uint32
 	Velocity = 1 << 0,
 	Position = 1 << 1,
 	FacingDirection = 1 << 3,
+	VelocityXY = 1 << 4,
+	PositionXY = 1 << 5,
 };
 ENUM_CLASS_FLAGS(EPoseSearchTrajectoryFlags);
 constexpr bool EnumHasAnyFlags(int32 Flags, EPoseSearchTrajectoryFlags Contains) { return (Flags & int32(Contains)) != 0; }
