PoseSearch: Fix for nondeterministic ddc keys caused by hashing struct memory. Refactored ddc key building to use object serialization.

diff --git a/Engine/Plugins/Experimental/Animation/PoseSearch/Source/Runtime/Private/PoseSearch.cpp b/Engine/Plugins/Experimental/Animation/PoseSearch/Source/Runtime/Private/PoseSearch.cpp
--- a/Engine/Plugins/Experimental/Animation/PoseSearch/Source/Runtime/Private/PoseSearch.cpp
+++ b/Engine/Plugins/Experimental/Animation/PoseSearch/Source/Runtime/Private/PoseSearch.cpp
@@ -4,7 +4,7 @@
 #include "PoseSearch/PoseSearchAnimNotifies.h"
 
 #if WITH_EDITORONLY_DATA
-// @todo: deleted this include after deleting the *_DEPRECATED properies
+// @todo: deleted this include after deleting the *_DEPRECATED properties
 #include "PoseSearch/PoseSearchFeatureChannels.h"
 #endif // WITH_EDITORONLY_DATA
 
@@ -41,6 +41,7 @@
 #if WITH_EDITOR
 #include "AssetRegistry/AssetRegistryModule.h"
 #include "AssetRegistry/ARFilter.h"
+#include "PoseSearch/PoseSearchDerivedDataKey.h"
 #endif // WITH_EDITOR
 
 IMPLEMENT_ANIMGRAPH_MESSAGE(UE::PoseSearch::IPoseHistoryProvider);
@@ -499,29 +500,6 @@ void UPoseSearchSchema::PostEditChangeProperty(FPropertyChangedEvent& PropertyCh
 
 	Super::PostEditChangeProperty(PropertyChangedEvent);
 }
-
-void UPoseSearchSchema::GenerateDDCKey(FBlake3& InOutKeyHasher) const
-{
-	for (const TObjectPtr<UPoseSearchFeatureChannel>& Channel : Channels)
-	{
-		if (Channel)
-		{
-			const FBlake3Hash& ChannelClassHash = Channel->GetClass()->GetSchemaHash(false);
-			InOutKeyHasher.Update(MakeMemoryView(ChannelClassHash.GetBytes()));
-			Channel->GenerateDDCKey(InOutKeyHasher);
-		}
-	}
-
-	InOutKeyHasher.Update(&DataPreprocessor, sizeof(DataPreprocessor));
-	InOutKeyHasher.Update(&EffectiveDataPreprocessor, sizeof(EffectiveDataPreprocessor));
-	InOutKeyHasher.Update(&SamplingInterval, sizeof(SamplingInterval));
-	InOutKeyHasher.Update(MakeMemoryView(BoneIndicesWithParents));
-	InOutKeyHasher.Update(&ContinuingPoseCostBias, sizeof(ContinuingPoseCostBias));
-	InOutKeyHasher.Update(&BaseCostBias, sizeof(BaseCostBias));
-	InOutKeyHasher.Update(&MirrorMismatchCostBias, sizeof(MirrorMismatchCostBias));
-
-	// @todo: add schema version if POSESEARCHDB_DERIVEDDATA_VER is not enought
-}
 #endif
 
 bool UPoseSearchSchema::IsValid() const
@@ -1129,69 +1107,20 @@ int32 UPoseSearchDatabase::GetNumberOfPrincipalComponents() const
 }
 
 #if WITH_EDITOR
-void UPoseSearchDatabase::AddDbSequenceToWriter(const FPoseSearchDatabaseSequence& DbSequence, FBlake3& InOutWriter)
-{
-	// Main Sequence
-	AddRawSequenceToWriter(DbSequence.Sequence, InOutWriter);
-	InOutWriter.Update(&DbSequence.SamplingRange, sizeof(DbSequence.SamplingRange));
-	if (DbSequence.Sequence)
-	{
-		InOutWriter.Update(&DbSequence.Sequence->bLoop, sizeof(DbSequence.Sequence->bLoop));
-	}
-	InOutWriter.Update(&DbSequence.MirrorOption, sizeof(DbSequence.MirrorOption));
-
-	// Lead in sequence
-	AddRawSequenceToWriter(DbSequence.LeadInSequence, InOutWriter);
-	if (DbSequence.LeadInSequence)
-	{
-		InOutWriter.Update(&DbSequence.LeadInSequence->bLoop, sizeof(DbSequence.LeadInSequence->bLoop));
-	}
-
-	// Follow up sequence
-	AddRawSequenceToWriter(DbSequence.FollowUpSequence, InOutWriter);
-	if (DbSequence.FollowUpSequence)
-	{
-		InOutWriter.Update(&DbSequence.FollowUpSequence->bLoop, sizeof(DbSequence.FollowUpSequence->bLoop));
-	}
-
-	// Tags
-	InOutWriter.Update(&DbSequence.GroupTags, sizeof(DbSequence.GroupTags));
 
-	// Notifies
-	AddPoseSearchNotifiesToWriter(DbSequence.Sequence, InOutWriter);
-}
-
-void UPoseSearchDatabase::AddRawSequenceToWriter(const UAnimSequence* Sequence,	FBlake3& InOutWriter)
+void UPoseSearchDatabase::AddRawSequenceToWriter(UAnimSequence* Sequence, UE::PoseSearch::FDerivedDataKeyBuilder& KeyBuilder)
 {
 	if (Sequence)
 	{
-		InOutWriter.Update(MakeMemoryView(Sequence->GetName()));
-		InOutWriter.Update(MakeMemoryView(Sequence->GetRawDataGuid().ToString()));
+		FName SequenceName = Sequence->GetFName();
+		FGuid SequenceGuid = Sequence->GetRawDataGuid();
+		KeyBuilder << SequenceName;
+		KeyBuilder << SequenceGuid;
+		KeyBuilder << Sequence->bLoop;
 	}
 }
 
-void UPoseSearchDatabase::AddDbBlendSpaceToWriter(const FPoseSearchDatabaseBlendSpace& DbBlendSpace, FBlake3& InOutWriter)
-{
-	if (IsValid(DbBlendSpace.BlendSpace))
-	{
-		const TArray<FBlendSample>& BlendSpaceSamples = DbBlendSpace.BlendSpace->GetBlendSamples();
-		for (const FBlendSample& Sample : BlendSpaceSamples)
-		{
-			AddRawSequenceToWriter(Sample.Animation, InOutWriter);
-			InOutWriter.Update(&Sample.SampleValue, sizeof(Sample.SampleValue));
-			InOutWriter.Update(&Sample.RateScale, sizeof(Sample.RateScale));
-		}
-
-		InOutWriter.Update(&DbBlendSpace.BlendSpace->bLoop, sizeof(DbBlendSpace.BlendSpace->bLoop));
-		InOutWriter.Update(&DbBlendSpace.MirrorOption, sizeof(DbBlendSpace.MirrorOption));
-		InOutWriter.Update(&DbBlendSpace.bUseGridForSampling, sizeof(DbBlendSpace.bUseGridForSampling));
-		InOutWriter.Update(&DbBlendSpace.NumberOfHorizontalSamples, sizeof(DbBlendSpace.NumberOfHorizontalSamples));
-		InOutWriter.Update(&DbBlendSpace.NumberOfVerticalSamples, sizeof(DbBlendSpace.NumberOfVerticalSamples));
-		InOutWriter.Update(&DbBlendSpace.GroupTags, sizeof(DbBlendSpace.GroupTags));
-	}
-}
-
-void UPoseSearchDatabase::AddPoseSearchNotifiesToWriter(const UAnimSequence* Sequence, FBlake3& InOutWriter)
+void UPoseSearchDatabase::AddPoseSearchNotifiesToWriter(UAnimSequence* Sequence, UE::PoseSearch::FDerivedDataKeyBuilder& KeyBuilder)
 {
 	if (!Sequence)
 	{
@@ -1211,41 +1140,60 @@ void UPoseSearchDatabase::AddPoseSearchNotifiesToWriter(const UAnimSequence* Seq
 
 		if (NotifyEvent->NotifyStateClass->IsA<UAnimNotifyState_PoseSearchBase>())
 		{
-			const float StartTime = NotifyEvent->GetTriggerTime();
-			const float EndTime = NotifyEvent->GetEndTriggerTime();
-			InOutWriter.Update(&StartTime, sizeof(float));
-			InOutWriter.Update(&EndTime, sizeof(float));
-
-			if (const auto ModifyCostNotifyState =
-				Cast<const UAnimNotifyState_PoseSearchModifyCost>(NotifyEvent->NotifyStateClass))
-			{
-				InOutWriter.Update(&ModifyCostNotifyState->CostAddend, sizeof(float));
-			}
+			float StartTime = NotifyEvent->GetTriggerTime();
+			float EndTime = NotifyEvent->GetEndTriggerTime();
+			KeyBuilder << StartTime;
+			KeyBuilder << EndTime;
+			KeyBuilder.Update(NotifyEvent->NotifyStateClass);
 		}
 	}
 }
 
-void UPoseSearchDatabase::GenerateDDCKey(FBlake3& InOutKeyHasher) const
+void UPoseSearchDatabase::AddDbSequenceToWriter(FPoseSearchDatabaseSequence& DbSequence, UE::PoseSearch::FDerivedDataKeyBuilder& InOutWriter)
 {
-	if (IsValid(Schema))
+	AddRawSequenceToWriter(DbSequence.Sequence, InOutWriter);
+	AddRawSequenceToWriter(DbSequence.LeadInSequence, InOutWriter);
+	AddRawSequenceToWriter(DbSequence.FollowUpSequence, InOutWriter);
+
+	AddPoseSearchNotifiesToWriter(DbSequence.Sequence, InOutWriter);
+}
+
+void UPoseSearchDatabase::AddDbBlendSpaceToWriter(FPoseSearchDatabaseBlendSpace& DbBlendSpace, UE::PoseSearch::FDerivedDataKeyBuilder& KeyBuilder)
+{
+	if (IsValid(DbBlendSpace.BlendSpace))
 	{
-		Schema->GenerateDDCKey(InOutKeyHasher);
+		const TArray<FBlendSample>& BlendSpaceSamples = DbBlendSpace.BlendSpace->GetBlendSamples();
+		for (const FBlendSample& Sample : BlendSpaceSamples)
+		{
+			AddRawSequenceToWriter(Sample.Animation, KeyBuilder);
+			FVector SampleValue = Sample.SampleValue;
+			float RateScale = Sample.RateScale;
+ 			KeyBuilder << SampleValue;
+ 			KeyBuilder << RateScale;
+		}
+
+		KeyBuilder << DbBlendSpace.BlendSpace->bLoop;
 	}
+}
 
-	for (int32 i = 0; i < Sequences.Num(); ++i)
+void UPoseSearchDatabase::BuildDerivedDataKey(UE::PoseSearch::FDerivedDataKeyBuilder& KeyBuilder)
+{
+	KeyBuilder.Update(this);
+
+	if (Schema)
 	{
-		AddDbSequenceToWriter(Sequences[i], InOutKeyHasher);
+		KeyBuilder.Update(Schema.Get());
 	}
 
-	for (int32 i = 0; i < BlendSpaces.Num(); ++i)
+	for (FPoseSearchDatabaseSequence& DbSequence : Sequences)
 	{
-		AddDbBlendSpaceToWriter(BlendSpaces[i], InOutKeyHasher);
+		AddDbSequenceToWriter(DbSequence, KeyBuilder);
 	}
 
-	InOutKeyHasher.Update(&NumberOfPrincipalComponents, sizeof(NumberOfPrincipalComponents));
-	InOutKeyHasher.Update(&KDTreeMaxLeafSize, sizeof(KDTreeMaxLeafSize));
-	InOutKeyHasher.Update(&KDTreeQueryNumNeighbors, sizeof(KDTreeQueryNumNeighbors));
-	InOutKeyHasher.Update(&PoseSearchMode, sizeof(PoseSearchMode));
+	for (FPoseSearchDatabaseBlendSpace& DbBlendSpace : BlendSpaces)
+	{
+		AddDbBlendSpaceToWriter(DbBlendSpace, KeyBuilder);
+	}
 
 	// @todo: add database version if POSESEARCHDB_DERIVEDDATA_VER is not enough
 }
diff --git a/Engine/Plugins/Experimental/Animation/PoseSearch/Source/Runtime/Private/PoseSearchDerivedData.cpp b/Engine/Plugins/Experimental/Animation/PoseSearch/Source/Runtime/Private/PoseSearchDerivedData.cpp
--- a/Engine/Plugins/Experimental/Animation/PoseSearch/Source/Runtime/Private/PoseSearchDerivedData.cpp
+++ b/Engine/Plugins/Experimental/Animation/PoseSearch/Source/Runtime/Private/PoseSearchDerivedData.cpp
@@ -10,8 +10,18 @@
 
 #if WITH_EDITOR
 #include "Serialization/BulkDataRegistry.h"
+#include "PoseSearch/PoseSearchDerivedDataKey.h"
 #endif // WITH_EDITOR
 
+#if UE_BUILD_DEBUG && WITH_EDITOR
+#define UE_POSE_SEARCH_DERIVED_DATA_LOGGING 1
+#endif
+
+#ifndef UE_POSE_SEARCH_DERIVED_DATA_LOGGING
+#define UE_POSE_SEARCH_DERIVED_DATA_LOGGING 0
+#endif
+
+
 #if WITH_EDITOR
 namespace UE::PoseSearch
 {
@@ -188,16 +198,25 @@ namespace UE::PoseSearch
 		Reader << DerivedData.SearchIndex;
 	}
 
-	FIoHash FPoseSearchDatabaseAsyncCacheTask::CreateKey(const UPoseSearchDatabase& Database)
+	FIoHash FPoseSearchDatabaseAsyncCacheTask::CreateKey(UPoseSearchDatabase& Database)
 	{
-		FBlake3 Writer;
-
-		const FGuid VersionGuid =  FDevSystemGuids::GetSystemGuid(FDevSystemGuids::Get().POSESEARCHDB_DERIVEDDATA_VER);
-		Writer.Update(MakeMemoryView(VersionGuid.ToString()));
-
-		Database.GenerateDDCKey(Writer);
-
-		return Writer.Finalize();
+		using UE::PoseSearch::FDerivedDataKeyBuilder;
+
+#if UE_POSE_SEARCH_DERIVED_DATA_LOGGING
+		const double StartTime = FPlatformTime::Seconds();
+#endif
+		FDerivedDataKeyBuilder KeyBuilder;
+		FGuid VersionGuid = FDevSystemGuids::GetSystemGuid(FDevSystemGuids::Get().POSESEARCHDB_DERIVEDDATA_VER);
+		KeyBuilder << VersionGuid;
+		Database.BuildDerivedDataKey(KeyBuilder);
+		FDerivedDataKeyBuilder::HashDigestType Hash = KeyBuilder.Finalize();
+
+#if UE_POSE_SEARCH_DERIVED_DATA_LOGGING
+		const double TotalTime = FPlatformTime::Seconds() - StartTime;
+		UE_LOG(LogPoseSearch, Log, TEXT("DDC key for '%s': %s (%.0f µs)"), *Database.GetName(), *WriteToString<48>(Hash), TotalTime * 1e6);
+#endif // UE_POSE_SEARCH_DERIVED_DATA_LOGGING
+
+		return Hash;
 	}
 
 #endif // WITH_EDITOR
diff --git a/Engine/Plugins/Experimental/Animation/PoseSearch/Source/Runtime/Private/PoseSearchDerivedData.h b/Engine/Plugins/Experimental/Animation/PoseSearch/Source/Runtime/Private/PoseSearchDerivedData.h
--- a/Engine/Plugins/Experimental/Animation/PoseSearch/Source/Runtime/Private/PoseSearchDerivedData.h
+++ b/Engine/Plugins/Experimental/Animation/PoseSearch/Source/Runtime/Private/PoseSearchDerivedData.h
@@ -92,7 +92,7 @@ namespace UE::PoseSearch
 		static const UE::DerivedData::FCacheBucket Bucket;
 
 	public:
-		static FIoHash CreateKey(const UPoseSearchDatabase& Database);
+		static FIoHash CreateKey(UPoseSearchDatabase& Database);
 	};
 #endif // WITH_EDITOR
 
diff --git a/Engine/Plugins/Experimental/Animation/PoseSearch/Source/Runtime/Private/PoseSearchFeatureChannels.cpp b/Engine/Plugins/Experimental/Animation/PoseSearch/Source/Runtime/Private/PoseSearchFeatureChannels.cpp
--- a/Engine/Plugins/Experimental/Animation/PoseSearch/Source/Runtime/Private/PoseSearchFeatureChannels.cpp
+++ b/Engine/Plugins/Experimental/Animation/PoseSearch/Source/Runtime/Private/PoseSearchFeatureChannels.cpp
@@ -298,6 +298,7 @@ static void CalculatePhasesFromLocalMinMax(const TArray<LocalMinMax>& MinMax, TA
 
 } // namespace UE::PoseSearch
 
+
 //////////////////////////////////////////////////////////////////////////
 // UPoseSearchFeatureChannel
 class USkeleton* UPoseSearchFeatureChannel::GetSkeleton(bool& bInvalidSkeletonIsError, const IPropertyHandle* PropertyHandle)
@@ -351,16 +352,6 @@ void UPoseSearchFeatureChannel_Position::IndexAsset(UE::PoseSearch::IAssetIndexe
 	}
 }
 
-void UPoseSearchFeatureChannel_Position::GenerateDDCKey(FBlake3& InOutKeyHasher) const
-{
-	InOutKeyHasher.Update(&Bone, sizeof(Bone));
-	InOutKeyHasher.Update(&Weight, sizeof(Weight));
-	InOutKeyHasher.Update(&SampleTimeOffset, sizeof(SampleTimeOffset));
-	InOutKeyHasher.Update(&bUseFeaturesFromContinuingPose, sizeof(bUseFeaturesFromContinuingPose));
-	InOutKeyHasher.Update(&bUseSampleTimeOffsetRootBone, sizeof(bUseSampleTimeOffsetRootBone));
-	InOutKeyHasher.Update(&MinimumMeanDeviation, sizeof(MinimumMeanDeviation));
-}
-
 bool UPoseSearchFeatureChannel_Position::BuildQuery(UE::PoseSearch::FSearchContext& SearchContext, FPoseSearchFeatureVectorBuilder& InOutQuery) const
 {
 	using namespace UE::PoseSearch;
@@ -485,17 +476,6 @@ void UPoseSearchFeatureChannel_Heading::IndexAsset(UE::PoseSearch::IAssetIndexer
 	}
 }
 
-void UPoseSearchFeatureChannel_Heading::GenerateDDCKey(FBlake3& InOutKeyHasher) const
-{
-	InOutKeyHasher.Update(&Bone, sizeof(Bone));
-	InOutKeyHasher.Update(&Weight, sizeof(Weight));
-	InOutKeyHasher.Update(&SampleTimeOffset, sizeof(SampleTimeOffset));
-	InOutKeyHasher.Update(&HeadingAxis, sizeof(HeadingAxis));
-	InOutKeyHasher.Update(&bUseFeaturesFromContinuingPose, sizeof(bUseFeaturesFromContinuingPose));
-	InOutKeyHasher.Update(&bUseSampleTimeOffsetRootBone, sizeof(bUseSampleTimeOffsetRootBone));
-	InOutKeyHasher.Update(&MinimumMeanDeviation, sizeof(MinimumMeanDeviation));
-}
-
 bool UPoseSearchFeatureChannel_Heading::BuildQuery(UE::PoseSearch::FSearchContext& SearchContext, FPoseSearchFeatureVectorBuilder& InOutQuery) const
 {
 	using namespace UE::PoseSearch;
@@ -887,15 +867,6 @@ void UPoseSearchFeatureChannel_Pose::AddPoseFeatures(UE::PoseSearch::IAssetIndex
 	check(DataOffset == ChannelDataOffset + ChannelCardinality);
 }
 
-void UPoseSearchFeatureChannel_Pose::GenerateDDCKey(FBlake3& InOutKeyHasher) const
-{
-	InOutKeyHasher.Update(&Weight, sizeof(Weight));
-	InOutKeyHasher.Update(MakeMemoryView(SampledBones));
-	InOutKeyHasher.Update(MakeMemoryView(SampleTimes));
-	InOutKeyHasher.Update(&bUseFeaturesFromContinuingPose, sizeof(bUseFeaturesFromContinuingPose));
-	InOutKeyHasher.Update(&MinimumMeanDeviation, sizeof(MinimumMeanDeviation));
-}
-
 bool UPoseSearchFeatureChannel_Pose::BuildQuery(UE::PoseSearch::FSearchContext& SearchContext, FPoseSearchFeatureVectorBuilder& InOutQuery) const
 {
 	using namespace UE::PoseSearch;
@@ -1206,6 +1177,7 @@ void UPoseSearchFeatureChannel_Pose::DebugDraw(const UE::PoseSearch::FDebugDrawP
 	check(DataOffset == ChannelDataOffset + ChannelCardinality);
 }
 
+
 //////////////////////////////////////////////////////////////////////////
 // UPoseSearchFeatureChannel_Trajectory
 void UPoseSearchFeatureChannel_Trajectory::PreSave(FObjectPreSaveContext ObjectSaveContext)
@@ -1450,14 +1422,6 @@ void UPoseSearchFeatureChannel_Trajectory::IndexAssetPrivate(const UE::PoseSearc
 	check(DataOffset == ChannelDataOffset + ChannelCardinality);
 }
 
-void UPoseSearchFeatureChannel_Trajectory::GenerateDDCKey(FBlake3& InOutKeyHasher) const
-{
-	InOutKeyHasher.Update(&Domain, sizeof(Domain));
-	InOutKeyHasher.Update(MakeMemoryView(Samples));
-	InOutKeyHasher.Update(&Weight, sizeof(Weight));
-	InOutKeyHasher.Update(&MinimumMeanDeviation, sizeof(MinimumMeanDeviation));
-}
-
 bool UPoseSearchFeatureChannel_Trajectory::BuildQuery(UE::PoseSearch::FSearchContext& SearchContext, FPoseSearchFeatureVectorBuilder& InOutQuery) const
 {
 	using namespace UE::PoseSearch;
diff --git a/Engine/Plugins/Experimental/Animation/PoseSearch/Source/Runtime/Public/PoseSearch/PoseSearch.h b/Engine/Plugins/Experimental/Animation/PoseSearch/Source/Runtime/Public/PoseSearch/PoseSearch.h
--- a/Engine/Plugins/Experimental/Animation/PoseSearch/Source/Runtime/Public/PoseSearch/PoseSearch.h
+++ b/Engine/Plugins/Experimental/Animation/PoseSearch/Source/Runtime/Public/PoseSearch/PoseSearch.h
@@ -20,6 +20,7 @@
 #include "GameplayTagContainer.h"
 #include "Interfaces/Interface_BoneReferenceSkeletonProvider.h"
 #include "PoseSearch/KDTree.h"
+#include "PoseSearch/PoseSearchDerivedDataKey.h"
 #include "ObjectTrace.h"
 #include "PoseSearch.generated.h"
 
@@ -228,10 +229,11 @@ struct POSESEARCH_API FPoseSearchBone
 	UPROPERTY(EditAnywhere, Category = Config)
 	float Weight = 1.f;
 
-	UPROPERTY(EditAnywhere, Category = Config)
+	UPROPERTY(EditAnywhere, Category = Config, meta=(ExcludeFromHash))
 	int32 ColorPresetIndex = 0;
 };
 
+
 //////////////////////////////////////////////////////////////////////////
 // Asset sampling and indexing
 
@@ -379,9 +381,6 @@ public:
 	// Called at database build time to calculate normalization values
 	virtual void ComputeMeanDeviations(const Eigen::MatrixXd& CenteredPoseMatrix, Eigen::VectorXd& MeanDeviations) const;
 
-	// Hash channel properties to produce a key for database derived data
-	virtual void GenerateDDCKey(FBlake3& InOutKeyHasher) const PURE_VIRTUAL(UPoseSearchFeatureChannel::GenerateDDCKey, );
-
 	// Called at runtime to add this channel's data to the query pose vector
 	virtual bool BuildQuery(UE::PoseSearch::FSearchContext& SearchContext, FPoseSearchFeatureVectorBuilder& InOutQuery) const PURE_VIRTUAL(UPoseSearchFeatureChannel::BuildQuery, return false;);
 
@@ -553,7 +552,6 @@ public: // IBoneReferenceSkeletonProvider
 
 #if WITH_EDITOR
 	virtual void PostEditChangeProperty(struct FPropertyChangedEvent& PropertyChangedEvent) override;
-	void GenerateDDCKey(FBlake3& InOutKeyHasher) const;
 #endif
 
 private:
@@ -1092,7 +1090,7 @@ struct FPoseCostDetails
 	TArray<float> CostVector;
 };
 
-}
+} // namespace UE::PoseSearch
 
 
 /** A data asset for indexing a collection of animation sequences. */
@@ -1178,12 +1176,12 @@ public:
 	int32 GetNumberOfPrincipalComponents() const;
 	
 #if WITH_EDITOR
-	void GenerateDDCKey(FBlake3& InOutKeyHasher) const;
+	void BuildDerivedDataKey(UE::PoseSearch::FDerivedDataKeyBuilder& KeyBuilder);
 private:
-	static void AddDbSequenceToWriter(const FPoseSearchDatabaseSequence& DbSequence, FBlake3& InOutWriter);
-	static void AddRawSequenceToWriter(const UAnimSequence* Sequence, FBlake3& InOutWriter);
-	static void AddPoseSearchNotifiesToWriter(const UAnimSequence* Sequence, FBlake3& InOutWriter);
-	static void AddDbBlendSpaceToWriter(const FPoseSearchDatabaseBlendSpace& DbBlendSpace, FBlake3& InOutWriter);
+	static void AddRawSequenceToWriter(UAnimSequence* Sequence, UE::PoseSearch::FDerivedDataKeyBuilder& KeyBuilder);
+	static void AddPoseSearchNotifiesToWriter(UAnimSequence* Sequence, UE::PoseSearch::FDerivedDataKeyBuilder& KeyBuilder);
+	static void AddDbSequenceToWriter(FPoseSearchDatabaseSequence& DbSequence, UE::PoseSearch::FDerivedDataKeyBuilder& InOutWriter);
+	static void AddDbBlendSpaceToWriter(FPoseSearchDatabaseBlendSpace& DbBlendSpace, UE::PoseSearch::FDerivedDataKeyBuilder& KeyBuilder);
 #endif // WITH_EDITOR
 
 public: // UObject
@@ -1195,7 +1193,7 @@ public: // UObject
 	virtual void PostEditChangeProperty(struct FPropertyChangedEvent& PropertyChangedEvent) override;
 	virtual void BeginCacheForCookedPlatformData(const ITargetPlatform* TargetPlatform) override;
 	virtual bool IsCachedCookedPlatformDataLoaded(const ITargetPlatform* TargetPlatform) override;
-#endif
+#endif // WITH_EDITOR
 
 private:
 	void CollectSimpleSequences();
diff --git a/Engine/Plugins/Experimental/Animation/PoseSearch/Source/Runtime/Public/PoseSearch/PoseSearchDerivedDataKey.h b/Engine/Plugins/Experimental/Animation/PoseSearch/Source/Runtime/Public/PoseSearch/PoseSearchDerivedDataKey.h
new file mode 100644
--- /dev/null
+++ b/Engine/Plugins/Experimental/Animation/PoseSearch/Source/Runtime/Public/PoseSearch/PoseSearchDerivedDataKey.h
@@ -0,0 +1,176 @@
+// Copyright Epic Games, Inc. All Rights Reserved.
+
+#pragma once
+
+#include "Containers/Queue.h"
+#include "Containers/Set.h"
+#include "Containers/UnrealString.h"
+#include "Serialization/ArchiveUObject.h"
+#include "Hash/Blake3.h"
+#include "UObject/UnrealType.h"
+#include "UObject/Object.h"
+
+#if WITH_EDITORONLY_DATA
+
+namespace UE::PoseSearch {
+
+//////////////////////////////////////////////////////////////////////////
+// TObjectHasher
+template <typename HashBuilder, typename HashDigest>
+class POSESEARCH_API TObjectHasher : public FArchiveUObject
+{
+public:
+	using HashDigestType = HashDigest;
+	using HashBuilderType = HashBuilder;
+
+	inline static const FName ExcludeFromHashName = FName(TEXT("ExcludeFromHash"));
+
+	/**
+	* Default constructor.
+	*/
+	TObjectHasher()
+	{
+		ArIgnoreOuterRef = true;
+
+		// Set TObjectHasher to be a saving archive instead of a reference collector.
+		// Reference collection causes FSoftObjectPtrs to be serialized by their weak pointer,
+		// which doesn't give a stable hash.  Serializing these to a saving archive will
+		// use a string reference instead, which is a more meaningful hash value.
+		SetIsSaving(true);
+	}
+
+	//~ Begin FArchive Interface
+	virtual bool ShouldSkipProperty(const FProperty* InProperty) const override
+	{
+		if (InProperty == nullptr)
+		{
+			return false;
+		}
+
+		if (FArchiveUObject::ShouldSkipProperty(InProperty))
+		{
+			return true;
+		}
+
+		if (InProperty->HasAllPropertyFlags(CPF_Transient))
+		{
+			return true;
+		}
+		const bool bExclude = InProperty->HasMetaData(ExcludeFromHashName);
+		return bExclude;
+	}
+
+	virtual void Serialize(void* Data, int64 Length) override
+	{
+		Hasher.Update(reinterpret_cast<uint8*>(Data), Length);
+	}
+
+	virtual FArchive& operator<<(FName& Name) override
+	{
+		// Don't include the name of the object being serialized, since that isn't technically part of the object's state
+		if (!ObjectBeingSerialized || (Name != ObjectBeingSerialized->GetFName()))
+		{
+			uint32 NameHash = GetTypeHash(Name);
+			Hasher.Update(&NameHash, sizeof(NameHash));
+		}
+
+		return *this;
+	}
+
+	virtual FArchive& operator<<(class UObject*& Object) override
+	{
+		FArchive& Ar = *this;
+
+		if (!RootObject || !Object || !Object->IsIn(RootObject))
+		{
+			auto UniqueName = GetPathNameSafe(Object);
+			Ar << UniqueName;
+		}
+		else
+		{
+			ObjectsToSerialize.Enqueue(Object);
+		}
+
+		return Ar;
+	}
+
+	using FArchiveUObject::operator<<; // For visibility of the overloads we don't override
+
+	virtual FString GetArchiveName() const override
+	{
+		return TEXT("TObjectHasher");
+	}
+	//~ End FArchive Interface
+
+	/**
+	* Serialize the given object and update the hash
+	*/
+	void Update(const UObject* Object, const UObject* Root)
+	{
+		RootObject = Root;
+		if (Object)
+		{
+			// Start with the given object
+			ObjectsToSerialize.Enqueue(Object);
+
+			// Continue until we no longer have any objects to serialize
+			while (ObjectsToSerialize.Dequeue(Object))
+			{
+				bool bAlreadyProcessed = false;
+				ObjectsAlreadySerialized.Add(Object, &bAlreadyProcessed);
+				// If we haven't already serialized this object
+				if (!bAlreadyProcessed)
+				{
+					// Serialize it
+					ObjectBeingSerialized = Object;
+					if (!CustomSerialize(Object))
+					{
+						const_cast<UObject*>(Object)->Serialize(*this);
+					}
+					ObjectBeingSerialized = nullptr;
+				}
+			}
+
+			// Cleanup
+			RootObject = nullptr;
+		}
+	}
+
+	void Update(const UObject* Object)
+	{
+		Update(Object, Object);
+	}
+
+	HashDigest Finalize()
+	{
+		return Hasher.Finalize();
+	}
+
+protected:
+	virtual bool CustomSerialize(const UObject* Object)
+	{
+		return false;
+	}
+
+	HashBuilder Hasher;
+
+	/** Queue of object references awaiting serialization */
+	TQueue<const UObject*> ObjectsToSerialize;
+
+	/** Set of objects that have already been serialized */
+	TSet<const UObject*> ObjectsAlreadySerialized;
+
+	/** Object currently being serialized */
+	const UObject* ObjectBeingSerialized = nullptr;
+
+	/** Root of object currently being serialized */
+	const UObject* RootObject = nullptr;
+};
+
+using FObjectHasherBlake3 = TObjectHasher<FBlake3, FBlake3Hash>;
+
+using FDerivedDataKeyBuilder = FObjectHasherBlake3;
+
+} // namespace UE::PoseSearch
+
+#endif // WITH_EDITORONLY_DATA
\ No newline at end of file
diff --git a/Engine/Plugins/Experimental/Animation/PoseSearch/Source/Runtime/Public/PoseSearch/PoseSearchFeatureChannels.h b/Engine/Plugins/Experimental/Animation/PoseSearch/Source/Runtime/Public/PoseSearch/PoseSearchFeatureChannels.h
--- a/Engine/Plugins/Experimental/Animation/PoseSearch/Source/Runtime/Public/PoseSearch/PoseSearchFeatureChannels.h
+++ b/Engine/Plugins/Experimental/Animation/PoseSearch/Source/Runtime/Public/PoseSearch/PoseSearchFeatureChannels.h
@@ -44,7 +44,6 @@ public:
 	virtual void InitializeSchema(UE::PoseSearch::FSchemaInitializer& Initializer) override;
 	virtual void FillWeights(TArray<float>& Weights) const override;
 	virtual void IndexAsset(UE::PoseSearch::IAssetIndexer& Indexer, UE::PoseSearch::FAssetIndexingOutput& IndexingOutput) const override;
-	virtual void GenerateDDCKey(FBlake3& InOutKeyHasher) const override;
 	virtual bool BuildQuery(UE::PoseSearch::FSearchContext& SearchContext, FPoseSearchFeatureVectorBuilder& InOutQuery) const override;
 	virtual void DebugDraw(const UE::PoseSearch::FDebugDrawParams& DrawParams, TArrayView<const float> PoseVector) const override;
 };
@@ -100,7 +99,6 @@ public:
 	virtual void InitializeSchema(UE::PoseSearch::FSchemaInitializer& Initializer) override;
 	virtual void FillWeights(TArray<float>& Weights) const override;
 	virtual void IndexAsset(UE::PoseSearch::IAssetIndexer& Indexer, UE::PoseSearch::FAssetIndexingOutput& IndexingOutput) const override;
-	virtual void GenerateDDCKey(FBlake3& InOutKeyHasher) const override;
 	virtual bool BuildQuery(UE::PoseSearch::FSearchContext& SearchContext, FPoseSearchFeatureVectorBuilder& InOutQuery) const override;
 	virtual void DebugDraw(const UE::PoseSearch::FDebugDrawParams& DrawParams, TArrayView<const float> PoseVector) const override;
 
@@ -139,7 +137,6 @@ public:
 	virtual void FillWeights(TArray<float>& Weights) const override;
 	virtual void IndexAsset(UE::PoseSearch::IAssetIndexer& Indexer, UE::PoseSearch::FAssetIndexingOutput& IndexingOutput) const override;
 	virtual void ComputeMeanDeviations(const Eigen::MatrixXd& CenteredPoseMatrix, Eigen::VectorXd& MeanDeviations) const override;
-	virtual void GenerateDDCKey(FBlake3& InOutKeyHasher) const override;
 	virtual bool BuildQuery(UE::PoseSearch::FSearchContext& SearchContext, FPoseSearchFeatureVectorBuilder& InOutQuery) const override;
 	virtual void DebugDraw(const UE::PoseSearch::FDebugDrawParams& DrawParams, TArrayView<const float> PoseVector) const override;
 
@@ -220,7 +217,6 @@ public:
 	virtual void FillWeights(TArray<float>& Weights) const override;
 	virtual void IndexAsset(UE::PoseSearch::IAssetIndexer& Indexer, UE::PoseSearch::FAssetIndexingOutput& IndexingOutput) const override;
 	virtual void ComputeMeanDeviations(const Eigen::MatrixXd& CenteredPoseMatrix, Eigen::VectorXd& MeanDeviations) const override;
-	virtual void GenerateDDCKey(FBlake3& InOutKeyHasher) const override;
 	virtual bool BuildQuery(UE::PoseSearch::FSearchContext& SearchContext, FPoseSearchFeatureVectorBuilder& InOutQuery) const override;
 	virtual void DebugDraw(const UE::PoseSearch::FDebugDrawParams& DrawParams, TArrayView<const float> PoseVector) const override;
 
