PoseSearch: Search now computes pose cost details.
- Removed FPoseCostDetials and and folded it into FPoseCost.
- Combined ComparePoses functions together into one path
- Full cost details are traced by the logger
- Rewind debugger rows and details panel are aware of cost details again

diff --git a/Engine/Plugins/Experimental/Animation/PoseSearch/Source/Editor/Private/PoseSearchDebugger.cpp b/Engine/Plugins/Experimental/Animation/PoseSearch/Source/Editor/Private/PoseSearchDebugger.cpp
--- a/Engine/Plugins/Experimental/Animation/PoseSearch/Source/Editor/Private/PoseSearchDebugger.cpp
+++ b/Engine/Plugins/Experimental/Animation/PoseSearch/Source/Editor/Private/PoseSearchDebugger.cpp
@@ -218,16 +218,16 @@ public:
 	
 	float GetChannelCost(int32 ChannelIdx) const
 	{
-		return PoseCostDetails.ChannelCosts.IsValidIndex(ChannelIdx)
-			? PoseCostDetails.ChannelCosts[ChannelIdx]
+		return PoseCost.ChannelCosts.IsValidIndex(ChannelIdx)
+			? PoseCost.ChannelCosts[ChannelIdx]
 			: 0.0f;
 	}
 	
-	float GetAddendsCost() const { return PoseCostDetails.NotifyCostAddend + PoseCostDetails.MirrorMismatchCostAddend; }
+	float GetAddendsCost() const { return PoseCost.GetCostAddend(); }
 
 	void CalculateColors(FChannelCostRange TotalCostRange, TArrayView<const FChannelCostRange> ChannelCostRanges)
 	{
-		const float CostColorBlend = (PoseCostDetails.PoseCost.GetTotalCost() - TotalCostRange.Min) / TotalCostRange.Delta;
+		const float CostColorBlend = (PoseCost.GetTotalCost() - TotalCostRange.Min) / TotalCostRange.Delta;
 		CostColor = UE::PoseSearch::LinearColorBlend(FLinearColor::Green, FLinearColor::Red, CostColorBlend);
 
 		ChannelCostColors.SetNum(ChannelCostRanges.Num());
@@ -259,7 +259,7 @@ public:
 	bool bMirrored = false;
 	bool bLooping = false;
 	FVector BlendParameters = FVector::Zero();
-	FPoseCostDetails PoseCostDetails;
+	FPoseSearchCost PoseCost;
 	FLinearColor CostColor = FLinearColor::White;
 	TArray<FLinearColor> ChannelCostColors;
 };
@@ -475,14 +475,14 @@ namespace DebuggerDatabaseColumns
 		
 		virtual FSortPredicate GetSortPredicate() const override
 		{
-			return [](const FRowDataRef& Row0, const FRowDataRef& Row1) -> bool { return Row0->PoseCostDetails.PoseCost < Row1->PoseCostDetails.PoseCost; };
+			return [](const FRowDataRef& Row0, const FRowDataRef& Row1) -> bool { return Row0->PoseCost < Row1->PoseCost; };
 		}
 		
 		virtual FText GetRowText(const FRowDataRef& Row) const override
 		{
-			if (Row->PoseCostDetails.PoseCost.IsValid())
+			if (Row->PoseCost.IsValid())
 			{
-				return FText::AsNumber(Row->PoseCostDetails.PoseCost.GetTotalCost());
+				return FText::AsNumber(Row->PoseCost.GetTotalCost());
 			}
 			else
 			{
@@ -519,7 +519,7 @@ namespace DebuggerDatabaseColumns
 
 		virtual FText GetRowText(const FRowDataRef& Row) const override
 		{
-			if (Row->PoseCostDetails.ChannelCosts.IsValidIndex(ChannelIdx))
+			if (Row->PoseCost.ChannelCosts.IsValidIndex(ChannelIdx))
 			{
 				return FText::AsNumber(Row->GetChannelCost(ChannelIdx));
 			}
@@ -932,7 +932,7 @@ void SDebuggerDatabaseView::CreateRows(const FTraceMotionMatchingStateMessage& S
 					Row->DbAssetIdx = SearchIndexAsset->SourceAssetIdx;
 					Row->AssetTime = Time;
 					Row->bMirrored = SearchIndexAsset->bMirrored;
-					Row->PoseCostDetails.PoseCost.Set(PoseEntry.Cost, 0.0f);
+					Row->PoseCost = PoseEntry.Cost;
 
 					if (SearchIndexAsset->Type == ESearchIndexAssetType::Sequence)
 					{
@@ -983,11 +983,11 @@ void SDebuggerDatabaseView::ComputeFilteredDatabaseRowsColors()
 	TArray<FChannelCostRange> ChannelCostRanges;
 	for (const TSharedRef<FDebuggerDatabaseRowData>& Row : FilteredDatabaseView.Rows)
 	{
-		const float Cost = Row->PoseCostDetails.PoseCost.GetTotalCost();
+		const float Cost = Row->PoseCost.GetTotalCost();
 		CostRange.Min = FMath::Min(CostRange.Min, Cost);
 		CostRange.Max = FMath::Max(CostRange.Max, Cost);
 
-		const int32 NumChannels = Row->PoseCostDetails.ChannelCosts.Num();
+		const int32 NumChannels = Row->PoseCost.ChannelCosts.Num();
 		if (ChannelCostRanges.Num() < NumChannels)
 		{
 			ChannelCostRanges.SetNum(NumChannels);
@@ -995,8 +995,8 @@ void SDebuggerDatabaseView::ComputeFilteredDatabaseRowsColors()
 
 		for (int32 ChannelIdx = 0; ChannelIdx != NumChannels; ++ChannelIdx)
 		{
-			ChannelCostRanges[ChannelIdx].Min = FMath::Min(ChannelCostRanges[ChannelIdx].Min, Row->PoseCostDetails.ChannelCosts[ChannelIdx]);
-			ChannelCostRanges[ChannelIdx].Max = FMath::Max(ChannelCostRanges[ChannelIdx].Max, Row->PoseCostDetails.ChannelCosts[ChannelIdx]);
+			ChannelCostRanges[ChannelIdx].Min = FMath::Min(ChannelCostRanges[ChannelIdx].Min, Row->PoseCost.ChannelCosts[ChannelIdx]);
+			ChannelCostRanges[ChannelIdx].Max = FMath::Max(ChannelCostRanges[ChannelIdx].Max, Row->PoseCost.ChannelCosts[ChannelIdx]);
 		}
 	}
 
@@ -1072,23 +1072,23 @@ void SDebuggerDatabaseView::Construct(const FArguments& InArgs)
 	auto CostColumn = MakeShared<FCost>(ColumnIdx++);
 	AddColumn(CostColumn);
 
-// 	int32 ChannelIdx = 0;
-// 	FChannelCostColumn::FParams ChannelCostParams;
-// 	ChannelCostParams.ChannelIdx = ChannelIdx++;
-// 	ChannelCostParams.SortIndex = ColumnIdx++;
-// 	AddColumn(MakeShared<FChannelCostColumn>(ChannelCostParams));
-// 
-// 	ChannelCostParams.ChannelIdx = ChannelIdx++;
-// 	ChannelCostParams.SortIndex = ColumnIdx++;
-// 	AddColumn(MakeShared<FChannelCostColumn>(ChannelCostParams));
-// 
-// 	ChannelCostParams.ChannelIdx = ChannelIdx++;
-// 	ChannelCostParams.SortIndex = ColumnIdx++;
-// 	AddColumn(MakeShared<FChannelCostColumn>(ChannelCostParams));
-// 
-// 	ChannelCostParams.ChannelIdx = ChannelIdx++;
-// 	ChannelCostParams.SortIndex = ColumnIdx++;
-// 	AddColumn(MakeShared<FChannelCostColumn>(ChannelCostParams));
+	int32 ChannelIdx = 0;
+	FChannelCostColumn::FParams ChannelCostParams;
+	ChannelCostParams.ChannelIdx = ChannelIdx++;
+	ChannelCostParams.SortIndex = ColumnIdx++;
+	AddColumn(MakeShared<FChannelCostColumn>(ChannelCostParams));
+
+	ChannelCostParams.ChannelIdx = ChannelIdx++;
+	ChannelCostParams.SortIndex = ColumnIdx++;
+	AddColumn(MakeShared<FChannelCostColumn>(ChannelCostParams));
+
+	ChannelCostParams.ChannelIdx = ChannelIdx++;
+	ChannelCostParams.SortIndex = ColumnIdx++;
+	AddColumn(MakeShared<FChannelCostColumn>(ChannelCostParams));
+
+	ChannelCostParams.ChannelIdx = ChannelIdx++;
+	ChannelCostParams.SortIndex = ColumnIdx++;
+	AddColumn(MakeShared<FChannelCostColumn>(ChannelCostParams));
 
 
 	AddColumn(MakeShared<FCostModifier>(ColumnIdx++));
@@ -1485,20 +1485,15 @@ void SDebuggerDetailsView::UpdateReflection(const FTraceMotionMatchingStateMessa
 				}
 			}
 
-			Reflection->CostVector = Selected->PoseCostDetails.CostVector;
-// 
-// 			TArray<TSharedRef<FDebuggerDatabaseRowData>> SelectedRows = DebuggerView->GetSelectedDatabaseRows();
-// 
-// 
-// 			State.DatabaseEntries[State.CurrentDbEntryIdx]
-// 
-// 			const TSharedRef<FDebuggerDatabaseRowData>& ActiveRow = DebuggerView->GetPoseIdxDatabaseRow(CurrentDbPoseIdx);
-// 
-// 			Reflection->CostVectorDifference = Reflection->CostVector;
-// 			for (int i = 0; i < Reflection->CostVectorDifference.Num(); ++i)
-// 			{
-// 				Reflection->CostVectorDifference[i] -= ActiveRow->PoseCostDetails.CostVector[i];
-// 			}
+			Reflection->CostVector = Selected->PoseCost.CostVector;
+
+			const FTraceMotionMatchingStatePoseEntry* ActivePoseEntry = State.GetCurrentPoseEntry();
+
+			Reflection->CostVectorDifference = Reflection->CostVector;
+			for (int i = 0; i < Reflection->CostVectorDifference.Num(); ++i)
+			{
+				Reflection->CostVectorDifference[i] -= ActivePoseEntry->Cost.CostVector[i];
+			}
 		}
 	}
 }
diff --git a/Engine/Plugins/Experimental/Animation/PoseSearch/Source/Runtime/Private/PoseSearch.cpp b/Engine/Plugins/Experimental/Animation/PoseSearch/Source/Runtime/Private/PoseSearch.cpp
--- a/Engine/Plugins/Experimental/Animation/PoseSearch/Source/Runtime/Private/PoseSearch.cpp
+++ b/Engine/Plugins/Experimental/Animation/PoseSearch/Source/Runtime/Private/PoseSearch.cpp
@@ -1860,6 +1860,7 @@ bool UPoseSearchDatabase::IsCachedCookedPlatformDataLoaded(const ITargetPlatform
 
 FPoseSearchCost UPoseSearchDatabase::ComparePoses(UE::PoseSearch::FSearchContext& SearchContext, int32 PoseIdx, UE::PoseSearch::EPoseComparisonFlags PoseComparisonFlags, int32 GroupIdx, const TArrayView<const float>& QueryValues) const
 {
+	using namespace Eigen;
 	using namespace UE::PoseSearch;
 
 	FPoseSearchCost Result;
@@ -1868,7 +1869,8 @@ FPoseSearchCost UPoseSearchDatabase::ComparePoses(UE::PoseSearch::FSearchContext
 	check(SearchIndex);
 
 	TArrayView<const float> PoseValues = SearchIndex->GetPoseValues(PoseIdx);
-	check(PoseValues.Num() == QueryValues.Num());
+	const int32 Dims = PoseValues.Num();
+	check(Dims == QueryValues.Num());
 
 	if (GroupIdx == INDEX_NONE)
 	{
@@ -1877,11 +1879,13 @@ FPoseSearchCost UPoseSearchDatabase::ComparePoses(UE::PoseSearch::FSearchContext
 		GroupIdx = SearchIndexAsset->SourceGroupIdx;
 	}
 
+	TArrayView<const float> Weights = SearchIndex->FindGroup(GroupIdx)->Weights;
+
 	Result.SetDissimilarity(CompareFeatureVectors(
 		PoseValues.Num(),
 		PoseValues.GetData(),
 		QueryValues.GetData(),
-		SearchIndex->FindGroup(GroupIdx)->Weights.GetData()));
+		Weights.GetData()));
 
 	const float MirrorMismatchAddend = SearchIndex->ComputeMirrorMismatchAddend(PoseIdx, SearchContext);
 	const float NotifyAddend = SearchIndex->ComputeNotifyAddend(PoseIdx);
@@ -1889,73 +1893,32 @@ FPoseSearchCost UPoseSearchDatabase::ComparePoses(UE::PoseSearch::FSearchContext
 
 	Result.SetCostAddend(NotifyAddend + MirrorMismatchAddend + ContinuingPoseCostAddend);
 
-	return Result;
-}
-
-FPoseSearchCost UPoseSearchDatabase::ComparePoses(UE::PoseSearch::FSearchContext& SearchContext, int32 PoseIdx, UE::PoseSearch::EPoseComparisonFlags PoseComparisonFlags, const TArrayView<const float>& QueryValues, UE::PoseSearch::FPoseCostDetails& OutPoseCostDetails) const
-{
-	using namespace Eigen;
-	using namespace UE::PoseSearch;
-
-	FPoseSearchCost Result;
-
-	const FPoseSearchIndex* SearchIndex = GetSearchIndex();
-	check(SearchIndex);
-
-	TArrayView<const float> PoseValues = SearchIndex->GetPoseValues(PoseIdx);
-	const int32 Dims = PoseValues.Num();
-	check(Dims == QueryValues.Num());
-
-	OutPoseCostDetails.CostVector.SetNum(Dims);
-
-	// Setup Eigen views onto our vectors
-	auto OutCostVector = Map<ArrayXf>(OutPoseCostDetails.CostVector.GetData(), Dims);
-	auto PoseVector = Map<const ArrayXf>(PoseValues.GetData(), Dims);
-	auto QueryVector = Map<const ArrayXf>(QueryValues.GetData(), Dims);
-
-	// Compute weighted squared difference vector
-	const FPoseSearchIndexAsset* SearchIndexAsset = SearchIndex->FindAssetForPose(PoseIdx);
-	const TArray<float>& Weights = SearchIndex->FindGroup(SearchIndexAsset->SourceGroupIdx)->Weights;
-
-	check(Weights.Num() == Dims);
-	auto WeightsVector = Map<const ArrayXf>(Weights.GetData(), Dims);
-
-	OutCostVector = WeightsVector * (PoseVector - QueryVector).square();
-	Result.SetDissimilarity(OutCostVector.sum());
-
-	const float MirrorMismatchAddend = SearchIndex->ComputeMirrorMismatchAddend(PoseIdx, SearchContext);
-	const float NotifyAddend = SearchIndex->ComputeNotifyAddend(PoseIdx);
-	const float ContinuingPoseCostAddend = SearchIndex->ComputeContinuingPoseCostAddend(PoseIdx, PoseComparisonFlags);
+#if UE_POSE_SEARCH_TRACE_ENABLED && WITH_EDITORONLY_DATA
+	if (SearchContext.bIsTracing)
+	{
+		// Setup Eigen views onto our vectors
+		Result.CostVector.SetNum(Dims);
+		auto OutCostVector = Map<ArrayXf>(Result.CostVector.GetData(), Dims);
+		auto PoseVector = Map<const ArrayXf>(PoseValues.GetData(), Dims);
+		auto QueryVector = Map<const ArrayXf>(QueryValues.GetData(), Dims);
+		auto WeightsVector = Map<const ArrayXf>(Weights.GetData(), Dims);
+		
+		// Compute weighted squared difference vector
+		OutCostVector = WeightsVector * (PoseVector - QueryVector).square();
 
-	Result.SetCostAddend(NotifyAddend + MirrorMismatchAddend + ContinuingPoseCostAddend);
+		// Output cost details
+		CalcChannelCosts(SearchIndex->Schema, Result.CostVector, Result.ChannelCosts);
+		Result.MirrorMismatchCostAddend = MirrorMismatchAddend;
+		Result.NotifyCostAddend = NotifyAddend;
+		Result.ContinuingPoseCostAddend = ContinuingPoseCostAddend;
 
-	// Output cost details
-	OutPoseCostDetails.NotifyCostAddend = NotifyAddend + ContinuingPoseCostAddend;
-	OutPoseCostDetails.MirrorMismatchCostAddend = MirrorMismatchAddend;
-	OutPoseCostDetails.PoseCost = Result;
-	CalcChannelCosts(
-		SearchIndex->Schema,
-		OutPoseCostDetails.CostVector,
-		OutPoseCostDetails.ChannelCosts);
-
-#if DO_GUARD_SLOW
-	{
-		// Verify details pose comparator agrees with runtime pose comparator
-		FPoseSearchCost RuntimeComparatorCost = ComparePoses(
-			SearchContext, 
-			PoseIdx, 
-			PoseComparisonFlags,
-			SearchIndexAsset->SourceGroupIdx,
-			QueryValues);
-		checkSlow(FMath::IsNearlyEqual(Result.GetTotalCost(), RuntimeComparatorCost.GetTotalCost(), 1e-3f));
-
-		// Verify channel cost decomposition agrees with runtime pose comparator
-		auto OutChannelCosts = Map<const ArrayXf>(
-			OutPoseCostDetails.ChannelCosts.GetData(), 
-			OutPoseCostDetails.ChannelCosts.Num());
-		checkSlow(FMath::IsNearlyEqual(OutChannelCosts.sum(), RuntimeComparatorCost.GetDissimilarity(), 1e-3f));
+		// Verify channel cost decomposition for tracing agrees with runtime pose comparator
+		auto ChannelCostsVector = Map<const ArrayXf>(
+			Result.ChannelCosts.GetData(),
+			Result.ChannelCosts.Num());
+		checkSlow(FMath::IsNearlyEqual(ChannelCostsVector.sum(), Result.GetDissimilarity(), 1e-3f));
 	}
-#endif
+#endif // UE_POSE_SEARCH_TRACE_ENABLED && WITH_EDITORONLY_DATA
 
 	return Result;
 }
@@ -2161,7 +2124,7 @@ UE::PoseSearch::FSearchResult UPoseSearchDatabase::SearchPCAKDTree(UE::PoseSearc
 				}
 
 #if UE_POSE_SEARCH_TRACE_ENABLED
-				SearchContext.BestCandidates.Add(PoseCost.GetTotalCost(), PoseIdx, this);
+				SearchContext.BestCandidates.Add(PoseCost, PoseIdx, this);
 #endif
 			}
 		}
@@ -2283,7 +2246,7 @@ UE::PoseSearch::FSearchResult UPoseSearchDatabase::SearchBruteForce(UE::PoseSear
 #if UE_POSE_SEARCH_TRACE_ENABLED
 				if (PoseSearchMode == EPoseSearchMode::BruteForce)
 				{
-					SearchContext.BestCandidates.Add(PoseCost.GetTotalCost(), PoseIdx, this);
+					SearchContext.BestCandidates.Add(PoseCost, PoseIdx, this);
 				}
 #endif
 			}
diff --git a/Engine/Plugins/Experimental/Animation/PoseSearch/Source/Runtime/Private/PoseSearchLibrary.cpp b/Engine/Plugins/Experimental/Animation/PoseSearch/Source/Runtime/Private/PoseSearchLibrary.cpp
--- a/Engine/Plugins/Experimental/Animation/PoseSearch/Source/Runtime/Private/PoseSearchLibrary.cpp
+++ b/Engine/Plugins/Experimental/Animation/PoseSearch/Source/Runtime/Private/PoseSearchLibrary.cpp
@@ -222,6 +222,11 @@ static void TraceMotionMatchingState(
 #if UE_POSE_SEARCH_TRACE_ENABLED
 	using namespace UE::PoseSearch;
 
+	if (!IsTracing(UpdateContext))
+	{
+		return;
+	}
+
 	if (!MotionMatchingState.CurrentSearchResult.IsValid())
 	{
 		return;
@@ -275,7 +280,7 @@ static void TraceMotionMatchingState(
 		int32 LastResultPoseEntryIdx = DbEntry.PoseEntries.AddUnique({LastResult.PoseIdx});
 		FTraceMotionMatchingStatePoseEntry& PoseEntry = DbEntry.PoseEntries[LastResultPoseEntryIdx];
 
-		PoseEntry.Cost = LastResult.PoseCost.GetTotalCost();
+		PoseEntry.Cost = LastResult.PoseCost;
 		PoseEntry.Flags |= FTraceMotionMatchingStatePoseEntry::EFlags::ContinuingPose;
 		bContinuingPoseTraced = true;
 	}
@@ -291,7 +296,7 @@ static void TraceMotionMatchingState(
 		int32 PoseEntryIdx = DbEntry.PoseEntries.AddUnique({MotionMatchingState.CurrentSearchResult.PoseIdx});
 		FTraceMotionMatchingStatePoseEntry& PoseEntry = DbEntry.PoseEntries[PoseEntryIdx];
 
-		PoseEntry.Cost = MotionMatchingState.CurrentSearchResult.PoseCost.GetTotalCost();
+		PoseEntry.Cost = MotionMatchingState.CurrentSearchResult.PoseCost;
 		PoseEntry.Flags |= FTraceMotionMatchingStatePoseEntry::EFlags::CurrentPose;
 		bCurrentPoseTraced = true;
 	}
@@ -429,6 +434,7 @@ void UpdateMotionMatchingState(
 		SearchContext.Trajectory = &Trajectory;
 		SearchContext.OwningComponent = Context.AnimInstanceProxy->GetSkelMeshComponent();
 		SearchContext.BoneContainer = &Context.AnimInstanceProxy->GetRequiredBones();
+		SearchContext.bIsTracing = IsTracing(Context);
 		SearchContext.bForceInterrupt = bForceInterrupt;
 		SearchContext.bCanAdvance = bCanAdvance;
 		SearchContext.CurrentResult = InOutMotionMatchingState.CurrentSearchResult;
diff --git a/Engine/Plugins/Experimental/Animation/PoseSearch/Source/Runtime/Private/Trace/PoseSearchTraceLogger.cpp b/Engine/Plugins/Experimental/Animation/PoseSearch/Source/Runtime/Private/Trace/PoseSearchTraceLogger.cpp
--- a/Engine/Plugins/Experimental/Animation/PoseSearch/Source/Runtime/Private/Trace/PoseSearchTraceLogger.cpp
+++ b/Engine/Plugins/Experimental/Animation/PoseSearch/Source/Runtime/Private/Trace/PoseSearchTraceLogger.cpp
@@ -17,8 +17,9 @@ namespace UE { namespace PoseSearch {
 const FName FTraceLogger::Name("PoseSearch");
 const FName FTraceMotionMatchingState::Name("MotionMatchingState");
 
-static bool ShouldTrace(const FAnimationBaseContext& InContext)
+bool IsTracing(const FAnimationBaseContext& InContext)
 {
+#if UE_POSE_SEARCH_TRACE_ENABLED
 	const bool bChannelEnabled = UE_TRACE_CHANNELEXPR_IS_ENABLED(PoseSearchChannel);
 	if (!bChannelEnabled)
 	{
@@ -32,6 +33,9 @@ static bool ShouldTrace(const FAnimationBaseContext& InContext)
 
 	check(InContext.AnimInstanceProxy);
 	return !CANNOT_TRACE_OBJECT(InContext.AnimInstanceProxy->GetSkelMeshComponent());
+#else
+	return false;
+#endif
 }
 
 FArchive& operator<<(FArchive& Ar, FTraceMessage& State)
@@ -47,7 +51,7 @@ FArchive& operator<<(FArchive& Ar, FTraceMessage& State)
 FArchive& operator<<(FArchive& Ar, FTraceMotionMatchingStatePoseEntry& Entry)
 {
 	Ar << Entry.DbPoseIdx;
-	Ar << Entry.Cost;
+	FPoseSearchCost::StaticStruct()->SerializeItem(Ar, &Entry.Cost, nullptr);
 	Ar << Entry.Flags;
 	return Ar;
 }
@@ -84,7 +88,7 @@ FArchive& operator<<(FArchive& Ar, FTraceMotionMatchingState& State)
 void FTraceMotionMatchingState::Output(const FAnimationBaseContext& InContext)
 {
 #if OBJECT_TRACE_ENABLED
-	if (!ShouldTrace(InContext))
+	if (!IsTracing(InContext))
 	{
 		return;
 	}
@@ -124,15 +128,20 @@ const UPoseSearchDatabase* FTraceMotionMatchingState::GetCurrentDatabase() const
 }
 
 int32 FTraceMotionMatchingState::GetCurrentDatabasePoseIndex() const
+{
+	const FTraceMotionMatchingStatePoseEntry* PoseEntry = GetCurrentPoseEntry();
+	return PoseEntry ? PoseEntry->DbPoseIdx : INDEX_NONE;
+}
+
+const FTraceMotionMatchingStatePoseEntry* FTraceMotionMatchingState::GetCurrentPoseEntry() const
 {
 	if ((CurrentDbEntryIdx == INDEX_NONE) || (CurrentPoseEntryIdx == INDEX_NONE))
 	{
-		return INDEX_NONE;
+		return nullptr;
 	}
 
 	const FTraceMotionMatchingStateDatabaseEntry& DbEntry = DatabaseEntries[CurrentDbEntryIdx];
-	const FTraceMotionMatchingStatePoseEntry& PoseEntry = DbEntry.PoseEntries[CurrentPoseEntryIdx];
-	return PoseEntry.DbPoseIdx;
+	return &DbEntry.PoseEntries[CurrentPoseEntryIdx];
 }
 
 }}
\ No newline at end of file
diff --git a/Engine/Plugins/Experimental/Animation/PoseSearch/Source/Runtime/Private/Trace/PoseSearchTraceLogger.h b/Engine/Plugins/Experimental/Animation/PoseSearch/Source/Runtime/Private/Trace/PoseSearchTraceLogger.h
--- a/Engine/Plugins/Experimental/Animation/PoseSearch/Source/Runtime/Private/Trace/PoseSearchTraceLogger.h
+++ b/Engine/Plugins/Experimental/Animation/PoseSearch/Source/Runtime/Private/Trace/PoseSearchTraceLogger.h
@@ -8,6 +8,8 @@ UE_TRACE_CHANNEL_EXTERN(PoseSearchChannel, POSESEARCH_API);
 
 namespace UE { namespace PoseSearch {
 
+bool IsTracing(const FAnimationBaseContext& InContext);
+
 struct POSESEARCH_API FTraceLogger
 {
 	/** Used for reading trace data */
@@ -42,10 +44,9 @@ struct POSESEARCH_API FTraceMotionMatchingStatePoseEntry
 		ContinuingPose	= 1 << 0,
 		CurrentPose		= 1 << 1,
 	};
-
 	int32 DbPoseIdx = INDEX_NONE;
-	float Cost = 0.f;
 	EFlags Flags = EFlags::None;
+	FPoseSearchCost Cost;
 
 	bool operator==(const FTraceMotionMatchingStatePoseEntry& Other) const { return DbPoseIdx == Other.DbPoseIdx; }
 };
@@ -124,6 +125,7 @@ struct POSESEARCH_API FTraceMotionMatchingState
 
 	const UPoseSearchDatabase* GetCurrentDatabase() const;
 	int32 GetCurrentDatabasePoseIndex() const;
+	const FTraceMotionMatchingStatePoseEntry* GetCurrentPoseEntry() const;
 
 	template<typename T>
 	static const T* GetObjectFromId(uint64 ObjectId)
diff --git a/Engine/Plugins/Experimental/Animation/PoseSearch/Source/Runtime/Public/PoseSearch/PoseSearch.h b/Engine/Plugins/Experimental/Animation/PoseSearch/Source/Runtime/Public/PoseSearch/PoseSearch.h
--- a/Engine/Plugins/Experimental/Animation/PoseSearch/Source/Runtime/Public/PoseSearch/PoseSearch.h
+++ b/Engine/Plugins/Experimental/Animation/PoseSearch/Source/Runtime/Public/PoseSearch/PoseSearch.h
@@ -867,11 +867,10 @@ public:
 	FGameplayTag Tag;
 };
 
-USTRUCT(BlueprintType)
-struct FPoseSearchCost
+USTRUCT()
+struct POSESEARCH_API FPoseSearchCost
 {
 	GENERATED_BODY()
-
 public:
 	FPoseSearchCost() = default;
 	FPoseSearchCost(float InDissimilarity, float InCostAddend)
@@ -926,8 +925,32 @@ protected:
 
 	UPROPERTY()
 	float TotalCost = MAX_flt;
+
+#if WITH_EDITORONLY_DATA
+public:
+	// Contribution from ModifyCost anim notify
+	UPROPERTY()
+	float NotifyCostAddend = 0.0f;
+
+	// Contribution from mirroring cost
+	UPROPERTY()
+	float MirrorMismatchCostAddend = 0.0f;
+
+	UPROPERTY()
+	float ContinuingPoseCostAddend = 0.0f;
+
+	// Cost breakdown per channel (e.g. pose cost, time-based trajectory cost, distance-based trajectory cost, etc.)
+	UPROPERTY()
+	TArray<float> ChannelCosts;
+
+	// Difference vector computed as W*((P-Q)^2) without the cost modifier applied
+	// Where P is the pose vector, Q is the query vector, W is the weights vector, and multiplication/exponentiation are element-wise operations
+	UPROPERTY()
+	TArray<float> CostVector;
+#endif // WITH_EDITORONLY_DATA
 };
 
+
 /**
 * Helper object for writing features into a float buffer according to a feature vector layout.
 * Keeps track of which features are present, allowing the feature vector to be built up piecemeal.
@@ -1031,7 +1054,6 @@ public:
 	virtual FPoseHistory& GetPoseHistory() = 0;
 };
 
-
 struct FSearchResult
 {
 	FPoseSearchCost PoseCost;
@@ -1069,27 +1091,6 @@ struct FSearchResult
 	void Reset();
 };
 
-/**
- * Cost details for pose analysis in the rewind debugger
- */
-struct FPoseCostDetails
-{
-	FPoseSearchCost PoseCost;
-
-	// Contribution from ModifyCost anim notify
-	float NotifyCostAddend = 0.0f;
-
-	// Contribution from mirroring cost
-	float MirrorMismatchCostAddend = 0.0f;
-
-	// Cost breakdown per channel (e.g. pose cost, time-based trajectory cost, distance-based trajectory cost, etc.)
-	TArray<float> ChannelCosts;
-
-	// Difference vector computed as W*((P-Q)^2) without the cost modifier applied
-	// Where P is the pose vector, Q is the query vector, W is the weights vector, and multiplication/exponentiation are element-wise operations
-	TArray<float> CostVector;
-};
-
 } // namespace UE::PoseSearch
 
 
@@ -1250,7 +1251,6 @@ public:
 	void BuildQuery(UE::PoseSearch::FSearchContext& SearchContext, FPoseSearchFeatureVectorBuilder& OutQuery) const;
 
 	FPoseSearchCost ComparePoses(UE::PoseSearch::FSearchContext& SearchContext, int32 PoseIdx, UE::PoseSearch::EPoseComparisonFlags PoseComparisonFlags, int32 GroupIdx, const TArrayView<const float>& QueryValues) const;
-	FPoseSearchCost ComparePoses(UE::PoseSearch::FSearchContext& SearchContext, int32 PoseIdx, UE::PoseSearch::EPoseComparisonFlags PoseComparisonFlags, const TArrayView<const float>& QueryValues, UE::PoseSearch::FPoseCostDetails& OutPoseCostDetails) const;
 
 protected:
 	UE::PoseSearch::FSearchResult SearchPCAKDTree(UE::PoseSearch::FSearchContext& SearchContext) const;
@@ -1403,6 +1403,7 @@ struct POSESEARCH_API FSearchContext
 	const FBoneContainer* BoneContainer = nullptr;
 	const FGameplayTagContainer* ActiveTagsContainer = nullptr;
 	float PoseJumpThresholdTime = 0.f;
+	bool bIsTracing = false;
 	bool bForceInterrupt = false;
 	// can the continuing pose advance? (if not we skip evaluating it)
 	bool bCanAdvance = true;
@@ -1456,11 +1457,11 @@ private:
 public:
 	struct FPoseCandidate
 	{
-		float Cost = 0.f;
+		FPoseSearchCost Cost;
 		int32 PoseIdx = 0;
 		const UPoseSearchDatabase* Database = nullptr;
 
-		bool operator<(const FPoseCandidate& Other) const { return Cost > Other.Cost; }
+		bool operator<(const FPoseCandidate& Other) const { return Cost < Other.Cost; }
 		bool operator==(const FSearchResult& SearchResult) const { return (PoseIdx == SearchResult.PoseIdx) && (Database == SearchResult.Database.Get()); }
 	};
 
@@ -1471,7 +1472,7 @@ public:
 
 		int32 MaxPoseCandidates = 100;
 
-		void Add(float Cost, int32 PoseIdx, const UPoseSearchDatabase* Database)
+		void Add(const FPoseSearchCost& Cost, int32 PoseIdx, const UPoseSearchDatabase* Database)
 		{
 			if (Num() < MaxPoseCandidates || Cost < HeapTop().Cost)
 			{
