PoseSearch: Fix candidate sorting in debugger

diff --git a/Engine/Plugins/Experimental/Animation/PoseSearch/Source/Runtime/Public/PoseSearch/PoseSearch.h b/Engine/Plugins/Experimental/Animation/PoseSearch/Source/Runtime/Public/PoseSearch/PoseSearch.h
--- a/Engine/Plugins/Experimental/Animation/PoseSearch/Source/Runtime/Public/PoseSearch/PoseSearch.h
+++ b/Engine/Plugins/Experimental/Animation/PoseSearch/Source/Runtime/Public/PoseSearch/PoseSearch.h
@@ -1461,7 +1461,7 @@ public:
 		int32 PoseIdx = 0;
 		const UPoseSearchDatabase* Database = nullptr;
 
-		bool operator<(const FPoseCandidate& Other) const { return Cost < Other.Cost; }
+		bool operator<(const FPoseCandidate& Other) const { return Other.Cost < Cost; } // Reverse compare because BestCandidates is a max heap
 		bool operator==(const FSearchResult& SearchResult) const { return (PoseIdx == SearchResult.PoseIdx) && (Database == SearchResult.Database.Get()); }
 	};
 
