motion matching - blend behaviour support for blendstack node

diff --git a/Engine/Plugins/Experimental/Animation/PoseSearch/Source/Runtime/Private/AnimNode_BlendStack.cpp b/Engine/Plugins/Experimental/Animation/PoseSearch/Source/Runtime/Private/AnimNode_BlendStack.cpp
--- a/Engine/Plugins/Experimental/Animation/PoseSearch/Source/Runtime/Private/AnimNode_BlendStack.cpp
+++ b/Engine/Plugins/Experimental/Animation/PoseSearch/Source/Runtime/Private/AnimNode_BlendStack.cpp
@@ -2,6 +2,7 @@
 
 #include "PoseSearch/AnimNode_BlendStack.h"
 
+#include "Algo/MaxElement.h"
 #include "Animation/AnimSequence.h"
 #include "Animation/BlendSpace.h"
 
@@ -9,11 +10,24 @@
 
 /////////////////////////////////////////////////////
 // FPoseSearchAnimPlayer
-void FPoseSearchAnimPlayer::Initialize(ESearchIndexAssetType InAssetType, UAnimationAsset* AnimationAsset, float AccumulatedTime, bool bLoop, bool bMirrored, UMirrorDataTable* MirrorDataTable, float BlendTime, FVector BlendParameters)
+void FPoseSearchAnimPlayer::Initialize(ESearchIndexAssetType InAssetType, UAnimationAsset* AnimationAsset, float AccumulatedTime, bool bLoop, bool bMirrored, UMirrorDataTable* MirrorDataTable, float BlendTime, const UBlendProfile* BlendProfile, FVector BlendParameters)
 {
 	check(AnimationAsset);
 	check(MirrorDataTable);
 
+	if (BlendProfile != nullptr)
+	{
+		const USkeleton* SkeletonAsset = BlendProfile->OwningSkeleton;
+		check(SkeletonAsset);
+
+		const FReferenceSkeleton& RefSkeleton = SkeletonAsset->GetReferenceSkeleton();
+		const int32 NumSkeletonBones = RefSkeleton.GetNum();
+		TotalBlendInTimePerBone.Init(BlendTime, NumSkeletonBones);
+
+		BlendProfile->FillSkeletonBoneDurationsArray(TotalBlendInTimePerBone, BlendTime);
+		BlendTime = *Algo::MaxElement(TotalBlendInTimePerBone);
+	}
+
 	AssetType = InAssetType;
 	TotalBlendInTime = BlendTime;
 	CurrentBlendInTime = 0.f;
@@ -108,6 +122,29 @@ float FPoseSearchAnimPlayer::GetBlendInPercentage() const
 	return FMath::Clamp(CurrentBlendInTime / TotalBlendInTime, 0.f, 1.f);
 }
 
+bool FPoseSearchAnimPlayer::GetBlendInWeights(TArray<float>& Weights) const
+{
+	const int32 NumBones = TotalBlendInTimePerBone.Num();
+	if (NumBones > 0)
+	{
+		Weights.SetNumUninitialized(NumBones);
+		for (int32 BoneIdx = 0; BoneIdx < NumBones; ++BoneIdx)
+		{
+			const float TotalBlendInTimeBoneIdx = TotalBlendInTimePerBone[BoneIdx];
+			if (FMath::IsNearlyZero(TotalBlendInTimeBoneIdx))
+			{
+				Weights[BoneIdx] = 1.f;
+			}
+			else
+			{
+				Weights[BoneIdx] = FMath::Clamp(CurrentBlendInTime / TotalBlendInTimeBoneIdx, 0.f, 1.f);
+			}
+		}
+		return true;
+	}
+	return false;
+}
+
 void FPoseSearchAnimPlayer::SetBlendWeight(float InBlendWeight)
 {
 	BlendWeight = InBlendWeight;
@@ -132,35 +169,34 @@ void FAnimNode_BlendStack::Evaluate_AnyThread(FPoseContext& Output)
 	}
 	else
 	{
-		TArray<FPoseContext, TMemStackAllocator<>> PoseContexts;
-		PoseContexts.Reserve(BlendStackSize);
-		for (int32 i = 0; i < BlendStackSize; ++i)
-		{
-			PoseContexts.Add(Output);
-			AnimPlayers[i].Evaluate_AnyThread(PoseContexts[i]);
-		}
+		AnimPlayers[BlendStackSize - 1].Evaluate_AnyThread(Output);
 
-		// @todo: optimize copies and allocations: is there a FAnimationRuntime::BlendPosesTogether(PoseContexts, Weights, AnimationPoseData)?
-		TArray<FCompactPose, TMemStackAllocator<>> Poses;
-		TArray<FBlendedCurve, TMemStackAllocator<>> Curves;
-		TArray<UE::Anim::FStackAttributeContainer, TMemStackAllocator<>> Attributes;
-		TArray<float, TMemStackAllocator<>> Weights;
+		FPoseContext EvaluationPoseContext(Output);
+		FPoseContext BlendedPoseContext(Output); // @todo: this should not be necessary (but FBaseBlendedCurve::InitFrom complains about "ensure(&InCurveToInitFrom != this)"): optimize it away!
+		FAnimationPoseData BlendedAnimationPoseData(BlendedPoseContext);
 
-		Poses.AddDefaulted(BlendStackSize);
-		Curves.AddDefaulted(BlendStackSize);
-		Attributes.AddDefaulted(BlendStackSize);
-		Weights.AddDefaulted(BlendStackSize);
+		const USkeleton* SkeletonAsset = Output.AnimInstanceProxy->GetRequiredBones().GetSkeletonAsset();
+		check(SkeletonAsset);
 
-		for (int32 i = 0; i < BlendStackSize; ++i)
+		const FReferenceSkeleton& RefSkeleton = SkeletonAsset->GetReferenceSkeleton();
+		const int32 NumSkeletonBones = RefSkeleton.GetNum();
+		TArray<float> Weights;
+		for (int32 i = BlendStackSize - 2; i >= 0; --i)
 		{
-			Poses[i] = PoseContexts[i].Pose;
-			Curves[i] = PoseContexts[i].Curve;
-			Attributes[i] = PoseContexts[i].CustomAttributes;
-			Weights[i] = AnimPlayers[i].GetBlendWeight();
+			AnimPlayers[i].Evaluate_AnyThread(EvaluationPoseContext);
+			
+			if (AnimPlayers[i].GetBlendInWeights(Weights))
+			{
+				// @todo: have BlendTwoPosesTogetherPerBone using a TArrayView for the Weights to avoid allocations
+				FAnimationRuntime::BlendTwoPosesTogetherPerBone(FAnimationPoseData(Output), FAnimationPoseData(EvaluationPoseContext), Weights, BlendedAnimationPoseData);
+			}
+			else
+			{
+				const float Weight = 1.f - AnimPlayers[i].GetBlendInPercentage();
+				FAnimationRuntime::BlendTwoPosesTogether(FAnimationPoseData(Output), FAnimationPoseData(EvaluationPoseContext), Weight, BlendedAnimationPoseData);
+			}
+			Output = BlendedPoseContext; // @todo: this should not be necessary either: optimize it away!
 		}
-
-		FAnimationPoseData AnimationPoseData(Output);
-		FAnimationRuntime::BlendPosesTogether(Poses, Curves, Attributes, Weights, AnimationPoseData);
 	}
 }
 
@@ -185,10 +221,10 @@ float FAnimNode_BlendStack::GetAccumulatedTime() const
 	return AnimPlayers.IsEmpty() ? 0.f : AnimPlayers.First().GetAccumulatedTime();
 }
 
-void FAnimNode_BlendStack::BlendTo(ESearchIndexAssetType AssetType, UAnimationAsset* AnimationAsset, float AccumulatedTime, bool bLoop, bool bMirrored, UMirrorDataTable* MirrorDataTable, int32 MaxActiveBlends, float BlendTime, FVector BlendParameters)
+void FAnimNode_BlendStack::BlendTo(ESearchIndexAssetType AssetType, UAnimationAsset* AnimationAsset, float AccumulatedTime, bool bLoop, bool bMirrored, UMirrorDataTable* MirrorDataTable, int32 MaxActiveBlends, float BlendTime, const UBlendProfile* BlendProfile, FVector BlendParameters)
 {
 	AnimPlayers.PushFirst(FPoseSearchAnimPlayer());
-	AnimPlayers.First().Initialize(AssetType, AnimationAsset, AccumulatedTime, bLoop, bMirrored, MirrorDataTable, BlendTime, BlendParameters);
+	AnimPlayers.First().Initialize(AssetType, AnimationAsset, AccumulatedTime, bLoop, bMirrored, MirrorDataTable, BlendTime, BlendProfile, BlendParameters);
 
 	CalculateWeights();
 	PruneBlendStack(MaxActiveBlends);
diff --git a/Engine/Plugins/Experimental/Animation/PoseSearch/Source/Runtime/Private/AnimNode_MotionMatching.cpp b/Engine/Plugins/Experimental/Animation/PoseSearch/Source/Runtime/Private/AnimNode_MotionMatching.cpp
--- a/Engine/Plugins/Experimental/Animation/PoseSearch/Source/Runtime/Private/AnimNode_MotionMatching.cpp
+++ b/Engine/Plugins/Experimental/Animation/PoseSearch/Source/Runtime/Private/AnimNode_MotionMatching.cpp
@@ -100,7 +100,7 @@ void FAnimNode_MotionMatching::UpdateAssetPlayer(const FAnimationUpdateContext&
 		const FPoseSearchDatabaseAnimationAssetBase& AnimationAssetBase = Database->GetAnimationSourceAsset(SearchIndexAsset);
 		BlendStackNode.BlendTo(	SearchIndexAsset->Type, AnimationAssetBase.GetAnimationAsset(), MotionMatchingState.CurrentSearchResult.AssetTime,
 								AnimationAssetBase.IsLooping(), SearchIndexAsset->bMirrored, Database->Schema->MirrorDataTable.Get(),
-								Settings.MaxActiveBlends, Settings.BlendTime, SearchIndexAsset->BlendParameters);
+								Settings.MaxActiveBlends, Settings.BlendTime, Settings.BlendProfile, SearchIndexAsset->BlendParameters);
 	}
 
 	if (!BlendStackNode.AnimPlayers.IsEmpty() && BlendStackNode.AnimPlayers[0].GetAssetType() == ESearchIndexAssetType::Sequence)
diff --git a/Engine/Plugins/Experimental/Animation/PoseSearch/Source/Runtime/Public/PoseSearch/AnimNode_BlendStack.h b/Engine/Plugins/Experimental/Animation/PoseSearch/Source/Runtime/Public/PoseSearch/AnimNode_BlendStack.h
--- a/Engine/Plugins/Experimental/Animation/PoseSearch/Source/Runtime/Public/PoseSearch/AnimNode_BlendStack.h
+++ b/Engine/Plugins/Experimental/Animation/PoseSearch/Source/Runtime/Public/PoseSearch/AnimNode_BlendStack.h
@@ -16,13 +16,14 @@ struct FPoseSearchAnimPlayer
 {
 	GENERATED_BODY()
 	
-	void Initialize(ESearchIndexAssetType InAssetType, UAnimationAsset* AnimationAsset, float AccumulatedTime, bool bLoop, bool bMirrored, UMirrorDataTable* MirrorDataTable, float BlendTime = 0.f, FVector BlendParameters = FVector::Zero());
+	void Initialize(ESearchIndexAssetType InAssetType, UAnimationAsset* AnimationAsset, float AccumulatedTime, bool bLoop, bool bMirrored, UMirrorDataTable* MirrorDataTable, float BlendTime = 0.f, const UBlendProfile* BlendProfile = nullptr, FVector BlendParameters = FVector::Zero());
 	void Evaluate_AnyThread(FPoseContext& Output);
 	void Update_AnyThread(const FAnimationUpdateContext& Context);
 	float GetAccumulatedTime() const;
 	float GetBlendWeight() const { return BlendWeight;  }
 	float GetBlendInPercentage() const;
 	void SetBlendWeight(float InBlendWeight);
+	bool GetBlendInWeights(TArray<float>& Weights) const;
 	ESearchIndexAssetType GetAssetType() const { return AssetType; }
 
 	// @todo: used only for DynamicPlayRateAdjustment. Remove once the functionality is integrated with the BlendStackNode
@@ -44,6 +45,8 @@ protected:
 
 	ESearchIndexAssetType AssetType = ESearchIndexAssetType::Sequence;
 
+	TCustomBoneIndexArray<float, FSkeletonPoseBoneIndex> TotalBlendInTimePerBone;
+
 	float TotalBlendInTime = 0.f;
 	float CurrentBlendInTime = 0.f;
 	float BlendWeight = 1.f;
@@ -60,7 +63,7 @@ struct POSESEARCH_API FAnimNode_BlendStack : public FAnimNode_AssetPlayerBase
 	virtual void Evaluate_AnyThread(FPoseContext& Output) override;
 	// End of FAnimNode_Base interface
 
-	void BlendTo(ESearchIndexAssetType AssetType, UAnimationAsset* AnimationAsset, float AccumulatedTime, bool bLoop, bool bMirrored, UMirrorDataTable* MirrorDataTable, int32 MaxActiveBlends = 3, float BlendTime = 0.3f, FVector BlendParameters = FVector::Zero());
+	void BlendTo(ESearchIndexAssetType AssetType, UAnimationAsset* AnimationAsset, float AccumulatedTime, bool bLoop, bool bMirrored, UMirrorDataTable* MirrorDataTable, int32 MaxActiveBlends = 3, float BlendTime = 0.3f, const UBlendProfile* BlendProfile = nullptr, FVector BlendParameters = FVector::Zero());
 	void CalculateWeights();
 	void PruneBlendStack(int32 MaxActiveBlends);
 
diff --git a/Engine/Plugins/Experimental/Animation/PoseSearch/Source/Runtime/Public/PoseSearch/PoseSearchLibrary.h b/Engine/Plugins/Experimental/Animation/PoseSearch/Source/Runtime/Public/PoseSearch/PoseSearchLibrary.h
--- a/Engine/Plugins/Experimental/Animation/PoseSearch/Source/Runtime/Public/PoseSearch/PoseSearchLibrary.h
+++ b/Engine/Plugins/Experimental/Animation/PoseSearch/Source/Runtime/Public/PoseSearch/PoseSearchLibrary.h
@@ -35,6 +35,9 @@ struct POSESEARCH_API FMotionMatchingSettings
 	UPROPERTY(EditAnywhere, BlueprintReadWrite, Category=Settings, meta=(ClampMin="0"))
 	int32 MaxActiveBlends = 0;
 
+	UPROPERTY(EditAnywhere, BlueprintReadWrite, Category=Settings, meta=(UseAsBlendProfile = true))
+	TObjectPtr<UBlendProfile> BlendProfile;
+
 	// If the pose jump requires a mirroring change and this value is greater than 0, it will be used instead of BlendTime
 	UPROPERTY(EditAnywhere, BlueprintReadWrite, Category=Settings, meta=(ClampMin="0", DislayAfter="BlendTime"))
 	float MirrorChangeBlendTime = 0.0f;
