motion matching - embedding mindeviation scaling into weights

diff --git a/Engine/Plugins/Experimental/Animation/PoseSearch/Source/Editor/Private/PoseSearchDebugger.cpp b/Engine/Plugins/Experimental/Animation/PoseSearch/Source/Editor/Private/PoseSearchDebugger.cpp
--- a/Engine/Plugins/Experimental/Animation/PoseSearch/Source/Editor/Private/PoseSearchDebugger.cpp
+++ b/Engine/Plugins/Experimental/Animation/PoseSearch/Source/Editor/Private/PoseSearchDebugger.cpp
@@ -1465,7 +1465,6 @@ void SDebuggerDetailsView::UpdateReflection(const FTraceMotionMatchingStateMessa
 	if (CurrentSearchIndex)
 	{
 		Reflection->ActivePoseVector = CurrentSearchIndex->GetPoseValues(CurrentDbPoseIdx);
-		CurrentSearchIndex->InverseNormalize(Reflection->ActivePoseVector);
 	}
 
 	auto DebuggerView = ParentDebuggerViewPtr.Pin();
@@ -1481,7 +1480,6 @@ void SDebuggerDetailsView::UpdateReflection(const FTraceMotionMatchingStateMessa
 				if (SelectedSearchIndex)
 				{
 					Reflection->SelectedPoseVector = SelectedSearchIndex->GetPoseValues(Selected->PoseIdx);
-					SelectedSearchIndex->InverseNormalize(Reflection->SelectedPoseVector);
 				}
 			}
 
diff --git a/Engine/Plugins/Experimental/Animation/PoseSearch/Source/Runtime/Private/PoseSearch.cpp b/Engine/Plugins/Experimental/Animation/PoseSearch/Source/Runtime/Private/PoseSearch.cpp
--- a/Engine/Plugins/Experimental/Animation/PoseSearch/Source/Runtime/Private/PoseSearch.cpp
+++ b/Engine/Plugins/Experimental/Animation/PoseSearch/Source/Runtime/Private/PoseSearch.cpp
@@ -323,12 +323,6 @@ void UPoseSearchSchema::Finalize()
 
 	BoneReferences = MoveTemp(Initializer.BoneReferences);
 
-	EffectiveDataPreprocessor = DataPreprocessor;
-	if (EffectiveDataPreprocessor == EPoseSearchDataPreprocessor::Automatic)
-	{
-		EffectiveDataPreprocessor = EPoseSearchDataPreprocessor::Normalize;
-	}
-
 	ResolveBoneReferences();
 }
 
@@ -691,60 +685,6 @@ void FPoseSearchIndex::Reset()
 	Schema = nullptr;
 }
 
-void FPoseSearchIndex::Normalize(TArrayView<float> InOutPoseVector) const
-{
-	using namespace Eigen;
-
-	auto TransformationMtx = Map<const Matrix<float, Dynamic, Dynamic, ColMajor>>
-	(
-		PreprocessInfo.TransformationMatrix.GetData(),
-		PreprocessInfo.NumDimensions,
-		PreprocessInfo.NumDimensions
-	);
-	auto SampleMean = Map<const Matrix<float, Dynamic, 1, ColMajor>>
-	(
-		PreprocessInfo.SampleMean.GetData(),
-		PreprocessInfo.NumDimensions
-	);
-
-	checkSlow(InOutPoseVector.Num() == PreprocessInfo.NumDimensions);
-
-	auto PoseVector = Map<Matrix<float, Dynamic, 1, ColMajor>>
-	(
-		InOutPoseVector.GetData(),
-		InOutPoseVector.Num()
-	);
-
-	PoseVector = TransformationMtx * (PoseVector - SampleMean);
-}
-
-void FPoseSearchIndex::InverseNormalize(TArrayView<float> InOutNormalizedPoseVector) const
-{
-	using namespace Eigen;
-
-	auto InverseTransformationMtx = Map<const Matrix<float, Dynamic, Dynamic, ColMajor>>
-	(
-		PreprocessInfo.InverseTransformationMatrix.GetData(),
-		PreprocessInfo.NumDimensions,
-		PreprocessInfo.NumDimensions
-	);
-	auto SampleMean = Map<const Matrix<float, Dynamic, 1, ColMajor>>
-	(
-		PreprocessInfo.SampleMean.GetData(),
-		PreprocessInfo.NumDimensions
-	);
-
-	checkSlow(InOutNormalizedPoseVector.Num() == PreprocessInfo.NumDimensions);
-
-	auto NormalizedPoseVector = Map<Matrix<float, Dynamic, 1, ColMajor>>
-	(
-		InOutNormalizedPoseVector.GetData(),
-		InOutNormalizedPoseVector.Num()
-	);
-
-	NormalizedPoseVector = (InverseTransformationMtx * NormalizedPoseVector) + SampleMean;
-}
-
 float FPoseSearchIndex::ComputeMirrorMismatchAddend(int32 PoseIdx, UE::PoseSearch::FSearchContext& SearchContext) const
 {
 	if (SearchContext.QueryMirrorRequest != EPoseSearchBooleanRequest::Indifferent)
@@ -822,10 +762,9 @@ UE::PoseSearch::FSearchResult UPoseSearchSequenceMetaData::Search(UE::PoseSearch
 	}
 
 	Schema->BuildQuery(SearchContext, Result.ComposedQuery);
-	Result.ComposedQuery.Normalize(SearchIndex);
-	TArrayView<const float> NormalizedQueryValues = Result.ComposedQuery.GetNormalizedValues();
+	TArrayView<const float> QueryValues = Result.ComposedQuery.GetValues();
 
-	if (!ensure(NormalizedQueryValues.Num() == SearchIndex.Schema->SchemaCardinality))
+	if (!ensure(QueryValues.Num() == SearchIndex.Schema->SchemaCardinality))
 	{
 		return Result;
 	}
@@ -844,7 +783,7 @@ UE::PoseSearch::FSearchResult UPoseSearchSequenceMetaData::Search(UE::PoseSearch
 				continue;
 			}
 
-			FPoseSearchCost PoseCost = ComparePoses(PoseIdx, EPoseComparisonFlags::ContinuingPose, NormalizedQueryValues);
+			FPoseSearchCost PoseCost = ComparePoses(PoseIdx, EPoseComparisonFlags::ContinuingPose, QueryValues);
 
 			if (PoseCost < BestPoseCost)
 			{
@@ -863,7 +802,7 @@ UE::PoseSearch::FSearchResult UPoseSearchSequenceMetaData::Search(UE::PoseSearch
 	DrawFeatureVector(SearchContext.DebugDrawParams, Result.PoseIdx);
 	
 	EnumAddFlags(SearchContext.DebugDrawParams.Flags, EDebugDrawFlags::DrawQuery);
-	DrawFeatureVector(SearchContext.DebugDrawParams, NormalizedQueryValues);
+	DrawFeatureVector(SearchContext.DebugDrawParams, QueryValues);
 
 	return Result;
 }
@@ -2054,7 +1993,7 @@ UE::PoseSearch::FSearchResult UPoseSearchDatabase::SearchPCAKDTree(UE::PoseSearc
 
 	SearchContext.GetOrBuildQuery(this, Result.ComposedQuery);
 
-	TArrayView<const float> NormalizedQueryValues = Result.ComposedQuery.GetNormalizedValues();
+	TArrayView<const float> QueryValues = Result.ComposedQuery.GetValues();
 
 	const bool IsCurrentResultFromThisDatabase = SearchContext.IsCurrentResultFromDatabase(this);
 
@@ -2066,7 +2005,7 @@ UE::PoseSearch::FSearchResult UPoseSearchDatabase::SearchPCAKDTree(UE::PoseSearc
 			SearchContext.CurrentResult.PoseIdx,
 			EPoseComparisonFlags::ContinuingPose,
 			SearchIndex->FindAssetForPose(SearchContext.CurrentResult.PoseIdx)->SourceGroupIdx,
-			NormalizedQueryValues);
+			QueryValues);
 
 		if (GetSkipSearchIfPossible())
 		{
@@ -2086,14 +2025,14 @@ UE::PoseSearch::FSearchResult UPoseSearchDatabase::SearchPCAKDTree(UE::PoseSearc
 			const RowMajorVectorMapConst MapWeights(GroupSearchIndex.Weights.GetData(), 1, NumDimensions);
 			FKDTree::KNNResultSet ResultSet(ClampedKDTreeQueryNumNeighbors, ResultIndexes, ResultDistanceSqr, NonSelectableIdx);
 
-			check(NormalizedQueryValues.Num() == NumDimensions);
+			check(QueryValues.Num() == NumDimensions);
 
 			const RowMajorVectorMapConst Mean(GroupSearchIndex.Mean.GetData(), 1, NumDimensions);
 			const ColMajorMatrixMapConst PCAProjectionMatrix(GroupSearchIndex.PCAProjectionMatrix.GetData(), NumDimensions, ClampedNumberOfPrincipalComponents);
 
 			// transforming query values into PCA space to query the KDTree
-			const RowMajorVectorMapConst QueryValues(NormalizedQueryValues.GetData(), 1, NumDimensions);
-			WeightedQueryValues = QueryValues.array() * MapWeights.array();
+			const RowMajorVectorMapConst QueryValuesMap(QueryValues.GetData(), 1, NumDimensions);
+			WeightedQueryValues = QueryValuesMap.array() * MapWeights.array();
 			CenteredQueryValues.noalias() = WeightedQueryValues - Mean;
 			ProjectedQueryValues.noalias() = CenteredQueryValues * PCAProjectionMatrix;
 
@@ -2115,7 +2054,7 @@ UE::PoseSearch::FSearchResult UPoseSearchDatabase::SearchPCAKDTree(UE::PoseSearc
 					PoseIdx,
 					EPoseComparisonFlags::None,
 					GroupSearchIndex.GroupIndex,
-					NormalizedQueryValues);
+					QueryValues);
 
 				if (PoseCost < BestPoseCost)
 				{
@@ -2170,7 +2109,7 @@ UE::PoseSearch::FSearchResult UPoseSearchDatabase::SearchBruteForce(UE::PoseSear
 	check(SearchIndex);
 
 	SearchContext.GetOrBuildQuery(this, Result.ComposedQuery);
-	TArrayView<const float> NormalizedQueryValues = Result.ComposedQuery.GetNormalizedValues();
+	TArrayView<const float> QueryValues = Result.ComposedQuery.GetValues();
 
 	const bool IsCurrentResultFromThisDatabase = SearchContext.IsCurrentResultFromDatabase(this);
 	if (IsCurrentResultFromThisDatabase)
@@ -2183,7 +2122,7 @@ UE::PoseSearch::FSearchResult UPoseSearchDatabase::SearchBruteForce(UE::PoseSear
 				SearchContext.CurrentResult.PoseIdx,
 				EPoseComparisonFlags::ContinuingPose,
 				SearchIndex->FindAssetForPose(SearchContext.CurrentResult.PoseIdx)->SourceGroupIdx,
-				NormalizedQueryValues);
+				QueryValues);
 
 			if (GetSkipSearchIfPossible())
 			{
@@ -2235,7 +2174,7 @@ UE::PoseSearch::FSearchResult UPoseSearchDatabase::SearchBruteForce(UE::PoseSear
 					PoseIdx, 
 					EPoseComparisonFlags::None,
 					Asset.SourceGroupIdx,
-					NormalizedQueryValues);
+					QueryValues);
 
 				if (PoseCost < BestPoseCost)
 				{
@@ -2288,7 +2227,6 @@ void UPoseSearchDatabase::BuildQuery(UE::PoseSearch::FSearchContext& SearchConte
 {
 	check(Schema && Schema->IsValid());
 	Schema->BuildQuery(SearchContext, OutQuery);
-	OutQuery.Normalize(*GetSearchIndex());
 }
 
 UE::PoseSearch::FSearchResult UPoseSearchDatabaseSet::Search(UE::PoseSearch::FSearchContext& SearchContext) const
@@ -2314,7 +2252,7 @@ UE::PoseSearch::FSearchResult UPoseSearchDatabaseSet::Search(UE::PoseSearch::FSe
 
 		SearchContext.GetOrBuildQuery(Database, Result.ComposedQuery);
 
-		TArrayView<const float> NormalizedQueryValues = Result.ComposedQuery.GetNormalizedValues();
+		TArrayView<const float> QueryValues = Result.ComposedQuery.GetValues();
 
 		const FPoseSearchIndexAsset* PoseSearchIndexAsset = PoseSearchIndex->FindAssetForPose(SearchContext.CurrentResult.PoseIdx);
 		check(PoseSearchIndexAsset);
@@ -2324,7 +2262,7 @@ UE::PoseSearch::FSearchResult UPoseSearchDatabaseSet::Search(UE::PoseSearch::FSe
 				SearchContext.CurrentResult.PoseIdx,
 				EPoseComparisonFlags::ContinuingPose,
 				PoseSearchIndexAsset->SourceGroupIdx,
-				NormalizedQueryValues);
+				QueryValues);
 
 		if (Database->GetSkipSearchIfPossible())
 		{
@@ -2400,26 +2338,18 @@ void FPoseSearchFeatureVectorBuilder::Reset()
 {
 	Schema = nullptr;
 	Values.Reset(0);
-	ValuesNormalized.Reset(0);
 }
 
 void FPoseSearchFeatureVectorBuilder::ResetFeatures()
 {
 	Values.Reset(0);
 	Values.SetNumZeroed(Schema->SchemaCardinality);
-	ValuesNormalized.Reset(0);
-	ValuesNormalized.SetNumZeroed(Schema->SchemaCardinality);
 }
 
 void FPoseSearchFeatureVectorBuilder::CopyFromSearchIndex(const FPoseSearchIndex& SearchIndex, int32 PoseIdx)
 {
 	check(Schema == SearchIndex.Schema);
-
-	TArrayView<const float> FeatureVector = SearchIndex.GetPoseValues(PoseIdx);
-
-	ValuesNormalized = FeatureVector;
-	Values = FeatureVector;
-	SearchIndex.InverseNormalize(Values);
+	Values = SearchIndex.GetPoseValues(PoseIdx);
 }
 
 bool FPoseSearchFeatureVectorBuilder::IsInitialized() const
@@ -2437,12 +2367,6 @@ bool FPoseSearchFeatureVectorBuilder::IsCompatible(const FPoseSearchFeatureVecto
 	return IsInitialized() && (Schema == OtherBuilder.Schema);
 }
 
-void FPoseSearchFeatureVectorBuilder::Normalize(const FPoseSearchIndex& ForSearchIndex)
-{
-	ValuesNormalized = Values;
-	ForSearchIndex.Normalize(ValuesNormalized);
-}
-
 namespace UE { namespace PoseSearch
 {
 
@@ -4514,17 +4438,7 @@ void DrawFeatureVector(const FDebugDrawParams& DrawParams, int32 PoseIdx)
 {
 	if (PoseIdx != INDEX_NONE && DrawParams.CanDraw())
 	{
-		const FPoseSearchIndex* SearchIndex = DrawParams.GetSearchIndex();
-
-		// PreprocessInfo happens to be invalid when updating the database
-		if (SearchIndex->PreprocessInfo.IsValid())
-		{
-			// @todo: embed normalization into the weights to remove the InverseNormalize step and avoid the copy here, using the TArrayView instead of the TArray for the PoseVector
-			TArray<float> PoseVector;
-			PoseVector = SearchIndex->GetPoseValues(PoseIdx);
-			SearchIndex->InverseNormalize(PoseVector);
-			DrawFeatureVector(DrawParams, PoseVector);
-		}
+		DrawFeatureVector(DrawParams, DrawParams.GetSearchIndex()->GetPoseValues(PoseIdx));
 	}
 }
 
@@ -4541,68 +4455,7 @@ void DrawSearchIndex(const FDebugDrawParams& DrawParams)
 	}
 }
 
-static void PreprocessSearchIndexNone(FPoseSearchIndex* SearchIndex)
-{
-	// This function leaves the data unmodified and simply outputs the transformation
-	// and inverse transformation matrices as the identity matrix and the sample mean
-	// as the zero vector.
-
-	using namespace Eigen;
-
-	check(SearchIndex->IsValid() && !SearchIndex->IsEmpty());
-
-	FPoseSearchIndexPreprocessInfo& Info = SearchIndex->PreprocessInfo;
-	Info.Reset();
-
-	const int32 NumPoses = SearchIndex->NumPoses;
-	const int32 NumDimensions = SearchIndex->Schema->SchemaCardinality;
-
-	Info.NumDimensions = NumDimensions;
-	Info.TransformationMatrix.SetNumZeroed(NumDimensions * NumPoses);
-	Info.InverseTransformationMatrix.SetNumZeroed(NumDimensions * NumPoses);
-	Info.SampleMean.SetNumZeroed(NumDimensions);
-
-	// Map output transformation matrix
-	auto TransformMap = Map<Matrix<float, Dynamic, Dynamic, ColMajor>>(
-		Info.TransformationMatrix.GetData(),
-		NumDimensions, NumPoses
-	);
-
-	// Map output inverse transformation matrix
-	auto InverseTransformMap = Map<Matrix<float, Dynamic, Dynamic, ColMajor>>(
-		Info.InverseTransformationMatrix.GetData(),
-		NumDimensions, NumPoses
-	);
-
-	// Map output sample mean vector
-	auto SampleMeanMap = Map<VectorXf>(Info.SampleMean.GetData(), NumDimensions);
-
-	// Write the transformation matrices and sample mean
-	TransformMap = MatrixXf::Identity(NumDimensions, NumPoses);
-	InverseTransformMap = MatrixXf::Identity(NumDimensions, NumPoses);
-	SampleMeanMap = VectorXf::Zero(NumDimensions);
-}
-
-static inline Eigen::VectorXd ComputeChannelsMeanDeviations(const Eigen::MatrixXd& CenteredPoseMatrix, const UPoseSearchSchema* Schema)
-{
-	using namespace Eigen;
-
-	const int32 NumPoses = CenteredPoseMatrix.cols();
-	const int32 NumDimensions = CenteredPoseMatrix.rows();
-
-	VectorXd MeanDeviations(NumDimensions);
-	MeanDeviations.setConstant(1.0);
-
-	for (int ChannelIdx = 0; ChannelIdx != Schema->Channels.Num(); ++ChannelIdx)
-	{
-		const UPoseSearchFeatureChannel* Channel = Schema->Channels[ChannelIdx].Get();
-		Channel->ComputeMeanDeviations(CenteredPoseMatrix, MeanDeviations);
-	}
-
-	return MeanDeviations;
-}
-
-static void PreprocessSearchIndexNormalize(FPoseSearchIndex* SearchIndex)
+static Eigen::VectorXd ComputeChannelsMeanDeviations(const FPoseSearchIndex* SearchIndex)
 {
 	// This function performs a modified z-score normalization where features are normalized
 	// by mean absolute deviation rather than standard deviation. Both methods are preferable
@@ -4613,36 +4466,6 @@ static void PreprocessSearchIndexNormalize(FPoseSearchIndex* SearchIndex)
 	// exponentially rather than additively and square rooting the sum of squares does not 
 	// remove that bias. [1]
 	//
-	// The pose matrix is transformed in place and the transformation matrix, its inverse,
-	// and data mean vector are computed and stored along with it.
-	//
-	// N:	number of dimensions for input column vectors
-	// P:	number of input column vectors
-	// X:	NxP input matrix
-	// x_p:	pth column vector of input matrix
-	// u:   mean column vector of X
-	//
-	// S:	mean absolute deviations of X, as diagonal NxN matrix with average distances replicated for each feature's axes
-	// s_n:	nth deviation
-	//
-	// Normalization by mean absolute deviation algorithm:
-	//
-	// 1) mean-center X
-	//    x_p := x_p - u
-	// 2) rescale X by inverse mean absolute deviation
-	//    x_p := x_p * s_n^(-1)
-	// 
-	// Let S^(-1) be the inverse of S where the nth diagonal element is s_n^(-1)
-	// then step 2 can be expressed as matrix multiplication:
-	// X := S^(-1) * X
-	//
-	// By persisting the mean vector u and linear transform S, we can bring an input vector q
-	// into the same space as the mean centered and scaled data matrix X:
-	// q := S^(-1) * (q - u)
-	//
-	// This operation is invertible, a normalized data vector x can be unscaled via:
-	// x := (S * x) + u
-	//
 	// References:
 	// [1] Gorard, S. (2005), "Revisiting a 90-Year-Old Debate: The Advantages of the Mean Deviation."
 	//     British Journal of Educational Studies, 53: 417-430.
@@ -4651,97 +4474,38 @@ static void PreprocessSearchIndexNormalize(FPoseSearchIndex* SearchIndex)
 
 	check(SearchIndex->IsValid() && !SearchIndex->IsEmpty());
 
-	FPoseSearchIndexPreprocessInfo& Info = SearchIndex->PreprocessInfo;
-	Info.Reset();
-
 	const int32 NumPoses = SearchIndex->NumPoses;
 	const int32 NumDimensions = SearchIndex->Schema->SchemaCardinality;
 
 	// Map input buffer
-	auto PoseMatrixSourceMap = Map<Matrix<float, Dynamic, Dynamic, RowMajor>>(
+	auto PoseMatrixSourceMap = RowMajorMatrixMapConst(
 		SearchIndex->Values.GetData(),
 		NumPoses,		// rows
 		NumDimensions	// cols
 	);
 
+	// @todo: evaluate removing the cast to double
+	
 	// Copy row major float matrix to column major double matrix
 	MatrixXd PoseMatrix = PoseMatrixSourceMap.transpose().cast<double>();
 	checkSlow(PoseMatrix.rows() == NumDimensions);
 	checkSlow(PoseMatrix.cols() == NumPoses);
 
-#if UE_POSE_SEARCH_EIGEN_DEBUG
-	MatrixXd PoseMatrixOriginal = PoseMatrix;
-#endif
-
 	// Mean center
 	VectorXd SampleMean = PoseMatrix.rowwise().mean();
 	PoseMatrix = PoseMatrix.colwise() - SampleMean;
 
-	// Compute per-feature average distances
-	VectorXd MeanDeviations = ComputeChannelsMeanDeviations(PoseMatrix, SearchIndex->Schema);
-
-	// Construct a scaling matrix that uniformly scales each feature by its average distance from the mean
-	MatrixXd ScalingMatrix = MeanDeviations.cwiseInverse().asDiagonal();
-
-	// Construct the inverse scaling matrix
-	MatrixXd InverseScalingMatrix = MeanDeviations.asDiagonal();
-
-	// Rescale data by transforming it with the scaling matrix
-	// Now each feature has an average Euclidean length = 1.
-	PoseMatrix = ScalingMatrix * PoseMatrix;
-
-	// Write normalized data back to source buffer, converting from column data back to row data
-	PoseMatrixSourceMap = PoseMatrix.transpose().cast<float>();
-
-	// Output preprocessing info
-	Info.NumDimensions = NumDimensions;
-	Info.TransformationMatrix.SetNumZeroed(ScalingMatrix.size());
-	Info.InverseTransformationMatrix.SetNumZeroed(InverseScalingMatrix.size());
-	Info.SampleMean.SetNumZeroed(SampleMean.size());
-
-	auto TransformMap = Map<Matrix<float, Dynamic, Dynamic, ColMajor>>(
-		Info.TransformationMatrix.GetData(),
-		ScalingMatrix.rows(), ScalingMatrix.cols()
-	);
-
-	auto InverseTransformMap = Map<Matrix<float, Dynamic, Dynamic, ColMajor>>(
-		Info.InverseTransformationMatrix.GetData(),
-		InverseScalingMatrix.rows(), InverseScalingMatrix.cols()
-	);
-
-	auto SampleMeanMap = Map<VectorXf>(Info.SampleMean.GetData(), SampleMean.size());
-
-	// Output scaling matrix, inverse scaling matrix, and mean vector
-	TransformMap = ScalingMatrix.cast<float>();
-	InverseTransformMap = InverseScalingMatrix.cast<float>();
-	SampleMeanMap = SampleMean.cast<float>();
-
-#if UE_POSE_SEARCH_EIGEN_DEBUG
-	FString MeanDevationsStr = EigenMatrixToString(MeanDeviations);
-	FString PoseMtxOriginalStr = EigenMatrixToString(PoseMatrixOriginal);
-	FString PoseMtxStr = EigenMatrixToString(PoseMatrix);
-	FString TransformationStr = EigenMatrixToString(TransformMap);
-	FString InverseTransformationStr = EigenMatrixToString(InverseTransformMap);
-	FString SampleMeanStr = EigenMatrixToString(SampleMeanMap);
-#endif // UE_POSE_SEARCH_EIGEN_DEBUG
-}
+	// Compute per Channel average distances
+	VectorXd MeanDeviations(NumDimensions);
+	MeanDeviations.setConstant(1.0);
 
-static void PreprocessSearchIndex(FPoseSearchIndex* SearchIndex)
-{
-	switch (SearchIndex->Schema->EffectiveDataPreprocessor)
+	for (int ChannelIdx = 0; ChannelIdx != SearchIndex->Schema->Channels.Num(); ++ChannelIdx)
 	{
-		case EPoseSearchDataPreprocessor::Normalize:
-			PreprocessSearchIndexNormalize(SearchIndex);
-		break;
-
-		case EPoseSearchDataPreprocessor::None:
-			PreprocessSearchIndexNone(SearchIndex);
-		break;
-
-		case EPoseSearchDataPreprocessor::Invalid:
-			checkNoEntry();
-		break;
+		const UPoseSearchFeatureChannel* Channel = SearchIndex->Schema->Channels[ChannelIdx].Get();
+		Channel->ComputeMeanDeviations(PoseMatrix, MeanDeviations);
 	}
+
+	return MeanDeviations;
 }
 
 static void PreprocessGroupSearchIndexWeights(FGroupSearchIndex& GroupSearchIndex, const UPoseSearchDatabase* Database)
@@ -4755,13 +4519,29 @@ static void PreprocessGroupSearchIndexWeights(FGroupSearchIndex& GroupSearchInde
 		Channel->FillWeights(GroupSearchIndex.Weights);
 	}
 
-	// normalizing user weights
+	// normalizing weights: the idea behind this step is to be able to compare poses from databases using different schemas
 	RowMajorVectorMap MapWeights(GroupSearchIndex.Weights.GetData(), 1, NumDimensions);
 	const float WeightsSum = MapWeights.sum();
 	if (!FMath::IsNearlyZero(WeightsSum))
 	{
 		MapWeights *= (1.0f / WeightsSum);
 	}
+
+	// @todo: remove groups to avoid calculating channels min deviations multiple times
+	EPoseSearchDataPreprocessor DataPreprocessor = Database->GetSearchIndex()->Schema->DataPreprocessor;
+
+	if (DataPreprocessor == EPoseSearchDataPreprocessor::Normalize)
+	{
+		Eigen::VectorXd ChannelsMeanDeviations = ComputeChannelsMeanDeviations(Database->GetSearchIndex());
+
+		for (int32 Dimension = 0; Dimension != NumDimensions; ++Dimension)
+		{
+			// the idea here is to premultiply the weights by the inverse of the variance (proportional to the square of the deviation) to have a "weighted Mahalanobis" distance
+			const float ChannelsMeanDeviation = ChannelsMeanDeviations[Dimension];
+			const float ChannelsVariance = ChannelsMeanDeviation * ChannelsMeanDeviation;
+			GroupSearchIndex.Weights[Dimension] /= ChannelsVariance;
+		}
+	}
 }
 
 // it calcualtes Mean and PCAProjectionMatrix
@@ -5000,8 +4780,8 @@ bool BuildIndex(const UAnimSequence* Sequence, UPoseSearchSequenceMetaData* Sequ
 	SequenceMetaData->SearchIndex.Assets.Add(SearchIndexAsset);
 	SequenceMetaData->SearchIndex.PoseMetadata = Indexer.Output.PoseMetadata;
 
-	PreprocessSearchIndex(&SequenceMetaData->SearchIndex);
-
+	// todo: do we need to initialize per group data?
+	// PreprocessGroupSearchIndex(&SequenceMetaData->SearchIndex, Database);
 	return true;
 }
 
@@ -5299,8 +5079,6 @@ bool BuildIndex(UPoseSearchDatabase* Database, FPoseSearchIndex& OutSearchIndex)
 	bSuccess &= DbIndexingContext.IndexAssets();
 	DbIndexingContext.JoinIndex();
 
-	PreprocessSearchIndex(&OutSearchIndex);
-
 	PreprocessGroupSearchIndex(OutSearchIndex, Database);
 	
 	return bSuccess;
diff --git a/Engine/Plugins/Experimental/Animation/PoseSearch/Source/Runtime/Private/PoseSearchDerivedData.cpp b/Engine/Plugins/Experimental/Animation/PoseSearch/Source/Runtime/Private/PoseSearchDerivedData.cpp
--- a/Engine/Plugins/Experimental/Animation/PoseSearch/Source/Runtime/Private/PoseSearchDerivedData.cpp
+++ b/Engine/Plugins/Experimental/Animation/PoseSearch/Source/Runtime/Private/PoseSearchDerivedData.cpp
@@ -221,47 +221,6 @@ namespace UE::PoseSearch
 
 #endif // WITH_EDITOR
 
-	FArchive& operator<<(FArchive& Ar, FPoseSearchIndexPreprocessInfo& Info)
-	{
-		int32 NumTransformationMatrix = 0;
-
-		if (Ar.IsSaving())
-		{
-			NumTransformationMatrix = Info.TransformationMatrix.Num();
-		}
-
-		Ar << Info.NumDimensions;
-		Ar << NumTransformationMatrix;
-
-		if (Ar.IsLoading())
-		{
-			Info.TransformationMatrix.SetNumUninitialized(NumTransformationMatrix);
-			Info.InverseTransformationMatrix.SetNumUninitialized(NumTransformationMatrix);
-			Info.SampleMean.SetNumUninitialized(Info.NumDimensions);
-		}
-
-		if (Info.TransformationMatrix.Num() > 0)
-		{
-			Ar.Serialize(
-				&Info.TransformationMatrix[0],
-				Info.TransformationMatrix.Num() * Info.TransformationMatrix.GetTypeSize());
-		}
-
-		if (Info.InverseTransformationMatrix.Num() > 0)
-		{
-			Ar.Serialize(
-				&Info.InverseTransformationMatrix[0],
-				Info.InverseTransformationMatrix.Num() * Info.InverseTransformationMatrix.GetTypeSize());
-		}
-
-		if (Info.SampleMean.Num() > 0)
-		{
-			Ar.Serialize(&Info.SampleMean[0], Info.SampleMean.Num() * Info.SampleMean.GetTypeSize());
-		}
-
-		return Ar;
-	}
-
 	FArchive& operator<<(FArchive& Ar, FPoseSearchIndex& Index)
 	{
 		int32 NumValues = 0;
@@ -312,8 +271,6 @@ namespace UE::PoseSearch
 			Ar.Serialize(&Index.Assets[0], Index.Assets.Num() * Index.Assets.GetTypeSize());
 		}
 
-		Ar << Index.PreprocessInfo;
-
 		for (int i = 0; i < Index.Groups.Num(); ++i)
 		{
 			FGroupSearchIndex& GroupSearchIndex = Index.Groups[i];
diff --git a/Engine/Plugins/Experimental/Animation/PoseSearch/Source/Runtime/Private/PoseSearchDerivedData.h b/Engine/Plugins/Experimental/Animation/PoseSearch/Source/Runtime/Private/PoseSearchDerivedData.h
--- a/Engine/Plugins/Experimental/Animation/PoseSearch/Source/Runtime/Private/PoseSearchDerivedData.h
+++ b/Engine/Plugins/Experimental/Animation/PoseSearch/Source/Runtime/Private/PoseSearchDerivedData.h
@@ -96,9 +96,6 @@ namespace UE::PoseSearch
 	};
 #endif // WITH_EDITOR
 
-	// Serialization for FPoseSearchIndexPreprocessInfo.
-	FArchive& operator<<(FArchive& Ar, FPoseSearchIndexPreprocessInfo& Info);
-
 	// Serialization for FPoseSearchIndex.
 	FArchive& operator<<(FArchive& Ar, FPoseSearchIndex& Index);
 }
\ No newline at end of file
diff --git a/Engine/Plugins/Experimental/Animation/PoseSearch/Source/Runtime/Private/PoseSearchLibrary.cpp b/Engine/Plugins/Experimental/Animation/PoseSearch/Source/Runtime/Private/PoseSearchLibrary.cpp
--- a/Engine/Plugins/Experimental/Animation/PoseSearch/Source/Runtime/Private/PoseSearchLibrary.cpp
+++ b/Engine/Plugins/Experimental/Animation/PoseSearch/Source/Runtime/Private/PoseSearchLibrary.cpp
@@ -374,7 +374,6 @@ static void TraceMotionMatchingState(
 	TraceState.SearchableAssetId = FTraceMotionMatchingState::GetIdFromObject(Searchable);
 	TraceState.ElapsedPoseJumpTime = MotionMatchingState.ElapsedPoseJumpTime;
 	TraceState.QueryVector = MotionMatchingState.CurrentSearchResult.ComposedQuery.GetValues();
-	TraceState.QueryVectorNormalized = MotionMatchingState.CurrentSearchResult.ComposedQuery.GetNormalizedValues();
 	TraceState.AssetPlayerTime = MotionMatchingState.CurrentSearchResult.AssetTime;
 	TraceState.DeltaTime = DeltaTime;
 	TraceState.SimLinearVelocity = SimLinearVelocity;
diff --git a/Engine/Plugins/Experimental/Animation/PoseSearch/Source/Runtime/Private/Trace/PoseSearchTraceLogger.cpp b/Engine/Plugins/Experimental/Animation/PoseSearch/Source/Runtime/Private/Trace/PoseSearchTraceLogger.cpp
--- a/Engine/Plugins/Experimental/Animation/PoseSearch/Source/Runtime/Private/Trace/PoseSearchTraceLogger.cpp
+++ b/Engine/Plugins/Experimental/Animation/PoseSearch/Source/Runtime/Private/Trace/PoseSearchTraceLogger.cpp
@@ -70,7 +70,6 @@ FArchive& operator<<(FArchive& Ar, FTraceMotionMatchingState& State)
 	Ar << State.ElapsedPoseJumpTime;
 	Ar << State.Flags;
 	Ar << State.QueryVector;
-	Ar << State.QueryVectorNormalized;
 	Ar << State.DatabaseSequenceFilter;
 	Ar << State.DatabaseBlendSpaceFilter;
 	Ar << State.AssetPlayerTime;
diff --git a/Engine/Plugins/Experimental/Animation/PoseSearch/Source/Runtime/Private/Trace/PoseSearchTraceLogger.h b/Engine/Plugins/Experimental/Animation/PoseSearch/Source/Runtime/Private/Trace/PoseSearchTraceLogger.h
--- a/Engine/Plugins/Experimental/Animation/PoseSearch/Source/Runtime/Private/Trace/PoseSearchTraceLogger.h
+++ b/Engine/Plugins/Experimental/Animation/PoseSearch/Source/Runtime/Private/Trace/PoseSearchTraceLogger.h
@@ -97,7 +97,6 @@ struct POSESEARCH_API FTraceMotionMatchingState
 
 	/** Search vectors in normalized and unnormalized forms */
 	TArray<float> QueryVector;
-	TArray<float> QueryVectorNormalized;
 
 	/** If true, groups are being filtered by DatabaseGroupQuery */
 	TArray<bool> DatabaseSequenceFilter;
diff --git a/Engine/Plugins/Experimental/Animation/PoseSearch/Source/Runtime/Public/PoseSearch/PoseSearch.h b/Engine/Plugins/Experimental/Animation/PoseSearch/Source/Runtime/Public/PoseSearch/PoseSearch.h
--- a/Engine/Plugins/Experimental/Animation/PoseSearch/Source/Runtime/Public/PoseSearch/PoseSearch.h
+++ b/Engine/Plugins/Experimental/Animation/PoseSearch/Source/Runtime/Public/PoseSearch/PoseSearch.h
@@ -101,7 +101,6 @@ UENUM()
 enum class EPoseSearchDataPreprocessor : int32
 {
 	None,
-	Automatic,
 	Normalize,
 
 	Num UMETA(Hidden),
@@ -503,10 +502,7 @@ public:
 	TObjectPtr<UMirrorDataTable> MirrorDataTable;
 
 	UPROPERTY(EditAnywhere, Category = "Schema")
-	EPoseSearchDataPreprocessor DataPreprocessor = EPoseSearchDataPreprocessor::Automatic;
-
-	UPROPERTY()
-	EPoseSearchDataPreprocessor EffectiveDataPreprocessor = EPoseSearchDataPreprocessor::Invalid;
+	EPoseSearchDataPreprocessor DataPreprocessor = EPoseSearchDataPreprocessor::Normalize;
 
 	UPROPERTY()
 	float SamplingInterval = 1.0f / DefaultSampleRate;
@@ -567,38 +563,6 @@ public:
 //////////////////////////////////////////////////////////////////////////
 // Search index
 
-USTRUCT()
-struct POSESEARCH_API FPoseSearchIndexPreprocessInfo
-{
-	GENERATED_BODY()
-
-	UPROPERTY()
-	int32 NumDimensions = 0;
-
-	UPROPERTY()
-	TArray<float> TransformationMatrix;
-
-	UPROPERTY()
-	TArray<float> InverseTransformationMatrix;
-
-	UPROPERTY()
-	TArray<float> SampleMean;
-
-	void Reset()
-	{
-		NumDimensions = 0;
-		TransformationMatrix.Reset();
-		InverseTransformationMatrix.Reset();
-		SampleMean.Reset();
-	}
-
-	bool IsValid() const
-	{
-		return NumDimensions > 0;
-	}
-};
-
-
 /**
  * This is kept for each pose in the search index along side the feature vector values and is used to influence the search.
  */
@@ -733,9 +697,6 @@ struct POSESEARCH_API FPoseSearchIndex
 	UPROPERTY()
 	TObjectPtr<const UPoseSearchSchema> Schema = nullptr;
 
-	UPROPERTY()
-	FPoseSearchIndexPreprocessInfo PreprocessInfo;
-
 	UPROPERTY()
 	TArray<FPoseSearchIndexAsset> Assets;
 
@@ -757,9 +718,6 @@ struct POSESEARCH_API FPoseSearchIndex
 
 	void Reset();
 
-	void Normalize (TArrayView<float> PoseVector) const;
-	void InverseNormalize (TArrayView<float> PoseVector) const;
-
 	// individual cost addends calculation methods
 	float ComputeMirrorMismatchAddend(int32 PoseIdx, UE::PoseSearch::FSearchContext& SearchContext) const;
 	float ComputeNotifyAddend(int32 PoseIdx) const;
@@ -970,7 +928,6 @@ public:
 
 	TArray<float>& EditValues() { return Values; }
 	TArrayView<const float> GetValues() const { return Values; }
-	TArrayView<const float> GetNormalizedValues() const { return ValuesNormalized; }
 
 	void CopyFromSearchIndex(const FPoseSearchIndex& SearchIndex, int32 PoseIdx);
 
@@ -978,14 +935,11 @@ public:
 	bool IsInitializedForSchema(const UPoseSearchSchema* Schema) const;
 	bool IsCompatible(const FPoseSearchFeatureVectorBuilder& OtherBuilder) const;
 
-	void Normalize(const FPoseSearchIndex& ForSearchIndex);
-
 private:
 	UPROPERTY(Transient)
 	TWeakObjectPtr<const UPoseSearchSchema> Schema = nullptr;
 
 	TArray<float> Values;
-	TArray<float> ValuesNormalized;
 };
 
 
