motion matching - requesting inertial blending only if blendstack is disabled

diff --git a/Engine/Plugins/Experimental/Animation/PoseSearch/Source/Runtime/Private/PoseSearchLibrary.cpp b/Engine/Plugins/Experimental/Animation/PoseSearch/Source/Runtime/Private/PoseSearchLibrary.cpp
--- a/Engine/Plugins/Experimental/Animation/PoseSearch/Source/Runtime/Private/PoseSearchLibrary.cpp
+++ b/Engine/Plugins/Experimental/Animation/PoseSearch/Source/Runtime/Private/PoseSearchLibrary.cpp
@@ -203,9 +203,14 @@ void FMotionMatchingState::JumpToPose(const FAnimationUpdateContext& Context, co
 	CurrentSearchResult = Result;
 
 	ElapsedPoseJumpTime = 0.0f;
+	
+	// requesting inertial blending only if blendstack is disabled
+	if (Settings.MaxActiveBlends <= 0)
+	{
+		const float JumpBlendTime = ComputeJumpBlendTime(Result, Settings);
+		RequestInertialBlend(Context, JumpBlendTime);
+	}
 
-	const float JumpBlendTime = ComputeJumpBlendTime(Result, Settings);
-	RequestInertialBlend(Context, JumpBlendTime);
 	Flags |= EMotionMatchingFlags::JumpedToPose;
 }
 
