motion matching - Added anim sequence percentage display to motion matching debug panel

diff --git a/Engine/Plugins/Experimental/Animation/PoseSearch/Source/Editor/Private/PoseSearchDebugger.cpp b/Engine/Plugins/Experimental/Animation/PoseSearch/Source/Editor/Private/PoseSearchDebugger.cpp
--- a/Engine/Plugins/Experimental/Animation/PoseSearch/Source/Editor/Private/PoseSearchDebugger.cpp
+++ b/Engine/Plugins/Experimental/Animation/PoseSearch/Source/Editor/Private/PoseSearchDebugger.cpp
@@ -9,6 +9,7 @@
 #include "Animation/BlendSpace.h"
 #include "Animation/MirrorDataTable.h"
 #include "Animation/AnimRootMotionProvider.h"
+#include "Preferences/PersonaOptions.h"
 #include "RewindDebuggerInterface/Public/IRewindDebugger.h"
 #include "ObjectTrace.h"
 #include "TraceServices/Model/AnalysisSession.h"
@@ -255,6 +256,7 @@ public:
 	FString AssetPath = "";
 	int32 DbAssetIdx = 0;
 	int32 AnimFrame = 0;
+	float AnimPercentage = 0.0f;
 	float AssetTime = 0.0f;
 	bool bMirrored = false;
 	bool bLooping = false;
@@ -446,12 +448,26 @@ namespace DebuggerDatabaseColumns
 				.SetUseGrouping(false)
 				.SetMaximumFractionalDigits(2);
 
+			FNumberFormattingOptions PercentageFormattingOptions = FNumberFormattingOptions()
+				.SetMaximumFractionalDigits(2);
+
 			if (Row->AssetType == ESearchIndexAssetType::Sequence)
 			{
-				return FText::Format(
-					FText::FromString("{0} ({1})"),
-					FText::AsNumber(Row->AnimFrame, &FNumberFormattingOptions::DefaultNoGrouping()),
-					FText::AsNumber(Row->AssetTime, &TimeFormattingOptions));
+				if (GetDefault<UPersonaOptions>()->bTimelineDisplayPercentage)
+				{
+					return FText::Format(
+						FText::FromString("{0} ({1}) ({2})"),
+						FText::AsNumber(Row->AnimFrame, &FNumberFormattingOptions::DefaultNoGrouping()),
+						FText::AsNumber(Row->AssetTime, &TimeFormattingOptions),
+						FText::AsPercent(Row->AnimPercentage, &PercentageFormattingOptions));
+				}
+				else
+				{
+					return FText::Format(
+						FText::FromString("{0} ({1})"),
+						FText::AsNumber(Row->AnimFrame, &FNumberFormattingOptions::DefaultNoGrouping()),
+						FText::AsNumber(Row->AssetTime, &TimeFormattingOptions));
+				}
 			}
 			else if (Row->AssetType == ESearchIndexAssetType::BlendSpace)
 			{
@@ -941,6 +957,7 @@ void SDebuggerDatabaseView::CreateRows(const FTraceMotionMatchingStateMessage& S
 						Row->AssetName = DbSequence.Sequence->GetName();
 						Row->AssetPath = DbSequence.Sequence->GetPathName();
 						Row->AnimFrame = DbSequence.Sequence->GetFrameAtTime(Time);
+						Row->AnimPercentage = Time / DbSequence.Sequence->GetPlayLength();
 						Row->bLooping = DbSequence.Sequence->bLoop;
 						Row->BlendParameters = FVector::Zero();
 					}
@@ -951,6 +968,7 @@ void SDebuggerDatabaseView::CreateRows(const FTraceMotionMatchingStateMessage& S
 						Row->AssetName = DbBlendSpace.BlendSpace->GetName();
 						Row->AssetPath = DbBlendSpace.BlendSpace->GetPathName();
 						Row->AnimFrame = 0; // There is no frame index associated with a blendspace
+						Row->AnimPercentage = 0.0f;
 						Row->bLooping = DbBlendSpace.BlendSpace->bLoop;
 						Row->BlendParameters = SearchIndexAsset->BlendParameters;
 					}
