motion matching - Added a.AnimNode.BlendStack.Enable console variable for previewing motion matching without blends

diff --git a/Engine/Plugins/Experimental/Animation/PoseSearch/Source/Runtime/Private/AnimNode_BlendStack.cpp b/Engine/Plugins/Experimental/Animation/PoseSearch/Source/Runtime/Private/AnimNode_BlendStack.cpp
--- a/Engine/Plugins/Experimental/Animation/PoseSearch/Source/Runtime/Private/AnimNode_BlendStack.cpp
+++ b/Engine/Plugins/Experimental/Animation/PoseSearch/Source/Runtime/Private/AnimNode_BlendStack.cpp
@@ -8,6 +8,8 @@
 
 #define LOCTEXT_NAMESPACE "AnimNode_BlendStack"
 
+TAutoConsoleVariable<int32> CVarAnimBlendStackEnable(TEXT("a.AnimNode.BlendStack.Enable"), 1, TEXT("Enable / Disable Blend Stack"));
+
 /////////////////////////////////////////////////////
 // FPoseSearchAnimPlayer
 void FPoseSearchAnimPlayer::Initialize(ESearchIndexAssetType InAssetType, UAnimationAsset* AnimationAsset, float AccumulatedTime, bool bLoop, bool bMirrored, UMirrorDataTable* MirrorDataTable, float BlendTime, const UBlendProfile* BlendProfile, FVector BlendParameters)
@@ -158,7 +160,17 @@ void FAnimNode_BlendStack::Evaluate_AnyThread(FPoseContext& Output)
 
 	DECLARE_SCOPE_HIERARCHICAL_COUNTER_ANIMNODE(Evaluate_AnyThread);
 
-	const int32 BlendStackSize = AnimPlayers.Num();
+	int32 BlendStackSize = AnimPlayers.Num();
+
+	// Disable blend stack if requested (for testing / debugging)
+	if (!CVarAnimBlendStackEnable.GetValueOnAnyThread())
+	{
+		if (BlendStackSize > 1)
+		{
+			BlendStackSize = 1;
+		}
+	}
+
 	if (BlendStackSize <= 0)
 	{
 		Output.ResetToRefPose();
