motion matching - Added VelocityDirection, VelocityDirectionXY, and FacingDirectionXY to UPoseSearchFeatureChannel_Trajectory

diff --git a/Engine/Plugins/Experimental/Animation/PoseSearch/Source/Runtime/Private/PoseSearch.cpp b/Engine/Plugins/Experimental/Animation/PoseSearch/Source/Runtime/Private/PoseSearch.cpp
--- a/Engine/Plugins/Experimental/Animation/PoseSearch/Source/Runtime/Private/PoseSearch.cpp
+++ b/Engine/Plugins/Experimental/Animation/PoseSearch/Source/Runtime/Private/PoseSearch.cpp
@@ -2866,6 +2866,14 @@ void FFeatureVectorHelper::ComputeMeanDeviations(float MinMeanDeviation, const E
 	DataOffset += Cardinality;
 }
 
+void FFeatureVectorHelper::SetMeanDeviations(float Deviation, Eigen::VectorXd& MeanDeviations, int32& DataOffset, int32 Cardinality)
+{
+	// Fill the feature's corresponding scaling axes with the supplied value
+	MeanDeviations.segment(DataOffset, Cardinality).setConstant(Deviation);
+
+	DataOffset += Cardinality;
+}
+
 //////////////////////////////////////////////////////////////////////////
 // FDebugDrawParams
 bool FDebugDrawParams::CanDraw() const
diff --git a/Engine/Plugins/Experimental/Animation/PoseSearch/Source/Runtime/Private/PoseSearchFeatureChannels.cpp b/Engine/Plugins/Experimental/Animation/PoseSearch/Source/Runtime/Private/PoseSearchFeatureChannels.cpp
--- a/Engine/Plugins/Experimental/Animation/PoseSearch/Source/Runtime/Private/PoseSearchFeatureChannels.cpp
+++ b/Engine/Plugins/Experimental/Animation/PoseSearch/Source/Runtime/Private/PoseSearchFeatureChannels.cpp
@@ -1218,10 +1218,23 @@ void UPoseSearchFeatureChannel_Trajectory::InitializeSchema(UE::PoseSearch::FSch
 			DataOffset += FFeatureVectorHelper::EncodeVector2DCardinality;
 		}
 
+		if (EnumHasAnyFlags(Sample.Flags, EPoseSearchTrajectoryFlags::VelocityDirection))
+		{
+			DataOffset += FFeatureVectorHelper::EncodeVectorCardinality;
+		}
+		if (EnumHasAnyFlags(Sample.Flags, EPoseSearchTrajectoryFlags::VelocityDirectionXY))
+		{
+			DataOffset += FFeatureVectorHelper::EncodeVector2DCardinality;
+		}
+
 		if (EnumHasAnyFlags(Sample.Flags, EPoseSearchTrajectoryFlags::FacingDirection))
 		{
 			DataOffset += FFeatureVectorHelper::EncodeVectorCardinality;
 		}
+		if (EnumHasAnyFlags(Sample.Flags, EPoseSearchTrajectoryFlags::FacingDirectionXY))
+		{
+			DataOffset += FFeatureVectorHelper::EncodeVector2DCardinality;
+		}
 	}
 
 	ChannelCardinality = DataOffset - ChannelDataOffset;
@@ -1270,6 +1283,23 @@ void UPoseSearchFeatureChannel_Trajectory::FillWeights(TArray<float>& Weights) c
 			DataOffset += FFeatureVectorHelper::EncodeVector2DCardinality;
 		}
 
+		if (EnumHasAnyFlags(Sample.Flags, EPoseSearchTrajectoryFlags::VelocityDirection))
+		{
+			for (int32 i = 0; i != FFeatureVectorHelper::EncodeVectorCardinality; ++i)
+			{
+				Weights[DataOffset + i] = Weight * Sample.Weight;
+			}
+			DataOffset += FFeatureVectorHelper::EncodeVectorCardinality;
+		}
+		if (EnumHasAnyFlags(Sample.Flags, EPoseSearchTrajectoryFlags::VelocityDirectionXY))
+		{
+			for (int32 i = 0; i != FFeatureVectorHelper::EncodeVector2DCardinality; ++i)
+			{
+				Weights[DataOffset + i] = Weight * Sample.Weight;
+			}
+			DataOffset += FFeatureVectorHelper::EncodeVector2DCardinality;
+		}
+
 		if (EnumHasAnyFlags(Sample.Flags, EPoseSearchTrajectoryFlags::FacingDirection))
 		{
 			for (int32 i = 0; i != FFeatureVectorHelper::EncodeVectorCardinality; ++i)
@@ -1278,6 +1308,14 @@ void UPoseSearchFeatureChannel_Trajectory::FillWeights(TArray<float>& Weights) c
 			}
 			DataOffset += FFeatureVectorHelper::EncodeVectorCardinality;
 		}
+		if (EnumHasAnyFlags(Sample.Flags, EPoseSearchTrajectoryFlags::FacingDirectionXY))
+		{
+			for (int32 i = 0; i != FFeatureVectorHelper::EncodeVector2DCardinality; ++i)
+			{
+				Weights[DataOffset + i] = Weight * Sample.Weight;
+			}
+			DataOffset += FFeatureVectorHelper::EncodeVector2DCardinality;
+		}
 	}
 
 	check(DataOffset == ChannelDataOffset + ChannelCardinality);
@@ -1319,9 +1357,22 @@ void UPoseSearchFeatureChannel_Trajectory::ComputeMeanDeviations(const Eigen::Ma
 			FFeatureVectorHelper::ComputeMeanDeviations(MinimumMeanDeviation, CenteredPoseMatrix, MeanDeviations, DataOffset, FFeatureVectorHelper::EncodeVector2DCardinality);
 		}
 
+		if (EnumHasAnyFlags(Sample.Flags, EPoseSearchTrajectoryFlags::VelocityDirection))
+		{
+			FFeatureVectorHelper::SetMeanDeviations(1.f, MeanDeviations, DataOffset, FFeatureVectorHelper::EncodeVectorCardinality);
+		}
+		if (EnumHasAnyFlags(Sample.Flags, EPoseSearchTrajectoryFlags::VelocityDirectionXY))
+		{
+			FFeatureVectorHelper::SetMeanDeviations(1.f, MeanDeviations, DataOffset, FFeatureVectorHelper::EncodeVector2DCardinality);
+		}
+
 		if (EnumHasAnyFlags(Sample.Flags, EPoseSearchTrajectoryFlags::FacingDirection))
 		{
-			FFeatureVectorHelper::ComputeMeanDeviations(MinimumMeanDeviation, CenteredPoseMatrix, MeanDeviations, DataOffset, FFeatureVectorHelper::EncodeVectorCardinality);
+			FFeatureVectorHelper::SetMeanDeviations(1.f, MeanDeviations, DataOffset, FFeatureVectorHelper::EncodeVectorCardinality);
+		}
+		if (EnumHasAnyFlags(Sample.Flags, EPoseSearchTrajectoryFlags::FacingDirectionXY))
+		{
+			FFeatureVectorHelper::SetMeanDeviations(1.f, MeanDeviations, DataOffset, FFeatureVectorHelper::EncodeVector2DCardinality);
 		}
 	}
 
@@ -1395,6 +1446,8 @@ void UPoseSearchFeatureChannel_Trajectory::IndexAssetPrivate(const UE::PoseSearc
 			LinearVelocity = (MirroredRootFuture.GetTranslation() - MirroredRootPast.GetTranslation()) / (IndexingContext.SamplingContext->FiniteDelta * 2.0f);
 		}
 
+		const FVector LinearVelocityDirection = LinearVelocity.GetClampedToMaxSize(1.0f);
+		const FVector FacingDirection = MirroredRootPresent.GetRotation().GetForwardVector();
 
 		if (EnumHasAnyFlags(Sample.Flags, EPoseSearchTrajectoryFlags::Position))
 		{
@@ -1414,9 +1467,22 @@ void UPoseSearchFeatureChannel_Trajectory::IndexAssetPrivate(const UE::PoseSearc
 			FFeatureVectorHelper::EncodeVector2D(FeatureVector.EditValues(), DataOffset, FVector2D(LinearVelocity.X, LinearVelocity.Y));
 		}
 
+		if (EnumHasAnyFlags(Sample.Flags, EPoseSearchTrajectoryFlags::VelocityDirection))
+		{
+			FFeatureVectorHelper::EncodeVector(FeatureVector.EditValues(), DataOffset, LinearVelocityDirection);
+		}
+		if (EnumHasAnyFlags(Sample.Flags, EPoseSearchTrajectoryFlags::VelocityDirectionXY))
+		{
+			FFeatureVectorHelper::EncodeVector2D(FeatureVector.EditValues(), DataOffset, FVector2D(LinearVelocityDirection.X, LinearVelocityDirection.Y));
+		}
+
 		if (EnumHasAnyFlags(Sample.Flags, EPoseSearchTrajectoryFlags::FacingDirection))
 		{
-			FFeatureVectorHelper::EncodeVector(FeatureVector.EditValues(), DataOffset, MirroredRootPresent.GetRotation().GetForwardVector());
+			FFeatureVectorHelper::EncodeVector(FeatureVector.EditValues(), DataOffset, FacingDirection);
+		}
+		if (EnumHasAnyFlags(Sample.Flags, EPoseSearchTrajectoryFlags::FacingDirectionXY))
+		{
+			FFeatureVectorHelper::EncodeVector2D(FeatureVector.EditValues(), DataOffset, FVector2D(FacingDirection.X, FacingDirection.Y));
 		}
 	}
 	check(DataOffset == ChannelDataOffset + ChannelCardinality);
@@ -1457,6 +1523,9 @@ bool UPoseSearchFeatureChannel_Trajectory::BuildQuery(UE::PoseSearch::FSearchCon
 		check(Sample.Offset >= PreviousOffset);
 		const FTrajectorySample TrajectorySample = FTrajectorySampleRange::IterSampleTrajectory(SearchContext.Trajectory->Samples, SampleDomain, Sample.Offset, NextIterStartIdx);
 
+		const FVector LinearVelocityDirection = TrajectorySample.LinearVelocity.GetClampedToMaxSize(1.0f);
+		const FVector FacingDirection = TrajectorySample.Transform.GetRotation().GetForwardVector();
+
 		if (EnumHasAnyFlags(Sample.Flags, EPoseSearchTrajectoryFlags::Position))
 		{
 			FFeatureVectorHelper::EncodeVector(InOutQuery.EditValues(), DataOffset, TrajectorySample.Transform.GetTranslation());
@@ -1475,9 +1544,22 @@ bool UPoseSearchFeatureChannel_Trajectory::BuildQuery(UE::PoseSearch::FSearchCon
 			FFeatureVectorHelper::EncodeVector2D(InOutQuery.EditValues(), DataOffset, FVector2D(TrajectorySample.LinearVelocity.X, TrajectorySample.LinearVelocity.Y));
 		}
 
+		if (EnumHasAnyFlags(Sample.Flags, EPoseSearchTrajectoryFlags::VelocityDirection))
+		{
+			FFeatureVectorHelper::EncodeVector(InOutQuery.EditValues(), DataOffset, LinearVelocityDirection);
+		}
+		if (EnumHasAnyFlags(Sample.Flags, EPoseSearchTrajectoryFlags::VelocityDirectionXY))
+		{
+			FFeatureVectorHelper::EncodeVector2D(InOutQuery.EditValues(), DataOffset, FVector2D(LinearVelocityDirection.X, LinearVelocityDirection.Y));
+		}
+
 		if (EnumHasAnyFlags(Sample.Flags, EPoseSearchTrajectoryFlags::FacingDirection))
 		{
-			FFeatureVectorHelper::EncodeVector(InOutQuery.EditValues(), DataOffset, TrajectorySample.Transform.GetRotation().GetForwardVector());
+			FFeatureVectorHelper::EncodeVector(InOutQuery.EditValues(), DataOffset, FacingDirection);
+		}
+		if (EnumHasAnyFlags(Sample.Flags, EPoseSearchTrajectoryFlags::FacingDirectionXY))
+		{
+			FFeatureVectorHelper::EncodeVector2D(InOutQuery.EditValues(), DataOffset, FVector2D(FacingDirection.X, FacingDirection.Y));
 		}
 	}
 	check(DataOffset == ChannelDataOffset + ChannelCardinality);
@@ -1545,10 +1627,23 @@ struct TrajectoryPositionReconstructor
 				DataOffset += UE::PoseSearch::FFeatureVectorHelper::EncodeVector2DCardinality;
 			}
 
+			if (EnumHasAnyFlags(Sample.Flags, EPoseSearchTrajectoryFlags::VelocityDirection))
+			{
+				DataOffset += UE::PoseSearch::FFeatureVectorHelper::EncodeVectorCardinality;
+			}
+			if (EnumHasAnyFlags(Sample.Flags, EPoseSearchTrajectoryFlags::VelocityDirectionXY))
+			{
+				DataOffset += UE::PoseSearch::FFeatureVectorHelper::EncodeVector2DCardinality;
+			}
+
 			if (EnumHasAnyFlags(Sample.Flags, EPoseSearchTrajectoryFlags::FacingDirection))
 			{
 				DataOffset += UE::PoseSearch::FFeatureVectorHelper::EncodeVectorCardinality;
 			}
+			if (EnumHasAnyFlags(Sample.Flags, EPoseSearchTrajectoryFlags::FacingDirectionXY))
+			{
+				DataOffset += UE::PoseSearch::FFeatureVectorHelper::EncodeVector2DCardinality;
+			}
 		}
 
 		if (bAddZeroOffsetSample)
@@ -1742,6 +1837,67 @@ void UPoseSearchFeatureChannel_Trajectory::DebugDraw(const UE::PoseSearch::FDebu
 			}
 		}
 
+		if (EnumHasAnyFlags(Sample.Flags, EPoseSearchTrajectoryFlags::VelocityDirection))
+		{
+			FVector TrajectoryVelDirection = FFeatureVectorHelper::DecodeVector(PoseVector, DataOffset);
+
+			const float GradientPercentage = NumSamples > 1 ? float(SampleIdx) / float(NumSamples - 1) : 0.f;
+			const FColor Color = DrawParams.GetColor(Sample.ColorPresetIndex, GradientPercentage);
+
+			TrajectoryVelDirection = DrawParams.RootTransform.TransformVector(TrajectoryVelDirection);
+
+			if (EnumHasAnyFlags(DrawParams.Flags, EDebugDrawFlags::DrawSearchIndex))
+			{
+				DrawDebugLine(DrawParams.World, TrajectoryPos, TrajectoryPos + TrajectoryVelDirection, Color, bPersistent, LifeTime, DepthPriority);
+			}
+			else
+			{
+				const float AdjustedThickness = EnumHasAnyFlags(DrawParams.Flags, EDebugDrawFlags::DrawFast) ? 0.0f : DrawDebugLineThickness;
+
+				DrawDebugLine(
+					DrawParams.World,
+					TrajectoryPos + TrajectoryVelDirection * DrawDebugSphereSize,
+					TrajectoryPos + TrajectoryVelDirection * DrawDebugSphereSize * 10.0f,
+					Color,
+					bPersistent,
+					LifeTime,
+					DepthPriority,
+					AdjustedThickness
+				);
+			}
+		}
+
+		if (EnumHasAnyFlags(Sample.Flags, EPoseSearchTrajectoryFlags::VelocityDirectionXY))
+		{
+			const FVector2D TrajectoryVelDirection2D = FFeatureVectorHelper::DecodeVector2D(PoseVector, DataOffset);
+			FVector TrajectoryVelDirection(TrajectoryVelDirection2D.X, TrajectoryVelDirection2D.Y, 0.f);
+
+			const float GradientPercentage = NumSamples > 1 ? float(SampleIdx) / float(NumSamples - 1) : 0.f;
+			const FColor Color = DrawParams.GetColor(Sample.ColorPresetIndex, GradientPercentage);
+
+			TrajectoryVelDirection = DrawParams.RootTransform.TransformVector(TrajectoryVelDirection);
+
+			if (EnumHasAnyFlags(DrawParams.Flags, EDebugDrawFlags::DrawSearchIndex))
+			{
+				DrawDebugLine(DrawParams.World, TrajectoryPos, TrajectoryPos + TrajectoryVelDirection, Color, bPersistent, LifeTime, DepthPriority);
+			}
+			else
+			{
+				const float AdjustedThickness = EnumHasAnyFlags(DrawParams.Flags, EDebugDrawFlags::DrawFast) ? 0.0f : DrawDebugLineThickness;
+
+				DrawDebugLine(
+					DrawParams.World,
+					TrajectoryPos + TrajectoryVelDirection * DrawDebugSphereSize,
+					TrajectoryPos + TrajectoryVelDirection * DrawDebugSphereSize * 10.0f,
+					Color,
+					bPersistent,
+					LifeTime,
+					DepthPriority,
+					AdjustedThickness
+				);
+			}
+		}
+
 		if (EnumHasAnyFlags(Sample.Flags, EPoseSearchTrajectoryFlags::FacingDirection))
 		{
 			FVector TrajectoryForward = FFeatureVectorHelper::DecodeVector(PoseVector, DataOffset);
@@ -1772,6 +1928,37 @@ void UPoseSearchFeatureChannel_Trajectory::DebugDraw(const UE::PoseSearch::FDebu
 			}
 		}
 
+		if (EnumHasAnyFlags(Sample.Flags, EPoseSearchTrajectoryFlags::FacingDirectionXY))
+		{
+			const FVector2D TrajectoryForward2D = FFeatureVectorHelper::DecodeVector2D(PoseVector, DataOffset);
+			FVector TrajectoryForward(TrajectoryForward2D.X, TrajectoryForward2D.Y, 0.f);
+
+			const float GradientPercentage = NumSamples > 1 ? float(SampleIdx) / float(NumSamples - 1) : 0.f;
+			const FColor Color = DrawParams.GetColor(Sample.ColorPresetIndex, GradientPercentage);
+
+			TrajectoryForward = DrawParams.RootTransform.TransformVector(TrajectoryForward);
+
+			if (EnumHasAnyFlags(DrawParams.Flags, EDebugDrawFlags::DrawSearchIndex))
+			{
+				DrawDebugLine(DrawParams.World, TrajectoryPos, TrajectoryPos + TrajectoryForward, Color, bPersistent, LifeTime, DepthPriority);
+			}
+			else
+			{
+				const float AdjustedThickness = EnumHasAnyFlags(DrawParams.Flags, EDebugDrawFlags::DrawFast) ? 0.0f : DrawDebugLineThickness;
+
+				DrawDebugLine(
+					DrawParams.World,
+					TrajectoryPos + TrajectoryForward * DrawDebugSphereSize,
+					TrajectoryPos + TrajectoryForward * DrawDebugSphereSize * 10.0f,
+					Color,
+					bPersistent,
+					LifeTime,
+					DepthPriority,
+					AdjustedThickness
+				);
+			}
+		}
+
 		if (EnumHasAnyFlags(DrawParams.Flags, EDebugDrawFlags::DrawSampleLabels))
 		{
 			const float GradientPercentage = NumSamples > 1 ? float(SampleIdx) / float(NumSamples - 1) : 0.f;
diff --git a/Engine/Plugins/Experimental/Animation/PoseSearch/Source/Runtime/Public/PoseSearch/PoseSearch.h b/Engine/Plugins/Experimental/Animation/PoseSearch/Source/Runtime/Public/PoseSearch/PoseSearch.h
--- a/Engine/Plugins/Experimental/Animation/PoseSearch/Source/Runtime/Public/PoseSearch/PoseSearch.h
+++ b/Engine/Plugins/Experimental/Animation/PoseSearch/Source/Runtime/Public/PoseSearch/PoseSearch.h
@@ -1273,9 +1273,12 @@ public:
 	static void EncodeVector2D(TArrayView<float> Values, int32& DataOffset, TArrayView<const float> PrevValues, TArrayView<const float> CurPrevValues, TArrayView<const float> NextPrevValues, float LerpValue);
 	static FVector2D DecodeVector2D(TArrayView<const float> Values, int32& DataOffset);
 
-	// populates MeanDeviations[DataOffset] ... MeanDeviations[DataOffset + Cardinality] with a single value the mean deviation calcualted from a cenetered matrix
+	// populates MeanDeviations[DataOffset] ... MeanDeviations[DataOffset + Cardinality] with a single value the mean deviation calculated from a centered matrix
 	static void ComputeMeanDeviations(float MinMeanDeviation, const Eigen::MatrixXd& CenteredPoseMatrix, Eigen::VectorXd& MeanDeviations, int32& DataOffset, int32 Cardinality);
 
+	// populates MeanDeviations[DataOffset] ... MeanDeviations[DataOffset + Cardinality] with a single value
+	static void SetMeanDeviations(float Deviation, Eigen::VectorXd& MeanDeviations, int32& DataOffset, int32 Cardinality);
+
 private:
 	static FQuat DecodeQuatInternal(TArrayView<const float> Values, int32 DataOffset);
 	static FVector DecodeVectorInternal(TArrayView<const float> Values, int32 DataOffset);
diff --git a/Engine/Plugins/Experimental/Animation/PoseSearch/Source/Runtime/Public/PoseSearch/PoseSearchFeatureChannels.h b/Engine/Plugins/Experimental/Animation/PoseSearch/Source/Runtime/Public/PoseSearch/PoseSearchFeatureChannels.h
--- a/Engine/Plugins/Experimental/Animation/PoseSearch/Source/Runtime/Public/PoseSearch/PoseSearchFeatureChannels.h
+++ b/Engine/Plugins/Experimental/Animation/PoseSearch/Source/Runtime/Public/PoseSearch/PoseSearchFeatureChannels.h
@@ -153,9 +153,12 @@ enum class EPoseSearchTrajectoryFlags : uint32
 {
 	Velocity = 1 << 0,
 	Position = 1 << 1,
+	VelocityDirection = 1 << 2,
 	FacingDirection = 1 << 3,
 	VelocityXY = 1 << 4,
 	PositionXY = 1 << 5,
+	VelocityDirectionXY = 1 << 6,
+	FacingDirectionXY = 1 << 7,
 };
 ENUM_CLASS_FLAGS(EPoseSearchTrajectoryFlags);
 constexpr bool EnumHasAnyFlags(int32 Flags, EPoseSearchTrajectoryFlags Contains) { return (Flags & int32(Contains)) != 0; }
