motion matching - Enable / Disable Assets in the Pose Search Database

diff --git a/Engine/Plugins/Experimental/Animation/PoseSearch/Source/Editor/Private/PoseSearchDatabaseViewModel.cpp b/Engine/Plugins/Experimental/Animation/PoseSearch/Source/Editor/Private/PoseSearchDatabaseViewModel.cpp
--- a/Engine/Plugins/Experimental/Animation/PoseSearch/Source/Editor/Private/PoseSearchDatabaseViewModel.cpp
+++ b/Engine/Plugins/Experimental/Animation/PoseSearch/Source/Editor/Private/PoseSearchDatabaseViewModel.cpp
@@ -397,6 +397,31 @@ namespace UE::PoseSearch
 		PoseSearchDatabase->Groups.RemoveAt(GroupIdx);
 	}
 
+	void FDatabaseViewModel::SetSelectedSequenceEnabled(int32 SequenceIndex, bool bEnabled)
+	{
+		PoseSearchDatabase->Sequences[SequenceIndex].bEnabled = bEnabled;
+		PoseSearchDatabase->NotifyAssetChange();
+		PoseSearchDatabase->BeginCacheDerivedData();
+	}	
+	
+	void FDatabaseViewModel::SetSelectedBlendSpaceEnabled(int32 BlendSpaceIndex, bool bEnabled)
+	{
+		PoseSearchDatabase->BlendSpaces[BlendSpaceIndex].bEnabled = bEnabled;
+		PoseSearchDatabase->NotifyAssetChange();
+		PoseSearchDatabase->BeginCacheDerivedData();
+	}
+
+	bool FDatabaseViewModel::IsSelectedSequenceEnabled(int32 SequenceIndex) const
+	{
+		return PoseSearchDatabase->Sequences[SequenceIndex].bEnabled;
+		PoseSearchDatabase->NotifyAssetChange();
+	}
+
+	bool FDatabaseViewModel::IsSelectedBlendSpaceEnabled(int32 BlendSpaceIndex) const
+	{
+		return PoseSearchDatabase->BlendSpaces[BlendSpaceIndex].bEnabled;
+	}
+
 	void FDatabaseViewModel::SetSelectedNodes(const TArrayView<TSharedPtr<FDatabaseAssetTreeNode>>& InSelectedNodes)
 	{
 		SelectedNodes = InSelectedNodes;
diff --git a/Engine/Plugins/Experimental/Animation/PoseSearch/Source/Editor/Private/SPoseSearchDatabaseAssetList.cpp b/Engine/Plugins/Experimental/Animation/PoseSearch/Source/Editor/Private/SPoseSearchDatabaseAssetList.cpp
--- a/Engine/Plugins/Experimental/Animation/PoseSearch/Source/Editor/Private/SPoseSearchDatabaseAssetList.cpp
+++ b/Engine/Plugins/Experimental/Animation/PoseSearch/Source/Editor/Private/SPoseSearchDatabaseAssetList.cpp
@@ -13,6 +13,7 @@
 #include "AssetSelection.h"
 
 #include "Widgets/Text/SRichTextBlock.h"
+#include "Widgets/Input/SCheckBox.h"
 #include "Framework/Commands/GenericCommands.h"
 #include "Framework/MultiBox/MultiBoxBuilder.h"
 #include "SPositiveActionButton.h"
@@ -20,7 +21,6 @@
 
 #include "ScopedTransaction.h"
 
-
 #define LOCTEXT_NAMESPACE "PoseSearchDatabaseAssetList"
 
 namespace UE::PoseSearch
@@ -244,6 +244,7 @@ namespace UE::PoseSearch
 			[
 				SNew(STextBlock)
 				.Text(this, &SDatabaseAssetListItem::GetName)
+				.ColorAndOpacity(this, &SDatabaseAssetListItem::GetNameTextColorAndOpacity)
 			]
 			+ SHorizontalBox::Slot()
 			.MaxWidth(18)
@@ -254,7 +255,25 @@ namespace UE::PoseSearch
 				SNew(SImage)
 				.Image(FAppStyle::Get().GetBrush("Icons.EyeDropper"))
 				.Visibility_Raw(this, &SDatabaseAssetListItem::GetSelectedActorIconVisbility)
+			]
+			+SHorizontalBox::Slot()
+			.MaxWidth(16)
+			.AutoWidth()
+			.HAlign(HAlign_Right)
+			.VAlign(VAlign_Center)
+			[
+				SNew(SCheckBox)
+				.IsChecked(this, &SDatabaseAssetListItem::GetAssetEnabledChecked)
+				.OnCheckStateChanged(const_cast<SDatabaseAssetListItem*>(this), &SDatabaseAssetListItem::OnAssetIsEnabledChanged)
+				.ToolTipText(this, &SDatabaseAssetListItem::GetAssetEnabledToolTip)
+				.CheckedImage(FAppStyle::Get().GetBrush("Icons.Visible"))
+				.CheckedHoveredImage(FAppStyle::Get().GetBrush("Icons.Visible"))
+				.CheckedPressedImage(FAppStyle::Get().GetBrush("Icons.Visible"))
+				.UncheckedImage(FAppStyle::Get().GetBrush("Icons.Hidden"))
+				.UncheckedHoveredImage(FAppStyle::Get().GetBrush("Icons.Hidden"))
+				.UncheckedPressedImage(FAppStyle::Get().GetBrush("Icons.Hidden"))
 			];
+
 		}
 
 		return ItemWidget.ToSharedRef();
@@ -328,6 +347,57 @@ namespace UE::PoseSearch
 		return EVisibility::Hidden;
 	}
 
+	ECheckBoxState SDatabaseAssetListItem::GetAssetEnabledChecked() const
+	{
+		TSharedPtr<FDatabaseViewModel> ViewModelPtr = EditorViewModel.Pin();
+		TSharedPtr<FDatabaseAssetTreeNode> TreeNodePtr = WeakAssetTreeNode.Pin();
+
+		if (TreeNodePtr->SourceAssetType == ESearchIndexAssetType::Sequence)
+		{
+			const bool bEnabled = ViewModelPtr->IsSelectedSequenceEnabled(TreeNodePtr->SourceAssetIdx);
+			return bEnabled ? ECheckBoxState::Checked : ECheckBoxState::Unchecked;
+		}
+		else if (TreeNodePtr->SourceAssetType == ESearchIndexAssetType::BlendSpace)
+		{
+			const bool bEnabled = ViewModelPtr->IsSelectedBlendSpaceEnabled(TreeNodePtr->SourceAssetIdx);
+			return bEnabled ? ECheckBoxState::Checked : ECheckBoxState::Unchecked;
+		}
+
+		return ECheckBoxState::Unchecked;
+	}
+
+	void SDatabaseAssetListItem::OnAssetIsEnabledChanged(ECheckBoxState NewCheckboxState)
+	{
+		TSharedPtr<FDatabaseViewModel> ViewModelPtr = EditorViewModel.Pin();
+		TSharedPtr<FDatabaseAssetTreeNode> TreeNodePtr = WeakAssetTreeNode.Pin();
+
+		if (TreeNodePtr->SourceAssetType == ESearchIndexAssetType::Sequence)
+		{
+			ViewModelPtr->SetSelectedSequenceEnabled(TreeNodePtr->SourceAssetIdx, NewCheckboxState == ECheckBoxState::Checked ? true : false);
+		}
+		else if (TreeNodePtr->SourceAssetType == ESearchIndexAssetType::BlendSpace)
+		{
+			ViewModelPtr->SetSelectedBlendSpaceEnabled(TreeNodePtr->SourceAssetIdx, NewCheckboxState == ECheckBoxState::Checked ? true : false);
+		}
+
+		SkeletonView.Pin()->RefreshTreeView(false, true);
+	}
+
+	FSlateColor SDatabaseAssetListItem::GetNameTextColorAndOpacity() const
+	{
+		return GetAssetEnabledChecked() == ECheckBoxState::Checked ? FLinearColor::White : FLinearColor(1.0f, 1.0f, 1.0f, 0.3f);
+	}
+
+	FText SDatabaseAssetListItem::GetAssetEnabledToolTip() const
+	{
+		if (GetAssetEnabledChecked() == ECheckBoxState::Checked)
+		{
+			return LOCTEXT("DisableAssetTooltip", "Disable this asset in the Pose Search Database.");
+		}
+		
+		return LOCTEXT("EnableAssetTooltip", "Enable this asset in the Pose Search Database.");
+	}
+
 	SDatabaseAssetTree::~SDatabaseAssetTree()
 	{
 	}
diff --git a/Engine/Plugins/Experimental/Animation/PoseSearch/Source/Editor/Public/PoseSearchDatabaseViewModel.h b/Engine/Plugins/Experimental/Animation/PoseSearch/Source/Editor/Public/PoseSearchDatabaseViewModel.h
--- a/Engine/Plugins/Experimental/Animation/PoseSearch/Source/Editor/Public/PoseSearchDatabaseViewModel.h
+++ b/Engine/Plugins/Experimental/Animation/PoseSearch/Source/Editor/Public/PoseSearchDatabaseViewModel.h
@@ -104,6 +104,12 @@ namespace UE::PoseSearch
 		void RemoveBlendSpaceFromGroup(int32 BlendSpaceIdx, int32 GroupIdx);
 		void DeleteGroup(int32 GroupIdx);
 
+		void SetSelectedSequenceEnabled(int32 SequenceIndex, bool bEnabled);
+		void SetSelectedBlendSpaceEnabled(int32 BlendSpaceIndex, bool bEnabled);
+		
+		bool IsSelectedSequenceEnabled(int32 SequenceIndex) const;
+		bool IsSelectedBlendSpaceEnabled(int32 BlendSpaceIndex) const;
+
 		void SetSelectedNodes(const TArrayView<TSharedPtr<FDatabaseAssetTreeNode>>& InSelectedNodes);
 		void ProcessSelectedActor(AActor* Actor);
 		
diff --git a/Engine/Plugins/Experimental/Animation/PoseSearch/Source/Editor/Public/SPoseSearchDatabaseAssetList.h b/Engine/Plugins/Experimental/Animation/PoseSearch/Source/Editor/Public/SPoseSearchDatabaseAssetList.h
--- a/Engine/Plugins/Experimental/Animation/PoseSearch/Source/Editor/Public/SPoseSearchDatabaseAssetList.h
+++ b/Engine/Plugins/Experimental/Animation/PoseSearch/Source/Editor/Public/SPoseSearchDatabaseAssetList.h
@@ -62,6 +62,11 @@ namespace UE::PoseSearch
 
 		EVisibility GetSelectedActorIconVisbility() const;
 
+		FSlateColor GetNameTextColorAndOpacity() const;
+		FText GetAssetEnabledToolTip() const;
+		ECheckBoxState GetAssetEnabledChecked() const;
+		void OnAssetIsEnabledChanged(ECheckBoxState NewCheckboxState);
+
 		TWeakPtr<FDatabaseAssetTreeNode> WeakAssetTreeNode;
 		TWeakPtr<FDatabaseViewModel> EditorViewModel;
 		TWeakPtr<SDatabaseAssetTree> SkeletonView;
diff --git a/Engine/Plugins/Experimental/Animation/PoseSearch/Source/Runtime/Private/PoseSearch.cpp b/Engine/Plugins/Experimental/Animation/PoseSearch/Source/Runtime/Private/PoseSearch.cpp
--- a/Engine/Plugins/Experimental/Animation/PoseSearch/Source/Runtime/Private/PoseSearch.cpp
+++ b/Engine/Plugins/Experimental/Animation/PoseSearch/Source/Runtime/Private/PoseSearch.cpp
@@ -1385,6 +1385,12 @@ bool UPoseSearchDatabase::TryInitSearchIndexAssets(FPoseSearchIndex& OutSearchIn
 	for (int32 SequenceIdx = 0; SequenceIdx < Sequences.Num(); ++SequenceIdx)
 	{
 		const FPoseSearchDatabaseSequence& Sequence = Sequences[SequenceIdx];
+
+		if (!Sequence.bEnabled)
+		{
+			continue;
+		}
+
 		bool bAddUnmirrored = 
 			Sequence.MirrorOption == EPoseSearchMirrorOption::UnmirroredOnly ||
 			Sequence.MirrorOption == EPoseSearchMirrorOption::UnmirroredAndMirrored;
@@ -1436,6 +1442,12 @@ bool UPoseSearchDatabase::TryInitSearchIndexAssets(FPoseSearchIndex& OutSearchIn
 	for (int32 BlendSpaceIdx = 0; BlendSpaceIdx < BlendSpaces.Num(); ++BlendSpaceIdx)
 	{
 		const FPoseSearchDatabaseBlendSpace& BlendSpace = BlendSpaces[BlendSpaceIdx];
+
+		if (!BlendSpace.bEnabled)
+		{
+			continue;
+		}
+
 		bool bAddUnmirrored =
 			BlendSpace.MirrorOption == EPoseSearchMirrorOption::UnmirroredOnly ||
 			BlendSpace.MirrorOption == EPoseSearchMirrorOption::UnmirroredAndMirrored;
diff --git a/Engine/Plugins/Experimental/Animation/PoseSearch/Source/Runtime/Public/PoseSearch/PoseSearch.h b/Engine/Plugins/Experimental/Animation/PoseSearch/Source/Runtime/Public/PoseSearch/PoseSearch.h
--- a/Engine/Plugins/Experimental/Animation/PoseSearch/Source/Runtime/Public/PoseSearch/PoseSearch.h
+++ b/Engine/Plugins/Experimental/Animation/PoseSearch/Source/Runtime/Public/PoseSearch/PoseSearch.h
@@ -746,6 +746,9 @@ struct POSESEARCH_API FPoseSearchDatabaseSequence : public FPoseSearchDatabaseAn
 	UPROPERTY(EditAnywhere, Category="Sequence")
 	TObjectPtr<UAnimSequence> Sequence = nullptr;
 
+	UPROPERTY(EditAnywhere, Category = "Sequence")
+	bool bEnabled = true;
+
 	UPROPERTY(EditAnywhere, Category="Sequence")
 	FFloatInterval SamplingRange = FFloatInterval(0.0f, 0.0f);
 
@@ -784,6 +787,9 @@ struct POSESEARCH_API FPoseSearchDatabaseBlendSpace : public FPoseSearchDatabase
 	UPROPERTY(EditAnywhere, Category = "BlendSpace")
 	TObjectPtr<UBlendSpace> BlendSpace = nullptr;
 
+	UPROPERTY(EditAnywhere, Category = "BlendSpace")
+	bool bEnabled = true;
+
 	UPROPERTY(EditAnywhere, Category = "BlendSpace")
 	EPoseSearchMirrorOption MirrorOption = EPoseSearchMirrorOption::UnmirroredOnly;
 
