motion matching - removing deprecated db groups

diff --git a/Engine/Plugins/Experimental/Animation/PoseSearch/Source/Editor/Private/PoseSearchCustomization.cpp b/Engine/Plugins/Experimental/Animation/PoseSearch/Source/Editor/Private/PoseSearchCustomization.cpp
--- a/Engine/Plugins/Experimental/Animation/PoseSearch/Source/Editor/Private/PoseSearchCustomization.cpp
+++ b/Engine/Plugins/Experimental/Animation/PoseSearch/Source/Editor/Private/PoseSearchCustomization.cpp
@@ -29,66 +29,12 @@
 #define LOCTEXT_NAMESPACE "PoseSearchCustomization"
 
 
-
-
-///////////////////////////////////////////////////////////////////////////////////////////////////////////////////////
-/// FPoseSearchDatabaseGroupCustomization
-
-void FPoseSearchDatabaseGroupCustomization::CustomizeHeader(TSharedRef<IPropertyHandle> InStructPropertyHandle, FDetailWidgetRow& HeaderRow, IPropertyTypeCustomizationUtils& StructCustomizationUtils)
-{
-	FText NamePropertyText;
-
-	TArray<UObject*> Objects;
-	InStructPropertyHandle->GetOuterObjects(Objects);
-	if (Objects.Num() == 1)
-	{
-		FPoseSearchDatabaseGroup* PoseSearchDatabaseGroup = (FPoseSearchDatabaseGroup*)InStructPropertyHandle->GetValueBaseAddress((uint8*)Objects[0]);
-		NamePropertyText = FText::FromString(PoseSearchDatabaseGroup->Tag.ToString());
-	}
-	else
-	{
-		TSharedPtr<IPropertyHandle> NamePropertyHandle = InStructPropertyHandle->GetChildHandle(GET_MEMBER_NAME_CHECKED(FPoseSearchDatabaseGroup, Tag));
-		NamePropertyHandle->GetValueAsFormattedText(NamePropertyText);
-	}
-
-	HeaderRow
-	.NameContent()
-	[
-		SNew(STextBlock)
-		.Text(NamePropertyText) 
-	]
-	.ValueContent()
-	[
-		InStructPropertyHandle->CreatePropertyValueWidget(false)
-	];
-
-	const FSimpleDelegate OnValueChanged = FSimpleDelegate::CreateLambda([&StructCustomizationUtils]()
-	{
-		StructCustomizationUtils.GetPropertyUtilities()->ForceRefresh();
-	});
-
-	InStructPropertyHandle->SetOnChildPropertyValueChanged(OnValueChanged);
-}
-
-void FPoseSearchDatabaseGroupCustomization::CustomizeChildren(TSharedRef<IPropertyHandle> InStructPropertyHandle, IDetailChildrenBuilder& StructBuilder, IPropertyTypeCustomizationUtils& StructCustomizationUtils)
-{
-	uint32 NumChildren;
-	InStructPropertyHandle->GetNumChildren(NumChildren);
-	for (uint32 ChildIndex = 0; ChildIndex < NumChildren; ++ChildIndex)
-	{
-		TSharedPtr<IPropertyHandle> ChildPropertyHandle = InStructPropertyHandle->GetChildHandle(ChildIndex);
-		StructBuilder.AddProperty(ChildPropertyHandle.ToSharedRef());
-	}
-}
-
-
 ///////////////////////////////////////////////////////////////////////////////////////////////////////////////////////
 /// FPoseSearchDatabaseSequenceCustomization
 
 void FPoseSearchDatabaseSequenceCustomization::CustomizeHeader(TSharedRef<IPropertyHandle> InStructPropertyHandle, FDetailWidgetRow& HeaderRow, IPropertyTypeCustomizationUtils& StructCustomizationUtils)
 {
 	FText SequenceNameText;
-	FString GroupsString; // = TEXT("Groups: ");
 
 	TArray<UObject*> Objects;
 	InStructPropertyHandle->GetOuterObjects(Objects);
@@ -106,23 +52,6 @@ void FPoseSearchDatabaseSequenceCustomization::CustomizeHeader(TSharedRef<IPrope
 		{
 			SequenceNameText = LOCTEXT("NewSequenceLabel", "New Sequence");
 		}
-		
-		const int32 NumGroups = PoseSearchDatabaseSequence->GroupTags.Num();
-		if (NumGroups == 0)
-		{
-			GroupsString.Append(TEXT("Default"));
-		}
-		else
-		{
-			for (int i = 0; i < NumGroups; ++i)
-			{
-				GroupsString.Append(PoseSearchDatabaseSequence->GroupTags.GetByIndex(i).ToString());
-				if (i < NumGroups - 1)
-				{
-					GroupsString.Append(TEXT(" | "));
-				}
-			}
-		}
 	}
 
 	HeaderRow
@@ -130,17 +59,6 @@ void FPoseSearchDatabaseSequenceCustomization::CustomizeHeader(TSharedRef<IPrope
 	[
 		SNew(STextBlock)
 		.Text(SequenceNameText)
-	]
-	.ValueContent()
-	[
-		SNew(SHorizontalBox)
-		+ SHorizontalBox::Slot()
-		.VAlign(VAlign_Center)
-		[
-			SNew(STextBlock)
-			.TextStyle(&FCoreStyle::Get().GetWidgetStyle<FTextBlockStyle>("SmallText"))
-			.Text(FText::FromString(GroupsString))
-		]
 	];
 
 	const FSimpleDelegate OnValueChanged = FSimpleDelegate::CreateLambda([&StructCustomizationUtils]()
@@ -176,7 +94,6 @@ void FPoseSearchDatabaseDetails::CustomizeDetails(IDetailLayoutBuilder& DetailBu
 	TArray<TSharedPtr<IPropertyHandle>> HiddenHandles;
 	HiddenHandles.Add(DetailBuilder.GetProperty(GET_MEMBER_NAME_CHECKED(UPoseSearchDatabase, Sequences)));
 	HiddenHandles.Add(DetailBuilder.GetProperty(GET_MEMBER_NAME_CHECKED(UPoseSearchDatabase, BlendSpaces)));
-	HiddenHandles.Add(DetailBuilder.GetProperty(GET_MEMBER_NAME_CHECKED(UPoseSearchDatabase, Groups)));
 	HiddenHandles.Add(DetailBuilder.GetProperty(GET_MEMBER_NAME_CHECKED(UPoseSearchDatabase, SimpleSequences)));
 	HiddenHandles.Add(DetailBuilder.GetProperty(GET_MEMBER_NAME_CHECKED(UPoseSearchDatabase, SimpleBlendSpaces)));
 
diff --git a/Engine/Plugins/Experimental/Animation/PoseSearch/Source/Editor/Private/PoseSearchDatabaseEditor.cpp b/Engine/Plugins/Experimental/Animation/PoseSearch/Source/Editor/Private/PoseSearchDatabaseEditor.cpp
--- a/Engine/Plugins/Experimental/Animation/PoseSearch/Source/Editor/Private/PoseSearchDatabaseEditor.cpp
+++ b/Engine/Plugins/Experimental/Animation/PoseSearch/Source/Editor/Private/PoseSearchDatabaseEditor.cpp
@@ -514,19 +514,6 @@ namespace UE::PoseSearch
 					NewSelectionReflection->SetSourceLink(SelectedItem, AssetTreeWidget);
 					SelectionReflections[BlendSpaceSelectionIndex].Add(NewSelectionReflection);
 				}
-				else
-				{
-					if (SelectedItem->SourceAssetIdx != INDEX_NONE)
-					{
-						UPoseSearchDatabaseGroupReflection* NewSelectionReflection =
-							NewObject<UPoseSearchDatabaseGroupReflection>();
-						NewSelectionReflection->AddToRoot();
-						NewSelectionReflection->Group =
-							GetPoseSearchDatabase()->Groups[SelectedItem->SourceAssetIdx];
-						NewSelectionReflection->SetSourceLink(SelectedItem, AssetTreeWidget);
-						SelectionReflections[GroupSelectionIndex].Add(NewSelectionReflection);
-					}
-				}
 			}
 		}
 
diff --git a/Engine/Plugins/Experimental/Animation/PoseSearch/Source/Editor/Private/PoseSearchDatabaseEditorReflection.cpp b/Engine/Plugins/Experimental/Animation/PoseSearch/Source/Editor/Private/PoseSearchDatabaseEditorReflection.cpp
--- a/Engine/Plugins/Experimental/Animation/PoseSearch/Source/Editor/Private/PoseSearchDatabaseEditorReflection.cpp
+++ b/Engine/Plugins/Experimental/Animation/PoseSearch/Source/Editor/Private/PoseSearchDatabaseEditorReflection.cpp
@@ -44,19 +44,4 @@ void UPoseSearchDatabaseBlendSpaceReflection::PostEditChangeProperty(
 	}
 }
 
-void UPoseSearchDatabaseGroupReflection::PostEditChangeProperty(
-	struct FPropertyChangedEvent& PropertyChangedEvent)
-{
-	Super::PostEditChangeProperty(PropertyChangedEvent);
-
-	check(WeakAssetTreeNode.Pin()->SourceAssetType == ESearchIndexAssetType::Invalid);
-
-	UPoseSearchDatabase* Database = WeakAssetTreeNode.Pin()->EditorViewModel.Pin()->GetPoseSearchDatabase();
-	if (IsValid(Database))
-	{
-		Database->Groups[WeakAssetTreeNode.Pin()->SourceAssetIdx] = Group;
-		AssetTreeWidget->FinalizeTreeChanges(true);
-	}
-}
-
 #endif // WITH_EDITOR
diff --git a/Engine/Plugins/Experimental/Animation/PoseSearch/Source/Editor/Private/PoseSearchDatabaseViewModel.cpp b/Engine/Plugins/Experimental/Animation/PoseSearch/Source/Editor/Private/PoseSearchDatabaseViewModel.cpp
--- a/Engine/Plugins/Experimental/Animation/PoseSearch/Source/Editor/Private/PoseSearchDatabaseViewModel.cpp
+++ b/Engine/Plugins/Experimental/Animation/PoseSearch/Source/Editor/Private/PoseSearchDatabaseViewModel.cpp
@@ -318,33 +318,16 @@ namespace UE::PoseSearch
 		return AnimationPreviewMode == PreviewMode;
 	}
 
-	void FDatabaseViewModel::AddSequenceToDatabase(UAnimSequence* AnimSequence, int InitialGroupIdx)
+	void FDatabaseViewModel::AddSequenceToDatabase(UAnimSequence* AnimSequence)
 	{
 		FPoseSearchDatabaseSequence& NewDbSequence = PoseSearchDatabase->Sequences.AddDefaulted_GetRef();
 		NewDbSequence.Sequence = AnimSequence;
-
-		if (InitialGroupIdx >= 0)
-		{
-			const FPoseSearchDatabaseGroup& InitialGroup = PoseSearchDatabase->Groups[InitialGroupIdx];
-			NewDbSequence.GroupTags.AddTag(InitialGroup.Tag);
-		}
 	}
 
-	void FDatabaseViewModel::AddBlendSpaceToDatabase(UBlendSpace* BlendSpace, int InitialGroupIdx)
+	void FDatabaseViewModel::AddBlendSpaceToDatabase(UBlendSpace* BlendSpace)
 	{
 		FPoseSearchDatabaseBlendSpace& NewDbBlendSpace = PoseSearchDatabase->BlendSpaces.AddDefaulted_GetRef();
 		NewDbBlendSpace.BlendSpace = BlendSpace;
-
-		if (InitialGroupIdx >= 0)
-		{
-			const FPoseSearchDatabaseGroup& InitialGroup = PoseSearchDatabase->Groups[InitialGroupIdx];
-			NewDbBlendSpace.GroupTags.AddTag(InitialGroup.Tag);
-		}
-	}
-
-	void FDatabaseViewModel::AddGroupToDatabase()
-	{
-		PoseSearchDatabase->Groups.AddDefaulted();
 	}
 
 	void FDatabaseViewModel::DeleteSequenceFromDatabase(int32 SequenceIdx)
@@ -352,51 +335,11 @@ namespace UE::PoseSearch
 		PoseSearchDatabase->Sequences.RemoveAt(SequenceIdx);
 	}
 
-	void FDatabaseViewModel::RemoveSequenceFromGroup(int32 SequenceIdx, int32 GroupIdx)
-	{
-		FPoseSearchDatabaseSequence& DbSequence = PoseSearchDatabase->Sequences[SequenceIdx];
-		FPoseSearchDatabaseGroup& DbGroup = PoseSearchDatabase->Groups[GroupIdx];
-		if (DbGroup.Tag.IsValid())
-		{
-			DbSequence.GroupTags.RemoveTag(DbGroup.Tag);
-		}
-	}
-
 	void FDatabaseViewModel::DeleteBlendSpaceFromDatabase(int32 BlendSpaceIdx)
 	{
 		PoseSearchDatabase->BlendSpaces.RemoveAt(BlendSpaceIdx);
 	}
 
-	void FDatabaseViewModel::RemoveBlendSpaceFromGroup(int32 BlendSpaceIdx, int32 GroupIdx)
-	{
-		FPoseSearchDatabaseBlendSpace& DbBlendSpace = PoseSearchDatabase->BlendSpaces[BlendSpaceIdx];
-		FPoseSearchDatabaseGroup& DbGroup = PoseSearchDatabase->Groups[GroupIdx];
-		if (DbGroup.Tag.IsValid())
-		{
-			DbBlendSpace.GroupTags.RemoveTag(DbGroup.Tag);
-		}
-	}
-
-	void FDatabaseViewModel::DeleteGroup(int32 GroupIdx)
-	{
-		FPoseSearchDatabaseGroup& DbGroup = PoseSearchDatabase->Groups[GroupIdx];
-
-		if (DbGroup.Tag.IsValid())
-		{
-			for (FPoseSearchDatabaseSequence& Sequence : PoseSearchDatabase->Sequences)
-			{
-				Sequence.GroupTags.RemoveTag(DbGroup.Tag);
-			}
-
-			for (FPoseSearchDatabaseBlendSpace& BlendSpace : PoseSearchDatabase->BlendSpaces)
-			{
-				BlendSpace.GroupTags.RemoveTag(DbGroup.Tag);
-			}
-		}
-
-		PoseSearchDatabase->Groups.RemoveAt(GroupIdx);
-	}
-
 	void FDatabaseViewModel::SetSelectedSequenceEnabled(int32 SequenceIndex, bool bEnabled)
 	{
 		PoseSearchDatabase->Sequences[SequenceIndex].bEnabled = bEnabled;
diff --git a/Engine/Plugins/Experimental/Animation/PoseSearch/Source/Editor/Private/PoseSearchDebugger.cpp b/Engine/Plugins/Experimental/Animation/PoseSearch/Source/Editor/Private/PoseSearchDebugger.cpp
--- a/Engine/Plugins/Experimental/Animation/PoseSearch/Source/Editor/Private/PoseSearchDebugger.cpp
+++ b/Engine/Plugins/Experimental/Animation/PoseSearch/Source/Editor/Private/PoseSearchDebugger.cpp
@@ -736,9 +736,6 @@ void SDebuggerDatabaseView::Update(const FTraceMotionMatchingStateMessage& State
 	check(ActiveView.Rows.Num() == 1);
 	check(ContinuingPoseView.Rows.Num() == 1);
 
-	DatabaseSequenceFilter = State.DatabaseSequenceFilter;
-	DatabaseBlendSpaceFilter = State.DatabaseBlendSpaceFilter;
-
 	SortDatabaseRows();
 	FilterDatabaseRows();
 
@@ -888,9 +885,7 @@ void SDebuggerDatabaseView::FilterDatabaseRows()
 	for (const auto& UnfilteredRow : UnfilteredDatabaseRows)
 	{
 		bool bPoseIsFromCurrentDb = EnumHasAnyFlags(UnfilteredRow->DatabaseEntryFlags, FTraceMotionMatchingStateDatabaseEntry::EFlags::CurrentDatabase);
-		bool bPassesSequenceFilter = !bPoseIsFromCurrentDb || (UnfilteredRow->AssetType == ESearchIndexAssetType::Sequence && DatabaseSequenceFilter[UnfilteredRow->DbAssetIdx]);
-		bool bPassesBlendSpaceFilter = !bPoseIsFromCurrentDb || (UnfilteredRow->AssetType == ESearchIndexAssetType::BlendSpace && DatabaseBlendSpaceFilter[UnfilteredRow->DbAssetIdx]);
-		bool bPassesAssetFilter = !bAssetFilterEnabled || bPassesSequenceFilter || bPassesBlendSpaceFilter;
+		bool bPassesAssetFilter = !bAssetFilterEnabled || !bPoseIsFromCurrentDb;
 
 		if (bPassesAssetFilter)
 		{
diff --git a/Engine/Plugins/Experimental/Animation/PoseSearch/Source/Editor/Private/PoseSearchDebugger.h b/Engine/Plugins/Experimental/Animation/PoseSearch/Source/Editor/Private/PoseSearchDebugger.h
--- a/Engine/Plugins/Experimental/Animation/PoseSearch/Source/Editor/Private/PoseSearchDebugger.h
+++ b/Engine/Plugins/Experimental/Animation/PoseSearch/Source/Editor/Private/PoseSearchDebugger.h
@@ -304,9 +304,6 @@ private:
 	/** All database poses */
 	TArray<TSharedRef<FDebuggerDatabaseRowData>> UnfilteredDatabaseRows;
 
-	TArrayView<const bool> DatabaseSequenceFilter;
-	TArrayView<const bool> DatabaseBlendSpaceFilter;
-
 	/** Database listing for filtered poses */
 	FTable FilteredDatabaseView;
 
diff --git a/Engine/Plugins/Experimental/Animation/PoseSearch/Source/Editor/Private/PoseSearchEditor.cpp b/Engine/Plugins/Experimental/Animation/PoseSearch/Source/Editor/Private/PoseSearchEditor.cpp
--- a/Engine/Plugins/Experimental/Animation/PoseSearch/Source/Editor/Private/PoseSearchEditor.cpp
+++ b/Engine/Plugins/Experimental/Animation/PoseSearch/Source/Editor/Private/PoseSearchEditor.cpp
@@ -195,7 +195,6 @@ void FEditorModule::RegisterCustomPropertyTypeLayout(FName PropertyTypeName, FOn
 void FEditorModule::RegisterPropertyTypeCustomizations()
 {
 	RegisterCustomPropertyTypeLayout("PoseSearchDatabaseSequence", FOnGetPropertyTypeCustomizationInstance::CreateStatic(&FPoseSearchDatabaseSequenceCustomization::MakeInstance));
-	RegisterCustomPropertyTypeLayout("PoseSearchDatabaseGroup", FOnGetPropertyTypeCustomizationInstance::CreateStatic(&FPoseSearchDatabaseGroupCustomization::MakeInstance));
 }
 
 void FEditorModule::RegisterObjectCustomizations()
diff --git a/Engine/Plugins/Experimental/Animation/PoseSearch/Source/Editor/Private/SPoseSearchDatabaseAssetList.cpp b/Engine/Plugins/Experimental/Animation/PoseSearch/Source/Editor/Private/SPoseSearchDatabaseAssetList.cpp
--- a/Engine/Plugins/Experimental/Animation/PoseSearch/Source/Editor/Private/SPoseSearchDatabaseAssetList.cpp
+++ b/Engine/Plugins/Experimental/Animation/PoseSearch/Source/Editor/Private/SPoseSearchDatabaseAssetList.cpp
@@ -104,21 +104,13 @@ namespace UE::PoseSearch
 
 	void SDatabaseAssetListItem::OnAddSequence()
 	{
-		TSharedPtr<FDatabaseAssetTreeNode> NodePtr = WeakAssetTreeNode.Pin();
-		const int32 GroupIdx = NodePtr->Parent.IsValid() ? NodePtr->Parent->SourceAssetIdx : NodePtr->SourceAssetIdx;
-
-		EditorViewModel.Pin()->AddSequenceToDatabase(nullptr, GroupIdx);
-
+		EditorViewModel.Pin()->AddSequenceToDatabase(nullptr);
 		SkeletonView.Pin()->RefreshTreeView(false);
 	}
 
 	void SDatabaseAssetListItem::OnAddBlendSpace()
 	{
-		TSharedPtr<FDatabaseAssetTreeNode> NodePtr = WeakAssetTreeNode.Pin();
-		const int32 GroupIdx = NodePtr->Parent.IsValid() ? NodePtr->Parent->SourceAssetIdx : NodePtr->SourceAssetIdx;
-
-		EditorViewModel.Pin()->AddBlendSpaceToDatabase(nullptr, GroupIdx);
-
+		EditorViewModel.Pin()->AddBlendSpaceToDatabase(nullptr);
 		SkeletonView.Pin()->RefreshTreeView(false);
 	}
 
@@ -154,18 +146,8 @@ namespace UE::PoseSearch
 			}
 			default:
 			{
-				if (Node->SourceAssetIdx == INDEX_NONE)
-				{
-					Name = LOCTEXT("Default", "Default");
-				}
-				else
-				{
-					FPoseSearchDatabaseGroup& Group = Database->Groups[Node->SourceAssetIdx];
-					if (Group.Tag.IsValid())
-					{
-						Name = FText::FromName(Group.Tag.GetTagName());
-					}
-				}
+				Name = LOCTEXT("Default", "Default");
+				break;
 			}
 			}
 		}
@@ -531,16 +513,6 @@ namespace UE::PoseSearch
 		// store selection so we can recover it afterwards (if possible)
 		TArray<TSharedPtr<FDatabaseAssetTreeNode>> PreviouslySelectedNodes = TreeView->GetSelectedItems();
 
-		// create all group nodes
-		for (int32 GroupIdx = 0; GroupIdx < Database->Groups.Num(); ++GroupIdx)
-		{
-			TSharedPtr<FDatabaseAssetTreeNode> GroupNode = MakeShared<FDatabaseAssetTreeNode>(
-				GroupIdx,
-				ESearchIndexAssetType::Invalid,
-				ViewModelRef);
-			AllNodes.Add(GroupNode);
-			RootNodes.Add(GroupNode);
-		}
 		TSharedPtr<FDatabaseAssetTreeNode> DefaultGroupNode = MakeShared<FDatabaseAssetTreeNode>(
 			INDEX_NONE,
 			ESearchIndexAssetType::Invalid,
@@ -566,54 +538,14 @@ namespace UE::PoseSearch
 		for (int32 SequenceIdx = 0; SequenceIdx < Database->Sequences.Num(); ++SequenceIdx)
 		{
 			const FPoseSearchDatabaseSequence& DbSequence = Database->Sequences[SequenceIdx];
-
-			int32 NumGroups = 0;
-			for (const FGameplayTag& GroupTag : DbSequence.GroupTags)
-			{
-				const int32 GroupIdx = Database->Groups.IndexOfByPredicate(
-					[GroupTag](const FPoseSearchDatabaseGroup& Group)
-				{
-					return Group.Tag == GroupTag;
-				});
-
-				if (GroupIdx != INDEX_NONE)
-				{
-					CreateAssetNode(SequenceIdx, ESearchIndexAssetType::Sequence, GroupIdx);
-					++NumGroups;
-				}
-			}
-
-			if (NumGroups == 0)
-			{
-				CreateAssetNode(SequenceIdx, ESearchIndexAssetType::Sequence, DefaultGroupIdx);
-			}
+			CreateAssetNode(SequenceIdx, ESearchIndexAssetType::Sequence, DefaultGroupIdx);
 		}
 
 		// create all blendspace nodes
 		for (int32 BlendSpaceIdx = 0; BlendSpaceIdx < Database->BlendSpaces.Num(); ++BlendSpaceIdx)
 		{
 			const FPoseSearchDatabaseBlendSpace& DbBlendSpace = Database->BlendSpaces[BlendSpaceIdx];
-
-			int32 NumGroups = 0;
-			for (const FGameplayTag& GroupTag : DbBlendSpace.GroupTags)
-			{
-				const int32 GroupIdx = Database->Groups.IndexOfByPredicate(
-					[GroupTag](const FPoseSearchDatabaseGroup& Group)
-				{
-					return Group.Tag == GroupTag;
-				});
-
-				if (GroupIdx != INDEX_NONE)
-				{
-					CreateAssetNode(BlendSpaceIdx, ESearchIndexAssetType::BlendSpace, GroupIdx);
-					++NumGroups;
-				}
-			}
-
-			if (NumGroups == 0)
-			{
-				CreateAssetNode(BlendSpaceIdx, ESearchIndexAssetType::BlendSpace, DefaultGroupIdx);
-			}
+			CreateAssetNode(BlendSpaceIdx, ESearchIndexAssetType::BlendSpace, DefaultGroupIdx);
 		}
 
 		TreeView->RequestTreeRefresh();
@@ -723,14 +655,12 @@ namespace UE::PoseSearch
 
 				if (AssetClass->IsChildOf(UAnimSequence::StaticClass()))
 				{
-					const int32 GroupIdx = FindGroupIndex(TargetItem);
-					ViewModel->AddSequenceToDatabase(Cast<UAnimSequence>(Asset), GroupIdx);
+					ViewModel->AddSequenceToDatabase(Cast<UAnimSequence>(Asset));
 					++AddedAssets;
 				}
 				else if (AssetClass->IsChildOf(UBlendSpace::StaticClass()))
 				{
-					const int32 GroupIdx = FindGroupIndex(TargetItem);
-					ViewModel->AddBlendSpaceToDatabase(Cast<UBlendSpace>(Asset), GroupIdx);
+					ViewModel->AddBlendSpaceToDatabase(Cast<UBlendSpace>(Asset));
 					++AddedAssets;
 				}
 			}
@@ -747,22 +677,6 @@ namespace UE::PoseSearch
 		return FReply::Handled();
 	}
 
-	int32 SDatabaseAssetTree::FindGroupIndex(TSharedPtr<FDatabaseAssetTreeNode> TargetItem)
-	{
-		if (!TargetItem.IsValid())
-		{
-			return INDEX_NONE;
-		}
-
-		if (TargetItem->SourceAssetType == ESearchIndexAssetType::Invalid)
-		{
-			return TargetItem->SourceAssetIdx;
-		}
-
-		check(TargetItem->Parent.IsValid() && TargetItem->Parent->SourceAssetType == ESearchIndexAssetType::Invalid);
-		return TargetItem->Parent->SourceAssetIdx;
-	}
-
 	TSharedRef<SWidget> SDatabaseAssetTree::CreateAddNewMenuWidget()
 	{
 		FMenuBuilder AddOptions(true, nullptr);
@@ -785,16 +699,6 @@ namespace UE::PoseSearch
 			EUserInterfaceActionType::Button);
 		AddOptions.EndSection();
 
-		AddOptions.BeginSection("AddOptions", LOCTEXT("GroupAddOptions", "Groups"));
-		AddOptions.AddMenuEntry(
-			LOCTEXT("GroupLabel", "Group"),
-			LOCTEXT("AddNewGroupTooltip", "Add new group"),
-			FSlateIcon(),
-			FUIAction(FExecuteAction::CreateSP(this, &SDatabaseAssetTree::OnAddGroup, true)),
-			NAME_None,
-			EUserInterfaceActionType::Button);
-		AddOptions.EndSection();
-
 		return AddOptions.MakeWidget();
 	}
 
@@ -820,23 +724,11 @@ namespace UE::PoseSearch
 		return MenuBuilder.MakeWidget();
 	}
 
-	void SDatabaseAssetTree::OnAddGroup(bool bFinalizeChanges)
-	{
-		FScopedTransaction Transaction(LOCTEXT("AddGroup", "Add Group"));
-
-		EditorViewModel.Pin()->AddGroupToDatabase();
-
-		if (bFinalizeChanges)
-		{
-			FinalizeTreeChanges();
-		}
-	}
-
 	void SDatabaseAssetTree::OnAddSequence(bool bFinalizeChanges)
 	{
 		FScopedTransaction Transaction(LOCTEXT("AddSequence", "Add Sequence"));
 
-		EditorViewModel.Pin()->AddSequenceToDatabase(nullptr, INDEX_NONE);
+		EditorViewModel.Pin()->AddSequenceToDatabase(nullptr);
 
 		if (bFinalizeChanges)
 		{
@@ -848,7 +740,7 @@ namespace UE::PoseSearch
 	{
 		FScopedTransaction Transaction(LOCTEXT("AddBlendSpaceTransaction", "Add Blend Space"));
 
-		EditorViewModel.Pin()->AddBlendSpaceToDatabase(nullptr, INDEX_NONE);
+		EditorViewModel.Pin()->AddBlendSpaceToDatabase(nullptr);
 
 		if (bFinalizeChanges)
 		{
@@ -879,48 +771,6 @@ namespace UE::PoseSearch
 		}
 	}
 
-	void SDatabaseAssetTree::OnRemoveFromGroup(TSharedPtr<FDatabaseAssetTreeNode> Node, bool bFinalizeChanges)
-	{
-		check(Node->Parent->SourceAssetType == ESearchIndexAssetType::Invalid);
-		
-		const int32 GroupIdx = Node->Parent->SourceAssetIdx;
-
-		FScopedTransaction Transaction(LOCTEXT("RemoveFromGroup", "Remove Asset From Group"));
-
-		if (Node->SourceAssetType == ESearchIndexAssetType::Sequence)
-		{
-			EditorViewModel.Pin()->RemoveSequenceFromGroup(Node->SourceAssetIdx, GroupIdx);
-		}
-		else if (Node->SourceAssetType == ESearchIndexAssetType::BlendSpace)
-		{
-			EditorViewModel.Pin()->RemoveBlendSpaceFromGroup(Node->SourceAssetIdx, GroupIdx);
-		}
-		else
-		{
-			checkNoEntry();
-		}
-
-		if (bFinalizeChanges)
-		{
-			FinalizeTreeChanges();
-		}
-	}
-
-	void SDatabaseAssetTree::OnDeleteGroup(TSharedPtr<FDatabaseAssetTreeNode> Node, bool bFinalizeChanges)
-	{
-		check(Node->SourceAssetType == ESearchIndexAssetType::Invalid);
-
-		FScopedTransaction Transaction(LOCTEXT("DeleteGroup", "Delete Group"));
-
-		const int32 GroupIdx = Node->SourceAssetIdx;
-		EditorViewModel.Pin()->DeleteGroup(GroupIdx);
-
-		if (bFinalizeChanges)
-		{
-			FinalizeTreeChanges();
-		}
-	}
-
 	void SDatabaseAssetTree::RegisterOnSelectionChanged(const FOnSelectionChanged& Delegate)
 	{
 		OnSelectionChanged.Add(Delegate);
@@ -1005,19 +855,7 @@ namespace UE::PoseSearch
 			{
 				if (SelectedNode->SourceAssetType != ESearchIndexAssetType::Invalid)
 				{
-					const int32 GroupIdx = SelectedNode->Parent->SourceAssetIdx;
-					if (GroupIdx == INDEX_NONE)
-					{
-						OnDeleteAsset(SelectedNode, false);
-					}
-					else
-					{
-						OnRemoveFromGroup(SelectedNode, false);
-					}
-				}
-				else if (SelectedNode->SourceAssetIdx != INDEX_NONE)
-				{
-					OnDeleteGroup(SelectedNode, false);
+					OnDeleteAsset(SelectedNode, false);
 				}
 			}
 
diff --git a/Engine/Plugins/Experimental/Animation/PoseSearch/Source/Editor/Public/PoseSearchCustomization.h b/Engine/Plugins/Experimental/Animation/PoseSearch/Source/Editor/Public/PoseSearchCustomization.h
--- a/Engine/Plugins/Experimental/Animation/PoseSearch/Source/Editor/Public/PoseSearchCustomization.h
+++ b/Engine/Plugins/Experimental/Animation/PoseSearch/Source/Editor/Public/PoseSearchCustomization.h
@@ -13,22 +13,6 @@
 class IPropertyHandle;
 class IDetailLayoutBuilder;
 
-/**
- * Customizes FPoseSearchDatabaseGroup to show the name of the group as the header
- */
-class FPoseSearchDatabaseGroupCustomization : public IPropertyTypeCustomization
-{
-public:
-	static TSharedRef<IPropertyTypeCustomization> MakeInstance()
-	{
-		return MakeShareable(new FPoseSearchDatabaseGroupCustomization);
-	}
-
-	/** IPropertyTypeCustomization interface */
-	virtual void CustomizeHeader(TSharedRef<class IPropertyHandle> InStructPropertyHandle, class FDetailWidgetRow& HeaderRow, IPropertyTypeCustomizationUtils& StructCustomizationUtils) override;
-	virtual void CustomizeChildren(TSharedRef<class IPropertyHandle> InStructPropertyHandle, class IDetailChildrenBuilder& StructBuilder, IPropertyTypeCustomizationUtils& StructCustomizationUtils) override;
-};
-
 /**
  * Customizes FPoseSearchDatabaseSequence to show the name of the sequence and the groups in the header
  */
diff --git a/Engine/Plugins/Experimental/Animation/PoseSearch/Source/Editor/Public/PoseSearchDatabaseEditorReflection.h b/Engine/Plugins/Experimental/Animation/PoseSearch/Source/Editor/Public/PoseSearchDatabaseEditorReflection.h
--- a/Engine/Plugins/Experimental/Animation/PoseSearch/Source/Editor/Public/PoseSearchDatabaseEditorReflection.h
+++ b/Engine/Plugins/Experimental/Animation/PoseSearch/Source/Editor/Public/PoseSearchDatabaseEditorReflection.h
@@ -58,17 +58,3 @@ public:
 	virtual void PostEditChangeProperty(struct FPropertyChangedEvent& PropertyChangedEvent) override;
 #endif
 };
-
-UCLASS()
-class UPoseSearchDatabaseGroupReflection : public UPoseSearchDatabaseReflectionBase
-{
-	GENERATED_BODY()
-
-public:
-	UPROPERTY(EditAnywhere, Category = "Selected Group")
-	FPoseSearchDatabaseGroup Group;
-
-#if WITH_EDITOR
-	virtual void PostEditChangeProperty(struct FPropertyChangedEvent& PropertyChangedEvent) override;
-#endif
-};
diff --git a/Engine/Plugins/Experimental/Animation/PoseSearch/Source/Editor/Public/PoseSearchDatabaseViewModel.h b/Engine/Plugins/Experimental/Animation/PoseSearch/Source/Editor/Public/PoseSearchDatabaseViewModel.h
--- a/Engine/Plugins/Experimental/Animation/PoseSearch/Source/Editor/Public/PoseSearchDatabaseViewModel.h
+++ b/Engine/Plugins/Experimental/Animation/PoseSearch/Source/Editor/Public/PoseSearchDatabaseViewModel.h
@@ -94,15 +94,11 @@ namespace UE::PoseSearch
 		void OnSetAnimationPreviewMode(EAnimationPreviewMode PreviewMode);
 		bool IsAnimationPreviewMode(EAnimationPreviewMode PreviewMode) const;
 
-		void AddSequenceToDatabase(UAnimSequence* AnimSequence, int InitialGroupIdx = -1);
-		void AddBlendSpaceToDatabase(UBlendSpace* BlendSpace, int InitialGroupIdx = -1);
-		void AddGroupToDatabase();
+		void AddSequenceToDatabase(UAnimSequence* AnimSequence);
+		void AddBlendSpaceToDatabase(UBlendSpace* BlendSpace);
 
 		void DeleteSequenceFromDatabase(int32 SequenceIdx);
-		void RemoveSequenceFromGroup(int32 SequenceIdx, int32 GroupIdx); 
 		void DeleteBlendSpaceFromDatabase(int32 BlendSpaceIdx);
-		void RemoveBlendSpaceFromGroup(int32 BlendSpaceIdx, int32 GroupIdx);
-		void DeleteGroup(int32 GroupIdx);
 
 		void SetSelectedSequenceEnabled(int32 SequenceIndex, bool bEnabled);
 		void SetSelectedBlendSpaceEnabled(int32 BlendSpaceIndex, bool bEnabled);
diff --git a/Engine/Plugins/Experimental/Animation/PoseSearch/Source/Editor/Public/SPoseSearchDatabaseAssetList.h b/Engine/Plugins/Experimental/Animation/PoseSearch/Source/Editor/Public/SPoseSearchDatabaseAssetList.h
--- a/Engine/Plugins/Experimental/Animation/PoseSearch/Source/Editor/Public/SPoseSearchDatabaseAssetList.h
+++ b/Engine/Plugins/Experimental/Animation/PoseSearch/Source/Editor/Public/SPoseSearchDatabaseAssetList.h
@@ -119,19 +119,13 @@ namespace UE::PoseSearch
 			EItemDropZone DropZone,
 			TSharedPtr<FDatabaseAssetTreeNode> TargetItem);
 
-		int32 FindGroupIndex(TSharedPtr<FDatabaseAssetTreeNode> TargetItem);
-
 		TSharedRef<SWidget> CreateAddNewMenuWidget();
 		TSharedPtr<SWidget> CreateContextMenu();
 		
-		void OnAddGroup(bool bFinalizeChanges = true);
 		void OnAddSequence(bool bFinalizeChanges = true);
 		void OnAddBlendSpace(bool bFinalizeChanges = true);
 
 		void OnDeleteAsset(TSharedPtr<FDatabaseAssetTreeNode> Node, bool bFinalizeChanges = true);
-		void OnRemoveFromGroup(TSharedPtr<FDatabaseAssetTreeNode> Node, bool bFinalizeChanges = true);
-		void OnDeleteGroup(TSharedPtr<FDatabaseAssetTreeNode> Node, bool bFinalizeChanges = true);
-
 		void CreateCommandList();
 
 		/** Removes existing selected component nodes from the tree*/
diff --git a/Engine/Plugins/Experimental/Animation/PoseSearch/Source/Runtime/Private/AnimNode_MotionMatching.cpp b/Engine/Plugins/Experimental/Animation/PoseSearch/Source/Runtime/Private/AnimNode_MotionMatching.cpp
--- a/Engine/Plugins/Experimental/Animation/PoseSearch/Source/Runtime/Private/AnimNode_MotionMatching.cpp
+++ b/Engine/Plugins/Experimental/Animation/PoseSearch/Source/Runtime/Private/AnimNode_MotionMatching.cpp
@@ -83,7 +83,6 @@ void FAnimNode_MotionMatching::UpdateAssetPlayer(const FAnimationUpdateContext&
 	UpdateMotionMatchingState(
 		Context,
 		Searchable,
-		bUseDatabaseTagQuery ? &DatabaseTagQuery : nullptr,
 		&ActiveTagsContainer,
 		Trajectory,
 		Settings,
diff --git a/Engine/Plugins/Experimental/Animation/PoseSearch/Source/Runtime/Private/PoseSearch.cpp b/Engine/Plugins/Experimental/Animation/PoseSearch/Source/Runtime/Private/PoseSearch.cpp
--- a/Engine/Plugins/Experimental/Animation/PoseSearch/Source/Runtime/Private/PoseSearch.cpp
+++ b/Engine/Plugins/Experimental/Animation/PoseSearch/Source/Runtime/Private/PoseSearch.cpp
@@ -176,7 +176,7 @@ static void CalcChannelCosts(const UPoseSearchSchema* Schema, TArrayView<const f
 	}
 }
 
-static int32 PopulateNonSelectableIdx(size_t* NonSelectableIdx, int32 NonSelectableIdxSize, const FSearchContext& SearchContext, const UPoseSearchDatabase* Database, int32 IdxOffset = 0)
+static int32 PopulateNonSelectableIdx(size_t* NonSelectableIdx, int32 NonSelectableIdxSize, const FSearchContext& SearchContext, const UPoseSearchDatabase* Database)
 {
 	check(Database);
 
@@ -205,7 +205,7 @@ static int32 PopulateNonSelectableIdx(size_t* NonSelectableIdx, int32 NonSelecta
 
 			if (CurrentIndexAsset->SamplingInterval.Contains(PoseAssetPlayerTime) && NonSelectableIdxUsedSize < NonSelectableIdxSize)
 			{
-				NonSelectableIdx[NonSelectableIdxUsedSize++] = PoseIdx + IdxOffset;
+				NonSelectableIdx[NonSelectableIdxUsedSize++] = PoseIdx;
 			}
 			else
 			{
@@ -230,7 +230,7 @@ static int32 PopulateNonSelectableIdx(size_t* NonSelectableIdx, int32 NonSelecta
 
 			if (CurrentIndexAsset->SamplingInterval.Contains(PoseAssetPlayerTime) && NonSelectableIdxUsedSize < NonSelectableIdxSize)
 			{
-				NonSelectableIdx[NonSelectableIdxUsedSize++] = PoseIdx + IdxOffset;
+				NonSelectableIdx[NonSelectableIdxUsedSize++] = PoseIdx;
 			}
 			else
 			{
@@ -250,7 +250,7 @@ static int32 PopulateNonSelectableIdx(size_t* NonSelectableIdx, int32 NonSelecta
 				const FHistoricalPoseIndex& HistoricalPoseIndex = It.Key();
 				if (HistoricalPoseIndex.DatabaseKey == DatabaseKey)
 				{
-					NonSelectableIdx[NonSelectableIdxUsedSize++] = HistoricalPoseIndex.PoseIndex + IdxOffset;
+					NonSelectableIdx[NonSelectableIdxUsedSize++] = HistoricalPoseIndex.PoseIndex;
 				}
 			}
 			else
@@ -586,19 +586,6 @@ int32 FPoseSearchIndex::FindAssetIndex(const FPoseSearchIndexAsset* Asset) const
 	return Result;
 }
 
-const FGroupSearchIndex* FPoseSearchIndex::FindGroup(int32 GroupIndex) const
-{
-	for (const FGroupSearchIndex& Group : Groups)
-	{
-		if (Group.GroupIndex == GroupIndex)
-		{
-			return &Group;
-		}
-	}
-	check(false);
-	return nullptr;
-}
-
 const FPoseSearchIndexAsset* FPoseSearchIndex::FindAssetForPose(int32 PoseIdx) const
 {
 	auto Predicate = [PoseIdx](const FPoseSearchIndexAsset& Asset)
@@ -1008,22 +995,6 @@ const bool UPoseSearchDatabase::IsSourceAssetLooping(const FPoseSearchIndexAsset
 	return false;
 }
 
-const FGameplayTagContainer* UPoseSearchDatabase::GetSourceAssetGroupTags(const FPoseSearchIndexAsset* SearchIndexAsset) const
-{
-	if (SearchIndexAsset->Type == ESearchIndexAssetType::Sequence)
-	{
-		return &GetSequenceSourceAsset(SearchIndexAsset).GroupTags;
-	}
-	
-	if (SearchIndexAsset->Type == ESearchIndexAssetType::BlendSpace)
-	{
-		return &GetBlendSpaceSourceAsset(SearchIndexAsset).GroupTags;
-	}
-	
-	checkNoEntry();
-	return nullptr;
-}
-
 const FString UPoseSearchDatabase::GetSourceAssetName(const FPoseSearchIndexAsset* SearchIndexAsset) const
 {
 	if (SearchIndexAsset->Type == ESearchIndexAssetType::Sequence)
@@ -1281,37 +1252,6 @@ static void FindValidSequenceIntervals(const FPoseSearchDatabaseSequence& DbSequ
 	}
 }
 
-static inline void CollectGroupIndices(
-	TArrayView<const FPoseSearchDatabaseGroup> Groups,
-	const FGameplayTagContainer& GroupTags,
-	const int32 Index,
-	TArray<int32>& GroupIndices,
-	TArray<int32>& BadGroupIndices)
-{
-	GroupIndices.Reset();
-
-	for (const FGameplayTag& GroupTag : GroupTags)
-	{
-		const int32 GroupIndex = Groups.IndexOfByPredicate([&](const FPoseSearchDatabaseGroup& DatabaseGroup)
-		{
-			return DatabaseGroup.Tag == GroupTag;
-		});
-
-		// we don't add INDEX_NONE because index none represents a choice to use the default group by not adding
-		// any group identifiers. If an added identifier doesn't match, that's an error. In the future this
-		// should be made robust enough to prevent these errors from happening
-		if (GroupIndex == INDEX_NONE)
-		{
-			BadGroupIndices.Add(Index);
-		}
-	}
-
-	if (GroupIndices.Num() == 0)
-	{
-		GroupIndices.Add(INDEX_NONE);
-	}
-}
-
 UAnimationAsset* FPoseSearchDatabaseBlendSpace::GetAnimationAsset() const
 {
 	return BlendSpace.Get();
@@ -1378,9 +1318,6 @@ bool UPoseSearchDatabase::TryInitSearchIndexAssets(FPoseSearchIndex& OutSearchIn
 	bool bAnyMirrored = false;
 	
 	TArray<FFloatRange> ValidRanges;
-	TArray<int32> GroupIndices;
-	TArray<int32> BadSequenceGroupIndices;
-	TArray<int32> BadBlendSpaceGroupIndices;
 
 	for (int32 SequenceIdx = 0; SequenceIdx < Sequences.Num(); ++SequenceIdx)
 	{
@@ -1398,41 +1335,29 @@ bool UPoseSearchDatabase::TryInitSearchIndexAssets(FPoseSearchIndex& OutSearchIn
 			Sequence.MirrorOption == EPoseSearchMirrorOption::MirroredOnly ||
 			Sequence.MirrorOption == EPoseSearchMirrorOption::UnmirroredAndMirrored;
 
-		CollectGroupIndices(
-			Groups,
-			Sequence.GroupTags,
-			SequenceIdx,
-			GroupIndices,
-			BadSequenceGroupIndices);
-
-		for (int32 GroupIndex : GroupIndices)
+		ValidRanges.Reset();
+		FindValidSequenceIntervals(Sequence, ValidRanges);
+		for (const FFloatRange& Range : ValidRanges)
 		{
-			ValidRanges.Reset();
-			FindValidSequenceIntervals(Sequence, ValidRanges);
-			for (const FFloatRange& Range : ValidRanges)
+			if (bAddUnmirrored)
 			{
-				if (bAddUnmirrored)
-				{
-					OutSearchIndex.Assets.Add(
-						FPoseSearchIndexAsset(
-							ESearchIndexAssetType::Sequence,
-							GroupIndex,
-							SequenceIdx,
-							false,
-							FFloatInterval(Range.GetLowerBoundValue(), Range.GetUpperBoundValue())));
-				}
+				OutSearchIndex.Assets.Add(
+					FPoseSearchIndexAsset(
+						ESearchIndexAssetType::Sequence,
+						SequenceIdx,
+						false,
+						FFloatInterval(Range.GetLowerBoundValue(), Range.GetUpperBoundValue())));
+			}
 
-				if (bAddMirrored)
-				{
-					OutSearchIndex.Assets.Add(
-						FPoseSearchIndexAsset(
-							ESearchIndexAssetType::Sequence,
-							GroupIndex,
-							SequenceIdx,
-							true,
-							FFloatInterval(Range.GetLowerBoundValue(), Range.GetUpperBoundValue())));
-					bAnyMirrored = true;
-				}
+			if (bAddMirrored)
+			{
+				OutSearchIndex.Assets.Add(
+					FPoseSearchIndexAsset(
+						ESearchIndexAssetType::Sequence,
+						SequenceIdx,
+						true,
+						FFloatInterval(Range.GetLowerBoundValue(), Range.GetUpperBoundValue())));
+				bAnyMirrored = true;
 			}
 		}
 	}
@@ -1455,81 +1380,62 @@ bool UPoseSearchDatabase::TryInitSearchIndexAssets(FPoseSearchIndex& OutSearchIn
 			BlendSpace.MirrorOption == EPoseSearchMirrorOption::MirroredOnly ||
 			BlendSpace.MirrorOption == EPoseSearchMirrorOption::UnmirroredAndMirrored;
 
-		CollectGroupIndices(
-			Groups,
-			BlendSpace.GroupTags,
-			BlendSpaceIdx,
-			GroupIndices,
-			BadBlendSpaceGroupIndices);
+		int32 HorizontalBlendNum, VerticalBlendNum;
+		float HorizontalBlendMin, HorizontalBlendMax, VerticalBlendMin, VerticalBlendMax;
 
-		for (int32 GroupIndex : GroupIndices)
-		{
-			int32 HorizontalBlendNum, VerticalBlendNum;
-			float HorizontalBlendMin, HorizontalBlendMax, VerticalBlendMin, VerticalBlendMax;
+		BlendSpace.GetBlendSpaceParameterSampleRanges(
+			HorizontalBlendNum,
+			VerticalBlendNum,
+			HorizontalBlendMin,
+			HorizontalBlendMax,
+			VerticalBlendMin,
+			VerticalBlendMax);
 
-			BlendSpace.GetBlendSpaceParameterSampleRanges(
-				HorizontalBlendNum,
-				VerticalBlendNum,
-				HorizontalBlendMin,
-				HorizontalBlendMax,
-				VerticalBlendMin,
-				VerticalBlendMax);
-
-			for (int32 HorizontalIndex = 0; HorizontalIndex < HorizontalBlendNum; HorizontalIndex++)
+		for (int32 HorizontalIndex = 0; HorizontalIndex < HorizontalBlendNum; HorizontalIndex++)
+		{
+			for (int32 VerticalIndex = 0; VerticalIndex < VerticalBlendNum; VerticalIndex++)
 			{
-				for (int32 VerticalIndex = 0; VerticalIndex < VerticalBlendNum; VerticalIndex++)
-				{
-					FVector BlendParameters = BlendParameterForSampleRanges(
-						HorizontalIndex,
-						VerticalIndex,
-						HorizontalBlendNum,
-						VerticalBlendNum,
-						HorizontalBlendMin,
-						HorizontalBlendMax,
-						VerticalBlendMin,
-						VerticalBlendMax);
+				FVector BlendParameters = BlendParameterForSampleRanges(
+					HorizontalIndex,
+					VerticalIndex,
+					HorizontalBlendNum,
+					VerticalBlendNum,
+					HorizontalBlendMin,
+					HorizontalBlendMax,
+					VerticalBlendMin,
+					VerticalBlendMax);
 						
-					int32 TriangulationIndex = 0;
-					BlendSpace.BlendSpace->GetSamplesFromBlendInput(BlendParameters, BlendSamples, TriangulationIndex, true);
+				int32 TriangulationIndex = 0;
+				BlendSpace.BlendSpace->GetSamplesFromBlendInput(BlendParameters, BlendSamples, TriangulationIndex, true);
 
-					float PlayLength = BlendSpace.BlendSpace->GetAnimationLengthFromSampleData(BlendSamples);
+				float PlayLength = BlendSpace.BlendSpace->GetAnimationLengthFromSampleData(BlendSamples);
 
-					if (bAddUnmirrored)
-					{
-						OutSearchIndex.Assets.Add(
-							FPoseSearchIndexAsset(
-								ESearchIndexAssetType::BlendSpace,
-								GroupIndex,
-								BlendSpaceIdx,
-								false,
-								FFloatInterval(0.0f, PlayLength),
-								BlendParameters));
-					}
+				if (bAddUnmirrored)
+				{
+					OutSearchIndex.Assets.Add(
+						FPoseSearchIndexAsset(
+							ESearchIndexAssetType::BlendSpace,
+							BlendSpaceIdx,
+							false,
+							FFloatInterval(0.0f, PlayLength),
+							BlendParameters));
+				}
 
-					if (bAddMirrored)
-					{
-						OutSearchIndex.Assets.Add(
-							FPoseSearchIndexAsset(
-								ESearchIndexAssetType::BlendSpace,
-								GroupIndex,
-								BlendSpaceIdx,
-								true,
-								FFloatInterval(0.0f, PlayLength),
-								BlendParameters));
-						bAnyMirrored = true;
-					}
+				if (bAddMirrored)
+				{
+					OutSearchIndex.Assets.Add(
+						FPoseSearchIndexAsset(
+							ESearchIndexAssetType::BlendSpace,
+							BlendSpaceIdx,
+							true,
+							FFloatInterval(0.0f, PlayLength),
+							BlendParameters));
+					bAnyMirrored = true;
 				}
 			}
 		}
 	}
 
-	// @todo: change the above for loops fill OutSearchIndex.Assets already in ascending group order
-	// sorting by ascending SourceGroupIdx
-	OutSearchIndex.Assets.Sort([](const FPoseSearchIndexAsset& InOne, const FPoseSearchIndexAsset& InTwo)
-	{
-		return InOne.SourceGroupIdx < InTwo.SourceGroupIdx;
-	});
-
 	if (bAnyMirrored && !Schema->MirrorDataTable)
 	{
 		UE_LOG(
@@ -1542,26 +1448,6 @@ bool UPoseSearchDatabase::TryInitSearchIndexAssets(FPoseSearchIndex& OutSearchIn
 		return false;
 	}
 
-	for (int32 BadGroupSequenceIdx : BadSequenceGroupIndices)
-	{
-		UE_LOG(
-			LogPoseSearch,
-			Warning,
-			TEXT("Database %s, sequence %s is asking for a group that doesn't exist"),
-			*GetNameSafe(this),
-			*GetNameSafe(Sequences[BadGroupSequenceIdx].Sequence));
-	}
-
-	for (int32 BadGroupBlendSpaceIdx : BadBlendSpaceGroupIndices)
-	{
-		UE_LOG(
-			LogPoseSearch,
-			Warning,
-			TEXT("Database %s, blendspace %s is asking for a group that doesn't exist"),
-			*GetNameSafe(this),
-			*GetNameSafe(BlendSpaces[BadGroupBlendSpaceIdx].BlendSpace));
-	}
-
 	return true;
 }
 
@@ -1770,11 +1656,6 @@ void UPoseSearchDatabase::PostEditChangeProperty(struct FPropertyChangedEvent& P
 		NotifyAssetChange();
 	}
 
-	if (PropertyChangedEvent.MemberProperty->GetFName() == GET_MEMBER_NAME_CHECKED(UPoseSearchDatabase, Groups))
-	{
-		NotifyGroupChange();
-	}
-
 	BeginCacheDerivedData();
 }
 
@@ -1809,7 +1690,7 @@ bool UPoseSearchDatabase::IsCachedCookedPlatformDataLoaded(const ITargetPlatform
 
 #endif // WITH_EDITOR
 
-FPoseSearchCost UPoseSearchDatabase::ComparePoses(UE::PoseSearch::FSearchContext& SearchContext, int32 PoseIdx, UE::PoseSearch::EPoseComparisonFlags PoseComparisonFlags, int32 GroupIdx, const TArrayView<const float>& QueryValues) const
+FPoseSearchCost UPoseSearchDatabase::ComparePoses(UE::PoseSearch::FSearchContext& SearchContext, int32 PoseIdx, UE::PoseSearch::EPoseComparisonFlags PoseComparisonFlags, const TArrayView<const float>& QueryValues) const
 {
 	using namespace Eigen;
 	using namespace UE::PoseSearch;
@@ -1823,14 +1704,7 @@ FPoseSearchCost UPoseSearchDatabase::ComparePoses(UE::PoseSearch::FSearchContext
 	const int32 Dims = PoseValues.Num();
 	check(Dims == QueryValues.Num());
 
-	if (GroupIdx == INDEX_NONE)
-	{
-		const FPoseSearchIndexAsset* SearchIndexAsset = SearchIndex->FindAssetForPose(PoseIdx);
-		check(SearchIndexAsset)
-		GroupIdx = SearchIndexAsset->SourceGroupIdx;
-	}
-
-	TArrayView<const float> Weights = SearchIndex->FindGroup(GroupIdx)->Weights;
+	TArrayView<const float> Weights(SearchIndex->Weights);
 
 	Result.SetDissimilarity(CompareFeatureVectors(
 		PoseValues.Num(),
@@ -1965,37 +1839,34 @@ UE::PoseSearch::FSearchResult UPoseSearchDatabase::SearchPCAKDTree(UE::PoseSearc
 	// KDTree in PCA space search
 	if (PoseSearchMode == EPoseSearchMode::PCAKDTree_Validate)
 	{
-		for (const FGroupSearchIndex& GroupSearchIndex : SearchIndex->Groups)
-		{
-			const RowMajorVectorMapConst MapWeights(GroupSearchIndex.Weights.GetData(), 1, NumDimensions);
+		const RowMajorVectorMapConst MapWeights(SearchIndex->Weights.GetData(), 1, NumDimensions);
 
-			// testing the KDTree is returning the proper searches for all the original points transformed in pca space
-			for (int32 PoseIdx = GroupSearchIndex.StartPoseIndex; PoseIdx < GroupSearchIndex.EndPoseIndex; ++PoseIdx)
-			{
-				FKDTree::KNNResultSet ResultSet(ClampedKDTreeQueryNumNeighbors, ResultIndexes, ResultDistanceSqr);
-				TArrayView<const float> PoseValues = SearchIndex->GetPoseValues(PoseIdx);
+		// testing the KDTree is returning the proper searches for all the original points transformed in pca space
+		for (int32 PoseIdx = 0; PoseIdx < SearchIndex->NumPoses; ++PoseIdx)
+		{
+			FKDTree::KNNResultSet ResultSet(ClampedKDTreeQueryNumNeighbors, ResultIndexes, ResultDistanceSqr);
+			TArrayView<const float> PoseValues = SearchIndex->GetPoseValues(PoseIdx);
 
-				const RowMajorVectorMapConst Mean(GroupSearchIndex.Mean.GetData(), 1, NumDimensions);
-				const ColMajorMatrixMapConst PCAProjectionMatrix(GroupSearchIndex.PCAProjectionMatrix.GetData(), NumDimensions, ClampedNumberOfPrincipalComponents);
+			const RowMajorVectorMapConst Mean(SearchIndex->Mean.GetData(), 1, NumDimensions);
+			const ColMajorMatrixMapConst PCAProjectionMatrix(SearchIndex->PCAProjectionMatrix.GetData(), NumDimensions, ClampedNumberOfPrincipalComponents);
 
-				const RowMajorVectorMapConst QueryValues(PoseValues.GetData(), 1, NumDimensions);
-				WeightedQueryValues = QueryValues.array() * MapWeights.array();
-				CenteredQueryValues.noalias() = WeightedQueryValues - Mean;
-				ProjectedQueryValues.noalias() = CenteredQueryValues * PCAProjectionMatrix;
+			const RowMajorVectorMapConst QueryValues(PoseValues.GetData(), 1, NumDimensions);
+			WeightedQueryValues = QueryValues.array() * MapWeights.array();
+			CenteredQueryValues.noalias() = WeightedQueryValues - Mean;
+			ProjectedQueryValues.noalias() = CenteredQueryValues * PCAProjectionMatrix;
 
-				GroupSearchIndex.KDTree.FindNeighbors(ResultSet, ProjectedQueryValues.data());
+			SearchIndex->KDTree.FindNeighbors(ResultSet, ProjectedQueryValues.data());
 
-				size_t ResultIndex = 0;
-				for (; ResultIndex < ResultSet.Num(); ++ResultIndex)
+			size_t ResultIndex = 0;
+			for (; ResultIndex < ResultSet.Num(); ++ResultIndex)
+			{
+				if (PoseIdx == ResultIndexes[ResultIndex])
 				{
-					if ((PoseIdx - GroupSearchIndex.StartPoseIndex) == ResultIndexes[ResultIndex])
-					{
-						check(ResultDistanceSqr[ResultIndex] < UE_KINDA_SMALL_NUMBER);
-						break;
-					}
+					check(ResultDistanceSqr[ResultIndex] < UE_KINDA_SMALL_NUMBER);
+					break;
 				}
-				check(ResultIndex < ResultSet.Num());
 			}
+			check(ResultIndex < ResultSet.Num());
 		}
 	}
 
@@ -2016,7 +1887,6 @@ UE::PoseSearch::FSearchResult UPoseSearchDatabase::SearchPCAKDTree(UE::PoseSearc
 			SearchContext,
 			SearchContext.CurrentResult.PoseIdx,
 			EPoseComparisonFlags::ContinuingPose,
-			SearchIndex->FindAssetForPose(SearchContext.CurrentResult.PoseIdx)->SourceGroupIdx,
 			QueryValues);
 
 		if (GetSkipSearchIfPossible())
@@ -2029,55 +1899,51 @@ UE::PoseSearch::FSearchResult UPoseSearchDatabase::SearchPCAKDTree(UE::PoseSearc
 	// there's no point in performing the search if CurrentBestTotalCost is already better than that
 	if (SearchContext.GetCurrentBestTotalCost() > SearchIndex->MinCostAddend)
 	{
-		for (const FGroupSearchIndex& GroupSearchIndex : SearchIndex->Groups)
-		{
-			// we're offsetting the NonSelectableIdx by "-GroupSearchIndex.StartPoseIndex" to have the indexes in kdtree space rather than in PoseIdx database space (kdtree index 0 is PoseIdx of GroupSearchIndex.StartPoseIndex)
-			const int NonSelectableIdxUsedSize = PopulateNonSelectableIdx(NonSelectableIdxData, NonSelectableIdxDataSize, SearchContext, this, -GroupSearchIndex.StartPoseIndex);
-			TArrayView<size_t> NonSelectableIdx(NonSelectableIdxData, NonSelectableIdxUsedSize);
-			const RowMajorVectorMapConst MapWeights(GroupSearchIndex.Weights.GetData(), 1, NumDimensions);
-			FKDTree::KNNResultSet ResultSet(ClampedKDTreeQueryNumNeighbors, ResultIndexes, ResultDistanceSqr, NonSelectableIdx);
+		// we're offsetting the NonSelectableIdx by "-GroupSearchIndex.StartPoseIndex" to have the indexes in kdtree space rather than in PoseIdx database space (kdtree index 0 is PoseIdx of GroupSearchIndex.StartPoseIndex)
+		const int NonSelectableIdxUsedSize = PopulateNonSelectableIdx(NonSelectableIdxData, NonSelectableIdxDataSize, SearchContext, this);
+		TArrayView<size_t> NonSelectableIdx(NonSelectableIdxData, NonSelectableIdxUsedSize);
+		const RowMajorVectorMapConst MapWeights(SearchIndex->Weights.GetData(), 1, NumDimensions);
+		FKDTree::KNNResultSet ResultSet(ClampedKDTreeQueryNumNeighbors, ResultIndexes, ResultDistanceSqr, NonSelectableIdx);
 
-			check(QueryValues.Num() == NumDimensions);
+		check(QueryValues.Num() == NumDimensions);
 
-			const RowMajorVectorMapConst Mean(GroupSearchIndex.Mean.GetData(), 1, NumDimensions);
-			const ColMajorMatrixMapConst PCAProjectionMatrix(GroupSearchIndex.PCAProjectionMatrix.GetData(), NumDimensions, ClampedNumberOfPrincipalComponents);
+		const RowMajorVectorMapConst Mean(SearchIndex->Mean.GetData(), 1, NumDimensions);
+		const ColMajorMatrixMapConst PCAProjectionMatrix(SearchIndex->PCAProjectionMatrix.GetData(), NumDimensions, ClampedNumberOfPrincipalComponents);
 
-			// transforming query values into PCA space to query the KDTree
-			const RowMajorVectorMapConst QueryValuesMap(QueryValues.GetData(), 1, NumDimensions);
-			WeightedQueryValues = QueryValuesMap.array() * MapWeights.array();
-			CenteredQueryValues.noalias() = WeightedQueryValues - Mean;
-			ProjectedQueryValues.noalias() = CenteredQueryValues * PCAProjectionMatrix;
+		// transforming query values into PCA space to query the KDTree
+		const RowMajorVectorMapConst QueryValuesMap(QueryValues.GetData(), 1, NumDimensions);
+		WeightedQueryValues = QueryValuesMap.array() * MapWeights.array();
+		CenteredQueryValues.noalias() = WeightedQueryValues - Mean;
+		ProjectedQueryValues.noalias() = CenteredQueryValues * PCAProjectionMatrix;
 
-			GroupSearchIndex.KDTree.FindNeighbors(ResultSet, ProjectedQueryValues.data());
+		SearchIndex->KDTree.FindNeighbors(ResultSet, ProjectedQueryValues.data());
 
-			for (size_t ResultIndex = 0; ResultIndex < ResultSet.Num(); ++ResultIndex)
-			{
-				const int32 PoseIdx = ResultIndexes[ResultIndex] + GroupSearchIndex.StartPoseIndex;
+		for (size_t ResultIndex = 0; ResultIndex < ResultSet.Num(); ++ResultIndex)
+		{
+			const int32 PoseIdx = ResultIndexes[ResultIndex];
 
-				const FPoseSearchPoseMetadata& Metadata = SearchIndex->PoseMetadata[PoseIdx];
+			const FPoseSearchPoseMetadata& Metadata = SearchIndex->PoseMetadata[PoseIdx];
 
-				if (EnumHasAnyFlags(Metadata.Flags, EPoseSearchPoseFlags::BlockTransition))
-				{
-					continue;
-				}
+			if (EnumHasAnyFlags(Metadata.Flags, EPoseSearchPoseFlags::BlockTransition))
+			{
+				continue;
+			}
 
-				FPoseSearchCost PoseCost = ComparePoses(
-					SearchContext,
-					PoseIdx,
-					EPoseComparisonFlags::None,
-					GroupSearchIndex.GroupIndex,
-					QueryValues);
+			FPoseSearchCost PoseCost = ComparePoses(
+				SearchContext,
+				PoseIdx,
+				EPoseComparisonFlags::None,
+				QueryValues);
 
-				if (PoseCost < BestPoseCost)
-				{
-					BestPoseCost = PoseCost;
-					BestPoseIdx = PoseIdx;
-				}
+			if (PoseCost < BestPoseCost)
+			{
+				BestPoseCost = PoseCost;
+				BestPoseIdx = PoseIdx;
+			}
 
 #if UE_POSE_SEARCH_TRACE_ENABLED
-				SearchContext.BestCandidates.Add(PoseCost, PoseIdx, this);
+			SearchContext.BestCandidates.Add(PoseCost, PoseIdx, this);
 #endif
-			}
 		}
 
 		if (GetSkipSearchIfPossible() && BestPoseCost.IsValid())
@@ -2133,7 +1999,6 @@ UE::PoseSearch::FSearchResult UPoseSearchDatabase::SearchBruteForce(UE::PoseSear
 				SearchContext,
 				SearchContext.CurrentResult.PoseIdx,
 				EPoseComparisonFlags::ContinuingPose,
-				SearchIndex->FindAssetForPose(SearchContext.CurrentResult.PoseIdx)->SourceGroupIdx,
 				QueryValues);
 
 			if (GetSkipSearchIfPossible())
@@ -2157,14 +2022,6 @@ UE::PoseSearch::FSearchResult UPoseSearchDatabase::SearchBruteForce(UE::PoseSear
 	{
 		for (const FPoseSearchIndexAsset& Asset : SearchIndex->Assets)
 		{
-			if (SearchContext.DatabaseTagQuery)
-			{
-				if (!SearchContext.DatabaseTagQuery->Matches(*GetSourceAssetGroupTags(&Asset)))
-				{
-					continue;
-				}
-			}
-
 			const bool CheckForNonSelectableIdx = IsCurrentResultFromThisDatabase && (&Asset == SearchContext.CurrentResult.SearchIndexAsset);
 			const int32 EndIndex = Asset.FirstPoseIdx + Asset.NumPoses;
 			for (int32 PoseIdx = Asset.FirstPoseIdx; PoseIdx < EndIndex; ++PoseIdx)
@@ -2185,7 +2042,6 @@ UE::PoseSearch::FSearchResult UPoseSearchDatabase::SearchBruteForce(UE::PoseSear
 					SearchContext,
 					PoseIdx, 
 					EPoseComparisonFlags::None,
-					Asset.SourceGroupIdx,
 					QueryValues);
 
 				if (PoseCost < BestPoseCost)
@@ -2273,7 +2129,6 @@ UE::PoseSearch::FSearchResult UPoseSearchDatabaseSet::Search(UE::PoseSearch::FSe
 				SearchContext,
 				SearchContext.CurrentResult.PoseIdx,
 				EPoseComparisonFlags::ContinuingPose,
-				PoseSearchIndexAsset->SourceGroupIdx,
 				QueryValues);
 
 		if (Database->GetSkipSearchIfPossible())
@@ -4528,29 +4383,26 @@ static Eigen::VectorXd ComputeChannelsMeanDeviations(const FPoseSearchIndex* Sea
 	return MeanDeviations;
 }
 
-static void PreprocessGroupSearchIndexWeights(FGroupSearchIndex& GroupSearchIndex, const UPoseSearchDatabase* Database)
+static void PreprocessSearchIndexWeights(FPoseSearchIndex& SearchIndex, const UPoseSearchDatabase* Database)
 {
 	const int32 NumDimensions = Database->Schema->SchemaCardinality;
-	GroupSearchIndex.Weights.Init(1.f, NumDimensions);
+	SearchIndex.Weights.Init(1.f, NumDimensions);
 
 	for (int ChannelIdx = 0; ChannelIdx != Database->Schema->Channels.Num(); ++ChannelIdx)
 	{
 		const UPoseSearchFeatureChannel* Channel = Database->Schema->Channels[ChannelIdx].Get();
-		Channel->FillWeights(GroupSearchIndex.Weights);
+		Channel->FillWeights(SearchIndex.Weights);
 	}
 
 	// normalizing weights: the idea behind this step is to be able to compare poses from databases using different schemas
-	RowMajorVectorMap MapWeights(GroupSearchIndex.Weights.GetData(), 1, NumDimensions);
+	RowMajorVectorMap MapWeights(SearchIndex.Weights.GetData(), 1, NumDimensions);
 	const float WeightsSum = MapWeights.sum();
 	if (!FMath::IsNearlyZero(WeightsSum))
 	{
 		MapWeights *= (1.0f / WeightsSum);
 	}
 
-	// @todo: remove groups to avoid calculating channels min deviations multiple times
-	EPoseSearchDataPreprocessor DataPreprocessor = Database->GetSearchIndex()->Schema->DataPreprocessor;
-
-	if (DataPreprocessor == EPoseSearchDataPreprocessor::Normalize)
+	if (Database->GetSearchIndex()->Schema->DataPreprocessor == EPoseSearchDataPreprocessor::Normalize)
 	{
 		Eigen::VectorXd ChannelsMeanDeviations = ComputeChannelsMeanDeviations(Database->GetSearchIndex());
 
@@ -4559,37 +4411,40 @@ static void PreprocessGroupSearchIndexWeights(FGroupSearchIndex& GroupSearchInde
 			// the idea here is to premultiply the weights by the inverse of the variance (proportional to the square of the deviation) to have a "weighted Mahalanobis" distance
 			const float ChannelsMeanDeviation = ChannelsMeanDeviations[Dimension];
 			const float ChannelsVariance = ChannelsMeanDeviation * ChannelsMeanDeviation;
-			GroupSearchIndex.Weights[Dimension] /= ChannelsVariance;
+			SearchIndex.Weights[Dimension] /= ChannelsVariance;
 		}
 	}
 }
 
 // it calcualtes Mean and PCAProjectionMatrix
-static void PreprocessGroupSearchIndexPCAData(FGroupSearchIndex& GroupSearchIndex, const UPoseSearchDatabase* Database, const float* GroupValues, float* GroupPCAValues)
+static void PreprocessSearchIndexPCAData(FPoseSearchIndex& SearchIndex, const UPoseSearchDatabase* Database)
 {
 	// binding SearchIndex.Values and SearchIndex.PCAValues Eigen row major matrix maps
 	const int32 NumDimensions = Database->Schema->SchemaCardinality;
-	const int32 NumGroupPoses = GroupSearchIndex.EndPoseIndex - GroupSearchIndex.StartPoseIndex;
+	const int32 NumPoses = SearchIndex.NumPoses;
 
-	const RowMajorVectorMapConst MapWeights(GroupSearchIndex.Weights.GetData(), 1, NumDimensions);
-	const RowMajorMatrixMapConst MapGroupValues(GroupValues, NumGroupPoses, NumDimensions);
-	const RowMajorMatrix WeightedGroupValues = MapGroupValues.array().rowwise() * MapWeights.array();
+	const RowMajorVectorMapConst MapWeights(SearchIndex.Weights.GetData(), 1, NumDimensions);
+	const RowMajorMatrixMapConst MapValues(SearchIndex.Values.GetData(), NumPoses, NumDimensions);
+	const RowMajorMatrix WeightedValues = MapValues.array().rowwise() * MapWeights.array();
 	const uint32 NumberOfPrincipalComponents = Database->GetNumberOfPrincipalComponents();
-	RowMajorMatrixMap MapGroupPCAValues(GroupPCAValues, NumGroupPoses, NumberOfPrincipalComponents);
+
+	SearchIndex.PCAValues.Reset();
+	SearchIndex.PCAValues.AddZeroed(NumPoses * NumberOfPrincipalComponents);
+	RowMajorMatrixMap MapPCAValues(SearchIndex.PCAValues.GetData(), NumPoses, NumberOfPrincipalComponents);
 
 	// calculating the mean
-	GroupSearchIndex.Mean.SetNumZeroed(NumDimensions);
-	RowMajorVectorMap Mean(GroupSearchIndex.Mean.GetData(), 1, NumDimensions);
-	Mean = WeightedGroupValues.colwise().mean();
+	SearchIndex.Mean.SetNumZeroed(NumDimensions);
+	RowMajorVectorMap Mean(SearchIndex.Mean.GetData(), 1, NumDimensions);
+	Mean = WeightedValues.colwise().mean();
 
 	// use the mean to center the data points
-	const RowMajorMatrix CenteredGroupValues = WeightedGroupValues.rowwise() - Mean;
+	const RowMajorMatrix CenteredValues = WeightedValues.rowwise() - Mean;
 
 
 	// estimating the covariance matrix (with dimensionality of NumDimensions, NumDimensions)
 	// formula: https://en.wikipedia.org/wiki/Covariance_matrix#Estimation
 	// details: https://en.wikipedia.org/wiki/Estimation_of_covariance_matrices
-	const ColMajorMatrix CovariantMatrix = (CenteredGroupValues.transpose() * CenteredGroupValues) / float(NumGroupPoses - 1);
+	const ColMajorMatrix CovariantMatrix = (CenteredValues.transpose() * CenteredValues) / float(NumPoses - 1);
 	const Eigen::SelfAdjointEigenSolver<ColMajorMatrix> EigenSolver(CovariantMatrix);
 
 	check(EigenSolver.info() == Eigen::Success);
@@ -4600,12 +4455,12 @@ static void PreprocessGroupSearchIndexPCAData(FGroupSearchIndex& GroupSearchInde
 	if (Database->PoseSearchMode == EPoseSearchMode::PCAKDTree_Validate && NumberOfPrincipalComponents == NumDimensions)
 	{
 		const RowMajorVector ReciprocalWeights = MapWeights.cwiseInverse();
-		const RowMajorMatrix ProjectedGroupValues = CenteredGroupValues * EigenVectors;
-		for (Eigen::Index RowIndex = 0; RowIndex < MapGroupValues.rows(); ++RowIndex)
+		const RowMajorMatrix ProjectedValues = CenteredValues * EigenVectors;
+		for (Eigen::Index RowIndex = 0; RowIndex < MapValues.rows(); ++RowIndex)
 		{
-			const RowMajorVector WeightedReconstructedPoint = ProjectedGroupValues.row(RowIndex) * EigenVectors.transpose() + Mean;
+			const RowMajorVector WeightedReconstructedPoint = ProjectedValues.row(RowIndex) * EigenVectors.transpose() + Mean;
 			const RowMajorVector ReconstructedPoint = WeightedReconstructedPoint.array() * ReciprocalWeights.array();
-			const float Error = (ReconstructedPoint - MapGroupValues.row(RowIndex)).squaredNorm();
+			const float Error = (ReconstructedPoint - MapValues.row(RowIndex)).squaredNorm();
 			check(Error < UE_KINDA_SMALL_NUMBER);
 		}
 	}
@@ -4624,8 +4479,8 @@ static void PreprocessGroupSearchIndexPCAData(FGroupSearchIndex& GroupSearchInde
 	});
 
 	// composing the PCA projection matrix with the PCANumComponents most significant EigenVectors
-	GroupSearchIndex.PCAProjectionMatrix.SetNumZeroed(NumDimensions * NumberOfPrincipalComponents);
-	ColMajorMatrixMap PCAProjectionMatrix(GroupSearchIndex.PCAProjectionMatrix.GetData(), NumDimensions, NumberOfPrincipalComponents);
+	SearchIndex.PCAProjectionMatrix.SetNumZeroed(NumDimensions * NumberOfPrincipalComponents);
+	ColMajorMatrixMap PCAProjectionMatrix(SearchIndex.PCAProjectionMatrix.GetData(), NumDimensions, NumberOfPrincipalComponents);
 	float AccumulatedVariance = 0.f;
 	for (size_t PCAComponentIndex = 0; PCAComponentIndex < NumberOfPrincipalComponents; ++PCAComponentIndex)
 	{
@@ -4638,37 +4493,37 @@ static void PreprocessGroupSearchIndexPCAData(FGroupSearchIndex& GroupSearchInde
 	// and explained variance as ratio between AccumulatedVariance and TotalVariance: https://ro-che.info/articles/2017-12-11-pca-explained-variance
 	const float ExplainedVariance = TotalVariance > UE_KINDA_SMALL_NUMBER ? AccumulatedVariance / TotalVariance : 0.f;
 
-	MapGroupPCAValues = CenteredGroupValues * PCAProjectionMatrix;
+	MapPCAValues = CenteredValues * PCAProjectionMatrix;
 
 	if (Database->PoseSearchMode == EPoseSearchMode::PCAKDTree_Validate && NumberOfPrincipalComponents == NumDimensions)
 	{
 		const RowMajorVector ReciprocalWeights = MapWeights.cwiseInverse();
-		for (Eigen::Index RowIndex = 0; RowIndex < MapGroupValues.rows(); ++RowIndex)
+		for (Eigen::Index RowIndex = 0; RowIndex < MapValues.rows(); ++RowIndex)
 		{
-			const RowMajorVector WeightedReconstructedValues = MapGroupPCAValues.row(RowIndex) * PCAProjectionMatrix.transpose() + Mean;
+			const RowMajorVector WeightedReconstructedValues = MapPCAValues.row(RowIndex) * PCAProjectionMatrix.transpose() + Mean;
 			const RowMajorVector ReconstructedValues = WeightedReconstructedValues.array() * ReciprocalWeights.array();
-			const float Error = (ReconstructedValues - MapGroupValues.row(RowIndex)).squaredNorm();
+			const float Error = (ReconstructedValues - MapValues.row(RowIndex)).squaredNorm();
 			check(Error < UE_KINDA_SMALL_NUMBER);
 		}
 	}
 }
 
-static void PreprocessGroupSearchIndexKDTree(FGroupSearchIndex& GroupSearchIndex, const UPoseSearchDatabase* Database, const float* GroupValues, const float* GroupPCAValues)
+static void PreprocessSearchIndexKDTree(FPoseSearchIndex& SearchIndex, const UPoseSearchDatabase* Database)
 {
-	const int32 NumGroupPoses = GroupSearchIndex.EndPoseIndex - GroupSearchIndex.StartPoseIndex;
+	const int32 NumPoses = SearchIndex.NumPoses;
 	const uint32 NumberOfPrincipalComponents = Database->GetNumberOfPrincipalComponents();
-	GroupSearchIndex.KDTree.Construct(NumGroupPoses, NumberOfPrincipalComponents, GroupPCAValues, Database->KDTreeMaxLeafSize);
+	SearchIndex.KDTree.Construct(NumPoses, NumberOfPrincipalComponents, SearchIndex.PCAValues.GetData(), Database->KDTreeMaxLeafSize);
 
 	if (Database->PoseSearchMode == EPoseSearchMode::PCAKDTree_Validate)
 	{
 		// testing the KDTree is returning the proper searches for all the points in pca space
-		for (size_t PointIndex = 0; PointIndex < NumGroupPoses; ++PointIndex)
+		for (size_t PointIndex = 0; PointIndex < NumPoses; ++PointIndex)
 		{
 			constexpr size_t NumResults = 10;
 			size_t ResultIndexes[NumResults + 1] = { 0 };
 			float ResultDistanceSqr[NumResults + 1] = { 0 };
 			FKDTree::KNNResultSet ResultSet(NumResults, ResultIndexes, ResultDistanceSqr);
-			GroupSearchIndex.KDTree.FindNeighbors(ResultSet, &GroupPCAValues[PointIndex * NumberOfPrincipalComponents]);
+			SearchIndex.KDTree.FindNeighbors(ResultSet, &SearchIndex.PCAValues[PointIndex * NumberOfPrincipalComponents]);
 
 			size_t ResultIndex = 0;
 			for (; ResultIndex < ResultSet.Num(); ++ResultIndex)
@@ -4684,23 +4539,23 @@ static void PreprocessGroupSearchIndexKDTree(FGroupSearchIndex& GroupSearchIndex
 
 		// testing the KDTree is returning the proper searches for all the original points transformed in pca space
 		const int32 NumDimensions = Database->Schema->SchemaCardinality;
-		for (size_t PointIndex = 0; PointIndex < NumGroupPoses; ++PointIndex)
+		for (size_t PointIndex = 0; PointIndex < NumPoses; ++PointIndex)
 		{
 			constexpr size_t NumResults = 10;
 			size_t ResultIndexes[NumResults + 1] = { 0 };
 			float ResultDistanceSqr[NumResults + 1] = { 0 };
 			FKDTree::KNNResultSet ResultSet(NumResults, ResultIndexes, ResultDistanceSqr);
 
-			const RowMajorVectorMapConst MapGroupValues(&GroupValues[PointIndex * NumDimensions], 1, NumDimensions);
-			const RowMajorVectorMapConst MapWeights(GroupSearchIndex.Weights.GetData(), 1, NumDimensions);
-			const RowMajorVectorMapConst Mean(GroupSearchIndex.Mean.GetData(), 1, NumDimensions);
-			const ColMajorMatrixMapConst PCAProjectionMatrix(GroupSearchIndex.PCAProjectionMatrix.GetData(), NumDimensions, NumberOfPrincipalComponents);
+			const RowMajorVectorMapConst MapValues(&SearchIndex.Values[PointIndex * NumDimensions], 1, NumDimensions);
+			const RowMajorVectorMapConst MapWeights(SearchIndex.Weights.GetData(), 1, NumDimensions);
+			const RowMajorVectorMapConst Mean(SearchIndex.Mean.GetData(), 1, NumDimensions);
+			const ColMajorMatrixMapConst PCAProjectionMatrix(SearchIndex.PCAProjectionMatrix.GetData(), NumDimensions, NumberOfPrincipalComponents);
 
-			const RowMajorMatrix WeightedGroupValues = MapGroupValues.array() * MapWeights.array();
-			const RowMajorMatrix CenteredGroupValues = WeightedGroupValues - Mean;
-			const RowMajorVector ProjectedGroupValues  = CenteredGroupValues * PCAProjectionMatrix;
+			const RowMajorMatrix WeightedValues = MapValues.array() * MapWeights.array();
+			const RowMajorMatrix CenteredValues = WeightedValues - Mean;
+			const RowMajorVector ProjectedValues  = CenteredValues * PCAProjectionMatrix;
 
-			GroupSearchIndex.KDTree.FindNeighbors(ResultSet, ProjectedGroupValues.data());
+			SearchIndex.KDTree.FindNeighbors(ResultSet, ProjectedValues.data());
 
 			size_t ResultIndex = 0;
 			for (; ResultIndex < ResultSet.Num(); ++ResultIndex)
@@ -4716,36 +4571,11 @@ static void PreprocessGroupSearchIndexKDTree(FGroupSearchIndex& GroupSearchIndex
 	}
 }
 
-static void PreprocessGroupSearchIndex(FPoseSearchIndex& SearchIndex, const UPoseSearchDatabase* Database)
+static void PreprocessSearchIndex(FPoseSearchIndex& SearchIndex, const UPoseSearchDatabase* Database)
 {
-	const uint32 NumberOfPrincipalComponents = Database->GetNumberOfPrincipalComponents();
-	if (NumberOfPrincipalComponents > 0)
-	{
-		const int32 NumDimensions = Database->Schema->SchemaCardinality;
-
-		// preallocating the PCAValues for all the groups
-		SearchIndex.PCAValues.Reset();
-		SearchIndex.PCAValues.AddZeroed(SearchIndex.NumPoses * NumberOfPrincipalComponents);
-		
-		// original serial for loop: 
-		// for (int32 i = 0; i != SearchIndex.Groups.Num(); ++i)
-		ParallelFor(SearchIndex.Groups.Num(), [&](int32 i)
-		{
-			FGroupSearchIndex& GroupSearchIndex = SearchIndex.Groups[i];
-			
-			const float* GroupValues = SearchIndex.Values.GetData() + GroupSearchIndex.StartPoseIndex * NumDimensions;
-			float* GroupPCAValues = SearchIndex.PCAValues.GetData() + GroupSearchIndex.StartPoseIndex * NumberOfPrincipalComponents;
-
-			PreprocessGroupSearchIndexWeights(GroupSearchIndex, Database);
-			PreprocessGroupSearchIndexPCAData(GroupSearchIndex, Database, GroupValues, GroupPCAValues);
-			PreprocessGroupSearchIndexKDTree(GroupSearchIndex, Database, GroupValues, GroupPCAValues);
-		});
-	}
-	else
-	{
-		// we don't need groups
-		SearchIndex.Groups.Reset();
-	}
+	PreprocessSearchIndexWeights(SearchIndex, Database);
+	PreprocessSearchIndexPCAData(SearchIndex, Database);
+	PreprocessSearchIndexKDTree(SearchIndex, Database);
 }
 
 bool BuildIndex(const UAnimSequence* Sequence, UPoseSearchSequenceMetaData* SequenceMetaData)
@@ -4800,8 +4630,8 @@ bool BuildIndex(const UAnimSequence* Sequence, UPoseSearchSequenceMetaData* Sequ
 	SequenceMetaData->SearchIndex.Assets.Add(SearchIndexAsset);
 	SequenceMetaData->SearchIndex.PoseMetadata = Indexer.Output.PoseMetadata;
 
-	// todo: do we need to initialize per group data?
-	// PreprocessGroupSearchIndex(&SequenceMetaData->SearchIndex, Database);
+	// todo: do we need to PreprocessSearchIndex?
+	// PreprocessSearchIndex(&SequenceMetaData->SearchIndex, Database);
 	return true;
 }
 
@@ -5014,59 +4844,32 @@ void FDatabaseIndexingContext::JoinIndex()
 	int32 TotalFloats = 0;
 
 	check(SearchIndex);
-	SearchIndex->Groups.Reset();
 
-	if (SearchIndex->Assets.Num() > 0)
-	{
-		SearchIndex->Groups.AddDefaulted();
-		SearchIndex->Groups.Last().GroupIndex = SearchIndex->Assets[0].SourceGroupIdx;
-
-		for (int32 AssetIdx = 0; AssetIdx != SearchIndex->Assets.Num(); ++AssetIdx)
-		{
-			const FAssetIndexer::FOutput& Output = Indexers[AssetIdx].Output;
-
-			FPoseSearchIndexAsset& SearchIndexAsset = SearchIndex->Assets[AssetIdx];
-			
-			if (SearchIndexAsset.SourceGroupIdx != SearchIndex->Groups.Last().GroupIndex)
-			{
-				// making sure groups are sorted correctly in ascending order
-				check(SearchIndexAsset.SourceGroupIdx > SearchIndex->Groups.Last().GroupIndex);
-				// finalizing the previous group before adding a new one
-				SearchIndex->Groups.Last().EndPoseIndex = TotalPoses;
-
-				SearchIndex->Groups.AddDefaulted();
-				SearchIndex->Groups.Last().GroupIndex = SearchIndexAsset.SourceGroupIdx;
-				SearchIndex->Groups.Last().StartPoseIndex = TotalPoses;
-			}
-			
-			SearchIndexAsset.NumPoses = Output.NumIndexedPoses;
-			SearchIndexAsset.FirstPoseIdx = TotalPoses;
-
-			TotalPoses += Output.NumIndexedPoses;
-			TotalFloats += Output.FeatureVectorTable.Num();
-		}
-		
-		// finalizing the last inserted group
-		SearchIndex->Groups.Last().EndPoseIndex = TotalPoses;
-	}
-	
-	check(TotalFloats == TotalPoses * Database->Schema->SchemaCardinality);
-	
 	// Join animation data into a single search index
-	SearchIndex->Values.Reset(TotalFloats);
-	SearchIndex->PoseMetadata.Reset(TotalPoses);
+	SearchIndex->Values.Reset();
+	SearchIndex->PoseMetadata.Reset();
 	SearchIndex->PCAValues.Reset();
-	
-	for (const FAssetIndexer& Indexer : Indexers)
+
+	for (int32 AssetIdx = 0; AssetIdx != SearchIndex->Assets.Num(); ++AssetIdx)
 	{
-		const FAssetIndexer::FOutput& Output = Indexer.Output;
+		const FAssetIndexer::FOutput& Output = Indexers[AssetIdx].Output;
+
+		FPoseSearchIndexAsset& SearchIndexAsset = SearchIndex->Assets[AssetIdx];
+		SearchIndexAsset.NumPoses = Output.NumIndexedPoses;
+		SearchIndexAsset.FirstPoseIdx = TotalPoses;
+
 		SearchIndex->Values.Append(Output.FeatureVectorTable.GetData(), Output.FeatureVectorTable.Num());
 		SearchIndex->PoseMetadata.Append(Output.PoseMetadata);
-	}
 
+		TotalPoses += Output.NumIndexedPoses;
+		TotalFloats += Output.FeatureVectorTable.Num();
+	}
+	
 	SearchIndex->NumPoses = TotalPoses;
 	SearchIndex->Schema = Database->Schema;
 	SearchIndex->MinCostAddend = CalculateMinCostAddend();
+
+	check(TotalFloats == TotalPoses * Database->Schema->SchemaCardinality);
 }
 
 bool BuildIndex(UPoseSearchDatabase* Database, FPoseSearchIndex& OutSearchIndex)
@@ -5099,7 +4902,7 @@ bool BuildIndex(UPoseSearchDatabase* Database, FPoseSearchIndex& OutSearchIndex)
 	bSuccess &= DbIndexingContext.IndexAssets();
 	DbIndexingContext.JoinIndex();
 
-	PreprocessGroupSearchIndex(OutSearchIndex, Database);
+	PreprocessSearchIndex(OutSearchIndex, Database);
 	
 	return bSuccess;
 }
diff --git a/Engine/Plugins/Experimental/Animation/PoseSearch/Source/Runtime/Private/PoseSearchDerivedData.cpp b/Engine/Plugins/Experimental/Animation/PoseSearch/Source/Runtime/Private/PoseSearchDerivedData.cpp
--- a/Engine/Plugins/Experimental/Animation/PoseSearch/Source/Runtime/Private/PoseSearchDerivedData.cpp
+++ b/Engine/Plugins/Experimental/Animation/PoseSearch/Source/Runtime/Private/PoseSearchDerivedData.cpp
@@ -226,21 +226,18 @@ namespace UE::PoseSearch
 		int32 NumValues = 0;
 		int32 NumPCAValues = 0;
 		int32 NumAssets = 0;
-		int32 NumGroups = 0;
 
 		if (Ar.IsSaving())
 		{
 			NumValues = Index.Values.Num();
 			NumPCAValues = Index.PCAValues.Num();
 			NumAssets = Index.Assets.Num();
-			NumGroups = Index.Groups.Num();
 		}
 
 		Ar << Index.NumPoses;
 		Ar << NumValues;
 		Ar << NumPCAValues;
 		Ar << NumAssets;
-		Ar << NumGroups;
 
 		if (Ar.IsLoading())
 		{
@@ -248,7 +245,6 @@ namespace UE::PoseSearch
 			Index.PCAValues.SetNumUninitialized(NumPCAValues);
 			Index.PoseMetadata.SetNumUninitialized(Index.NumPoses);
 			Index.Assets.SetNumUninitialized(NumAssets);
-			Index.Groups.SetNum(NumGroups);
 		}
 
 		if (Index.Values.Num() > 0)
@@ -271,22 +267,14 @@ namespace UE::PoseSearch
 			Ar.Serialize(&Index.Assets[0], Index.Assets.Num() * Index.Assets.GetTypeSize());
 		}
 
-		for (int i = 0; i < Index.Groups.Num(); ++i)
-		{
-			FGroupSearchIndex& GroupSearchIndex = Index.Groups[i];
-			Ar << GroupSearchIndex.StartPoseIndex;
-			Ar << GroupSearchIndex.EndPoseIndex;
-			Ar << GroupSearchIndex.GroupIndex;
-			Ar << GroupSearchIndex.Weights;
-			Ar << GroupSearchIndex.Mean;
-			Ar << GroupSearchIndex.PCAProjectionMatrix;
-
-			check(GroupSearchIndex.PCAProjectionMatrix.Num() > 0 && GroupSearchIndex.Mean.Num() > 0);
-			check(GroupSearchIndex.PCAProjectionMatrix.Num() % GroupSearchIndex.Mean.Num() == 0);
-			const int NumberOfPrincipalComponents = GroupSearchIndex.PCAProjectionMatrix.Num() / GroupSearchIndex.Mean.Num();
-			float* GroupPCAValues = Index.PCAValues.GetData() + GroupSearchIndex.StartPoseIndex * NumberOfPrincipalComponents;
-			Serialize(Ar, GroupSearchIndex.KDTree, GroupPCAValues);
-		}
+		Ar << Index.Weights;
+		Ar << Index.Mean;
+		Ar << Index.PCAProjectionMatrix;
+
+		check(Index.PCAProjectionMatrix.Num() > 0 && Index.Mean.Num() > 0);
+		check(Index.PCAProjectionMatrix.Num() % Index.Mean.Num() == 0);
+		const int NumberOfPrincipalComponents = Index.PCAProjectionMatrix.Num() / Index.Mean.Num();
+		Serialize(Ar, Index.KDTree, Index.PCAValues.GetData());
 
 		return Ar;
 	}
diff --git a/Engine/Plugins/Experimental/Animation/PoseSearch/Source/Runtime/Private/PoseSearchLibrary.cpp b/Engine/Plugins/Experimental/Animation/PoseSearch/Source/Runtime/Private/PoseSearchLibrary.cpp
--- a/Engine/Plugins/Experimental/Animation/PoseSearch/Source/Runtime/Private/PoseSearchLibrary.cpp
+++ b/Engine/Plugins/Experimental/Animation/PoseSearch/Source/Runtime/Private/PoseSearchLibrary.cpp
@@ -13,56 +13,6 @@
 
 #define LOCTEXT_NAMESPACE "PoseSearchLibrary"
 
-namespace UE::PoseSearch
-{
-static void ComputeDatabaseSequenceFilter(
-	const UPoseSearchDatabase* Database, 
-	const FGameplayTagQuery* Query, 
-	TArray<bool>& OutDbSequenceFilter)
-{
-	OutDbSequenceFilter.SetNum(Database->Sequences.Num());
-
-	if (Query)
-	{
-		for (int SeqIdx = 0; SeqIdx < Database->Sequences.Num(); ++SeqIdx)
-		{
-			OutDbSequenceFilter[SeqIdx] = Query->Matches(Database->Sequences[SeqIdx].GroupTags);
-		}
-	}
-	else
-	{
-		for (int SeqIdx = 0; SeqIdx < Database->Sequences.Num(); ++SeqIdx)
-		{
-			OutDbSequenceFilter[SeqIdx] = true;
-		}
-	}
-}
-
-static void ComputeDatabaseBlendSpaceFilter(
-	const UPoseSearchDatabase* Database,
-	const FGameplayTagQuery* Query,
-	TArray<bool>& OutDbBlendSpaceFilter)
-{
-	OutDbBlendSpaceFilter.SetNum(Database->BlendSpaces.Num());
-
-	if (Query)
-	{
-		for (int BlendSpaceIdx = 0; BlendSpaceIdx < Database->BlendSpaces.Num(); ++BlendSpaceIdx)
-		{
-			OutDbBlendSpaceFilter[BlendSpaceIdx] = Query->Matches(Database->BlendSpaces[BlendSpaceIdx].GroupTags);
-		}
-	}
-	else
-	{
-		for (int BlendSpaceIdx = 0; BlendSpaceIdx < Database->BlendSpaces.Num(); ++BlendSpaceIdx)
-		{
-			OutDbBlendSpaceFilter[BlendSpaceIdx] = true;
-		}
-	}
-}
-
-} // namespace UE::PoseSearch
-
 //////////////////////////////////////////////////////////////////////////
 // FMotionMatchingState
 
@@ -219,7 +169,6 @@ static void TraceMotionMatchingState(
 	const UPoseSearchSearchableAsset* Searchable,
 	UE::PoseSearch::FSearchContext& SearchContext,
 	const FMotionMatchingState& MotionMatchingState,
-	const FGameplayTagQuery* DatabaseTagQuery,
 	const FTrajectorySampleRange& Trajectory,
 	const UE::PoseSearch::FSearchResult& LastResult
 )
@@ -359,18 +308,6 @@ static void TraceMotionMatchingState(
 		AnimAngularVelocity = 0.0f;
 	}
 
-	TArray<bool> DatabaseSequenceFilter;
-	ComputeDatabaseSequenceFilter(
-		MotionMatchingState.CurrentSearchResult.Database.Get(),
-		DatabaseTagQuery, 
-		DatabaseSequenceFilter);
-
-	TArray<bool> DatabaseBlendSpaceFilter;
-	ComputeDatabaseBlendSpaceFilter(
-		MotionMatchingState.CurrentSearchResult.Database.Get(),
-		DatabaseTagQuery, 
-		DatabaseBlendSpaceFilter);
-
 	if (EnumHasAnyFlags(MotionMatchingState.Flags, EMotionMatchingFlags::JumpedToFollowUp))
 	{
 		TraceState.Flags |= FTraceMotionMatchingState::EFlags::FollowupAnimation;
@@ -385,8 +322,6 @@ static void TraceMotionMatchingState(
 	TraceState.SimAngularVelocity = SimAngularVelocity;
 	TraceState.AnimLinearVelocity = AnimLinearVelocity;
 	TraceState.AnimAngularVelocity = AnimAngularVelocity;
-	TraceState.DatabaseSequenceFilter = DatabaseSequenceFilter;
-	TraceState.DatabaseBlendSpaceFilter = DatabaseBlendSpaceFilter;
 
 	TraceState.Output(UpdateContext);
 #endif
@@ -395,7 +330,6 @@ static void TraceMotionMatchingState(
 void UpdateMotionMatchingState(
 	const FAnimationUpdateContext& Context,
 	const UPoseSearchSearchableAsset* Searchable,
-	const FGameplayTagQuery* DatabaseTagQuery,
 	const FGameplayTagContainer* ActiveTagsContainer,
 	const FTrajectorySampleRange& Trajectory,
 	const FMotionMatchingSettings& Settings,
@@ -433,7 +367,6 @@ void UpdateMotionMatchingState(
 	if (!bCanAdvance || (InOutMotionMatchingState.ElapsedPoseJumpTime >= Settings.SearchThrottleTime))
 	{
 		// Build the search context
-		SearchContext.DatabaseTagQuery = DatabaseTagQuery;
 		SearchContext.ActiveTagsContainer = ActiveTagsContainer;
 		SearchContext.Trajectory = &Trajectory;
 		SearchContext.OwningComponent = Context.AnimInstanceProxy->GetSkelMeshComponent();
@@ -506,7 +439,6 @@ void UpdateMotionMatchingState(
 		Searchable,
 		SearchContext,
 		InOutMotionMatchingState,
-		DatabaseTagQuery,
 		Trajectory,
 		LastResult
 	);
diff --git a/Engine/Plugins/Experimental/Animation/PoseSearch/Source/Runtime/Private/Trace/PoseSearchTraceLogger.cpp b/Engine/Plugins/Experimental/Animation/PoseSearch/Source/Runtime/Private/Trace/PoseSearchTraceLogger.cpp
--- a/Engine/Plugins/Experimental/Animation/PoseSearch/Source/Runtime/Private/Trace/PoseSearchTraceLogger.cpp
+++ b/Engine/Plugins/Experimental/Animation/PoseSearch/Source/Runtime/Private/Trace/PoseSearchTraceLogger.cpp
@@ -70,8 +70,6 @@ FArchive& operator<<(FArchive& Ar, FTraceMotionMatchingState& State)
 	Ar << State.ElapsedPoseJumpTime;
 	Ar << State.Flags;
 	Ar << State.QueryVector;
-	Ar << State.DatabaseSequenceFilter;
-	Ar << State.DatabaseBlendSpaceFilter;
 	Ar << State.AssetPlayerTime;
 	Ar << State.DeltaTime;
 	Ar << State.SimLinearVelocity;
diff --git a/Engine/Plugins/Experimental/Animation/PoseSearch/Source/Runtime/Private/Trace/PoseSearchTraceLogger.h b/Engine/Plugins/Experimental/Animation/PoseSearch/Source/Runtime/Private/Trace/PoseSearchTraceLogger.h
--- a/Engine/Plugins/Experimental/Animation/PoseSearch/Source/Runtime/Private/Trace/PoseSearchTraceLogger.h
+++ b/Engine/Plugins/Experimental/Animation/PoseSearch/Source/Runtime/Private/Trace/PoseSearchTraceLogger.h
@@ -98,12 +98,6 @@ struct POSESEARCH_API FTraceMotionMatchingState
 	/** Search vectors in normalized and unnormalized forms */
 	TArray<float> QueryVector;
 
-	/** If true, groups are being filtered by DatabaseGroupQuery */
-	TArray<bool> DatabaseSequenceFilter;
-
-	/** If true, groups are being filtered by DatabaseGroupQuery */
-	TArray<bool> DatabaseBlendSpaceFilter;
-
 	float AssetPlayerTime = 0.0f;
 	float DeltaTime = 0.0f;
 	float SimLinearVelocity = 0.0f;
diff --git a/Engine/Plugins/Experimental/Animation/PoseSearch/Source/Runtime/Public/PoseSearch/AnimNode_MotionMatching.h b/Engine/Plugins/Experimental/Animation/PoseSearch/Source/Runtime/Public/PoseSearch/AnimNode_MotionMatching.h
--- a/Engine/Plugins/Experimental/Animation/PoseSearch/Source/Runtime/Public/PoseSearch/AnimNode_MotionMatching.h
+++ b/Engine/Plugins/Experimental/Animation/PoseSearch/Source/Runtime/Public/PoseSearch/AnimNode_MotionMatching.h
@@ -24,13 +24,6 @@ public:
 	UPROPERTY(EditAnywhere, BlueprintReadWrite, Category=Settings, meta=(PinShownByDefault))
 	TObjectPtr<const UPoseSearchSearchableAsset> Searchable = nullptr;
 
-	UPROPERTY(EditAnywhere, BlueprintReadWrite, Category = Settings, meta = (PinHiddenByDefault))
-	bool bUseDatabaseTagQuery = false;
-
-	// Query used to filter database groups which can be searched
-	UPROPERTY(EditAnywhere, BlueprintReadWrite, Category = Settings, meta = (PinShownByDefault))
-	FGameplayTagQuery DatabaseTagQuery;
-
 	UPROPERTY(EditAnywhere, BlueprintReadWrite, Category = Settings, meta = (PinShownByDefault))
 	FGameplayTagContainer ActiveTagsContainer;
 
diff --git a/Engine/Plugins/Experimental/Animation/PoseSearch/Source/Runtime/Public/PoseSearch/PoseSearch.h b/Engine/Plugins/Experimental/Animation/PoseSearch/Source/Runtime/Public/PoseSearch/PoseSearch.h
--- a/Engine/Plugins/Experimental/Animation/PoseSearch/Source/Runtime/Public/PoseSearch/PoseSearch.h
+++ b/Engine/Plugins/Experimental/Animation/PoseSearch/Source/Runtime/Public/PoseSearch/PoseSearch.h
@@ -598,13 +598,11 @@ public:
 
 	FPoseSearchIndexAsset(
 		ESearchIndexAssetType InType,
-		int32 InSourceGroupIdx, 
 		int32 InSourceAssetIdx, 
 		bool bInMirrored, 
 		const FFloatInterval& InSamplingInterval,
 		FVector InBlendParameters = FVector::Zero())
 		: Type(InType)
-		, SourceGroupIdx(InSourceGroupIdx)
 		, SourceAssetIdx(InSourceAssetIdx)
 		, bMirrored(bInMirrored)
 		, BlendParameters(InBlendParameters)
@@ -616,9 +614,6 @@ public:
 	UPROPERTY()
 	ESearchIndexAssetType Type = ESearchIndexAssetType::Sequence;
 
-	UPROPERTY()
-	int32 SourceGroupIdx = INDEX_NONE;
-
 	// Index of the source asset in search index's container (i.e. UPoseSearchDatabase)
 	UPROPERTY()
 	int32 SourceAssetIdx = INDEX_NONE;
@@ -644,32 +639,6 @@ public:
 	}
 };
 
-USTRUCT()
-struct POSESEARCH_API FGroupSearchIndex
-{
-	GENERATED_BODY()
-
-	UE::PoseSearch::FKDTree KDTree;
-
-	UPROPERTY()
-	TArray<float> PCAProjectionMatrix;
-
-	UPROPERTY()
-	TArray<float> Mean;
-
-	UPROPERTY()
-	int32 StartPoseIndex = 0;
-	
-	UPROPERTY()
-	int32 EndPoseIndex = 0;
-
-	UPROPERTY()
-	int32 GroupIndex = 0;
-
-	UPROPERTY()
-	TArray<float> Weights;
-};
-
 /**
 * A search index for animation poses. The structure of the search index is determined by its UPoseSearchSchema.
 * May represent a single animation (see UPoseSearchSequenceMetaData) or a collection (see UPoseSearchDatabase).
@@ -688,8 +657,16 @@ struct POSESEARCH_API FPoseSearchIndex
 	UPROPERTY()
 	TArray<float> PCAValues;
 
+	UE::PoseSearch::FKDTree KDTree;
+
+	UPROPERTY()
+	TArray<float> PCAProjectionMatrix;
+
 	UPROPERTY()
-	TArray<FGroupSearchIndex> Groups;
+	TArray<float> Mean;
+
+	UPROPERTY()
+	TArray<float> Weights;
 
 	UPROPERTY()
 	TArray<FPoseSearchPoseMetadata> PoseMetadata;
@@ -711,7 +688,6 @@ struct POSESEARCH_API FPoseSearchIndex
 	TArrayView<const float> GetPoseValues(int32 PoseIdx) const;
 
 	int32 FindAssetIndex(const FPoseSearchIndexAsset* Asset) const;
-	const FGroupSearchIndex* FindGroup(int32 GroupIndex) const;
 
 	const FPoseSearchIndexAsset* FindAssetForPose(int32 PoseIdx) const;
 	float GetAssetTime(int32 PoseIdx, const FPoseSearchIndexAsset* Asset) const;
@@ -769,9 +745,6 @@ struct POSESEARCH_API FPoseSearchDatabaseSequence : public FPoseSearchDatabaseAn
 	UPROPERTY(EditAnywhere, Category="Sequence")
 	TObjectPtr<UAnimSequence> FollowUpSequence = nullptr;
 
-	UPROPERTY(EditAnywhere, Category = "Group")
-	FGameplayTagContainer GroupTags;
-
 	FFloatInterval GetEffectiveSamplingRange() const;
 
 	virtual UAnimationAsset* GetAnimationAsset() const override;
@@ -804,9 +777,6 @@ struct POSESEARCH_API FPoseSearchDatabaseBlendSpace : public FPoseSearchDatabase
 	UPROPERTY(EditAnywhere, Category = "BlendSpace", meta = (ClampMin = "1", UIMin = "1", UIMax = "25"))
 	int32 NumberOfVerticalSamples = 5;
 
-	UPROPERTY(EditAnywhere, Category = "Group")
-	FGameplayTagContainer GroupTags;
-
 	virtual UAnimationAsset* GetAnimationAsset() const override;
 	virtual bool IsLooping() const override;
 
@@ -821,16 +791,6 @@ public:
 		float& VerticalBlendMax) const;
 };
 
-USTRUCT()
-struct POSESEARCH_API FPoseSearchDatabaseGroup
-{
-	GENERATED_BODY()
-
-public:
-	UPROPERTY(EditAnywhere, Category = "Settings")
-	FGameplayTag Tag;
-};
-
 USTRUCT()
 struct POSESEARCH_API FPoseSearchCost
 {
@@ -1081,9 +1041,6 @@ public:
 	UPROPERTY(EditAnywhere, Category = "Database")
 	FPoseSearchBlockTransitionParameters BlockTransitionParameters = { 0.0f, 0.2f };
 
-	UPROPERTY(EditAnywhere, Category = "Database")
-	TArray<FPoseSearchDatabaseGroup> Groups;
-
 	// Drag and drop animations here to add them in bulk to Sequences
 	UPROPERTY(EditAnywhere, Category = "Database", DisplayName="Drag And Drop Anims Here")
 	TArray<TObjectPtr<UAnimSequence>> SimpleSequences;
@@ -1132,7 +1089,6 @@ public:
 	const FPoseSearchDatabaseSequence& GetSequenceSourceAsset(const FPoseSearchIndexAsset* SearchIndexAsset) const;
 	const FPoseSearchDatabaseBlendSpace& GetBlendSpaceSourceAsset(const FPoseSearchIndexAsset* SearchIndexAsset) const;
 	const bool IsSourceAssetLooping(const FPoseSearchIndexAsset* SearchIndexAsset) const;
-	const FGameplayTagContainer* GetSourceAssetGroupTags(const FPoseSearchIndexAsset* SearchIndexAsset) const;
 	const FString GetSourceAssetName(const FPoseSearchIndexAsset* SearchIndexAsset) const;
 	int32 GetNumberOfPrincipalComponents() const;
 	
@@ -1210,7 +1166,7 @@ public:
 
 	void BuildQuery(UE::PoseSearch::FSearchContext& SearchContext, FPoseSearchFeatureVectorBuilder& OutQuery) const;
 
-	FPoseSearchCost ComparePoses(UE::PoseSearch::FSearchContext& SearchContext, int32 PoseIdx, UE::PoseSearch::EPoseComparisonFlags PoseComparisonFlags, int32 GroupIdx, const TArrayView<const float>& QueryValues) const;
+	FPoseSearchCost ComparePoses(UE::PoseSearch::FSearchContext& SearchContext, int32 PoseIdx, UE::PoseSearch::EPoseComparisonFlags PoseComparisonFlags, const TArrayView<const float>& QueryValues) const;
 
 protected:
 	UE::PoseSearch::FSearchResult SearchPCAKDTree(UE::PoseSearch::FSearchContext& SearchContext) const;
@@ -1357,7 +1313,6 @@ struct FPoseIndicesHistory
 struct POSESEARCH_API FSearchContext
 {
 	EPoseSearchBooleanRequest QueryMirrorRequest = EPoseSearchBooleanRequest::Indifferent;
-	const FGameplayTagQuery* DatabaseTagQuery = nullptr;
 	UE::PoseSearch::FDebugDrawParams DebugDrawParams;
 	UE::PoseSearch::FPoseHistory* History = nullptr;
 	const FTrajectorySampleRange* Trajectory = nullptr;
diff --git a/Engine/Plugins/Experimental/Animation/PoseSearch/Source/Runtime/Public/PoseSearch/PoseSearchLibrary.h b/Engine/Plugins/Experimental/Animation/PoseSearch/Source/Runtime/Public/PoseSearch/PoseSearchLibrary.h
--- a/Engine/Plugins/Experimental/Animation/PoseSearch/Source/Runtime/Public/PoseSearch/PoseSearchLibrary.h
+++ b/Engine/Plugins/Experimental/Animation/PoseSearch/Source/Runtime/Public/PoseSearch/PoseSearchLibrary.h
@@ -117,7 +117,6 @@ struct POSESEARCH_API FMotionMatchingState
 POSESEARCH_API void UpdateMotionMatchingState(
 	const FAnimationUpdateContext& Context,
 	const UPoseSearchSearchableAsset* Searchable,
-	const FGameplayTagQuery* DatabaseTagQuery,
 	const FGameplayTagContainer* ActiveTagsContainer,
 	const FTrajectorySampleRange& Trajectory,
 	const FMotionMatchingSettings& Settings,
