Update PoseSearchDebugger to start and stop tracing when recordings are started and stopped

diff --git a/Engine/Plugins/Experimental/Animation/PoseSearch/Source/Editor/Private/PoseSearchDebugger.cpp b/Engine/Plugins/Experimental/Animation/PoseSearch/Source/Editor/Private/PoseSearchDebugger.cpp
--- a/Engine/Plugins/Experimental/Animation/PoseSearch/Source/Editor/Private/PoseSearchDebugger.cpp
+++ b/Engine/Plugins/Experimental/Animation/PoseSearch/Source/Editor/Private/PoseSearchDebugger.cpp
@@ -2660,7 +2660,6 @@ void FDebugger::Initialize()
 {
 	Debugger = new FDebugger;
 	IModularFeatures::Get().RegisterModularFeature(IRewindDebuggerExtension::ModularFeatureName, Debugger);
-
 }
 
 void FDebugger::Shutdown()
@@ -2669,6 +2668,16 @@ void FDebugger::Shutdown()
 	delete Debugger;
 }
 
+void FDebugger::RecordingStarted(IRewindDebugger*)
+{
+	UE::Trace::ToggleChannel(TEXT("PoseSearch"), true);
+}
+
+void FDebugger::RecordingStopped(IRewindDebugger*)
+{
+	UE::Trace::ToggleChannel(TEXT("PoseSearch"), false);
+}
+
 bool FDebugger::IsPIESimulating()
 {
 	return Debugger->RewindDebugger->IsPIESimulating();
diff --git a/Engine/Plugins/Experimental/Animation/PoseSearch/Source/Editor/Private/PoseSearchDebugger.h b/Engine/Plugins/Experimental/Animation/PoseSearch/Source/Editor/Private/PoseSearchDebugger.h
--- a/Engine/Plugins/Experimental/Animation/PoseSearch/Source/Editor/Private/PoseSearchDebugger.h
+++ b/Engine/Plugins/Experimental/Animation/PoseSearch/Source/Editor/Private/PoseSearchDebugger.h
@@ -646,6 +646,9 @@ public:
 	virtual void Update(float DeltaTime, IRewindDebugger* InRewindDebugger) override;
 	virtual ~FDebugger() = default;
 
+	virtual void RecordingStarted(IRewindDebugger* RewindDebugger);
+	virtual void RecordingStopped(IRewindDebugger* RewindDebugger);
+
 	static FDebugger* Get() { return Debugger; }
 	static void Initialize();
 	static void Shutdown();
diff --git a/Engine/Plugins/Experimental/Animation/PoseSearch/Source/Runtime/Private/PoseSearch.cpp b/Engine/Plugins/Experimental/Animation/PoseSearch/Source/Runtime/Private/PoseSearch.cpp
--- a/Engine/Plugins/Experimental/Animation/PoseSearch/Source/Runtime/Private/PoseSearch.cpp
+++ b/Engine/Plugins/Experimental/Animation/PoseSearch/Source/Runtime/Private/PoseSearch.cpp
@@ -4929,11 +4929,6 @@ void FModule::StartupModule()
 {
 	IModularFeatures::Get().RegisterModularFeature(UE::Anim::IPoseSearchProvider::ModularFeatureName, this);
 
-#if UE_POSE_SEARCH_TRACE_ENABLED
-	// Enable the PoseSearch trace channel
-	UE::Trace::ToggleChannel(*FTraceLogger::Name.ToString(), true);
-#endif
-
 #if WITH_EDITOR
 	FCoreUObjectDelegates::OnObjectPreSave.AddRaw(this, &FModule::OnObjectSaved);
 #endif // WITH_EDITOR
