motion matching - fix for UPoseSearchDatabaseSet::Search crash when editing schemas while in PIE

diff --git a/Engine/Plugins/Experimental/Animation/PoseSearch/Source/Runtime/Private/PoseSearch.cpp b/Engine/Plugins/Experimental/Animation/PoseSearch/Source/Runtime/Private/PoseSearch.cpp
--- a/Engine/Plugins/Experimental/Animation/PoseSearch/Source/Runtime/Private/PoseSearch.cpp
+++ b/Engine/Plugins/Experimental/Animation/PoseSearch/Source/Runtime/Private/PoseSearch.cpp
@@ -862,12 +862,32 @@ FPoseSearchIndex* UPoseSearchDatabase::GetSearchIndex()
 
 const FPoseSearchIndex* UPoseSearchDatabase::GetSearchIndex() const
 {
-	if (PrivateDerivedData == nullptr)
+	return PrivateDerivedData ? &PrivateDerivedData->SearchIndex : nullptr;
+}
+
+const FPoseSearchIndex* UPoseSearchDatabase::GetSearchIndexSafe() const
+{
+	const FPoseSearchIndex* SearchIndex = GetSearchIndex();
+
+	if (!SearchIndex || !SearchIndex->IsValid() || SearchIndex->IsEmpty())
 	{
-		return nullptr;
+		if (!Schema)
+		{
+			UE_LOG(LogAnimation, Warning, TEXT("UPoseSearchDatabase %s failed to index. Reason: no Schema!"), *GetName());
+		}
+		else if (!Schema->IsValid())
+		{
+			UE_LOG(LogAnimation, Warning, TEXT("UPoseSearchDatabase %s failed to index. Reason: Schema %s is invalid"), *GetName(), *Schema->GetName());
+		}
+		else
+		{
+			UE_LOG(LogAnimation, Warning, TEXT("UPoseSearchDatabase %s failed to index. Reason: is there any unsaved modified asset?"), *GetName());
+		}
+
+		SearchIndex = nullptr;
 	}
 
-	return &PrivateDerivedData->SearchIndex;
+	return SearchIndex;
 }
 
 float UPoseSearchDatabase::GetAssetTime(int32 PoseIdx, const FPoseSearchIndexAsset* Asset) const
@@ -1761,29 +1781,12 @@ UE::PoseSearch::FSearchResult UPoseSearchDatabase::Search(UE::PoseSearch::FSearc
 	}
 #endif
 
-	const FPoseSearchIndex* SearchIndex = GetSearchIndex();
+	const FPoseSearchIndex* SearchIndex = GetSearchIndexSafe();
 	if (!SearchIndex)
 	{
 		return Result;
 	}
 
-	if (!SearchIndex->IsValid() || SearchIndex->IsEmpty())
-	{
-		if (!Schema)
-		{
-			UE_LOG(LogAnimation, Warning, TEXT("UPoseSearchDatabase %s failed to index. Reason: no Schema!"), *GetName());
-		}
-		else if(!Schema->IsValid())
-		{
-			UE_LOG(LogAnimation, Warning, TEXT("UPoseSearchDatabase %s failed to index. Reason: Schema %s is invalid"), *GetName(), *Schema->GetName());
-		}
-		else
-		{
-			UE_LOG(LogAnimation, Warning, TEXT("UPoseSearchDatabase %s failed to index. Reason: is there any unsaved modified asset?"), *GetName());
-		}
-		return Result;
-	}
-
 	if (PoseSearchMode == EPoseSearchMode::BruteForce || PoseSearchMode == EPoseSearchMode::PCAKDTree_Compare)
 	{
 		Result = SearchBruteForce(SearchContext);
@@ -2115,25 +2118,26 @@ UE::PoseSearch::FSearchResult UPoseSearchDatabaseSet::Search(UE::PoseSearch::FSe
 	{
 		const UPoseSearchDatabase* Database = SearchContext.CurrentResult.Database.Get();
 		check(Database);
-		const FPoseSearchIndex* PoseSearchIndex = Database->GetSearchIndex();
-		check(PoseSearchIndex);
-
-		SearchContext.GetOrBuildQuery(Database, Result.ComposedQuery);
+		const FPoseSearchIndex* SearchIndex = Database->GetSearchIndexSafe();
+		if (SearchIndex)
+		{
+			SearchContext.GetOrBuildQuery(Database, Result.ComposedQuery);
 
-		TArrayView<const float> QueryValues = Result.ComposedQuery.GetValues();
+			TArrayView<const float> QueryValues = Result.ComposedQuery.GetValues();
 
-		const FPoseSearchIndexAsset* PoseSearchIndexAsset = PoseSearchIndex->FindAssetForPose(SearchContext.CurrentResult.PoseIdx);
-		check(PoseSearchIndexAsset);
+			const FPoseSearchIndexAsset* PoseSearchIndexAsset = SearchIndex->FindAssetForPose(SearchContext.CurrentResult.PoseIdx);
+			check(PoseSearchIndexAsset);
 
-		Result.ContinuingPoseCost = Database->ComparePoses(
+			Result.ContinuingPoseCost = Database->ComparePoses(
 				SearchContext,
 				SearchContext.CurrentResult.PoseIdx,
 				EPoseComparisonFlags::ContinuingPose,
 				QueryValues);
 
-		if (Database->GetSkipSearchIfPossible())
-		{
-			SearchContext.UpdateCurrentBestCost(Result.ContinuingPoseCost);
+			if (Database->GetSkipSearchIfPossible())
+			{
+				SearchContext.UpdateCurrentBestCost(Result.ContinuingPoseCost);
+			}
 		}
 	}
 
diff --git a/Engine/Plugins/Experimental/Animation/PoseSearch/Source/Runtime/Public/PoseSearch/PoseSearch.h b/Engine/Plugins/Experimental/Animation/PoseSearch/Source/Runtime/Public/PoseSearch/PoseSearch.h
--- a/Engine/Plugins/Experimental/Animation/PoseSearch/Source/Runtime/Public/PoseSearch/PoseSearch.h
+++ b/Engine/Plugins/Experimental/Animation/PoseSearch/Source/Runtime/Public/PoseSearch/PoseSearch.h
@@ -1075,6 +1075,7 @@ public:
 
 	FPoseSearchIndex* GetSearchIndex();
 	const FPoseSearchIndex* GetSearchIndex() const;
+	const FPoseSearchIndex* GetSearchIndexSafe() const;
 	
 	bool GetSkipSearchIfPossible() const;
 	bool IsValidForIndexing() const;
