motion matching - fix for pose search debugger active and continuing pose cost reporting, dynamic coloring after asset filtering

diff --git a/Engine/Plugins/Experimental/Animation/PoseSearch/Source/Editor/Private/PoseSearchDebugger.cpp b/Engine/Plugins/Experimental/Animation/PoseSearch/Source/Editor/Private/PoseSearchDebugger.cpp
--- a/Engine/Plugins/Experimental/Animation/PoseSearch/Source/Editor/Private/PoseSearchDebugger.cpp
+++ b/Engine/Plugins/Experimental/Animation/PoseSearch/Source/Editor/Private/PoseSearchDebugger.cpp
@@ -249,7 +249,6 @@ public:
 	int32 PoseIdx = 0;
 	TWeakObjectPtr<const UPoseSearchDatabase> SourceDatabase = nullptr;
 	FTraceMotionMatchingStatePoseEntry::EFlags PoseEntryFlags = FTraceMotionMatchingStatePoseEntry::EFlags::None;
-	FTraceMotionMatchingStateDatabaseEntry::EFlags DatabaseEntryFlags = FTraceMotionMatchingStateDatabaseEntry::EFlags::None;
 	FString DatabaseName = "";
 	FString DatabasePath = "";
 	FString AssetName = "";
@@ -733,19 +732,15 @@ void SDebuggerDatabaseView::Update(const FTraceMotionMatchingStateMessage& State
 {
 	CreateRows(State);
 
-	check(ActiveView.Rows.Num() == 1);
-	check(ContinuingPoseView.Rows.Num() == 1);
-
 	SortDatabaseRows();
-	FilterDatabaseRows();
 
+	PopulateViewRows();
 
 	//	if (bNewDatabase)
 	// 	{
 	// 		FilteredDatabaseView.ListView->ClearSelection();
 	// 	}
 
-	ComputeFilteredDatabaseRowsColors();
 }
 
 void SDebuggerDatabaseView::RefreshColumns()
@@ -818,7 +813,7 @@ void SDebuggerDatabaseView::OnColumnSortModeChanged(const EColumnSortPriority::T
 	SortColumn = ColumnId;
 	SortMode = InSortMode;
 	SortDatabaseRows();
-	FilterDatabaseRows();
+	PopulateViewRows();
 }
 
 void SDebuggerDatabaseView::OnColumnWidthChanged(const float NewWidth, FName ColumnId) const
@@ -831,10 +826,9 @@ void SDebuggerDatabaseView::OnColumnWidthChanged(const float NewWidth, FName Col
 void SDebuggerDatabaseView::OnFilterTextChanged(const FText& SearchText)
 {
 	FilterText = SearchText;
-	FilterDatabaseRows();
+	PopulateViewRows();
 }
 
-
 void SDebuggerDatabaseView::OnDatabaseRowSelectionChanged(
 	TSharedPtr<FDebuggerDatabaseRowData> Row, 
 	ESelectInfo::Type SelectInfo)
@@ -845,17 +839,6 @@ void SDebuggerDatabaseView::OnDatabaseRowSelectionChanged(
 	}
 }
 
-ECheckBoxState SDebuggerDatabaseView::IsAssetFilterEnabled() const
-{
-	return bAssetFilterEnabled ? ECheckBoxState::Checked : ECheckBoxState::Unchecked;
-}
-
-void SDebuggerDatabaseView::OnAssetFilterEnabledChanged(ECheckBoxState NewState)
-{
-	bAssetFilterEnabled = NewState == ECheckBoxState::Checked;
-	FilterDatabaseRows();
-}
-
 void SDebuggerDatabaseView::SortDatabaseRows()
 {
 	if (SortMode == EColumnSortMode::Ascending)
@@ -873,8 +856,10 @@ void SDebuggerDatabaseView::SortDatabaseRows()
 	}
 }
 
-void SDebuggerDatabaseView::FilterDatabaseRows()
+void SDebuggerDatabaseView::PopulateViewRows()
 {
+	ActiveView.Rows.Reset();
+	ContinuingPoseView.Rows.Reset();
 	FilteredDatabaseView.Rows.Empty();
 
 	FString FilterString = FilterText.ToString();
@@ -884,15 +869,25 @@ void SDebuggerDatabaseView::FilterDatabaseRows()
 
 	for (const auto& UnfilteredRow : UnfilteredDatabaseRows)
 	{
-		bool bPoseIsFromCurrentDb = EnumHasAnyFlags(UnfilteredRow->DatabaseEntryFlags, FTraceMotionMatchingStateDatabaseEntry::EFlags::CurrentDatabase);
-		bool bPassesAssetFilter = !bAssetFilterEnabled || !bPoseIsFromCurrentDb;
+		bool bTryAddToFilteredDatabaseViewRows = true;
+		if (EnumHasAnyFlags(UnfilteredRow->PoseEntryFlags, FTraceMotionMatchingStatePoseEntry::EFlags::ContinuingPose))
+		{
+			ContinuingPoseView.Rows.Add(UnfilteredRow);
+			bTryAddToFilteredDatabaseViewRows = false;
+		}
+
+		if (EnumHasAnyFlags(UnfilteredRow->PoseEntryFlags, FTraceMotionMatchingStatePoseEntry::EFlags::CurrentPose))
+		{
+			ActiveView.Rows.Add(UnfilteredRow);
+			bTryAddToFilteredDatabaseViewRows = false;
+		}
 
-		if (bPassesAssetFilter)
+		if (bTryAddToFilteredDatabaseViewRows)
 		{
 			bool bPassesNameFilter = true;
 			if (bHasNameFilter)
 			{
-				bPassesNameFilter = Algo::AllOf(Tokens,[&](FString Token)
+				bPassesNameFilter = Algo::AllOf(Tokens, [&](FString Token)
 				{
 					return UnfilteredRow->AssetName.Contains(Token);
 				});
@@ -905,6 +900,10 @@ void SDebuggerDatabaseView::FilterDatabaseRows()
 		}
 	}
 
+	ComputeViewRowsColors();
+
+	ActiveView.ListView->RequestListRefresh();
+	ContinuingPoseView.ListView->RequestListRefresh();
 	FilteredDatabaseView.ListView->RequestListRefresh();
 }
 
@@ -912,12 +911,6 @@ void SDebuggerDatabaseView::CreateRows(const FTraceMotionMatchingStateMessage& S
 {
 	UnfilteredDatabaseRows.Reset();
 
-	ActiveView.Rows.Reset();
-	ActiveView.Rows.Add(MakeShared<FDebuggerDatabaseRowData>());
-
-	ContinuingPoseView.Rows.Reset();
-	ContinuingPoseView.Rows.Add(MakeShared<FDebuggerDatabaseRowData>());
-
 	for (const FTraceMotionMatchingStateDatabaseEntry& DbEntry : State.DatabaseEntries)
 	{
 		const auto* Database = FTraceMotionMatchingState::GetObjectFromId<UPoseSearchDatabase>(DbEntry.DatabaseId);
@@ -938,7 +931,6 @@ void SDebuggerDatabaseView::CreateRows(const FTraceMotionMatchingStateMessage& S
 					Row->SourceDatabase = Database;
 					Row->DatabaseName = Database->GetName();
 					Row->DatabasePath = Database->GetPathName();
-					Row->DatabaseEntryFlags = DbEntry.Flags;
 					Row->PoseEntryFlags = PoseEntry.Flags;
 					Row->DbAssetIdx = SearchIndexAsset->SourceAssetIdx;
 					Row->AssetTime = Time;
@@ -971,30 +963,15 @@ void SDebuggerDatabaseView::CreateRows(const FTraceMotionMatchingStateMessage& S
 					{
 						checkNoEntry();
 					}
-
-					if (EnumHasAnyFlags(PoseEntry.Flags, FTraceMotionMatchingStatePoseEntry::EFlags::ContinuingPose))
-					{
-						*ContinuingPoseView.Rows[0] = Row.Get();
-					}
-
-					if (EnumHasAnyFlags(PoseEntry.Flags, FTraceMotionMatchingStatePoseEntry::EFlags::CurrentPose))
-					{
-						*ActiveView.Rows[0] = Row.Get();
-					}
 				}
 			}
 		}
 	}
-
-	ActiveView.ListView->RequestListRefresh();
-	ContinuingPoseView.ListView->RequestListRefresh();
 }
 
-void SDebuggerDatabaseView::ComputeFilteredDatabaseRowsColors()
+void SDebuggerDatabaseView::ComputeViewRowsColors()
 {
-	FChannelCostRange CostRange;
-	TArray<FChannelCostRange> ChannelCostRanges;
-	for (const TSharedRef<FDebuggerDatabaseRowData>& Row : FilteredDatabaseView.Rows)
+	auto UpdateCostRanges = [](const TSharedRef<FDebuggerDatabaseRowData>& Row, FChannelCostRange& CostRange, TArray<FChannelCostRange>& ChannelCostRanges) -> void
 	{
 		const float Cost = Row->PoseCost.GetTotalCost();
 		CostRange.Min = FMath::Min(CostRange.Min, Cost);
@@ -1011,6 +988,21 @@ void SDebuggerDatabaseView::ComputeFilteredDatabaseRowsColors()
 			ChannelCostRanges[ChannelIdx].Min = FMath::Min(ChannelCostRanges[ChannelIdx].Min, Row->PoseCost.ChannelCosts[ChannelIdx]);
 			ChannelCostRanges[ChannelIdx].Max = FMath::Max(ChannelCostRanges[ChannelIdx].Max, Row->PoseCost.ChannelCosts[ChannelIdx]);
 		}
+	};
+
+	FChannelCostRange CostRange;
+	TArray<FChannelCostRange> ChannelCostRanges;
+	for (const TSharedRef<FDebuggerDatabaseRowData>& Row : FilteredDatabaseView.Rows)
+	{
+		UpdateCostRanges(Row, CostRange, ChannelCostRanges);
+	}
+	for (const TSharedRef<FDebuggerDatabaseRowData>& Row : ActiveView.Rows)
+	{
+		UpdateCostRanges(Row, CostRange, ChannelCostRanges);
+	}
+	for (const TSharedRef<FDebuggerDatabaseRowData>& Row : ContinuingPoseView.Rows)
+	{
+		UpdateCostRanges(Row, CostRange, ChannelCostRanges);
 	}
 
 	CostRange.Delta = CostRange.Max - CostRange.Min;
@@ -1033,15 +1025,13 @@ void SDebuggerDatabaseView::ComputeFilteredDatabaseRowsColors()
 	{
 		Row->CalculateColors(CostRange, ChannelCostRanges);
 	}
-
-	if (!ActiveView.Rows.IsEmpty())
+	for (const TSharedRef<FDebuggerDatabaseRowData>& Row : ActiveView.Rows)
 	{
-		ActiveView.Rows[0]->CalculateColors(CostRange, ChannelCostRanges);
+		Row->CalculateColors(CostRange, ChannelCostRanges);
 	}
-
-	if (!ContinuingPoseView.Rows.IsEmpty())
+	for (const TSharedRef<FDebuggerDatabaseRowData>& Row : ContinuingPoseView.Rows)
 	{
-		ContinuingPoseView.Rows[0]->CalculateColors(CostRange, ChannelCostRanges);
+		Row->CalculateColors(CostRange, ChannelCostRanges);
 	}
 }
 
@@ -1351,18 +1341,7 @@ void SDebuggerDatabaseView::Construct(const FArguments& InArgs)
 			[
 				SNew(SHorizontalBox)
 				+ SHorizontalBox::Slot()
-				.AutoWidth()
 				.Padding(10, 5, 10, 5)
-				[
-					SNew(SCheckBox)
-					.IsChecked(this, &SDebuggerDatabaseView::IsAssetFilterEnabled)
-					.OnCheckStateChanged(this, &SDebuggerDatabaseView::OnAssetFilterEnabledChanged)
-					[
-						SNew(STextBlock)
-						.Text(LOCTEXT("PoseSearchDebuggerGroupFiltering", "Apply Group Filtering"))
-					]
-				]
-				+ SHorizontalBox::Slot()
 				[
 					SAssignNew(FilterBox, SSearchBox)
 					.OnTextChanged(this, &SDebuggerDatabaseView::OnFilterTextChanged)
@@ -1470,7 +1449,7 @@ void SDebuggerDetailsView::UpdateReflection(const FTraceMotionMatchingStateMessa
 	Reflection->SimAngularVelocity = State.SimAngularVelocity;
 	Reflection->AnimLinearVelocity = State.AnimLinearVelocity;
 	Reflection->AnimAngularVelocity = State.AnimAngularVelocity;
-
+	 
 	// Query pose
 	Reflection->QueryPoseVector = State.QueryVector;
 
diff --git a/Engine/Plugins/Experimental/Animation/PoseSearch/Source/Editor/Private/PoseSearchDebugger.h b/Engine/Plugins/Experimental/Animation/PoseSearch/Source/Editor/Private/PoseSearchDebugger.h
--- a/Engine/Plugins/Experimental/Animation/PoseSearch/Source/Editor/Private/PoseSearchDebugger.h
+++ b/Engine/Plugins/Experimental/Animation/PoseSearch/Source/Editor/Private/PoseSearchDebugger.h
@@ -222,9 +222,9 @@ private:
 	/** Sorts the database by the current sort predicate, updating the view order */
 	void SortDatabaseRows();
 
-	void FilterDatabaseRows();
+	void PopulateViewRows();
 
-	void ComputeFilteredDatabaseRowsColors();
+	void ComputeViewRowsColors();
 
 	/** Acquires sort predicate for the given column */
 	EColumnSortMode::Type GetColumnSortMode(const FName ColumnId) const;
@@ -245,12 +245,6 @@ private:
 	/** Row selection to update model view */
 	void OnDatabaseRowSelectionChanged(TSharedPtr<FDebuggerDatabaseRowData> Row, ESelectInfo::Type SelectInfo);
 
-	/** Informs the widget if asset filtering is enabled */
-	ECheckBoxState IsAssetFilterEnabled() const;
-
-	/** Updates the state of asset filtering */
-	void OnAssetFilterEnabledChanged(ECheckBoxState NewState);
-
 	/** Generates a database row widget for the given data */
 	TSharedRef<ITableRow> HandleGenerateDatabaseRow(TSharedRef<FDebuggerDatabaseRowData> Item, const TSharedRef<STableViewBase>& OwnerTable) const;
 	
@@ -312,9 +306,6 @@ private:
 
 	/** Text used to filter DatabaseView */
 	FText FilterText;
-
-	/** True if only asets that pass the asset group query are being displayed */
-	bool bAssetFilterEnabled = true;
 };
 
 /**
diff --git a/Engine/Plugins/Experimental/Animation/PoseSearch/Source/Runtime/Private/PoseSearch.cpp b/Engine/Plugins/Experimental/Animation/PoseSearch/Source/Runtime/Private/PoseSearch.cpp
--- a/Engine/Plugins/Experimental/Animation/PoseSearch/Source/Runtime/Private/PoseSearch.cpp
+++ b/Engine/Plugins/Experimental/Animation/PoseSearch/Source/Runtime/Private/PoseSearch.cpp
@@ -806,14 +806,14 @@ FPoseSearchCost UPoseSearchSequenceMetaData::ComparePoses(
 	TArrayView<const float> PoseValues = SearchIndex.GetPoseValues(PoseIdx);
 	check(PoseValues.Num() == QueryValues.Num());
 
-	Result.SetDissimilarity(CompareFeatureVectors(PoseValues.Num(), PoseValues.GetData(), QueryValues.GetData()));
+	const float DissimilarityCost = CompareFeatureVectors(PoseValues.Num(), PoseValues.GetData(), QueryValues.GetData());
 
 	// @todo: shouldn't we include MirrorMismatchAddend as well?
 	//const float MirrorMismatchAddend = SearchIndex.ComputeMirrorMismatchAddend(PoseIdx, SearchContext);
 	const float NotifyAddend = SearchIndex.ComputeNotifyAddend(PoseIdx);
 	const float ContinuingPoseCostAddend = SearchIndex.ComputeContinuingPoseCostAddend(PoseIdx, PoseComparisonFlags);
 
-	Result.SetCostAddend(NotifyAddend + ContinuingPoseCostAddend);
+	Result.Set(DissimilarityCost, NotifyAddend + ContinuingPoseCostAddend);
 
 	return Result;
 }
@@ -1726,17 +1726,13 @@ FPoseSearchCost UPoseSearchDatabase::ComparePoses(UE::PoseSearch::FSearchContext
 
 	TArrayView<const float> Weights(SearchIndex->Weights);
 
-	Result.SetDissimilarity(CompareFeatureVectors(
-		PoseValues.Num(),
-		PoseValues.GetData(),
-		QueryValues.GetData(),
-		Weights.GetData()));
+	const float DissimilarityCost = CompareFeatureVectors(PoseValues.Num(), PoseValues.GetData(), QueryValues.GetData(), Weights.GetData());
 
 	const float MirrorMismatchAddend = SearchIndex->ComputeMirrorMismatchAddend(PoseIdx, SearchContext);
 	const float NotifyAddend = SearchIndex->ComputeNotifyAddend(PoseIdx);
 	const float ContinuingPoseCostAddend = SearchIndex->ComputeContinuingPoseCostAddend(PoseIdx, PoseComparisonFlags);
 
-	Result.SetCostAddend(NotifyAddend + MirrorMismatchAddend + ContinuingPoseCostAddend);
+	Result.Set(DissimilarityCost, NotifyAddend + MirrorMismatchAddend + ContinuingPoseCostAddend);
 
 #if UE_POSE_SEARCH_TRACE_ENABLED && WITH_EDITORONLY_DATA
 	if (SearchContext.bIsTracing)
@@ -1873,10 +1869,6 @@ UE::PoseSearch::FSearchResult UPoseSearchDatabase::SearchPCAKDTree(UE::PoseSearc
 		}
 	}
 
-	// @todo: implement support for DatabaseTagQuery
-	FPoseSearchCost BestPoseCost;
-	int32 BestPoseIdx = INDEX_NONE;
-
 	SearchContext.GetOrBuildQuery(this, Result.ComposedQuery);
 
 	TArrayView<const float> QueryValues = Result.ComposedQuery.GetValues();
@@ -1886,15 +1878,13 @@ UE::PoseSearch::FSearchResult UPoseSearchDatabase::SearchPCAKDTree(UE::PoseSearc
 	// evaluating the continuing pose only if it hasn't already being evaluated and the related animation can advance
 	if (IsCurrentResultFromThisDatabase && SearchContext.bCanAdvance && !Result.ContinuingPoseCost.IsValid())
 	{
-		Result.ContinuingPoseCost = ComparePoses(
-			SearchContext,
-			SearchContext.CurrentResult.PoseIdx,
-			EPoseComparisonFlags::ContinuingPose,
-			QueryValues);
+		Result.PoseIdx = SearchContext.CurrentResult.PoseIdx;
+		Result.PoseCost = ComparePoses(SearchContext, Result.PoseIdx, EPoseComparisonFlags::ContinuingPose, QueryValues);
+		Result.ContinuingPoseCost = Result.PoseCost;
 
 		if (GetSkipSearchIfPossible())
 		{
-			SearchContext.UpdateCurrentBestCost(Result.ContinuingPoseCost);
+			SearchContext.UpdateCurrentBestCost(Result.PoseCost);
 		}
 	}
 
@@ -1938,10 +1928,10 @@ UE::PoseSearch::FSearchResult UPoseSearchDatabase::SearchPCAKDTree(UE::PoseSearc
 				EPoseComparisonFlags::None,
 				QueryValues);
 
-			if (PoseCost < BestPoseCost)
+			if (PoseCost < Result.PoseCost)
 			{
-				BestPoseCost = PoseCost;
-				BestPoseIdx = PoseIdx;
+				Result.PoseCost = PoseCost;
+				Result.PoseIdx = PoseIdx;
 			}
 
 #if UE_POSE_SEARCH_TRACE_ENABLED
@@ -1949,25 +1939,17 @@ UE::PoseSearch::FSearchResult UPoseSearchDatabase::SearchPCAKDTree(UE::PoseSearc
 #endif
 		}
 
-		if (GetSkipSearchIfPossible() && BestPoseCost.IsValid())
+		if (GetSkipSearchIfPossible() && Result.PoseCost.IsValid())
 		{
-			SearchContext.UpdateCurrentBestCost(BestPoseCost);
+			SearchContext.UpdateCurrentBestCost(Result.PoseCost);
 		}
 	}
-	else if (Result.ContinuingPoseCost.IsValid())
-	{
-		// if the search were skipped we'll reuse the continuing pose search result and subtract the ContinuingPoseCostAddend from the cost
-		BestPoseCost = Result.ContinuingPoseCost;
-		BestPoseIdx = SearchContext.CurrentResult.PoseIdx;
-		BestPoseCost.SetCostAddend(BestPoseCost.GetCostAddend() - SearchIndex->ComputeContinuingPoseCostAddend(SearchContext.CurrentResult.PoseIdx, EPoseComparisonFlags::ContinuingPose));
-	}
 
-	if (BestPoseIdx != INDEX_NONE)
+	// finalizing Result properties
+	if (Result.PoseIdx != INDEX_NONE)
 	{
-		Result.PoseCost = BestPoseCost;
-		Result.PoseIdx = BestPoseIdx;
-		Result.SearchIndexAsset = SearchIndex->FindAssetForPose(BestPoseIdx);
-		Result.AssetTime = SearchIndex->GetAssetTime(BestPoseIdx, Result.SearchIndexAsset);
+		Result.SearchIndexAsset = SearchIndex->FindAssetForPose(Result.PoseIdx);
+		Result.AssetTime = SearchIndex->GetAssetTime(Result.PoseIdx, Result.SearchIndexAsset);
 		Result.Database = this;
 #if WITH_EDITOR
 		Result.SearchIndexHash = GetSearchIndexHash();
@@ -1998,15 +1980,13 @@ UE::PoseSearch::FSearchResult UPoseSearchDatabase::SearchBruteForce(UE::PoseSear
 		// evaluating the continuing pose only if it hasn't already being evaluated and the related animation can advance
 		if (SearchContext.bCanAdvance && !Result.ContinuingPoseCost.IsValid())
 		{
-			Result.ContinuingPoseCost = ComparePoses(
-				SearchContext,
-				SearchContext.CurrentResult.PoseIdx,
-				EPoseComparisonFlags::ContinuingPose,
-				QueryValues);
+			Result.PoseIdx = SearchContext.CurrentResult.PoseIdx;
+			Result.PoseCost = ComparePoses(SearchContext, Result.PoseIdx, EPoseComparisonFlags::ContinuingPose, QueryValues);
+			Result.ContinuingPoseCost = Result.PoseCost;
 
 			if (GetSkipSearchIfPossible())
 			{
-				SearchContext.UpdateCurrentBestCost(Result.ContinuingPoseCost);
+				SearchContext.UpdateCurrentBestCost(Result.PoseCost);
 			}
 		}
 	}
@@ -2016,9 +1996,6 @@ UE::PoseSearch::FSearchResult UPoseSearchDatabase::SearchBruteForce(UE::PoseSear
 	const int NonSelectableIdxUsedSize = PopulateNonSelectableIdx(NonSelectableIdxData, NonSelectableIdxDataSize, SearchContext, this);
 	TArrayView<size_t> NonSelectableIdx(NonSelectableIdxData, NonSelectableIdxUsedSize);
 
-	FPoseSearchCost BestPoseCost;
-	int32 BestPoseIdx = INDEX_NONE;
-
 	// since any PoseCost calculated here is at least SearchIndex->MinCostAddend,
 	// there's no point in performing the search if CurrentBestTotalCost is already better than that
 	if (SearchContext.GetCurrentBestTotalCost() > SearchIndex->MinCostAddend)
@@ -2047,10 +2024,10 @@ UE::PoseSearch::FSearchResult UPoseSearchDatabase::SearchBruteForce(UE::PoseSear
 					EPoseComparisonFlags::None,
 					QueryValues);
 
-				if (PoseCost < BestPoseCost)
+				if (PoseCost < Result.PoseCost)
 				{
-					BestPoseCost = PoseCost;
-					BestPoseIdx = PoseIdx;
+					Result.PoseCost = PoseCost;
+					Result.PoseIdx = PoseIdx;
 				}
 
 #if UE_POSE_SEARCH_TRACE_ENABLED
@@ -2062,25 +2039,17 @@ UE::PoseSearch::FSearchResult UPoseSearchDatabase::SearchBruteForce(UE::PoseSear
 			}
 		}
 
-		if (GetSkipSearchIfPossible() && BestPoseCost.IsValid())
+		if (GetSkipSearchIfPossible() && Result.PoseCost.IsValid())
 		{
-			SearchContext.UpdateCurrentBestCost(BestPoseCost);
+			SearchContext.UpdateCurrentBestCost(Result.PoseCost);
 		}
 	}
-	else if (Result.ContinuingPoseCost.IsValid())
-	{
-		// if the search were skipped we'll reuse the continuing pose search result and subtract the ContinuingPoseCostAddend from the cost
-		BestPoseCost = Result.ContinuingPoseCost;
-		BestPoseIdx = SearchContext.CurrentResult.PoseIdx;
-		BestPoseCost.SetCostAddend(BestPoseCost.GetCostAddend() - SearchIndex->ComputeContinuingPoseCostAddend(SearchContext.CurrentResult.PoseIdx, EPoseComparisonFlags::ContinuingPose));
-	}
 
-	if (BestPoseIdx != INDEX_NONE)
+	// finalizing Result properties
+	if (Result.PoseIdx != INDEX_NONE)
 	{
-		Result.PoseCost = BestPoseCost;
-		Result.PoseIdx = BestPoseIdx;
-		Result.SearchIndexAsset = SearchIndex->FindAssetForPose(BestPoseIdx);
-		Result.AssetTime = SearchIndex->GetAssetTime(BestPoseIdx, Result.SearchIndexAsset);
+		Result.SearchIndexAsset = SearchIndex->FindAssetForPose(Result.PoseIdx);
+		Result.AssetTime = SearchIndex->GetAssetTime(Result.PoseIdx, Result.SearchIndexAsset);
 		Result.Database = this;
 #if WITH_EDITOR
 		Result.SearchIndexHash = GetSearchIndexHash();
@@ -2128,15 +2097,14 @@ UE::PoseSearch::FSearchResult UPoseSearchDatabaseSet::Search(UE::PoseSearch::FSe
 			const FPoseSearchIndexAsset* PoseSearchIndexAsset = SearchIndex->FindAssetForPose(SearchContext.CurrentResult.PoseIdx);
 			check(PoseSearchIndexAsset);
 
-			Result.ContinuingPoseCost = Database->ComparePoses(
-				SearchContext,
-				SearchContext.CurrentResult.PoseIdx,
-				EPoseComparisonFlags::ContinuingPose,
-				QueryValues);
+			Result.PoseIdx = SearchContext.CurrentResult.PoseIdx;
+			Result.PoseCost = Database->ComparePoses(SearchContext, Result.PoseIdx, EPoseComparisonFlags::ContinuingPose, QueryValues);
+			Result.ContinuingPoseCost = Result.PoseCost;
+			ContinuingCost = Result.PoseCost;
 
 			if (Database->GetSkipSearchIfPossible())
 			{
-				SearchContext.UpdateCurrentBestCost(Result.ContinuingPoseCost);
+				SearchContext.UpdateCurrentBestCost(Result.PoseCost);
 			}
 		}
 	}
diff --git a/Engine/Plugins/Experimental/Animation/PoseSearch/Source/Runtime/Private/PoseSearchLibrary.cpp b/Engine/Plugins/Experimental/Animation/PoseSearch/Source/Runtime/Private/PoseSearchLibrary.cpp
--- a/Engine/Plugins/Experimental/Animation/PoseSearch/Source/Runtime/Private/PoseSearchLibrary.cpp
+++ b/Engine/Plugins/Experimental/Animation/PoseSearch/Source/Runtime/Private/PoseSearchLibrary.cpp
@@ -164,16 +164,15 @@ void FMotionMatchingState::JumpToPose(const FAnimationUpdateContext& Context, co
 	Flags |= EMotionMatchingFlags::JumpedToPose;
 }
 
+#if UE_POSE_SEARCH_TRACE_ENABLED
 static void TraceMotionMatchingState(
 	const FAnimationUpdateContext& UpdateContext,
 	const UPoseSearchSearchableAsset* Searchable,
 	UE::PoseSearch::FSearchContext& SearchContext,
 	const FMotionMatchingState& MotionMatchingState,
 	const FTrajectorySampleRange& Trajectory,
-	const UE::PoseSearch::FSearchResult& LastResult
-)
+	const UE::PoseSearch::FSearchResult& LastResult)
 {
-#if UE_POSE_SEARCH_TRACE_ENABLED
 	using namespace UE::PoseSearch;
 
 	if (!IsTracing(UpdateContext))
@@ -189,8 +188,6 @@ static void TraceMotionMatchingState(
 	const float DeltaTime = UpdateContext.GetDeltaTime();
 
 	FTraceMotionMatchingState TraceState;
-	bool bContinuingPoseTraced = false;
-	bool bCurrentPoseTraced = false;
 	while (!SearchContext.BestCandidates.IsEmpty())
 	{
 		FSearchContext::FPoseCandidate PoseCandidate;
@@ -204,79 +201,44 @@ static void TraceMotionMatchingState(
 		PoseEntry.DbPoseIdx = PoseCandidate.PoseIdx;
 		PoseEntry.Cost = PoseCandidate.Cost;
 
-		if (PoseCandidate.Database == MotionMatchingState.CurrentSearchResult.Database)
-		{
-			DbEntry.Flags |= FTraceMotionMatchingStateDatabaseEntry::EFlags::CurrentDatabase;
-		}
-		
-		if (PoseCandidate == MotionMatchingState.CurrentSearchResult)
-		{
-			PoseEntry.Flags |= FTraceMotionMatchingStatePoseEntry::EFlags::CurrentPose;
-			bCurrentPoseTraced = true;
-		}
-
-		if (PoseCandidate == LastResult)
-		{
-			PoseEntry.Flags |= FTraceMotionMatchingStatePoseEntry::EFlags::ContinuingPose;
-			bContinuingPoseTraced = true;
-		}
-
 		DbEntry.PoseEntries.Add(PoseEntry);
 	}
 
-	if (!bContinuingPoseTraced && LastResult.IsValid())
+	if (MotionMatchingState.CurrentSearchResult.ContinuingPoseCost.IsValid())
 	{
+		check(LastResult.IsValid());
 		const uint64 LastResultDbId = FTraceMotionMatchingState::GetIdFromObject(LastResult.Database.Get());
 
 		int32 LastResultDbIdx = TraceState.DatabaseEntries.AddUnique({LastResultDbId});
 		FTraceMotionMatchingStateDatabaseEntry& DbEntry = TraceState.DatabaseEntries[LastResultDbIdx];
 
-		int32 LastResultPoseEntryIdx = DbEntry.PoseEntries.AddUnique({LastResult.PoseIdx});
+		int32 LastResultPoseEntryIdx = DbEntry.PoseEntries.Add({LastResult.PoseIdx});
 		FTraceMotionMatchingStatePoseEntry& PoseEntry = DbEntry.PoseEntries[LastResultPoseEntryIdx];
 
-		PoseEntry.Cost = LastResult.PoseCost;
-		PoseEntry.Flags |= FTraceMotionMatchingStatePoseEntry::EFlags::ContinuingPose;
-		bContinuingPoseTraced = true;
+		PoseEntry.Cost = MotionMatchingState.CurrentSearchResult.ContinuingPoseCost;
+		PoseEntry.Flags = FTraceMotionMatchingStatePoseEntry::EFlags::ContinuingPose;
 	}
 
-	if (!bCurrentPoseTraced && MotionMatchingState.CurrentSearchResult.IsValid())
+	if (MotionMatchingState.CurrentSearchResult.PoseCost.IsValid())
 	{
 		const uint64 CurrentResultDbId = FTraceMotionMatchingState::GetIdFromObject(MotionMatchingState.CurrentSearchResult.Database.Get());
 
 		int32 DbEntryIdx = TraceState.DatabaseEntries.AddUnique({CurrentResultDbId});
 		FTraceMotionMatchingStateDatabaseEntry& DbEntry = TraceState.DatabaseEntries[DbEntryIdx];
-		DbEntry.Flags |= FTraceMotionMatchingStateDatabaseEntry::EFlags::CurrentDatabase;
 
-		int32 PoseEntryIdx = DbEntry.PoseEntries.AddUnique({MotionMatchingState.CurrentSearchResult.PoseIdx});
+		int32 PoseEntryIdx = DbEntry.PoseEntries.Add({MotionMatchingState.CurrentSearchResult.PoseIdx});
 		FTraceMotionMatchingStatePoseEntry& PoseEntry = DbEntry.PoseEntries[PoseEntryIdx];
 
 		PoseEntry.Cost = MotionMatchingState.CurrentSearchResult.PoseCost;
-		PoseEntry.Flags |= FTraceMotionMatchingStatePoseEntry::EFlags::CurrentPose;
-		bCurrentPoseTraced = true;
-	}
-
-	TraceState.CurrentDbEntryIdx = TraceState.DatabaseEntries.IndexOfByPredicate(
-		[](const auto& DbEntry)
-		{
-			return EnumHasAnyFlags(DbEntry.Flags, FTraceMotionMatchingStateDatabaseEntry::EFlags::CurrentDatabase);
-		});
+		PoseEntry.Flags = FTraceMotionMatchingStatePoseEntry::EFlags::CurrentPose;
 
-	if (TraceState.CurrentDbEntryIdx != INDEX_NONE)
-	{
-		TraceState.CurrentPoseEntryIdx = TraceState.DatabaseEntries[TraceState.CurrentDbEntryIdx].PoseEntries.IndexOfByPredicate(
-			[](const auto& PoseEntry)
-			{
-				return EnumHasAnyFlags(PoseEntry.Flags, FTraceMotionMatchingStatePoseEntry::EFlags::CurrentPose);
-			});
+		TraceState.CurrentDbEntryIdx = DbEntryIdx;
+		TraceState.CurrentPoseEntryIdx = PoseEntryIdx;
 	}
 
-
-	float SimLinearVelocity, SimAngularVelocity, AnimLinearVelocity, AnimAngularVelocity;
-
 	if (DeltaTime > SMALL_NUMBER)
 	{
 		// simulation
-
 		int32 FirstIdx = 0;
 		const FTrajectorySample PrevSample = FTrajectorySampleRange::IterSampleTrajectory(
 			Trajectory.Samples,
@@ -290,22 +252,14 @@ static void TraceMotionMatchingState(
 
 		const FTransform SimDelta = CurrSample.Transform.GetRelativeTransform(PrevSample.Transform);
 
-		SimLinearVelocity = SimDelta.GetTranslation().Size() / DeltaTime;
-		SimAngularVelocity = FMath::RadiansToDegrees(SimDelta.GetRotation().GetAngle()) / DeltaTime;
+		TraceState.SimLinearVelocity = SimDelta.GetTranslation().Size() / DeltaTime;
+		TraceState.SimAngularVelocity = FMath::RadiansToDegrees(SimDelta.GetRotation().GetAngle()) / DeltaTime;
 
 		// animation
-
 		const FTransform AnimDelta = MotionMatchingState.RootMotionTransformDelta;
 
-		AnimLinearVelocity = AnimDelta.GetTranslation().Size() / DeltaTime;
-		AnimAngularVelocity = FMath::RadiansToDegrees(AnimDelta.GetRotation().GetAngle()) / DeltaTime;
-	}
-	else
-	{
-		SimLinearVelocity = 0.0f;
-		SimAngularVelocity = 0.0f;
-		AnimLinearVelocity = 0.0f;
-		AnimAngularVelocity = 0.0f;
+		TraceState.AnimLinearVelocity = AnimDelta.GetTranslation().Size() / DeltaTime;
+		TraceState.AnimAngularVelocity = FMath::RadiansToDegrees(AnimDelta.GetRotation().GetAngle()) / DeltaTime;
 	}
 
 	if (EnumHasAnyFlags(MotionMatchingState.Flags, EMotionMatchingFlags::JumpedToFollowUp))
@@ -318,14 +272,10 @@ static void TraceMotionMatchingState(
 	TraceState.QueryVector = MotionMatchingState.CurrentSearchResult.ComposedQuery.GetValues();
 	TraceState.AssetPlayerTime = MotionMatchingState.CurrentSearchResult.AssetTime;
 	TraceState.DeltaTime = DeltaTime;
-	TraceState.SimLinearVelocity = SimLinearVelocity;
-	TraceState.SimAngularVelocity = SimAngularVelocity;
-	TraceState.AnimLinearVelocity = AnimLinearVelocity;
-	TraceState.AnimAngularVelocity = AnimAngularVelocity;
 
 	TraceState.Output(UpdateContext);
-#endif
 }
+#endif
 
 void UpdateMotionMatchingState(
 	const FAnimationUpdateContext& Context,
@@ -354,8 +304,10 @@ void UpdateMotionMatchingState(
 	// Reset State Flags
 	InOutMotionMatchingState.Flags = EMotionMatchingFlags::None;
 
+#if UE_POSE_SEARCH_TRACE_ENABLED
 	// Record Current Pose Index for Debugger
 	const FSearchResult LastResult = InOutMotionMatchingState.CurrentSearchResult;
+#endif
 
 	// Check if we can advance. Includes the case where we can advance but only by switching to a follow up asset.
 	bool bAdvanceToFollowUpAsset = false;
@@ -433,6 +385,7 @@ void UpdateMotionMatchingState(
 
 	InOutMotionMatchingState.PoseIndicesHistory.Update(InOutMotionMatchingState.CurrentSearchResult, DeltaTime, Settings.PoseReselectHistory);
 	
+#if UE_POSE_SEARCH_TRACE_ENABLED
 	// Record debugger details
 	TraceMotionMatchingState(
 		Context,
@@ -440,8 +393,8 @@ void UpdateMotionMatchingState(
 		SearchContext,
 		InOutMotionMatchingState,
 		Trajectory,
-		LastResult
-	);
+		LastResult);
+#endif
 }
 
 const FPoseSearchIndexAsset* FMotionMatchingState::GetCurrentSearchIndexAsset() const
diff --git a/Engine/Plugins/Experimental/Animation/PoseSearch/Source/Runtime/Private/Trace/PoseSearchTraceLogger.cpp b/Engine/Plugins/Experimental/Animation/PoseSearch/Source/Runtime/Private/Trace/PoseSearchTraceLogger.cpp
--- a/Engine/Plugins/Experimental/Animation/PoseSearch/Source/Runtime/Private/Trace/PoseSearchTraceLogger.cpp
+++ b/Engine/Plugins/Experimental/Animation/PoseSearch/Source/Runtime/Private/Trace/PoseSearchTraceLogger.cpp
@@ -59,7 +59,6 @@ FArchive& operator<<(FArchive& Ar, FTraceMotionMatchingStatePoseEntry& Entry)
 FArchive& operator<<(FArchive& Ar, FTraceMotionMatchingStateDatabaseEntry& Entry)
 {
 	Ar << Entry.DatabaseId;
-	Ar << Entry.Flags;
 	Ar << Entry.PoseEntries;
 	return Ar;
 }
diff --git a/Engine/Plugins/Experimental/Animation/PoseSearch/Source/Runtime/Private/Trace/PoseSearchTraceLogger.h b/Engine/Plugins/Experimental/Animation/PoseSearch/Source/Runtime/Private/Trace/PoseSearchTraceLogger.h
--- a/Engine/Plugins/Experimental/Animation/PoseSearch/Source/Runtime/Private/Trace/PoseSearchTraceLogger.h
+++ b/Engine/Plugins/Experimental/Animation/PoseSearch/Source/Runtime/Private/Trace/PoseSearchTraceLogger.h
@@ -55,21 +55,13 @@ POSESEARCH_API FArchive& operator<<(FArchive& Ar, FTraceMotionMatchingStatePoseE
 
 struct POSESEARCH_API FTraceMotionMatchingStateDatabaseEntry
 {
-	enum class EFlags : uint8
-	{
-		None			= 0,
-		CurrentDatabase	= 1 << 0,
-	};
-
 	// @todo: can we use UPoseSearchDatabase* instead of DatabaseId?
 	//UPoseSearchDatabase const* Database = nullptr;
 	uint64 DatabaseId = 0;
-	EFlags Flags = EFlags::None;
 	TArray<FTraceMotionMatchingStatePoseEntry> PoseEntries;
 
 	bool operator==(const FTraceMotionMatchingStateDatabaseEntry& Other) const { return DatabaseId == Other.DatabaseId; }
 };
-ENUM_CLASS_FLAGS(FTraceMotionMatchingStateDatabaseEntry::EFlags);
 POSESEARCH_API FArchive& operator<<(FArchive& Ar, FTraceMotionMatchingStateDatabaseEntry& Entry);
 
 /**
diff --git a/Engine/Plugins/Experimental/Animation/PoseSearch/Source/Runtime/Public/PoseSearch/PoseSearch.h b/Engine/Plugins/Experimental/Animation/PoseSearch/Source/Runtime/Public/PoseSearch/PoseSearch.h
--- a/Engine/Plugins/Experimental/Animation/PoseSearch/Source/Runtime/Public/PoseSearch/PoseSearch.h
+++ b/Engine/Plugins/Experimental/Animation/PoseSearch/Source/Runtime/Public/PoseSearch/PoseSearch.h
@@ -800,43 +800,26 @@ public:
 	FPoseSearchCost(float InDissimilarity, float InCostAddend)
 		: Dissimilarity(InDissimilarity)
 		, CostAddend(InCostAddend)
-		, TotalCost(Dissimilarity + CostAddend)
 	{}
 
-	bool IsValid() const { return TotalCost != MAX_flt; }
+	bool IsValid() const { return Dissimilarity != MAX_flt; }
 
 	void Set(float InDissimilarity, float InCostAddend)
 	{
 		Dissimilarity = InDissimilarity;
 		CostAddend = InCostAddend;
-		TotalCost = Dissimilarity + CostAddend;
 	}
 
 	float GetDissimilarity() const { return Dissimilarity; }
-
-	void SetDissimilarity(float InDissimilarity)
-	{
-		Dissimilarity = InDissimilarity;
-		TotalCost = Dissimilarity + CostAddend;
-	}
-
 	float GetCostAddend() const { return CostAddend; }
+	float GetTotalCost() const { return Dissimilarity + CostAddend; }
 
-	void SetCostAddend(float InCostAddend)
-	{
-		CostAddend = InCostAddend;
-		TotalCost = Dissimilarity + CostAddend;
-	}
-
-	float GetTotalCost() const { return TotalCost; }
-
-	bool operator<(const FPoseSearchCost& Other) const { return TotalCost < Other.TotalCost; }
+	bool operator<(const FPoseSearchCost& Other) const { return GetTotalCost() < Other.GetTotalCost(); }
 
 	void Reset()
 	{
 		Dissimilarity = MAX_flt;
 		CostAddend = 0.0f;
-		TotalCost = MAX_flt;
 	}
 
 protected:
@@ -847,9 +830,6 @@ protected:
 	UPROPERTY()
 	float CostAddend = 0.0f;
 
-	UPROPERTY()
-	float TotalCost = MAX_flt;
-
 #if WITH_EDITORONLY_DATA
 public:
 	// Contribution from ModifyCost anim notify
@@ -976,6 +956,7 @@ public:
 
 struct FSearchResult
 {
+	// best cost of the currently selected PoseIdx (it could be equal to ContinuingPoseCost)
 	FPoseSearchCost PoseCost;
 	int32 PoseIdx = INDEX_NONE;
 
