motion matching - adding UPoseSearchDatabaseReflection to display index property in the db editor, adding NormalizeOnlyByDeviation to premultiply weights only by data deviation, avoiding normalizing user weights, adding minor quality of life improvements

diff --git a/Engine/Plugins/Experimental/Animation/PoseSearch/Source/Editor/Private/PoseSearchDatabaseEditor.cpp b/Engine/Plugins/Experimental/Animation/PoseSearch/Source/Editor/Private/PoseSearchDatabaseEditor.cpp
--- a/Engine/Plugins/Experimental/Animation/PoseSearch/Source/Editor/Private/PoseSearchDatabaseEditor.cpp
+++ b/Engine/Plugins/Experimental/Animation/PoseSearch/Source/Editor/Private/PoseSearchDatabaseEditor.cpp
@@ -514,6 +514,25 @@ namespace UE::PoseSearch
 					NewSelectionReflection->SetSourceLink(SelectedItem, AssetTreeWidget);
 					SelectionReflections[BlendSpaceSelectionIndex].Add(NewSelectionReflection);
 				}
+				else
+				{
+					UPoseSearchDatabaseReflection* NewSelectionReflection =
+						NewObject<UPoseSearchDatabaseReflection>();
+					NewSelectionReflection->AddToRoot();
+
+					const FPoseSearchIndex* SearchIndex = GetPoseSearchDatabase()->GetSearchIndex();
+					if (SearchIndex)
+					{
+						NewSelectionReflection->SearchIndex = *SearchIndex;
+					}
+					else
+					{
+						NewSelectionReflection->SearchIndex = FPoseSearchIndex();
+					}
+
+					NewSelectionReflection->SetSourceLink(SelectedItem, AssetTreeWidget);
+					SelectionReflections[GroupSelectionIndex].Add(NewSelectionReflection);
+				}
 			}
 		}
 
diff --git a/Engine/Plugins/Experimental/Animation/PoseSearch/Source/Editor/Private/PoseSearchDebugger.cpp b/Engine/Plugins/Experimental/Animation/PoseSearch/Source/Editor/Private/PoseSearchDebugger.cpp
--- a/Engine/Plugins/Experimental/Animation/PoseSearch/Source/Editor/Private/PoseSearchDebugger.cpp
+++ b/Engine/Plugins/Experimental/Animation/PoseSearch/Source/Editor/Private/PoseSearchDebugger.cpp
@@ -1707,8 +1707,6 @@ void SDebuggerView::DrawFeatures(
 	DrawParams.DefaultLifeTime = 0.0f; // Single frame render
 	DrawParams.Mesh = Mesh;
 
-	const TObjectPtr<UPoseSearchDebuggerReflection>& Reflection = DetailsView->GetReflection();
-
 	auto SetDrawFlags = [](FDebugDrawParams& InDrawParams, const FPoseSearchDebuggerFeatureDrawOptions& Options)
 	{
 		InDrawParams.Flags = EDebugDrawFlags::None;
@@ -1731,11 +1729,12 @@ void SDebuggerView::DrawFeatures(
 		}
 	};
 
+	const TObjectPtr<UPoseSearchDebuggerReflection> Reflection = DetailsView->GetReflection();
 
 	// Draw query vector
 	{
 		DrawParams.Database = ViewModel.Get()->GetCurrentDatabase();
-		SetDrawFlags(DrawParams, Reflection->QueryDrawOptions);
+		SetDrawFlags(DrawParams, Reflection ? Reflection->QueryDrawOptions : FPoseSearchDebuggerFeatureDrawOptions());
 		EnumAddFlags(DrawParams.Flags, EDebugDrawFlags::DrawQuery);
 		DrawFeatureVector(DrawParams, State.QueryVector);
 		EnumRemoveFlags(DrawParams.Flags, EDebugDrawFlags::DrawQuery);
@@ -1748,7 +1747,7 @@ void SDebuggerView::DrawFeatures(
 		TArray<TSharedRef<FDebuggerDatabaseRowData>> SelectedRows = DatabaseRows->GetSelectedItems();
 	
 		// Red for non-active database view
-		SetDrawFlags(DrawParams, Reflection->SelectedPoseDrawOptions);
+		SetDrawFlags(DrawParams, Reflection ? Reflection->SelectedPoseDrawOptions : FPoseSearchDebuggerFeatureDrawOptions());
 
 		// Draw any selected database vectors
 		for (const TSharedRef<FDebuggerDatabaseRowData>& Row : SelectedRows)
@@ -1793,11 +1792,11 @@ void SDebuggerView::DrawFeatures(
 	// Draw skeleton
 	{
 		FSkeletonDrawParams SkeletonDrawParams;
-		if (Reflection->bDrawSelectedSkeleton)
+		if (Reflection && Reflection->bDrawSelectedSkeleton)
 		{
 			SkeletonDrawParams.Flags |= ESkeletonDrawFlags::SelectedPose;
 		}
-		if (Reflection->bDrawActiveSkeleton)
+		if (Reflection && Reflection->bDrawActiveSkeleton)
 		{
 			SkeletonDrawParams.Flags |= ESkeletonDrawFlags::ActivePose;
 		}
diff --git a/Engine/Plugins/Experimental/Animation/PoseSearch/Source/Editor/Private/PoseSearchDebugger.h b/Engine/Plugins/Experimental/Animation/PoseSearch/Source/Editor/Private/PoseSearchDebugger.h
--- a/Engine/Plugins/Experimental/Animation/PoseSearch/Source/Editor/Private/PoseSearchDebugger.h
+++ b/Engine/Plugins/Experimental/Animation/PoseSearch/Source/Editor/Private/PoseSearchDebugger.h
@@ -324,7 +324,7 @@ class SDebuggerDetailsView : public SCompoundWidget
 	void Update(const FTraceMotionMatchingStateMessage& State) const;
 
 	/** Get a const version of our reflection object */
-	const TObjectPtr<UPoseSearchDebuggerReflection>& GetReflection() const { return Reflection; }
+	const TObjectPtr<UPoseSearchDebuggerReflection> GetReflection() const { return Reflection; }
 	
 private:
 	/** Update our details view object with new state information */
diff --git a/Engine/Plugins/Experimental/Animation/PoseSearch/Source/Editor/Private/SPoseSearchDatabaseAssetList.cpp b/Engine/Plugins/Experimental/Animation/PoseSearch/Source/Editor/Private/SPoseSearchDatabaseAssetList.cpp
--- a/Engine/Plugins/Experimental/Animation/PoseSearch/Source/Editor/Private/SPoseSearchDatabaseAssetList.cpp
+++ b/Engine/Plugins/Experimental/Animation/PoseSearch/Source/Editor/Private/SPoseSearchDatabaseAssetList.cpp
@@ -146,7 +146,7 @@ namespace UE::PoseSearch
 			}
 			default:
 			{
-				Name = LOCTEXT("Default", "Default");
+				Name = FText::FromString(Database->GetName());
 				break;
 			}
 			}
diff --git a/Engine/Plugins/Experimental/Animation/PoseSearch/Source/Editor/Public/PoseSearchDatabaseEditorReflection.h b/Engine/Plugins/Experimental/Animation/PoseSearch/Source/Editor/Public/PoseSearchDatabaseEditorReflection.h
--- a/Engine/Plugins/Experimental/Animation/PoseSearch/Source/Editor/Public/PoseSearchDatabaseEditorReflection.h
+++ b/Engine/Plugins/Experimental/Animation/PoseSearch/Source/Editor/Public/PoseSearchDatabaseEditorReflection.h
@@ -58,3 +58,13 @@ public:
 	virtual void PostEditChangeProperty(struct FPropertyChangedEvent& PropertyChangedEvent) override;
 #endif
 };
+
+UCLASS()
+class UPoseSearchDatabaseReflection : public UPoseSearchDatabaseReflectionBase
+{
+	GENERATED_BODY()
+
+public:
+	UPROPERTY(VisibleAnywhere, Category = "Selected Database Search Index")
+	FPoseSearchIndex SearchIndex;
+};
diff --git a/Engine/Plugins/Experimental/Animation/PoseSearch/Source/Runtime/Private/PoseSearch.cpp b/Engine/Plugins/Experimental/Animation/PoseSearch/Source/Runtime/Private/PoseSearch.cpp
--- a/Engine/Plugins/Experimental/Animation/PoseSearch/Source/Runtime/Private/PoseSearch.cpp
+++ b/Engine/Plugins/Experimental/Animation/PoseSearch/Source/Runtime/Private/PoseSearch.cpp
@@ -1331,7 +1331,7 @@ static FVector BlendParameterForSampleRanges(
 		0.0f);
 }
 
-bool UPoseSearchDatabase::TryInitSearchIndexAssets(FPoseSearchIndex& OutSearchIndex)
+bool UPoseSearchDatabase::TryInitSearchIndexAssets(FPoseSearchIndex& OutSearchIndex) const
 {
 	OutSearchIndex.Assets.Empty();
 	
@@ -4283,7 +4283,8 @@ void DrawFeatureVector(const FDebugDrawParams& DrawParams, TArrayView<const floa
 
 void DrawFeatureVector(const FDebugDrawParams& DrawParams, int32 PoseIdx)
 {
-	if (PoseIdx != INDEX_NONE && DrawParams.CanDraw())
+	// if we're editing the schema while in PIE with Rewind Debugger active, PoseIdx could be out of bound / stale
+	if (DrawParams.CanDraw() && PoseIdx >= 0 && PoseIdx < DrawParams.GetSearchIndex()->NumPoses)
 	{
 		DrawFeatureVector(DrawParams, DrawParams.GetSearchIndex()->GetPoseValues(PoseIdx));
 	}
@@ -4366,24 +4367,33 @@ static void PreprocessSearchIndexWeights(FPoseSearchIndex& SearchIndex, const UP
 		Channel->FillWeights(SearchIndex.Weights);
 	}
 
-	// normalizing weights: the idea behind this step is to be able to compare poses from databases using different schemas
-	RowMajorVectorMap MapWeights(SearchIndex.Weights.GetData(), 1, NumDimensions);
-	const float WeightsSum = MapWeights.sum();
-	if (!FMath::IsNearlyZero(WeightsSum))
+	Eigen::VectorXd ChannelsMeanDeviations = ComputeChannelsMeanDeviations(Database->GetSearchIndex());
+	SearchIndex.Variance.Init(1.f, NumDimensions);
+	for (int32 Dimension = 0; Dimension != NumDimensions; ++Dimension)
 	{
-		MapWeights *= (1.0f / WeightsSum);
+		const float ChannelsMeanDeviation = ChannelsMeanDeviations[Dimension];
+		SearchIndex.Variance[Dimension] = ChannelsMeanDeviation * ChannelsMeanDeviation;
 	}
 
-	if (Database->GetSearchIndex()->Schema->DataPreprocessor == EPoseSearchDataPreprocessor::Normalize)
+	EPoseSearchDataPreprocessor DataPreprocessor = Database->GetSearchIndex()->Schema->DataPreprocessor;
+
+	if (DataPreprocessor == EPoseSearchDataPreprocessor::Normalize)
 	{
-		Eigen::VectorXd ChannelsMeanDeviations = ComputeChannelsMeanDeviations(Database->GetSearchIndex());
+		// normalizing user weights: the idea behind this step is to be able to compare poses from databases using different schemas
+		RowMajorVectorMap MapWeights(SearchIndex.Weights.GetData(), 1, NumDimensions);
+		const float WeightsSum = MapWeights.sum();
+		if (!FMath::IsNearlyZero(WeightsSum))
+		{
+			MapWeights *= (1.0f / WeightsSum);
+		}
+	}
 
+	if (DataPreprocessor == EPoseSearchDataPreprocessor::Normalize || DataPreprocessor == EPoseSearchDataPreprocessor::NormalizeOnlyByDeviation)
+	{
 		for (int32 Dimension = 0; Dimension != NumDimensions; ++Dimension)
 		{
 			// the idea here is to premultiply the weights by the inverse of the variance (proportional to the square of the deviation) to have a "weighted Mahalanobis" distance
-			const float ChannelsMeanDeviation = ChannelsMeanDeviations[Dimension];
-			const float ChannelsVariance = ChannelsMeanDeviation * ChannelsMeanDeviation;
-			SearchIndex.Weights[Dimension] /= ChannelsVariance;
+			SearchIndex.Weights[Dimension] /= SearchIndex.Variance[Dimension];
 		}
 	}
 }
@@ -4412,7 +4422,6 @@ static void PreprocessSearchIndexPCAData(FPoseSearchIndex& SearchIndex, const UP
 	// use the mean to center the data points
 	const RowMajorMatrix CenteredValues = WeightedValues.rowwise() - Mean;
 
-
 	// estimating the covariance matrix (with dimensionality of NumDimensions, NumDimensions)
 	// formula: https://en.wikipedia.org/wiki/Covariance_matrix#Estimation
 	// details: https://en.wikipedia.org/wiki/Estimation_of_covariance_matrices
@@ -4463,7 +4472,7 @@ static void PreprocessSearchIndexPCAData(FPoseSearchIndex& SearchIndex, const UP
 	// calculating the total variance knowing that eigen values measure variance along the principal components:
 	const float TotalVariance = EigenValues.sum();
 	// and explained variance as ratio between AccumulatedVariance and TotalVariance: https://ro-che.info/articles/2017-12-11-pca-explained-variance
-	const float ExplainedVariance = TotalVariance > UE_KINDA_SMALL_NUMBER ? AccumulatedVariance / TotalVariance : 0.f;
+	SearchIndex.PCAExplainedVariance = TotalVariance > UE_KINDA_SMALL_NUMBER ? AccumulatedVariance / TotalVariance : 0.f;
 
 	MapPCAValues = CenteredValues * PCAProjectionMatrix;
 
@@ -4489,12 +4498,14 @@ static void PreprocessSearchIndexKDTree(FPoseSearchIndex& SearchIndex, const UPo
 	if (Database->PoseSearchMode == EPoseSearchMode::PCAKDTree_Validate)
 	{
 		// testing the KDTree is returning the proper searches for all the points in pca space
+		int32 NumberOfFailingPoints = 0;
 		for (size_t PointIndex = 0; PointIndex < NumPoses; ++PointIndex)
 		{
-			constexpr size_t NumResults = 10;
-			size_t ResultIndexes[NumResults + 1] = { 0 };
-			float ResultDistanceSqr[NumResults + 1] = { 0 };
-			FKDTree::KNNResultSet ResultSet(NumResults, ResultIndexes, ResultDistanceSqr);
+			TArray<size_t> ResultIndexes;
+			TArray<float> ResultDistanceSqr;
+			ResultIndexes.SetNum(Database->KDTreeQueryNumNeighbors + 1);
+			ResultDistanceSqr.SetNum(Database->KDTreeQueryNumNeighbors + 1);
+			FKDTree::KNNResultSet ResultSet(Database->KDTreeQueryNumNeighbors, ResultIndexes, ResultDistanceSqr);
 			SearchIndex.KDTree.FindNeighbors(ResultSet, &SearchIndex.PCAValues[PointIndex * NumberOfPrincipalComponents]);
 
 			size_t ResultIndex = 0;
@@ -4506,17 +4517,24 @@ static void PreprocessSearchIndexKDTree(FPoseSearchIndex& SearchIndex, const UPo
 					break;
 				}
 			}
-			check(ResultIndex < ResultSet.Num());
+			if (ResultIndex == ResultSet.Num())
+			{
+				++NumberOfFailingPoints;
+			}
 		}
 
+		check(NumberOfFailingPoints == 0);
+
 		// testing the KDTree is returning the proper searches for all the original points transformed in pca space
+		NumberOfFailingPoints = 0;
 		const int32 NumDimensions = Database->Schema->SchemaCardinality;
 		for (size_t PointIndex = 0; PointIndex < NumPoses; ++PointIndex)
 		{
-			constexpr size_t NumResults = 10;
-			size_t ResultIndexes[NumResults + 1] = { 0 };
-			float ResultDistanceSqr[NumResults + 1] = { 0 };
-			FKDTree::KNNResultSet ResultSet(NumResults, ResultIndexes, ResultDistanceSqr);
+			TArray<size_t> ResultIndexes;
+			TArray<float> ResultDistanceSqr;
+			ResultIndexes.SetNum(Database->KDTreeQueryNumNeighbors + 1);
+			ResultDistanceSqr.SetNum(Database->KDTreeQueryNumNeighbors + 1);
+			FKDTree::KNNResultSet ResultSet(Database->KDTreeQueryNumNeighbors, ResultIndexes, ResultDistanceSqr);
 
 			const RowMajorVectorMapConst MapValues(&SearchIndex.Values[PointIndex * NumDimensions], 1, NumDimensions);
 			const RowMajorVectorMapConst MapWeights(SearchIndex.Weights.GetData(), 1, NumDimensions);
@@ -4538,8 +4556,13 @@ static void PreprocessSearchIndexKDTree(FPoseSearchIndex& SearchIndex, const UPo
 					break;
 				}
 			}
-			check(ResultIndex < ResultSet.Num());
+			if (ResultIndex == ResultSet.Num())
+			{
+				++NumberOfFailingPoints;
+			}
 		}
+
+		check(NumberOfFailingPoints == 0);
 	}
 }
 
@@ -4609,7 +4632,7 @@ bool BuildIndex(const UAnimSequence* Sequence, UPoseSearchSequenceMetaData* Sequ
 
 struct FDatabaseIndexingContext
 {
-	UPoseSearchDatabase* Database = nullptr;
+	const UPoseSearchDatabase* Database = nullptr;
 	FPoseSearchIndex* SearchIndex = nullptr;
 
 	FAssetSamplingContext SamplingContext;
@@ -4844,7 +4867,7 @@ void FDatabaseIndexingContext::JoinIndex()
 	check(TotalFloats == TotalPoses * Database->Schema->SchemaCardinality);
 }
 
-bool BuildIndex(UPoseSearchDatabase* Database, FPoseSearchIndex& OutSearchIndex)
+bool BuildIndex(const UPoseSearchDatabase* Database, FPoseSearchIndex& OutSearchIndex)
 {
 	check(Database);
 
diff --git a/Engine/Plugins/Experimental/Animation/PoseSearch/Source/Runtime/Public/PoseSearch/PoseSearch.h b/Engine/Plugins/Experimental/Animation/PoseSearch/Source/Runtime/Public/PoseSearch/PoseSearch.h
--- a/Engine/Plugins/Experimental/Animation/PoseSearch/Source/Runtime/Public/PoseSearch/PoseSearch.h
+++ b/Engine/Plugins/Experimental/Animation/PoseSearch/Source/Runtime/Public/PoseSearch/PoseSearch.h
@@ -102,6 +102,7 @@ enum class EPoseSearchDataPreprocessor : int32
 {
 	None,
 	Normalize,
+	NormalizeOnlyByDeviation,
 
 	Num UMETA(Hidden),
 	Invalid = Num UMETA(Hidden)
@@ -648,7 +649,7 @@ struct POSESEARCH_API FPoseSearchIndex
 {
 	GENERATED_BODY()
 
-	UPROPERTY()
+	UPROPERTY(Category = Info, VisibleAnywhere)
 	int32 NumPoses = 0;
 
 	UPROPERTY()
@@ -657,15 +658,23 @@ struct POSESEARCH_API FPoseSearchIndex
 	UPROPERTY()
 	TArray<float> PCAValues;
 
+	// @todo: editor only data?
+	UPROPERTY(Category = Info, VisibleAnywhere)
+	float PCAExplainedVariance = 0.f;
+	
+	// @todo: editor only data?
+	UPROPERTY(Category = Info, VisibleAnywhere)
+	TArray<float> Variance;
+
 	UE::PoseSearch::FKDTree KDTree;
 
-	UPROPERTY()
+	UPROPERTY(Category = Info, VisibleAnywhere)
 	TArray<float> PCAProjectionMatrix;
 
-	UPROPERTY()
+	UPROPERTY(Category = Info, VisibleAnywhere)
 	TArray<float> Mean;
 
-	UPROPERTY()
+	UPROPERTY(Category = Info, VisibleAnywhere) 
 	TArray<float> Weights;
 
 	UPROPERTY()
@@ -678,7 +687,7 @@ struct POSESEARCH_API FPoseSearchIndex
 	TArray<FPoseSearchIndexAsset> Assets;
 
 	// minimum of the database metadata CostAddend: it represents the minimum cost of any search for the associated database (we'll skip the search in case the search result total cost is already less than MinCostAddend)
-	UPROPERTY()
+	UPROPERTY(Category = Info, VisibleAnywhere)
 	float MinCostAddend = -MAX_FLT;
 
 	bool IsValid() const;
@@ -1036,18 +1045,18 @@ public:
 	UPROPERTY(EditAnywhere, Category = "Database")
 	TArray<FPoseSearchDatabaseBlendSpace> BlendSpaces;
 
-	UPROPERTY(EditAnywhere, Category = "Performance", meta = (ClampMin = "1", ClampMax = "64", UIMin = "1", UIMax = "64"))
+	UPROPERTY(EditAnywhere, Category = "Performance")
+	EPoseSearchMode PoseSearchMode = EPoseSearchMode::BruteForce;
+
+	UPROPERTY(EditAnywhere, Category = "Performance", meta = (EditCondition = "PoseSearchMode != EPoseSearchMode::BruteForce", EditConditionHides, ClampMin = "1", ClampMax = "64", UIMin = "1", UIMax = "64"))
 	int32 NumberOfPrincipalComponents = 4;
 
-	UPROPERTY(EditAnywhere, Category = "Performance", meta = (ClampMin = "1", ClampMax = "256", UIMin = "1", UIMax = "256"))
+	UPROPERTY(EditAnywhere, Category = "Performance", meta = (EditCondition = "PoseSearchMode != EPoseSearchMode::BruteForce", EditConditionHides, ClampMin = "1", ClampMax = "256", UIMin = "1", UIMax = "256"))
 	int32 KDTreeMaxLeafSize = 8;
 	
-	UPROPERTY(EditAnywhere, Category = "Performance", meta = (ClampMin = "1", ClampMax = "600", UIMin = "1", UIMax = "600"))
+	UPROPERTY(EditAnywhere, Category = "Performance", meta = (EditCondition = "PoseSearchMode != EPoseSearchMode::BruteForce", EditConditionHides, ClampMin = "1", ClampMax = "600", UIMin = "1", UIMax = "600"))
 	int32 KDTreeQueryNumNeighbors = 100;
 
-	UPROPERTY(EditAnywhere, Category = "Performance")
-	EPoseSearchMode PoseSearchMode = EPoseSearchMode::BruteForce;
-
 	// if true, this database search will be skipped if cannot decrease the pose cost, and poses will not be listed into the PoseSearchDebugger
 	UPROPERTY(EditAnywhere, Category = "Performance")
 	bool bSkipSearchIfPossible = false;
@@ -1100,7 +1109,7 @@ private:
 
 public:
 	// Populates the FPoseSearchIndex::Assets array by evaluating the data in the Sequences array
-	bool TryInitSearchIndexAssets(FPoseSearchIndex& OutSearchIndex);
+	bool TryInitSearchIndexAssets(FPoseSearchIndex& OutSearchIndex) const;
 
 private:
 	FPoseSearchDatabaseDerivedData* PrivateDerivedData;
@@ -1425,7 +1434,7 @@ POSESEARCH_API bool BuildIndex(const UAnimSequence* Sequence, UPoseSearchSequenc
 * 
 * @return Whether the index was built successfully
 */
-POSESEARCH_API bool BuildIndex(UPoseSearchDatabase* Database, FPoseSearchIndex& OutSearchIndex);
+POSESEARCH_API bool BuildIndex(const UPoseSearchDatabase* Database, FPoseSearchIndex& OutSearchIndex);
 
 } // namespace UE::PoseSearch
 
