motion matching - fix for numerical instability by storing sqrt of weights. fix for pca proj matrix construction by premultiplying the dataset by the sqrt of the weights rather than the weights

diff --git a/Engine/Plugins/Experimental/Animation/PoseSearch/Source/Runtime/Private/PoseSearch.cpp b/Engine/Plugins/Experimental/Animation/PoseSearch/Source/Runtime/Private/PoseSearch.cpp
--- a/Engine/Plugins/Experimental/Animation/PoseSearch/Source/Runtime/Private/PoseSearch.cpp
+++ b/Engine/Plugins/Experimental/Animation/PoseSearch/Source/Runtime/Private/PoseSearch.cpp
@@ -81,23 +81,39 @@ static bool IsSamplingRangeValid(FFloatInterval Range)
 	return Range.IsValid() && (Range.Min >= 0.0f);
 }
 
-static inline float CompareFeatureVectors(int32 NumValues, const float* A, const float* B, const float* Weights)
+static inline float CompareFeatureVectors(TConstArrayView<const float> A, TConstArrayView<const float> B, TConstArrayView<const float> WeightsSqrt)
 {
-	Eigen::Map<const Eigen::ArrayXf> VA(A, NumValues);
-	Eigen::Map<const Eigen::ArrayXf> VB(B, NumValues);
-	Eigen::Map<const Eigen::ArrayXf> VW(Weights, NumValues);
+	check(A.Num() == B.Num() && A.Num() == WeightsSqrt.Num());
 
-	return ((VA - VB).square() * VW).sum();
+	Eigen::Map<const Eigen::ArrayXf> VA(A.GetData(), A.Num());
+	Eigen::Map<const Eigen::ArrayXf> VB(B.GetData(), B.Num());
+	Eigen::Map<const Eigen::ArrayXf> VW(WeightsSqrt.GetData(), WeightsSqrt.Num());
+
+	return ((VA - VB) * VW).square().sum();
 }
 
-static inline float CompareFeatureVectors(int32 NumValues, const float* A, const float* B)
+static inline float CompareFeatureVectors(TConstArrayView<const float> A, TConstArrayView<const float> B)
 {
-	Eigen::Map<const Eigen::ArrayXf> VA(A, NumValues);
-	Eigen::Map<const Eigen::ArrayXf> VB(B, NumValues);
+	check(A.Num() == B.Num());
+
+	Eigen::Map<const Eigen::ArrayXf> VA(A.GetData(), A.Num());
+	Eigen::Map<const Eigen::ArrayXf> VB(B.GetData(), B.Num());
 
 	return (VA - VB).square().sum();
 }
 
+static inline void CompareFeatureVectors(TConstArrayView<const float> A, TConstArrayView<const float> B, TConstArrayView<const float> WeightsSqrt, TArrayView<float> Result)
+{
+	check(A.Num() == B.Num() && A.Num() == WeightsSqrt.Num() && A.Num() == Result.Num());
+
+	Eigen::Map<const Eigen::ArrayXf> VA(A.GetData(), A.Num());
+	Eigen::Map<const Eigen::ArrayXf> VB(B.GetData(), B.Num());
+	Eigen::Map<const Eigen::ArrayXf> VW(WeightsSqrt.GetData(), WeightsSqrt.Num());
+	Eigen::Map<Eigen::ArrayXf> VR(Result.GetData(), Result.Num());
+
+	VR = ((VA - VB) * VW).square();
+}
+
 static FFloatInterval GetEffectiveSamplingRange(const UAnimSequenceBase* Sequence, FFloatInterval RequestedSamplingRange)
 {
 	const bool bSampleAll = (RequestedSamplingRange.Min == 0.0f) && (RequestedSamplingRange.Max == 0.0f);
@@ -657,7 +673,7 @@ bool FPoseSearchIndex::IsEmpty() const
 	return bEmpty;
 }
 
-TArrayView<const float> FPoseSearchIndex::GetPoseValues(int32 PoseIdx) const
+TConstArrayView<const float> FPoseSearchIndex::GetPoseValues(int32 PoseIdx) const
 {
 	check(PoseIdx < NumPoses);
 	int32 ValueOffset = PoseIdx * Schema->SchemaCardinality;
@@ -797,25 +813,18 @@ UE::PoseSearch::FSearchResult UPoseSearchSequenceMetaData::Search(UE::PoseSearch
 FPoseSearchCost UPoseSearchSequenceMetaData::ComparePoses(
 	int32 PoseIdx,
 	UE::PoseSearch::EPoseComparisonFlags PoseComparisonFlags,
-	const TArrayView<const float>& QueryValues) const
+	TConstArrayView<const float> QueryValues) const
 {
 	using namespace UE::PoseSearch;
 
-	FPoseSearchCost Result;
-
-	TArrayView<const float> PoseValues = SearchIndex.GetPoseValues(PoseIdx);
-	check(PoseValues.Num() == QueryValues.Num());
-
-	const float DissimilarityCost = CompareFeatureVectors(PoseValues.Num(), PoseValues.GetData(), QueryValues.GetData());
+	const float DissimilarityCost = CompareFeatureVectors(SearchIndex.GetPoseValues(PoseIdx), QueryValues);
 
 	// @todo: shouldn't we include MirrorMismatchAddend as well?
 	//const float MirrorMismatchAddend = SearchIndex.ComputeMirrorMismatchAddend(PoseIdx, SearchContext);
 	const float NotifyAddend = SearchIndex.ComputeNotifyAddend(PoseIdx);
 	const float ContinuingPoseCostAddend = SearchIndex.ComputeContinuingPoseCostAddend(PoseIdx, PoseComparisonFlags);
 
-	Result.Set(DissimilarityCost, NotifyAddend + ContinuingPoseCostAddend);
-
-	return Result;
+	return FPoseSearchCost(DissimilarityCost, NotifyAddend + ContinuingPoseCostAddend);
 }
 
 
@@ -1720,13 +1729,9 @@ FPoseSearchCost UPoseSearchDatabase::ComparePoses(UE::PoseSearch::FSearchContext
 	const FPoseSearchIndex* SearchIndex = GetSearchIndex();
 	check(SearchIndex);
 
-	TArrayView<const float> PoseValues = SearchIndex->GetPoseValues(PoseIdx);
-	const int32 Dims = PoseValues.Num();
-	check(Dims == QueryValues.Num());
-
-	TArrayView<const float> Weights(SearchIndex->Weights);
+	TConstArrayView<const float> PoseValues = SearchIndex->GetPoseValues(PoseIdx);
 
-	const float DissimilarityCost = CompareFeatureVectors(PoseValues.Num(), PoseValues.GetData(), QueryValues.GetData(), Weights.GetData());
+	const float DissimilarityCost = CompareFeatureVectors(PoseValues, QueryValues, SearchIndex->WeightsSqrt);
 
 	const float MirrorMismatchAddend = SearchIndex->ComputeMirrorMismatchAddend(PoseIdx, SearchContext);
 	const float NotifyAddend = SearchIndex->ComputeNotifyAddend(PoseIdx);
@@ -1738,14 +1743,10 @@ FPoseSearchCost UPoseSearchDatabase::ComparePoses(UE::PoseSearch::FSearchContext
 	if (SearchContext.bIsTracing)
 	{
 		// Setup Eigen views onto our vectors
-		Result.CostVector.SetNum(Dims);
-		auto OutCostVector = Map<ArrayXf>(Result.CostVector.GetData(), Dims);
-		auto PoseVector = Map<const ArrayXf>(PoseValues.GetData(), Dims);
-		auto QueryVector = Map<const ArrayXf>(QueryValues.GetData(), Dims);
-		auto WeightsVector = Map<const ArrayXf>(Weights.GetData(), Dims);
-		
+		Result.CostVector.SetNum(PoseValues.Num());
+
 		// Compute weighted squared difference vector
-		OutCostVector = WeightsVector * (PoseVector - QueryVector).square();
+		CompareFeatureVectors(PoseValues, QueryValues, SearchIndex->WeightsSqrt, Result.CostVector);
 
 		// Output cost details
 		CalcChannelCosts(SearchIndex->Schema, Result.CostVector, Result.ChannelCosts);
@@ -1838,7 +1839,7 @@ UE::PoseSearch::FSearchResult UPoseSearchDatabase::SearchPCAKDTree(UE::PoseSearc
 	// KDTree in PCA space search
 	if (PoseSearchMode == EPoseSearchMode::PCAKDTree_Validate)
 	{
-		const RowMajorVectorMapConst MapWeights(SearchIndex->Weights.GetData(), 1, NumDimensions);
+		const RowMajorVectorMapConst MapWeightsSqrt(SearchIndex->WeightsSqrt.GetData(), 1, NumDimensions);
 
 		// testing the KDTree is returning the proper searches for all the original points transformed in pca space
 		for (int32 PoseIdx = 0; PoseIdx < SearchIndex->NumPoses; ++PoseIdx)
@@ -1850,7 +1851,7 @@ UE::PoseSearch::FSearchResult UPoseSearchDatabase::SearchPCAKDTree(UE::PoseSearc
 			const ColMajorMatrixMapConst PCAProjectionMatrix(SearchIndex->PCAProjectionMatrix.GetData(), NumDimensions, ClampedNumberOfPrincipalComponents);
 
 			const RowMajorVectorMapConst QueryValues(PoseValues.GetData(), 1, NumDimensions);
-			WeightedQueryValues = QueryValues.array() * MapWeights.array();
+			WeightedQueryValues = QueryValues.array() * MapWeightsSqrt.array();
 			CenteredQueryValues.noalias() = WeightedQueryValues - Mean;
 			ProjectedQueryValues.noalias() = CenteredQueryValues * PCAProjectionMatrix;
 
@@ -1892,10 +1893,9 @@ UE::PoseSearch::FSearchResult UPoseSearchDatabase::SearchPCAKDTree(UE::PoseSearc
 	// there's no point in performing the search if CurrentBestTotalCost is already better than that
 	if (SearchContext.GetCurrentBestTotalCost() > SearchIndex->MinCostAddend)
 	{
-		// we're offsetting the NonSelectableIdx by "-GroupSearchIndex.StartPoseIndex" to have the indexes in kdtree space rather than in PoseIdx database space (kdtree index 0 is PoseIdx of GroupSearchIndex.StartPoseIndex)
 		const int NonSelectableIdxUsedSize = PopulateNonSelectableIdx(NonSelectableIdxData, NonSelectableIdxDataSize, SearchContext, this);
 		TArrayView<size_t> NonSelectableIdx(NonSelectableIdxData, NonSelectableIdxUsedSize);
-		const RowMajorVectorMapConst MapWeights(SearchIndex->Weights.GetData(), 1, NumDimensions);
+		const RowMajorVectorMapConst MapWeightsSqrt(SearchIndex->WeightsSqrt.GetData(), 1, NumDimensions);
 		FKDTree::KNNResultSet ResultSet(ClampedKDTreeQueryNumNeighbors, ResultIndexes, ResultDistanceSqr, NonSelectableIdx);
 
 		check(QueryValues.Num() == NumDimensions);
@@ -1905,7 +1905,7 @@ UE::PoseSearch::FSearchResult UPoseSearchDatabase::SearchPCAKDTree(UE::PoseSearc
 
 		// transforming query values into PCA space to query the KDTree
 		const RowMajorVectorMapConst QueryValuesMap(QueryValues.GetData(), 1, NumDimensions);
-		WeightedQueryValues = QueryValuesMap.array() * MapWeights.array();
+		WeightedQueryValues = QueryValuesMap.array() * MapWeightsSqrt.array();
 		CenteredQueryValues.noalias() = WeightedQueryValues - Mean;
 		ProjectedQueryValues.noalias() = CenteredQueryValues * PCAProjectionMatrix;
 
@@ -4303,7 +4303,7 @@ void DrawSearchIndex(const FDebugDrawParams& DrawParams)
 	}
 }
 
-static Eigen::VectorXd ComputeChannelsMeanDeviations(const FPoseSearchIndex* SearchIndex)
+static Eigen::VectorXd ComputeChannelsDeviations(const FPoseSearchIndex* SearchIndex)
 {
 	// This function performs a modified z-score normalization where features are normalized
 	// by mean absolute deviation rather than standard deviation. Both methods are preferable
@@ -4359,20 +4359,20 @@ static Eigen::VectorXd ComputeChannelsMeanDeviations(const FPoseSearchIndex* Sea
 static void PreprocessSearchIndexWeights(FPoseSearchIndex& SearchIndex, const UPoseSearchDatabase* Database)
 {
 	const int32 NumDimensions = Database->Schema->SchemaCardinality;
-	SearchIndex.Weights.Init(1.f, NumDimensions);
+	SearchIndex.WeightsSqrt.Init(1.f, NumDimensions);
 
 	for (int ChannelIdx = 0; ChannelIdx != Database->Schema->Channels.Num(); ++ChannelIdx)
 	{
 		const UPoseSearchFeatureChannel* Channel = Database->Schema->Channels[ChannelIdx].Get();
-		Channel->FillWeights(SearchIndex.Weights);
+		Channel->FillWeights(SearchIndex.WeightsSqrt);
 	}
 
-	Eigen::VectorXd ChannelsMeanDeviations = ComputeChannelsMeanDeviations(Database->GetSearchIndex());
-	SearchIndex.Variance.Init(1.f, NumDimensions);
+	Eigen::VectorXd ChannelsMeanDeviations = ComputeChannelsDeviations(Database->GetSearchIndex());
+	SearchIndex.Deviation.Init(1.f, NumDimensions);
 	for (int32 Dimension = 0; Dimension != NumDimensions; ++Dimension)
 	{
 		const float ChannelsMeanDeviation = ChannelsMeanDeviations[Dimension];
-		SearchIndex.Variance[Dimension] = ChannelsMeanDeviation * ChannelsMeanDeviation;
+		SearchIndex.Deviation[Dimension] = ChannelsMeanDeviation;
 	}
 
 	EPoseSearchDataPreprocessor DataPreprocessor = Database->GetSearchIndex()->Schema->DataPreprocessor;
@@ -4380,7 +4380,7 @@ static void PreprocessSearchIndexWeights(FPoseSearchIndex& SearchIndex, const UP
 	if (DataPreprocessor == EPoseSearchDataPreprocessor::Normalize)
 	{
 		// normalizing user weights: the idea behind this step is to be able to compare poses from databases using different schemas
-		RowMajorVectorMap MapWeights(SearchIndex.Weights.GetData(), 1, NumDimensions);
+		RowMajorVectorMap MapWeights(SearchIndex.WeightsSqrt.GetData(), 1, NumDimensions);
 		const float WeightsSum = MapWeights.sum();
 		if (!FMath::IsNearlyZero(WeightsSum))
 		{
@@ -4388,12 +4388,18 @@ static void PreprocessSearchIndexWeights(FPoseSearchIndex& SearchIndex, const UP
 		}
 	}
 
+	// extracting the square root
+	for (int32 Dimension = 0; Dimension != NumDimensions; ++Dimension)
+	{
+		SearchIndex.WeightsSqrt[Dimension] = FMath::Sqrt(SearchIndex.WeightsSqrt[Dimension]);
+	}
+
 	if (DataPreprocessor == EPoseSearchDataPreprocessor::Normalize || DataPreprocessor == EPoseSearchDataPreprocessor::NormalizeOnlyByDeviation)
 	{
 		for (int32 Dimension = 0; Dimension != NumDimensions; ++Dimension)
 		{
 			// the idea here is to premultiply the weights by the inverse of the variance (proportional to the square of the deviation) to have a "weighted Mahalanobis" distance
-			SearchIndex.Weights[Dimension] /= SearchIndex.Variance[Dimension];
+			SearchIndex.WeightsSqrt[Dimension] /= SearchIndex.Deviation[Dimension];
 		}
 	}
 }
@@ -4405,9 +4411,9 @@ static void PreprocessSearchIndexPCAData(FPoseSearchIndex& SearchIndex, const UP
 	const int32 NumDimensions = Database->Schema->SchemaCardinality;
 	const int32 NumPoses = SearchIndex.NumPoses;
 
-	const RowMajorVectorMapConst MapWeights(SearchIndex.Weights.GetData(), 1, NumDimensions);
+	const RowMajorVectorMapConst MapWeightsSqrt(SearchIndex.WeightsSqrt.GetData(), 1, NumDimensions);
 	const RowMajorMatrixMapConst MapValues(SearchIndex.Values.GetData(), NumPoses, NumDimensions);
-	const RowMajorMatrix WeightedValues = MapValues.array().rowwise() * MapWeights.array();
+	const RowMajorMatrix WeightedValues = MapValues.array().rowwise() * MapWeightsSqrt.array();
 	const uint32 NumberOfPrincipalComponents = Database->GetNumberOfPrincipalComponents();
 
 	SearchIndex.PCAValues.Reset();
@@ -4435,12 +4441,12 @@ static void PreprocessSearchIndexPCAData(FPoseSearchIndex& SearchIndex, const UP
 
 	if (Database->PoseSearchMode == EPoseSearchMode::PCAKDTree_Validate && NumberOfPrincipalComponents == NumDimensions)
 	{
-		const RowMajorVector ReciprocalWeights = MapWeights.cwiseInverse();
+		const RowMajorVector ReciprocalWeightsSqrt = MapWeightsSqrt.cwiseInverse();
 		const RowMajorMatrix ProjectedValues = CenteredValues * EigenVectors;
 		for (Eigen::Index RowIndex = 0; RowIndex < MapValues.rows(); ++RowIndex)
 		{
 			const RowMajorVector WeightedReconstructedPoint = ProjectedValues.row(RowIndex) * EigenVectors.transpose() + Mean;
-			const RowMajorVector ReconstructedPoint = WeightedReconstructedPoint.array() * ReciprocalWeights.array();
+			const RowMajorVector ReconstructedPoint = WeightedReconstructedPoint.array() * ReciprocalWeightsSqrt.array();
 			const float Error = (ReconstructedPoint - MapValues.row(RowIndex)).squaredNorm();
 			check(Error < UE_KINDA_SMALL_NUMBER);
 		}
@@ -4478,11 +4484,11 @@ static void PreprocessSearchIndexPCAData(FPoseSearchIndex& SearchIndex, const UP
 
 	if (Database->PoseSearchMode == EPoseSearchMode::PCAKDTree_Validate && NumberOfPrincipalComponents == NumDimensions)
 	{
-		const RowMajorVector ReciprocalWeights = MapWeights.cwiseInverse();
+		const RowMajorVector ReciprocalWeightsSqrt = MapWeightsSqrt.cwiseInverse();
 		for (Eigen::Index RowIndex = 0; RowIndex < MapValues.rows(); ++RowIndex)
 		{
 			const RowMajorVector WeightedReconstructedValues = MapPCAValues.row(RowIndex) * PCAProjectionMatrix.transpose() + Mean;
-			const RowMajorVector ReconstructedValues = WeightedReconstructedValues.array() * ReciprocalWeights.array();
+			const RowMajorVector ReconstructedValues = WeightedReconstructedValues.array() * ReciprocalWeightsSqrt.array();
 			const float Error = (ReconstructedValues - MapValues.row(RowIndex)).squaredNorm();
 			check(Error < UE_KINDA_SMALL_NUMBER);
 		}
@@ -4537,11 +4543,11 @@ static void PreprocessSearchIndexKDTree(FPoseSearchIndex& SearchIndex, const UPo
 			FKDTree::KNNResultSet ResultSet(Database->KDTreeQueryNumNeighbors, ResultIndexes, ResultDistanceSqr);
 
 			const RowMajorVectorMapConst MapValues(&SearchIndex.Values[PointIndex * NumDimensions], 1, NumDimensions);
-			const RowMajorVectorMapConst MapWeights(SearchIndex.Weights.GetData(), 1, NumDimensions);
+			const RowMajorVectorMapConst MapWeightsSqrt(SearchIndex.WeightsSqrt.GetData(), 1, NumDimensions);
 			const RowMajorVectorMapConst Mean(SearchIndex.Mean.GetData(), 1, NumDimensions);
 			const ColMajorMatrixMapConst PCAProjectionMatrix(SearchIndex.PCAProjectionMatrix.GetData(), NumDimensions, NumberOfPrincipalComponents);
 
-			const RowMajorMatrix WeightedValues = MapValues.array() * MapWeights.array();
+			const RowMajorMatrix WeightedValues = MapValues.array() * MapWeightsSqrt.array();
 			const RowMajorMatrix CenteredValues = WeightedValues - Mean;
 			const RowMajorVector ProjectedValues  = CenteredValues * PCAProjectionMatrix;
 
diff --git a/Engine/Plugins/Experimental/Animation/PoseSearch/Source/Runtime/Private/PoseSearchDerivedData.cpp b/Engine/Plugins/Experimental/Animation/PoseSearch/Source/Runtime/Private/PoseSearchDerivedData.cpp
--- a/Engine/Plugins/Experimental/Animation/PoseSearch/Source/Runtime/Private/PoseSearchDerivedData.cpp
+++ b/Engine/Plugins/Experimental/Animation/PoseSearch/Source/Runtime/Private/PoseSearchDerivedData.cpp
@@ -267,7 +267,7 @@ namespace UE::PoseSearch
 			Ar.Serialize(&Index.Assets[0], Index.Assets.Num() * Index.Assets.GetTypeSize());
 		}
 
-		Ar << Index.Weights;
+		Ar << Index.WeightsSqrt;
 		Ar << Index.Mean;
 		Ar << Index.PCAProjectionMatrix;
 
diff --git a/Engine/Plugins/Experimental/Animation/PoseSearch/Source/Runtime/Public/PoseSearch/PoseSearch.h b/Engine/Plugins/Experimental/Animation/PoseSearch/Source/Runtime/Public/PoseSearch/PoseSearch.h
--- a/Engine/Plugins/Experimental/Animation/PoseSearch/Source/Runtime/Public/PoseSearch/PoseSearch.h
+++ b/Engine/Plugins/Experimental/Animation/PoseSearch/Source/Runtime/Public/PoseSearch/PoseSearch.h
@@ -664,7 +664,7 @@ struct POSESEARCH_API FPoseSearchIndex
 	
 	// @todo: editor only data?
 	UPROPERTY(Category = Info, VisibleAnywhere)
-	TArray<float> Variance;
+	TArray<float> Deviation;
 
 	UE::PoseSearch::FKDTree KDTree;
 
@@ -674,8 +674,13 @@ struct POSESEARCH_API FPoseSearchIndex
 	UPROPERTY(Category = Info, VisibleAnywhere)
 	TArray<float> Mean;
 
+	// we store weights square roots to reduce numerical errors when CompareFeatureVectors 
+	// ((VA - VB) * VW).square().sum()
+	// instead of
+	// ((VA - VB).square() * VW).sum()
+	// since (VA - VB).square() could lead to big numbers, and VW being multiplied by the variance of the dataset
 	UPROPERTY(Category = Info, VisibleAnywhere) 
-	TArray<float> Weights;
+	TArray<float> WeightsSqrt;
 
 	UPROPERTY()
 	TArray<FPoseSearchPoseMetadata> PoseMetadata;
@@ -694,7 +699,7 @@ struct POSESEARCH_API FPoseSearchIndex
 	bool IsValidPoseIndex(int32 PoseIdx) const { return PoseIdx < NumPoses; }
 	bool IsEmpty() const;
 
-	TArrayView<const float> GetPoseValues(int32 PoseIdx) const;
+	TConstArrayView<const float> GetPoseValues(int32 PoseIdx) const;
 
 	int32 FindAssetIndex(const FPoseSearchIndexAsset* Asset) const;
 
@@ -1193,7 +1198,7 @@ public:
 	UE::PoseSearch::FSearchResult Search(UE::PoseSearch::FSearchContext& SearchContext) const;
 
 protected:
-	FPoseSearchCost ComparePoses(int32 PoseIdx, UE::PoseSearch::EPoseComparisonFlags PoseComparisonFlags, const TArrayView<const float>& QueryValues) const;
+	FPoseSearchCost ComparePoses(int32 PoseIdx, UE::PoseSearch::EPoseComparisonFlags PoseComparisonFlags, TConstArrayView<const float> QueryValues) const;
 
 public: // UObject
 	virtual void PreSave(FObjectPreSaveContext ObjectSaveContext) override;
