Update Pose Search UI
- Include Filtering
- Mass Enable / Disable
- Entries Appear in Alphabetical Order

diff --git a/Engine/Plugins/Experimental/Animation/PoseSearch/Source/Editor/Private/PoseSearchDatabaseViewModel.cpp b/Engine/Plugins/Experimental/Animation/PoseSearch/Source/Editor/Private/PoseSearchDatabaseViewModel.cpp
--- a/Engine/Plugins/Experimental/Animation/PoseSearch/Source/Editor/Private/PoseSearchDatabaseViewModel.cpp
+++ b/Engine/Plugins/Experimental/Animation/PoseSearch/Source/Editor/Private/PoseSearchDatabaseViewModel.cpp
@@ -321,6 +321,7 @@ namespace UE::PoseSearch
 	void FDatabaseViewModel::AddSequenceToDatabase(UAnimSequence* AnimSequence)
 	{
 		FPoseSearchDatabaseSequence& NewDbSequence = PoseSearchDatabase->Sequences.AddDefaulted_GetRef();
+
 		NewDbSequence.Sequence = AnimSequence;
 	}
 
@@ -343,21 +344,16 @@ namespace UE::PoseSearch
 	void FDatabaseViewModel::SetSelectedSequenceEnabled(int32 SequenceIndex, bool bEnabled)
 	{
 		PoseSearchDatabase->Sequences[SequenceIndex].bEnabled = bEnabled;
-		PoseSearchDatabase->NotifyAssetChange();
-		PoseSearchDatabase->BeginCacheDerivedData();
 	}	
 	
 	void FDatabaseViewModel::SetSelectedBlendSpaceEnabled(int32 BlendSpaceIndex, bool bEnabled)
 	{
 		PoseSearchDatabase->BlendSpaces[BlendSpaceIndex].bEnabled = bEnabled;
-		PoseSearchDatabase->NotifyAssetChange();
-		PoseSearchDatabase->BeginCacheDerivedData();
 	}
 
 	bool FDatabaseViewModel::IsSelectedSequenceEnabled(int32 SequenceIndex) const
 	{
 		return PoseSearchDatabase->Sequences[SequenceIndex].bEnabled;
-		PoseSearchDatabase->NotifyAssetChange();
 	}
 
 	bool FDatabaseViewModel::IsSelectedBlendSpaceEnabled(int32 BlendSpaceIndex) const
diff --git a/Engine/Plugins/Experimental/Animation/PoseSearch/Source/Editor/Private/SPoseSearchDatabaseAssetList.cpp b/Engine/Plugins/Experimental/Animation/PoseSearch/Source/Editor/Private/SPoseSearchDatabaseAssetList.cpp
--- a/Engine/Plugins/Experimental/Animation/PoseSearch/Source/Editor/Private/SPoseSearchDatabaseAssetList.cpp
+++ b/Engine/Plugins/Experimental/Animation/PoseSearch/Source/Editor/Private/SPoseSearchDatabaseAssetList.cpp
@@ -14,6 +14,7 @@
 
 #include "Widgets/Text/SRichTextBlock.h"
 #include "Widgets/Input/SCheckBox.h"
+#include "Widgets/Input/SSearchBox.h"
 #include "Framework/Commands/GenericCommands.h"
 #include "Framework/MultiBox/MultiBoxBuilder.h"
 #include "SPositiveActionButton.h"
@@ -410,6 +411,14 @@ namespace UE::PoseSearch
 					.ToolTipText(LOCTEXT("AddNewToolTip", "Add a new Sequence, Blend Space or Group"))
 					.OnGetMenuContent(this, &SDatabaseAssetTree::CreateAddNewMenuWidget)
 				]
+				+ SHorizontalBox::Slot()
+				.AutoWidth()
+				.VAlign(VAlign_Center)
+				.HAlign(HAlign_Right)
+				.Padding(2, 0, 0, 0)
+				[
+					GenerateFilterBoxWidget()
+				]
 			]
 			+SVerticalBox::Slot()
 			.Padding(0.0f, 0.0f)
@@ -534,18 +543,85 @@ namespace UE::PoseSearch
 			AllNodes.Add(SequenceGroupNode);
 		};
 
+		// Build an index based off of alphabetical order than iterate the index instead
+		TArray<uint32> SequenceIndexArray;
+		SequenceIndexArray.SetNumUninitialized(Database->Sequences.Num());
+		for (int32 SequenceIdx = 0; SequenceIdx < Database->Sequences.Num(); ++SequenceIdx)
+		{
+			SequenceIndexArray[SequenceIdx] = SequenceIdx;
+		}
+
+		SequenceIndexArray.Sort([Database](int32 SequenceIdxA, int32 SequenceIdxB)
+		{
+			const FPoseSearchDatabaseSequence& A = Database->Sequences[SequenceIdxA];
+			const FPoseSearchDatabaseSequence& B = Database->Sequences[SequenceIdxB];
+
+			//If its null add it to the end of the list 
+			if (B.Sequence.IsNull())
+			{
+				return true;
+			}
+
+			if (A.Sequence.IsNull())
+			{
+				return false;
+			}
+
+			const int32 Comparison = A.Sequence->GetName().Compare(B.Sequence->GetName());
+			return Comparison < 0;
+		});
+
 		// create all sequence nodes
 		for (int32 SequenceIdx = 0; SequenceIdx < Database->Sequences.Num(); ++SequenceIdx)
 		{
-			const FPoseSearchDatabaseSequence& DbSequence = Database->Sequences[SequenceIdx];
-			CreateAssetNode(SequenceIdx, ESearchIndexAssetType::Sequence, DefaultGroupIdx);
+			const int32 MappedId = SequenceIndexArray[SequenceIdx];
+			const FPoseSearchDatabaseSequence& DbSequence = Database->Sequences[MappedId];
+			const bool bFiltered = (DbSequence.Sequence == nullptr || GetAssetFilterString().IsEmpty()) ? false : !DbSequence.Sequence->GetName().Contains(GetAssetFilterString());
+
+			if (!bFiltered)
+			{
+				CreateAssetNode(MappedId, ESearchIndexAssetType::Sequence, DefaultGroupIdx);
+			}
 		}
 
+		TArray<uint32> BlendspaceIndexArray;
+		SequenceIndexArray.SetNumUninitialized(Database->BlendSpaces.Num());
+		for (int32 BlendspaceIdx = 0; BlendspaceIdx < Database->BlendSpaces.Num(); ++BlendspaceIdx)
+		{
+			BlendspaceIndexArray[BlendspaceIdx] = BlendspaceIdx;
+		}
+
+		BlendspaceIndexArray.Sort([Database](int32 BlendspaceIdxA, int32 BlendspaceIdxB)
+		{
+			const FPoseSearchDatabaseBlendSpace& A = Database->BlendSpaces[BlendspaceIdxA];
+			const FPoseSearchDatabaseBlendSpace& B = Database->BlendSpaces[BlendspaceIdxB];
+
+			//If its null add it to the end of the list 
+			if (B.BlendSpace.IsNull())
+			{
+				return true;
+			}
+
+			if (A.BlendSpace.IsNull())
+			{
+				return false;
+			}
+
+			const int32 Comparison = A.BlendSpace->GetName().Compare(B.BlendSpace->GetName());
+			return Comparison < 0;
+		});
+
 		// create all blendspace nodes
 		for (int32 BlendSpaceIdx = 0; BlendSpaceIdx < Database->BlendSpaces.Num(); ++BlendSpaceIdx)
 		{
-			const FPoseSearchDatabaseBlendSpace& DbBlendSpace = Database->BlendSpaces[BlendSpaceIdx];
-			CreateAssetNode(BlendSpaceIdx, ESearchIndexAssetType::BlendSpace, DefaultGroupIdx);
+			const int32 MappedId = BlendspaceIndexArray[BlendSpaceIdx];
+			const FPoseSearchDatabaseBlendSpace& DbBlendSpace = Database->BlendSpaces[MappedId];
+			const bool bFiltered = (DbBlendSpace.BlendSpace == nullptr || GetAssetFilterString().IsEmpty()) ? false : !DbBlendSpace.BlendSpace->GetName().Contains(GetAssetFilterString());
+
+			if (!bFiltered)
+			{
+				CreateAssetNode(MappedId, ESearchIndexAssetType::BlendSpace, DefaultGroupIdx);
+			}
 		}
 
 		TreeView->RequestTreeRefresh();
@@ -719,11 +795,57 @@ namespace UE::PoseSearch
 				FUIAction(FExecuteAction::CreateSP(this, &SDatabaseAssetTree::OnDeleteNodes)),
 				NAME_None,
 				EUserInterfaceActionType::Button);
+
+			MenuBuilder.AddMenuEntry(
+				LOCTEXT("Enable", "Enable"),
+				LOCTEXT(
+					"EnableTooltip",
+					"Sets Assets Enabled."),
+				FSlateIcon(),
+				FUIAction(FExecuteAction::CreateSP(this, &SDatabaseAssetTree::OnEnableNodes)),
+				NAME_None,
+				EUserInterfaceActionType::Button);
+
+			MenuBuilder.AddMenuEntry(
+				LOCTEXT("Disable", "Disable"),
+				LOCTEXT(
+					"DisableToolTip",
+					"Sets Assets Disabled."),
+				FSlateIcon(),
+				FUIAction(FExecuteAction::CreateSP(this, &SDatabaseAssetTree::OnDisableNodes)),
+				NAME_None,
+				EUserInterfaceActionType::Button);
 		}
 
 		return MenuBuilder.MakeWidget();
 	}
 
+	TSharedRef<SWidget> SDatabaseAssetTree::GenerateFilterBoxWidget()
+	{
+		TSharedPtr<SSearchBox> SearchBox;
+		SAssignNew(SearchBox, SSearchBox)
+			.MinDesiredWidth(300.0f)
+			.InitialText(this, &SDatabaseAssetTree::GetFilterText)
+			.ToolTipText(FText::FromString(TEXT("Enter Asset Filter...")))
+			.OnTextChanged(this, &SDatabaseAssetTree::OnAssetFilterTextCommitted, ETextCommit::Default)
+			.OnTextCommitted(this, &SDatabaseAssetTree::OnAssetFilterTextCommitted);
+
+		return SearchBox.ToSharedRef();
+	}
+
+
+	FText SDatabaseAssetTree::GetFilterText() const
+	{
+		return FText::FromString(GetAssetFilterString());
+	}
+
+	void SDatabaseAssetTree::OnAssetFilterTextCommitted(const FText& InText, ETextCommit::Type CommitInfo)
+	{
+		SetAssetFilterString(InText.ToString());
+		RefreshTreeView(false);
+	}
+
+
 	void SDatabaseAssetTree::OnAddSequence(bool bFinalizeChanges)
 	{
 		FScopedTransaction Transaction(LOCTEXT("AddSequence", "Add Sequence"));
@@ -858,6 +980,48 @@ namespace UE::PoseSearch
 					OnDeleteAsset(SelectedNode, false);
 				}
 			}
+			
+			FinalizeTreeChanges();
+		}
+	}
+
+	void SDatabaseAssetTree::OnEnableNodes()
+	{
+		TArray<TSharedPtr<FDatabaseAssetTreeNode>> SelectedNodes = TreeView->GetSelectedItems();
+		if (!SelectedNodes.IsEmpty())
+		{
+			for (TSharedPtr<FDatabaseAssetTreeNode> SelectedNode : SelectedNodes)
+			{
+				if (SelectedNode->SourceAssetType == ESearchIndexAssetType::Sequence)
+				{
+					EditorViewModel.Pin()->SetSelectedSequenceEnabled(SelectedNode->SourceAssetIdx, true);
+				}
+				else if (SelectedNode->SourceAssetType == ESearchIndexAssetType::BlendSpace)
+				{
+					EditorViewModel.Pin()->SetSelectedBlendSpaceEnabled(SelectedNode->SourceAssetIdx, true);
+				}
+			}
+		
+			FinalizeTreeChanges();
+		}
+	}
+
+	void SDatabaseAssetTree::OnDisableNodes()
+	{
+		TArray<TSharedPtr<FDatabaseAssetTreeNode>> SelectedNodes = TreeView->GetSelectedItems();
+		if (!SelectedNodes.IsEmpty())
+		{
+			for (TSharedPtr<FDatabaseAssetTreeNode> SelectedNode : SelectedNodes)
+			{
+				if (SelectedNode->SourceAssetType == ESearchIndexAssetType::Sequence)
+				{
+					EditorViewModel.Pin()->SetSelectedSequenceEnabled(SelectedNode->SourceAssetIdx, false);
+				}
+				else if (SelectedNode->SourceAssetType == ESearchIndexAssetType::BlendSpace)
+				{
+					EditorViewModel.Pin()->SetSelectedBlendSpaceEnabled(SelectedNode->SourceAssetIdx, false);
+				}
+			}
 
 			FinalizeTreeChanges();
 		}
diff --git a/Engine/Plugins/Experimental/Animation/PoseSearch/Source/Editor/Public/SPoseSearchDatabaseAssetList.h b/Engine/Plugins/Experimental/Animation/PoseSearch/Source/Editor/Public/SPoseSearchDatabaseAssetList.h
--- a/Engine/Plugins/Experimental/Animation/PoseSearch/Source/Editor/Public/SPoseSearchDatabaseAssetList.h
+++ b/Engine/Plugins/Experimental/Animation/PoseSearch/Source/Editor/Public/SPoseSearchDatabaseAssetList.h
@@ -121,7 +121,18 @@ namespace UE::PoseSearch
 
 		TSharedRef<SWidget> CreateAddNewMenuWidget();
 		TSharedPtr<SWidget> CreateContextMenu();
-		
+
+		TSharedRef<SWidget> GenerateFilterBoxWidget();
+
+		FText GetFilterText() const;
+		void OnAssetFilterTextCommitted(const FText& InText, ETextCommit::Type CommitInfo);
+
+		void SetAssetFilterString(FString InString) { AssetFilterString = InString; }
+		FString GetAssetFilterString() const { return AssetFilterString; }
+
+		//String Filter For Database View
+		FString AssetFilterString;
+
 		void OnAddSequence(bool bFinalizeChanges = true);
 		void OnAddBlendSpace(bool bFinalizeChanges = true);
 
@@ -132,6 +143,9 @@ namespace UE::PoseSearch
 		bool CanDeleteNodes() const;
 		void OnDeleteNodes();
 
+		void OnEnableNodes();
+		void OnDisableNodes();
+
 		friend SDatabaseAssetListItem;
 
 	protected:
diff --git a/Engine/Plugins/Experimental/Animation/PoseSearch/Source/Runtime/Private/PoseSearch.cpp b/Engine/Plugins/Experimental/Animation/PoseSearch/Source/Runtime/Private/PoseSearch.cpp
--- a/Engine/Plugins/Experimental/Animation/PoseSearch/Source/Runtime/Private/PoseSearch.cpp
+++ b/Engine/Plugins/Experimental/Animation/PoseSearch/Source/Runtime/Private/PoseSearch.cpp
@@ -4320,37 +4320,40 @@ static Eigen::VectorXd ComputeChannelsDeviations(const FPoseSearchIndex* SearchI
 
 	using namespace Eigen;
 
-	check(SearchIndex->IsValid() && !SearchIndex->IsEmpty());
+	check(SearchIndex->IsValid());
 
 	const int32 NumPoses = SearchIndex->NumPoses;
 	const int32 NumDimensions = SearchIndex->Schema->SchemaCardinality;
 
-	// Map input buffer
-	auto PoseMatrixSourceMap = RowMajorMatrixMapConst(
-		SearchIndex->Values.GetData(),
-		NumPoses,		// rows
-		NumDimensions	// cols
-	);
-
-	// @todo: evaluate removing the cast to double
-	
-	// Copy row major float matrix to column major double matrix
-	MatrixXd PoseMatrix = PoseMatrixSourceMap.transpose().cast<double>();
-	checkSlow(PoseMatrix.rows() == NumDimensions);
-	checkSlow(PoseMatrix.cols() == NumPoses);
-
-	// Mean center
-	VectorXd SampleMean = PoseMatrix.rowwise().mean();
-	PoseMatrix = PoseMatrix.colwise() - SampleMean;
-
 	// Compute per Channel average distances
 	VectorXd MeanDeviations(NumDimensions);
 	MeanDeviations.setConstant(1.0);
 
-	for (int ChannelIdx = 0; ChannelIdx != SearchIndex->Schema->Channels.Num(); ++ChannelIdx)
+	if (NumPoses > 0)
 	{
-		const UPoseSearchFeatureChannel* Channel = SearchIndex->Schema->Channels[ChannelIdx].Get();
-		Channel->ComputeMeanDeviations(PoseMatrix, MeanDeviations);
+		// Map input buffer
+		auto PoseMatrixSourceMap = RowMajorMatrixMapConst(
+			SearchIndex->Values.GetData(),
+			NumPoses,		// rows
+			NumDimensions	// cols
+		);
+
+		// @todo: evaluate removing the cast to double
+
+		// Copy row major float matrix to column major double matrix
+		MatrixXd PoseMatrix = PoseMatrixSourceMap.transpose().cast<double>();
+		checkSlow(PoseMatrix.rows() == NumDimensions);
+		checkSlow(PoseMatrix.cols() == NumPoses);
+
+		// Mean center
+		VectorXd SampleMean = PoseMatrix.rowwise().mean();
+		PoseMatrix = PoseMatrix.colwise() - SampleMean;
+
+		for (int ChannelIdx = 0; ChannelIdx != SearchIndex->Schema->Channels.Num(); ++ChannelIdx)
+		{
+			const UPoseSearchFeatureChannel* Channel = SearchIndex->Schema->Channels[ChannelIdx].Get();
+			Channel->ComputeMeanDeviations(PoseMatrix, MeanDeviations);
+		}
 	}
 
 	return MeanDeviations;
