pose search debugger QOL improvements

diff --git a/Engine/Plugins/Experimental/Animation/PoseSearch/Source/Editor/Private/PoseSearchDebugger.cpp b/Engine/Plugins/Experimental/Animation/PoseSearch/Source/Editor/Private/PoseSearchDebugger.cpp
--- a/Engine/Plugins/Experimental/Animation/PoseSearch/Source/Editor/Private/PoseSearchDebugger.cpp
+++ b/Engine/Plugins/Experimental/Animation/PoseSearch/Source/Editor/Private/PoseSearchDebugger.cpp
@@ -2252,12 +2252,15 @@ void FDebuggerViewModel::OnDraw(FSkeletonDrawParams& DrawParams)
 	// Returns if it is to be drawn this frame
 	auto SetDrawSkeleton = [this](UPoseSearchMeshComponent* InComponent, bool bDraw)
 	{
-		const bool bIsDrawingSkeleton = InComponent->ShouldDrawDebugSkeleton();
-		if (bIsDrawingSkeleton != bDraw)
+		if (InComponent && InComponent->RequiredBones.IsValid())
 		{
-			InComponent->SetDrawDebugSkeleton(bDraw);
+			const bool bIsDrawingSkeleton = InComponent->ShouldDrawDebugSkeleton();
+			if (bIsDrawingSkeleton != bDraw)
+			{
+				InComponent->SetDrawDebugSkeleton(bDraw);
+			}
+			InComponent->MarkRenderStateDirty();
 		}
-		InComponent->MarkRenderStateDirty();
 	};
 	const bool bDrawActivePose = EnumHasAnyFlags(DrawParams.Flags, ESkeletonDrawFlags::ActivePose);
 	SetDrawSkeleton(Skeletons[ActivePose].Component, bDrawActivePose);
@@ -2276,42 +2279,44 @@ void FDebuggerViewModel::OnDraw(FSkeletonDrawParams& DrawParams)
 	if (bDrawSelectedPose)
 	{
 		UPoseSearchMeshComponent* Component = Skeletons[SelectedPose].Component;
-
-		if (Skeletons[SelectedPose].Type == ESearchIndexAssetType::Sequence)
+		if (Component && Component->RequiredBones.IsValid())
 		{
-			const FPoseSearchDatabaseSequence* DatabaseSequence = Skeletons[SelectedPose].GetAnimSequence();
-			if (DatabaseSequence)
+			if (Skeletons[SelectedPose].Type == ESearchIndexAssetType::Sequence)
 			{
-				UpdateContext.Type = ESearchIndexAssetType::Sequence;
-				UpdateContext.Sequence = DatabaseSequence->Sequence;
-				UpdateContext.StartTime = Skeletons[SelectedPose].Time;
-				UpdateContext.Time = Skeletons[SelectedPose].Time;
-				UpdateContext.bMirrored = Skeletons[SelectedPose].bMirrored;
-				UpdateContext.bLoop = DatabaseSequence->Sequence->bLoop;
+				const FPoseSearchDatabaseSequence* DatabaseSequence = Skeletons[SelectedPose].GetAnimSequence();
+				if (DatabaseSequence)
+				{
+					UpdateContext.Type = ESearchIndexAssetType::Sequence;
+					UpdateContext.Sequence = DatabaseSequence->Sequence;
+					UpdateContext.StartTime = Skeletons[SelectedPose].Time;
+					UpdateContext.Time = Skeletons[SelectedPose].Time;
+					UpdateContext.bMirrored = Skeletons[SelectedPose].bMirrored;
+					UpdateContext.bLoop = DatabaseSequence->Sequence->bLoop;
+				}
 			}
-		}
-		else if (Skeletons[SelectedPose].Type == ESearchIndexAssetType::BlendSpace)
-		{
-			const FPoseSearchDatabaseBlendSpace* DatabaseBlendSpace = Skeletons[SelectedPose].GetBlendSpace();
-			if (DatabaseBlendSpace)
+			else if (Skeletons[SelectedPose].Type == ESearchIndexAssetType::BlendSpace)
 			{
-				UpdateContext.Type = ESearchIndexAssetType::BlendSpace;
-				UpdateContext.BlendSpace = DatabaseBlendSpace->BlendSpace;
-				UpdateContext.StartTime = Skeletons[SelectedPose].Time;
-				UpdateContext.Time = Skeletons[SelectedPose].Time;
-				UpdateContext.bMirrored = Skeletons[SelectedPose].bMirrored;
-				UpdateContext.bLoop = DatabaseBlendSpace->BlendSpace->bLoop;
-				UpdateContext.BlendParameters = Skeletons[SelectedPose].BlendParameters;
+				const FPoseSearchDatabaseBlendSpace* DatabaseBlendSpace = Skeletons[SelectedPose].GetBlendSpace();
+				if (DatabaseBlendSpace)
+				{
+					UpdateContext.Type = ESearchIndexAssetType::BlendSpace;
+					UpdateContext.BlendSpace = DatabaseBlendSpace->BlendSpace;
+					UpdateContext.StartTime = Skeletons[SelectedPose].Time;
+					UpdateContext.Time = Skeletons[SelectedPose].Time;
+					UpdateContext.bMirrored = Skeletons[SelectedPose].bMirrored;
+					UpdateContext.bLoop = DatabaseBlendSpace->BlendSpace->bLoop;
+					UpdateContext.BlendParameters = Skeletons[SelectedPose].BlendParameters;
+				}
+			}
+			else
+			{
+				checkNoEntry();
 			}
-		}
-		else
-		{
-			checkNoEntry();
-		}
 
-		if (UpdateContext.Type != ESearchIndexAssetType::Invalid)
-		{
-			Component->UpdatePose(UpdateContext);
+			if (UpdateContext.Type != ESearchIndexAssetType::Invalid)
+			{
+				Component->UpdatePose(UpdateContext);
+			}
 		}
 	}
 
@@ -2319,43 +2324,46 @@ void FDebuggerViewModel::OnDraw(FSkeletonDrawParams& DrawParams)
 	if (bDrawAsset && AssetData.bActive)
 	{
 		UPoseSearchMeshComponent* Component = Skeletons[Asset].Component;
-		SetDrawSkeleton(Component, true);
-
-		if (Skeletons[SelectedPose].Type == ESearchIndexAssetType::Sequence)
+		if (Component && Component->RequiredBones.IsValid())
 		{
-			const FPoseSearchDatabaseSequence* DatabaseSequence = Skeletons[Asset].GetAnimSequence();
-			if (DatabaseSequence)
+			SetDrawSkeleton(Component, true);
+
+			if (Skeletons[SelectedPose].Type == ESearchIndexAssetType::Sequence)
 			{
-				UpdateContext.Type = ESearchIndexAssetType::Sequence;
-				UpdateContext.Sequence = DatabaseSequence->Sequence;
-				UpdateContext.StartTime = Skeletons[Asset].Time;
-				UpdateContext.Time = Skeletons[Asset].Time;
-				UpdateContext.bMirrored = Skeletons[Asset].bMirrored;
-				UpdateContext.bLoop = DatabaseSequence->Sequence->bLoop;
+				const FPoseSearchDatabaseSequence* DatabaseSequence = Skeletons[Asset].GetAnimSequence();
+				if (DatabaseSequence)
+				{
+					UpdateContext.Type = ESearchIndexAssetType::Sequence;
+					UpdateContext.Sequence = DatabaseSequence->Sequence;
+					UpdateContext.StartTime = Skeletons[Asset].Time;
+					UpdateContext.Time = Skeletons[Asset].Time;
+					UpdateContext.bMirrored = Skeletons[Asset].bMirrored;
+					UpdateContext.bLoop = DatabaseSequence->Sequence->bLoop;
+				}
 			}
-		}
-		else if (Skeletons[SelectedPose].Type == ESearchIndexAssetType::BlendSpace)
-		{
-			const FPoseSearchDatabaseBlendSpace* DatabaseBlendSpace = Skeletons[Asset].GetBlendSpace();
-			if (DatabaseBlendSpace)
+			else if (Skeletons[SelectedPose].Type == ESearchIndexAssetType::BlendSpace)
 			{
-				UpdateContext.Type = ESearchIndexAssetType::BlendSpace;
-				UpdateContext.BlendSpace = DatabaseBlendSpace->BlendSpace;
-				UpdateContext.StartTime = Skeletons[Asset].Time;
-				UpdateContext.Time = Skeletons[Asset].Time;
-				UpdateContext.bMirrored = Skeletons[Asset].bMirrored;
-				UpdateContext.bLoop = DatabaseBlendSpace->BlendSpace->bLoop;
-				UpdateContext.BlendParameters = Skeletons[Asset].BlendParameters;
+				const FPoseSearchDatabaseBlendSpace* DatabaseBlendSpace = Skeletons[Asset].GetBlendSpace();
+				if (DatabaseBlendSpace)
+				{
+					UpdateContext.Type = ESearchIndexAssetType::BlendSpace;
+					UpdateContext.BlendSpace = DatabaseBlendSpace->BlendSpace;
+					UpdateContext.StartTime = Skeletons[Asset].Time;
+					UpdateContext.Time = Skeletons[Asset].Time;
+					UpdateContext.bMirrored = Skeletons[Asset].bMirrored;
+					UpdateContext.bLoop = DatabaseBlendSpace->BlendSpace->bLoop;
+					UpdateContext.BlendParameters = Skeletons[Asset].BlendParameters;
+				}
+			}
+			else
+			{
+				checkNoEntry();
 			}
-		}
-		else
-		{
-			checkNoEntry();
-		}
 
-		if (UpdateContext.Type != ESearchIndexAssetType::Invalid)
-		{
-			Component->UpdatePose(UpdateContext);
+			if (UpdateContext.Type != ESearchIndexAssetType::Invalid)
+			{
+				Component->UpdatePose(UpdateContext);
+			}
 		}
 	}
 }
