pose search debugger - defaulting debug lines to minimum verbosity

diff --git a/Engine/Plugins/Experimental/Animation/PoseSearch/Source/Editor/Private/PoseSearchDebugger.h b/Engine/Plugins/Experimental/Animation/PoseSearch/Source/Editor/Private/PoseSearchDebugger.h
--- a/Engine/Plugins/Experimental/Animation/PoseSearch/Source/Editor/Private/PoseSearchDebugger.h
+++ b/Engine/Plugins/Experimental/Animation/PoseSearch/Source/Editor/Private/PoseSearchDebugger.h
@@ -81,13 +81,13 @@ struct FPoseSearchDebuggerFeatureDrawOptions
     bool bDisable = false;
 
 	UPROPERTY(EditAnywhere, Category = "Draw Options", Meta = (EditCondition = "!bDisable"))
-	bool bDrawBoneNames = true;
+	bool bDrawBoneNames = false;
 
 	UPROPERTY(EditAnywhere, Category = "Draw Options", Meta = (EditCondition = "!bDisable"))
-	bool bDrawSampleLabels = true;
+	bool bDrawSampleLabels = false;
 
 	UPROPERTY(EditAnywhere, Category = "Draw Options", Meta = (EditCondition = "!bDisable"))
-	bool bDrawSamplesWithColorGradient = true;
+	bool bDrawSamplesWithColorGradient = false;
 };
 
 /**
@@ -139,10 +139,10 @@ public:
 	FPoseSearchDebuggerFeatureDrawOptions SelectedPoseDrawOptions;
 
 	UPROPERTY(EditAnywhere, Category = "Draw Options")
-	bool bDrawActiveSkeleton = true;
+	bool bDrawActiveSkeleton = false;
 	
 	UPROPERTY(EditAnywhere, Category = "Draw Options")
-	bool bDrawSelectedSkeleton = true;
+	bool bDrawSelectedSkeleton = false;
 
     UPROPERTY(VisibleAnywhere, Category="Pose Vectors")
 	TArray<float> QueryPoseVector;
