motion matching - query creation with raw or interpolated continuing pose

diff --git a/Engine/Plugins/Experimental/Animation/PoseSearch/Source/Runtime/Private/PoseSearch.cpp b/Engine/Plugins/Experimental/Animation/PoseSearch/Source/Runtime/Private/PoseSearch.cpp
--- a/Engine/Plugins/Experimental/Animation/PoseSearch/Source/Runtime/Private/PoseSearch.cpp
+++ b/Engine/Plugins/Experimental/Animation/PoseSearch/Source/Runtime/Private/PoseSearch.cpp
@@ -2337,27 +2337,33 @@ bool FSearchContext::GetOrBuildQuery(const UPoseSearchDatabase* Database, FPoseS
 	return false;
 }
 
-void FSearchContext::CacheCurrentResultFeatureVectors()
+bool FSearchContext::IsCurrentResultFromDatabase(const UPoseSearchDatabase* Database) const
 {
-	if (!CurrentResultPoseVector.IsInitialized())
-	{
-		check(CurrentResult.IsValid());
-		const FPoseSearchIndex* SearchIndex = CurrentResult.Database->GetSearchIndex();
-		check(SearchIndex);
+	return !bForceInterrupt && CurrentResult.IsValid() && CurrentResult.Database == Database;
+}
 
-		CurrentResultPoseVector.Init(CurrentResult.Database->Schema);
-		CurrentResultPrevPoseVector.Init(CurrentResult.Database->Schema);
-		CurrentResultNextPoseVector.Init(CurrentResult.Database->Schema);
+TConstArrayView<const float> FSearchContext::GetCurrentResultPrevPoseVector() const
+{
+	check(CurrentResult.IsValid());
+	const FPoseSearchIndex* SearchIndex = CurrentResult.Database->GetSearchIndex();
+	check(SearchIndex);
+	return SearchIndex->GetPoseValues(CurrentResult.PrevPoseIdx);
+}
 
-		CurrentResultPoseVector.CopyFromSearchIndex(*SearchIndex, CurrentResult.PoseIdx);
-		CurrentResultPrevPoseVector.CopyFromSearchIndex(*SearchIndex, CurrentResult.PrevPoseIdx);
-		CurrentResultNextPoseVector.CopyFromSearchIndex(*SearchIndex, CurrentResult.NextPoseIdx);
-	}
+TConstArrayView<const float> FSearchContext::GetCurrentResultPoseVector() const
+{
+	check(CurrentResult.IsValid());
+	const FPoseSearchIndex* SearchIndex = CurrentResult.Database->GetSearchIndex();
+	check(SearchIndex);
+	return SearchIndex->GetPoseValues(CurrentResult.PoseIdx);
 }
 
-bool FSearchContext::IsCurrentResultFromDatabase(const UPoseSearchDatabase* Database) const
+TConstArrayView<const float> FSearchContext::GetCurrentResultNextPoseVector() const
 {
-	return !bForceInterrupt && CurrentResult.IsValid() && CurrentResult.Database == Database;
+	check(CurrentResult.IsValid());
+	const FPoseSearchIndex* SearchIndex = CurrentResult.Database->GetSearchIndex();
+	check(SearchIndex);
+	return SearchIndex->GetPoseValues(CurrentResult.NextPoseIdx);
 }
 
 //////////////////////////////////////////////////////////////////////////
@@ -2562,9 +2568,9 @@ void FFeatureVectorHelper::EncodeQuat(TArrayView<float> Values, int32& DataOffse
 	DataOffset += EncodeQuatCardinality;
 }
 
-void FFeatureVectorHelper::EncodeQuat(TArrayView<float> Values, int32& DataOffset, TArrayView<const float> PrevValues, TArrayView<const float> CurPrevValues, TArrayView<const float> NextPrevValues, float LerpValue)
+void FFeatureVectorHelper::EncodeQuat(TArrayView<float> Values, int32& DataOffset, TArrayView<const float> PrevValues, TArrayView<const float> CurValues, TArrayView<const float> NextValues, float LerpValue)
 {
-	FQuat Quat = DecodeQuatInternal(CurPrevValues, DataOffset);
+	FQuat Quat = DecodeQuatInternal(CurValues, DataOffset);
 
 	// linear interpolation
 	if (!FMath::IsNearlyZero(LerpValue))
@@ -2575,7 +2581,7 @@ void FFeatureVectorHelper::EncodeQuat(TArrayView<float> Values, int32& DataOffse
 		}
 		else
 		{
-			Quat = FQuat::Slerp(Quat, DecodeQuatInternal(NextPrevValues, DataOffset), LerpValue);
+			Quat = FQuat::Slerp(Quat, DecodeQuatInternal(NextValues, DataOffset), LerpValue);
 		}
 	}
 	// @todo: do we need to add options for cubic interpolation?
@@ -2612,9 +2618,9 @@ void FFeatureVectorHelper::EncodeVector(TArrayView<float> Values, int32& DataOff
 	DataOffset += EncodeVectorCardinality;
 }
 
-void FFeatureVectorHelper::EncodeVector(TArrayView<float> Values, int32& DataOffset, TArrayView<const float> PrevValues, TArrayView<const float> CurPrevValues, TArrayView<const float> NextPrevValues, float LerpValue, bool bNormalize)
+void FFeatureVectorHelper::EncodeVector(TArrayView<float> Values, int32& DataOffset, TArrayView<const float> PrevValues, TArrayView<const float> CurValues, TArrayView<const float> NextValues, float LerpValue, bool bNormalize)
 {
-	FVector Vector = DecodeVectorInternal(CurPrevValues, DataOffset);
+	FVector Vector = DecodeVectorInternal(CurValues, DataOffset);
 
 	// linear interpolation
 	if (!FMath::IsNearlyZero(LerpValue))
@@ -2625,7 +2631,7 @@ void FFeatureVectorHelper::EncodeVector(TArrayView<float> Values, int32& DataOff
 		}
 		else
 		{
-			Vector = FMath::Lerp(Vector, DecodeVectorInternal(NextPrevValues, DataOffset), LerpValue);
+			Vector = FMath::Lerp(Vector, DecodeVectorInternal(NextValues, DataOffset), LerpValue);
 		}
 	}
 	// @todo: do we need to add options for cubic interpolation?
@@ -2657,9 +2663,9 @@ void FFeatureVectorHelper::EncodeVector2D(TArrayView<float> Values, int32& DataO
 	DataOffset += EncodeVector2DCardinality;
 }
 
-void FFeatureVectorHelper::EncodeVector2D(TArrayView<float> Values, int32& DataOffset, TArrayView<const float> PrevValues, TArrayView<const float> CurPrevValues, TArrayView<const float> NextPrevValues, float LerpValue)
+void FFeatureVectorHelper::EncodeVector2D(TArrayView<float> Values, int32& DataOffset, TArrayView<const float> PrevValues, TArrayView<const float> CurValues, TArrayView<const float> NextValues, float LerpValue)
 {
-	FVector2D Vector2D = DecodeVector2DInternal(CurPrevValues, DataOffset);
+	FVector2D Vector2D = DecodeVector2DInternal(CurValues, DataOffset);
 
 	// linear interpolation
 	if (!FMath::IsNearlyZero(LerpValue))
@@ -2670,7 +2676,7 @@ void FFeatureVectorHelper::EncodeVector2D(TArrayView<float> Values, int32& DataO
 		}
 		else
 		{
-			Vector2D = FMath::Lerp(Vector2D, DecodeVector2DInternal(NextPrevValues, DataOffset), LerpValue);
+			Vector2D = FMath::Lerp(Vector2D, DecodeVector2DInternal(NextValues, DataOffset), LerpValue);
 		}
 	}
 	// @todo: do we need to add options for cubic interpolation?
diff --git a/Engine/Plugins/Experimental/Animation/PoseSearch/Source/Runtime/Private/PoseSearchFeatureChannels.cpp b/Engine/Plugins/Experimental/Animation/PoseSearch/Source/Runtime/Private/PoseSearchFeatureChannels.cpp
--- a/Engine/Plugins/Experimental/Animation/PoseSearch/Source/Runtime/Private/PoseSearchFeatureChannels.cpp
+++ b/Engine/Plugins/Experimental/Animation/PoseSearch/Source/Runtime/Private/PoseSearchFeatureChannels.cpp
@@ -357,14 +357,14 @@ bool UPoseSearchFeatureChannel_Position::BuildQuery(UE::PoseSearch::FSearchConte
 	using namespace UE::PoseSearch;
 
 	const bool bIsCurrentResultValid = SearchContext.CurrentResult.IsValid();
-	const bool bSkip = bUseFeaturesFromContinuingPose && bIsCurrentResultValid && SearchContext.CurrentResult.Database->Schema == InOutQuery.GetSchema();
+	const bool bSkip = InputQueryPose != InputQueryPose::UseCharacterPose && bIsCurrentResultValid && SearchContext.CurrentResult.Database->Schema == InOutQuery.GetSchema();
 	if (bSkip || !SearchContext.History)
 	{
 		if (bIsCurrentResultValid)
 		{
-			SearchContext.CacheCurrentResultFeatureVectors();
+			const float LerpValue = InputQueryPose == InputQueryPose::UseInterpolatedContinuingPose ? SearchContext.CurrentResult.LerpValue : 0.f;
 			int32 DataOffset = ChannelDataOffset;
-			FFeatureVectorHelper::EncodeVector(InOutQuery.EditValues(), DataOffset, SearchContext.CurrentResultPrevPoseVector.GetValues(), SearchContext.CurrentResultPoseVector.GetValues(), SearchContext.CurrentResultNextPoseVector.GetValues(), SearchContext.CurrentResult.LerpValue);
+			FFeatureVectorHelper::EncodeVector(InOutQuery.EditValues(), DataOffset, SearchContext.GetCurrentResultPrevPoseVector(), SearchContext.GetCurrentResultPoseVector(), SearchContext.GetCurrentResultNextPoseVector(), LerpValue);
 		}
 		return bSkip;
 	}
@@ -481,14 +481,14 @@ bool UPoseSearchFeatureChannel_Heading::BuildQuery(UE::PoseSearch::FSearchContex
 	using namespace UE::PoseSearch;
 
 	const bool bIsCurrentResultValid = SearchContext.CurrentResult.IsValid();
-	const bool bSkip = bUseFeaturesFromContinuingPose && bIsCurrentResultValid && SearchContext.CurrentResult.Database->Schema == InOutQuery.GetSchema();
+	const bool bSkip = InputQueryPose != InputQueryPose::UseCharacterPose && bIsCurrentResultValid && SearchContext.CurrentResult.Database->Schema == InOutQuery.GetSchema();
 	if (bSkip || !SearchContext.History)
 	{
 		if (bIsCurrentResultValid)
 		{
-			SearchContext.CacheCurrentResultFeatureVectors();
+			const float LerpValue = InputQueryPose == InputQueryPose::UseInterpolatedContinuingPose ? SearchContext.CurrentResult.LerpValue : 0.f;
 			int32 DataOffset = ChannelDataOffset;
-			FFeatureVectorHelper::EncodeVector(InOutQuery.EditValues(), DataOffset, SearchContext.CurrentResultPrevPoseVector.GetValues(), SearchContext.CurrentResultPoseVector.GetValues(), SearchContext.CurrentResultNextPoseVector.GetValues(), SearchContext.CurrentResult.LerpValue, true);
+			FFeatureVectorHelper::EncodeVector(InOutQuery.EditValues(), DataOffset, SearchContext.GetCurrentResultPrevPoseVector(), SearchContext.GetCurrentResultPoseVector(), SearchContext.GetCurrentResultNextPoseVector(), LerpValue, true);
 			check(DataOffset == ChannelDataOffset + ChannelCardinality);
 		}
 		return bSkip;
@@ -872,13 +872,12 @@ bool UPoseSearchFeatureChannel_Pose::BuildQuery(UE::PoseSearch::FSearchContext&
 	using namespace UE::PoseSearch;
 
 	const bool bIsCurrentResultValid = SearchContext.CurrentResult.IsValid();
-	const bool bSkip = bUseFeaturesFromContinuingPose && bIsCurrentResultValid && SearchContext.CurrentResult.Database->Schema == InOutQuery.GetSchema();
+	const bool bSkip = InputQueryPose != InputQueryPose::UseCharacterPose && bIsCurrentResultValid && SearchContext.CurrentResult.Database->Schema == InOutQuery.GetSchema();
 	if (bSkip || !SearchContext.History)
 	{
 		if (bIsCurrentResultValid)
 		{
-			SearchContext.CacheCurrentResultFeatureVectors();
-
+			const float LerpValue = InputQueryPose == InputQueryPose::UseInterpolatedContinuingPose ? SearchContext.CurrentResult.LerpValue : 0.f;
 			int32 DataOffset = ChannelDataOffset;
 			for (int32 SampledBoneIdx = 0; SampledBoneIdx != SampledBones.Num(); ++SampledBoneIdx)
 			{
@@ -887,7 +886,7 @@ bool UPoseSearchFeatureChannel_Pose::BuildQuery(UE::PoseSearch::FSearchContext&
 				{
 					for (int32 SubsampleIdx = 0; SubsampleIdx != SampleTimes.Num(); ++SubsampleIdx)
 					{
-						FFeatureVectorHelper::EncodeVector(InOutQuery.EditValues(), DataOffset, SearchContext.CurrentResultPrevPoseVector.GetValues(), SearchContext.CurrentResultPoseVector.GetValues(), SearchContext.CurrentResultNextPoseVector.GetValues(), SearchContext.CurrentResult.LerpValue);
+						FFeatureVectorHelper::EncodeVector(InOutQuery.EditValues(), DataOffset, SearchContext.GetCurrentResultPrevPoseVector(), SearchContext.GetCurrentResultPoseVector(), SearchContext.GetCurrentResultNextPoseVector(), LerpValue);
 					}
 				}
 
@@ -895,7 +894,7 @@ bool UPoseSearchFeatureChannel_Pose::BuildQuery(UE::PoseSearch::FSearchContext&
 				{
 					for (int32 SubsampleIdx = 0; SubsampleIdx != SampleTimes.Num(); ++SubsampleIdx)
 					{
-						FFeatureVectorHelper::EncodeQuat(InOutQuery.EditValues(), DataOffset, SearchContext.CurrentResultPrevPoseVector.GetValues(), SearchContext.CurrentResultPoseVector.GetValues(), SearchContext.CurrentResultNextPoseVector.GetValues(), SearchContext.CurrentResult.LerpValue);
+						FFeatureVectorHelper::EncodeQuat(InOutQuery.EditValues(), DataOffset, SearchContext.GetCurrentResultPrevPoseVector(), SearchContext.GetCurrentResultPoseVector(), SearchContext.GetCurrentResultNextPoseVector(), LerpValue);
 					}
 				}
 
@@ -903,7 +902,7 @@ bool UPoseSearchFeatureChannel_Pose::BuildQuery(UE::PoseSearch::FSearchContext&
 				{
 					for (int32 SubsampleIdx = 0; SubsampleIdx != SampleTimes.Num(); ++SubsampleIdx)
 					{
-						FFeatureVectorHelper::EncodeVector(InOutQuery.EditValues(), DataOffset, SearchContext.CurrentResultPrevPoseVector.GetValues(), SearchContext.CurrentResultPoseVector.GetValues(), SearchContext.CurrentResultNextPoseVector.GetValues(), SearchContext.CurrentResult.LerpValue);
+						FFeatureVectorHelper::EncodeVector(InOutQuery.EditValues(), DataOffset, SearchContext.GetCurrentResultPrevPoseVector(), SearchContext.GetCurrentResultPoseVector(), SearchContext.GetCurrentResultNextPoseVector(), LerpValue);
 					}
 				}
 
@@ -911,7 +910,7 @@ bool UPoseSearchFeatureChannel_Pose::BuildQuery(UE::PoseSearch::FSearchContext&
 				{
 					for (int32 SubsampleIdx = 0; SubsampleIdx != SampleTimes.Num(); ++SubsampleIdx)
 					{
-						FFeatureVectorHelper::EncodeVector2D(InOutQuery.EditValues(), DataOffset, SearchContext.CurrentResultPrevPoseVector.GetValues(), SearchContext.CurrentResultPoseVector.GetValues(), SearchContext.CurrentResultNextPoseVector.GetValues(), SearchContext.CurrentResult.LerpValue);
+						FFeatureVectorHelper::EncodeVector2D(InOutQuery.EditValues(), DataOffset, SearchContext.GetCurrentResultPrevPoseVector(), SearchContext.GetCurrentResultPoseVector(), SearchContext.GetCurrentResultNextPoseVector(), LerpValue);
 					}
 				}
 			}
diff --git a/Engine/Plugins/Experimental/Animation/PoseSearch/Source/Runtime/Public/PoseSearch/PoseSearch.h b/Engine/Plugins/Experimental/Animation/PoseSearch/Source/Runtime/Public/PoseSearch/PoseSearch.h
--- a/Engine/Plugins/Experimental/Animation/PoseSearch/Source/Runtime/Public/PoseSearch/PoseSearch.h
+++ b/Engine/Plugins/Experimental/Animation/PoseSearch/Source/Runtime/Public/PoseSearch/PoseSearch.h
@@ -1218,17 +1218,17 @@ class POSESEARCH_API FFeatureVectorHelper
 public:
 	enum { EncodeQuatCardinality = 6 };
 	static void EncodeQuat(TArrayView<float> Values, int32& DataOffset, const FQuat& Quat);
-	static void EncodeQuat(TArrayView<float> Values, int32& DataOffset, TArrayView<const float> PrevValues, TArrayView<const float> CurPrevValues, TArrayView<const float> NextPrevValues, float LerpValue);
+	static void EncodeQuat(TArrayView<float> Values, int32& DataOffset, TArrayView<const float> PrevValues, TArrayView<const float> CurValues, TArrayView<const float> NextValues, float LerpValue);
 	static FQuat DecodeQuat(TArrayView<const float> Values, int32& DataOffset);
 
 	enum { EncodeVectorCardinality = 3 };
 	static void EncodeVector(TArrayView<float> Values, int32& DataOffset, const FVector& Vector);
-	static void EncodeVector(TArrayView<float> Values, int32& DataOffset, TArrayView<const float> PrevValues, TArrayView<const float> CurPrevValues, TArrayView<const float> NextPrevValues, float LerpValue, bool bNormalize = false);
+	static void EncodeVector(TArrayView<float> Values, int32& DataOffset, TArrayView<const float> PrevValues, TArrayView<const float> CurValues, TArrayView<const float> NextValues, float LerpValue, bool bNormalize = false);
 	static FVector DecodeVector(TArrayView<const float> Values, int32& DataOffset);
 
 	enum { EncodeVector2DCardinality = 2 };
 	static void EncodeVector2D(TArrayView<float> Values, int32& DataOffset, const FVector2D& Vector2D);
-	static void EncodeVector2D(TArrayView<float> Values, int32& DataOffset, TArrayView<const float> PrevValues, TArrayView<const float> CurPrevValues, TArrayView<const float> NextPrevValues, float LerpValue);
+	static void EncodeVector2D(TArrayView<float> Values, int32& DataOffset, TArrayView<const float> PrevValues, TArrayView<const float> CurValues, TArrayView<const float> NextValues, float LerpValue);
 	static FVector2D DecodeVector2D(TArrayView<const float> Values, int32& DataOffset);
 
 	// populates MeanDeviations[DataOffset] ... MeanDeviations[DataOffset + Cardinality] with a single value the mean deviation calculated from a centered matrix
@@ -1330,15 +1330,15 @@ struct POSESEARCH_API FSearchContext
 	float GetCurrentBestTotalCost() const { return CurrentBestTotalCost; }
 
 	bool GetOrBuildQuery(const UPoseSearchDatabase* Database, FPoseSearchFeatureVectorBuilder& FeatureVectorBuilder);
-	void CacheCurrentResultFeatureVectors();
 
 	bool IsCurrentResultFromDatabase(const UPoseSearchDatabase* Database) const;
 
+	TConstArrayView<const float> GetCurrentResultPrevPoseVector() const;
+	TConstArrayView<const float> GetCurrentResultPoseVector() const;
+	TConstArrayView<const float> GetCurrentResultNextPoseVector() const;
+
 	static constexpr int8 SchemaRootBoneIdx = -1;
 
-	FPoseSearchFeatureVectorBuilder CurrentResultPoseVector;
-	FPoseSearchFeatureVectorBuilder CurrentResultPrevPoseVector;
-	FPoseSearchFeatureVectorBuilder CurrentResultNextPoseVector;
 	const FPoseIndicesHistory* PoseIndicesHistory = nullptr;
 
 private:
diff --git a/Engine/Plugins/Experimental/Animation/PoseSearch/Source/Runtime/Public/PoseSearch/PoseSearchFeatureChannels.h b/Engine/Plugins/Experimental/Animation/PoseSearch/Source/Runtime/Public/PoseSearch/PoseSearchFeatureChannels.h
--- a/Engine/Plugins/Experimental/Animation/PoseSearch/Source/Runtime/Public/PoseSearch/PoseSearchFeatureChannels.h
+++ b/Engine/Plugins/Experimental/Animation/PoseSearch/Source/Runtime/Public/PoseSearch/PoseSearchFeatureChannels.h
@@ -8,6 +8,22 @@
 
 #include "PoseSearchFeatureChannels.generated.h"
 
+UENUM(BlueprintType)
+enum class InputQueryPose : uint8
+{
+	// use character pose to compose the query
+	UseCharacterPose,
+
+	// if available reuse continuing pose from the database to compose the query or else UseCharacterPose
+	UseContinuingPose,
+
+	// if available reuse and interpolate continuing pose from the database to compose the query or else UseCharacterPose
+	UseInterpolatedContinuingPose,
+
+	Num UMETA(Hidden),
+	Invalid = Num UMETA(Hidden)
+};
+
 //////////////////////////////////////////////////////////////////////////
 // UPoseSearchFeatureChannel_Position
 UCLASS(BlueprintType, EditInlineNew)
@@ -31,9 +47,8 @@ public:
 	UPROPERTY(EditAnywhere, Category = "Settings")
 	int32 ColorPresetIndex = 0;
 
-	// reuse continuing pose database feature values if available
 	UPROPERTY(EditAnywhere, Category = "Settings")
-	bool bUseFeaturesFromContinuingPose = true;
+	InputQueryPose InputQueryPose = InputQueryPose::UseContinuingPose;
 
 	// if UseSampleTimeOffsetRootBone is true, this UPoseSearchFeatureChannel_Position will calculate the position of Bone from the pose SampleTimeOffset seconds away from the current time pose root bone
 	// if false the calculated position of Bone will be in component space from the pose SampleTimeOffset seconds away
@@ -86,9 +101,8 @@ public:
 	UPROPERTY(EditAnywhere, Category = "Settings")
 	int32 ColorPresetIndex = 0;
 
-	// reuse continuing pose database feature values if available
 	UPROPERTY(EditAnywhere, Category = "Settings")
-	bool bUseFeaturesFromContinuingPose = true;
+	InputQueryPose InputQueryPose = InputQueryPose::UseContinuingPose;
 
 	// if UseSampleTimeOffsetRootBone is true, this UPoseSearchFeatureChannel_Position will calculate the position of Bone from the pose SampleTimeOffset seconds away from the current time pose root bone
 	// if false the calculated position of Bone will be in component space from the pose SampleTimeOffset seconds away
@@ -125,9 +139,8 @@ public:
 	UPROPERTY()
 	TArray<int8> SchemaBoneIdx;
 
-	// reuse continuing pose database feature values if available
 	UPROPERTY(EditAnywhere, Category = "Settings")
-	bool bUseFeaturesFromContinuingPose = true;
+	InputQueryPose InputQueryPose = InputQueryPose::UseContinuingPose;
 
 	// UObject interface
 	virtual void PreSave(FObjectPreSaveContext ObjectSaveContext) override;
