motion matching - blend stack fix for null MirrorDataTable

diff --git a/Engine/Plugins/Experimental/Animation/PoseSearch/Source/Runtime/Private/AnimNode_BlendStack.cpp b/Engine/Plugins/Experimental/Animation/PoseSearch/Source/Runtime/Private/AnimNode_BlendStack.cpp
--- a/Engine/Plugins/Experimental/Animation/PoseSearch/Source/Runtime/Private/AnimNode_BlendStack.cpp
+++ b/Engine/Plugins/Experimental/Animation/PoseSearch/Source/Runtime/Private/AnimNode_BlendStack.cpp
@@ -15,7 +15,15 @@ TAutoConsoleVariable<int32> CVarAnimBlendStackEnable(TEXT("a.AnimNode.BlendStack
 void FPoseSearchAnimPlayer::Initialize(ESearchIndexAssetType InAssetType, UAnimationAsset* AnimationAsset, float AccumulatedTime, bool bLoop, bool bMirrored, UMirrorDataTable* MirrorDataTable, float BlendTime, const UBlendProfile* BlendProfile, FVector BlendParameters)
 {
 	check(AnimationAsset);
-	check(MirrorDataTable);
+
+	if (bMirrored && !MirrorDataTable)
+	{
+		UE_LOG(
+			LogPoseSearch,
+			Error,
+			TEXT("FPoseSearchAnimPlayer failed to Initialize for %s. Mirroring will not work becasue MirrorDataTable is missing"),
+			*GetNameSafe(AnimationAsset));
+	}
 
 	if (BlendProfile != nullptr)
 	{
