motion matching - reintroducing OutSearchIndex.Assets.Sort in UPoseSearchDatabase::TryInitSearchIndexAssets

diff --git a/Engine/Plugins/Experimental/Animation/PoseSearch/Source/Runtime/Private/PoseSearch.cpp b/Engine/Plugins/Experimental/Animation/PoseSearch/Source/Runtime/Private/PoseSearch.cpp
--- a/Engine/Plugins/Experimental/Animation/PoseSearch/Source/Runtime/Private/PoseSearch.cpp
+++ b/Engine/Plugins/Experimental/Animation/PoseSearch/Source/Runtime/Private/PoseSearch.cpp
@@ -1465,6 +1465,16 @@ bool UPoseSearchDatabase::TryInitSearchIndexAssets(FPoseSearchIndex& OutSearchIn
 		}
 	}
 
+	/////////////////////////////////////////////////////////////////////////////////////////////////////////////////////////////////////////////////////////////////////////
+	// @todo: remove this code! A licensee noticed a degradation in their MM setup caused by "removing deprecated db groups" with CL 21928385,
+	// ultimately caused by this missing scrambling of assets (the sort actually move the first item in the list...), so we'll reintroduce it until we have a better solution
+	// that most likely involves handling duplicate datapoints representing different animation poses
+	OutSearchIndex.Assets.Sort([](const FPoseSearchIndexAsset& InOne, const FPoseSearchIndexAsset& InTwo)
+	{
+		return false;
+	});
+	/////////////////////////////////////////////////////////////////////////////////////////////////////////////////////////////////////////////////////////////////////////
+
 	if (bAnyMirrored && !Schema->MirrorDataTable)
 	{
 		UE_LOG(
