motion matching - fix for PopulateNonSelectableIdx log spamming

diff --git a/Engine/Plugins/Experimental/Animation/PoseSearch/Source/Runtime/Private/PoseSearch.cpp b/Engine/Plugins/Experimental/Animation/PoseSearch/Source/Runtime/Private/PoseSearch.cpp
--- a/Engine/Plugins/Experimental/Animation/PoseSearch/Source/Runtime/Private/PoseSearch.cpp
+++ b/Engine/Plugins/Experimental/Animation/PoseSearch/Source/Runtime/Private/PoseSearch.cpp
@@ -219,9 +219,17 @@ static int32 PopulateNonSelectableIdx(size_t* NonSelectableIdx, int32 NonSelecta
 				}
 			}
 
-			if (CurrentIndexAsset->SamplingInterval.Contains(PoseAssetPlayerTime) && NonSelectableIdxUsedSize < NonSelectableIdxSize)
+			if (CurrentIndexAsset->SamplingInterval.Contains(PoseAssetPlayerTime))
 			{
-				NonSelectableIdx[NonSelectableIdxUsedSize++] = PoseIdx;
+				if (NonSelectableIdxUsedSize < NonSelectableIdxSize)
+				{
+					NonSelectableIdx[NonSelectableIdxUsedSize++] = PoseIdx;
+				}
+				else
+				{
+					UE_LOG(LogPoseSearch, Warning, TEXT("PopulateNonSelectableIdx couldn't add all the NonSelectableIdx"));
+					return NonSelectableIdxUsedSize;
+				}
 			}
 			else
 			{
@@ -244,13 +252,20 @@ static int32 PopulateNonSelectableIdx(size_t* NonSelectableIdx, int32 NonSelecta
 				}
 			}
 
-			if (CurrentIndexAsset->SamplingInterval.Contains(PoseAssetPlayerTime) && NonSelectableIdxUsedSize < NonSelectableIdxSize)
+			if (CurrentIndexAsset->SamplingInterval.Contains(PoseAssetPlayerTime))
 			{
-				NonSelectableIdx[NonSelectableIdxUsedSize++] = PoseIdx;
+				if (NonSelectableIdxUsedSize < NonSelectableIdxSize)
+				{
+					NonSelectableIdx[NonSelectableIdxUsedSize++] = PoseIdx;
+				}
+				else
+				{
+					UE_LOG(LogPoseSearch, Warning, TEXT("PopulateNonSelectableIdx couldn't add all the NonSelectableIdx"));
+					return NonSelectableIdxUsedSize;
+				}
 			}
 			else
 			{
-				UE_LOG(LogPoseSearch, Warning, TEXT("PopulateNonSelectableIdx couldn't add all the NonSelectableIdx"));
 				break;
 			}
 		}
@@ -261,18 +276,18 @@ static int32 PopulateNonSelectableIdx(size_t* NonSelectableIdx, int32 NonSelecta
 		const FObjectKey DatabaseKey(Database);
 		for (auto It = SearchContext.PoseIndicesHistory->IndexToTime.CreateConstIterator(); It; ++It)
 		{
-			if (NonSelectableIdxUsedSize < NonSelectableIdxSize)
+			const FHistoricalPoseIndex& HistoricalPoseIndex = It.Key();
+			if (HistoricalPoseIndex.DatabaseKey == DatabaseKey)
 			{
-				const FHistoricalPoseIndex& HistoricalPoseIndex = It.Key();
-				if (HistoricalPoseIndex.DatabaseKey == DatabaseKey)
+				if (NonSelectableIdxUsedSize < NonSelectableIdxSize)
 				{
 					NonSelectableIdx[NonSelectableIdxUsedSize++] = HistoricalPoseIndex.PoseIndex;
 				}
-			}
-			else
-			{
-				UE_LOG(LogPoseSearch, Warning, TEXT("PopulateNonSelectableIdx couldn't add all the NonSelectableIdx"));
-				break;
+				else
+				{
+					UE_LOG(LogPoseSearch, Warning, TEXT("PopulateNonSelectableIdx couldn't add all the NonSelectableIdx"));
+					return NonSelectableIdxUsedSize;
+				}
 			}
 		}
 	}
