motion matching - UPoseSearchDatabaseSet::Search bEvaluateContinuingPoseFirst fix - it was causing  "Invalid result searching..." log spamming

diff --git a/Engine/Plugins/Experimental/Animation/PoseSearch/Source/Runtime/Private/PoseSearch.cpp b/Engine/Plugins/Experimental/Animation/PoseSearch/Source/Runtime/Private/PoseSearch.cpp
--- a/Engine/Plugins/Experimental/Animation/PoseSearch/Source/Runtime/Private/PoseSearch.cpp
+++ b/Engine/Plugins/Experimental/Animation/PoseSearch/Source/Runtime/Private/PoseSearch.cpp
@@ -2127,6 +2127,13 @@ UE::PoseSearch::FSearchResult UPoseSearchDatabaseSet::Search(UE::PoseSearch::FSe
 			Result.ContinuingPoseCost = Result.PoseCost;
 			ContinuingCost = Result.PoseCost;
 
+			Result.SearchIndexAsset = SearchIndex->FindAssetForPose(Result.PoseIdx);
+			Result.AssetTime = SearchIndex->GetAssetTime(Result.PoseIdx, Result.SearchIndexAsset);
+			Result.Database = Database;
+#if WITH_EDITOR
+			Result.SearchIndexHash = Database->GetSearchIndexHash();
+#endif // WITH_EDITOR
+
 			if (Database->GetSkipSearchIfPossible())
 			{
 				SearchContext.UpdateCurrentBestCost(Result.PoseCost);
