motion matching - pose search debugger channel cost breakdowns

diff --git a/Engine/Plugins/Experimental/Animation/PoseSearch/Source/Editor/Private/PoseSearchDebugger.cpp b/Engine/Plugins/Experimental/Animation/PoseSearch/Source/Editor/Private/PoseSearchDebugger.cpp
--- a/Engine/Plugins/Experimental/Animation/PoseSearch/Source/Editor/Private/PoseSearchDebugger.cpp
+++ b/Engine/Plugins/Experimental/Animation/PoseSearch/Source/Editor/Private/PoseSearchDebugger.cpp
@@ -38,6 +38,16 @@ static FLinearColor LinearColorBlend(FLinearColor LinearColorA, FLinearColor Lin
 	return LinearColorA + (LinearColorB - LinearColorA) * BlendParam;
 }
 
+static void LinearColorBlend(FLinearColor LinearColorA, FLinearColor LinearColorB, TConstArrayView<const float> BlendParam, TArray<FLinearColor>& BlendedColors)
+{
+	const int32 Num = BlendParam.Num();
+	BlendedColors.SetNumUninitialized(Num);
+	for (int i = 0; i < Num; ++i)
+	{
+		BlendedColors[i] = LinearColorBlend(LinearColorA, LinearColorB, BlendParam[i]);
+	}
+}
+
 } // namespace UE::PoseSearch
 
 void UPoseSearchMeshComponent::Initialize(const FTransform& InComponentToWorld)
@@ -205,46 +215,11 @@ void UPoseSearchMeshComponent::UpdatePose(const FUpdateContext& UpdateContext)
 
 namespace UE::PoseSearch {
 
-struct FChannelCostRange
-{
-	float Min = MAX_flt;
-	float Max = -MAX_flt;
-	float Delta = 1.0f;
-};
-
 class FDebuggerDatabaseRowData : public TSharedFromThis<FDebuggerDatabaseRowData>
 {
 public:
 	FDebuggerDatabaseRowData() = default;
 	
-	float GetChannelCost(int32 ChannelIdx) const
-	{
-		return PoseCost.ChannelCosts.IsValidIndex(ChannelIdx)
-			? PoseCost.ChannelCosts[ChannelIdx]
-			: 0.0f;
-	}
-	
-	float GetAddendsCost() const { return PoseCost.GetCostAddend(); }
-
-	void CalculateColors(FChannelCostRange TotalCostRange, TArrayView<const FChannelCostRange> ChannelCostRanges)
-	{
-		const float CostColorBlend = (PoseCost.GetTotalCost() - TotalCostRange.Min) / TotalCostRange.Delta;
-		CostColor = UE::PoseSearch::LinearColorBlend(FLinearColor::Green, FLinearColor::Red, CostColorBlend);
-
-		ChannelCostColors.SetNum(ChannelCostRanges.Num());
-		for (int32 ChannelIdx = 0; ChannelIdx != ChannelCostRanges.Num(); ++ChannelIdx)
-		{
-			const FChannelCostRange& CostRange = ChannelCostRanges[ChannelIdx];
-			const float ColorBlend = (GetChannelCost(ChannelIdx) - CostRange.Min) / CostRange.Delta;
-			ChannelCostColors[ChannelIdx] = UE::PoseSearch::LinearColorBlend(FLinearColor::Green, FLinearColor::Red, ColorBlend);
-		}
-	}
-
-	FLinearColor GetChannelCostColor (int32 ChannelIdx) const
-	{
-		return ChannelCostColors.IsValidIndex(ChannelIdx) ? ChannelCostColors[ChannelIdx] : FLinearColor::White;
-	}
-
 	ESearchIndexAssetType AssetType = ESearchIndexAssetType::Invalid;
 	int32 PoseIdx = 0;
 	TWeakObjectPtr<const UPoseSearchDatabase> SourceDatabase = nullptr;
@@ -262,7 +237,9 @@ public:
 	FVector BlendParameters = FVector::Zero();
 	FPoseSearchCost PoseCost;
 	FLinearColor CostColor = FLinearColor::White;
-	TArray<FLinearColor> ChannelCostColors;
+	TArray<float> CostBreakdowns;
+	TArray<FLinearColor> CostBreakdownsColors;
+	TArray<float> CostVector;
 };
 
 namespace DebuggerDatabaseColumns
@@ -331,7 +308,7 @@ namespace DebuggerDatabaseColumns
 	{
 		using ITextColumn::ITextColumn;
 
-		virtual FText GetLabel() const override { return LOCTEXT("ColumnLabelPoseIndex", "Pose Index"); }
+		virtual FText GetLabel() const override { return LOCTEXT("ColumnLabelPoseIndex", "Index"); }
 
 		virtual FSortPredicate GetSortPredicate() const override
 		{
@@ -402,34 +379,6 @@ namespace DebuggerDatabaseColumns
 		}
 	};
 
-	struct FAssetType : ITextColumn
-	{
-		using ITextColumn::ITextColumn;
-
-		virtual FText GetLabel() const override { return LOCTEXT("ColumnLabelAssetType", "Type"); }
-
-		virtual FSortPredicate GetSortPredicate() const override
-		{
-			return [](const FRowDataRef& Row0, const FRowDataRef& Row1) -> bool { return Row0->AssetType < Row1->AssetType; };
-		}
-
-		virtual FText GetRowText(const FRowDataRef& Row) const override
-		{
-			if (Row->AssetType == ESearchIndexAssetType::Sequence)
-			{
-				return FText::FromString(TEXT("Sequence"));
-			}
-			else if (Row->AssetType == ESearchIndexAssetType::BlendSpace)
-			{
-				return FText::FromString(TEXT("BlendSpace"));
-			}
-			else
-			{
-				return FText::FromString(TEXT("None"));
-			}
-		}
-	};
-
 	struct FFrame : ITextColumn
 	{
 		using ITextColumn::ITextColumn;
@@ -511,69 +460,68 @@ namespace DebuggerDatabaseColumns
 		}
 	};
 
-	struct FChannelCostColumn : ITextColumn
+	struct FChannelBreakdownCostColumn : ITextColumn
 	{
-		struct FParams
+		FChannelBreakdownCostColumn(int32 SortIndex, int32 InBreakdownCostIndex, const FText& InLabel)
+			: ITextColumn(SortIndex)
+			, Label(InLabel)
+			, BreakdownCostIndex(InBreakdownCostIndex)
 		{
-			int32 SortIndex;
-			bool bEnabled = true;
-			int32 ChannelIdx;
-		};
-
-		FChannelCostColumn(const FParams& Params)
-			: ITextColumn(Params.SortIndex, Params.bEnabled)
-			, ChannelIdx(Params.ChannelIdx)
-		{}
+		}
 
-		virtual FText GetLabel() const override { return FText::Format(LOCTEXT("ColumnLabelChannelCost", "Cost[{0}]"), ChannelIdx); }
+		virtual FText GetLabel() const override { return Label; }
 
 		virtual FSortPredicate GetSortPredicate() const override
 		{
-			return [this](const FRowDataRef& Row0, const FRowDataRef& Row1) -> bool { return Row0->GetChannelCost(ChannelIdx) < Row1->GetChannelCost(ChannelIdx); };
+			return [this](const FRowDataRef& Row0, const FRowDataRef& Row1) -> bool
+			{
+				return Row0->CostBreakdowns[BreakdownCostIndex] < Row1->CostBreakdowns[BreakdownCostIndex];
+			};
 		}
 
 		virtual FText GetRowText(const FRowDataRef& Row) const override
 		{
-			if (Row->PoseCost.ChannelCosts.IsValidIndex(ChannelIdx))
+			const float LabelCost = Row->CostBreakdowns[BreakdownCostIndex];
+			if (LabelCost != UE_MAX_FLT)
 			{
-				return FText::AsNumber(Row->GetChannelCost(ChannelIdx));
-			}
-			else
-			{
-				return FText::FromString(TEXT("-"));
+				return FText::AsNumber(LabelCost);
 			}
+			return FText::FromString(TEXT("-"));
 		}
 
 		virtual FSlateColor GetColorAndOpacity(const FRowDataRef& Row) const override
 		{
-			return Row->GetChannelCostColor(ChannelIdx);
+			return Row->CostBreakdownsColors[BreakdownCostIndex];
 		}
 
-		int32 ChannelIdx = -1;
+		FText Label;
+		int32 BreakdownCostIndex = -1;
 	};
 
+#if WITH_EDITORONLY_DATA
 	struct FCostModifier : ITextColumn
 	{
 		using ITextColumn::ITextColumn;
 
-		virtual FText GetLabel() const override { return LOCTEXT("ColumnLabelCostModifier", "Cost Modifier"); }
+		virtual FText GetLabel() const override { return LOCTEXT("ColumnLabelCostModifier", "Bias"); }
 
 		virtual FSortPredicate GetSortPredicate() const override
 		{
-			return [](const FRowDataRef& Row0, const FRowDataRef& Row1) -> bool { return Row0->GetAddendsCost() < Row1->GetAddendsCost(); };
+			return [](const FRowDataRef& Row0, const FRowDataRef& Row1) -> bool { return Row0->PoseCost.GetCostAddend() < Row1->PoseCost.GetCostAddend(); };
 		}
 
 		virtual FText GetRowText(const FRowDataRef& Row) const override
 		{
-			return FText::AsNumber(Row->GetAddendsCost());
+			return FText::AsNumber(Row->PoseCost.GetCostAddend());
 		}
 	};
+#endif // WITH_EDITORONLY_DATA
 
 	struct FMirrored : ITextColumn
 	{
 		using ITextColumn::ITextColumn;
 
-		virtual FText GetLabel() const override { return LOCTEXT("ColumnLabelMirrored", "Mirrored"); }
+		virtual FText GetLabel() const override { return LOCTEXT("ColumnLabelMirrored", "Mirror"); }
 
 		virtual FSortPredicate GetSortPredicate() const override
 		{
@@ -590,7 +538,7 @@ namespace DebuggerDatabaseColumns
 	{
 		using ITextColumn::ITextColumn;
 
-		virtual FText GetLabel() const override { return LOCTEXT("ColumnLabelLooping", "Looping"); }
+		virtual FText GetLabel() const override { return LOCTEXT("ColumnLabelLooping", "Loop"); }
 
 		virtual FSortPredicate GetSortPredicate() const override
 		{
@@ -607,7 +555,7 @@ namespace DebuggerDatabaseColumns
 	{
 		using ITextColumn::ITextColumn;
 
-		virtual FText GetLabel() const override { return LOCTEXT("ColumnLabelBlendParams", "Blend Parameters"); }
+		virtual FText GetLabel() const override { return LOCTEXT("ColumnLabelBlendParams", "Blend Params"); }
 
 		virtual FSortPredicate GetSortPredicate() const override
 		{
@@ -728,60 +676,312 @@ class SDebuggerMessageBox : public SCompoundWidget
 	}
 };
 
-void SDebuggerDatabaseView::Update(const FTraceMotionMatchingStateMessage& State)
+class SCostBreakDownData : public ICostBreakDownData
 {
-	CreateRows(State);
+public:
+	SCostBreakDownData(TArray<TSharedRef<FDebuggerDatabaseRowData>>& UnfilteredDatabaseRows_, bool bIsVerbose_)
+	: UnfilteredDatabaseRows(UnfilteredDatabaseRows_)
+	, bIsVerbose(bIsVerbose_)
+	{
+	}
 
-	SortDatabaseRows();
+	// begin ICostBreakDownData
+	virtual ~SCostBreakDownData()
+	{
+		check(CurrentLabeIndex == -1);
+	}
 
-	PopulateViewRows();
+	virtual int32 Num() const override
+	{
+		return UnfilteredDatabaseRows.Num();
+	}
 
-	//	if (bNewDatabase)
-	// 	{
-	// 		FilteredDatabaseView.ListView->ClearSelection();
-	// 	}
+	virtual bool IsCostVectorFromSchema(int32 Index, const UPoseSearchSchema* Schema) const override
+	{
+		return UnfilteredDatabaseRows[Index]->SourceDatabase->Schema == Schema;
+	}
 
-}
+	virtual TConstArrayView<const float> GetCostVector(int32 Index, const UPoseSearchSchema* Schema) const override
+	{
+		check(UnfilteredDatabaseRows[Index]->SourceDatabase->Schema == Schema);
+		return UnfilteredDatabaseRows[Index]->CostVector;
+	}
 
-void SDebuggerDatabaseView::RefreshColumns()
+	virtual void BeginBreakDownSection(const FText& Label) override
+	{
+		check(CurrentLabeIndex == -1);
+		CurrentLabeIndex = GetOrAddLabelIndex(Label);
+	}
+
+	virtual void SetCostBreakDown(float CostBreakDown, int32 Index, const UPoseSearchSchema* Schema) override
+	{
+		FDebuggerDatabaseRowData& UnfilteredDatabaseRow = UnfilteredDatabaseRows[Index].Get();
+		check(UnfilteredDatabaseRow.SourceDatabase->Schema == Schema);
+		TArray<float>& CostBreakdowns = UnfilteredDatabaseRow.CostBreakdowns;
+		CostBreakdowns.SetNumZeroed(CurrentLabeIndex + 1, false);
+		CostBreakdowns[CurrentLabeIndex] = CostBreakDown;
+	}
+
+	virtual void EndBreakDownSection(const FText& Label) override
+	{
+		check(CurrentLabeIndex == GetOrAddLabelIndex(Label));
+		CurrentLabeIndex = -1;
+	}
+
+	virtual bool IsVerbose() const override
+	{
+		return bIsVerbose;
+	}
+	// end ICostBreakDownData
+
+	int32 GetOrAddLabelIndex(const FText& Label)
+	{
+		for (int32 i = 0; i < Labels.Num(); ++i)
+		{
+			if (Labels[i].EqualTo(Label))
+			{
+				return i;
+			}
+		}
+		Labels.Add(Label);
+		return Labels.Num() - 1;
+	}
+
+	bool AreLabelsEqualTo(TArray<FText> OtherLabels)
+	{
+		if (Labels.Num() != OtherLabels.Num())
+		{
+			return false;
+		}
+
+		for (int32 i = 0; i < Labels.Num(); ++i)
+		{
+			if (!Labels[i].EqualTo(OtherLabels[i]))
+			{
+				return false;
+			}
+		}
+
+		return true;
+	}
+
+	const TArray<FText>& GetLabels() const
+	{
+		return Labels;
+	}
+
+private:
+	TArray<TSharedRef<FDebuggerDatabaseRowData>>& UnfilteredDatabaseRows;
+
+	TArray<FText> Labels;
+
+	int32 CurrentLabeIndex = -1;
+
+	bool bIsVerbose = true;
+};
+
+void SDebuggerDatabaseView::Update(const FTraceMotionMatchingStateMessage& State)
 {
 	using namespace DebuggerDatabaseColumns;
-	
-	ActiveView.HeaderRow->ClearColumns();
-	ContinuingPoseView.HeaderRow->ClearColumns();
-	FilteredDatabaseView.HeaderRow->ClearColumns();
-	
-	// Sort columns by index
-	Columns.ValueSort([](const TSharedRef<IColumn> Column0, const TSharedRef<IColumn> Column1)
+
+	bool bIsVerbose = false;
+	auto DebuggerView = ParentDebuggerViewPtr.Pin();
+	if (DebuggerView.IsValid())
 	{
-		return Column0->SortIndex < Column1->SortIndex;
-	});
+		bIsVerbose = DebuggerView->GetViewModel()->IsVerbose();
+	}
 
-	// Add columns from map to header row
-	for(TPair<FName, TSharedRef<IColumn>>& ColumnPair : Columns)
+	UnfilteredDatabaseRows.Reset();
+	for (const FTraceMotionMatchingStateDatabaseEntry& DbEntry : State.DatabaseEntries)
 	{
-		IColumn& Column = ColumnPair.Value.Get();
-		if(ColumnPair.Value->bEnabled)
+		const UPoseSearchDatabase* Database = FTraceMotionMatchingState::GetObjectFromId<UPoseSearchDatabase>(DbEntry.DatabaseId);
+
+		const FPoseSearchIndex* SearchIndex = Database->GetSearchIndex();
+		if (SearchIndex && SearchIndex->IsValid())
 		{
-			SHeaderRow::FColumn::FArguments ColumnArgs = SHeaderRow::FColumn::FArguments()
-				.ColumnId(Column.ColumnId)
-				.DefaultLabel(Column.GetLabel())
-				.SortMode(this, &SDebuggerDatabaseView::GetColumnSortMode, Column.ColumnId)
-				.OnSort(this, &SDebuggerDatabaseView::OnColumnSortModeChanged)
-				.FillWidth(this, &SDebuggerDatabaseView::GetColumnWidth, Column.ColumnId)
-				.VAlignCell(VAlign_Center)
-				.VAlignHeader(VAlign_Center)
-				.HAlignHeader(HAlign_Center)
-				.HAlignCell(HAlign_Fill);
+			for (const FTraceMotionMatchingStatePoseEntry& PoseEntry : DbEntry.PoseEntries)
+			{
+				const FPoseSearchIndexAsset* SearchIndexAsset = SearchIndex->FindAssetForPose(PoseEntry.DbPoseIdx);
+				if (SearchIndexAsset)
+				{
+					TSharedRef<FDebuggerDatabaseRowData>& Row = UnfilteredDatabaseRows.Add_GetRef(MakeShared<FDebuggerDatabaseRowData>());
 
-			FilteredDatabaseView.HeaderRow->AddColumn(ColumnArgs);
-			
-			// Every time the active column is changed, update the database column
-			ActiveView.HeaderRow->AddColumn(ColumnArgs.OnWidthChanged(this, &SDebuggerDatabaseView::OnColumnWidthChanged, Column.ColumnId));
-			
-			ContinuingPoseView.HeaderRow->AddColumn(ColumnArgs.OnWidthChanged(this, &SDebuggerDatabaseView::OnColumnWidthChanged, Column.ColumnId));
+					const float Time = SearchIndex->GetAssetTime(PoseEntry.DbPoseIdx, SearchIndexAsset);
+
+					Row->PoseIdx = PoseEntry.DbPoseIdx;
+					Row->SourceDatabase = Database;
+					Row->DatabaseName = Database->GetName();
+					Row->DatabasePath = Database->GetPathName();
+					Row->PoseEntryFlags = PoseEntry.Flags;
+					Row->DbAssetIdx = SearchIndexAsset->SourceAssetIdx;
+					Row->AssetTime = Time;
+					Row->bMirrored = SearchIndexAsset->bMirrored;
+					Row->PoseCost = PoseEntry.Cost;
+
+					Row->CostVector.SetNum(DbEntry.QueryVector.Num());
+					TConstArrayView<const float> PoseValues = SearchIndex->GetPoseValues(PoseEntry.DbPoseIdx);
+					CompareFeatureVectors(PoseValues, DbEntry.QueryVector, SearchIndex->WeightsSqrt, Row->CostVector);
+
+					if (SearchIndexAsset->Type == ESearchIndexAssetType::Sequence)
+					{
+						const FPoseSearchDatabaseSequence& DbSequence = Database->GetSequenceSourceAsset(SearchIndexAsset);
+						Row->AssetType = ESearchIndexAssetType::Sequence;
+						Row->AssetName = DbSequence.Sequence->GetName();
+						Row->AssetPath = DbSequence.Sequence->GetPathName();
+						Row->AnimFrame = DbSequence.Sequence->GetFrameAtTime(Time);
+						Row->AnimPercentage = Time / DbSequence.Sequence->GetPlayLength();
+						Row->bLooping = DbSequence.Sequence->bLoop;
+						Row->BlendParameters = FVector::Zero();
+					}
+					else if (SearchIndexAsset->Type == ESearchIndexAssetType::BlendSpace)
+					{
+						const FPoseSearchDatabaseBlendSpace& DbBlendSpace = Database->GetBlendSpaceSourceAsset(SearchIndexAsset);
+						Row->AssetType = ESearchIndexAssetType::BlendSpace;
+						Row->AssetName = DbBlendSpace.BlendSpace->GetName();
+						Row->AssetPath = DbBlendSpace.BlendSpace->GetPathName();
+						Row->AnimFrame = 0; // There is no frame index associated with a blendspace
+						Row->AnimPercentage = 0.0f;
+						Row->bLooping = DbBlendSpace.BlendSpace->bLoop;
+						Row->BlendParameters = SearchIndexAsset->BlendParameters;
+					}
+					else
+					{
+						checkNoEntry();
+					}
+				}
+			}
 		}
 	}
+
+	SCostBreakDownData CostBreakDownData(UnfilteredDatabaseRows, bIsVerbose);
+	if (!UnfilteredDatabaseRows.IsEmpty())
+	{
+		for (const FTraceMotionMatchingStateDatabaseEntry& DbEntry : State.DatabaseEntries)
+		{
+			const UPoseSearchDatabase* Database = FTraceMotionMatchingState::GetObjectFromId<UPoseSearchDatabase>(DbEntry.DatabaseId);
+			const FPoseSearchIndex* SearchIndex = Database->GetSearchIndex();
+			if (SearchIndex && SearchIndex->IsValid())
+			{
+				Database->Schema->ComputeCostBreakdowns(CostBreakDownData);
+			}
+		}
+
+		// calculating breakdowns min max and colors
+		TArray<float> MinCostBreakdowns;
+		TArray<float> MaxCostBreakdowns;
+
+		const int32 CostBreakdownsCardinality = UnfilteredDatabaseRows[0]->CostBreakdowns.Num();
+		MinCostBreakdowns.Init(UE_MAX_FLT, CostBreakdownsCardinality);
+		MaxCostBreakdowns.Init(-UE_MAX_FLT, CostBreakdownsCardinality);
+
+		for (auto& UnfilteredRow : UnfilteredDatabaseRows)
+		{
+			ArrayMinMax(UnfilteredRow->CostBreakdowns, MinCostBreakdowns, MaxCostBreakdowns, UE_MAX_FLT);
+		}
+
+		TArray<float> CostBreakdownsColorBlend;
+		CostBreakdownsColorBlend.Init(0, CostBreakdownsCardinality);
+		for (auto& UnfilteredRow : UnfilteredDatabaseRows)
+		{
+			ArraySafeNormalize(UnfilteredRow->CostBreakdowns, MinCostBreakdowns, MaxCostBreakdowns, CostBreakdownsColorBlend);
+			LinearColorBlend(FLinearColor::Green, FLinearColor::Red, CostBreakdownsColorBlend, UnfilteredRow->CostBreakdownsColors);
+		}
+
+		float MinCost = UE_MAX_FLT;
+		float MaxCost = -UE_MAX_FLT;
+		for (auto& UnfilteredRow : UnfilteredDatabaseRows)
+		{
+			const float Cost = UnfilteredRow->PoseCost.GetTotalCost();
+			MinCost = FMath::Min(MinCost, Cost);
+			MaxCost = FMath::Max(MaxCost, Cost);
+		}
+
+		const float DeltaCost = MaxCost - MinCost;
+		for (auto& UnfilteredRow : UnfilteredDatabaseRows)
+		{
+			const float CostColorBlend = DeltaCost > UE_KINDA_SMALL_NUMBER ? (UnfilteredRow->PoseCost.GetTotalCost() - MinCost) / DeltaCost : 0.f;
+			UnfilteredRow->CostColor = UE::PoseSearch::LinearColorBlend(FLinearColor::Green, FLinearColor::Red, CostColorBlend);
+		}
+	}
+
+	if (!CostBreakDownData.AreLabelsEqualTo(OldLabels))
+	{
+		OldLabels = CostBreakDownData.GetLabels();
+
+		// recreating and binding the columns
+
+		Columns.Reset();
+
+		// @TODO: Support runtime reordering of these indices
+		// Construct all column types
+		int32 ColumnIdx = 0;
+		AddColumn(MakeShared<FDatabaseName>(ColumnIdx++));
+		AddColumn(MakeShared<FAssetName>(ColumnIdx++));
+
+		auto CostColumn = MakeShared<FCost>(ColumnIdx++);
+		AddColumn(CostColumn);
+
+		int32 LabelIdx = 0;
+		for (const FText& Label : CostBreakDownData.GetLabels())
+		{
+			AddColumn(MakeShared<FChannelBreakdownCostColumn>(ColumnIdx++, LabelIdx++, Label));
+		}
+
+#if WITH_EDITORONLY_DATA
+		AddColumn(MakeShared<FCostModifier>(ColumnIdx++));
+#endif // WITH_EDITORONLY_DATA
+		AddColumn(MakeShared<FFrame>(ColumnIdx++));
+		AddColumn(MakeShared<FMirrored>(ColumnIdx++));
+		AddColumn(MakeShared<FLooping>(ColumnIdx++));
+		AddColumn(MakeShared<FPoseIdx>(ColumnIdx++));
+		AddColumn(MakeShared<FBlendParameters>(ColumnIdx++));
+
+		SortColumn = CostColumn->ColumnId;
+
+		// Active and Continuing Pose view scroll bars only for indenting the columns to align w/ database
+		ActiveView.ScrollBar->SetVisibility(EVisibility::Hidden);
+		ContinuingPoseView.ScrollBar->SetVisibility(EVisibility::Hidden);
+
+		// Refresh Columns
+		ActiveView.HeaderRow->ClearColumns();
+		ContinuingPoseView.HeaderRow->ClearColumns();
+		FilteredDatabaseView.HeaderRow->ClearColumns();
+
+		// Sort columns by index
+		Columns.ValueSort([](const TSharedRef<IColumn> Column0, const TSharedRef<IColumn> Column1)
+		{
+			return Column0->SortIndex < Column1->SortIndex;
+		});
+
+		// Add columns from map to header row
+		for (TPair<FName, TSharedRef<IColumn>>& ColumnPair : Columns)
+		{
+			IColumn& Column = ColumnPair.Value.Get();
+			if (ColumnPair.Value->bEnabled)
+			{
+				SHeaderRow::FColumn::FArguments ColumnArgs = SHeaderRow::FColumn::FArguments()
+					.ColumnId(Column.ColumnId)
+					.DefaultLabel(Column.GetLabel())
+					.SortMode(this, &SDebuggerDatabaseView::GetColumnSortMode, Column.ColumnId)
+					.OnSort(this, &SDebuggerDatabaseView::OnColumnSortModeChanged)
+					.FillWidth(this, &SDebuggerDatabaseView::GetColumnWidth, Column.ColumnId)
+					.VAlignCell(VAlign_Center)
+					.VAlignHeader(VAlign_Center)
+					.HAlignHeader(HAlign_Center)
+					.HAlignCell(HAlign_Fill);
+
+				FilteredDatabaseView.HeaderRow->AddColumn(ColumnArgs);
+
+				// Every time the active column is changed, update the database column
+				ActiveView.HeaderRow->AddColumn(ColumnArgs.OnWidthChanged(this, &SDebuggerDatabaseView::OnColumnWidthChanged, Column.ColumnId));
+
+				ContinuingPoseView.HeaderRow->AddColumn(ColumnArgs.OnWidthChanged(this, &SDebuggerDatabaseView::OnColumnWidthChanged, Column.ColumnId));
+			}
+		}
+	}
+
+	SortDatabaseRows();
+	PopulateViewRows();
 }
 
 void SDebuggerDatabaseView::AddColumn(TSharedRef<DebuggerDatabaseColumns::IColumn>&& Column)
@@ -900,213 +1100,40 @@ void SDebuggerDatabaseView::PopulateViewRows()
 		}
 	}
 
-	ComputeViewRowsColors();
-
 	ActiveView.ListView->RequestListRefresh();
 	ContinuingPoseView.ListView->RequestListRefresh();
 	FilteredDatabaseView.ListView->RequestListRefresh();
 }
 
-void SDebuggerDatabaseView::CreateRows(const FTraceMotionMatchingStateMessage& State)
-{
-	UnfilteredDatabaseRows.Reset();
-
-	for (const FTraceMotionMatchingStateDatabaseEntry& DbEntry : State.DatabaseEntries)
-	{
-		const auto* Database = FTraceMotionMatchingState::GetObjectFromId<UPoseSearchDatabase>(DbEntry.DatabaseId);
-
-		const FPoseSearchIndex* SearchIndex = Database->GetSearchIndex();
-		if (SearchIndex && SearchIndex->IsValid())
-		{
-			for (const FTraceMotionMatchingStatePoseEntry& PoseEntry : DbEntry.PoseEntries)
-			{
-				const FPoseSearchIndexAsset* SearchIndexAsset = SearchIndex->FindAssetForPose(PoseEntry.DbPoseIdx);
-				if (SearchIndexAsset)
-				{
-					TSharedRef<FDebuggerDatabaseRowData>& Row = UnfilteredDatabaseRows.Add_GetRef(MakeShared<FDebuggerDatabaseRowData>());
-
-					const float Time = SearchIndex->GetAssetTime(PoseEntry.DbPoseIdx, SearchIndexAsset);
-
-					Row->PoseIdx = PoseEntry.DbPoseIdx;
-					Row->SourceDatabase = Database;
-					Row->DatabaseName = Database->GetName();
-					Row->DatabasePath = Database->GetPathName();
-					Row->PoseEntryFlags = PoseEntry.Flags;
-					Row->DbAssetIdx = SearchIndexAsset->SourceAssetIdx;
-					Row->AssetTime = Time;
-					Row->bMirrored = SearchIndexAsset->bMirrored;
-					Row->PoseCost = PoseEntry.Cost;
-
-					if (SearchIndexAsset->Type == ESearchIndexAssetType::Sequence)
-					{
-						const FPoseSearchDatabaseSequence& DbSequence = Database->GetSequenceSourceAsset(SearchIndexAsset);
-						Row->AssetType = ESearchIndexAssetType::Sequence;
-						Row->AssetName = DbSequence.Sequence->GetName();
-						Row->AssetPath = DbSequence.Sequence->GetPathName();
-						Row->AnimFrame = DbSequence.Sequence->GetFrameAtTime(Time);
-						Row->AnimPercentage = Time / DbSequence.Sequence->GetPlayLength();
-						Row->bLooping = DbSequence.Sequence->bLoop;
-						Row->BlendParameters = FVector::Zero();
-					}
-					else if (SearchIndexAsset->Type == ESearchIndexAssetType::BlendSpace)
-					{
-						const FPoseSearchDatabaseBlendSpace& DbBlendSpace = Database->GetBlendSpaceSourceAsset(SearchIndexAsset);
-						Row->AssetType = ESearchIndexAssetType::BlendSpace;
-						Row->AssetName = DbBlendSpace.BlendSpace->GetName();
-						Row->AssetPath = DbBlendSpace.BlendSpace->GetPathName();
-						Row->AnimFrame = 0; // There is no frame index associated with a blendspace
-						Row->AnimPercentage = 0.0f;
-						Row->bLooping = DbBlendSpace.BlendSpace->bLoop;
-						Row->BlendParameters = SearchIndexAsset->BlendParameters;
-					}
-					else
-					{
-						checkNoEntry();
-					}
-				}
-			}
-		}
-	}
-}
-
-void SDebuggerDatabaseView::ComputeViewRowsColors()
-{
-	auto UpdateCostRanges = [](const TSharedRef<FDebuggerDatabaseRowData>& Row, FChannelCostRange& CostRange, TArray<FChannelCostRange>& ChannelCostRanges) -> void
-	{
-		const float Cost = Row->PoseCost.GetTotalCost();
-		CostRange.Min = FMath::Min(CostRange.Min, Cost);
-		CostRange.Max = FMath::Max(CostRange.Max, Cost);
-
-		const int32 NumChannels = Row->PoseCost.ChannelCosts.Num();
-		if (ChannelCostRanges.Num() < NumChannels)
-		{
-			ChannelCostRanges.SetNum(NumChannels);
-		}
-
-		for (int32 ChannelIdx = 0; ChannelIdx != NumChannels; ++ChannelIdx)
-		{
-			ChannelCostRanges[ChannelIdx].Min = FMath::Min(ChannelCostRanges[ChannelIdx].Min, Row->PoseCost.ChannelCosts[ChannelIdx]);
-			ChannelCostRanges[ChannelIdx].Max = FMath::Max(ChannelCostRanges[ChannelIdx].Max, Row->PoseCost.ChannelCosts[ChannelIdx]);
-		}
-	};
-
-	FChannelCostRange CostRange;
-	TArray<FChannelCostRange> ChannelCostRanges;
-	for (const TSharedRef<FDebuggerDatabaseRowData>& Row : FilteredDatabaseView.Rows)
-	{
-		UpdateCostRanges(Row, CostRange, ChannelCostRanges);
-	}
-	for (const TSharedRef<FDebuggerDatabaseRowData>& Row : ActiveView.Rows)
-	{
-		UpdateCostRanges(Row, CostRange, ChannelCostRanges);
-	}
-	for (const TSharedRef<FDebuggerDatabaseRowData>& Row : ContinuingPoseView.Rows)
-	{
-		UpdateCostRanges(Row, CostRange, ChannelCostRanges);
-	}
-
-	CostRange.Delta = CostRange.Max - CostRange.Min;
-	if (FMath::IsNearlyZero(CostRange.Delta))
-	{
-		CostRange.Delta = 1.f;
-	}
-
-	for (int32 ChannelIdx = 0; ChannelIdx != ChannelCostRanges.Num(); ++ChannelIdx)
-	{
-		FChannelCostRange ChannelCostRange = ChannelCostRanges[ChannelIdx];
-		ChannelCostRange.Delta = ChannelCostRange.Max - ChannelCostRange.Min;
-		if (FMath::IsNearlyZero(ChannelCostRange.Delta))
-		{
-			ChannelCostRange.Delta = 1.f;
-		}
-	}
-
-	for (const TSharedRef<FDebuggerDatabaseRowData>& Row : FilteredDatabaseView.Rows)
-	{
-		Row->CalculateColors(CostRange, ChannelCostRanges);
-	}
-	for (const TSharedRef<FDebuggerDatabaseRowData>& Row : ActiveView.Rows)
-	{
-		Row->CalculateColors(CostRange, ChannelCostRanges);
-	}
-	for (const TSharedRef<FDebuggerDatabaseRowData>& Row : ContinuingPoseView.Rows)
-	{
-		Row->CalculateColors(CostRange, ChannelCostRanges);
-	}
-}
-
 TSharedRef<ITableRow> SDebuggerDatabaseView::HandleGenerateDatabaseRow(TSharedRef<FDebuggerDatabaseRowData> Item, const TSharedRef<STableViewBase>& OwnerTable) const
 {
-	return
-		SNew(SDebuggerDatabaseRow, OwnerTable, Item, FilteredDatabaseView.RowStyle, &FilteredDatabaseView.RowBrush, FMargin(0.0f, 2.0f, 6.0f, 2.0f))
-		.ColumnMap(this, &SDebuggerDatabaseView::GetColumnMap);
-	
+	return SNew(SDebuggerDatabaseRow, OwnerTable, Item, FilteredDatabaseView.RowStyle, &FilteredDatabaseView.RowBrush, FMargin(0.0f, 2.0f, 6.0f, 2.0f))
+		.ColumnMap(this, &SDebuggerDatabaseView::GetColumnMap);	
 }
 
 TSharedRef<ITableRow> SDebuggerDatabaseView::HandleGenerateActiveRow(TSharedRef<FDebuggerDatabaseRowData> Item, const TSharedRef<STableViewBase>& OwnerTable) const
 {
-	return
-		SNew(SDebuggerDatabaseRow, OwnerTable, Item, ActiveView.RowStyle, &ActiveView.RowBrush, FMargin(0.0f, 2.0f, 6.0f, 4.0f))
+	return SNew(SDebuggerDatabaseRow, OwnerTable, Item, ActiveView.RowStyle, &ActiveView.RowBrush, FMargin(0.0f, 2.0f, 6.0f, 4.0f))
 		.ColumnMap(this, &SDebuggerDatabaseView::GetColumnMap);
 }
 
 TSharedRef<ITableRow> SDebuggerDatabaseView::HandleGenerateContinuingPoseRow(TSharedRef<FDebuggerDatabaseRowData> Item, const TSharedRef<STableViewBase>& OwnerTable) const
 {
-	return
-		SNew(SDebuggerDatabaseRow, OwnerTable, Item, ContinuingPoseView.RowStyle, &ContinuingPoseView.RowBrush, FMargin(0.0f, 2.0f, 6.0f, 4.0f))
+	return SNew(SDebuggerDatabaseRow, OwnerTable, Item, ContinuingPoseView.RowStyle, &ContinuingPoseView.RowBrush, FMargin(0.0f, 2.0f, 6.0f, 4.0f))
 		.ColumnMap(this, &SDebuggerDatabaseView::GetColumnMap);
 }
 
 void SDebuggerDatabaseView::Construct(const FArguments& InArgs)
 {
-	using namespace DebuggerDatabaseColumns;
-
 	ParentDebuggerViewPtr = InArgs._Parent;
 	OnPoseSelectionChanged = InArgs._OnPoseSelectionChanged;
 	check(OnPoseSelectionChanged.IsBound());
 
-	// @TODO: Support runtime reordering of these indices
-	// Construct all column types
-	int32 ColumnIdx = 0;
-	AddColumn(MakeShared<FDatabaseName>(ColumnIdx++));
-	AddColumn(MakeShared<FAssetName>(ColumnIdx++));
-	AddColumn(MakeShared<FAssetType>(ColumnIdx++));
-
-	auto CostColumn = MakeShared<FCost>(ColumnIdx++);
-	AddColumn(CostColumn);
-
-	int32 ChannelIdx = 0;
-	FChannelCostColumn::FParams ChannelCostParams;
-	ChannelCostParams.ChannelIdx = ChannelIdx++;
-	ChannelCostParams.SortIndex = ColumnIdx++;
-	AddColumn(MakeShared<FChannelCostColumn>(ChannelCostParams));
-
-	ChannelCostParams.ChannelIdx = ChannelIdx++;
-	ChannelCostParams.SortIndex = ColumnIdx++;
-	AddColumn(MakeShared<FChannelCostColumn>(ChannelCostParams));
-
-	ChannelCostParams.ChannelIdx = ChannelIdx++;
-	ChannelCostParams.SortIndex = ColumnIdx++;
-	AddColumn(MakeShared<FChannelCostColumn>(ChannelCostParams));
-
-	ChannelCostParams.ChannelIdx = ChannelIdx++;
-	ChannelCostParams.SortIndex = ColumnIdx++;
-	AddColumn(MakeShared<FChannelCostColumn>(ChannelCostParams));
-
-
-	AddColumn(MakeShared<FCostModifier>(ColumnIdx++));
-	AddColumn(MakeShared<FFrame>(ColumnIdx++));
-	AddColumn(MakeShared<FMirrored>(ColumnIdx++));
-	AddColumn(MakeShared<FLooping>(ColumnIdx++));
-	AddColumn(MakeShared<FBlendParameters>(ColumnIdx++));
-	AddColumn(MakeShared<FPoseIdx>(ColumnIdx++));
-
 	// Active Row
 	ActiveView.HeaderRow = SNew(SHeaderRow);
 
 	// Used for spacing
-	ActiveView.ScrollBar =
-		SNew(SScrollBar)
+	ActiveView.ScrollBar = SNew(SScrollBar)
 		.Orientation(Orient_Vertical)
 		.HideWhenNotInUse(false)
 		.AlwaysShowScrollbar(true)
@@ -1127,8 +1154,7 @@ void SDebuggerDatabaseView::Construct(const FArguments& InArgs)
 	ContinuingPoseView.HeaderRow = SNew(SHeaderRow).Visibility(EVisibility::Collapsed);
 
 	// Used for spacing
-	ContinuingPoseView.ScrollBar =
-		SNew(SScrollBar)
+	ContinuingPoseView.ScrollBar = SNew(SScrollBar)
 		.Orientation(Orient_Vertical)
 		.HideWhenNotInUse(false)
 		.AlwaysShowScrollbar(true)
@@ -1370,15 +1396,10 @@ void SDebuggerDatabaseView::Construct(const FArguments& InArgs)
 			]
 		]
 	];
-	
-	SortColumn = CostColumn->ColumnId;
-	SortMode = EColumnSortMode::Ascending;
 
-	// Active and Continuing Pose view scroll bars only for indenting the columns to align w/ database
-	ActiveView.ScrollBar->SetVisibility(EVisibility::Hidden);
-	ContinuingPoseView.ScrollBar->SetVisibility(EVisibility::Hidden);
-
-	RefreshColumns();
+	SortMode = EColumnSortMode::Ascending;
+	OldLabels.Reset();
+	Columns.Reset();
 }
 
 void SDebuggerDetailsView::Construct(const FArguments& InArgs)
@@ -1450,8 +1471,9 @@ void SDebuggerDetailsView::UpdateReflection(const FTraceMotionMatchingStateMessa
 	Reflection->AnimLinearVelocity = State.AnimLinearVelocity;
 	Reflection->AnimAngularVelocity = State.AnimAngularVelocity;
 	 
+	// @todo: fix me
 	// Query pose
-	Reflection->QueryPoseVector = State.QueryVector;
+	//Reflection->QueryPoseVector = State.QueryVector;
 
 	// Active pose
 	if (CurrentSearchIndex)
@@ -1475,15 +1497,15 @@ void SDebuggerDetailsView::UpdateReflection(const FTraceMotionMatchingStateMessa
 				}
 			}
 
-			Reflection->CostVector = Selected->PoseCost.CostVector;
-
-			const FTraceMotionMatchingStatePoseEntry* ActivePoseEntry = State.GetCurrentPoseEntry();
+			Reflection->CostVector = Selected->CostVector;
 
-			Reflection->CostVectorDifference = Reflection->CostVector;
-			for (int i = 0; i < Reflection->CostVectorDifference.Num(); ++i)
-			{
-				Reflection->CostVectorDifference[i] -= ActivePoseEntry->Cost.CostVector[i];
-			}
+			// @todo: fix me
+			//const FTraceMotionMatchingStatePoseEntry* ActivePoseEntry = State.GetCurrentPoseEntry();
+			//Reflection->CostVectorDifference = Reflection->CostVector;
+			//for (int i = 0; i < Reflection->CostVectorDifference.Num(); ++i)
+			//{
+			//	Reflection->CostVectorDifference[i] -= ActivePoseEntry->Cost.CostVector[i];
+			//}
 		}
 	}
 }
@@ -1492,9 +1514,7 @@ void SDebuggerView::Construct(const FArguments& InArgs, uint64 InAnimInstanceId)
 {
 	ViewModel = InArgs._ViewModel;
 	OnViewClosed = InArgs._OnViewClosed;
-
 	
-
 	// Validate the existence of the passed getters
 	check(ViewModel.IsBound())
 	check(OnViewClosed.IsBound());
@@ -1731,13 +1751,14 @@ void SDebuggerView::DrawFeatures(
 
 	const TObjectPtr<UPoseSearchDebuggerReflection> Reflection = DetailsView->GetReflection();
 
+	// @todo: fix me
 	// Draw query vector
 	{
-		DrawParams.Database = ViewModel.Get()->GetCurrentDatabase();
-		SetDrawFlags(DrawParams, Reflection ? Reflection->QueryDrawOptions : FPoseSearchDebuggerFeatureDrawOptions());
-		EnumAddFlags(DrawParams.Flags, EDebugDrawFlags::DrawQuery);
-		DrawFeatureVector(DrawParams, State.QueryVector);
-		EnumRemoveFlags(DrawParams.Flags, EDebugDrawFlags::DrawQuery);
+		//DrawParams.Database = ViewModel.Get()->GetCurrentDatabase();
+		//SetDrawFlags(DrawParams, Reflection ? Reflection->QueryDrawOptions : FPoseSearchDebuggerFeatureDrawOptions());
+		//EnumAddFlags(DrawParams.Flags, EDebugDrawFlags::DrawQuery);
+		//DrawFeatureVector(DrawParams, State.QueryVector);
+		//EnumRemoveFlags(DrawParams.Flags, EDebugDrawFlags::DrawQuery);
 	}
 
 
@@ -2033,6 +2054,35 @@ TSharedRef<SWidget> SDebuggerView::GenerateNodeDebuggerView()
 		]
 	];
 
+	ReturnButtonView->AddSlot()
+	.VAlign(VAlign_Center)
+	.HAlign(HAlign_Left)
+	.Padding(64, 5, 0, 0)
+	.AutoWidth()
+	[
+		SNew(SHorizontalBox)
+
+		+ SHorizontalBox::Slot()
+		.AutoWidth()
+		.Padding(0, 5, 0, 0)
+		[
+			SNew(SCheckBox)
+			.IsChecked_Lambda([this]
+			{
+				return ViewModel.Get()->IsVerbose() ? ECheckBoxState::Checked : ECheckBoxState::Unchecked; 
+			})
+			.OnCheckStateChanged_Lambda([this](ECheckBoxState State)
+			{
+				ViewModel.Get()->SetVerbose(State == ECheckBoxState::Checked);
+				UpdateViews();
+			})
+			[
+				SNew(STextBlock)
+				.Text(LOCTEXT("PoseSearchDebuggerShowVerbose", "Channels Breakdown"))
+			]
+		]
+	];
+
 	return 
 		SNew(SSplitter)
 		.Orientation(Orient_Horizontal)
diff --git a/Engine/Plugins/Experimental/Animation/PoseSearch/Source/Editor/Private/PoseSearchDebugger.h b/Engine/Plugins/Experimental/Animation/PoseSearch/Source/Editor/Private/PoseSearchDebugger.h
--- a/Engine/Plugins/Experimental/Animation/PoseSearch/Source/Editor/Private/PoseSearchDebugger.h
+++ b/Engine/Plugins/Experimental/Animation/PoseSearch/Source/Editor/Private/PoseSearchDebugger.h
@@ -207,25 +207,17 @@ class SDebuggerDatabaseView : public SCompoundWidget
 	using FColumnMap = TMap<FName, TSharedRef<DebuggerDatabaseColumns::IColumn>>;
 
 private:
-	/** Deletes existing columns and initializes a new set */
-	void RefreshColumns();
-
 	/** Adds a column to the existing list */
 	void AddColumn(TSharedRef<DebuggerDatabaseColumns::IColumn>&& Column);
 
 	/** Retrieves current column map, used as an attribute by rows */
 	const FColumnMap* GetColumnMap() const { return &Columns; }
 
-	/** Creates widgets for every pose in the database, initializing the static data in the process */
-	void CreateRows(const FTraceMotionMatchingStateMessage& State);
-
 	/** Sorts the database by the current sort predicate, updating the view order */
 	void SortDatabaseRows();
 
 	void PopulateViewRows();
 
-	void ComputeViewRowsColors();
-
 	/** Acquires sort predicate for the given column */
 	EColumnSortMode::Type GetColumnSortMode(const FName ColumnId) const;
 
@@ -267,6 +259,8 @@ private:
 	/** Column data container, used to emplace defined column structures of various types */
     FColumnMap Columns;
 
+	TArray<FText> OldLabels;
+
 	struct FTable
 	{
 		/** Header row*/
@@ -508,6 +502,10 @@ public:
 	/** Current play rate of the asset selection player */
 	float GetPlayRate() const { return AssetPlayRate; }
 	
+	void SetVerbose(bool bVerbose) { bIsVerbose = bVerbose; }
+
+	bool IsVerbose() const { return bIsVerbose; }
+
 	/** Callback to reset debug skeletons for the active world */
 	void OnWorldCleanup(UWorld* InWorld, bool bSessionEnded, bool bCleanupResources);
 
@@ -623,6 +621,8 @@ private:
 	/** Current play rate of the asset selection player */
 	float AssetPlayRate = 1.0f;
 
+	bool bIsVerbose = false;
+
 	/** Limits some public API */
 	friend class FDebugger;
 };
diff --git a/Engine/Plugins/Experimental/Animation/PoseSearch/Source/Runtime/Private/PoseSearch.cpp b/Engine/Plugins/Experimental/Animation/PoseSearch/Source/Runtime/Private/PoseSearch.cpp
--- a/Engine/Plugins/Experimental/Animation/PoseSearch/Source/Runtime/Private/PoseSearch.cpp
+++ b/Engine/Plugins/Experimental/Animation/PoseSearch/Source/Runtime/Private/PoseSearch.cpp
@@ -76,12 +76,51 @@ namespace UE { namespace PoseSearch {
 	constexpr EParallelForFlags ParallelForFlags = EParallelForFlags::None;
 #endif // UE_POSE_SEARCH_FORCE_SINGLE_THREAD
 
-static bool IsSamplingRangeValid(FFloatInterval Range)
+float ArraySum(TConstArrayView<const float> View, int32 StartIndex, int32 Offset)
 {
-	return Range.IsValid() && (Range.Min >= 0.0f);
+	float Sum = 0.f;
+	const int32 EndIndex = StartIndex + Offset;
+	for (int i = StartIndex; i < EndIndex; ++i)
+	{
+		Sum += View[i];
+	}
+	return Sum;
+}
+
+void ArrayMinMax(TConstArrayView<const float> View, TArrayView<float> Min, TArrayView<float> Max, float InvalidValue)
+{
+	const int32 Num = View.Num();
+	check(Num == Min.Num() && Num == Max.Num());
+	for (int i = 0; i < Num; ++i)
+	{
+		const float Value = View[i];
+		if (Value != InvalidValue)
+		{
+			Min[i] = FMath::Min(Min[i], Value);
+			Max[i] = FMath::Max(Max[i], Value);
+		}
+	}
+}
+
+void ArraySafeNormalize(TConstArrayView<const float> View, TConstArrayView<const float> Min, TConstArrayView<const float> Max, TArrayView<float> NormalizedView)
+{
+	const int32 Num = View.Num();
+	check(Num == Min.Num() && Num == Max.Num() && Num == NormalizedView.Num());
+	for (int i = 0; i < Num; ++i)
+	{
+		const float Delta = Max[i] - Min[i];
+		if (FMath::IsNearlyZero(Delta, UE_KINDA_SMALL_NUMBER))
+		{
+			NormalizedView[i] = 0.f;
+		}
+		else
+		{
+			NormalizedView[i] = (View[i] - Min[i]) / Delta;
+		}
+	}
 }
 
-static inline float CompareFeatureVectors(TConstArrayView<const float> A, TConstArrayView<const float> B, TConstArrayView<const float> WeightsSqrt)
+float CompareFeatureVectors(TConstArrayView<const float> A, TConstArrayView<const float> B, TConstArrayView<const float> WeightsSqrt)
 {
 	check(A.Num() == B.Num() && A.Num() == WeightsSqrt.Num());
 
@@ -92,7 +131,7 @@ static inline float CompareFeatureVectors(TConstArrayView<const float> A, TConst
 	return ((VA - VB) * VW).square().sum();
 }
 
-static inline float CompareFeatureVectors(TConstArrayView<const float> A, TConstArrayView<const float> B)
+float CompareFeatureVectors(TConstArrayView<const float> A, TConstArrayView<const float> B)
 {
 	check(A.Num() == B.Num());
 
@@ -102,7 +141,7 @@ static inline float CompareFeatureVectors(TConstArrayView<const float> A, TConst
 	return (VA - VB).square().sum();
 }
 
-static inline void CompareFeatureVectors(TConstArrayView<const float> A, TConstArrayView<const float> B, TConstArrayView<const float> WeightsSqrt, TArrayView<float> Result)
+void CompareFeatureVectors(TConstArrayView<const float> A, TConstArrayView<const float> B, TConstArrayView<const float> WeightsSqrt, TArrayView<float> Result)
 {
 	check(A.Num() == B.Num() && A.Num() == WeightsSqrt.Num() && A.Num() == Result.Num());
 
@@ -114,6 +153,11 @@ static inline void CompareFeatureVectors(TConstArrayView<const float> A, TConstA
 	VR = ((VA - VB) * VW).square();
 }
 
+static inline bool IsSamplingRangeValid(FFloatInterval Range)
+{
+	return Range.IsValid() && (Range.Min >= 0.0f);
+}
+
 static FFloatInterval GetEffectiveSamplingRange(const UAnimSequenceBase* Sequence, FFloatInterval RequestedSamplingRange)
 {
 	const bool bSampleAll = (RequestedSamplingRange.Min == 0.0f) && (RequestedSamplingRange.Max == 0.0f);
@@ -124,7 +168,6 @@ static FFloatInterval GetEffectiveSamplingRange(const UAnimSequenceBase* Sequenc
 	return Range;
 }
 
-
 /**
 * Algo::LowerBound adapted to TIndexedContainerIterator for use with indexable but not necessarily contiguous containers. Used here with TRingBuffer.
 *
@@ -177,21 +220,6 @@ FORCEINLINE auto LowerBound(IteratorType First, IteratorType Last, const ValueTy
 	return LowerBound(First, Last, Value, FIdentityFunctor(), TLess<>());
 }
 
-static void CalcChannelCosts(const UPoseSearchSchema* Schema, TArrayView<const float> CostVector, TArray<float>& OutChannelCosts)
-{
-	OutChannelCosts.Reset();
-	OutChannelCosts.AddZeroed(Schema->Channels.Num());
-	for (int ChannelIdx = 0; ChannelIdx != Schema->Channels.Num(); ++ChannelIdx)
-	{
-		const UPoseSearchFeatureChannel* Channel = Schema->Channels[ChannelIdx].Get();
-		const int32 ValueTerm = Channel->GetChannelDataOffset() + Channel->GetChannelCardinality();
-		for (int32 ValueIdx = Channel->GetChannelDataOffset(); ValueIdx != ValueTerm; ++ValueIdx)
-		{
-			OutChannelCosts[ChannelIdx] += CostVector[ValueIdx];
-		}
-	}
-}
-
 static int32 PopulateNonSelectableIdx(size_t* NonSelectableIdx, int32 NonSelectableIdxSize, const FSearchContext& SearchContext, const UPoseSearchDatabase* Database)
 {
 	check(Database);
@@ -306,6 +334,13 @@ void UPoseSearchFeatureChannel::InitializeSchema(UE::PoseSearch::FSchemaInitiali
 	ChannelDataOffset = Initializer.GetCurrentChannelDataOffset();
 }
 
+#if WITH_EDITOR
+void UPoseSearchFeatureChannel::ComputeCostBreakdowns(UE::PoseSearch::ICostBreakDownData& CostBreakDownData, const UPoseSearchSchema* Schema) const
+{
+	CostBreakDownData.AddEntireBreakDownSection(FText::FromString(GetName()), Schema, ChannelDataOffset, ChannelCardinality);
+}
+#endif // WITH_EDITOR
+
 // base implementation calculating a single mean deviation value (replicated ChannelCardinality times into MeanDeviations starting at DataOffset index) from all the features data associated to this channel
 void UPoseSearchFeatureChannel::ComputeMeanDeviations(const Eigen::MatrixXd& CenteredPoseMatrix, Eigen::VectorXd& MeanDeviations) const
 {
@@ -525,6 +560,17 @@ void UPoseSearchSchema::PostEditChangeProperty(FPropertyChangedEvent& PropertyCh
 
 	Super::PostEditChangeProperty(PropertyChangedEvent);
 }
+
+void UPoseSearchSchema::ComputeCostBreakdowns(UE::PoseSearch::ICostBreakDownData& CostBreakDownData) const
+{
+	for (const TObjectPtr<UPoseSearchFeatureChannel>& Channel : Channels)
+	{
+		if (!Channel.IsNull())
+		{
+			Channel->ComputeCostBreakdowns(CostBreakDownData, this);
+		}
+	}
+}
 #endif
 
 bool UPoseSearchSchema::IsValid() const
@@ -599,24 +645,6 @@ bool UPoseSearchSchema::BuildQuery(UE::PoseSearch::FSearchContext& SearchContext
 //////////////////////////////////////////////////////////////////////////
 // FPoseSearchIndexAsset
 // 
-int32 FPoseSearchIndex::FindAssetIndex(const FPoseSearchIndexAsset* Asset) const
-{
-	if (Asset == nullptr || Assets.Num() == 0)
-	{
-		return INDEX_NONE;
-	}
-
-	const FPoseSearchIndexAsset* Start = &Assets[0];
-	int32 Result = Asset - Start;
-
-	if (!Assets.IsValidIndex(Result))
-	{
-		return INDEX_NONE;
-	}
-
-	return Result;
-}
-
 const FPoseSearchIndexAsset* FPoseSearchIndex::FindAssetForPose(int32 PoseIdx) const
 {
 	auto Predicate = [PoseIdx](const FPoseSearchIndexAsset& Asset)
@@ -690,8 +718,8 @@ bool FPoseSearchIndex::IsEmpty() const
 
 TConstArrayView<const float> FPoseSearchIndex::GetPoseValues(int32 PoseIdx) const
 {
-	check(PoseIdx < NumPoses);
-	int32 ValueOffset = PoseIdx * Schema->SchemaCardinality;
+	check(PoseIdx >= 0 && PoseIdx < NumPoses && Schema && Schema->SchemaCardinality >= 0);
+	const int32 ValueOffset = PoseIdx * Schema->SchemaCardinality;
 	return MakeArrayView(&Values[ValueOffset], Schema->SchemaCardinality);
 }
 
@@ -835,11 +863,11 @@ FPoseSearchCost UPoseSearchSequenceMetaData::ComparePoses(
 	const float DissimilarityCost = CompareFeatureVectors(SearchIndex.GetPoseValues(PoseIdx), QueryValues);
 
 	// @todo: shouldn't we include MirrorMismatchAddend as well?
-	//const float MirrorMismatchAddend = SearchIndex.ComputeMirrorMismatchAddend(PoseIdx, SearchContext);
+	const float MirrorMismatchAddend = 0.f; // SearchIndex.ComputeMirrorMismatchAddend(PoseIdx, SearchContext);
 	const float NotifyAddend = SearchIndex.ComputeNotifyAddend(PoseIdx);
 	const float ContinuingPoseCostAddend = SearchIndex.ComputeContinuingPoseCostAddend(PoseIdx, PoseComparisonFlags);
 
-	return FPoseSearchCost(DissimilarityCost, NotifyAddend + ContinuingPoseCostAddend);
+	return FPoseSearchCost(DissimilarityCost, NotifyAddend, MirrorMismatchAddend, ContinuingPoseCostAddend);
 }
 
 
@@ -1744,50 +1772,17 @@ bool UPoseSearchDatabase::IsCachedCookedPlatformDataLoaded(const ITargetPlatform
 
 #endif // WITH_EDITOR
 
-FPoseSearchCost UPoseSearchDatabase::ComparePoses(UE::PoseSearch::FSearchContext& SearchContext, int32 PoseIdx, UE::PoseSearch::EPoseComparisonFlags PoseComparisonFlags, const TArrayView<const float>& QueryValues) const
+FPoseSearchCost UPoseSearchDatabase::ComparePoses(UE::PoseSearch::FSearchContext& SearchContext, int32 PoseIdx, UE::PoseSearch::EPoseComparisonFlags PoseComparisonFlags, TConstArrayView<const float> QueryValues) const
 {
-	using namespace Eigen;
-	using namespace UE::PoseSearch;
-
-	FPoseSearchCost Result;
-
 	const FPoseSearchIndex* SearchIndex = GetSearchIndex();
 	check(SearchIndex);
 
-	TConstArrayView<const float> PoseValues = SearchIndex->GetPoseValues(PoseIdx);
-
-	const float DissimilarityCost = CompareFeatureVectors(PoseValues, QueryValues, SearchIndex->WeightsSqrt);
-
+	const float DissimilarityCost = UE::PoseSearch::CompareFeatureVectors(SearchIndex->GetPoseValues(PoseIdx), QueryValues, SearchIndex->WeightsSqrt);
 	const float MirrorMismatchAddend = SearchIndex->ComputeMirrorMismatchAddend(PoseIdx, SearchContext);
 	const float NotifyAddend = SearchIndex->ComputeNotifyAddend(PoseIdx);
 	const float ContinuingPoseCostAddend = SearchIndex->ComputeContinuingPoseCostAddend(PoseIdx, PoseComparisonFlags);
 
-	Result.Set(DissimilarityCost, NotifyAddend + MirrorMismatchAddend + ContinuingPoseCostAddend);
-
-#if UE_POSE_SEARCH_TRACE_ENABLED && WITH_EDITORONLY_DATA
-	if (SearchContext.bIsTracing)
-	{
-		// Setup Eigen views onto our vectors
-		Result.CostVector.SetNum(PoseValues.Num());
-
-		// Compute weighted squared difference vector
-		CompareFeatureVectors(PoseValues, QueryValues, SearchIndex->WeightsSqrt, Result.CostVector);
-
-		// Output cost details
-		CalcChannelCosts(SearchIndex->Schema, Result.CostVector, Result.ChannelCosts);
-		Result.MirrorMismatchCostAddend = MirrorMismatchAddend;
-		Result.NotifyCostAddend = NotifyAddend;
-		Result.ContinuingPoseCostAddend = ContinuingPoseCostAddend;
-
-		// Verify channel cost decomposition for tracing agrees with runtime pose comparator
-		auto ChannelCostsVector = Map<const ArrayXf>(
-			Result.ChannelCosts.GetData(),
-			Result.ChannelCosts.Num());
-		checkSlow(FMath::IsNearlyEqual(ChannelCostsVector.sum(), Result.GetDissimilarity(), 1e-3f));
-	}
-#endif // UE_POSE_SEARCH_TRACE_ENABLED && WITH_EDITORONLY_DATA
-
-	return Result;
+	return FPoseSearchCost(DissimilarityCost, NotifyAddend, MirrorMismatchAddend, ContinuingPoseCostAddend);
 }
 
 UE::PoseSearch::FSearchResult UPoseSearchDatabase::Search(UE::PoseSearch::FSearchContext& SearchContext) const
@@ -2349,7 +2344,7 @@ void FSearchContext::UpdateCurrentBestCost(const FPoseSearchCost& PoseSearchCost
 	};
 }
 
-bool FSearchContext::GetOrBuildQuery(const UPoseSearchDatabase* Database, FPoseSearchFeatureVectorBuilder& FeatureVectorBuilder)
+const FPoseSearchFeatureVectorBuilder* FSearchContext::GetCachedQuery(const UPoseSearchDatabase* Database) const
 {
 	const FSearchContext::FCachedQuery* CachedQuery = CachedQueries.FindByPredicate([Database](const FSearchContext::FCachedQuery& CachedQuery)
 	{
@@ -2358,7 +2353,17 @@ bool FSearchContext::GetOrBuildQuery(const UPoseSearchDatabase* Database, FPoseS
 
 	if (CachedQuery)
 	{
-		FeatureVectorBuilder = CachedQuery->FeatureVectorBuilder;
+		return &CachedQuery->FeatureVectorBuilder;
+	}
+	return nullptr;
+}
+
+bool FSearchContext::GetOrBuildQuery(const UPoseSearchDatabase* Database, FPoseSearchFeatureVectorBuilder& FeatureVectorBuilder)
+{
+	const FPoseSearchFeatureVectorBuilder* CachedFeatureVectorBuilder = GetCachedQuery(Database);
+	if (CachedFeatureVectorBuilder)
+	{
+		FeatureVectorBuilder = *CachedFeatureVectorBuilder;
 		return true;
 	}
 
diff --git a/Engine/Plugins/Experimental/Animation/PoseSearch/Source/Runtime/Private/PoseSearchFeatureChannels.cpp b/Engine/Plugins/Experimental/Animation/PoseSearch/Source/Runtime/Private/PoseSearchFeatureChannels.cpp
--- a/Engine/Plugins/Experimental/Animation/PoseSearch/Source/Runtime/Private/PoseSearchFeatureChannels.cpp
+++ b/Engine/Plugins/Experimental/Animation/PoseSearch/Source/Runtime/Private/PoseSearchFeatureChannels.cpp
@@ -10,6 +10,8 @@
 #include "PoseSearchEigenHelper.h"
 #include "UObject/ObjectSaveContext.h"
 
+#define LOCTEXT_NAMESPACE "PoseSearchDebugger"
+
 namespace UE::PoseSearch {
 
 //////////////////////////////////////////////////////////////////////////
@@ -1176,6 +1178,59 @@ void UPoseSearchFeatureChannel_Pose::DebugDraw(const UE::PoseSearch::FDebugDrawP
 	check(DataOffset == ChannelDataOffset + ChannelCardinality);
 }
 
+#if WITH_EDITOR
+void UPoseSearchFeatureChannel_Pose::ComputeCostBreakdowns(UE::PoseSearch::ICostBreakDownData& CostBreakDownData, const UPoseSearchSchema* Schema) const
+{
+	using namespace UE::PoseSearch;
+
+	CostBreakDownData.AddEntireBreakDownSection(LOCTEXT("ColumnLabelPoseChannelTotal", "Pose Total"), Schema, ChannelDataOffset, ChannelCardinality);
+
+	if (CostBreakDownData.IsVerbose())
+	{
+		int32 DataOffset = ChannelDataOffset;
+
+		for (int32 ChannelBoneIdx = 0; ChannelBoneIdx != SampledBones.Num(); ++ChannelBoneIdx)
+		{
+			const FPoseSearchBone& SampledBone = SampledBones[ChannelBoneIdx];
+			if (EnumHasAnyFlags(SampledBone.Flags, EPoseSearchBoneFlags::Position))
+			{
+				for (int32 SubsampleIdx = 0; SubsampleIdx != SampleTimes.Num(); ++SubsampleIdx)
+				{
+					CostBreakDownData.AddEntireBreakDownSection(FText::Format(LOCTEXT("ColumnLabelPoseChannelPosition", "{0} Pos {1}"), FText::FromName(SampledBone.Reference.BoneName), SampleTimes[SubsampleIdx]), Schema, DataOffset, FFeatureVectorHelper::EncodeVectorCardinality);
+					DataOffset += FFeatureVectorHelper::EncodeVectorCardinality;
+				}
+			}
+			if (EnumHasAnyFlags(SampledBone.Flags, EPoseSearchBoneFlags::Rotation))
+			{
+				for (int32 SubsampleIdx = 0; SubsampleIdx != SampleTimes.Num(); ++SubsampleIdx)
+				{
+					CostBreakDownData.AddEntireBreakDownSection(FText::Format(LOCTEXT("ColumnLabelPoseChannelRotation", "{0} Rot {1}"), FText::FromName(SampledBone.Reference.BoneName), SampleTimes[SubsampleIdx]), Schema, DataOffset, FFeatureVectorHelper::EncodeQuatCardinality);
+					DataOffset += FFeatureVectorHelper::EncodeQuatCardinality;
+				}
+			}
+			if (EnumHasAnyFlags(SampledBone.Flags, EPoseSearchBoneFlags::Velocity))
+			{
+				for (int32 SubsampleIdx = 0; SubsampleIdx != SampleTimes.Num(); ++SubsampleIdx)
+				{
+					CostBreakDownData.AddEntireBreakDownSection(FText::Format(LOCTEXT("ColumnLabelPoseChannelVelocity", "{0} Vel {1}"), FText::FromName(SampledBone.Reference.BoneName), SampleTimes[SubsampleIdx]), Schema, DataOffset, FFeatureVectorHelper::EncodeVectorCardinality);
+					DataOffset += FFeatureVectorHelper::EncodeVectorCardinality;
+				}
+			}
+			if (EnumHasAnyFlags(SampledBone.Flags, EPoseSearchBoneFlags::Phase))
+			{
+				for (int32 SubsampleIdx = 0; SubsampleIdx != SampleTimes.Num(); ++SubsampleIdx)
+				{
+					CostBreakDownData.AddEntireBreakDownSection(FText::Format(LOCTEXT("ColumnLabelPoseChannelVelocity", "{0} Pha {1}"), FText::FromName(SampledBone.Reference.BoneName), SampleTimes[SubsampleIdx]), Schema, DataOffset, FFeatureVectorHelper::EncodeVector2DCardinality);
+					DataOffset += FFeatureVectorHelper::EncodeVector2DCardinality;
+				}
+			}
+		}
+
+		check(DataOffset == ChannelDataOffset + ChannelCardinality);
+	}
+}
+
+#endif // WITH_EDITOR
 
 //////////////////////////////////////////////////////////////////////////
 // UPoseSearchFeatureChannel_Trajectory
@@ -1980,4 +2035,70 @@ void UPoseSearchFeatureChannel_Trajectory::DebugDraw(const UE::PoseSearch::FDebu
 	}
 
 	check(DataOffset == ChannelDataOffset + ChannelCardinality);
-}
\ No newline at end of file
+}
+
+#if WITH_EDITOR
+void UPoseSearchFeatureChannel_Trajectory::ComputeCostBreakdowns(UE::PoseSearch::ICostBreakDownData& CostBreakDownData, const UPoseSearchSchema* Schema) const
+{
+	using namespace UE::PoseSearch;
+
+	CostBreakDownData.AddEntireBreakDownSection(LOCTEXT("ColumnLabelTrajChannelTotal", "Traj Total"), Schema, ChannelDataOffset, ChannelCardinality);
+
+	if (CostBreakDownData.IsVerbose())
+	{
+		int32 DataOffset = ChannelDataOffset;
+
+		for (const FPoseSearchTrajectorySample& Sample : Samples)
+		{
+			if (EnumHasAnyFlags(Sample.Flags, EPoseSearchTrajectoryFlags::Position))
+			{
+				CostBreakDownData.AddEntireBreakDownSection(FText::Format(LOCTEXT("ColumnLabelTrajChannelPosition", "Traj Pos {0}"), Sample.Offset), Schema, DataOffset, FFeatureVectorHelper::EncodeVectorCardinality);
+				DataOffset += FFeatureVectorHelper::EncodeVectorCardinality;
+			}
+			if (EnumHasAnyFlags(Sample.Flags, EPoseSearchTrajectoryFlags::PositionXY))
+			{
+				CostBreakDownData.AddEntireBreakDownSection(FText::Format(LOCTEXT("ColumnLabelTrajChannelPositionXY", "Traj PosXY {0}"), Sample.Offset), Schema, DataOffset, FFeatureVectorHelper::EncodeVector2DCardinality);
+				DataOffset += FFeatureVectorHelper::EncodeVector2DCardinality;
+			}
+
+			if (EnumHasAnyFlags(Sample.Flags, EPoseSearchTrajectoryFlags::Velocity))
+			{
+				CostBreakDownData.AddEntireBreakDownSection(FText::Format(LOCTEXT("ColumnLabelTrajChannelVelocity", "Traj Vel {0}"), Sample.Offset), Schema, DataOffset, FFeatureVectorHelper::EncodeVectorCardinality);
+				DataOffset += FFeatureVectorHelper::EncodeVectorCardinality;
+			}
+			if (EnumHasAnyFlags(Sample.Flags, EPoseSearchTrajectoryFlags::VelocityXY))
+			{
+				CostBreakDownData.AddEntireBreakDownSection(FText::Format(LOCTEXT("ColumnLabelTrajChannelVelocityXY", "Traj VelXY {0}"), Sample.Offset), Schema, DataOffset, FFeatureVectorHelper::EncodeVector2DCardinality);
+				DataOffset += FFeatureVectorHelper::EncodeVector2DCardinality;
+			}
+
+			if (EnumHasAnyFlags(Sample.Flags, EPoseSearchTrajectoryFlags::VelocityDirection))
+			{
+				CostBreakDownData.AddEntireBreakDownSection(FText::Format(LOCTEXT("ColumnLabelTrajChannelVelocityDirection", "Traj VelDir {0}"), Sample.Offset), Schema, DataOffset, FFeatureVectorHelper::EncodeVectorCardinality);
+				DataOffset += FFeatureVectorHelper::EncodeVectorCardinality;
+			}
+			if (EnumHasAnyFlags(Sample.Flags, EPoseSearchTrajectoryFlags::VelocityDirectionXY))
+			{
+				CostBreakDownData.AddEntireBreakDownSection(FText::Format(LOCTEXT("ColumnLabelTrajChannelVelocityDirectionXY", "Traj VelDirXY {0}"), Sample.Offset), Schema, DataOffset, FFeatureVectorHelper::EncodeVector2DCardinality);
+				DataOffset += FFeatureVectorHelper::EncodeVector2DCardinality;
+			}
+
+			if (EnumHasAnyFlags(Sample.Flags, EPoseSearchTrajectoryFlags::FacingDirection))
+			{
+				CostBreakDownData.AddEntireBreakDownSection(FText::Format(LOCTEXT("ColumnLabelTrajChannelFacingDirection", "Traj Fac {0}"), Sample.Offset), Schema, DataOffset, FFeatureVectorHelper::EncodeVectorCardinality);
+				DataOffset += FFeatureVectorHelper::EncodeVectorCardinality;
+			}
+			if (EnumHasAnyFlags(Sample.Flags, EPoseSearchTrajectoryFlags::FacingDirectionXY))
+			{
+				CostBreakDownData.AddEntireBreakDownSection(FText::Format(LOCTEXT("ColumnLabelTrajChannelFacingDirectionXY", "Traj FacXY {0}"), Sample.Offset), Schema, DataOffset, FFeatureVectorHelper::EncodeVector2DCardinality);
+				DataOffset += FFeatureVectorHelper::EncodeVector2DCardinality;
+			}
+		}
+
+		check(DataOffset == ChannelDataOffset + ChannelCardinality);
+	}
+}
+
+#endif // WITH_EDITOR
+
+#undef LOCTEXT_NAMESPACE
\ No newline at end of file
diff --git a/Engine/Plugins/Experimental/Animation/PoseSearch/Source/Runtime/Private/PoseSearchLibrary.cpp b/Engine/Plugins/Experimental/Animation/PoseSearch/Source/Runtime/Private/PoseSearchLibrary.cpp
--- a/Engine/Plugins/Experimental/Animation/PoseSearch/Source/Runtime/Private/PoseSearchLibrary.cpp
+++ b/Engine/Plugins/Experimental/Animation/PoseSearch/Source/Runtime/Private/PoseSearchLibrary.cpp
@@ -187,14 +187,39 @@ static void TraceMotionMatchingState(
 
 	const float DeltaTime = UpdateContext.GetDeltaTime();
 
+	auto AddUniqueDatabase = [](TArray<FTraceMotionMatchingStateDatabaseEntry>& DatabaseEntries, const UPoseSearchDatabase* Database, const UE::PoseSearch::FSearchContext& SearchContext) -> int32
+	{
+		const uint64 DatabaseId = FTraceMotionMatchingState::GetIdFromObject(Database);
+
+		int32 DbEntryIdx = -1;
+		for (int32 i = 0; i < DatabaseEntries.Num(); ++i)
+		{
+			if (DatabaseEntries[i].DatabaseId == DatabaseId)
+			{
+				DbEntryIdx = i;
+				break;
+			}
+		}
+		if (DbEntryIdx == -1)
+		{
+			DbEntryIdx = DatabaseEntries.Add({ DatabaseId });
+
+			const FPoseSearchFeatureVectorBuilder* CachedQuery = SearchContext.GetCachedQuery(Database);
+			check(CachedQuery);
+
+			DatabaseEntries[DbEntryIdx].QueryVector = CachedQuery->GetValues();
+		}
+
+		return DbEntryIdx;
+	};
+
 	FTraceMotionMatchingState TraceState;
 	while (!SearchContext.BestCandidates.IsEmpty())
 	{
 		FSearchContext::FPoseCandidate PoseCandidate;
 		SearchContext.BestCandidates.Pop(PoseCandidate);
 
-		const uint64 DatabaseId = FTraceMotionMatchingState::GetIdFromObject(PoseCandidate.Database);
-		const int32 DbEntryIdx = TraceState.DatabaseEntries.AddUnique({DatabaseId});
+		const int32 DbEntryIdx = AddUniqueDatabase(TraceState.DatabaseEntries, PoseCandidate.Database, SearchContext);
 		FTraceMotionMatchingStateDatabaseEntry& DbEntry = TraceState.DatabaseEntries[DbEntryIdx];
 
 		FTraceMotionMatchingStatePoseEntry PoseEntry;
@@ -207,12 +232,11 @@ static void TraceMotionMatchingState(
 	if (MotionMatchingState.CurrentSearchResult.ContinuingPoseCost.IsValid())
 	{
 		check(LastResult.IsValid());
-		const uint64 LastResultDbId = FTraceMotionMatchingState::GetIdFromObject(LastResult.Database.Get());
 
-		int32 LastResultDbIdx = TraceState.DatabaseEntries.AddUnique({LastResultDbId});
-		FTraceMotionMatchingStateDatabaseEntry& DbEntry = TraceState.DatabaseEntries[LastResultDbIdx];
+		const int32 DbEntryIdx = AddUniqueDatabase(TraceState.DatabaseEntries, LastResult.Database.Get(), SearchContext);
+		FTraceMotionMatchingStateDatabaseEntry& DbEntry = TraceState.DatabaseEntries[DbEntryIdx];
 
-		int32 LastResultPoseEntryIdx = DbEntry.PoseEntries.Add({LastResult.PoseIdx});
+		const int32 LastResultPoseEntryIdx = DbEntry.PoseEntries.Add({LastResult.PoseIdx});
 		FTraceMotionMatchingStatePoseEntry& PoseEntry = DbEntry.PoseEntries[LastResultPoseEntryIdx];
 
 		PoseEntry.Cost = MotionMatchingState.CurrentSearchResult.ContinuingPoseCost;
@@ -221,12 +245,10 @@ static void TraceMotionMatchingState(
 
 	if (MotionMatchingState.CurrentSearchResult.PoseCost.IsValid())
 	{
-		const uint64 CurrentResultDbId = FTraceMotionMatchingState::GetIdFromObject(MotionMatchingState.CurrentSearchResult.Database.Get());
-
-		int32 DbEntryIdx = TraceState.DatabaseEntries.AddUnique({CurrentResultDbId});
+		const int32 DbEntryIdx = AddUniqueDatabase(TraceState.DatabaseEntries, MotionMatchingState.CurrentSearchResult.Database.Get(), SearchContext);
 		FTraceMotionMatchingStateDatabaseEntry& DbEntry = TraceState.DatabaseEntries[DbEntryIdx];
 
-		int32 PoseEntryIdx = DbEntry.PoseEntries.Add({MotionMatchingState.CurrentSearchResult.PoseIdx});
+		const int32 PoseEntryIdx = DbEntry.PoseEntries.Add({MotionMatchingState.CurrentSearchResult.PoseIdx});
 		FTraceMotionMatchingStatePoseEntry& PoseEntry = DbEntry.PoseEntries[PoseEntryIdx];
 
 		PoseEntry.Cost = MotionMatchingState.CurrentSearchResult.PoseCost;
@@ -269,7 +291,6 @@ static void TraceMotionMatchingState(
 
 	TraceState.SearchableAssetId = FTraceMotionMatchingState::GetIdFromObject(Searchable);
 	TraceState.ElapsedPoseJumpTime = MotionMatchingState.ElapsedPoseJumpTime;
-	TraceState.QueryVector = MotionMatchingState.CurrentSearchResult.ComposedQuery.GetValues();
 	TraceState.AssetPlayerTime = MotionMatchingState.CurrentSearchResult.AssetTime;
 	TraceState.DeltaTime = DeltaTime;
 
diff --git a/Engine/Plugins/Experimental/Animation/PoseSearch/Source/Runtime/Private/Trace/PoseSearchTraceLogger.cpp b/Engine/Plugins/Experimental/Animation/PoseSearch/Source/Runtime/Private/Trace/PoseSearchTraceLogger.cpp
--- a/Engine/Plugins/Experimental/Animation/PoseSearch/Source/Runtime/Private/Trace/PoseSearchTraceLogger.cpp
+++ b/Engine/Plugins/Experimental/Animation/PoseSearch/Source/Runtime/Private/Trace/PoseSearchTraceLogger.cpp
@@ -59,6 +59,7 @@ FArchive& operator<<(FArchive& Ar, FTraceMotionMatchingStatePoseEntry& Entry)
 FArchive& operator<<(FArchive& Ar, FTraceMotionMatchingStateDatabaseEntry& Entry)
 {
 	Ar << Entry.DatabaseId;
+	Ar << Entry.QueryVector;
 	Ar << Entry.PoseEntries;
 	return Ar;
 }
@@ -68,7 +69,6 @@ FArchive& operator<<(FArchive& Ar, FTraceMotionMatchingState& State)
 	Ar << State.SearchableAssetId;
 	Ar << State.ElapsedPoseJumpTime;
 	Ar << State.Flags;
-	Ar << State.QueryVector;
 	Ar << State.AssetPlayerTime;
 	Ar << State.DeltaTime;
 	Ar << State.SimLinearVelocity;
diff --git a/Engine/Plugins/Experimental/Animation/PoseSearch/Source/Runtime/Private/Trace/PoseSearchTraceLogger.h b/Engine/Plugins/Experimental/Animation/PoseSearch/Source/Runtime/Private/Trace/PoseSearchTraceLogger.h
--- a/Engine/Plugins/Experimental/Animation/PoseSearch/Source/Runtime/Private/Trace/PoseSearchTraceLogger.h
+++ b/Engine/Plugins/Experimental/Animation/PoseSearch/Source/Runtime/Private/Trace/PoseSearchTraceLogger.h
@@ -58,8 +58,9 @@ struct POSESEARCH_API FTraceMotionMatchingStateDatabaseEntry
 	// @todo: can we use UPoseSearchDatabase* instead of DatabaseId?
 	//UPoseSearchDatabase const* Database = nullptr;
 	uint64 DatabaseId = 0;
+	TArray<float> QueryVector;
 	TArray<FTraceMotionMatchingStatePoseEntry> PoseEntries;
-
+	
 	bool operator==(const FTraceMotionMatchingStateDatabaseEntry& Other) const { return DatabaseId == Other.DatabaseId; }
 };
 POSESEARCH_API FArchive& operator<<(FArchive& Ar, FTraceMotionMatchingStateDatabaseEntry& Entry);
@@ -87,9 +88,6 @@ struct POSESEARCH_API FTraceMotionMatchingState
 	/** Storage container for state booleans */
 	EFlags Flags = EFlags::None;
 
-	/** Search vectors in normalized and unnormalized forms */
-	TArray<float> QueryVector;
-
 	float AssetPlayerTime = 0.0f;
 	float DeltaTime = 0.0f;
 	float SimLinearVelocity = 0.0f;
diff --git a/Engine/Plugins/Experimental/Animation/PoseSearch/Source/Runtime/Public/PoseSearch/PoseSearch.h b/Engine/Plugins/Experimental/Animation/PoseSearch/Source/Runtime/Public/PoseSearch/PoseSearch.h
--- a/Engine/Plugins/Experimental/Animation/PoseSearch/Source/Runtime/Public/PoseSearch/PoseSearch.h
+++ b/Engine/Plugins/Experimental/Animation/PoseSearch/Source/Runtime/Public/PoseSearch/PoseSearch.h
@@ -60,6 +60,18 @@ namespace Eigen {
 	using VectorXd = Matrix<double, -1, 1, 0, -1, 1>;
 }
 
+
+namespace UE::PoseSearch
+{
+	// @todo: move this in an utils file and make them inline
+	POSESEARCH_API float ArraySum(TConstArrayView<const float> View, int32 StartIndex, int32 Offset);
+	POSESEARCH_API void ArrayMinMax(TConstArrayView<const float> View, TArrayView<float> Min, TArrayView<float> Max, float InvalidValue);
+	POSESEARCH_API void ArraySafeNormalize(TConstArrayView<const float> View, TConstArrayView<const float> Min, TConstArrayView<const float> Max, TArrayView<float> NormalizedView);
+	POSESEARCH_API float CompareFeatureVectors(TConstArrayView<const float> A, TConstArrayView<const float> B, TConstArrayView<const float> WeightsSqrt);
+	POSESEARCH_API float CompareFeatureVectors(TConstArrayView<const float> A, TConstArrayView<const float> B);
+	POSESEARCH_API void CompareFeatureVectors(TConstArrayView<const float> A, TConstArrayView<const float> B, TConstArrayView<const float> WeightsSqrt, TArrayView<float> Result);
+}
+
 //////////////////////////////////////////////////////////////////////////
 // Constants
 
@@ -351,6 +363,50 @@ public:
 	virtual FTransform GetTransformAndCacheResults(float SampleTime, float OriginTime, int8 SchemaBoneIdx, bool& Clamped) = 0;
 };
 
+class POSESEARCH_API ICostBreakDownData
+{
+public:
+	virtual ~ICostBreakDownData() {}
+
+	// returns the size of the dataset
+	virtual int32 Num() const = 0;
+
+	// returns true if Index-th cost data vector is associated with Schema
+	virtual bool IsCostVectorFromSchema(int32 Index, const UPoseSearchSchema* Schema) const = 0;
+
+	// returns the Index-th cost data vector
+	virtual TConstArrayView<const float> GetCostVector(int32 Index, const UPoseSearchSchema* Schema) const = 0;
+
+	// every breakdown section start by calling BeginBreakDownSection...
+	virtual void BeginBreakDownSection(const FText& Label) = 0;
+
+	// ...then add as many SetCostBreakDown into the section...
+	virtual void SetCostBreakDown(float CostBreakDown, int32 Index, const UPoseSearchSchema* Schema) = 0;
+
+	// ...to finally wrap the section up by calling EndBreakDownSection
+	virtual void EndBreakDownSection(const FText& Label) = 0;
+
+	// most common case implementation
+	void AddEntireBreakDownSection(const FText& Label, const UPoseSearchSchema* Schema, int32 DataOffset, int32 Cardinality)
+	{
+		BeginBreakDownSection(Label);
+
+		const int32 Count = Num();
+		for (int32 i = 0; i < Count; ++i)
+		{
+			if (IsCostVectorFromSchema(i, Schema))
+			{
+				const float CostBreakdown = ArraySum(GetCostVector(i, Schema), DataOffset, Cardinality);
+				SetCostBreakDown(CostBreakdown, i, Schema);
+			}
+		}
+
+		EndBreakDownSection(Label);
+	}
+
+	virtual bool IsVerbose() const { return true; }
+};
+
 } // namespace UE::PoseSearch
 
 
@@ -387,6 +443,10 @@ public:
 	// Draw this channel's data for the given pose vector
 	virtual void DebugDraw(const UE::PoseSearch::FDebugDrawParams& DrawParams, TArrayView<const float> PoseVector) const PURE_VIRTUAL(UPoseSearchFeatureChannel::DebugDraw, );
 
+#if WITH_EDITOR
+	virtual void ComputeCostBreakdowns(UE::PoseSearch::ICostBreakDownData& CostBreakDownData, const UPoseSearchSchema* Schema) const;
+#endif
+
 	// Used during data normalization. If a feature has less than this amount of deviation from the mean across all poses in a database, then it will not be normalized.
 	UPROPERTY(EditAnywhere, Category = "Settings", meta = (ClampMin = "0.0001"))
 	float MinimumMeanDeviation = 1.0f;
@@ -540,14 +600,16 @@ public:
 
 	int32 GetNumBones () const { return BoneIndices.Num(); }
 
-public: // UObject
+	// UObject
 	virtual void PreSave(FObjectPreSaveContext ObjectSaveContext) override;
 	virtual void PostLoad() override;
 
-public: // IBoneReferenceSkeletonProvider
+	// IBoneReferenceSkeletonProvider
 	class USkeleton* GetSkeleton(bool& bInvalidSkeletonIsError, const IPropertyHandle* PropertyHandle) override { bInvalidSkeletonIsError = false; return Skeleton; }
 
 #if WITH_EDITOR
+	void ComputeCostBreakdowns(UE::PoseSearch::ICostBreakDownData& CostBreakDownData) const;
+
 	virtual void PostEditChangeProperty(struct FPropertyChangedEvent& PropertyChangedEvent) override;
 #endif
 
@@ -701,8 +763,6 @@ struct POSESEARCH_API FPoseSearchIndex
 
 	TConstArrayView<const float> GetPoseValues(int32 PoseIdx) const;
 
-	int32 FindAssetIndex(const FPoseSearchIndexAsset* Asset) const;
-
 	const FPoseSearchIndexAsset* FindAssetForPose(int32 PoseIdx) const;
 	float GetAssetTime(int32 PoseIdx, const FPoseSearchIndexAsset* Asset) const;
 
@@ -811,60 +871,45 @@ struct POSESEARCH_API FPoseSearchCost
 	GENERATED_BODY()
 public:
 	FPoseSearchCost() = default;
-	FPoseSearchCost(float InDissimilarity, float InCostAddend)
-		: Dissimilarity(InDissimilarity)
-		, CostAddend(InCostAddend)
-	{}
-
-	bool IsValid() const { return Dissimilarity != MAX_flt; }
-
-	void Set(float InDissimilarity, float InCostAddend)
+	FPoseSearchCost(float InDissimilarityCost, float InNotifyCostAddend, float InMirrorMismatchCostAddend, float InContinuingPoseCostAddend)
+	: TotalCost(InDissimilarityCost + InNotifyCostAddend + InMirrorMismatchCostAddend + InContinuingPoseCostAddend)
 	{
-		Dissimilarity = InDissimilarity;
-		CostAddend = InCostAddend;
+#if WITH_EDITORONLY_DATA
+		NotifyCostAddend = InNotifyCostAddend;
+		MirrorMismatchCostAddend = InMirrorMismatchCostAddend;
+		ContinuingPoseCostAddend = InContinuingPoseCostAddend;
+#endif // WITH_EDITORONLY_DATA
 	}
 
-	float GetDissimilarity() const { return Dissimilarity; }
-	float GetCostAddend() const { return CostAddend; }
-	float GetTotalCost() const { return Dissimilarity + CostAddend; }
+	bool IsValid() const { return TotalCost != MAX_flt; }
 
-	bool operator<(const FPoseSearchCost& Other) const { return GetTotalCost() < Other.GetTotalCost(); }
+	float GetTotalCost() const { return TotalCost; }
 
-	void Reset()
-	{
-		Dissimilarity = MAX_flt;
-		CostAddend = 0.0f;
-	}
+	bool operator<(const FPoseSearchCost& Other) const { return TotalCost < Other.TotalCost; }
 
 protected:
-
-	UPROPERTY()
-	float Dissimilarity = MAX_flt;
-
 	UPROPERTY()
-	float CostAddend = 0.0f;
+	float TotalCost = MAX_flt;
 
 #if WITH_EDITORONLY_DATA
 public:
+
+	float GetCostAddend() const
+	{
+		return NotifyCostAddend + MirrorMismatchCostAddend + ContinuingPoseCostAddend;
+	}
+
 	// Contribution from ModifyCost anim notify
 	UPROPERTY()
-	float NotifyCostAddend = 0.0f;
+	float NotifyCostAddend = 0.f;
 
 	// Contribution from mirroring cost
 	UPROPERTY()
-	float MirrorMismatchCostAddend = 0.0f;
-
-	UPROPERTY()
-	float ContinuingPoseCostAddend = 0.0f;
+	float MirrorMismatchCostAddend = 0.f;
 
-	// Cost breakdown per channel (e.g. pose cost, time-based trajectory cost, distance-based trajectory cost, etc.)
 	UPROPERTY()
-	TArray<float> ChannelCosts;
+	float ContinuingPoseCostAddend = 0.f;
 
-	// Difference vector computed as W*((P-Q)^2) without the cost modifier applied
-	// Where P is the pose vector, Q is the query vector, W is the weights vector, and multiplication/exponentiation are element-wise operations
-	UPROPERTY()
-	TArray<float> CostVector;
 #endif // WITH_EDITORONLY_DATA
 };
 
@@ -1162,7 +1207,7 @@ public:
 
 	void BuildQuery(UE::PoseSearch::FSearchContext& SearchContext, FPoseSearchFeatureVectorBuilder& OutQuery) const;
 
-	FPoseSearchCost ComparePoses(UE::PoseSearch::FSearchContext& SearchContext, int32 PoseIdx, UE::PoseSearch::EPoseComparisonFlags PoseComparisonFlags, const TArrayView<const float>& QueryValues) const;
+	FPoseSearchCost ComparePoses(UE::PoseSearch::FSearchContext& SearchContext, int32 PoseIdx, UE::PoseSearch::EPoseComparisonFlags PoseComparisonFlags, TConstArrayView<const float> QueryValues) const;
 
 protected:
 	UE::PoseSearch::FSearchResult SearchPCAKDTree(UE::PoseSearch::FSearchContext& SearchContext) const;
@@ -1330,6 +1375,7 @@ struct POSESEARCH_API FSearchContext
 	float GetCurrentBestTotalCost() const { return CurrentBestTotalCost; }
 
 	bool GetOrBuildQuery(const UPoseSearchDatabase* Database, FPoseSearchFeatureVectorBuilder& FeatureVectorBuilder);
+	const FPoseSearchFeatureVectorBuilder* GetCachedQuery(const UPoseSearchDatabase* Database) const;
 
 	bool IsCurrentResultFromDatabase(const UPoseSearchDatabase* Database) const;
 
diff --git a/Engine/Plugins/Experimental/Animation/PoseSearch/Source/Runtime/Public/PoseSearch/PoseSearchFeatureChannels.h b/Engine/Plugins/Experimental/Animation/PoseSearch/Source/Runtime/Public/PoseSearch/PoseSearchFeatureChannels.h
--- a/Engine/Plugins/Experimental/Animation/PoseSearch/Source/Runtime/Public/PoseSearch/PoseSearchFeatureChannels.h
+++ b/Engine/Plugins/Experimental/Animation/PoseSearch/Source/Runtime/Public/PoseSearch/PoseSearchFeatureChannels.h
@@ -153,6 +153,10 @@ public:
 	virtual bool BuildQuery(UE::PoseSearch::FSearchContext& SearchContext, FPoseSearchFeatureVectorBuilder& InOutQuery) const override;
 	virtual void DebugDraw(const UE::PoseSearch::FDebugDrawParams& DrawParams, TArrayView<const float> PoseVector) const override;
 
+#if WITH_EDITOR
+	virtual void ComputeCostBreakdowns(UE::PoseSearch::ICostBreakDownData& CostBreakDownData, const UPoseSearchSchema* Schema) const override;
+#endif
+
 protected:
 	void AddPoseFeatures(UE::PoseSearch::IAssetIndexer& Indexer, int32 SampleIdx, FPoseSearchFeatureVectorBuilder& FeatureVector, const TArray<TArray<FVector2D>>& Phases) const;
 	void CalculatePhases(UE::PoseSearch::IAssetIndexer& Indexer, UE::PoseSearch::FAssetIndexingOutput& IndexingOutput, TArray<TArray<FVector2D>>& OutPhases) const;
@@ -236,6 +240,10 @@ public:
 	virtual bool BuildQuery(UE::PoseSearch::FSearchContext& SearchContext, FPoseSearchFeatureVectorBuilder& InOutQuery) const override;
 	virtual void DebugDraw(const UE::PoseSearch::FDebugDrawParams& DrawParams, TArrayView<const float> PoseVector) const override;
 
+#if WITH_EDITOR
+	virtual void ComputeCostBreakdowns(UE::PoseSearch::ICostBreakDownData& CostBreakDownData, const UPoseSearchSchema* Schema) const override;
+#endif
+
 protected:
 	void IndexAssetPrivate(const UE::PoseSearch::IAssetIndexer& Indexer, int32 SampleIdx, FPoseSearchFeatureVectorBuilder& FeatureVector) const;
 	float GetSampleTime(const UE::PoseSearch::IAssetIndexer& Indexer, float Offset, float SampleTime, float RootDistance) const;
