motion matching - fix for Text conflict from LOCTEXT macro for namespace "PoseSearchDebugger"

diff --git a/Engine/Plugins/Experimental/Animation/PoseSearch/Source/Runtime/Private/PoseSearchFeatureChannels.cpp b/Engine/Plugins/Experimental/Animation/PoseSearch/Source/Runtime/Private/PoseSearchFeatureChannels.cpp
--- a/Engine/Plugins/Experimental/Animation/PoseSearch/Source/Runtime/Private/PoseSearchFeatureChannels.cpp
+++ b/Engine/Plugins/Experimental/Animation/PoseSearch/Source/Runtime/Private/PoseSearchFeatureChannels.cpp
@@ -10,7 +10,7 @@
 #include "PoseSearchEigenHelper.h"
 #include "UObject/ObjectSaveContext.h"
 
-#define LOCTEXT_NAMESPACE "PoseSearchDebugger"
+#define LOCTEXT_NAMESPACE "PoseSearchFeatureChannels"
 
 namespace UE::PoseSearch {
 
