- adding mpose search debugger multi selection
- trajectory channel debug lines improvements using DrawCentripetalCatmullRomSpline
- removing channels GradientPercentage
- removed no longer necessary OutSearchIndex.Assets.Sort fix
- stability improvements

diff --git a/Engine/Plugins/Experimental/Animation/PoseSearch/Source/Editor/Private/PoseSearchDebugger.cpp b/Engine/Plugins/Experimental/Animation/PoseSearch/Source/Editor/Private/PoseSearchDebugger.cpp
--- a/Engine/Plugins/Experimental/Animation/PoseSearch/Source/Editor/Private/PoseSearchDebugger.cpp
+++ b/Engine/Plugins/Experimental/Animation/PoseSearch/Source/Editor/Private/PoseSearchDebugger.cpp
@@ -818,9 +818,13 @@ void SDebuggerDatabaseView::Update(const FTraceMotionMatchingStateMessage& State
 					Row->bMirrored = SearchIndexAsset->bMirrored;
 					Row->PoseCost = PoseEntry.Cost;
 
-					Row->CostVector.SetNum(DbEntry.QueryVector.Num());
+					Row->CostVector.SetNum(Database->Schema->SchemaCardinality);
 					TConstArrayView<const float> PoseValues = SearchIndex->GetPoseValues(PoseEntry.DbPoseIdx);
-					CompareFeatureVectors(PoseValues, DbEntry.QueryVector, SearchIndex->WeightsSqrt, Row->CostVector);
+
+					// in case we modify the schema while PIE is paused and displaying the Pose Search Editor, we could end up with a stale State with a DbEntry.QueryVector saved with the previous schema
+					// so the cardinality of DbEntry.QueryVector and PoseValues don't match. In that case we just use PoseValues as query to have all costs set to zero
+					const bool bIsQueryVectorValid = DbEntry.QueryVector.Num() == PoseValues.Num();
+					CompareFeatureVectors(PoseValues, bIsQueryVectorValid ? DbEntry.QueryVector : PoseValues, SearchIndex->WeightsSqrt, Row->CostVector);
 
 					if (SearchIndexAsset->Type == ESearchIndexAssetType::Sequence)
 					{
@@ -1185,7 +1189,7 @@ void SDebuggerDatabaseView::Construct(const FArguments& InArgs)
 		.HeaderRow(FilteredDatabaseView.HeaderRow.ToSharedRef())
 		.OnGenerateRow(this, &SDebuggerDatabaseView::HandleGenerateDatabaseRow)
 		.ExternalScrollbar(FilteredDatabaseView.ScrollBar)
-		.SelectionMode(ESelectionMode::SingleToggle)
+		.SelectionMode(ESelectionMode::Multi)
 		.ConsumeMouseWheel(EConsumeMouseWheel::WhenScrollingPossible)
 		.OnSelectionChanged(this, &SDebuggerDatabaseView::OnDatabaseRowSelectionChanged);
 
@@ -1453,7 +1457,6 @@ void SDebuggerDetailsView::UpdateReflection(const FTraceMotionMatchingStateMessa
 	Reflection->ElapsedPoseJumpTime = State.ElapsedPoseJumpTime;
 	Reflection->bFollowUpAnimation = EnumHasAnyFlags(State.Flags, FTraceMotionMatchingState::EFlags::FollowupAnimation);
 
-
 	Reflection->AssetPlayerAssetName = "None";
 	if (CurrentSearchIndex)
 	{
@@ -1463,7 +1466,6 @@ void SDebuggerDetailsView::UpdateReflection(const FTraceMotionMatchingStateMessa
 		}
 	}
 
-
 	Reflection->AssetPlayerTime = State.AssetPlayerTime;
 	Reflection->LastDeltaTime = State.DeltaTime;
 	Reflection->SimLinearVelocity = State.SimLinearVelocity;
@@ -1471,9 +1473,21 @@ void SDebuggerDetailsView::UpdateReflection(const FTraceMotionMatchingStateMessa
 	Reflection->AnimLinearVelocity = State.AnimLinearVelocity;
 	Reflection->AnimAngularVelocity = State.AnimAngularVelocity;
 	 
-	// @todo: fix me
 	// Query pose
-	//Reflection->QueryPoseVector = State.QueryVector;
+	Reflection->QueryPoseVector.Reset();
+	for (const FTraceMotionMatchingStateDatabaseEntry& DbEntry : State.DatabaseEntries)
+	{
+		const UPoseSearchDatabase* Database = FTraceMotionMatchingState::GetObjectFromId<UPoseSearchDatabase>(DbEntry.DatabaseId);
+		if (Database == CurrentDatabase)
+		{
+			const FPoseSearchIndex* SearchIndex = Database->GetSearchIndex();
+			if (SearchIndex && SearchIndex->IsValid())
+			{
+				Reflection->QueryPoseVector = DbEntry.QueryVector;
+				break;
+			}
+		}
+	}
 
 	// Active pose
 	if (CurrentSearchIndex)
@@ -1498,14 +1512,6 @@ void SDebuggerDetailsView::UpdateReflection(const FTraceMotionMatchingStateMessa
 			}
 
 			Reflection->CostVector = Selected->CostVector;
-
-			// @todo: fix me
-			//const FTraceMotionMatchingStatePoseEntry* ActivePoseEntry = State.GetCurrentPoseEntry();
-			//Reflection->CostVectorDifference = Reflection->CostVector;
-			//for (int i = 0; i < Reflection->CostVectorDifference.Num(); ++i)
-			//{
-			//	Reflection->CostVectorDifference[i] -= ActivePoseEntry->Cost.CostVector[i];
-			//}
 		}
 	}
 }
@@ -1720,13 +1726,6 @@ void SDebuggerView::DrawFeatures(
 	const USkinnedMeshComponent* Mesh
 ) const
 {
-	// Set shared state
-	FDebugDrawParams DrawParams;
-	DrawParams.World = &DebuggerWorld;
-	DrawParams.RootTransform = Transform;
-	DrawParams.DefaultLifeTime = 0.0f; // Single frame render
-	DrawParams.Mesh = Mesh;
-
 	auto SetDrawFlags = [](FDebugDrawParams& InDrawParams, const FPoseSearchDebuggerFeatureDrawOptions& Options)
 	{
 		InDrawParams.Flags = EDebugDrawFlags::None;
@@ -1741,33 +1740,53 @@ void SDebuggerView::DrawFeatures(
 			{
 				EnumAddFlags(InDrawParams.Flags, EDebugDrawFlags::DrawSampleLabels);
 			}
-
-			if (Options.bDrawSamplesWithColorGradient)
-			{
-				EnumAddFlags(InDrawParams.Flags, EDebugDrawFlags::DrawSamplesWithColorGradient);
-			}
 		}
 	};
 
 	const TObjectPtr<UPoseSearchDebuggerReflection> Reflection = DetailsView->GetReflection();
 
-	// @todo: fix me
 	// Draw query vector
 	{
-		//DrawParams.Database = ViewModel.Get()->GetCurrentDatabase();
-		//SetDrawFlags(DrawParams, Reflection ? Reflection->QueryDrawOptions : FPoseSearchDebuggerFeatureDrawOptions());
-		//EnumAddFlags(DrawParams.Flags, EDebugDrawFlags::DrawQuery);
-		//DrawFeatureVector(DrawParams, State.QueryVector);
-		//EnumRemoveFlags(DrawParams.Flags, EDebugDrawFlags::DrawQuery);
+		const UPoseSearchDatabase* CurrentDatabase = ViewModel.Get()->GetCurrentDatabase();
+		if (CurrentDatabase)
+		{
+			for (const FTraceMotionMatchingStateDatabaseEntry& DbEntry : State.DatabaseEntries)
+			{
+				const UPoseSearchDatabase* Database = FTraceMotionMatchingState::GetObjectFromId<UPoseSearchDatabase>(DbEntry.DatabaseId);
+				if (Database == CurrentDatabase)
+				{
+					const FPoseSearchIndex* SearchIndex = Database->GetSearchIndex();
+					if (SearchIndex && SearchIndex->IsValid() && DbEntry.QueryVector.Num() == Database->Schema->SchemaCardinality)
+					{
+						// Set shared state
+						FDebugDrawParams DrawParams;
+						DrawParams.World = &DebuggerWorld;
+						DrawParams.RootTransform = Transform;
+						DrawParams.DefaultLifeTime = 0.0f; // Single frame render
+						DrawParams.Mesh = Mesh;
+						DrawParams.Database = CurrentDatabase;
+						SetDrawFlags(DrawParams, Reflection ? Reflection->QueryDrawOptions : FPoseSearchDebuggerFeatureDrawOptions());
+						EnumAddFlags(DrawParams.Flags, EDebugDrawFlags::DrawQuery);
+						DrawFeatureVector(DrawParams, DbEntry.QueryVector);
+						EnumRemoveFlags(DrawParams.Flags, EDebugDrawFlags::DrawQuery);
+						break;
+					}
+				}
+			}
+		}
 	}
 
-
 	// Draw selected poses
 	{
 		const TSharedPtr<SListView<TSharedRef<FDebuggerDatabaseRowData>>>& DatabaseRows = DatabaseView->GetDatabaseRows();
 		TArray<TSharedRef<FDebuggerDatabaseRowData>> SelectedRows = DatabaseRows->GetSelectedItems();
 	
 		// Red for non-active database view
+		FDebugDrawParams DrawParams;
+		DrawParams.World = &DebuggerWorld;
+		DrawParams.RootTransform = Transform;
+		DrawParams.DefaultLifeTime = 0.0f; // Single frame render
+		DrawParams.Mesh = Mesh;
 		SetDrawFlags(DrawParams, Reflection ? Reflection->SelectedPoseDrawOptions : FPoseSearchDebuggerFeatureDrawOptions());
 
 		// Draw any selected database vectors
@@ -1789,6 +1808,11 @@ void SDebuggerView::DrawFeatures(
 		if (!ActiveRows.IsEmpty())
 		{		
 			// Use the motion-matching state's pose idx, as the active row may be update-throttled at this point
+			FDebugDrawParams DrawParams;
+			DrawParams.World = &DebuggerWorld;
+			DrawParams.RootTransform = Transform;
+			DrawParams.DefaultLifeTime = 0.0f; // Single frame render
+			DrawParams.Mesh = Mesh;
 			DrawParams.Database = ActiveRows[0]->SourceDatabase.Get();
 			DrawFeatureVector(DrawParams, ActiveRows[0]->PoseIdx);
 		}
@@ -1804,6 +1828,11 @@ void SDebuggerView::DrawFeatures(
 
 		if (!ContinuingRows.IsEmpty())
 		{
+			FDebugDrawParams DrawParams;
+			DrawParams.World = &DebuggerWorld;
+			DrawParams.RootTransform = Transform;
+			DrawParams.DefaultLifeTime = 0.0f; // Single frame render
+			DrawParams.Mesh = Mesh;
 			DrawParams.Database = ContinuingRows[0]->SourceDatabase.Get();
 			DrawFeatureVector(DrawParams, ContinuingRows[0]->PoseIdx);
 		}
@@ -2174,7 +2203,7 @@ const UPoseSearchSearchableAsset* FDebuggerViewModel::GetSearchableAsset() const
 
 void FDebuggerViewModel::ShowSelectedSkeleton(const UPoseSearchDatabase* Database, int32 DbPoseIdx, float Time)
 {
-	UPoseSearchMeshComponent* Component = Skeletons[SelectedPose].Component;
+	UPoseSearchMeshComponent* Component = Skeletons[SelectedPose].Component.Get();
 	if (!Component)
 	{
 		return;
@@ -2237,7 +2266,7 @@ void FDebuggerViewModel::OnUpdate()
 			Skeleton.Actor = World->SpawnActor<AActor>(ActorSpawnParameters);
 			Skeleton.Actor->SetActorLabel(TEXT("PoseSearch"));
 			Skeleton.Component = NewObject<UPoseSearchMeshComponent>(Skeleton.Actor.Get());
-			Skeleton.Actor->AddInstanceComponent(Skeleton.Component);
+			Skeleton.Actor->AddInstanceComponent(Skeleton.Component.Get());
 			Skeleton.Component->RegisterComponentWithWorld(World);
 		}
 		FWorldDelegates::OnWorldCleanup.AddRaw(this, &FDebuggerViewModel::OnWorldCleanup);
@@ -2313,10 +2342,10 @@ void FDebuggerViewModel::OnDraw(FSkeletonDrawParams& DrawParams)
 		}
 	};
 	const bool bDrawActivePose = EnumHasAnyFlags(DrawParams.Flags, ESkeletonDrawFlags::ActivePose);
-	SetDrawSkeleton(Skeletons[ActivePose].Component, bDrawActivePose);
+	SetDrawSkeleton(Skeletons[ActivePose].Component.Get(), bDrawActivePose);
 	// If flag is set and we are currently in a valid drawing state
 	const bool bDrawSelectedPose = EnumHasAnyFlags(DrawParams.Flags, ESkeletonDrawFlags::SelectedPose) && bSelecting;
-	SetDrawSkeleton(Skeletons[SelectedPose].Component, bDrawSelectedPose);
+	SetDrawSkeleton(Skeletons[SelectedPose].Component.Get(), bDrawSelectedPose);
 
 	FillCompactPoseAndComponentRefRotations();
 
@@ -2328,7 +2357,7 @@ void FDebuggerViewModel::OnDraw(FSkeletonDrawParams& DrawParams)
 
 	if (bDrawSelectedPose)
 	{
-		UPoseSearchMeshComponent* Component = Skeletons[SelectedPose].Component;
+		UPoseSearchMeshComponent* Component = Skeletons[SelectedPose].Component.Get();
 		if (Component && Component->RequiredBones.IsValid())
 		{
 			if (Skeletons[SelectedPose].Type == ESearchIndexAssetType::Sequence)
@@ -2373,7 +2402,7 @@ void FDebuggerViewModel::OnDraw(FSkeletonDrawParams& DrawParams)
 	const bool bDrawAsset = EnumHasAnyFlags(DrawParams.Flags, ESkeletonDrawFlags::Asset);
 	if (bDrawAsset && AssetData.bActive)
 	{
-		UPoseSearchMeshComponent* Component = Skeletons[Asset].Component;
+		UPoseSearchMeshComponent* Component = Skeletons[Asset].Component.Get();
 		if (Component && Component->RequiredBones.IsValid())
 		{
 			SetDrawSkeleton(Component, true);
@@ -2472,9 +2501,9 @@ void FDebuggerViewModel::UpdateFromTimeline()
 
 				return TraceServices::EEventEnumerate::Stop;
 			}
-			UPoseSearchMeshComponent* ActiveComponent = Skeletons[ActivePose].Component;
-			UPoseSearchMeshComponent* SelectedComponent = Skeletons[SelectedPose].Component;
-			UPoseSearchMeshComponent* AssetComponent = Skeletons[Asset].Component;
+			UPoseSearchMeshComponent* ActiveComponent = Skeletons[ActivePose].Component.Get();
+			UPoseSearchMeshComponent* SelectedComponent = Skeletons[SelectedPose].Component.Get();
+			UPoseSearchMeshComponent* AssetComponent = Skeletons[Asset].Component.Get();
 			USkeletalMesh* SkeletalMesh = TSoftObjectPtr<USkeletalMesh>(FSoftObjectPath(SkeletalMeshObjectInfo->PathName)).LoadSynchronous();
 			if (SkeletalMesh)
 			{
@@ -2510,7 +2539,7 @@ void FDebuggerViewModel::UpdateAsset()
 	}
 
 	FSkeleton& AssetSkeleton = Skeletons[Asset];
-	UPoseSearchMeshComponent* Component = AssetSkeleton.Component;
+	UPoseSearchMeshComponent* Component = AssetSkeleton.Component.Get();
 
 	auto RestartAsset = [&]()
 	{
@@ -2597,7 +2626,7 @@ const USkinnedMeshComponent* FDebuggerViewModel::GetMeshComponent() const
 {
 	if (Skeletons.Num() > FDebuggerViewModel::Asset)
 	{
-		return Skeletons[FDebuggerViewModel::Asset].Component;
+		return Skeletons[FDebuggerViewModel::Asset].Component.Get();
 	}
 	return nullptr;
 }
@@ -2633,7 +2662,7 @@ void FDebuggerViewModel::PlaySelection(int32 PoseIdx, float Time)
 	{
 		return;
 	}
-	UPoseSearchMeshComponent* Component = Skeletons[Asset].Component;
+	UPoseSearchMeshComponent* Component = Skeletons[Asset].Component.Get();
 	if (!Component)
 	{
 		return;
@@ -2658,7 +2687,7 @@ void FDebuggerViewModel::PlaySelection(int32 PoseIdx, float Time)
 }
 void FDebuggerViewModel::StopSelection()
 {
-	UPoseSearchMeshComponent* Component = Skeletons[Asset].Component;
+	UPoseSearchMeshComponent* Component = Skeletons[Asset].Component.Get();
 	if (!Component)
 	{
 		return;
diff --git a/Engine/Plugins/Experimental/Animation/PoseSearch/Source/Editor/Private/PoseSearchDebugger.h b/Engine/Plugins/Experimental/Animation/PoseSearch/Source/Editor/Private/PoseSearchDebugger.h
--- a/Engine/Plugins/Experimental/Animation/PoseSearch/Source/Editor/Private/PoseSearchDebugger.h
+++ b/Engine/Plugins/Experimental/Animation/PoseSearch/Source/Editor/Private/PoseSearchDebugger.h
@@ -85,9 +85,6 @@ struct FPoseSearchDebuggerFeatureDrawOptions
 
 	UPROPERTY(EditAnywhere, Category = "Draw Options", Meta = (EditCondition = "!bDisable"))
 	bool bDrawSampleLabels = false;
-
-	UPROPERTY(EditAnywhere, Category = "Draw Options", Meta = (EditCondition = "!bDisable"))
-	bool bDrawSamplesWithColorGradient = false;
 };
 
 /**
@@ -155,11 +152,6 @@ public:
 
 	UPROPERTY(VisibleAnywhere, Category="Pose Vectors")
 	TArray<float> CostVector;
-
-	// Cost comparison of selected and active poses. A negative value indicates the cost in the selected pose is lower;
-	// a positive value indicates the cost in the selected pose is higher.
-	UPROPERTY(VisibleAnywhere, Category="Pose Vectors")
-	TArray<float> CostVectorDifference;
 };
 
 
@@ -558,7 +550,7 @@ private:
 		TWeakObjectPtr<AActor> Actor = nullptr;
 
 		/** Derived skeletal mesh for setting the skeleton in the scene */
-		UPoseSearchMeshComponent* Component = nullptr;
+		TWeakObjectPtr<UPoseSearchMeshComponent> Component = nullptr;
 		
 		/** Type of the asset being played */
 		ESearchIndexAssetType Type = ESearchIndexAssetType::Invalid;
@@ -637,8 +629,8 @@ public:
 	virtual void Update(float DeltaTime, IRewindDebugger* InRewindDebugger) override;
 	virtual ~FDebugger() = default;
 
-	virtual void RecordingStarted(IRewindDebugger* RewindDebugger);
-	virtual void RecordingStopped(IRewindDebugger* RewindDebugger);
+	virtual void RecordingStarted(IRewindDebugger* RewindDebugger) override;
+	virtual void RecordingStopped(IRewindDebugger* RewindDebugger) override;
 
 	static FDebugger* Get() { return Debugger; }
 	static void Initialize();
diff --git a/Engine/Plugins/Experimental/Animation/PoseSearch/Source/Runtime/Private/PoseSearch.cpp b/Engine/Plugins/Experimental/Animation/PoseSearch/Source/Runtime/Private/PoseSearch.cpp
--- a/Engine/Plugins/Experimental/Animation/PoseSearch/Source/Runtime/Private/PoseSearch.cpp
+++ b/Engine/Plugins/Experimental/Animation/PoseSearch/Source/Runtime/Private/PoseSearch.cpp
@@ -56,11 +56,13 @@ DECLARE_CYCLE_STAT_EXTERN(TEXT("Search PCA/KNN"), STAT_PoseSearchPCAKNN, STATGRO
 DEFINE_STAT(STAT_PoseSearchBruteForce);
 DEFINE_STAT(STAT_PoseSearchPCAKNN);
 
-namespace UE { namespace PoseSearch {
+namespace UE::PoseSearch
+{
 
 //////////////////////////////////////////////////////////////////////////
 // Constants and utilities
 
+// @todo: set UE_POSE_SEARCH_FORCE_SINGLE_THREAD to 0
 // Temporarily disable parallel indexing due to a bug
 // A single BoneContainer is being used while indexing assets, but some
 // const BoneContainer methods change mutable lookup tables.
@@ -323,7 +325,7 @@ static int32 PopulateNonSelectableIdx(size_t* NonSelectableIdx, int32 NonSelecta
 	return NonSelectableIdxUsedSize;
 }
 
-}} // namespace UE::PoseSearch
+} // namespace UE::PoseSearch
 
 
 //////////////////////////////////////////////////////////////////////////
@@ -1508,16 +1510,6 @@ bool UPoseSearchDatabase::TryInitSearchIndexAssets(FPoseSearchIndex& OutSearchIn
 		}
 	}
 
-	/////////////////////////////////////////////////////////////////////////////////////////////////////////////////////////////////////////////////////////////////////////
-	// @todo: remove this code! A licensee noticed a degradation in their MM setup caused by "removing deprecated db groups" with CL 21928385,
-	// ultimately caused by this missing scrambling of assets (the sort actually move the first item in the list...), so we'll reintroduce it until we have a better solution
-	// that most likely involves handling duplicate datapoints representing different animation poses
-	OutSearchIndex.Assets.Sort([](const FPoseSearchIndexAsset& InOne, const FPoseSearchIndexAsset& InTwo)
-	{
-		return false;
-	});
-	/////////////////////////////////////////////////////////////////////////////////////////////////////////////////////////////////////////////////////////////////////////
-
 	if (bAnyMirrored && !Schema->MirrorDataTable)
 	{
 		UE_LOG(
@@ -2233,7 +2225,7 @@ bool FPoseSearchFeatureVectorBuilder::IsCompatible(const FPoseSearchFeatureVecto
 	return IsInitialized() && (Schema == OtherBuilder.Schema);
 }
 
-namespace UE { namespace PoseSearch
+namespace UE::PoseSearch
 {
 
 void FPoseIndicesHistory::Update(const FSearchResult& SearchResult, float DeltaTime, float MaxTime)
@@ -2774,7 +2766,7 @@ bool FDebugDrawParams::CanDraw() const
 	return SearchIndex->IsValid() && !SearchIndex->IsEmpty();
 }
 
-FColor FDebugDrawParams::GetColor(int32 ColorPreset, float GradientPercentage) const
+FColor FDebugDrawParams::GetColor(int32 ColorPreset) const
 {
 	FLinearColor Color = FLinearColor::Red;
 
@@ -2806,12 +2798,6 @@ FColor FDebugDrawParams::GetColor(int32 ColorPreset, float GradientPercentage) c
 		}
 	}
 
-	if (EnumHasAnyFlags(Flags, EDebugDrawFlags::DrawSamplesWithColorGradient))
-	{
-		constexpr float DrawDebugGradientStrength = 0.8f;
-		Color = Color * (1.0f - DrawDebugGradientStrength * GradientPercentage);
-	}
-
 	return Color.ToFColor(true);
 }
 
@@ -2967,6 +2953,25 @@ FTransform FAssetSamplingContext::MirrorTransform(const FTransform& InTransform)
 	return Result;
 }
 
+//////////////////////////////////////////////////////////////////////////
+// class ICostBreakDownData
+
+void ICostBreakDownData::AddEntireBreakDownSection(const FText& Label, const UPoseSearchSchema* Schema, int32 DataOffset, int32 Cardinality)
+{
+	BeginBreakDownSection(Label);
+
+	const int32 Count = Num();
+	for (int32 i = 0; i < Count; ++i)
+	{
+		if (IsCostVectorFromSchema(i, Schema))
+		{
+			const float CostBreakdown = ArraySum(GetCostVector(i, Schema), DataOffset, Cardinality);
+			SetCostBreakDown(CostBreakdown, i, Schema);
+		}
+	}
+
+	EndBreakDownSection(Label);
+}
 
 //////////////////////////////////////////////////////////////////////////
 // Root motion extrapolation
@@ -5108,7 +5113,7 @@ void FModule::OnObjectSaved(UObject* SavedObject, FObjectPreSaveContext SaveCont
 }
 #endif // WITH_EDITOR
 
-}} // namespace UE::PoseSearch
+} // namespace UE::PoseSearch
 
 #undef LOCTEXT_NAMESPACE
 
diff --git a/Engine/Plugins/Experimental/Animation/PoseSearch/Source/Runtime/Private/PoseSearchFeatureChannels.cpp b/Engine/Plugins/Experimental/Animation/PoseSearch/Source/Runtime/Private/PoseSearchFeatureChannels.cpp
--- a/Engine/Plugins/Experimental/Animation/PoseSearch/Source/Runtime/Private/PoseSearchFeatureChannels.cpp
+++ b/Engine/Plugins/Experimental/Animation/PoseSearch/Source/Runtime/Private/PoseSearchFeatureChannels.cpp
@@ -1076,8 +1076,7 @@ void UPoseSearchFeatureChannel_Pose::DebugDraw(const UE::PoseSearch::FDebugDrawP
 			{
 				BonePos[SubsampleIdx] = FFeatureVectorHelper::DecodeVector(PoseVector, DataOffset);
 
-				const float GradientPercentage = NumSubsamples > 1 ? float(SubsampleIdx) / float(NumSubsamples - 1) : 0.f;
-				const FColor Color = DrawParams.GetColor(SampledBone.ColorPresetIndex, GradientPercentage);
+				const FColor Color = DrawParams.GetColor(SampledBone.ColorPresetIndex);
 
 				BonePos[SubsampleIdx] = DrawParams.RootTransform.TransformPosition(BonePos[SubsampleIdx]);
 				if (EnumHasAnyFlags(DrawParams.Flags, EDebugDrawFlags::DrawFast | EDebugDrawFlags::DrawSearchIndex))
@@ -1155,8 +1154,7 @@ void UPoseSearchFeatureChannel_Pose::DebugDraw(const UE::PoseSearch::FDebugDrawP
 			{
 				const FVector2D Phase = FFeatureVectorHelper::DecodeVector2D(PoseVector, DataOffset);
 
-				const float GradientPercentage = NumSubsamples > 1 ? float(SubsampleIdx) / float(NumSubsamples - 1) : 0.f;
-				const FColor Color = DrawParams.GetColor(SampledBone.ColorPresetIndex, GradientPercentage);
+				const FColor Color = DrawParams.GetColor(SampledBone.ColorPresetIndex);
 
 				static float ScaleFactor = 1.f;
 
@@ -1773,6 +1771,8 @@ void UPoseSearchFeatureChannel_Trajectory::DebugDraw(const UE::PoseSearch::FDebu
 	int32 DataOffset = ChannelDataOffset;
 	int32 SampleIdx = 0;
 	TrajectoryPositionReconstructor TrajectoryPositionReconstructor;
+	TArray<FVector> TrajSplinePos;
+	TArray<FColor> TrajSplineColor;
 	for (const FPoseSearchTrajectorySample& Sample : Samples)
 	{
 		bool bIsTrajectoryPosValid = false;
@@ -1782,13 +1782,13 @@ void UPoseSearchFeatureChannel_Trajectory::DebugDraw(const UE::PoseSearch::FDebu
 		{
 			TrajectoryPos = FFeatureVectorHelper::DecodeVector(PoseVector, DataOffset);
 			TrajectoryPos = DrawParams.RootTransform.TransformPosition(TrajectoryPos);
+			
 			bIsTrajectoryPosValid = true;
 
 			// validating TrajectoryPositionReconstructor
 			check((TrajectoryPositionReconstructor.GetReconstructedTrajectoryPos(*this, PoseVector, DrawParams.RootTransform, Sample.Offset) - TrajectoryPos).IsNearlyZero());
 
-			const float GradientPercentage = NumSamples > 1 ? float(SampleIdx) / float(NumSamples - 1) : 0.f;
-			const FColor Color = DrawParams.GetColor(Sample.ColorPresetIndex, GradientPercentage);
+			const FColor Color = DrawParams.GetColor(Sample.ColorPresetIndex);
 
 			if (EnumHasAnyFlags(DrawParams.Flags, EDebugDrawFlags::DrawFast | EDebugDrawFlags::DrawSearchIndex))
 			{
@@ -1798,6 +1798,9 @@ void UPoseSearchFeatureChannel_Trajectory::DebugDraw(const UE::PoseSearch::FDebu
 			{
 				DrawDebugSphere(DrawParams.World, TrajectoryPos, DrawDebugSphereSize, DrawDebugSphereSegments, Color, bPersistent, LifeTime, DepthPriority);
 			}
+
+			TrajSplinePos.Add(TrajectoryPos);
+			TrajSplineColor.Add(Color);
 		}
 		if (EnumHasAnyFlags(Sample.Flags, EPoseSearchTrajectoryFlags::PositionXY))
 		{
@@ -1808,8 +1811,7 @@ void UPoseSearchFeatureChannel_Trajectory::DebugDraw(const UE::PoseSearch::FDebu
 				TrajectoryPos = DrawParams.RootTransform.TransformPosition(TrajectoryPos);
 				bIsTrajectoryPosValid = true;
 
-				const float GradientPercentage = NumSamples > 1 ? float(SampleIdx) / float(NumSamples - 1) : 0.f;
-				const FColor Color = DrawParams.GetColor(Sample.ColorPresetIndex, GradientPercentage);
+				const FColor Color = DrawParams.GetColor(Sample.ColorPresetIndex);
 
 				if (EnumHasAnyFlags(DrawParams.Flags, EDebugDrawFlags::DrawFast | EDebugDrawFlags::DrawSearchIndex))
 				{
@@ -1819,20 +1821,26 @@ void UPoseSearchFeatureChannel_Trajectory::DebugDraw(const UE::PoseSearch::FDebu
 				{
 					DrawDebugSphere(DrawParams.World, TrajectoryPos, DrawDebugSphereSize, DrawDebugSphereSegments, Color, bPersistent, LifeTime, DepthPriority);
 				}
+
+				TrajSplinePos.Add(TrajectoryPos);
+				TrajSplineColor.Add(Color);
 			}
 		}
 		
 		if (!bIsTrajectoryPosValid)
 		{
 			TrajectoryPos = TrajectoryPositionReconstructor.GetReconstructedTrajectoryPos(*this, PoseVector, DrawParams.RootTransform, Sample.Offset);
+			
+			TrajSplinePos.Add(TrajectoryPos);
+			const FColor Color = TrajSplineColor.IsEmpty() ? FColor::Black : TrajSplineColor.Last();
+			TrajSplineColor.Add(Color);
 		}
 
 		if (EnumHasAnyFlags(Sample.Flags, EPoseSearchTrajectoryFlags::Velocity))
 		{
 			FVector TrajectoryVel = FFeatureVectorHelper::DecodeVector(PoseVector, DataOffset);
 
-			const float GradientPercentage = NumSamples > 1 ? float(SampleIdx) / float(NumSamples - 1) : 0.f;
-			const FColor Color = DrawParams.GetColor(Sample.ColorPresetIndex, GradientPercentage);
+			const FColor Color = DrawParams.GetColor(Sample.ColorPresetIndex);
 
 			TrajectoryVel *= DrawDebugVelocityScale;
 			TrajectoryVel = DrawParams.RootTransform.TransformVector(TrajectoryVel);
@@ -1863,8 +1871,7 @@ void UPoseSearchFeatureChannel_Trajectory::DebugDraw(const UE::PoseSearch::FDebu
 			const FVector2D TrajectoryVel2D = FFeatureVectorHelper::DecodeVector2D(PoseVector, DataOffset);
 			FVector TrajectoryVel(TrajectoryVel2D.X, TrajectoryVel2D.Y, 0.f);
 
-			const float GradientPercentage = NumSamples > 1 ? float(SampleIdx) / float(NumSamples - 1) : 0.f;
-			const FColor Color = DrawParams.GetColor(Sample.ColorPresetIndex, GradientPercentage);
+			const FColor Color = DrawParams.GetColor(Sample.ColorPresetIndex);
 
 			TrajectoryVel *= DrawDebugVelocityScale;
 			TrajectoryVel = DrawParams.RootTransform.TransformVector(TrajectoryVel);
@@ -1895,8 +1902,7 @@ void UPoseSearchFeatureChannel_Trajectory::DebugDraw(const UE::PoseSearch::FDebu
 		{
 			FVector TrajectoryVelDirection = FFeatureVectorHelper::DecodeVector(PoseVector, DataOffset);
 
-			const float GradientPercentage = NumSamples > 1 ? float(SampleIdx) / float(NumSamples - 1) : 0.f;
-			const FColor Color = DrawParams.GetColor(Sample.ColorPresetIndex, GradientPercentage);
+			const FColor Color = DrawParams.GetColor(Sample.ColorPresetIndex);
 
 			TrajectoryVelDirection = DrawParams.RootTransform.TransformVector(TrajectoryVelDirection);
 
@@ -1926,8 +1932,7 @@ void UPoseSearchFeatureChannel_Trajectory::DebugDraw(const UE::PoseSearch::FDebu
 			const FVector2D TrajectoryVelDirection2D = FFeatureVectorHelper::DecodeVector2D(PoseVector, DataOffset);
 			FVector TrajectoryVelDirection(TrajectoryVelDirection2D.X, TrajectoryVelDirection2D.Y, 0.f);
 
-			const float GradientPercentage = NumSamples > 1 ? float(SampleIdx) / float(NumSamples - 1) : 0.f;
-			const FColor Color = DrawParams.GetColor(Sample.ColorPresetIndex, GradientPercentage);
+			const FColor Color = DrawParams.GetColor(Sample.ColorPresetIndex);
 
 			TrajectoryVelDirection = DrawParams.RootTransform.TransformVector(TrajectoryVelDirection);
 
@@ -1956,8 +1961,7 @@ void UPoseSearchFeatureChannel_Trajectory::DebugDraw(const UE::PoseSearch::FDebu
 		{
 			FVector TrajectoryForward = FFeatureVectorHelper::DecodeVector(PoseVector, DataOffset);
 
-			const float GradientPercentage = NumSamples > 1 ? float(SampleIdx) / float(NumSamples - 1) : 0.f;
-			const FColor Color = DrawParams.GetColor(Sample.ColorPresetIndex, GradientPercentage);
+			const FColor Color = DrawParams.GetColor(Sample.ColorPresetIndex);
 
 			TrajectoryForward = DrawParams.RootTransform.TransformVector(TrajectoryForward);
 
@@ -1987,8 +1991,7 @@ void UPoseSearchFeatureChannel_Trajectory::DebugDraw(const UE::PoseSearch::FDebu
 			const FVector2D TrajectoryForward2D = FFeatureVectorHelper::DecodeVector2D(PoseVector, DataOffset);
 			FVector TrajectoryForward(TrajectoryForward2D.X, TrajectoryForward2D.Y, 0.f);
 
-			const float GradientPercentage = NumSamples > 1 ? float(SampleIdx) / float(NumSamples - 1) : 0.f;
-			const FColor Color = DrawParams.GetColor(Sample.ColorPresetIndex, GradientPercentage);
+			const FColor Color = DrawParams.GetColor(Sample.ColorPresetIndex);
 
 			TrajectoryForward = DrawParams.RootTransform.TransformVector(TrajectoryForward);
 
@@ -2015,8 +2018,7 @@ void UPoseSearchFeatureChannel_Trajectory::DebugDraw(const UE::PoseSearch::FDebu
 
 		if (EnumHasAnyFlags(DrawParams.Flags, EDebugDrawFlags::DrawSampleLabels))
 		{
-			const float GradientPercentage = NumSamples > 1 ? float(SampleIdx) / float(NumSamples - 1) : 0.f;
-			const FColor Color = DrawParams.GetColor(Sample.ColorPresetIndex, GradientPercentage);
+			const FColor Color = DrawParams.GetColor(Sample.ColorPresetIndex);
 
 			const FString SampleLabel = FString::Format(TEXT("{0}"), { SampleIdx });
 
@@ -2034,6 +2036,8 @@ void UPoseSearchFeatureChannel_Trajectory::DebugDraw(const UE::PoseSearch::FDebu
 		++SampleIdx;
 	}
 
+	DrawCentripetalCatmullRomSpline(DrawParams.World, TrajSplinePos, TrajSplineColor, 0.5f, 8.f, bPersistent, LifeTime, DepthPriority, 0.f);
+
 	check(DataOffset == ChannelDataOffset + ChannelCardinality);
 }
 
diff --git a/Engine/Plugins/Experimental/Animation/PoseSearch/Source/Runtime/Public/PoseSearch/PoseSearch.h b/Engine/Plugins/Experimental/Animation/PoseSearch/Source/Runtime/Public/PoseSearch/PoseSearch.h
--- a/Engine/Plugins/Experimental/Animation/PoseSearch/Source/Runtime/Public/PoseSearch/PoseSearch.h
+++ b/Engine/Plugins/Experimental/Animation/PoseSearch/Source/Runtime/Public/PoseSearch/PoseSearch.h
@@ -53,7 +53,9 @@ struct FSearchContext;
 
 } // namespace UE::PoseSearch
 
-namespace Eigen {
+// eigen forward declaration
+namespace Eigen
+{
 	template<typename _Scalar, int _Rows, int _Cols, int _Options, int _MaxRows, int _MaxCols>
 	class Matrix;
 	using MatrixXd = Matrix<double, -1, -1, 0, -1, -1>;
@@ -386,25 +388,11 @@ public:
 	// ...to finally wrap the section up by calling EndBreakDownSection
 	virtual void EndBreakDownSection(const FText& Label) = 0;
 
-	// most common case implementation
-	void AddEntireBreakDownSection(const FText& Label, const UPoseSearchSchema* Schema, int32 DataOffset, int32 Cardinality)
-	{
-		BeginBreakDownSection(Label);
-
-		const int32 Count = Num();
-		for (int32 i = 0; i < Count; ++i)
-		{
-			if (IsCostVectorFromSchema(i, Schema))
-			{
-				const float CostBreakdown = ArraySum(GetCostVector(i, Schema), DataOffset, Cardinality);
-				SetCostBreakDown(CostBreakdown, i, Schema);
-			}
-		}
-
-		EndBreakDownSection(Label);
-	}
-
+	// true if want the channel to be verbose and generate the cost breakdown labels
 	virtual bool IsVerbose() const { return true; }
+
+	// most common case implementation
+	void AddEntireBreakDownSection(const FText& Label, const UPoseSearchSchema* Schema, int32 DataOffset, int32 Cardinality);
 };
 
 } // namespace UE::PoseSearch
@@ -970,9 +958,6 @@ enum class EDebugDrawFlags : uint32
 	// Label samples with their indices
 	DrawSampleLabels = 1 << 3,
 
-	// Fade colors
-	DrawSamplesWithColorGradient = 1 << 4,
-
 	// Draw Bone Names
 	DrawBoneNames = 1 << 5,
 
@@ -998,7 +983,7 @@ struct POSESEARCH_API FDebugDrawParams
 	TWeakObjectPtr<const USkinnedMeshComponent> Mesh = nullptr;
 
 	bool CanDraw() const;
-	FColor GetColor(int32 ColorPreset, float GradientPercentage = 0.f) const;
+	FColor GetColor(int32 ColorPreset) const;
 	const FPoseSearchIndex* GetSearchIndex() const;
 	const UPoseSearchSchema* GetSchema() const;
 };
@@ -1025,6 +1010,7 @@ struct FSearchResult
 	// lerp value to find AssetTime from PrevPoseIdx -> AssetTime -> NextPoseIdx, within range [-0.5, 0.5]
 	float LerpValue = 0.f;
 
+	// @todo: it should be a weak pointer
 	const FPoseSearchIndexAsset* SearchIndexAsset = nullptr;
 	TWeakObjectPtr<const UPoseSearchDatabase> Database = nullptr;
 	FPoseSearchFeatureVectorBuilder ComposedQuery;
