pose search plugin
- removing _DEPRECATED
- small code refactoring

diff --git a/Engine/Plugins/Experimental/Animation/PoseSearch/Source/Editor/Private/PoseSearchDebugger.cpp b/Engine/Plugins/Experimental/Animation/PoseSearch/Source/Editor/Private/PoseSearchDebugger.cpp
--- a/Engine/Plugins/Experimental/Animation/PoseSearch/Source/Editor/Private/PoseSearchDebugger.cpp
+++ b/Engine/Plugins/Experimental/Animation/PoseSearch/Source/Editor/Private/PoseSearchDebugger.cpp
@@ -31,7 +31,8 @@
 
 #define LOCTEXT_NAMESPACE "PoseSearchDebugger"
 
-namespace UE::PoseSearch {
+namespace UE::PoseSearch
+{
 
 static FLinearColor LinearColorBlend(FLinearColor LinearColorA, FLinearColor LinearColorB, float BlendParam)
 {
@@ -48,6 +49,39 @@ static void LinearColorBlend(FLinearColor LinearColorA, FLinearColor LinearColor
 	}
 }
 
+static void ArrayMinMax(TConstArrayView<const float> View, TArrayView<float> Min, TArrayView<float> Max, float InvalidValue)
+{
+	const int32 Num = View.Num();
+	check(Num == Min.Num() && Num == Max.Num());
+	for (int i = 0; i < Num; ++i)
+	{
+		const float Value = View[i];
+		if (Value != InvalidValue)
+		{
+			Min[i] = FMath::Min(Min[i], Value);
+			Max[i] = FMath::Max(Max[i], Value);
+		}
+	}
+}
+
+static void ArraySafeNormalize(TConstArrayView<const float> View, TConstArrayView<const float> Min, TConstArrayView<const float> Max, TArrayView<float> NormalizedView)
+{
+	const int32 Num = View.Num();
+	check(Num == Min.Num() && Num == Max.Num() && Num == NormalizedView.Num());
+	for (int i = 0; i < Num; ++i)
+	{
+		const float Delta = Max[i] - Min[i];
+		if (FMath::IsNearlyZero(Delta, UE_KINDA_SMALL_NUMBER))
+		{
+			NormalizedView[i] = 0.f;
+		}
+		else
+		{
+			NormalizedView[i] = (View[i] - Min[i]) / Delta;
+		}
+	}
+}
+
 } // namespace UE::PoseSearch
 
 void UPoseSearchMeshComponent::Initialize(const FTransform& InComponentToWorld)
@@ -213,7 +247,8 @@ void UPoseSearchMeshComponent::UpdatePose(const FUpdateContext& UpdateContext)
 }
 
 
-namespace UE::PoseSearch {
+namespace UE::PoseSearch
+{
 
 class FDebuggerDatabaseRowData : public TSharedFromThis<FDebuggerDatabaseRowData>
 {
diff --git a/Engine/Plugins/Experimental/Animation/PoseSearch/Source/Editor/Private/Trace/PoseSearchTraceAnalyzer.cpp b/Engine/Plugins/Experimental/Animation/PoseSearch/Source/Editor/Private/Trace/PoseSearchTraceAnalyzer.cpp
--- a/Engine/Plugins/Experimental/Animation/PoseSearch/Source/Editor/Private/Trace/PoseSearchTraceAnalyzer.cpp
+++ b/Engine/Plugins/Experimental/Animation/PoseSearch/Source/Editor/Private/Trace/PoseSearchTraceAnalyzer.cpp
@@ -7,7 +7,8 @@
 #include "Runtime/Private/Trace/PoseSearchTraceLogger.h"
 #include "TraceServices/Model/AnalysisSession.h"
 
-namespace UE { namespace PoseSearch {
+namespace UE::PoseSearch
+{
 
 FTraceAnalyzer::FTraceAnalyzer(TraceServices::IAnalysisSession& InSession, FTraceProvider& InTraceProvider)
 : Session(InSession), TraceProvider(InTraceProvider)
@@ -48,4 +49,4 @@ bool FTraceAnalyzer::OnEvent(uint16 RouteId, EStyle Style, const FOnEventContext
 
 	return true;
 }
-}}
+} // namespace UE::PoseSearch
diff --git a/Engine/Plugins/Experimental/Animation/PoseSearch/Source/Editor/Private/Trace/PoseSearchTraceAnalyzer.h b/Engine/Plugins/Experimental/Animation/PoseSearch/Source/Editor/Private/Trace/PoseSearchTraceAnalyzer.h
--- a/Engine/Plugins/Experimental/Animation/PoseSearch/Source/Editor/Private/Trace/PoseSearchTraceAnalyzer.h
+++ b/Engine/Plugins/Experimental/Animation/PoseSearch/Source/Editor/Private/Trace/PoseSearchTraceAnalyzer.h
@@ -10,7 +10,8 @@ namespace TraceServices
 class IAnalysisSession;
 }
 
-namespace UE { namespace PoseSearch { 
+namespace UE::PoseSearch
+{ 
 
 class FTraceProvider;
 
@@ -37,4 +38,4 @@ private:
 	FTraceProvider& TraceProvider;
 };
 
-}}
+} // namespace UE::PoseSearch
diff --git a/Engine/Plugins/Experimental/Animation/PoseSearch/Source/Editor/Private/Trace/PoseSearchTraceModule.cpp b/Engine/Plugins/Experimental/Animation/PoseSearch/Source/Editor/Private/Trace/PoseSearchTraceModule.cpp
--- a/Engine/Plugins/Experimental/Animation/PoseSearch/Source/Editor/Private/Trace/PoseSearchTraceModule.cpp
+++ b/Engine/Plugins/Experimental/Animation/PoseSearch/Source/Editor/Private/Trace/PoseSearchTraceModule.cpp
@@ -4,7 +4,8 @@
 #include "PoseSearchTraceAnalyzer.h"
 #include "PoseSearchTraceProvider.h"
 
-namespace UE { namespace PoseSearch {
+namespace UE::PoseSearch
+{
 
 const FName FTraceModule::ModuleName("PoseSearchTrace");
 
@@ -31,4 +32,4 @@ void FTraceModule::GenerateReports(const TraceServices::IAnalysisSession& Sessio
 {
 }
 
-}}
+} // namespace UE::PoseSearch
diff --git a/Engine/Plugins/Experimental/Animation/PoseSearch/Source/Editor/Private/Trace/PoseSearchTraceModule.h b/Engine/Plugins/Experimental/Animation/PoseSearch/Source/Editor/Private/Trace/PoseSearchTraceModule.h
--- a/Engine/Plugins/Experimental/Animation/PoseSearch/Source/Editor/Private/Trace/PoseSearchTraceModule.h
+++ b/Engine/Plugins/Experimental/Animation/PoseSearch/Source/Editor/Private/Trace/PoseSearchTraceModule.h
@@ -5,7 +5,8 @@
 #include "CoreMinimal.h"
 #include "TraceServices/ModuleService.h"
 
-namespace UE { namespace PoseSearch {
+namespace UE::PoseSearch
+{
 
 /**
 * Module used for loading our PoseSearch tracing systems
@@ -26,4 +27,4 @@ protected:
 	virtual const TCHAR* GetCommandLineArgument() override { return TEXT("posesearchtrace"); }
 };
 
-}}
\ No newline at end of file
+} // namespace UE::PoseSearch
\ No newline at end of file
diff --git a/Engine/Plugins/Experimental/Animation/PoseSearch/Source/Editor/Private/Trace/PoseSearchTraceProvider.cpp b/Engine/Plugins/Experimental/Animation/PoseSearch/Source/Editor/Private/Trace/PoseSearchTraceProvider.cpp
--- a/Engine/Plugins/Experimental/Animation/PoseSearch/Source/Editor/Private/Trace/PoseSearchTraceProvider.cpp
+++ b/Engine/Plugins/Experimental/Animation/PoseSearch/Source/Editor/Private/Trace/PoseSearchTraceProvider.cpp
@@ -2,7 +2,8 @@
 
 #include "PoseSearchTraceProvider.h"
 
-namespace UE { namespace PoseSearch {
+namespace UE::PoseSearch
+{
 
 const FName FTraceProvider::ProviderName("PoseSearchTraceProvider");
 
@@ -45,4 +46,4 @@ FArchive& operator<<(FArchive& Ar, FTraceMotionMatchingStateMessage& State)
 	return Ar;
 }
 
-}}
+} // namespace UE::PoseSearch
diff --git a/Engine/Plugins/Experimental/Animation/PoseSearch/Source/Editor/Private/Trace/PoseSearchTraceProvider.h b/Engine/Plugins/Experimental/Animation/PoseSearch/Source/Editor/Private/Trace/PoseSearchTraceProvider.h
--- a/Engine/Plugins/Experimental/Animation/PoseSearch/Source/Editor/Private/Trace/PoseSearchTraceProvider.h
+++ b/Engine/Plugins/Experimental/Animation/PoseSearch/Source/Editor/Private/Trace/PoseSearchTraceProvider.h
@@ -8,8 +8,8 @@
 #include "TraceServices/Model/AnalysisSession.h"
 
 
-namespace UE { namespace PoseSearch {
-
+namespace UE::PoseSearch
+{
 
 /** Motion matching state message container */
 struct FTraceMotionMatchingStateMessage : FTraceMotionMatchingState, FTraceMessage
@@ -162,4 +162,4 @@ private:
 
 	TraceServices::IAnalysisSession& Session;
 };
-}}
+} // namespace UE::PoseSearch
diff --git a/Engine/Plugins/Experimental/Animation/PoseSearch/Source/Runtime/Private/AnimNode_PoseSearchHistoryCollector.cpp b/Engine/Plugins/Experimental/Animation/PoseSearch/Source/Runtime/Private/AnimNode_PoseSearchHistoryCollector.cpp
--- a/Engine/Plugins/Experimental/Animation/PoseSearch/Source/Runtime/Private/AnimNode_PoseSearchHistoryCollector.cpp
+++ b/Engine/Plugins/Experimental/Animation/PoseSearch/Source/Runtime/Private/AnimNode_PoseSearchHistoryCollector.cpp
@@ -6,7 +6,8 @@
 
 #define LOCTEXT_NAMESPACE "AnimNode_PoseSearchHistoryCollector"
 
-namespace UE { namespace PoseSearch { namespace Private {
+namespace UE::PoseSearch::Private
+{
 
 class FPoseHistoryProvider : public IPoseHistoryProvider
 {
@@ -30,7 +31,7 @@ public:
 	FAnimNode_PoseSearchHistoryCollector& Node;
 };
 
-}}} // namespace UE::PoseSearch::Private
+} // namespace UE::PoseSearch::Private
 
 
 /////////////////////////////////////////////////////
diff --git a/Engine/Plugins/Experimental/Animation/PoseSearch/Source/Runtime/Private/KDTree.cpp b/Engine/Plugins/Experimental/Animation/PoseSearch/Source/Runtime/Private/KDTree.cpp
--- a/Engine/Plugins/Experimental/Animation/PoseSearch/Source/Runtime/Private/KDTree.cpp
+++ b/Engine/Plugins/Experimental/Animation/PoseSearch/Source/Runtime/Private/KDTree.cpp
@@ -12,7 +12,7 @@ THIRD_PARTY_INCLUDES_END
 #endif
 // @third party code - END nanoflann
 
-namespace UE { namespace PoseSearch
+namespace UE::PoseSearch
 {
 
 #if UE_POSE_SEARCH_USE_NANOFLANN
@@ -179,4 +179,4 @@ FArchive& Serialize(FArchive& Ar, FKDTree& KDTree, const float* KDTreeData)
 	return Ar;
 }
 
-} }
+} // namespace UE::PoseSearch
diff --git a/Engine/Plugins/Experimental/Animation/PoseSearch/Source/Runtime/Private/PoseSearch.cpp b/Engine/Plugins/Experimental/Animation/PoseSearch/Source/Runtime/Private/PoseSearch.cpp
--- a/Engine/Plugins/Experimental/Animation/PoseSearch/Source/Runtime/Private/PoseSearch.cpp
+++ b/Engine/Plugins/Experimental/Animation/PoseSearch/Source/Runtime/Private/PoseSearch.cpp
@@ -3,11 +3,6 @@
 #include "PoseSearch/PoseSearch.h"
 #include "PoseSearch/PoseSearchAnimNotifies.h"
 
-#if WITH_EDITORONLY_DATA
-// @todo: deleted this include after deleting the *_DEPRECATED properties
-#include "PoseSearch/PoseSearchFeatureChannels.h"
-#endif // WITH_EDITORONLY_DATA
-
 #include "PoseSearchDerivedData.h"
 #include "PoseSearchEigenHelper.h"
 
@@ -78,7 +73,7 @@ namespace UE::PoseSearch
 	constexpr EParallelForFlags ParallelForFlags = EParallelForFlags::None;
 #endif // UE_POSE_SEARCH_FORCE_SINGLE_THREAD
 
-float ArraySum(TConstArrayView<const float> View, int32 StartIndex, int32 Offset)
+static inline float ArraySum(TConstArrayView<const float> View, int32 StartIndex, int32 Offset)
 {
 	float Sum = 0.f;
 	const int32 EndIndex = StartIndex + Offset;
@@ -89,40 +84,7 @@ float ArraySum(TConstArrayView<const float> View, int32 StartIndex, int32 Offset
 	return Sum;
 }
 
-void ArrayMinMax(TConstArrayView<const float> View, TArrayView<float> Min, TArrayView<float> Max, float InvalidValue)
-{
-	const int32 Num = View.Num();
-	check(Num == Min.Num() && Num == Max.Num());
-	for (int i = 0; i < Num; ++i)
-	{
-		const float Value = View[i];
-		if (Value != InvalidValue)
-		{
-			Min[i] = FMath::Min(Min[i], Value);
-			Max[i] = FMath::Max(Max[i], Value);
-		}
-	}
-}
-
-void ArraySafeNormalize(TConstArrayView<const float> View, TConstArrayView<const float> Min, TConstArrayView<const float> Max, TArrayView<float> NormalizedView)
-{
-	const int32 Num = View.Num();
-	check(Num == Min.Num() && Num == Max.Num() && Num == NormalizedView.Num());
-	for (int i = 0; i < Num; ++i)
-	{
-		const float Delta = Max[i] - Min[i];
-		if (FMath::IsNearlyZero(Delta, UE_KINDA_SMALL_NUMBER))
-		{
-			NormalizedView[i] = 0.f;
-		}
-		else
-		{
-			NormalizedView[i] = (View[i] - Min[i]) / Delta;
-		}
-	}
-}
-
-float CompareFeatureVectors(TConstArrayView<const float> A, TConstArrayView<const float> B, TConstArrayView<const float> WeightsSqrt)
+static inline float CompareFeatureVectors(TConstArrayView<const float> A, TConstArrayView<const float> B, TConstArrayView<const float> WeightsSqrt)
 {
 	check(A.Num() == B.Num() && A.Num() == WeightsSqrt.Num());
 
@@ -133,7 +95,7 @@ float CompareFeatureVectors(TConstArrayView<const float> A, TConstArrayView<cons
 	return ((VA - VB) * VW).square().sum();
 }
 
-float CompareFeatureVectors(TConstArrayView<const float> A, TConstArrayView<const float> B)
+static inline float CompareFeatureVectors(TConstArrayView<const float> A, TConstArrayView<const float> B)
 {
 	check(A.Num() == B.Num());
 
@@ -160,7 +122,7 @@ static inline bool IsSamplingRangeValid(FFloatInterval Range)
 	return Range.IsValid() && (Range.Min >= 0.0f);
 }
 
-static FFloatInterval GetEffectiveSamplingRange(const UAnimSequenceBase* Sequence, FFloatInterval RequestedSamplingRange)
+static inline FFloatInterval GetEffectiveSamplingRange(const UAnimSequenceBase* Sequence, FFloatInterval RequestedSamplingRange)
 {
 	const bool bSampleAll = (RequestedSamplingRange.Min == 0.0f) && (RequestedSamplingRange.Max == 0.0f);
 	const float SequencePlayLength = Sequence->GetPlayLength();
@@ -357,7 +319,8 @@ void UPoseSearchFeatureChannel::ComputeMeanDeviations(const Eigen::MatrixXd& Cen
 //////////////////////////////////////////////////////////////////////////
 // UPoseSearchSchema
 
-namespace UE::PoseSearch {
+namespace UE::PoseSearch
+{
 
 int32 FSchemaInitializer::AddBoneReference(const FBoneReference& BoneReference)
 {
@@ -404,155 +367,7 @@ void UPoseSearchSchema::PreSave(FObjectPreSaveContext ObjectSaveContext)
 void UPoseSearchSchema::PostLoad()
 {
 	Super::PostLoad();
-
-	bool bNeedFinalize = false;
-
-	// Migrate deprecated schema properties into channels
-#if WITH_EDITORONLY_DATA
-	if (SchemaCardinality == 0)
-	{
-		bNeedFinalize = true;
-	}
-
-	if (!TrajectorySampleTimes_DEPRECATED.IsEmpty())
-	{
-		bNeedFinalize = true;
-
-		UPoseSearchFeatureChannel_Trajectory* Channel = NewObject<UPoseSearchFeatureChannel_Trajectory>(this);
-		Channel->Domain = EPoseSearchFeatureDomain::Time;
-		Channel->SampleOffsets_DEPRECATED = TrajectorySampleTimes_DEPRECATED;
-		Channel->bUseFacingDirections_DEPRECATED = bUseTrajectoryForwardVectors_DEPRECATED;
-		Channel->bUseLinearVelocities_DEPRECATED = bUseTrajectoryVelocities_DEPRECATED;
-		Channel->bUsePositions_DEPRECATED = bUseTrajectoryPositions_DEPRECATED;
-		Channels.Add(Channel);
-
-		TrajectorySampleTimes_DEPRECATED.Empty();
-	}
-
-	if (!TrajectorySampleDistances_DEPRECATED.IsEmpty())
-	{
-		bNeedFinalize = true;
-
-		UPoseSearchFeatureChannel_Trajectory* Channel = NewObject<UPoseSearchFeatureChannel_Trajectory>(this);
-		Channel->Domain = EPoseSearchFeatureDomain::Distance;
-		Channel->SampleOffsets_DEPRECATED = TrajectorySampleDistances_DEPRECATED;
-		Channel->bUseFacingDirections_DEPRECATED = bUseTrajectoryForwardVectors_DEPRECATED;
-		Channel->bUseLinearVelocities_DEPRECATED = bUseTrajectoryVelocities_DEPRECATED;
-		Channel->bUsePositions_DEPRECATED = bUseTrajectoryPositions_DEPRECATED;
-		Channels.Add(Channel);
-
-		TrajectorySampleDistances_DEPRECATED.Empty();
-	}
-
-	if (!PoseSampleTimes_DEPRECATED.IsEmpty())
-	{
-		bNeedFinalize = true;
-
-		UPoseSearchFeatureChannel_Pose* Channel = NewObject<UPoseSearchFeatureChannel_Pose>(this);
-		Channel->SampledBones = SampledBones_DEPRECATED;
-		Channel->SampleTimes = PoseSampleTimes_DEPRECATED;
-		Channels.Add(Channel);
-
-		SampledBones_DEPRECATED.Empty();
-		PoseSampleTimes_DEPRECATED.Empty();
-	}
-
-	if (!bNeedFinalize)
-	{
-		for (UPoseSearchFeatureChannel* Channel : Channels)
-		{
-			if (Channel->ChannelDataOffset == -1 || Channel->ChannelCardinality == -1)
-			{
-				bNeedFinalize = true;
-				break;
-			}
-		}
-	}
-
-	for (UPoseSearchFeatureChannel* Channel : Channels)
-	{
-		if (UPoseSearchFeatureChannel_Pose* ChannelPose = Cast<UPoseSearchFeatureChannel_Pose>(Channel))
-		{
-			for (FPoseSearchBone& SampledBone : ChannelPose->SampledBones)
-			if (SampledBone.bUseVelocity_DEPRECATED ||
-				SampledBone.bUsePosition_DEPRECATED ||
-				SampledBone.bUseRotation_DEPRECATED ||
-				SampledBone.bUsePhase_DEPRECATED)
-			{
-				SampledBone.Flags = 0;
-				if (SampledBone.bUseVelocity_DEPRECATED)
-				{
-					SampledBone.Flags |= EPoseSearchBoneFlags::Velocity;
-					SampledBone.bUseVelocity_DEPRECATED = false;
-				}
-				
-				if (SampledBone.bUsePosition_DEPRECATED)
-				{
-					SampledBone.Flags |= EPoseSearchBoneFlags::Position;
-					SampledBone.bUsePosition_DEPRECATED = false;
-				}
-
-				if (SampledBone.bUseRotation_DEPRECATED)
-				{
-					SampledBone.Flags |= EPoseSearchBoneFlags::Rotation;
-					SampledBone.bUseRotation_DEPRECATED = false;
-				}
-
-				if (SampledBone.bUsePhase_DEPRECATED)
-				{
-					SampledBone.Flags |= EPoseSearchBoneFlags::Phase;
-					SampledBone.bUsePhase_DEPRECATED = false;
-				}
-
-				bNeedFinalize = true;
-			}
-		}
-
-		if (UPoseSearchFeatureChannel_Trajectory* ChannelTrajectory = Cast<UPoseSearchFeatureChannel_Trajectory>(Channel))
-		{
-			if (ChannelTrajectory->SampleOffsets_DEPRECATED.Num() > 0)
-			{
-				ChannelTrajectory->Samples.Reset();
-
-				for (float Offset : ChannelTrajectory->SampleOffsets_DEPRECATED)
-				{
-					FPoseSearchTrajectorySample Sample;
-					Sample.Flags = 0;
-					if (ChannelTrajectory->bUsePositions_DEPRECATED)
-					{
-						Sample.Flags |= EPoseSearchTrajectoryFlags::Position;
-					}
-					if (ChannelTrajectory->bUseLinearVelocities_DEPRECATED)
-					{
-						Sample.Flags |= EPoseSearchTrajectoryFlags::Velocity;
-					}
-					if (ChannelTrajectory->bUseFacingDirections_DEPRECATED)
-					{
-						Sample.Flags |= EPoseSearchTrajectoryFlags::FacingDirection;
-					}
-					Sample.Offset = Offset;
-					Sample.Weight = 1.f;
-					ChannelTrajectory->Samples.Add(Sample);
-				}
-
-				ChannelTrajectory->SampleOffsets_DEPRECATED.Reset();
-				ChannelTrajectory->bUsePositions_DEPRECATED = false;
-				ChannelTrajectory->bUseLinearVelocities_DEPRECATED = false;
-				ChannelTrajectory->bUseFacingDirections_DEPRECATED = false;
-			}
-		}		
-	}
-
-#endif // WITH_EDITORONLY_DATA
-
-	if (bNeedFinalize)
-	{
-		Finalize();
-	}
-	else
-	{
-		ResolveBoneReferences();
-	}
+	ResolveBoneReferences();
 }
 
 #if WITH_EDITOR
@@ -4419,11 +4234,12 @@ static void PreprocessSearchIndexWeights(FPoseSearchIndex& SearchIndex, const UP
 	}
 
 	Eigen::VectorXd ChannelsMeanDeviations = ComputeChannelsDeviations(Database->GetSearchIndex());
-	SearchIndex.Deviation.Init(1.f, NumDimensions);
+	TArray<float> Deviation;
+	Deviation.Init(1.f, NumDimensions);
 	for (int32 Dimension = 0; Dimension != NumDimensions; ++Dimension)
 	{
 		const float ChannelsMeanDeviation = ChannelsMeanDeviations[Dimension];
-		SearchIndex.Deviation[Dimension] = ChannelsMeanDeviation;
+		Deviation[Dimension] = ChannelsMeanDeviation;
 	}
 
 	EPoseSearchDataPreprocessor DataPreprocessor = Database->GetSearchIndex()->Schema->DataPreprocessor;
@@ -4450,9 +4266,13 @@ static void PreprocessSearchIndexWeights(FPoseSearchIndex& SearchIndex, const UP
 		for (int32 Dimension = 0; Dimension != NumDimensions; ++Dimension)
 		{
 			// the idea here is to premultiply the weights by the inverse of the variance (proportional to the square of the deviation) to have a "weighted Mahalanobis" distance
-			SearchIndex.WeightsSqrt[Dimension] /= SearchIndex.Deviation[Dimension];
+			SearchIndex.WeightsSqrt[Dimension] /= Deviation[Dimension];
 		}
 	}
+
+#if WITH_EDITORONLY_DATA
+	SearchIndex.Deviation = Deviation;
+#endif // WITH_EDITORONLY_DATA
 }
 
 // it calcualtes Mean and PCAProjectionMatrix
@@ -4526,10 +4346,12 @@ static void PreprocessSearchIndexPCAData(FPoseSearchIndex& SearchIndex, const UP
 		AccumulatedVariance += EigenValues[Indexer[PCAComponentIndex]];
 	}
 
+#if WITH_EDITORONLY_DATA
 	// calculating the total variance knowing that eigen values measure variance along the principal components:
 	const float TotalVariance = EigenValues.sum();
 	// and explained variance as ratio between AccumulatedVariance and TotalVariance: https://ro-che.info/articles/2017-12-11-pca-explained-variance
 	SearchIndex.PCAExplainedVariance = TotalVariance > UE_KINDA_SMALL_NUMBER ? AccumulatedVariance / TotalVariance : 0.f;
+#endif // WITH_EDITORONLY_DATA
 
 	MapPCAValues = CenteredValues * PCAProjectionMatrix;
 
diff --git a/Engine/Plugins/Experimental/Animation/PoseSearch/Source/Runtime/Private/PoseSearchEigenHelper.h b/Engine/Plugins/Experimental/Animation/PoseSearch/Source/Runtime/Private/PoseSearchEigenHelper.h
--- a/Engine/Plugins/Experimental/Animation/PoseSearch/Source/Runtime/Private/PoseSearchEigenHelper.h
+++ b/Engine/Plugins/Experimental/Animation/PoseSearch/Source/Runtime/Private/PoseSearchEigenHelper.h
@@ -48,7 +48,7 @@ PRAGMA_DEFAULT_VISIBILITY_END
 #include <sstream>
 #endif
 
-namespace UE { namespace PoseSearch
+namespace UE::PoseSearch
 {
 	using RowMajorVector = Eigen::Matrix<float, 1, Eigen::Dynamic, Eigen::RowMajor>;
 	using RowMajorVectorMap = Eigen::Map<RowMajorVector, Eigen::RowMajor>;
@@ -77,4 +77,4 @@ namespace UE { namespace PoseSearch
 	}
 #endif
 
-}}
\ No newline at end of file
+} // namespace UE::PoseSearch
\ No newline at end of file
diff --git a/Engine/Plugins/Experimental/Animation/PoseSearch/Source/Runtime/Private/PoseSearchFeatureChannels.cpp b/Engine/Plugins/Experimental/Animation/PoseSearch/Source/Runtime/Private/PoseSearchFeatureChannels.cpp
--- a/Engine/Plugins/Experimental/Animation/PoseSearch/Source/Runtime/Private/PoseSearchFeatureChannels.cpp
+++ b/Engine/Plugins/Experimental/Animation/PoseSearch/Source/Runtime/Private/PoseSearchFeatureChannels.cpp
@@ -12,7 +12,8 @@
 
 #define LOCTEXT_NAMESPACE "PoseSearchFeatureChannels"
 
-namespace UE::PoseSearch {
+namespace UE::PoseSearch
+{
 
 //////////////////////////////////////////////////////////////////////////
 // Constants
@@ -359,12 +360,12 @@ bool UPoseSearchFeatureChannel_Position::BuildQuery(UE::PoseSearch::FSearchConte
 	using namespace UE::PoseSearch;
 
 	const bool bIsCurrentResultValid = SearchContext.CurrentResult.IsValid();
-	const bool bSkip = InputQueryPose != InputQueryPose::UseCharacterPose && bIsCurrentResultValid && SearchContext.CurrentResult.Database->Schema == InOutQuery.GetSchema();
+	const bool bSkip = InputQueryPose != EInputQueryPose::UseCharacterPose && bIsCurrentResultValid && SearchContext.CurrentResult.Database->Schema == InOutQuery.GetSchema();
 	if (bSkip || !SearchContext.History)
 	{
 		if (bIsCurrentResultValid)
 		{
-			const float LerpValue = InputQueryPose == InputQueryPose::UseInterpolatedContinuingPose ? SearchContext.CurrentResult.LerpValue : 0.f;
+			const float LerpValue = InputQueryPose == EInputQueryPose::UseInterpolatedContinuingPose ? SearchContext.CurrentResult.LerpValue : 0.f;
 			int32 DataOffset = ChannelDataOffset;
 			FFeatureVectorHelper::EncodeVector(InOutQuery.EditValues(), DataOffset, SearchContext.GetCurrentResultPrevPoseVector(), SearchContext.GetCurrentResultPoseVector(), SearchContext.GetCurrentResultNextPoseVector(), LerpValue);
 		}
@@ -443,11 +444,11 @@ FVector UPoseSearchFeatureChannel_Heading::GetAxis(const FQuat& Rotation) const
 {
 	switch (HeadingAxis)
 	{
-	case HeadingAxis::X:
+	case EHeadingAxis::X:
 		return Rotation.GetAxisX();
-	case HeadingAxis::Y:
+	case EHeadingAxis::Y:
 		return Rotation.GetAxisY();
-	case HeadingAxis::Z:
+	case EHeadingAxis::Z:
 		return Rotation.GetAxisZ();
 	}
 
@@ -483,12 +484,12 @@ bool UPoseSearchFeatureChannel_Heading::BuildQuery(UE::PoseSearch::FSearchContex
 	using namespace UE::PoseSearch;
 
 	const bool bIsCurrentResultValid = SearchContext.CurrentResult.IsValid();
-	const bool bSkip = InputQueryPose != InputQueryPose::UseCharacterPose && bIsCurrentResultValid && SearchContext.CurrentResult.Database->Schema == InOutQuery.GetSchema();
+	const bool bSkip = InputQueryPose != EInputQueryPose::UseCharacterPose && bIsCurrentResultValid && SearchContext.CurrentResult.Database->Schema == InOutQuery.GetSchema();
 	if (bSkip || !SearchContext.History)
 	{
 		if (bIsCurrentResultValid)
 		{
-			const float LerpValue = InputQueryPose == InputQueryPose::UseInterpolatedContinuingPose ? SearchContext.CurrentResult.LerpValue : 0.f;
+			const float LerpValue = InputQueryPose == EInputQueryPose::UseInterpolatedContinuingPose ? SearchContext.CurrentResult.LerpValue : 0.f;
 			int32 DataOffset = ChannelDataOffset;
 			FFeatureVectorHelper::EncodeVector(InOutQuery.EditValues(), DataOffset, SearchContext.GetCurrentResultPrevPoseVector(), SearchContext.GetCurrentResultPoseVector(), SearchContext.GetCurrentResultNextPoseVector(), LerpValue, true);
 			check(DataOffset == ChannelDataOffset + ChannelCardinality);
@@ -874,12 +875,12 @@ bool UPoseSearchFeatureChannel_Pose::BuildQuery(UE::PoseSearch::FSearchContext&
 	using namespace UE::PoseSearch;
 
 	const bool bIsCurrentResultValid = SearchContext.CurrentResult.IsValid();
-	const bool bSkip = InputQueryPose != InputQueryPose::UseCharacterPose && bIsCurrentResultValid && SearchContext.CurrentResult.Database->Schema == InOutQuery.GetSchema();
+	const bool bSkip = InputQueryPose != EInputQueryPose::UseCharacterPose && bIsCurrentResultValid && SearchContext.CurrentResult.Database->Schema == InOutQuery.GetSchema();
 	if (bSkip || !SearchContext.History)
 	{
 		if (bIsCurrentResultValid)
 		{
-			const float LerpValue = InputQueryPose == InputQueryPose::UseInterpolatedContinuingPose ? SearchContext.CurrentResult.LerpValue : 0.f;
+			const float LerpValue = InputQueryPose == EInputQueryPose::UseInterpolatedContinuingPose ? SearchContext.CurrentResult.LerpValue : 0.f;
 			int32 DataOffset = ChannelDataOffset;
 			for (int32 SampledBoneIdx = 0; SampledBoneIdx != SampledBones.Num(); ++SampledBoneIdx)
 			{
diff --git a/Engine/Plugins/Experimental/Animation/PoseSearch/Source/Runtime/Private/Trace/PoseSearchTraceLogger.cpp b/Engine/Plugins/Experimental/Animation/PoseSearch/Source/Runtime/Private/Trace/PoseSearchTraceLogger.cpp
--- a/Engine/Plugins/Experimental/Animation/PoseSearch/Source/Runtime/Private/Trace/PoseSearchTraceLogger.cpp
+++ b/Engine/Plugins/Experimental/Animation/PoseSearch/Source/Runtime/Private/Trace/PoseSearchTraceLogger.cpp
@@ -12,7 +12,8 @@ UE_TRACE_EVENT_BEGIN(PoseSearch, MotionMatchingState)
 	UE_TRACE_EVENT_FIELD(uint8[], Data)
 UE_TRACE_EVENT_END()
 
-namespace UE { namespace PoseSearch {
+namespace UE::PoseSearch
+{
 
 const FName FTraceLogger::Name("PoseSearch");
 const FName FTraceMotionMatchingState::Name("MotionMatchingState");
@@ -140,4 +141,4 @@ const FTraceMotionMatchingStatePoseEntry* FTraceMotionMatchingState::GetCurrentP
 	return &DbEntry.PoseEntries[CurrentPoseEntryIdx];
 }
 
-}}
\ No newline at end of file
+} // namespace UE::PoseSearch
\ No newline at end of file
diff --git a/Engine/Plugins/Experimental/Animation/PoseSearch/Source/Runtime/Private/Trace/PoseSearchTraceLogger.h b/Engine/Plugins/Experimental/Animation/PoseSearch/Source/Runtime/Private/Trace/PoseSearchTraceLogger.h
--- a/Engine/Plugins/Experimental/Animation/PoseSearch/Source/Runtime/Private/Trace/PoseSearchTraceLogger.h
+++ b/Engine/Plugins/Experimental/Animation/PoseSearch/Source/Runtime/Private/Trace/PoseSearchTraceLogger.h
@@ -6,7 +6,8 @@
 
 UE_TRACE_CHANNEL_EXTERN(PoseSearchChannel, POSESEARCH_API);
 
-namespace UE { namespace PoseSearch {
+namespace UE::PoseSearch
+{
 
 bool IsTracing(const FAnimationBaseContext& InContext);
 
@@ -142,4 +143,4 @@ ENUM_CLASS_FLAGS(FTraceMotionMatchingState::EFlags)
 
 POSESEARCH_API FArchive& operator<<(FArchive& Ar, FTraceMotionMatchingState& State);
 
-}}
+} // namespace UE::PoseSearch
diff --git a/Engine/Plugins/Experimental/Animation/PoseSearch/Source/Runtime/Public/PoseSearch/KDTree.h b/Engine/Plugins/Experimental/Animation/PoseSearch/Source/Runtime/Public/PoseSearch/KDTree.h
--- a/Engine/Plugins/Experimental/Animation/PoseSearch/Source/Runtime/Public/PoseSearch/KDTree.h
+++ b/Engine/Plugins/Experimental/Animation/PoseSearch/Source/Runtime/Public/PoseSearch/KDTree.h
@@ -7,7 +7,7 @@
 
 class FArchive;
 
-namespace UE { namespace PoseSearch
+namespace UE::PoseSearch
 {
 struct FKDTreeImplementation;
 
@@ -123,4 +123,4 @@ struct POSESEARCH_API FKDTree
 };
 
 FArchive& Serialize(FArchive& Ar, FKDTree& KDTree, const float* data);
-}} // namespace UE::PoseSearch
\ No newline at end of file
+} // namespace UE::PoseSearch
\ No newline at end of file
diff --git a/Engine/Plugins/Experimental/Animation/PoseSearch/Source/Runtime/Public/PoseSearch/PoseSearch.h b/Engine/Plugins/Experimental/Animation/PoseSearch/Source/Runtime/Public/PoseSearch/PoseSearch.h
--- a/Engine/Plugins/Experimental/Animation/PoseSearch/Source/Runtime/Public/PoseSearch/PoseSearch.h
+++ b/Engine/Plugins/Experimental/Animation/PoseSearch/Source/Runtime/Public/PoseSearch/PoseSearch.h
@@ -50,7 +50,6 @@ struct FDebugDrawParams;
 struct FSchemaInitializer;
 struct FQueryBuildingContext;
 struct FSearchContext;
-
 } // namespace UE::PoseSearch
 
 // eigen forward declaration
@@ -62,32 +61,9 @@ namespace Eigen
 	using VectorXd = Matrix<double, -1, 1, 0, -1, 1>;
 }
 
-
-namespace UE::PoseSearch
-{
-	// @todo: move this in an utils file and make them inline
-	POSESEARCH_API float ArraySum(TConstArrayView<const float> View, int32 StartIndex, int32 Offset);
-	POSESEARCH_API void ArrayMinMax(TConstArrayView<const float> View, TArrayView<float> Min, TArrayView<float> Max, float InvalidValue);
-	POSESEARCH_API void ArraySafeNormalize(TConstArrayView<const float> View, TConstArrayView<const float> Min, TConstArrayView<const float> Max, TArrayView<float> NormalizedView);
-	POSESEARCH_API float CompareFeatureVectors(TConstArrayView<const float> A, TConstArrayView<const float> B, TConstArrayView<const float> WeightsSqrt);
-	POSESEARCH_API float CompareFeatureVectors(TConstArrayView<const float> A, TConstArrayView<const float> B);
-	POSESEARCH_API void CompareFeatureVectors(TConstArrayView<const float> A, TConstArrayView<const float> B, TConstArrayView<const float> WeightsSqrt, TArrayView<float> Result);
-}
-
 //////////////////////////////////////////////////////////////////////////
 // Constants
 
-// @todo: move this enum in PoseSearchFeatureChannel.h since only used by UPoseSearchFeatureChannel_Trajectory
-UENUM()
-enum class EPoseSearchFeatureDomain : int32
-{
-	Time,
-	Distance,
-
-	Num UMETA(Hidden),
-	Invalid = Num UMETA(Hidden)
-};
-
 UENUM()
 enum class EPoseSearchBooleanRequest : uint8
 {
@@ -140,7 +116,8 @@ enum class ESearchIndexAssetType : int32
 	BlendSpace,
 };
 
-namespace UE::PoseSearch {
+namespace UE::PoseSearch
+{
 	
 enum class EPoseComparisonFlags : int32
 {
@@ -200,58 +177,11 @@ struct FPoseSearchBlockTransitionParameters
 	float SequenceEndInterval = 0.0f;
 };
 
-// @todo: move it into PoseSearchFeatureChannels after removing SampledBones_DEPRECATED 
-UENUM(meta = (Bitflags, UseEnumValuesAsMaskValuesInEditor = "true"))
-enum class EPoseSearchBoneFlags : uint32
-{
-	Velocity = 1 << 0,
-	Position = 1 << 1,
-	Rotation = 1 << 2,
-	Phase = 1 << 3,
-};
-ENUM_CLASS_FLAGS(EPoseSearchBoneFlags);
-constexpr bool EnumHasAnyFlags(int32 Flags, EPoseSearchBoneFlags Contains) { return (Flags & int32(Contains)) != 0; }
-inline int32& operator|=(int32& Lhs, EPoseSearchBoneFlags Rhs) { return Lhs |= int32(Rhs); }
-
-// @todo: move it into PoseSearchFeatureChannels after removing SampledBones_DEPRECATED 
-USTRUCT()
-struct POSESEARCH_API FPoseSearchBone
-{
-	GENERATED_BODY()
-
-	UPROPERTY(EditAnywhere, Category = Config)
-	FBoneReference Reference;
-
-#if WITH_EDITORONLY_DATA
-	UPROPERTY()
-	bool bUseVelocity_DEPRECATED = false;
-
-	UPROPERTY()
-	bool bUsePosition_DEPRECATED = false;
-
-	UPROPERTY()
-	bool bUseRotation_DEPRECATED = false;
-
-	UPROPERTY()
-	bool bUsePhase_DEPRECATED = false;
-#endif
-
-	UPROPERTY(EditAnywhere, meta = (Bitmask, BitmaskEnum = "/Script/PoseSearch.EPoseSearchBoneFlags"), Category = Config)
-	int32 Flags = int32(EPoseSearchBoneFlags::Position);
-
-	// @todo: temporary location for the channel bone weight to help the weights refactoring
-	UPROPERTY(EditAnywhere, Category = Config)
-	float Weight = 1.f;
-
-	UPROPERTY(EditAnywhere, Category = Config, meta=(ExcludeFromHash))
-	int32 ColorPresetIndex = 0;
-};
-
-
 //////////////////////////////////////////////////////////////////////////
 // Asset sampling and indexing
 
-namespace UE::PoseSearch {
+namespace UE::PoseSearch
+{
 
 struct POSESEARCH_API FAssetSamplingContext
 {
@@ -462,7 +392,8 @@ protected:
 //////////////////////////////////////////////////////////////////////////
 // Schema
 
-namespace UE::PoseSearch {
+namespace UE::PoseSearch
+{
 
 struct POSESEARCH_API FSchemaInitializer
 {
@@ -523,29 +454,6 @@ public:
 	UPROPERTY(EditAnywhere, Instanced, BlueprintReadWrite, Category = "Schema")
 	TArray<TObjectPtr<UPoseSearchFeatureChannel>> Channels;
 
-#if WITH_EDITORONLY_DATA
-	UPROPERTY()
-	bool bUseTrajectoryVelocities_DEPRECATED = true;
-
-	UPROPERTY()
-	bool bUseTrajectoryPositions_DEPRECATED = true;
-
-	UPROPERTY()
-	bool bUseTrajectoryForwardVectors_DEPRECATED = false;
-
-	UPROPERTY()
-	TArray<FPoseSearchBone> SampledBones_DEPRECATED;
-
-	UPROPERTY()
-	TArray<float> PoseSampleTimes_DEPRECATED;
-
-	UPROPERTY()
-	TArray<float> TrajectorySampleTimes_DEPRECATED;
-
-	UPROPERTY()
-	TArray<float> TrajectorySampleDistances_DEPRECATED;
-#endif // WITH_EDITOR
-
 	// If set, this schema will support mirroring pose search databases
 	UPROPERTY(EditAnywhere, Category = "Schema")
 	TObjectPtr<UMirrorDataTable> MirrorDataTable;
@@ -708,13 +616,13 @@ struct POSESEARCH_API FPoseSearchIndex
 	UPROPERTY()
 	TArray<float> PCAValues;
 
-	// @todo: editor only data?
+#if WITH_EDITORONLY_DATA
 	UPROPERTY(Category = Info, VisibleAnywhere)
 	float PCAExplainedVariance = 0.f;
 	
-	// @todo: editor only data?
 	UPROPERTY(Category = Info, VisibleAnywhere)
 	TArray<float> Deviation;
+#endif // WITH_EDITORONLY_DATA
 
 	UE::PoseSearch::FKDTree KDTree;
 
@@ -1452,6 +1360,7 @@ public:
 POSESEARCH_API void DrawFeatureVector(const FDebugDrawParams& DrawParams, TArrayView<const float> PoseVector);
 POSESEARCH_API void DrawFeatureVector(const FDebugDrawParams& DrawParams, int32 PoseIdx);
 POSESEARCH_API void DrawSearchIndex(const FDebugDrawParams& DrawParams);
+POSESEARCH_API void CompareFeatureVectors(TConstArrayView<const float> A, TConstArrayView<const float> B, TConstArrayView<const float> WeightsSqrt, TArrayView<float> Result);
 
 /**
 * Creates a pose search index for an animation sequence
diff --git a/Engine/Plugins/Experimental/Animation/PoseSearch/Source/Runtime/Public/PoseSearch/PoseSearchFeatureChannels.h b/Engine/Plugins/Experimental/Animation/PoseSearch/Source/Runtime/Public/PoseSearch/PoseSearchFeatureChannels.h
--- a/Engine/Plugins/Experimental/Animation/PoseSearch/Source/Runtime/Public/PoseSearch/PoseSearchFeatureChannels.h
+++ b/Engine/Plugins/Experimental/Animation/PoseSearch/Source/Runtime/Public/PoseSearch/PoseSearchFeatureChannels.h
@@ -9,7 +9,7 @@
 #include "PoseSearchFeatureChannels.generated.h"
 
 UENUM(BlueprintType)
-enum class InputQueryPose : uint8
+enum class EInputQueryPose : uint8
 {
 	// use character pose to compose the query
 	UseCharacterPose,
@@ -48,7 +48,7 @@ public:
 	int32 ColorPresetIndex = 0;
 
 	UPROPERTY(EditAnywhere, Category = "Settings")
-	InputQueryPose InputQueryPose = InputQueryPose::UseContinuingPose;
+	EInputQueryPose InputQueryPose = EInputQueryPose::UseContinuingPose;
 
 	// if UseSampleTimeOffsetRootBone is true, this UPoseSearchFeatureChannel_Position will calculate the position of Bone from the pose SampleTimeOffset seconds away from the current time pose root bone
 	// if false the calculated position of Bone will be in component space from the pose SampleTimeOffset seconds away
@@ -67,7 +67,7 @@ public:
 // UPoseSearchFeatureChannel_Heading
 
 UENUM(BlueprintType)
-enum class HeadingAxis : uint8
+enum class EHeadingAxis : uint8
 {
 	X,
 	Y,
@@ -93,7 +93,7 @@ public:
 	float SampleTimeOffset = 0.f;
 
 	UPROPERTY(EditAnywhere, Category = "Settings")
-	HeadingAxis HeadingAxis = HeadingAxis::X;	
+	EHeadingAxis HeadingAxis = EHeadingAxis::X;	
 
 	UPROPERTY()
 	int8 SchemaBoneIdx;
@@ -102,7 +102,7 @@ public:
 	int32 ColorPresetIndex = 0;
 
 	UPROPERTY(EditAnywhere, Category = "Settings")
-	InputQueryPose InputQueryPose = InputQueryPose::UseContinuingPose;
+	EInputQueryPose InputQueryPose = EInputQueryPose::UseContinuingPose;
 
 	// if UseSampleTimeOffsetRootBone is true, this UPoseSearchFeatureChannel_Position will calculate the position of Bone from the pose SampleTimeOffset seconds away from the current time pose root bone
 	// if false the calculated position of Bone will be in component space from the pose SampleTimeOffset seconds away
@@ -120,6 +120,36 @@ public:
 };
 
 //////////////////////////////////////////////////////////////////////////
+UENUM(meta = (Bitflags, UseEnumValuesAsMaskValuesInEditor = "true"))
+enum class EPoseSearchBoneFlags : uint32
+{
+	Velocity = 1 << 0,
+	Position = 1 << 1,
+	Rotation = 1 << 2,
+	Phase = 1 << 3,
+};
+ENUM_CLASS_FLAGS(EPoseSearchBoneFlags);
+constexpr bool EnumHasAnyFlags(int32 Flags, EPoseSearchBoneFlags Contains) { return (Flags & int32(Contains)) != 0; }
+inline int32& operator|=(int32& Lhs, EPoseSearchBoneFlags Rhs) { return Lhs |= int32(Rhs); }
+
+USTRUCT()
+struct POSESEARCH_API FPoseSearchBone
+{
+	GENERATED_BODY()
+
+	UPROPERTY(EditAnywhere, Category = Config)
+	FBoneReference Reference;
+
+	UPROPERTY(EditAnywhere, meta = (Bitmask, BitmaskEnum = "/Script/PoseSearch.EPoseSearchBoneFlags"), Category = Config)
+	int32 Flags = int32(EPoseSearchBoneFlags::Position);
+
+	UPROPERTY(EditAnywhere, Category = Config)
+	float Weight = 1.f;
+
+	UPROPERTY(EditAnywhere, Category = Config, meta=(ExcludeFromHash))
+	int32 ColorPresetIndex = 0;
+};
+
 // UPoseSearchFeatureChannel_Pose
 UCLASS(BlueprintType, EditInlineNew)
 class POSESEARCH_API UPoseSearchFeatureChannel_Pose : public UPoseSearchFeatureChannel
@@ -140,7 +170,7 @@ public:
 	TArray<int8> SchemaBoneIdx;
 
 	UPROPERTY(EditAnywhere, Category = "Settings")
-	InputQueryPose InputQueryPose = InputQueryPose::UseContinuingPose;
+	EInputQueryPose InputQueryPose = EInputQueryPose::UseContinuingPose;
 
 	// UObject interface
 	virtual void PreSave(FObjectPreSaveContext ObjectSaveContext) override;
@@ -164,6 +194,15 @@ protected:
 
 //////////////////////////////////////////////////////////////////////////
 // UPoseSearchFeatureChannel_Trajectory
+UENUM()
+enum class EPoseSearchFeatureDomain : int32
+{
+	Time,
+	Distance,
+
+	Num UMETA(Hidden),
+	Invalid = Num UMETA(Hidden)
+};
 
 UENUM(meta = (Bitflags, UseEnumValuesAsMaskValuesInEditor = "true"))
 enum class EPoseSearchTrajectoryFlags : uint32
@@ -205,27 +244,12 @@ class POSESEARCH_API UPoseSearchFeatureChannel_Trajectory : public UPoseSearchFe
 	GENERATED_BODY()
 
 public:
-#if WITH_EDITORONLY_DATA
-	UPROPERTY()
-	bool bUseLinearVelocities_DEPRECATED = true;
-
-	UPROPERTY()
-	bool bUsePositions_DEPRECATED = true;
-
-	UPROPERTY()
-	bool bUseFacingDirections_DEPRECATED = false;
-#endif
-
 	UPROPERTY(EditAnywhere, Category="Settings")
 	EPoseSearchFeatureDomain Domain = EPoseSearchFeatureDomain::Time;
 
 	UPROPERTY(EditAnywhere, Category = "Settings")
 	float Weight = 1.f;
 
-#if WITH_EDITORONLY_DATA
-	UPROPERTY()
-	TArray<float> SampleOffsets_DEPRECATED;
-#endif
 	UPROPERTY(EditAnywhere, Category = "Settings")
 	TArray<FPoseSearchTrajectorySample> Samples;
 
