fix for https://horde.devtools.epicgames.com/job/632b89917e1da7bc7ddcc9fd?step=fa9e&issue=215632

diff --git a/Engine/Plugins/Experimental/Animation/PoseSearch/Source/Editor/Private/PoseSearchDebugger.cpp b/Engine/Plugins/Experimental/Animation/PoseSearch/Source/Editor/Private/PoseSearchDebugger.cpp
--- a/Engine/Plugins/Experimental/Animation/PoseSearch/Source/Editor/Private/PoseSearchDebugger.cpp
+++ b/Engine/Plugins/Experimental/Animation/PoseSearch/Source/Editor/Private/PoseSearchDebugger.cpp
@@ -830,62 +830,64 @@ void SDebuggerDatabaseView::Update(const FTraceMotionMatchingStateMessage& State
 	for (const FTraceMotionMatchingStateDatabaseEntry& DbEntry : State.DatabaseEntries)
 	{
 		const UPoseSearchDatabase* Database = FTraceMotionMatchingState::GetObjectFromId<UPoseSearchDatabase>(DbEntry.DatabaseId);
-
-		const FPoseSearchIndex* SearchIndex = Database->GetSearchIndex();
-		if (SearchIndex && SearchIndex->IsValid())
+		if (Database)
 		{
-			for (const FTraceMotionMatchingStatePoseEntry& PoseEntry : DbEntry.PoseEntries)
+			const FPoseSearchIndex* SearchIndex = Database->GetSearchIndex();
+			if (SearchIndex && SearchIndex->IsValid())
 			{
-				const FPoseSearchIndexAsset* SearchIndexAsset = SearchIndex->FindAssetForPose(PoseEntry.DbPoseIdx);
-				if (SearchIndexAsset)
+				for (const FTraceMotionMatchingStatePoseEntry& PoseEntry : DbEntry.PoseEntries)
 				{
-					TSharedRef<FDebuggerDatabaseRowData>& Row = UnfilteredDatabaseRows.Add_GetRef(MakeShared<FDebuggerDatabaseRowData>());
-
-					const float Time = SearchIndex->GetAssetTime(PoseEntry.DbPoseIdx, SearchIndexAsset);
-
-					Row->PoseIdx = PoseEntry.DbPoseIdx;
-					Row->SourceDatabase = Database;
-					Row->DatabaseName = Database->GetName();
-					Row->DatabasePath = Database->GetPathName();
-					Row->PoseEntryFlags = PoseEntry.Flags;
-					Row->DbAssetIdx = SearchIndexAsset->SourceAssetIdx;
-					Row->AssetTime = Time;
-					Row->bMirrored = SearchIndexAsset->bMirrored;
-					Row->PoseCost = PoseEntry.Cost;
-
-					Row->CostVector.SetNum(Database->Schema->SchemaCardinality);
-					TConstArrayView<const float> PoseValues = SearchIndex->GetPoseValues(PoseEntry.DbPoseIdx);
-
-					// in case we modify the schema while PIE is paused and displaying the Pose Search Editor, we could end up with a stale State with a DbEntry.QueryVector saved with the previous schema
-					// so the cardinality of DbEntry.QueryVector and PoseValues don't match. In that case we just use PoseValues as query to have all costs set to zero
-					const bool bIsQueryVectorValid = DbEntry.QueryVector.Num() == PoseValues.Num();
-					CompareFeatureVectors(PoseValues, bIsQueryVectorValid ? DbEntry.QueryVector : PoseValues, SearchIndex->WeightsSqrt, Row->CostVector);
-
-					if (SearchIndexAsset->Type == ESearchIndexAssetType::Sequence)
+					const FPoseSearchIndexAsset* SearchIndexAsset = SearchIndex->FindAssetForPose(PoseEntry.DbPoseIdx);
+					if (SearchIndexAsset)
 					{
-						const FPoseSearchDatabaseSequence& DbSequence = Database->GetSequenceSourceAsset(SearchIndexAsset);
-						Row->AssetType = ESearchIndexAssetType::Sequence;
-						Row->AssetName = DbSequence.Sequence->GetName();
-						Row->AssetPath = DbSequence.Sequence->GetPathName();
-						Row->AnimFrame = DbSequence.Sequence->GetFrameAtTime(Time);
-						Row->AnimPercentage = Time / DbSequence.Sequence->GetPlayLength();
-						Row->bLooping = DbSequence.Sequence->bLoop;
-						Row->BlendParameters = FVector::Zero();
-					}
-					else if (SearchIndexAsset->Type == ESearchIndexAssetType::BlendSpace)
-					{
-						const FPoseSearchDatabaseBlendSpace& DbBlendSpace = Database->GetBlendSpaceSourceAsset(SearchIndexAsset);
-						Row->AssetType = ESearchIndexAssetType::BlendSpace;
-						Row->AssetName = DbBlendSpace.BlendSpace->GetName();
-						Row->AssetPath = DbBlendSpace.BlendSpace->GetPathName();
-						Row->AnimFrame = 0; // There is no frame index associated with a blendspace
-						Row->AnimPercentage = 0.0f;
-						Row->bLooping = DbBlendSpace.BlendSpace->bLoop;
-						Row->BlendParameters = SearchIndexAsset->BlendParameters;
-					}
-					else
-					{
-						checkNoEntry();
+						TSharedRef<FDebuggerDatabaseRowData>& Row = UnfilteredDatabaseRows.Add_GetRef(MakeShared<FDebuggerDatabaseRowData>());
+
+						const float Time = SearchIndex->GetAssetTime(PoseEntry.DbPoseIdx, SearchIndexAsset);
+
+						Row->PoseIdx = PoseEntry.DbPoseIdx;
+						Row->SourceDatabase = Database;
+						Row->DatabaseName = Database->GetName();
+						Row->DatabasePath = Database->GetPathName();
+						Row->PoseEntryFlags = PoseEntry.Flags;
+						Row->DbAssetIdx = SearchIndexAsset->SourceAssetIdx;
+						Row->AssetTime = Time;
+						Row->bMirrored = SearchIndexAsset->bMirrored;
+						Row->PoseCost = PoseEntry.Cost;
+
+						Row->CostVector.SetNum(Database->Schema->SchemaCardinality);
+						TConstArrayView<const float> PoseValues = SearchIndex->GetPoseValues(PoseEntry.DbPoseIdx);
+
+						// in case we modify the schema while PIE is paused and displaying the Pose Search Editor, we could end up with a stale State with a DbEntry.QueryVector saved with the previous schema
+						// so the cardinality of DbEntry.QueryVector and PoseValues don't match. In that case we just use PoseValues as query to have all costs set to zero
+						const bool bIsQueryVectorValid = DbEntry.QueryVector.Num() == PoseValues.Num();
+						CompareFeatureVectors(PoseValues, bIsQueryVectorValid ? DbEntry.QueryVector : PoseValues, SearchIndex->WeightsSqrt, Row->CostVector);
+
+						if (SearchIndexAsset->Type == ESearchIndexAssetType::Sequence)
+						{
+							const FPoseSearchDatabaseSequence& DbSequence = Database->GetSequenceSourceAsset(SearchIndexAsset);
+							Row->AssetType = ESearchIndexAssetType::Sequence;
+							Row->AssetName = DbSequence.Sequence->GetName();
+							Row->AssetPath = DbSequence.Sequence->GetPathName();
+							Row->AnimFrame = DbSequence.Sequence->GetFrameAtTime(Time);
+							Row->AnimPercentage = Time / DbSequence.Sequence->GetPlayLength();
+							Row->bLooping = DbSequence.Sequence->bLoop;
+							Row->BlendParameters = FVector::Zero();
+						}
+						else if (SearchIndexAsset->Type == ESearchIndexAssetType::BlendSpace)
+						{
+							const FPoseSearchDatabaseBlendSpace& DbBlendSpace = Database->GetBlendSpaceSourceAsset(SearchIndexAsset);
+							Row->AssetType = ESearchIndexAssetType::BlendSpace;
+							Row->AssetName = DbBlendSpace.BlendSpace->GetName();
+							Row->AssetPath = DbBlendSpace.BlendSpace->GetPathName();
+							Row->AnimFrame = 0; // There is no frame index associated with a blendspace
+							Row->AnimPercentage = 0.0f;
+							Row->bLooping = DbBlendSpace.BlendSpace->bLoop;
+							Row->BlendParameters = SearchIndexAsset->BlendParameters;
+						}
+						else
+						{
+							checkNoEntry();
+						}
 					}
 				}
 			}
@@ -898,10 +900,13 @@ void SDebuggerDatabaseView::Update(const FTraceMotionMatchingStateMessage& State
 		for (const FTraceMotionMatchingStateDatabaseEntry& DbEntry : State.DatabaseEntries)
 		{
 			const UPoseSearchDatabase* Database = FTraceMotionMatchingState::GetObjectFromId<UPoseSearchDatabase>(DbEntry.DatabaseId);
-			const FPoseSearchIndex* SearchIndex = Database->GetSearchIndex();
-			if (SearchIndex && SearchIndex->IsValid())
+			if (Database)
 			{
-				Database->Schema->ComputeCostBreakdowns(CostBreakDownData);
+				const FPoseSearchIndex* SearchIndex = Database->GetSearchIndex();
+				if (SearchIndex && SearchIndex->IsValid())
+				{
+					Database->Schema->ComputeCostBreakdowns(CostBreakDownData);
+				}
 			}
 		}
 
@@ -1513,7 +1518,7 @@ void SDebuggerDetailsView::UpdateReflection(const FTraceMotionMatchingStateMessa
 	for (const FTraceMotionMatchingStateDatabaseEntry& DbEntry : State.DatabaseEntries)
 	{
 		const UPoseSearchDatabase* Database = FTraceMotionMatchingState::GetObjectFromId<UPoseSearchDatabase>(DbEntry.DatabaseId);
-		if (Database == CurrentDatabase)
+		if (Database && Database == CurrentDatabase)
 		{
 			const FPoseSearchIndex* SearchIndex = Database->GetSearchIndex();
 			if (SearchIndex && SearchIndex->IsValid())
@@ -1788,7 +1793,7 @@ void SDebuggerView::DrawFeatures(
 			for (const FTraceMotionMatchingStateDatabaseEntry& DbEntry : State.DatabaseEntries)
 			{
 				const UPoseSearchDatabase* Database = FTraceMotionMatchingState::GetObjectFromId<UPoseSearchDatabase>(DbEntry.DatabaseId);
-				if (Database == CurrentDatabase)
+				if (Database && Database == CurrentDatabase)
 				{
 					const FPoseSearchIndex* SearchIndex = Database->GetSearchIndex();
 					if (SearchIndex && SearchIndex->IsValid() && DbEntry.QueryVector.Num() == Database->Schema->SchemaCardinality)
@@ -2238,6 +2243,11 @@ const UPoseSearchSearchableAsset* FDebuggerViewModel::GetSearchableAsset() const
 
 void FDebuggerViewModel::ShowSelectedSkeleton(const UPoseSearchDatabase* Database, int32 DbPoseIdx, float Time)
 {
+	if (!Database)
+	{
+		return;
+	}
+
 	UPoseSearchMeshComponent* Component = Skeletons[SelectedPose].Component.Get();
 	if (!Component)
 	{
@@ -2669,7 +2679,7 @@ const USkinnedMeshComponent* FDebuggerViewModel::GetMeshComponent() const
 void FDebuggerViewModel::FillCompactPoseAndComponentRefRotations()
 {
 	const UPoseSearchDatabase* Database = GetCurrentDatabase();
-	if (Database != nullptr)
+	if (Database)
 	{
 		UMirrorDataTable* MirrorDataTable = Database->Schema->MirrorDataTable;
 		if (MirrorDataTable != nullptr)
