Fix for Issue 214761: Localization warnings in PoseSearchFeatureChannels.cpp

diff --git a/Engine/Plugins/Experimental/Animation/PoseSearch/Source/Runtime/Private/PoseSearchFeatureChannels.cpp b/Engine/Plugins/Experimental/Animation/PoseSearch/Source/Runtime/Private/PoseSearchFeatureChannels.cpp
--- a/Engine/Plugins/Experimental/Animation/PoseSearch/Source/Runtime/Private/PoseSearchFeatureChannels.cpp
+++ b/Engine/Plugins/Experimental/Animation/PoseSearch/Source/Runtime/Private/PoseSearchFeatureChannels.cpp
@@ -1219,7 +1219,7 @@ void UPoseSearchFeatureChannel_Pose::ComputeCostBreakdowns(UE::PoseSearch::ICost
 			{
 				for (int32 SubsampleIdx = 0; SubsampleIdx != SampleTimes.Num(); ++SubsampleIdx)
 				{
-					CostBreakDownData.AddEntireBreakDownSection(FText::Format(LOCTEXT("ColumnLabelPoseChannelVelocity", "{0} Pha {1}"), FText::FromName(SampledBone.Reference.BoneName), SampleTimes[SubsampleIdx]), Schema, DataOffset, FFeatureVectorHelper::EncodeVector2DCardinality);
+					CostBreakDownData.AddEntireBreakDownSection(FText::Format(LOCTEXT("ColumnLabelPoseChannelPhase", "{0} Pha {1}"), FText::FromName(SampledBone.Reference.BoneName), SampleTimes[SubsampleIdx]), Schema, DataOffset, FFeatureVectorHelper::EncodeVector2DCardinality);
 					DataOffset += FFeatureVectorHelper::EncodeVector2DCardinality;
 				}
 			}
