motion matching - reporting discarded poses to pose search debugger

diff --git a/Engine/Plugins/Experimental/Animation/PoseSearch/Source/Editor/Private/PoseSearchDebugger.cpp b/Engine/Plugins/Experimental/Animation/PoseSearch/Source/Editor/Private/PoseSearchDebugger.cpp
--- a/Engine/Plugins/Experimental/Animation/PoseSearch/Source/Editor/Private/PoseSearchDebugger.cpp
+++ b/Engine/Plugins/Experimental/Animation/PoseSearch/Source/Editor/Private/PoseSearchDebugger.cpp
@@ -34,12 +34,12 @@
 namespace UE::PoseSearch
 {
 
-static FLinearColor LinearColorBlend(FLinearColor LinearColorA, FLinearColor LinearColorB, float BlendParam)
+inline static FLinearColor LinearColorBlend(FLinearColor LinearColorA, FLinearColor LinearColorB, float BlendParam)
 {
 	return LinearColorA + (LinearColorB - LinearColorA) * BlendParam;
 }
 
-static void LinearColorBlend(FLinearColor LinearColorA, FLinearColor LinearColorB, TConstArrayView<const float> BlendParam, TArray<FLinearColor>& BlendedColors)
+inline static void LinearColorBlend(FLinearColor LinearColorA, FLinearColor LinearColorB, TConstArrayView<const float> BlendParam, TArray<FLinearColor>& BlendedColors)
 {
 	const int32 Num = BlendParam.Num();
 	BlendedColors.SetNumUninitialized(Num);
@@ -258,7 +258,7 @@ public:
 	ESearchIndexAssetType AssetType = ESearchIndexAssetType::Invalid;
 	int32 PoseIdx = 0;
 	TWeakObjectPtr<const UPoseSearchDatabase> SourceDatabase = nullptr;
-	FTraceMotionMatchingStatePoseEntry::EFlags PoseEntryFlags = FTraceMotionMatchingStatePoseEntry::EFlags::None;
+	EPoseCandidateFlags PoseCandidateFlags = EPoseCandidateFlags::None;
 	FString DatabaseName = "";
 	FString DatabasePath = "";
 	FString AssetName = "";
@@ -617,6 +617,58 @@ namespace DebuggerDatabaseColumns
 			}
 		}
 	};
+
+	struct FPoseCandidateFlags : ITextColumn
+	{
+		using ITextColumn::ITextColumn;
+
+		virtual FText GetLabel() const override { return LOCTEXT("ColumnLabelPoseCandidateFlags", "Flags"); }
+
+		virtual FSortPredicate GetSortPredicate() const override
+		{
+			return [](const FRowDataRef& Row0, const FRowDataRef& Row1) -> bool { return Row0->PoseCandidateFlags < Row1->PoseCandidateFlags; };
+		}
+
+		virtual FText GetRowText(const FRowDataRef& Row) const override
+		{
+			if (EnumHasAnyFlags(Row->PoseCandidateFlags, EPoseCandidateFlags::AnyDiscardedMask))
+			{
+				auto AddDelimiter = [](FTextBuilder& TextBuilder, bool& bNeedDelimiter)
+				{
+					if (bNeedDelimiter)
+					{
+						TextBuilder.AppendLine(LOCTEXT("DiscardedBy_Delimiter", " | "));
+					}
+					bNeedDelimiter = true;
+				};
+
+				bool bNeedDelimiter = false;
+
+				FTextBuilder TextBuilder;
+				if (EnumHasAnyFlags(Row->PoseCandidateFlags, EPoseCandidateFlags::DiscardedBy_PoseJumpThresholdTime))
+				{
+					AddDelimiter(TextBuilder, bNeedDelimiter);
+					TextBuilder.AppendLine(LOCTEXT("DiscardedBy_PoseJumpThresholdTime", "PoseJumpThresholdTime"));
+				}
+				
+				if (EnumHasAnyFlags(Row->PoseCandidateFlags, EPoseCandidateFlags::DiscardedBy_PoseReselectHistory))
+				{
+					AddDelimiter(TextBuilder, bNeedDelimiter);
+					TextBuilder.AppendLine(LOCTEXT("DiscardedBy_PoseReselectHistory", "PoseReselectHistory"));
+				}
+
+				if (EnumHasAnyFlags(Row->PoseCandidateFlags, EPoseCandidateFlags::DiscardedBy_BlockTransition))
+				{
+					AddDelimiter(TextBuilder, bNeedDelimiter);
+					TextBuilder.AppendLine(LOCTEXT("DiscardedBy_BlockTransition", "BlockTransition"));
+				}
+
+				return TextBuilder.ToText();
+			}
+
+			return FText::GetEmpty();
+		}
+	};
 }
 
 
@@ -817,6 +869,11 @@ private:
 
 void SDebuggerDatabaseView::Update(const FTraceMotionMatchingStateMessage& State)
 {
+	// row cost color palette
+	static const FLinearColor DiscardedRowColor(0.314f, 0.314f, 0.314f); // darker gray
+	static const FLinearColor BestScoreRowColor = FLinearColor::Green;
+	static const FLinearColor WorstScoreRowColor = FLinearColor::Red;
+
 	using namespace DebuggerDatabaseColumns;
 
 	bool bIsVerbose = false;
@@ -848,7 +905,7 @@ void SDebuggerDatabaseView::Update(const FTraceMotionMatchingStateMessage& State
 						Row->SourceDatabase = Database;
 						Row->DatabaseName = Database->GetName();
 						Row->DatabasePath = Database->GetPathName();
-						Row->PoseEntryFlags = PoseEntry.Flags;
+						Row->PoseCandidateFlags = PoseEntry.PoseCandidateFlags;
 						Row->DbAssetIdx = SearchIndexAsset->SourceAssetIdx;
 						Row->AssetTime = Time;
 						Row->bMirrored = SearchIndexAsset->bMirrored;
@@ -922,13 +979,20 @@ void SDebuggerDatabaseView::Update(const FTraceMotionMatchingStateMessage& State
 		{
 			ArrayMinMax(UnfilteredRow->CostBreakdowns, MinCostBreakdowns, MaxCostBreakdowns, UE_MAX_FLT);
 		}
-
+		
 		TArray<float> CostBreakdownsColorBlend;
 		CostBreakdownsColorBlend.Init(0, CostBreakdownsCardinality);
 		for (auto& UnfilteredRow : UnfilteredDatabaseRows)
 		{
-			ArraySafeNormalize(UnfilteredRow->CostBreakdowns, MinCostBreakdowns, MaxCostBreakdowns, CostBreakdownsColorBlend);
-			LinearColorBlend(FLinearColor::Green, FLinearColor::Red, CostBreakdownsColorBlend, UnfilteredRow->CostBreakdownsColors);
+			if (EnumHasAnyFlags(UnfilteredRow->PoseCandidateFlags, EPoseCandidateFlags::AnyValidMask))
+			{
+				ArraySafeNormalize(UnfilteredRow->CostBreakdowns, MinCostBreakdowns, MaxCostBreakdowns, CostBreakdownsColorBlend);
+				LinearColorBlend(BestScoreRowColor, WorstScoreRowColor, CostBreakdownsColorBlend, UnfilteredRow->CostBreakdownsColors);
+			}
+			else
+			{
+				UnfilteredRow->CostBreakdownsColors.Init(DiscardedRowColor, CostBreakdownsCardinality);
+			}
 		}
 
 		float MinCost = UE_MAX_FLT;
@@ -943,8 +1007,15 @@ void SDebuggerDatabaseView::Update(const FTraceMotionMatchingStateMessage& State
 		const float DeltaCost = MaxCost - MinCost;
 		for (auto& UnfilteredRow : UnfilteredDatabaseRows)
 		{
-			const float CostColorBlend = DeltaCost > UE_KINDA_SMALL_NUMBER ? (UnfilteredRow->PoseCost.GetTotalCost() - MinCost) / DeltaCost : 0.f;
-			UnfilteredRow->CostColor = UE::PoseSearch::LinearColorBlend(FLinearColor::Green, FLinearColor::Red, CostColorBlend);
+			if (EnumHasAnyFlags(UnfilteredRow->PoseCandidateFlags, EPoseCandidateFlags::AnyValidMask))
+			{
+				const float CostColorBlend = DeltaCost > UE_KINDA_SMALL_NUMBER ? (UnfilteredRow->PoseCost.GetTotalCost() - MinCost) / DeltaCost : 0.f;
+				UnfilteredRow->CostColor = UE::PoseSearch::LinearColorBlend(BestScoreRowColor, WorstScoreRowColor, CostColorBlend);
+			}
+			else
+			{
+				UnfilteredRow->CostColor = DiscardedRowColor;
+			}
 		}
 	}
 
@@ -979,6 +1050,7 @@ void SDebuggerDatabaseView::Update(const FTraceMotionMatchingStateMessage& State
 		AddColumn(MakeShared<FLooping>(ColumnIdx++));
 		AddColumn(MakeShared<FPoseIdx>(ColumnIdx++));
 		AddColumn(MakeShared<FBlendParameters>(ColumnIdx++));
+		AddColumn(MakeShared<FPoseCandidateFlags>(ColumnIdx++));
 
 		SortColumn = CostColumn->ColumnId;
 
@@ -1114,13 +1186,13 @@ void SDebuggerDatabaseView::PopulateViewRows()
 	for (const auto& UnfilteredRow : UnfilteredDatabaseRows)
 	{
 		bool bTryAddToFilteredDatabaseViewRows = true;
-		if (EnumHasAnyFlags(UnfilteredRow->PoseEntryFlags, FTraceMotionMatchingStatePoseEntry::EFlags::ContinuingPose))
+		if (EnumHasAnyFlags(UnfilteredRow->PoseCandidateFlags, EPoseCandidateFlags::Valid_ContinuingPose))
 		{
 			ContinuingPoseView.Rows.Add(UnfilteredRow);
 			bTryAddToFilteredDatabaseViewRows = false;
 		}
 
-		if (EnumHasAnyFlags(UnfilteredRow->PoseEntryFlags, FTraceMotionMatchingStatePoseEntry::EFlags::CurrentPose))
+		if (EnumHasAnyFlags(UnfilteredRow->PoseCandidateFlags, EPoseCandidateFlags::Valid_CurrentPose))
 		{
 			ActiveView.Rows.Add(UnfilteredRow);
 			bTryAddToFilteredDatabaseViewRows = false;
diff --git a/Engine/Plugins/Experimental/Animation/PoseSearch/Source/Runtime/Private/PoseSearch.cpp b/Engine/Plugins/Experimental/Animation/PoseSearch/Source/Runtime/Private/PoseSearch.cpp
--- a/Engine/Plugins/Experimental/Animation/PoseSearch/Source/Runtime/Private/PoseSearch.cpp
+++ b/Engine/Plugins/Experimental/Animation/PoseSearch/Source/Runtime/Private/PoseSearch.cpp
@@ -184,22 +184,21 @@ FORCEINLINE auto LowerBound(IteratorType First, IteratorType Last, const ValueTy
 	return LowerBound(First, Last, Value, FIdentityFunctor(), TLess<>());
 }
 
-static int32 PopulateNonSelectableIdx(size_t* NonSelectableIdx, int32 NonSelectableIdxSize, const FSearchContext& SearchContext, const UPoseSearchDatabase* Database)
+static int32 PopulateNonSelectableIdx(size_t* NonSelectableIdx, int32 NonSelectableIdxSize, FSearchContext& SearchContext, const UPoseSearchDatabase* Database, TArrayView<const float> QueryValues)
 {
 	check(Database);
 
 	int32 NonSelectableIdxUsedSize = 0;
 	if (SearchContext.IsCurrentResultFromDatabase(Database) && SearchContext.CurrentResult.IsValid() && SearchContext.PoseJumpThresholdTime > 0.f)
 	{
-		const UPoseSearchDatabase* CurrentResultDatabase = SearchContext.CurrentResult.Database.Get();
 		const FPoseSearchIndexAsset* CurrentIndexAsset = SearchContext.CurrentResult.SearchIndexAsset;
-		const int32 PoseJumpIndexThreshold = FMath::FloorToInt(SearchContext.PoseJumpThresholdTime / CurrentResultDatabase->Schema->SamplingInterval);
-		const bool IsLooping = CurrentResultDatabase->IsSourceAssetLooping(CurrentIndexAsset);
+		const int32 PoseJumpIndexThreshold = FMath::FloorToInt(SearchContext.PoseJumpThresholdTime / Database->Schema->SamplingInterval);
+		const bool IsLooping = Database->IsSourceAssetLooping(CurrentIndexAsset);
 
 		for (int32 i = -PoseJumpIndexThreshold; i <= -1; ++i)
 		{
 			const int32 PoseIdx = SearchContext.CurrentResult.PoseIdx + i;
-			const float DeltaTime = i * CurrentResultDatabase->Schema->SamplingInterval;
+			const float DeltaTime = i * Database->Schema->SamplingInterval;
 
 			// @todo: should we use the quantized time associated to InOutMotionMatchingState.DbPoseIdx instead of InOutMotionMatchingState.AssetPlayerTime?
 			float PoseAssetPlayerTime = SearchContext.CurrentResult.AssetTime + DeltaTime;
@@ -216,6 +215,11 @@ static int32 PopulateNonSelectableIdx(size_t* NonSelectableIdx, int32 NonSelecta
 				if (NonSelectableIdxUsedSize < NonSelectableIdxSize)
 				{
 					NonSelectableIdx[NonSelectableIdxUsedSize++] = PoseIdx;
+
+#if UE_POSE_SEARCH_TRACE_ENABLED
+					const FPoseSearchCost PoseCost = Database->ComparePoses(SearchContext, PoseIdx, EPoseComparisonFlags::None, QueryValues);
+					SearchContext.BestCandidates.Add(PoseCost, PoseIdx, Database, EPoseCandidateFlags::DiscardedBy_PoseJumpThresholdTime);
+#endif
 				}
 				else
 				{
@@ -232,7 +236,7 @@ static int32 PopulateNonSelectableIdx(size_t* NonSelectableIdx, int32 NonSelecta
 		for (int32 i = 0; i <= PoseJumpIndexThreshold; ++i)
 		{
 			const int32 PoseIdx = SearchContext.CurrentResult.PoseIdx + i;
-			const float DeltaTime = i * CurrentResultDatabase->Schema->SamplingInterval;
+			const float DeltaTime = i * Database->Schema->SamplingInterval;
 
 			// @todo: should we use the quantized time associated to InOutMotionMatchingState.DbPoseIdx instead of InOutMotionMatchingState.AssetPlayerTime?
 			float PoseAssetPlayerTime = SearchContext.CurrentResult.AssetTime + DeltaTime;
@@ -249,6 +253,11 @@ static int32 PopulateNonSelectableIdx(size_t* NonSelectableIdx, int32 NonSelecta
 				if (NonSelectableIdxUsedSize < NonSelectableIdxSize)
 				{
 					NonSelectableIdx[NonSelectableIdxUsedSize++] = PoseIdx;
+
+#if UE_POSE_SEARCH_TRACE_ENABLED
+					const FPoseSearchCost PoseCost = Database->ComparePoses(SearchContext, PoseIdx, EPoseComparisonFlags::None, QueryValues);
+					SearchContext.BestCandidates.Add(PoseCost, PoseIdx, Database, EPoseCandidateFlags::DiscardedBy_PoseJumpThresholdTime);
+#endif
 				}
 				else
 				{
@@ -274,6 +283,11 @@ static int32 PopulateNonSelectableIdx(size_t* NonSelectableIdx, int32 NonSelecta
 				if (NonSelectableIdxUsedSize < NonSelectableIdxSize)
 				{
 					NonSelectableIdx[NonSelectableIdxUsedSize++] = HistoricalPoseIndex.PoseIndex;
+
+#if UE_POSE_SEARCH_TRACE_ENABLED
+					const FPoseSearchCost PoseCost = Database->ComparePoses(SearchContext, HistoricalPoseIndex.PoseIndex, EPoseComparisonFlags::None, QueryValues);
+					SearchContext.BestCandidates.Add(PoseCost, HistoricalPoseIndex.PoseIndex, Database, EPoseCandidateFlags::DiscardedBy_PoseReselectHistory);
+#endif
 				}
 				else
 				{
@@ -646,8 +660,7 @@ UE::PoseSearch::FSearchResult UPoseSearchSequenceMetaData::Search(UE::PoseSearch
 				continue;
 			}
 
-			FPoseSearchCost PoseCost = ComparePoses(PoseIdx, EPoseComparisonFlags::ContinuingPose, QueryValues);
-
+			const FPoseSearchCost PoseCost = ComparePoses(PoseIdx, EPoseComparisonFlags::ContinuingPose, QueryValues);
 			if (PoseCost < BestPoseCost)
 			{
 				BestPoseCost = PoseCost;
@@ -1720,7 +1733,7 @@ UE::PoseSearch::FSearchResult UPoseSearchDatabase::SearchPCAKDTree(UE::PoseSearc
 	// there's no point in performing the search if CurrentBestTotalCost is already better than that
 	if (SearchContext.GetCurrentBestTotalCost() > SearchIndex->MinCostAddend)
 	{
-		const int NonSelectableIdxUsedSize = PopulateNonSelectableIdx(NonSelectableIdxData, NonSelectableIdxDataSize, SearchContext, this);
+		const int NonSelectableIdxUsedSize = PopulateNonSelectableIdx(NonSelectableIdxData, NonSelectableIdxDataSize, SearchContext, this, QueryValues);
 		TArrayView<size_t> NonSelectableIdx(NonSelectableIdxData, NonSelectableIdxUsedSize);
 		const RowMajorVectorMapConst MapWeightsSqrt(SearchIndex->WeightsSqrt.GetData(), 1, NumDimensions);
 		FKDTree::KNNResultSet ResultSet(ClampedKDTreeQueryNumNeighbors, ResultIndexes, ResultDistanceSqr, NonSelectableIdx);
@@ -1746,15 +1759,14 @@ UE::PoseSearch::FSearchResult UPoseSearchDatabase::SearchPCAKDTree(UE::PoseSearc
 
 			if (EnumHasAnyFlags(Metadata.Flags, EPoseSearchPoseFlags::BlockTransition))
 			{
+#if UE_POSE_SEARCH_TRACE_ENABLED
+				const FPoseSearchCost PoseCost = ComparePoses(SearchContext, PoseIdx, EPoseComparisonFlags::None, QueryValues);
+				SearchContext.BestCandidates.Add(PoseCost, PoseIdx, this, EPoseCandidateFlags::DiscardedBy_BlockTransition);
+#endif
 				continue;
 			}
 
-			FPoseSearchCost PoseCost = ComparePoses(
-				SearchContext,
-				PoseIdx,
-				EPoseComparisonFlags::None,
-				QueryValues);
-
+			const FPoseSearchCost PoseCost = ComparePoses(SearchContext, PoseIdx, EPoseComparisonFlags::None, QueryValues);
 			if (PoseCost < Result.PoseCost)
 			{
 				Result.PoseCost = PoseCost;
@@ -1762,7 +1774,7 @@ UE::PoseSearch::FSearchResult UPoseSearchDatabase::SearchPCAKDTree(UE::PoseSearc
 			}
 
 #if UE_POSE_SEARCH_TRACE_ENABLED
-			SearchContext.BestCandidates.Add(PoseCost, PoseIdx, this);
+			SearchContext.BestCandidates.Add(PoseCost, PoseIdx, this, EPoseCandidateFlags::Valid_Pose);
 #endif
 		}
 
@@ -1820,7 +1832,7 @@ UE::PoseSearch::FSearchResult UPoseSearchDatabase::SearchBruteForce(UE::PoseSear
 
 	constexpr int NonSelectableIdxDataSize = 128;
 	size_t* NonSelectableIdxData((size_t*)FMemory_Alloca(NonSelectableIdxDataSize * sizeof(size_t)));
-	const int NonSelectableIdxUsedSize = PopulateNonSelectableIdx(NonSelectableIdxData, NonSelectableIdxDataSize, SearchContext, this);
+	const int NonSelectableIdxUsedSize = PopulateNonSelectableIdx(NonSelectableIdxData, NonSelectableIdxDataSize, SearchContext, this, QueryValues);
 	TArrayView<size_t> NonSelectableIdx(NonSelectableIdxData, NonSelectableIdxUsedSize);
 
 	// since any PoseCost calculated here is at least SearchIndex->MinCostAddend,
@@ -1837,6 +1849,10 @@ UE::PoseSearch::FSearchResult UPoseSearchDatabase::SearchBruteForce(UE::PoseSear
 
 				if (EnumHasAnyFlags(Metadata.Flags, EPoseSearchPoseFlags::BlockTransition))
 				{
+#if UE_POSE_SEARCH_TRACE_ENABLED
+					const FPoseSearchCost PoseCost = ComparePoses(SearchContext, PoseIdx, EPoseComparisonFlags::None, QueryValues);
+					SearchContext.BestCandidates.Add(PoseCost, PoseIdx, this, EPoseCandidateFlags::DiscardedBy_BlockTransition);
+#endif
 					continue;
 				}
 
@@ -1845,12 +1861,7 @@ UE::PoseSearch::FSearchResult UPoseSearchDatabase::SearchBruteForce(UE::PoseSear
 					continue;
 				}
 
-				FPoseSearchCost PoseCost = ComparePoses(
-					SearchContext,
-					PoseIdx, 
-					EPoseComparisonFlags::None,
-					QueryValues);
-
+				const FPoseSearchCost PoseCost = ComparePoses(SearchContext, PoseIdx, EPoseComparisonFlags::None, QueryValues);
 				if (PoseCost < Result.PoseCost)
 				{
 					Result.PoseCost = PoseCost;
@@ -1860,7 +1871,7 @@ UE::PoseSearch::FSearchResult UPoseSearchDatabase::SearchBruteForce(UE::PoseSear
 #if UE_POSE_SEARCH_TRACE_ENABLED
 				if (PoseSearchMode == EPoseSearchMode::BruteForce)
 				{
-					SearchContext.BestCandidates.Add(PoseCost, PoseIdx, this);
+					SearchContext.BestCandidates.Add(PoseCost, PoseIdx, this, EPoseCandidateFlags::Valid_Pose);
 				}
 #endif
 			}
diff --git a/Engine/Plugins/Experimental/Animation/PoseSearch/Source/Runtime/Private/PoseSearchLibrary.cpp b/Engine/Plugins/Experimental/Animation/PoseSearch/Source/Runtime/Private/PoseSearchLibrary.cpp
--- a/Engine/Plugins/Experimental/Animation/PoseSearch/Source/Runtime/Private/PoseSearchLibrary.cpp
+++ b/Engine/Plugins/Experimental/Animation/PoseSearch/Source/Runtime/Private/PoseSearchLibrary.cpp
@@ -225,6 +225,7 @@ static void TraceMotionMatchingState(
 		FTraceMotionMatchingStatePoseEntry PoseEntry;
 		PoseEntry.DbPoseIdx = PoseCandidate.PoseIdx;
 		PoseEntry.Cost = PoseCandidate.Cost;
+		PoseEntry.PoseCandidateFlags = PoseCandidate.PoseCandidateFlags;
 
 		DbEntry.PoseEntries.Add(PoseEntry);
 	}
@@ -240,7 +241,7 @@ static void TraceMotionMatchingState(
 		FTraceMotionMatchingStatePoseEntry& PoseEntry = DbEntry.PoseEntries[LastResultPoseEntryIdx];
 
 		PoseEntry.Cost = MotionMatchingState.CurrentSearchResult.ContinuingPoseCost;
-		PoseEntry.Flags = FTraceMotionMatchingStatePoseEntry::EFlags::ContinuingPose;
+		PoseEntry.PoseCandidateFlags = EPoseCandidateFlags::Valid_ContinuingPose;
 	}
 
 	if (MotionMatchingState.CurrentSearchResult.PoseCost.IsValid())
@@ -252,7 +253,7 @@ static void TraceMotionMatchingState(
 		FTraceMotionMatchingStatePoseEntry& PoseEntry = DbEntry.PoseEntries[PoseEntryIdx];
 
 		PoseEntry.Cost = MotionMatchingState.CurrentSearchResult.PoseCost;
-		PoseEntry.Flags = FTraceMotionMatchingStatePoseEntry::EFlags::CurrentPose;
+		PoseEntry.PoseCandidateFlags = EPoseCandidateFlags::Valid_CurrentPose;
 
 		TraceState.CurrentDbEntryIdx = DbEntryIdx;
 		TraceState.CurrentPoseEntryIdx = PoseEntryIdx;
diff --git a/Engine/Plugins/Experimental/Animation/PoseSearch/Source/Runtime/Private/Trace/PoseSearchTraceLogger.cpp b/Engine/Plugins/Experimental/Animation/PoseSearch/Source/Runtime/Private/Trace/PoseSearchTraceLogger.cpp
--- a/Engine/Plugins/Experimental/Animation/PoseSearch/Source/Runtime/Private/Trace/PoseSearchTraceLogger.cpp
+++ b/Engine/Plugins/Experimental/Animation/PoseSearch/Source/Runtime/Private/Trace/PoseSearchTraceLogger.cpp
@@ -53,7 +53,7 @@ FArchive& operator<<(FArchive& Ar, FTraceMotionMatchingStatePoseEntry& Entry)
 {
 	Ar << Entry.DbPoseIdx;
 	FPoseSearchCost::StaticStruct()->SerializeItem(Ar, &Entry.Cost, nullptr);
-	Ar << Entry.Flags;
+	Ar << Entry.PoseCandidateFlags;
 	return Ar;
 }
 
diff --git a/Engine/Plugins/Experimental/Animation/PoseSearch/Source/Runtime/Private/Trace/PoseSearchTraceLogger.h b/Engine/Plugins/Experimental/Animation/PoseSearch/Source/Runtime/Private/Trace/PoseSearchTraceLogger.h
--- a/Engine/Plugins/Experimental/Animation/PoseSearch/Source/Runtime/Private/Trace/PoseSearchTraceLogger.h
+++ b/Engine/Plugins/Experimental/Animation/PoseSearch/Source/Runtime/Private/Trace/PoseSearchTraceLogger.h
@@ -39,19 +39,12 @@ POSESEARCH_API FArchive& operator<<(FArchive& Ar, FTraceMessage& State);
 
 struct POSESEARCH_API FTraceMotionMatchingStatePoseEntry
 {
-	enum class EFlags : uint8
-	{
-		None			= 0,
-		ContinuingPose	= 1 << 0,
-		CurrentPose		= 1 << 1,
-	};
 	int32 DbPoseIdx = INDEX_NONE;
-	EFlags Flags = EFlags::None;
+	EPoseCandidateFlags PoseCandidateFlags = EPoseCandidateFlags::None;
 	FPoseSearchCost Cost;
 
 	bool operator==(const FTraceMotionMatchingStatePoseEntry& Other) const { return DbPoseIdx == Other.DbPoseIdx; }
 };
-ENUM_CLASS_FLAGS(FTraceMotionMatchingStatePoseEntry::EFlags);
 POSESEARCH_API FArchive& operator<<(FArchive& Ar, FTraceMotionMatchingStatePoseEntry& Entry);
 
 struct POSESEARCH_API FTraceMotionMatchingStateDatabaseEntry
diff --git a/Engine/Plugins/Experimental/Animation/PoseSearch/Source/Runtime/Public/PoseSearch/PoseSearch.h b/Engine/Plugins/Experimental/Animation/PoseSearch/Source/Runtime/Public/PoseSearch/PoseSearch.h
--- a/Engine/Plugins/Experimental/Animation/PoseSearch/Source/Runtime/Public/PoseSearch/PoseSearch.h
+++ b/Engine/Plugins/Experimental/Animation/PoseSearch/Source/Runtime/Public/PoseSearch/PoseSearch.h
@@ -1151,6 +1151,24 @@ public: // UObject
 namespace UE::PoseSearch
 {
 
+enum class EPoseCandidateFlags : uint8
+{
+	None = 0,
+
+	Valid_Pose = 1 << 0,
+	Valid_ContinuingPose = 1 << 1,
+	Valid_CurrentPose = 1 << 2,
+
+	AnyValidMask = Valid_Pose | Valid_ContinuingPose | Valid_CurrentPose,
+
+	DiscardedBy_PoseJumpThresholdTime = 1 << 3,
+	DiscardedBy_PoseReselectHistory = 1 << 4,
+	DiscardedBy_BlockTransition = 1 << 5,
+
+	AnyDiscardedMask = DiscardedBy_PoseJumpThresholdTime | DiscardedBy_PoseReselectHistory | DiscardedBy_BlockTransition,
+};
+ENUM_CLASS_FLAGS(EPoseCandidateFlags);
+
 /** Helper class for extracting and encoding features into a float buffer */
 class POSESEARCH_API FFeatureVectorHelper
 {
@@ -1314,6 +1332,7 @@ public:
 		FPoseSearchCost Cost;
 		int32 PoseIdx = 0;
 		const UPoseSearchDatabase* Database = nullptr;
+		EPoseCandidateFlags PoseCandidateFlags = EPoseCandidateFlags::None;
 
 		bool operator<(const FPoseCandidate& Other) const { return Other.Cost < Cost; } // Reverse compare because BestCandidates is a max heap
 		bool operator==(const FSearchResult& SearchResult) const { return (PoseIdx == SearchResult.PoseIdx) && (Database == SearchResult.Database.Get()); }
@@ -1326,7 +1345,7 @@ public:
 
 		int32 MaxPoseCandidates = 100;
 
-		void Add(const FPoseSearchCost& Cost, int32 PoseIdx, const UPoseSearchDatabase* Database)
+		void Add(const FPoseSearchCost& Cost, int32 PoseIdx, const UPoseSearchDatabase* Database, EPoseCandidateFlags PoseCandidateFlags)
 		{
 			if (Num() < MaxPoseCandidates || Cost < HeapTop().Cost)
 			{
@@ -1340,6 +1359,7 @@ public:
 				PoseCandidate.Cost = Cost;
 				PoseCandidate.PoseIdx = PoseIdx;
 				PoseCandidate.Database = Database;
+				PoseCandidate.PoseCandidateFlags = PoseCandidateFlags;
 				HeapPush(PoseCandidate);
 			}
 		}
