Fix for blend stack not blending curves correctly when one of the source poses doesn't have the curve
-review-22156501 [REDACTED_HANDLE]
#p4v-preflight-copy 22156491

diff --git a/Engine/Plugins/Experimental/Animation/PoseSearch/Source/Runtime/Private/AnimNode_BlendStack.cpp b/Engine/Plugins/Experimental/Animation/PoseSearch/Source/Runtime/Private/AnimNode_BlendStack.cpp
--- a/Engine/Plugins/Experimental/Animation/PoseSearch/Source/Runtime/Private/AnimNode_BlendStack.cpp
+++ b/Engine/Plugins/Experimental/Animation/PoseSearch/Source/Runtime/Private/AnimNode_BlendStack.cpp
@@ -203,6 +203,8 @@ void FAnimNode_BlendStack::Evaluate_AnyThread(FPoseContext& Output)
 		TArray<float> Weights;
 		for (int32 i = BlendStackSize - 2; i >= 0; --i)
 		{
+			// Since we're re-using the same pose context for different players, curves need to be reset before extraction.
+			EvaluationPoseContext.Curve.InitFrom(Output.AnimInstanceProxy->GetRequiredBones());
 			AnimPlayers[i].Evaluate_AnyThread(EvaluationPoseContext);
 			
 			if (AnimPlayers[i].GetBlendInWeights(Weights))
