motion matching - sort indexes in PopulateNonSelectableIdx to reduce complexity while searching in kdtree/bruteforce

diff --git a/Engine/Plugins/Experimental/Animation/PoseSearch/Source/Editor/Private/PoseSearchDebugger.cpp b/Engine/Plugins/Experimental/Animation/PoseSearch/Source/Editor/Private/PoseSearchDebugger.cpp
--- a/Engine/Plugins/Experimental/Animation/PoseSearch/Source/Editor/Private/PoseSearchDebugger.cpp
+++ b/Engine/Plugins/Experimental/Animation/PoseSearch/Source/Editor/Private/PoseSearchDebugger.cpp
@@ -39,7 +39,7 @@ inline static FLinearColor LinearColorBlend(FLinearColor LinearColorA, FLinearCo
 	return LinearColorA + (LinearColorB - LinearColorA) * BlendParam;
 }
 
-inline static void LinearColorBlend(FLinearColor LinearColorA, FLinearColor LinearColorB, TConstArrayView<const float> BlendParam, TArray<FLinearColor>& BlendedColors)
+inline static void LinearColorBlend(FLinearColor LinearColorA, FLinearColor LinearColorB, TConstArrayView<float> BlendParam, TArray<FLinearColor>& BlendedColors)
 {
 	const int32 Num = BlendParam.Num();
 	BlendedColors.SetNumUninitialized(Num);
@@ -49,7 +49,7 @@ inline static void LinearColorBlend(FLinearColor LinearColorA, FLinearColor Line
 	}
 }
 
-static void ArrayMinMax(TConstArrayView<const float> View, TArrayView<float> Min, TArrayView<float> Max, float InvalidValue)
+static void ArrayMinMax(TConstArrayView<float> View, TArrayView<float> Min, TArrayView<float> Max, float InvalidValue)
 {
 	const int32 Num = View.Num();
 	check(Num == Min.Num() && Num == Max.Num());
@@ -64,7 +64,7 @@ static void ArrayMinMax(TConstArrayView<const float> View, TArrayView<float> Min
 	}
 }
 
-static void ArraySafeNormalize(TConstArrayView<const float> View, TConstArrayView<const float> Min, TConstArrayView<const float> Max, TArrayView<float> NormalizedView)
+static void ArraySafeNormalize(TConstArrayView<float> View, TConstArrayView<float> Min, TConstArrayView<float> Max, TArrayView<float> NormalizedView)
 {
 	const int32 Num = View.Num();
 	check(Num == Min.Num() && Num == Max.Num() && Num == NormalizedView.Num());
@@ -788,7 +788,7 @@ public:
 		return UnfilteredDatabaseRows[Index]->SourceDatabase->Schema == Schema;
 	}
 
-	virtual TConstArrayView<const float> GetCostVector(int32 Index, const UPoseSearchSchema* Schema) const override
+	virtual TConstArrayView<float> GetCostVector(int32 Index, const UPoseSearchSchema* Schema) const override
 	{
 		check(UnfilteredDatabaseRows[Index]->SourceDatabase->Schema == Schema);
 		return UnfilteredDatabaseRows[Index]->CostVector;
@@ -912,7 +912,7 @@ void SDebuggerDatabaseView::Update(const FTraceMotionMatchingStateMessage& State
 						Row->PoseCost = PoseEntry.Cost;
 
 						Row->CostVector.SetNum(Database->Schema->SchemaCardinality);
-						TConstArrayView<const float> PoseValues = SearchIndex->GetPoseValues(PoseEntry.DbPoseIdx);
+						TConstArrayView<float> PoseValues = SearchIndex->GetPoseValues(PoseEntry.DbPoseIdx);
 
 						// in case we modify the schema while PIE is paused and displaying the Pose Search Editor, we could end up with a stale State with a DbEntry.QueryVector saved with the previous schema
 						// so the cardinality of DbEntry.QueryVector and PoseValues don't match. In that case we just use PoseValues as query to have all costs set to zero
diff --git a/Engine/Plugins/Experimental/Animation/PoseSearch/Source/Runtime/Private/PoseSearch.cpp b/Engine/Plugins/Experimental/Animation/PoseSearch/Source/Runtime/Private/PoseSearch.cpp
--- a/Engine/Plugins/Experimental/Animation/PoseSearch/Source/Runtime/Private/PoseSearch.cpp
+++ b/Engine/Plugins/Experimental/Animation/PoseSearch/Source/Runtime/Private/PoseSearch.cpp
@@ -73,7 +73,7 @@ namespace UE::PoseSearch
 	constexpr EParallelForFlags ParallelForFlags = EParallelForFlags::None;
 #endif // UE_POSE_SEARCH_FORCE_SINGLE_THREAD
 
-static inline float ArraySum(TConstArrayView<const float> View, int32 StartIndex, int32 Offset)
+static inline float ArraySum(TConstArrayView<float> View, int32 StartIndex, int32 Offset)
 {
 	float Sum = 0.f;
 	const int32 EndIndex = StartIndex + Offset;
@@ -84,7 +84,7 @@ static inline float ArraySum(TConstArrayView<const float> View, int32 StartIndex
 	return Sum;
 }
 
-static inline float CompareFeatureVectors(TConstArrayView<const float> A, TConstArrayView<const float> B, TConstArrayView<const float> WeightsSqrt)
+static inline float CompareFeatureVectors(TConstArrayView<float> A, TConstArrayView<float> B, TConstArrayView<float> WeightsSqrt)
 {
 	check(A.Num() == B.Num() && A.Num() == WeightsSqrt.Num());
 
@@ -95,7 +95,7 @@ static inline float CompareFeatureVectors(TConstArrayView<const float> A, TConst
 	return ((VA - VB) * VW).square().sum();
 }
 
-static inline float CompareFeatureVectors(TConstArrayView<const float> A, TConstArrayView<const float> B)
+static inline float CompareFeatureVectors(TConstArrayView<float> A, TConstArrayView<float> B)
 {
 	check(A.Num() == B.Num());
 
@@ -105,7 +105,7 @@ static inline float CompareFeatureVectors(TConstArrayView<const float> A, TConst
 	return (VA - VB).square().sum();
 }
 
-void CompareFeatureVectors(TConstArrayView<const float> A, TConstArrayView<const float> B, TConstArrayView<const float> WeightsSqrt, TArrayView<float> Result)
+void CompareFeatureVectors(TConstArrayView<float> A, TConstArrayView<float> B, TConstArrayView<float> WeightsSqrt, TArrayView<float> Result)
 {
 	check(A.Num() == B.Num() && A.Num() == WeightsSqrt.Num() && A.Num() == Result.Num());
 
@@ -184,7 +184,7 @@ FORCEINLINE auto LowerBound(IteratorType First, IteratorType Last, const ValueTy
 	return LowerBound(First, Last, Value, FIdentityFunctor(), TLess<>());
 }
 
-static int32 PopulateNonSelectableIdx(size_t* NonSelectableIdx, int32 NonSelectableIdxSize, FSearchContext& SearchContext, const UPoseSearchDatabase* Database, TArrayView<const float> QueryValues)
+static TArrayView<size_t> PopulateNonSelectableIdx(TArrayView<size_t> NonSelectableIdxBuffer, FSearchContext& SearchContext, const UPoseSearchDatabase* Database, TArrayView<const float> QueryValues)
 {
 	check(Database);
 
@@ -212,9 +212,9 @@ static int32 PopulateNonSelectableIdx(size_t* NonSelectableIdx, int32 NonSelecta
 
 			if (CurrentIndexAsset->SamplingInterval.Contains(PoseAssetPlayerTime))
 			{
-				if (NonSelectableIdxUsedSize < NonSelectableIdxSize)
+				if (NonSelectableIdxUsedSize < NonSelectableIdxBuffer.Num())
 				{
-					NonSelectableIdx[NonSelectableIdxUsedSize++] = PoseIdx;
+					NonSelectableIdxBuffer[NonSelectableIdxUsedSize++] = PoseIdx;
 
 #if UE_POSE_SEARCH_TRACE_ENABLED
 					const FPoseSearchCost PoseCost = Database->ComparePoses(SearchContext, PoseIdx, EPoseComparisonFlags::None, QueryValues);
@@ -224,7 +224,9 @@ static int32 PopulateNonSelectableIdx(size_t* NonSelectableIdx, int32 NonSelecta
 				else
 				{
 					UE_LOG(LogPoseSearch, Warning, TEXT("PopulateNonSelectableIdx couldn't add all the NonSelectableIdx"));
-					return NonSelectableIdxUsedSize;
+					TArrayView<size_t> NonSelectableIdx(NonSelectableIdxBuffer.GetData(), NonSelectableIdxUsedSize);
+					NonSelectableIdx.Sort();
+					return NonSelectableIdx;
 				}
 			}
 			else
@@ -250,9 +252,9 @@ static int32 PopulateNonSelectableIdx(size_t* NonSelectableIdx, int32 NonSelecta
 
 			if (CurrentIndexAsset->SamplingInterval.Contains(PoseAssetPlayerTime))
 			{
-				if (NonSelectableIdxUsedSize < NonSelectableIdxSize)
+				if (NonSelectableIdxUsedSize < NonSelectableIdxBuffer.Num())
 				{
-					NonSelectableIdx[NonSelectableIdxUsedSize++] = PoseIdx;
+					NonSelectableIdxBuffer[NonSelectableIdxUsedSize++] = PoseIdx;
 
 #if UE_POSE_SEARCH_TRACE_ENABLED
 					const FPoseSearchCost PoseCost = Database->ComparePoses(SearchContext, PoseIdx, EPoseComparisonFlags::None, QueryValues);
@@ -262,7 +264,9 @@ static int32 PopulateNonSelectableIdx(size_t* NonSelectableIdx, int32 NonSelecta
 				else
 				{
 					UE_LOG(LogPoseSearch, Warning, TEXT("PopulateNonSelectableIdx couldn't add all the NonSelectableIdx"));
-					return NonSelectableIdxUsedSize;
+					TArrayView<size_t> NonSelectableIdx(NonSelectableIdxBuffer.GetData(), NonSelectableIdxUsedSize);
+					NonSelectableIdx.Sort();
+					return NonSelectableIdx;
 				}
 			}
 			else
@@ -280,9 +284,9 @@ static int32 PopulateNonSelectableIdx(size_t* NonSelectableIdx, int32 NonSelecta
 			const FHistoricalPoseIndex& HistoricalPoseIndex = It.Key();
 			if (HistoricalPoseIndex.DatabaseKey == DatabaseKey)
 			{
-				if (NonSelectableIdxUsedSize < NonSelectableIdxSize)
+				if (NonSelectableIdxUsedSize < NonSelectableIdxBuffer.Num())
 				{
-					NonSelectableIdx[NonSelectableIdxUsedSize++] = HistoricalPoseIndex.PoseIndex;
+					NonSelectableIdxBuffer[NonSelectableIdxUsedSize++] = HistoricalPoseIndex.PoseIndex;
 
 #if UE_POSE_SEARCH_TRACE_ENABLED
 					const FPoseSearchCost PoseCost = Database->ComparePoses(SearchContext, HistoricalPoseIndex.PoseIndex, EPoseComparisonFlags::None, QueryValues);
@@ -292,13 +296,17 @@ static int32 PopulateNonSelectableIdx(size_t* NonSelectableIdx, int32 NonSelecta
 				else
 				{
 					UE_LOG(LogPoseSearch, Warning, TEXT("PopulateNonSelectableIdx couldn't add all the NonSelectableIdx"));
-					return NonSelectableIdxUsedSize;
+					TArrayView<size_t> NonSelectableIdx(NonSelectableIdxBuffer.GetData(), NonSelectableIdxUsedSize);
+					NonSelectableIdx.Sort();
+					return NonSelectableIdx;
 				}
 			}
 		}
 	}
 
-	return NonSelectableIdxUsedSize;
+	TArrayView<size_t> NonSelectableIdx(NonSelectableIdxBuffer.GetData(), NonSelectableIdxUsedSize);
+	NonSelectableIdx.Sort();
+	return NonSelectableIdx;
 }
 
 } // namespace UE::PoseSearch
@@ -547,7 +555,7 @@ bool FPoseSearchIndex::IsEmpty() const
 	return bEmpty;
 }
 
-TConstArrayView<const float> FPoseSearchIndex::GetPoseValues(int32 PoseIdx) const
+TConstArrayView<float> FPoseSearchIndex::GetPoseValues(int32 PoseIdx) const
 {
 	check(PoseIdx >= 0 && PoseIdx < NumPoses && Schema && Schema->SchemaCardinality >= 0);
 	const int32 ValueOffset = PoseIdx * Schema->SchemaCardinality;
@@ -686,7 +694,7 @@ UE::PoseSearch::FSearchResult UPoseSearchSequenceMetaData::Search(UE::PoseSearch
 FPoseSearchCost UPoseSearchSequenceMetaData::ComparePoses(
 	int32 PoseIdx,
 	UE::PoseSearch::EPoseComparisonFlags PoseComparisonFlags,
-	TConstArrayView<const float> QueryValues) const
+	TConstArrayView<float> QueryValues) const
 {
 	using namespace UE::PoseSearch;
 
@@ -1592,7 +1600,7 @@ bool UPoseSearchDatabase::IsCachedCookedPlatformDataLoaded(const ITargetPlatform
 
 #endif // WITH_EDITOR
 
-FPoseSearchCost UPoseSearchDatabase::ComparePoses(UE::PoseSearch::FSearchContext& SearchContext, int32 PoseIdx, UE::PoseSearch::EPoseComparisonFlags PoseComparisonFlags, TConstArrayView<const float> QueryValues) const
+FPoseSearchCost UPoseSearchDatabase::ComparePoses(UE::PoseSearch::FSearchContext& SearchContext, int32 PoseIdx, UE::PoseSearch::EPoseComparisonFlags PoseComparisonFlags, TConstArrayView<float> QueryValues) const
 {
 	const FPoseSearchIndex* SearchIndex = GetSearchIndex();
 	check(SearchIndex);
@@ -1675,7 +1683,8 @@ UE::PoseSearch::FSearchResult UPoseSearchDatabase::SearchPCAKDTree(UE::PoseSearc
 	
 	constexpr int NonSelectableIdxDataSize = 128;
 	size_t* NonSelectableIdxData((size_t*)FMemory_Alloca(NonSelectableIdxDataSize * sizeof(size_t)));
-	
+	TArrayView<size_t> NonSelectableIdxBuffer(NonSelectableIdxData, NonSelectableIdxDataSize);
+
 	// KDTree in PCA space search
 	if (PoseSearchMode == EPoseSearchMode::PCAKDTree_Validate)
 	{
@@ -1733,8 +1742,7 @@ UE::PoseSearch::FSearchResult UPoseSearchDatabase::SearchPCAKDTree(UE::PoseSearc
 	// there's no point in performing the search if CurrentBestTotalCost is already better than that
 	if (SearchContext.GetCurrentBestTotalCost() > SearchIndex->MinCostAddend)
 	{
-		const int NonSelectableIdxUsedSize = PopulateNonSelectableIdx(NonSelectableIdxData, NonSelectableIdxDataSize, SearchContext, this, QueryValues);
-		TArrayView<size_t> NonSelectableIdx(NonSelectableIdxData, NonSelectableIdxUsedSize);
+		TConstArrayView<size_t> NonSelectableIdx = PopulateNonSelectableIdx(NonSelectableIdxBuffer, SearchContext, this, QueryValues);
 		const RowMajorVectorMapConst MapWeightsSqrt(SearchIndex->WeightsSqrt.GetData(), 1, NumDimensions);
 		FKDTree::KNNResultSet ResultSet(ClampedKDTreeQueryNumNeighbors, ResultIndexes, ResultDistanceSqr, NonSelectableIdx);
 
@@ -1832,8 +1840,9 @@ UE::PoseSearch::FSearchResult UPoseSearchDatabase::SearchBruteForce(UE::PoseSear
 
 	constexpr int NonSelectableIdxDataSize = 128;
 	size_t* NonSelectableIdxData((size_t*)FMemory_Alloca(NonSelectableIdxDataSize * sizeof(size_t)));
-	const int NonSelectableIdxUsedSize = PopulateNonSelectableIdx(NonSelectableIdxData, NonSelectableIdxDataSize, SearchContext, this, QueryValues);
-	TArrayView<size_t> NonSelectableIdx(NonSelectableIdxData, NonSelectableIdxUsedSize);
+	TArrayView<size_t> NonSelectableIdxBuffer(NonSelectableIdxData, NonSelectableIdxDataSize);
+	TConstArrayView<size_t> NonSelectableIdx = PopulateNonSelectableIdx(NonSelectableIdxBuffer, SearchContext, this, QueryValues);
+	check(Algo::IsSorted(NonSelectableIdx));
 
 	// since any PoseCost calculated here is at least SearchIndex->MinCostAddend,
 	// there's no point in performing the search if CurrentBestTotalCost is already better than that
@@ -1856,7 +1865,7 @@ UE::PoseSearch::FSearchResult UPoseSearchDatabase::SearchBruteForce(UE::PoseSear
 					continue;
 				}
 
-				if (CheckForNonSelectableIdx && NonSelectableIdx.Contains(PoseIdx))
+				if (CheckForNonSelectableIdx && Algo::BinarySearch(NonSelectableIdx, PoseIdx) != INDEX_NONE)
 				{
 					continue;
 				}
@@ -2197,7 +2206,7 @@ bool FSearchContext::IsCurrentResultFromDatabase(const UPoseSearchDatabase* Data
 	return !bForceInterrupt && CurrentResult.IsValid() && CurrentResult.Database == Database;
 }
 
-TConstArrayView<const float> FSearchContext::GetCurrentResultPrevPoseVector() const
+TConstArrayView<float> FSearchContext::GetCurrentResultPrevPoseVector() const
 {
 	check(CurrentResult.IsValid());
 	const FPoseSearchIndex* SearchIndex = CurrentResult.Database->GetSearchIndex();
@@ -2205,7 +2214,7 @@ TConstArrayView<const float> FSearchContext::GetCurrentResultPrevPoseVector() co
 	return SearchIndex->GetPoseValues(CurrentResult.PrevPoseIdx);
 }
 
-TConstArrayView<const float> FSearchContext::GetCurrentResultPoseVector() const
+TConstArrayView<float> FSearchContext::GetCurrentResultPoseVector() const
 {
 	check(CurrentResult.IsValid());
 	const FPoseSearchIndex* SearchIndex = CurrentResult.Database->GetSearchIndex();
@@ -2213,7 +2222,7 @@ TConstArrayView<const float> FSearchContext::GetCurrentResultPoseVector() const
 	return SearchIndex->GetPoseValues(CurrentResult.PoseIdx);
 }
 
-TConstArrayView<const float> FSearchContext::GetCurrentResultNextPoseVector() const
+TConstArrayView<float> FSearchContext::GetCurrentResultNextPoseVector() const
 {
 	check(CurrentResult.IsValid());
 	const FPoseSearchIndex* SearchIndex = CurrentResult.Database->GetSearchIndex();
diff --git a/Engine/Plugins/Experimental/Animation/PoseSearch/Source/Runtime/Public/PoseSearch/KDTree.h b/Engine/Plugins/Experimental/Animation/PoseSearch/Source/Runtime/Public/PoseSearch/KDTree.h
--- a/Engine/Plugins/Experimental/Animation/PoseSearch/Source/Runtime/Public/PoseSearch/KDTree.h
+++ b/Engine/Plugins/Experimental/Animation/PoseSearch/Source/Runtime/Public/PoseSearch/KDTree.h
@@ -4,6 +4,7 @@
 
 #include "HAL/Platform.h"
 #include "Containers/ArrayView.h"
+#include "Algo/IsSorted.h"
 
 class FArchive;
 
@@ -42,7 +43,7 @@ struct POSESEARCH_API FKDTree
 
 	struct KNNResultSet
 	{
-		inline KNNResultSet(size_t InNumNeighbors, TArrayView<size_t> InIndexes, TArrayView<float> InDistances, TArrayView<size_t> InExcludeFromSearchIndexes = TArrayView<size_t>())
+		inline KNNResultSet(size_t InNumNeighbors, TArrayView<size_t> InIndexes, TArrayView<float> InDistances, TConstArrayView<size_t> InExcludeFromSearchIndexes = TConstArrayView<size_t>())
 		: Indexes(InIndexes)
 		, Distances(InDistances)
 		, NumNeighbors(InNumNeighbors)
@@ -53,6 +54,7 @@ struct POSESEARCH_API FKDTree
 			check(NumNeighbors > 0);
 			check(InIndexes.Num() > NumNeighbors);
 			check(InDistances.Num() > NumNeighbors);
+			check(Algo::IsSorted(InExcludeFromSearchIndexes));
 
 			Distances[NumNeighbors - 1] = UE_BIG_NUMBER;
 		}
@@ -69,8 +71,10 @@ struct POSESEARCH_API FKDTree
 
 		inline bool addPoint(size_t dist, size_t index)
 		{
-			if (ExcludeFromSearchIndexes.Contains(index))
+			if (Algo::BinarySearch(ExcludeFromSearchIndexes, index) != INDEX_NONE)
+			{
 				return true;
+			}
 
 			// shifting Distances[i] and Indexes[i] to make space for "dist" and "index" at the right "i"th slot
 			size_t i;
@@ -102,7 +106,7 @@ struct POSESEARCH_API FKDTree
 		TArrayView<float> Distances;
 		size_t NumNeighbors;
 		size_t Count;
-		const TArrayView<size_t> ExcludeFromSearchIndexes; // @todo: perhaps make it an hash if it gets too big
+		TConstArrayView<size_t> ExcludeFromSearchIndexes; // sorted array view
 	};
 	
 
diff --git a/Engine/Plugins/Experimental/Animation/PoseSearch/Source/Runtime/Public/PoseSearch/PoseSearch.h b/Engine/Plugins/Experimental/Animation/PoseSearch/Source/Runtime/Public/PoseSearch/PoseSearch.h
--- a/Engine/Plugins/Experimental/Animation/PoseSearch/Source/Runtime/Public/PoseSearch/PoseSearch.h
+++ b/Engine/Plugins/Experimental/Animation/PoseSearch/Source/Runtime/Public/PoseSearch/PoseSearch.h
@@ -307,7 +307,7 @@ public:
 	virtual bool IsCostVectorFromSchema(int32 Index, const UPoseSearchSchema* Schema) const = 0;
 
 	// returns the Index-th cost data vector
-	virtual TConstArrayView<const float> GetCostVector(int32 Index, const UPoseSearchSchema* Schema) const = 0;
+	virtual TConstArrayView<float> GetCostVector(int32 Index, const UPoseSearchSchema* Schema) const = 0;
 
 	// every breakdown section start by calling BeginBreakDownSection...
 	virtual void BeginBreakDownSection(const FText& Label) = 0;
@@ -657,7 +657,7 @@ struct POSESEARCH_API FPoseSearchIndex
 	bool IsValidPoseIndex(int32 PoseIdx) const { return PoseIdx < NumPoses; }
 	bool IsEmpty() const;
 
-	TConstArrayView<const float> GetPoseValues(int32 PoseIdx) const;
+	TConstArrayView<float> GetPoseValues(int32 PoseIdx) const;
 
 	const FPoseSearchIndexAsset* FindAssetForPose(int32 PoseIdx) const;
 	float GetAssetTime(int32 PoseIdx, const FPoseSearchIndexAsset* Asset) const;
@@ -1101,7 +1101,7 @@ public:
 
 	void BuildQuery(UE::PoseSearch::FSearchContext& SearchContext, FPoseSearchFeatureVectorBuilder& OutQuery) const;
 
-	FPoseSearchCost ComparePoses(UE::PoseSearch::FSearchContext& SearchContext, int32 PoseIdx, UE::PoseSearch::EPoseComparisonFlags PoseComparisonFlags, TConstArrayView<const float> QueryValues) const;
+	FPoseSearchCost ComparePoses(UE::PoseSearch::FSearchContext& SearchContext, int32 PoseIdx, UE::PoseSearch::EPoseComparisonFlags PoseComparisonFlags, TConstArrayView<float> QueryValues) const;
 
 protected:
 	UE::PoseSearch::FSearchResult SearchPCAKDTree(UE::PoseSearch::FSearchContext& SearchContext) const;
@@ -1137,7 +1137,7 @@ public:
 	UE::PoseSearch::FSearchResult Search(UE::PoseSearch::FSearchContext& SearchContext) const;
 
 protected:
-	FPoseSearchCost ComparePoses(int32 PoseIdx, UE::PoseSearch::EPoseComparisonFlags PoseComparisonFlags, TConstArrayView<const float> QueryValues) const;
+	FPoseSearchCost ComparePoses(int32 PoseIdx, UE::PoseSearch::EPoseComparisonFlags PoseComparisonFlags, TConstArrayView<float> QueryValues) const;
 
 public: // UObject
 	virtual void PreSave(FObjectPreSaveContext ObjectSaveContext) override;
@@ -1291,9 +1291,9 @@ struct POSESEARCH_API FSearchContext
 
 	bool IsCurrentResultFromDatabase(const UPoseSearchDatabase* Database) const;
 
-	TConstArrayView<const float> GetCurrentResultPrevPoseVector() const;
-	TConstArrayView<const float> GetCurrentResultPoseVector() const;
-	TConstArrayView<const float> GetCurrentResultNextPoseVector() const;
+	TConstArrayView<float> GetCurrentResultPrevPoseVector() const;
+	TConstArrayView<float> GetCurrentResultPoseVector() const;
+	TConstArrayView<float> GetCurrentResultNextPoseVector() const;
 
 	static constexpr int8 SchemaRootBoneIdx = -1;
 
@@ -1380,7 +1380,7 @@ public:
 POSESEARCH_API void DrawFeatureVector(const FDebugDrawParams& DrawParams, TArrayView<const float> PoseVector);
 POSESEARCH_API void DrawFeatureVector(const FDebugDrawParams& DrawParams, int32 PoseIdx);
 POSESEARCH_API void DrawSearchIndex(const FDebugDrawParams& DrawParams);
-POSESEARCH_API void CompareFeatureVectors(TConstArrayView<const float> A, TConstArrayView<const float> B, TConstArrayView<const float> WeightsSqrt, TArrayView<float> Result);
+POSESEARCH_API void CompareFeatureVectors(TConstArrayView<float> A, TConstArrayView<float> B, TConstArrayView<float> WeightsSqrt, TArrayView<float> Result);
 
 /**
 * Creates a pose search index for an animation sequence
