blend stack - adding EAlphaBlendOption to blend stack

diff --git a/Engine/Plugins/Experimental/Animation/PoseSearch/Source/Runtime/Private/AnimNode_BlendStack.cpp b/Engine/Plugins/Experimental/Animation/PoseSearch/Source/Runtime/Private/AnimNode_BlendStack.cpp
--- a/Engine/Plugins/Experimental/Animation/PoseSearch/Source/Runtime/Private/AnimNode_BlendStack.cpp
+++ b/Engine/Plugins/Experimental/Animation/PoseSearch/Source/Runtime/Private/AnimNode_BlendStack.cpp
@@ -12,7 +12,7 @@ TAutoConsoleVariable<int32> CVarAnimBlendStackEnable(TEXT("a.AnimNode.BlendStack
 
 /////////////////////////////////////////////////////
 // FPoseSearchAnimPlayer
-void FPoseSearchAnimPlayer::Initialize(ESearchIndexAssetType InAssetType, UAnimationAsset* AnimationAsset, float AccumulatedTime, bool bLoop, bool bMirrored, UMirrorDataTable* MirrorDataTable, float BlendTime, const UBlendProfile* BlendProfile, FVector BlendParameters)
+void FPoseSearchAnimPlayer::Initialize(ESearchIndexAssetType InAssetType, UAnimationAsset* AnimationAsset, float AccumulatedTime, bool bLoop, bool bMirrored, UMirrorDataTable* MirrorDataTable, float BlendTime, const UBlendProfile* BlendProfile, EAlphaBlendOption InBlendOption, FVector BlendParameters)
 {
 	check(AnimationAsset);
 
@@ -37,6 +37,7 @@ void FPoseSearchAnimPlayer::Initialize(ESearchIndexAssetType InAssetType, UAnima
 		BlendProfile->FillSkeletonBoneDurationsArray(TotalBlendInTimePerBone, BlendTime);
 		BlendTime = *Algo::MaxElement(TotalBlendInTimePerBone);
 	}
+	BlendOption = InBlendOption;
 
 	AssetType = InAssetType;
 	TotalBlendInTime = BlendTime;
@@ -147,7 +148,8 @@ bool FPoseSearchAnimPlayer::GetBlendInWeights(TArray<float>& Weights) const
 			}
 			else
 			{
-				Weights[BoneIdx] = FMath::Clamp(CurrentBlendInTime / TotalBlendInTimeBoneIdx, 0.f, 1.f);
+				const float unclampedLinearWeight = CurrentBlendInTime / TotalBlendInTimeBoneIdx;
+				Weights[BoneIdx] = FAlphaBlend::AlphaToBlendOption(unclampedLinearWeight, BlendOption);
 			}
 		}
 		return true;
@@ -214,7 +216,7 @@ void FAnimNode_BlendStack::Evaluate_AnyThread(FPoseContext& Output)
 			}
 			else
 			{
-				const float Weight = 1.f - AnimPlayers[i].GetBlendInPercentage();
+				const float Weight = 1.f - FAlphaBlend::AlphaToBlendOption(AnimPlayers[i].GetBlendInPercentage(), AnimPlayers[i].GetBlendOption());
 				FAnimationRuntime::BlendTwoPosesTogether(FAnimationPoseData(Output), FAnimationPoseData(EvaluationPoseContext), Weight, BlendedAnimationPoseData);
 			}
 			Output = BlendedPoseContext; // @todo: this should not be necessary either: optimize it away!
@@ -243,10 +245,10 @@ float FAnimNode_BlendStack::GetAccumulatedTime() const
 	return AnimPlayers.IsEmpty() ? 0.f : AnimPlayers.First().GetAccumulatedTime();
 }
 
-void FAnimNode_BlendStack::BlendTo(ESearchIndexAssetType AssetType, UAnimationAsset* AnimationAsset, float AccumulatedTime, bool bLoop, bool bMirrored, UMirrorDataTable* MirrorDataTable, int32 MaxActiveBlends, float BlendTime, const UBlendProfile* BlendProfile, FVector BlendParameters)
+void FAnimNode_BlendStack::BlendTo(ESearchIndexAssetType AssetType, UAnimationAsset* AnimationAsset, float AccumulatedTime, bool bLoop, bool bMirrored, UMirrorDataTable* MirrorDataTable, int32 MaxActiveBlends, float BlendTime, const UBlendProfile* BlendProfile, EAlphaBlendOption BlendOption, FVector BlendParameters)
 {
 	AnimPlayers.PushFirst(FPoseSearchAnimPlayer());
-	AnimPlayers.First().Initialize(AssetType, AnimationAsset, AccumulatedTime, bLoop, bMirrored, MirrorDataTable, BlendTime, BlendProfile, BlendParameters);
+	AnimPlayers.First().Initialize(AssetType, AnimationAsset, AccumulatedTime, bLoop, bMirrored, MirrorDataTable, BlendTime, BlendProfile, BlendOption, BlendParameters);
 
 	CalculateWeights();
 	PruneBlendStack(MaxActiveBlends);
diff --git a/Engine/Plugins/Experimental/Animation/PoseSearch/Source/Runtime/Private/AnimNode_MotionMatching.cpp b/Engine/Plugins/Experimental/Animation/PoseSearch/Source/Runtime/Private/AnimNode_MotionMatching.cpp
--- a/Engine/Plugins/Experimental/Animation/PoseSearch/Source/Runtime/Private/AnimNode_MotionMatching.cpp
+++ b/Engine/Plugins/Experimental/Animation/PoseSearch/Source/Runtime/Private/AnimNode_MotionMatching.cpp
@@ -99,7 +99,7 @@ void FAnimNode_MotionMatching::UpdateAssetPlayer(const FAnimationUpdateContext&
 		const FPoseSearchDatabaseAnimationAssetBase& AnimationAssetBase = Database->GetAnimationSourceAsset(SearchIndexAsset);
 		BlendStackNode.BlendTo(	SearchIndexAsset->Type, AnimationAssetBase.GetAnimationAsset(), MotionMatchingState.CurrentSearchResult.AssetTime,
 								AnimationAssetBase.IsLooping(), SearchIndexAsset->bMirrored, Database->Schema->MirrorDataTable.Get(),
-								Settings.MaxActiveBlends, Settings.BlendTime, Settings.BlendProfile, SearchIndexAsset->BlendParameters);
+								Settings.MaxActiveBlends, Settings.BlendTime, Settings.BlendProfile, Settings.BlendOption, SearchIndexAsset->BlendParameters);
 	}
 
 	if (!BlendStackNode.AnimPlayers.IsEmpty() && BlendStackNode.AnimPlayers[0].GetAssetType() == ESearchIndexAssetType::Sequence)
diff --git a/Engine/Plugins/Experimental/Animation/PoseSearch/Source/Runtime/Public/PoseSearch/AnimNode_BlendStack.h b/Engine/Plugins/Experimental/Animation/PoseSearch/Source/Runtime/Public/PoseSearch/AnimNode_BlendStack.h
--- a/Engine/Plugins/Experimental/Animation/PoseSearch/Source/Runtime/Public/PoseSearch/AnimNode_BlendStack.h
+++ b/Engine/Plugins/Experimental/Animation/PoseSearch/Source/Runtime/Public/PoseSearch/AnimNode_BlendStack.h
@@ -16,7 +16,7 @@ struct FPoseSearchAnimPlayer
 {
 	GENERATED_BODY()
 	
-	void Initialize(ESearchIndexAssetType InAssetType, UAnimationAsset* AnimationAsset, float AccumulatedTime, bool bLoop, bool bMirrored, UMirrorDataTable* MirrorDataTable, float BlendTime = 0.f, const UBlendProfile* BlendProfile = nullptr, FVector BlendParameters = FVector::Zero());
+	void Initialize(ESearchIndexAssetType InAssetType, UAnimationAsset* AnimationAsset, float AccumulatedTime, bool bLoop, bool bMirrored, UMirrorDataTable* MirrorDataTable, float BlendTime, const UBlendProfile* BlendProfile, EAlphaBlendOption InBlendOption, FVector BlendParameters);
 	void Evaluate_AnyThread(FPoseContext& Output);
 	void Update_AnyThread(const FAnimationUpdateContext& Context);
 	float GetAccumulatedTime() const;
@@ -25,6 +25,7 @@ struct FPoseSearchAnimPlayer
 	void SetBlendWeight(float InBlendWeight);
 	bool GetBlendInWeights(TArray<float>& Weights) const;
 	ESearchIndexAssetType GetAssetType() const { return AssetType; }
+	EAlphaBlendOption GetBlendOption() const { return BlendOption; }
 
 	// @todo: used only for DynamicPlayRateAdjustment. Remove once the functionality is integrated with the BlendStackNode
 	FAnimNode_SequencePlayer_Standalone GetSequencePlayerNode() { return SequencePlayerNode; }
@@ -44,6 +45,7 @@ protected:
 	FAnimNode_Mirror_Standalone MirrorNode;
 
 	ESearchIndexAssetType AssetType = ESearchIndexAssetType::Sequence;
+	EAlphaBlendOption BlendOption = EAlphaBlendOption::Linear;
 
 	TCustomBoneIndexArray<float, FSkeletonPoseBoneIndex> TotalBlendInTimePerBone;
 
@@ -63,7 +65,7 @@ struct POSESEARCH_API FAnimNode_BlendStack : public FAnimNode_AssetPlayerBase
 	virtual void Evaluate_AnyThread(FPoseContext& Output) override;
 	// End of FAnimNode_Base interface
 
-	void BlendTo(ESearchIndexAssetType AssetType, UAnimationAsset* AnimationAsset, float AccumulatedTime, bool bLoop, bool bMirrored, UMirrorDataTable* MirrorDataTable, int32 MaxActiveBlends = 3, float BlendTime = 0.3f, const UBlendProfile* BlendProfile = nullptr, FVector BlendParameters = FVector::Zero());
+	void BlendTo(ESearchIndexAssetType AssetType, UAnimationAsset* AnimationAsset, float AccumulatedTime = 0.f, bool bLoop = false, bool bMirrored = false, UMirrorDataTable* MirrorDataTable = nullptr, int32 MaxActiveBlends = 3, float BlendTime = 0.3f, const UBlendProfile* BlendProfile = nullptr, EAlphaBlendOption BlendOption = EAlphaBlendOption::Linear, FVector BlendParameters = FVector::Zero());
 	void CalculateWeights();
 	void PruneBlendStack(int32 MaxActiveBlends);
 
diff --git a/Engine/Plugins/Experimental/Animation/PoseSearch/Source/Runtime/Public/PoseSearch/PoseSearchLibrary.h b/Engine/Plugins/Experimental/Animation/PoseSearch/Source/Runtime/Public/PoseSearch/PoseSearchLibrary.h
--- a/Engine/Plugins/Experimental/Animation/PoseSearch/Source/Runtime/Public/PoseSearch/PoseSearchLibrary.h
+++ b/Engine/Plugins/Experimental/Animation/PoseSearch/Source/Runtime/Public/PoseSearch/PoseSearchLibrary.h
@@ -38,6 +38,9 @@ struct POSESEARCH_API FMotionMatchingSettings
 	UPROPERTY(EditAnywhere, BlueprintReadWrite, Category=Settings, meta=(UseAsBlendProfile = true))
 	TObjectPtr<UBlendProfile> BlendProfile;
 
+	UPROPERTY(EditAnywhere, BlueprintReadWrite, Category = Settings)
+	EAlphaBlendOption BlendOption = EAlphaBlendOption::Linear;
+
 	// If the pose jump requires a mirroring change and this value is greater than 0, it will be used instead of BlendTime
 	UPROPERTY(EditAnywhere, BlueprintReadWrite, Category=Settings, meta=(ClampMin="0", DislayAfter="BlendTime"))
 	float MirrorChangeBlendTime = 0.0f;
