blend stack - adding GatherDebugData, StorePoseContext/RestorePoseContext to accumulate poses if there are too many RequestedMaxActiveBlends

diff --git a/Engine/Plugins/Experimental/Animation/PoseSearch/Source/Runtime/Private/AnimNode_BlendStack.cpp b/Engine/Plugins/Experimental/Animation/PoseSearch/Source/Runtime/Private/AnimNode_BlendStack.cpp
--- a/Engine/Plugins/Experimental/Animation/PoseSearch/Source/Runtime/Private/AnimNode_BlendStack.cpp
+++ b/Engine/Plugins/Experimental/Animation/PoseSearch/Source/Runtime/Private/AnimNode_BlendStack.cpp
@@ -8,7 +8,10 @@
 
 #define LOCTEXT_NAMESPACE "AnimNode_BlendStack"
 
+#if ENABLE_ANIM_DEBUG
 TAutoConsoleVariable<int32> CVarAnimBlendStackEnable(TEXT("a.AnimNode.BlendStack.Enable"), 1, TEXT("Enable / Disable Blend Stack"));
+TAutoConsoleVariable<int32> CVarAnimBlendStackPruningEnable(TEXT("a.AnimNode.BlendStack.Pruning.Enable"), 1, TEXT("Enable / Disable Blend Stack Pruning"));
+#endif
 
 /////////////////////////////////////////////////////
 // FPoseSearchAnimPlayer
@@ -42,7 +45,6 @@ void FPoseSearchAnimPlayer::Initialize(ESearchIndexAssetType InAssetType, UAnima
 	AssetType = InAssetType;
 	TotalBlendInTime = BlendTime;
 	CurrentBlendInTime = 0.f;
-	BlendWeight = 0.f;
 
 	MirrorNode.SetMirrorDataTable(MirrorDataTable);
 	MirrorNode.SetMirror(bMirrored);
@@ -68,14 +70,60 @@ void FPoseSearchAnimPlayer::Initialize(ESearchIndexAssetType InAssetType, UAnima
 		BlendSpacePlayerNode.SetPlayRate(1.0f);
 		BlendSpacePlayerNode.SetPosition(BlendParameters);
 	}
-	else
+	else 
 	{
 		checkNoEntry();
 	}
+	UpdateSourceLinkNode();
+}
 
+void FPoseSearchAnimPlayer::StorePoseContext(const FPoseContext& PoseContext)
+{
+	AssetType = ESearchIndexAssetType::Invalid;
 	UpdateSourceLinkNode();
+
+	const FBoneContainer& BoneContainer = PoseContext.Pose.GetBoneContainer();
+	const USkeleton* SkeletonAsset = BoneContainer.GetSkeletonAsset();
+	check(SkeletonAsset);
+
+	const FReferenceSkeleton& RefSkeleton = SkeletonAsset->GetReferenceSkeleton();
+	const int32 NumSkeletonBones = RefSkeleton.GetNum();
+
+	StoredPose.SetNum(NumSkeletonBones);
+	for (FSkeletonPoseBoneIndex SkeletonBoneIdx(0); SkeletonBoneIdx != NumSkeletonBones; ++SkeletonBoneIdx)
+	{
+		FCompactPoseBoneIndex CompactBoneIdx = BoneContainer.GetCompactPoseIndexFromSkeletonPoseIndex(SkeletonBoneIdx);
+		StoredPose[SkeletonBoneIdx.GetInt()] = CompactBoneIdx.IsValid() ? PoseContext.Pose[CompactBoneIdx] : RefSkeleton.GetRefBonePose()[SkeletonBoneIdx.GetInt()];
+	}
+
+	// @todo: perhaps copy PoseContext.Curve and PoseContext.CustomAttributes?
 }
 
+void FPoseSearchAnimPlayer::RestorePoseContext(FPoseContext& PoseContext) const
+{
+	check(AssetType == ESearchIndexAssetType::Invalid);
+
+	const FBoneContainer& BoneContainer = PoseContext.Pose.GetBoneContainer();
+	const USkeleton* SkeletonAsset = BoneContainer.GetSkeletonAsset();
+	check(SkeletonAsset);
+
+	const FReferenceSkeleton& RefSkeleton = SkeletonAsset->GetReferenceSkeleton();
+	const int32 NumSkeletonBones = RefSkeleton.GetNum();
+
+	for (const FBoneIndexType BoneIdx : BoneContainer.GetBoneIndicesArray())
+	{
+		const FCompactPoseBoneIndex CompactBoneIdx(BoneIdx);
+		const FSkeletonPoseBoneIndex SkeletonBoneIdx = BoneContainer.GetSkeletonPoseIndexFromCompactPoseIndex(CompactBoneIdx);
+		if (SkeletonBoneIdx.IsValid() && BoneIdx < StoredPose.Num())
+		{
+			PoseContext.Pose[CompactBoneIdx] = StoredPose[SkeletonBoneIdx.GetInt()];
+		}
+	}
+
+	// @todo: perhaps copy PoseContext.Curve and PoseContext.CustomAttributes?
+}
+
+
 // @todo: maybe implement copy/move constructors and assignement operator do so (or use a list instead of an array)
 // since we're making copies and moving this object in memory, we're using this method to set the MirrorNode SourceLinkNode when necessary
 void FPoseSearchAnimPlayer::UpdateSourceLinkNode()
@@ -90,21 +138,30 @@ void FPoseSearchAnimPlayer::UpdateSourceLinkNode()
 	}
 	else
 	{
-		checkNoEntry();
+		MirrorNode.SetSourceLinkNode(nullptr);
 	}
 }
 
 void FPoseSearchAnimPlayer::Evaluate_AnyThread(FPoseContext& Output)
 {
-	UpdateSourceLinkNode();
-	MirrorNode.Evaluate_AnyThread(Output);
+	if (AssetType != ESearchIndexAssetType::Invalid)
+	{
+		UpdateSourceLinkNode();
+		MirrorNode.Evaluate_AnyThread(Output);
+	}
+	else
+	{
+		RestorePoseContext(Output);
+	}
 }
 
-void FPoseSearchAnimPlayer::Update_AnyThread(const FAnimationUpdateContext& Context)
+void FPoseSearchAnimPlayer::Update_AnyThread(const FAnimationUpdateContext& Context, float BlendWeight)
 {
+	const FAnimationUpdateContext AnimPlayerContext = Context.FractionalWeightAndRootMotion(BlendWeight, BlendWeight);
+
 	UpdateSourceLinkNode();
-	MirrorNode.Update_AnyThread(Context);
-	CurrentBlendInTime += Context.GetDeltaTime();
+	MirrorNode.Update_AnyThread(AnimPlayerContext);
+	CurrentBlendInTime += AnimPlayerContext.GetDeltaTime();
 }
 
 float FPoseSearchAnimPlayer::GetAccumulatedTime() const
@@ -119,10 +176,26 @@ float FPoseSearchAnimPlayer::GetAccumulatedTime() const
 		return BlendSpacePlayerNode.GetAccumulatedTime();
 	}
 	
-	checkNoEntry();
 	return 0.f;
 }
 
+FString FPoseSearchAnimPlayer::GetAnimName() const
+{
+	if (AssetType == ESearchIndexAssetType::Sequence)
+	{
+		check(SequencePlayerNode.GetSequence());
+		return SequencePlayerNode.GetSequence()->GetName();
+	}
+
+	if (AssetType == ESearchIndexAssetType::BlendSpace)
+	{
+		check(BlendSpacePlayerNode.GetBlendSpace());
+		return BlendSpacePlayerNode.GetBlendSpace()->GetName();
+	}
+
+	return FString("StoredPose");
+}
+
 float FPoseSearchAnimPlayer::GetBlendInPercentage() const
 {
 	if (FMath::IsNearlyZero(TotalBlendInTime))
@@ -157,11 +230,6 @@ bool FPoseSearchAnimPlayer::GetBlendInWeights(TArray<float>& Weights) const
 	return false;
 }
 
-void FPoseSearchAnimPlayer::SetBlendWeight(float InBlendWeight)
-{
-	BlendWeight = InBlendWeight;
-}
-
 /////////////////////////////////////////////////////
 // FAnimNode_BlendStack
 void FAnimNode_BlendStack::Evaluate_AnyThread(FPoseContext& Output)
@@ -170,17 +238,7 @@ void FAnimNode_BlendStack::Evaluate_AnyThread(FPoseContext& Output)
 
 	DECLARE_SCOPE_HIERARCHICAL_COUNTER_ANIMNODE(Evaluate_AnyThread);
 
-	int32 BlendStackSize = AnimPlayers.Num();
-
-	// Disable blend stack if requested (for testing / debugging)
-	if (!CVarAnimBlendStackEnable.GetValueOnAnyThread())
-	{
-		if (BlendStackSize > 1)
-		{
-			BlendStackSize = 1;
-		}
-	}
-
+	const int32 BlendStackSize = AnimPlayers.Num();
 	if (BlendStackSize <= 0)
 	{
 		Output.ResetToRefPose();
@@ -189,10 +247,24 @@ void FAnimNode_BlendStack::Evaluate_AnyThread(FPoseContext& Output)
 	{
 		AnimPlayers.First().Evaluate_AnyThread(Output);
 	}
+	else if (RequestedMaxActiveBlends <= 0 
+#if ENABLE_ANIM_DEBUG
+		|| !CVarAnimBlendStackEnable.GetValueOnAnyThread()
+#endif // ENABLE_ANIM_DEBUG
+		)
+	{
+		// Disable blend stack if requested (for testing / debugging) by removing all the AnimPlayers except the first
+		while (AnimPlayers.Num() > 1)
+		{
+			AnimPlayers.PopLast();
+		}
+		AnimPlayers.First().Evaluate_AnyThread(Output);
+	}
 	else
 	{
+		// evaluating the last AnimPlayer into Output...
 		AnimPlayers[BlendStackSize - 1].Evaluate_AnyThread(Output);
-
+		 
 		FPoseContext EvaluationPoseContext(Output);
 		FPoseContext BlendedPoseContext(Output); // @todo: this should not be necessary (but FBaseBlendedCurve::InitFrom complains about "ensure(&InCurveToInitFrom != this)"): optimize it away!
 		FAnimationPoseData BlendedAnimationPoseData(BlendedPoseContext);
@@ -203,10 +275,11 @@ void FAnimNode_BlendStack::Evaluate_AnyThread(FPoseContext& Output)
 		const FReferenceSkeleton& RefSkeleton = SkeletonAsset->GetReferenceSkeleton();
 		const int32 NumSkeletonBones = RefSkeleton.GetNum();
 		TArray<float> Weights;
+
+		const int BoundaryIndex = RequestedMaxActiveBlends + 1;
+		// evaluating from the second last to the first into EvaluationPoseContext and then blend it with the with the Output (initialized with the last AnimPlayer evaluation)
 		for (int32 i = BlendStackSize - 2; i >= 0; --i)
 		{
-			// Since we're re-using the same pose context for different players, curves need to be reset before extraction.
-			EvaluationPoseContext.Curve.InitFrom(Output.AnimInstanceProxy->GetRequiredBones());
 			AnimPlayers[i].Evaluate_AnyThread(EvaluationPoseContext);
 			
 			if (AnimPlayers[i].GetBlendInWeights(Weights))
@@ -220,6 +293,30 @@ void FAnimNode_BlendStack::Evaluate_AnyThread(FPoseContext& Output)
 				FAnimationRuntime::BlendTwoPosesTogether(FAnimationPoseData(Output), FAnimationPoseData(EvaluationPoseContext), Weight, BlendedAnimationPoseData);
 			}
 			Output = BlendedPoseContext; // @todo: this should not be necessary either: optimize it away!
+
+			if (i >= BoundaryIndex 
+#if ENABLE_ANIM_DEBUG
+				&& CVarAnimBlendStackPruningEnable.GetValueOnAnyThread()
+#endif // ENABLE_ANIM_DEBUG
+				)
+			{
+				// too many AnimPlayers! we don't have enought available blends to hold them all, so we accumulate the blended poses into Output / BlendedPoseContext, until...
+				AnimPlayers.PopLast();
+				
+				if (i == BoundaryIndex)
+				{
+					// we can store Output / BlendedPoseContext into the last AnimPlayer, that will hold a static pose, no longer an animation playing
+					AnimPlayers[i].StorePoseContext(Output);
+				}
+			}
+		}
+
+		const int32 ActiveBlends = AnimPlayers.Num() - 1;
+		if (ActiveBlends >= RequestedMaxActiveBlends)
+		{
+			UE_LOG(LogPoseSearch, Display,
+				TEXT("FAnimNode_BlendStack NumBlends/MaxNumBlends %d / %d"),
+				ActiveBlends, RequestedMaxActiveBlends);
 		}
 	}
 }
@@ -228,15 +325,35 @@ void FAnimNode_BlendStack::UpdateAssetPlayer(const FAnimationUpdateContext& Cont
 {
 	Super::UpdateAssetPlayer(Context);
 
+	// AnimPlayers[0] is the most newly inserted AnimPlayer, AnimPlayers[AnimPlayers.Num()-1] is the oldest, so to calculate the weights
+	// we ask AnimPlayers[0] its BlendInPercentage and then distribute the left over (CurrentWeightMultiplier) to the rest of the AnimPlayers
+	// AnimPlayers[AnimPlayerIndex].GetBlendWeight() will now store the weighted contribution of AnimPlayers[AnimPlayerIndex] to be able to calculate root motion from animation
+	float CurrentWeightMultiplier = 1.f;
 	const int32 BlendStackSize = AnimPlayers.Num();
+	int32 AnimPlayerIndex = 0;
+	for (; AnimPlayerIndex < BlendStackSize; ++AnimPlayerIndex)
+	{
+		FPoseSearchAnimPlayer& AnimPlayer = AnimPlayers[AnimPlayerIndex];
+		const bool bIsLastAnimPlayers = AnimPlayerIndex == BlendStackSize - 1;
+		const float BlendInPercentage = bIsLastAnimPlayers ? 1.f : AnimPlayer.GetBlendInPercentage();
+		const float AnimPlayerBlendWeight = CurrentWeightMultiplier * BlendInPercentage;
+
+		// don't break for AnimPlayerIndex == 0 since FAnimNode_BlendStack::BlendTo initialize the AnimPlayer with a weight of zero
+		if (AnimPlayerIndex > 0 && AnimPlayerBlendWeight < UE_KINDA_SMALL_NUMBER)
+		{
+			break;
+		}
 
-	CalculateWeights();
-	PruneBlendStack(BlendStackSize);
+		AnimPlayer.Update_AnyThread(Context, AnimPlayerBlendWeight);
 
-	for (FPoseSearchAnimPlayer& AnimPlayer : AnimPlayers)
+		CurrentWeightMultiplier *= (1.f - BlendInPercentage);
+	}
+
+	// AnimPlayers[AnimPlayerIndex] is the first FPoseSearchAnimPlayer with a weight contribution of zero, so we can discard it and all the successive AnimPlayers as well
+	const int32 WantedAnimPlayersNum = FMath::Max(1, AnimPlayerIndex); // we save at least one FPoseSearchAnimPlayer
+	while (AnimPlayers.Num() > WantedAnimPlayersNum)
 	{
-		const FAnimationUpdateContext AnimPlayerContext = Context.FractionalWeightAndRootMotion(AnimPlayer.GetBlendWeight(), AnimPlayer.GetBlendWeight());
-		AnimPlayer.Update_AnyThread(AnimPlayerContext);
+		AnimPlayers.PopLast();
 	}
 }
 
@@ -247,47 +364,28 @@ float FAnimNode_BlendStack::GetAccumulatedTime() const
 
 void FAnimNode_BlendStack::BlendTo(ESearchIndexAssetType AssetType, UAnimationAsset* AnimationAsset, float AccumulatedTime, bool bLoop, bool bMirrored, UMirrorDataTable* MirrorDataTable, int32 MaxActiveBlends, float BlendTime, const UBlendProfile* BlendProfile, EAlphaBlendOption BlendOption, FVector BlendParameters)
 {
+	RequestedMaxActiveBlends = MaxActiveBlends;
 	AnimPlayers.PushFirst(FPoseSearchAnimPlayer());
 	AnimPlayers.First().Initialize(AssetType, AnimationAsset, AccumulatedTime, bLoop, bMirrored, MirrorDataTable, BlendTime, BlendProfile, BlendOption, BlendParameters);
-
-	CalculateWeights();
-	PruneBlendStack(MaxActiveBlends);
-}
-
-void FAnimNode_BlendStack::CalculateWeights()
-{
-	// AnimPlayers[0] is the most newly inserted AnimPlayer, AnimPlayers[AnimPlayers.Num()-1] is the oldest, so to calculate the weights
-	// we ask AnimPlayers[0] its BlendInPercentage and then distribute the left over (CurrentWeightMultiplier) to the rest of the AnimPlayers
-	float CurrentWeightMultiplier = 1.f;
-	const int32 BlendStackSize = AnimPlayers.Num();
-	for (int32 i = 0; i < BlendStackSize; ++i)
-	{
-		const bool bIsLastAnimPlayers = i == BlendStackSize - 1;
-		const float BlendInPercentage = bIsLastAnimPlayers ? 1.f : AnimPlayers[i].GetBlendInPercentage();
-		AnimPlayers[i].SetBlendWeight(CurrentWeightMultiplier * BlendInPercentage);
-
-		CurrentWeightMultiplier *= (1.f - BlendInPercentage);
-	}
 }
 
-void FAnimNode_BlendStack::PruneBlendStack(int32 MaxActiveBlends)
+void FAnimNode_BlendStack::GatherDebugData(FNodeDebugData& DebugData)
 {
-	int32 FirstZeroWeightAnimPlayerIndex = 1;
-	for (; FirstZeroWeightAnimPlayerIndex < AnimPlayers.Num(); ++FirstZeroWeightAnimPlayerIndex)
+#if ENABLE_ANIM_DEBUG
+	DebugData.AddDebugItem(FString::Printf(TEXT("%s"), *DebugData.GetNodeName(this)));
+	for (int32 i = 0; i < AnimPlayers.Num(); ++i)
 	{
-		if (FMath::IsNearlyZero(AnimPlayers[FirstZeroWeightAnimPlayerIndex].GetBlendWeight()))
-		{
-			break;
-		}
+		const FPoseSearchAnimPlayer& AnimPlayer = AnimPlayers[i];
+		DebugData.AddDebugItem(FString::Printf(TEXT("%d) t:%.2f/%.2f m:%d %s"),
+				i, AnimPlayer.GetCurrentBlendInTime(), AnimPlayer.GetTotalBlendInTime(),
+				AnimPlayer.GetMirror() ? 1 : 0, *AnimPlayer.GetAnimName()));
 	}
-	// if the weight of the FirstZeroWeightAnimPlayerIndex-th AnimPlayer is zero, all the subsequent AnimPlayers contribution will be zero, so we delete from FirstZeroWeightAnimPlayerIndex forward
-
-	const int32 WantedAnimPlayersNum = FMath::Max(1, FMath::Min(FirstZeroWeightAnimPlayerIndex, MaxActiveBlends)); // we save at least one FPoseSearchAnimPlayer
+#endif // ENABLE_ANIM_DEBUG
 
-	while (AnimPlayers.Num() > WantedAnimPlayersNum)
+	// propagating GatherDebugData to the AnimPlayers
+	for (FPoseSearchAnimPlayer& AnimPlayer : AnimPlayers)
 	{
-		// @todo: instead of just popping the unwanted AnimPlayers we should store the blended pose + weight in a FPoseContext
-		AnimPlayers.PopLast();
+		AnimPlayer.GetMirrorNode().GatherDebugData(DebugData);
 	}
 }
 
diff --git a/Engine/Plugins/Experimental/Animation/PoseSearch/Source/Runtime/Public/PoseSearch/AnimNode_BlendStack.h b/Engine/Plugins/Experimental/Animation/PoseSearch/Source/Runtime/Public/PoseSearch/AnimNode_BlendStack.h
--- a/Engine/Plugins/Experimental/Animation/PoseSearch/Source/Runtime/Public/PoseSearch/AnimNode_BlendStack.h
+++ b/Engine/Plugins/Experimental/Animation/PoseSearch/Source/Runtime/Public/PoseSearch/AnimNode_BlendStack.h
@@ -18,23 +18,28 @@ struct FPoseSearchAnimPlayer
 	
 	void Initialize(ESearchIndexAssetType InAssetType, UAnimationAsset* AnimationAsset, float AccumulatedTime, bool bLoop, bool bMirrored, UMirrorDataTable* MirrorDataTable, float BlendTime, const UBlendProfile* BlendProfile, EAlphaBlendOption InBlendOption, FVector BlendParameters);
 	void Evaluate_AnyThread(FPoseContext& Output);
-	void Update_AnyThread(const FAnimationUpdateContext& Context);
+	void Update_AnyThread(const FAnimationUpdateContext& Context, float BlendWeight);
 	float GetAccumulatedTime() const;
-	float GetBlendWeight() const { return BlendWeight;  }
+	
 	float GetBlendInPercentage() const;
-	void SetBlendWeight(float InBlendWeight);
 	bool GetBlendInWeights(TArray<float>& Weights) const;
 	ESearchIndexAssetType GetAssetType() const { return AssetType; }
 	EAlphaBlendOption GetBlendOption() const { return BlendOption; }
+	void StorePoseContext(const FPoseContext& PoseContext);
+
+	float GetTotalBlendInTime() const { return TotalBlendInTime; }
+	float GetCurrentBlendInTime() const { return CurrentBlendInTime; }
+	bool GetMirror() const { return MirrorNode.GetMirror(); }
+	FString GetAnimName() const;
+	FAnimNode_Mirror_Standalone& GetMirrorNode() { return MirrorNode; }
 
 	// @todo: used only for DynamicPlayRateAdjustment. Remove once the functionality is integrated with the BlendStackNode
 	FAnimNode_SequencePlayer_Standalone GetSequencePlayerNode() { return SequencePlayerNode; }
-
+	
 protected:
+	void RestorePoseContext(FPoseContext& PoseContext) const;
 	void UpdateSourceLinkNode();
 
-	// @todo: maybe use an union between the standalone players?
-	
 	// Embedded sequence player node for playing animations from the motion matching database
 	FAnimNode_SequencePlayer_Standalone SequencePlayerNode;
 
@@ -44,6 +49,10 @@ protected:
 	// Embedded mirror node to handle mirroring if the pose search results in a mirrored sequence
 	FAnimNode_Mirror_Standalone MirrorNode;
 
+	// @todo: use FCompactHeapPose
+	// if AssetType == ESearchIndexAssetType::Invalid instead of using SequencePlayerNode or BlendSpacePlayerNode (wrapped in MirrorNode), the pose output will be from StoredPose
+	TArray<FTransform> StoredPose;
+
 	ESearchIndexAssetType AssetType = ESearchIndexAssetType::Sequence;
 	EAlphaBlendOption BlendOption = EAlphaBlendOption::Linear;
 
@@ -51,7 +60,6 @@ protected:
 
 	float TotalBlendInTime = 0.f;
 	float CurrentBlendInTime = 0.f;
-	float BlendWeight = 1.f;
 };
 
 USTRUCT(BlueprintInternalUseOnly)
@@ -60,14 +68,14 @@ struct POSESEARCH_API FAnimNode_BlendStack : public FAnimNode_AssetPlayerBase
 	GENERATED_BODY()
 
 	TDeque<FPoseSearchAnimPlayer> AnimPlayers;
+	int32 RequestedMaxActiveBlends = 0;
 
 	// FAnimNode_Base interface
 	virtual void Evaluate_AnyThread(FPoseContext& Output) override;
+	virtual void GatherDebugData(FNodeDebugData& DebugData) override;
 	// End of FAnimNode_Base interface
 
 	void BlendTo(ESearchIndexAssetType AssetType, UAnimationAsset* AnimationAsset, float AccumulatedTime = 0.f, bool bLoop = false, bool bMirrored = false, UMirrorDataTable* MirrorDataTable = nullptr, int32 MaxActiveBlends = 3, float BlendTime = 0.3f, const UBlendProfile* BlendProfile = nullptr, EAlphaBlendOption BlendOption = EAlphaBlendOption::Linear, FVector BlendParameters = FVector::Zero());
-	void CalculateWeights();
-	void PruneBlendStack(int32 MaxActiveBlends);
 
 	// FAnimNode_AssetPlayerBase interface
 	virtual float GetAccumulatedTime() const override;
