motion matching - fix for blendspace support, QOL improvements

diff --git a/Engine/Plugins/Experimental/Animation/PoseSearch/Source/Editor/Private/PoseSearchDatabaseViewModel.cpp b/Engine/Plugins/Experimental/Animation/PoseSearch/Source/Editor/Private/PoseSearchDatabaseViewModel.cpp
--- a/Engine/Plugins/Experimental/Animation/PoseSearch/Source/Editor/Private/PoseSearchDatabaseViewModel.cpp
+++ b/Engine/Plugins/Experimental/Animation/PoseSearch/Source/Editor/Private/PoseSearchDatabaseViewModel.cpp
@@ -121,6 +121,7 @@ namespace UE::PoseSearch
 
 	void FDatabaseViewModel::BuildSearchIndex()
 	{
+		PoseSearchDatabase->NotifyAssetChange();
 		PoseSearchDatabase->BeginCacheDerivedData();
 	}
 
@@ -321,7 +322,6 @@ namespace UE::PoseSearch
 	void FDatabaseViewModel::AddSequenceToDatabase(UAnimSequence* AnimSequence)
 	{
 		FPoseSearchDatabaseSequence& NewDbSequence = PoseSearchDatabase->Sequences.AddDefaulted_GetRef();
-
 		NewDbSequence.Sequence = AnimSequence;
 	}
 
diff --git a/Engine/Plugins/Experimental/Animation/PoseSearch/Source/Editor/Private/SPoseSearchDatabaseAssetList.cpp b/Engine/Plugins/Experimental/Animation/PoseSearch/Source/Editor/Private/SPoseSearchDatabaseAssetList.cpp
--- a/Engine/Plugins/Experimental/Animation/PoseSearch/Source/Editor/Private/SPoseSearchDatabaseAssetList.cpp
+++ b/Engine/Plugins/Experimental/Animation/PoseSearch/Source/Editor/Private/SPoseSearchDatabaseAssetList.cpp
@@ -585,7 +585,7 @@ namespace UE::PoseSearch
 		}
 
 		TArray<uint32> BlendspaceIndexArray;
-		SequenceIndexArray.SetNumUninitialized(Database->BlendSpaces.Num());
+		BlendspaceIndexArray.SetNumUninitialized(Database->BlendSpaces.Num());
 		for (int32 BlendspaceIdx = 0; BlendspaceIdx < Database->BlendSpaces.Num(); ++BlendspaceIdx)
 		{
 			BlendspaceIndexArray[BlendspaceIdx] = BlendspaceIdx;
diff --git a/Engine/Plugins/Experimental/Animation/PoseSearch/Source/Runtime/Private/PoseSearch.cpp b/Engine/Plugins/Experimental/Animation/PoseSearch/Source/Runtime/Private/PoseSearch.cpp
--- a/Engine/Plugins/Experimental/Animation/PoseSearch/Source/Runtime/Private/PoseSearch.cpp
+++ b/Engine/Plugins/Experimental/Animation/PoseSearch/Source/Runtime/Private/PoseSearch.cpp
@@ -289,8 +289,16 @@ static TArrayView<size_t> PopulateNonSelectableIdx(TArrayView<size_t> NonSelecta
 					NonSelectableIdxBuffer[NonSelectableIdxUsedSize++] = HistoricalPoseIndex.PoseIndex;
 
 #if UE_POSE_SEARCH_TRACE_ENABLED
-					const FPoseSearchCost PoseCost = Database->ComparePoses(SearchContext, HistoricalPoseIndex.PoseIndex, EPoseComparisonFlags::None, QueryValues);
-					SearchContext.BestCandidates.Add(PoseCost, HistoricalPoseIndex.PoseIndex, Database, EPoseCandidateFlags::DiscardedBy_PoseReselectHistory);
+					check(HistoricalPoseIndex.PoseIndex >= 0);
+					const FPoseSearchIndex* SearchIndex = Database->GetSearchIndex();
+					check(SearchIndex);
+					
+					// if we're editing the database and removing assets it's possible that the PoseIndicesHistory contains invalid pose indexes
+					if (HistoricalPoseIndex.PoseIndex < SearchIndex->NumPoses)
+					{
+						const FPoseSearchCost PoseCost = Database->ComparePoses(SearchContext, HistoricalPoseIndex.PoseIndex, EPoseComparisonFlags::None, QueryValues);
+						SearchContext.BestCandidates.Add(PoseCost, HistoricalPoseIndex.PoseIndex, Database, EPoseCandidateFlags::DiscardedBy_PoseReselectHistory);
+					}
 #endif
 				}
 				else
@@ -1032,28 +1040,41 @@ bool UPoseSearchDatabase::GetSkipSearchIfPossible() const
 
 bool UPoseSearchDatabase::IsValidForIndexing() const
 {
-	bool bValid = Schema && Schema->IsValid() && !Sequences.IsEmpty();
+	bool bValid = Schema && Schema->IsValid() && (!Sequences.IsEmpty() || !BlendSpaces.IsEmpty());
 
 	if (bValid)
 	{
-		bool bSequencesValid = true;
 		for (const FPoseSearchDatabaseSequence& DbSequence : Sequences)
 		{
 			if (!DbSequence.Sequence)
 			{
-				bSequencesValid = false;
+				bValid = false;
 				break;
 			}
 
-			USkeleton* SeqSkeleton = DbSequence.Sequence->GetSkeleton();
+			const USkeleton* SeqSkeleton = DbSequence.Sequence->GetSkeleton();
 			if (!SeqSkeleton || !SeqSkeleton->IsCompatible(Schema->Skeleton))
 			{
-				bSequencesValid = false;
+				bValid = false;
 				break;
 			}
 		}
 
-		bValid = bSequencesValid;
+		for (const FPoseSearchDatabaseBlendSpace& DbBlendSpace : BlendSpaces)
+		{
+			if (!DbBlendSpace.BlendSpace)
+			{
+				bValid = false;
+				break;
+			}
+
+			const USkeleton* SeqSkeleton = DbBlendSpace.BlendSpace->GetSkeleton();
+			if (!SeqSkeleton || !SeqSkeleton->IsCompatible(Schema->Skeleton))
+			{
+				bValid = false;
+				break;
+			}
+		}
 	}
 
 	return bValid;
