motion matching - removing deprecated DynamicPlayRateAdjustment

diff --git a/Engine/Plugins/Experimental/Animation/PoseSearch/Source/Runtime/PoseSearch.Build.cs b/Engine/Plugins/Experimental/Animation/PoseSearch/Source/Runtime/PoseSearch.Build.cs
--- a/Engine/Plugins/Experimental/Animation/PoseSearch/Source/Runtime/PoseSearch.Build.cs
+++ b/Engine/Plugins/Experimental/Animation/PoseSearch/Source/Runtime/PoseSearch.Build.cs
@@ -22,7 +22,6 @@ public class PoseSearch : ModuleRules
 				"AnimationCore",
 				"TraceLog",
 				"AnimGraphRuntime",
-				"DynamicPlayRate",
 				"GameplayTags",
 				"MotionTrajectory"
 			}
diff --git a/Engine/Plugins/Experimental/Animation/PoseSearch/Source/Runtime/Private/AnimNode_MotionMatching.cpp b/Engine/Plugins/Experimental/Animation/PoseSearch/Source/Runtime/Private/AnimNode_MotionMatching.cpp
--- a/Engine/Plugins/Experimental/Animation/PoseSearch/Source/Runtime/Private/AnimNode_MotionMatching.cpp
+++ b/Engine/Plugins/Experimental/Animation/PoseSearch/Source/Runtime/Private/AnimNode_MotionMatching.cpp
@@ -8,7 +8,6 @@
 #include "Animation/BlendSpace.h"
 #include "Animation/MotionTrajectoryTypes.h"
 #include "Animation/AnimRootMotionProvider.h"
-#include "DynamicPlayRate/DynamicPlayRateLibrary.h"
 #include "PoseSearch/AnimNode_PoseSearchHistoryCollector.h"
 #include "Trace/PoseSearchTraceLogger.h"
 
@@ -102,24 +101,6 @@ void FAnimNode_MotionMatching::UpdateAssetPlayer(const FAnimationUpdateContext&
 								Settings.MaxActiveBlends, Settings.BlendTime, Settings.BlendProfile, Settings.BlendOption, SearchIndexAsset->BlendParameters);
 	}
 
-	if (!BlendStackNode.AnimPlayers.IsEmpty() && BlendStackNode.AnimPlayers[0].GetAssetType() == ESearchIndexAssetType::Sequence)
-	{
-		// @todo: integrate this functionality with the BlendStackNode
-		
-		// Optionally applying dynamic play rate adjustment to chosen sequences based on predictive motion analysis
-		const float PlayRate = DynamicPlayRateAdjustment(
-			Context,
-			Trajectory,
-			DynamicPlayRateSettings,
-			BlendStackNode.AnimPlayers[0].GetSequencePlayerNode().GetSequence(),
-			GetAccumulatedTime(),
-			BlendStackNode.AnimPlayers[0].GetSequencePlayerNode().GetPlayRate(),
-			BlendStackNode.AnimPlayers[0].GetSequencePlayerNode().GetLoopAnimation()
-		);
-
-		BlendStackNode.AnimPlayers[0].GetSequencePlayerNode().SetPlayRate(PlayRate);
-	}
-
 	Source.Update(Context);
 }
 
diff --git a/Engine/Plugins/Experimental/Animation/PoseSearch/Source/Runtime/Public/PoseSearch/AnimNode_BlendStack.h b/Engine/Plugins/Experimental/Animation/PoseSearch/Source/Runtime/Public/PoseSearch/AnimNode_BlendStack.h
--- a/Engine/Plugins/Experimental/Animation/PoseSearch/Source/Runtime/Public/PoseSearch/AnimNode_BlendStack.h
+++ b/Engine/Plugins/Experimental/Animation/PoseSearch/Source/Runtime/Public/PoseSearch/AnimNode_BlendStack.h
@@ -33,9 +33,6 @@ struct FPoseSearchAnimPlayer
 	FString GetAnimName() const;
 	FAnimNode_Mirror_Standalone& GetMirrorNode() { return MirrorNode; }
 
-	// @todo: used only for DynamicPlayRateAdjustment. Remove once the functionality is integrated with the BlendStackNode
-	FAnimNode_SequencePlayer_Standalone GetSequencePlayerNode() { return SequencePlayerNode; }
-	
 protected:
 	void RestorePoseContext(FPoseContext& PoseContext) const;
 	void UpdateSourceLinkNode();
diff --git a/Engine/Plugins/Experimental/Animation/PoseSearch/Source/Runtime/Public/PoseSearch/AnimNode_MotionMatching.h b/Engine/Plugins/Experimental/Animation/PoseSearch/Source/Runtime/Public/PoseSearch/AnimNode_MotionMatching.h
--- a/Engine/Plugins/Experimental/Animation/PoseSearch/Source/Runtime/Public/PoseSearch/AnimNode_MotionMatching.h
+++ b/Engine/Plugins/Experimental/Animation/PoseSearch/Source/Runtime/Public/PoseSearch/AnimNode_MotionMatching.h
@@ -4,7 +4,6 @@
 
 #include "Animation/AnimNode_AssetPlayerBase.h"
 #include "Animation/MotionTrajectoryTypes.h"
-#include "DynamicPlayRate/DynamicPlayRateLibrary.h"
 #include "PoseSearch/AnimNode_BlendStack.h"
 #include "PoseSearch/PoseSearch.h"
 #include "PoseSearch/PoseSearchLibrary.h"
@@ -31,10 +30,6 @@ public:
 	UPROPERTY(EditAnywhere, BlueprintReadWrite, Category=Settings, meta=(PinShownByDefault))
 	FTrajectorySampleRange Trajectory;
 
-	// Settings for dynamic play rate adjustment on sequences chosen by motion matching
-	UPROPERTY(EditAnywhere, BlueprintReadWrite, Category=Settings, meta=(PinHiddenByDefault))
-	FDynamicPlayRateSettings DynamicPlayRateSettings;
-
 	// Settings for the core motion matching algorithm evaluation
 	UPROPERTY(EditAnywhere, BlueprintReadWrite, Category=Settings, meta=(PinHiddenByDefault))
 	FMotionMatchingSettings Settings;
diff --git a/Engine/Source/Runtime/Experimental/Animation/DynamicPlayRate/DynamicPlayRate.Build.cs b/Engine/Source/Runtime/Experimental/Animation/DynamicPlayRate/DynamicPlayRate.Build.cs
deleted file mode 100644
diff --git a/Engine/Source/Runtime/Experimental/Animation/DynamicPlayRate/Private/DynamicPlayRate/DynamicPlayRateLibrary.cpp b/Engine/Source/Runtime/Experimental/Animation/DynamicPlayRate/Private/DynamicPlayRate/DynamicPlayRateLibrary.cpp
deleted file mode 100644
diff --git a/Engine/Source/Runtime/Experimental/Animation/DynamicPlayRate/Public/DynamicPlayRate/DynamicPlayRateLibrary.h b/Engine/Source/Runtime/Experimental/Animation/DynamicPlayRate/Public/DynamicPlayRate/DynamicPlayRateLibrary.h
deleted file mode 100644
