motion matching - adding db ExcludeFromDatabaseParameters and deprecating BlockTransitionInterval

diff --git a/Engine/Plugins/Experimental/Animation/PoseSearch/Source/Runtime/Private/PoseSearch.cpp b/Engine/Plugins/Experimental/Animation/PoseSearch/Source/Runtime/Private/PoseSearch.cpp
--- a/Engine/Plugins/Experimental/Animation/PoseSearch/Source/Runtime/Private/PoseSearch.cpp
+++ b/Engine/Plugins/Experimental/Animation/PoseSearch/Source/Runtime/Private/PoseSearch.cpp
@@ -192,13 +192,13 @@ static TArrayView<size_t> PopulateNonSelectableIdx(TArrayView<size_t> NonSelecta
 	if (SearchContext.IsCurrentResultFromDatabase(Database) && SearchContext.CurrentResult.IsValid() && SearchContext.PoseJumpThresholdTime > 0.f)
 	{
 		const FPoseSearchIndexAsset* CurrentIndexAsset = SearchContext.CurrentResult.SearchIndexAsset;
-		const int32 PoseJumpIndexThreshold = FMath::FloorToInt(SearchContext.PoseJumpThresholdTime / Database->Schema->SamplingInterval);
+		const int32 PoseJumpIndexThreshold = FMath::FloorToInt(SearchContext.PoseJumpThresholdTime / Database->Schema->GetSamplingInterval());
 		const bool IsLooping = Database->IsSourceAssetLooping(CurrentIndexAsset);
 
 		for (int32 i = -PoseJumpIndexThreshold; i <= -1; ++i)
 		{
 			const int32 PoseIdx = SearchContext.CurrentResult.PoseIdx + i;
-			const float DeltaTime = i * Database->Schema->SamplingInterval;
+			const float DeltaTime = i * Database->Schema->GetSamplingInterval();
 
 			// @todo: should we use the quantized time associated to InOutMotionMatchingState.DbPoseIdx instead of InOutMotionMatchingState.AssetPlayerTime?
 			float PoseAssetPlayerTime = SearchContext.CurrentResult.AssetTime + DeltaTime;
@@ -238,7 +238,7 @@ static TArrayView<size_t> PopulateNonSelectableIdx(TArrayView<size_t> NonSelecta
 		for (int32 i = 0; i <= PoseJumpIndexThreshold; ++i)
 		{
 			const int32 PoseIdx = SearchContext.CurrentResult.PoseIdx + i;
-			const float DeltaTime = i * Database->Schema->SamplingInterval;
+			const float DeltaTime = i * Database->Schema->GetSamplingInterval();
 
 			// @todo: should we use the quantized time associated to InOutMotionMatchingState.DbPoseIdx instead of InOutMotionMatchingState.AssetPlayerTime?
 			float PoseAssetPlayerTime = SearchContext.CurrentResult.AssetTime + DeltaTime;
@@ -354,7 +354,6 @@ namespace UE::PoseSearch
 
 int32 FSchemaInitializer::AddBoneReference(const FBoneReference& BoneReference)
 {
-	check(BoneReferences.Num() + 1 <= UPoseSearchSchema::MaxBoneReferences);
 	return BoneReferences.AddUnique(BoneReference);
 }
 
@@ -364,9 +363,6 @@ void UPoseSearchSchema::Finalize()
 {
 	using namespace UE::PoseSearch;
 
-	SampleRate = FMath::Clamp(SampleRate, 1, 60);
-	SamplingInterval = 1.0f / SampleRate;
-
 	// Discard null channels
 	Channels.RemoveAll([](TObjectPtr<UPoseSearchFeatureChannel>& Channel) { return Channel.IsNull(); });
 
@@ -523,9 +519,7 @@ float FPoseSearchIndex::GetAssetTime(int32 PoseIdx, const FPoseSearchIndexAsset*
 	{
 		const FFloatInterval SamplingRange = Asset->SamplingInterval;
 
-		float AssetTime = FMath::Min(
-			SamplingRange.Min + Schema->SamplingInterval * (PoseIdx - Asset->FirstPoseIdx),
-			SamplingRange.Max);
+		float AssetTime = FMath::Min(SamplingRange.Min + Schema->GetSamplingInterval() * (PoseIdx - Asset->FirstPoseIdx), SamplingRange.Max);
 		return AssetTime;
 	}
 	
@@ -535,9 +529,7 @@ float FPoseSearchIndex::GetAssetTime(int32 PoseIdx, const FPoseSearchIndexAsset*
 
 		// For BlendSpaces the AssetTime is in the range [0, 1] while the Sampling Range
 		// is in real time (seconds)
-		float AssetTime = FMath::Min(
-			SamplingRange.Min + Schema->SamplingInterval * (PoseIdx - Asset->FirstPoseIdx),
-			SamplingRange.Max) / (Asset->NumPoses * Schema->SamplingInterval);
+		float AssetTime = FMath::Min(SamplingRange.Min + Schema->GetSamplingInterval() * (PoseIdx - Asset->FirstPoseIdx), SamplingRange.Max) / (Asset->NumPoses * Schema->GetSamplingInterval());
 		return AssetTime;
 	}
 	
@@ -716,24 +708,6 @@ FPoseSearchCost UPoseSearchSequenceMetaData::ComparePoses(
 	return FPoseSearchCost(DissimilarityCost, NotifyAddend, MirrorMismatchAddend, ContinuingPoseCostAddend);
 }
 
-
-//////////////////////////////////////////////////////////////////////////
-// FPoseSearchDatabaseSequence
-FFloatInterval FPoseSearchDatabaseSequence::GetEffectiveSamplingRange() const
-{
-	return UE::PoseSearch::GetEffectiveSamplingRange(Sequence, SamplingRange);
-}
-
-UAnimationAsset* FPoseSearchDatabaseSequence::GetAnimationAsset() const
-{
-	return Sequence;
-}
-
-bool FPoseSearchDatabaseSequence::IsLooping() const
-{
-	return Sequence->bLoop;
-}
-
 //////////////////////////////////////////////////////////////////////////
 // UPoseSearchDatabase
 
@@ -1139,45 +1113,46 @@ void UPoseSearchDatabase::CollectSimpleBlendSpaces()
 	SimpleBlendSpaces.Reset();
 }
 
-static void FindValidSequenceIntervals(const FPoseSearchDatabaseSequence& DbSequence, TArray<FFloatRange>& ValidRanges)
+void UPoseSearchDatabase::FindValidSequenceIntervals(const FPoseSearchDatabaseSequence& DbSequence, TArray<FFloatRange>& ValidRanges) const
 {
 	const UAnimSequence* Sequence = DbSequence.Sequence;
 	check(DbSequence.Sequence);
 
 	const float SequenceLength = DbSequence.Sequence->GetPlayLength();
-	const FFloatInterval EffectiveSamplingInterval = DbSequence.GetEffectiveSamplingRange();
+
+	const FFloatInterval EffectiveSamplingInterval = UE::PoseSearch::GetEffectiveSamplingRange(DbSequence.Sequence, DbSequence.SamplingRange);
+	FFloatRange EffectiveSamplingRange = FFloatRange::Inclusive(EffectiveSamplingInterval.Min, EffectiveSamplingInterval.Max);
+	if (!DbSequence.IsLooping())
+	{
+		const FFloatRange ExcludeFromDatabaseRange(ExcludeFromDatabaseParameters.SequenceStartInterval, SequenceLength - ExcludeFromDatabaseParameters.SequenceEndInterval);
+		EffectiveSamplingRange = FFloatRange::Intersection(EffectiveSamplingRange, ExcludeFromDatabaseRange);
+	}
 
 	// start from a single interval defined by the database sequence sampling range
 	ValidRanges.Empty();
-	ValidRanges.Add(FFloatRange::Inclusive(EffectiveSamplingInterval.Min, EffectiveSamplingInterval.Max));
+	ValidRanges.Add(EffectiveSamplingRange);
 
 	FAnimNotifyContext NotifyContext;
 	Sequence->GetAnimNotifies(0.0f, SequenceLength, NotifyContext);
 
 	for (const FAnimNotifyEventReference& EventReference : NotifyContext.ActiveNotifies)
 	{
-		const FAnimNotifyEvent* NotifyEvent = EventReference.GetNotify();
-		if (!NotifyEvent)
-		{
-			continue;
-		}
-
-		const UAnimNotifyState_PoseSearchExcludeFromDatabase* ExclusionNotifyState =
-			Cast<const UAnimNotifyState_PoseSearchExcludeFromDatabase>(NotifyEvent->NotifyStateClass);
-		if (ExclusionNotifyState)
+		if (const FAnimNotifyEvent* NotifyEvent = EventReference.GetNotify())
 		{
-			FFloatRange ExclusionRange = 
-				FFloatRange::Inclusive(NotifyEvent->GetTriggerTime(), NotifyEvent->GetEndTriggerTime());
-
-			// Split every valid range based on the exclusion range just found. Because this might increase the 
-			// number of ranges in ValidRanges, the algorithm iterates from end to start.
-			for (int RangeIdx = ValidRanges.Num() - 1; RangeIdx >= 0; --RangeIdx)
+			if (const UAnimNotifyState_PoseSearchExcludeFromDatabase* ExclusionNotifyState = Cast<const UAnimNotifyState_PoseSearchExcludeFromDatabase>(NotifyEvent->NotifyStateClass))
 			{
-				FFloatRange EvaluatedRange = ValidRanges[RangeIdx];
-				ValidRanges.RemoveAt(RangeIdx);
+				FFloatRange ExclusionRange = FFloatRange::Inclusive(NotifyEvent->GetTriggerTime(), NotifyEvent->GetEndTriggerTime());
+
+				// Split every valid range based on the exclusion range just found. Because this might increase the 
+				// number of ranges in ValidRanges, the algorithm iterates from end to start.
+				for (int RangeIdx = ValidRanges.Num() - 1; RangeIdx >= 0; --RangeIdx)
+				{
+					FFloatRange EvaluatedRange = ValidRanges[RangeIdx];
+					ValidRanges.RemoveAt(RangeIdx);
 
-				TArray<FFloatRange> Diff = FFloatRange::Difference(EvaluatedRange, ExclusionRange);
-				ValidRanges.Append(Diff);
+					TArray<FFloatRange> Diff = FFloatRange::Difference(EvaluatedRange, ExclusionRange);
+					ValidRanges.Append(Diff);
+				}
 			}
 		}
 	}
@@ -4054,20 +4029,9 @@ FTransform FAssetIndexer::MirrorTransform(const FTransform& Transform) const
 void FAssetIndexer::AddMetadata(int32 SampleIdx)
 {
 	const float SequenceLength = IndexingContext.MainSampler->GetPlayLength();
-	const float SampleTime = FMath::Min(SampleIdx * IndexingContext.Schema->SamplingInterval, SequenceLength);
+	const float SampleTime = FMath::Min(SampleIdx * IndexingContext.Schema->GetSamplingInterval(), SequenceLength);
 
 	Metadata = FPoseSearchPoseMetadata();
-
-	const bool bBlockTransition =
-		!IndexingContext.MainSampler->IsLoopable() &&
-		(SampleTime < IndexingContext.RequestedSamplingRange.Min + IndexingContext.BlockTransitionParameters.SequenceStartInterval ||
-		 SampleTime > IndexingContext.RequestedSamplingRange.Max - IndexingContext.BlockTransitionParameters.SequenceEndInterval);
-
-	if (bBlockTransition)
-	{
-		EnumAddFlags(Metadata.Flags, EPoseSearchPoseFlags::BlockTransition);
-	}
-
 	Metadata.CostAddend = IndexingContext.Schema->BaseCostBias;
 	Metadata.ContinuingPoseCostAddend = IndexingContext.Schema->ContinuingPoseCostBias;
 	
@@ -4685,6 +4649,7 @@ void FDatabaseIndexingContext::PrepareIndexers()
 	{
 		return BlendSpace ? &BlendSpaceSamplers[BlendSpaceSamplerMap[{BlendSpace, BlendParameters}]] : nullptr;
 	};
+
 	Indexers.Reserve(Database->GetSearchIndex()->Assets.Num());
 	for (int32 AssetIdx = 0; AssetIdx != SearchIndex->Assets.Num(); ++AssetIdx)
 	{
@@ -4693,7 +4658,6 @@ void FDatabaseIndexingContext::PrepareIndexers()
 		FAssetIndexingContext IndexerContext;
 		IndexerContext.SamplingContext = &SamplingContext;
 		IndexerContext.Schema = Database->Schema;
-		IndexerContext.BlockTransitionParameters = Database->BlockTransitionParameters;
 		IndexerContext.RequestedSamplingRange = SearchIndexAsset.SamplingInterval;
 		IndexerContext.bMirrored = SearchIndexAsset.bMirrored;
 
@@ -4702,10 +4666,8 @@ void FDatabaseIndexingContext::PrepareIndexers()
 			const FPoseSearchDatabaseSequence& DbSequence = Database->GetSequenceSourceAsset(&SearchIndexAsset);
 			const float SequenceLength = DbSequence.Sequence->GetPlayLength();
 			IndexerContext.MainSampler = GetSequenceSampler(DbSequence.Sequence);
-			IndexerContext.LeadInSampler =
-				SearchIndexAsset.SamplingInterval.Min == 0.0f ? GetSequenceSampler(DbSequence.LeadInSequence) : nullptr;
-			IndexerContext.FollowUpSampler =
-				SearchIndexAsset.SamplingInterval.Max == SequenceLength ? GetSequenceSampler(DbSequence.FollowUpSequence) : nullptr;
+			IndexerContext.LeadInSampler = SearchIndexAsset.SamplingInterval.Min == 0.0f ? GetSequenceSampler(DbSequence.LeadInSequence) : nullptr;
+			IndexerContext.FollowUpSampler = SearchIndexAsset.SamplingInterval.Max == SequenceLength ? GetSequenceSampler(DbSequence.FollowUpSequence) : nullptr;
 		}
 		else if (SearchIndexAsset.Type == ESearchIndexAssetType::BlendSpace)
 		{
diff --git a/Engine/Plugins/Experimental/Animation/PoseSearch/Source/Runtime/Private/PoseSearchFeatureChannels.cpp b/Engine/Plugins/Experimental/Animation/PoseSearch/Source/Runtime/Private/PoseSearchFeatureChannels.cpp
--- a/Engine/Plugins/Experimental/Animation/PoseSearch/Source/Runtime/Private/PoseSearchFeatureChannels.cpp
+++ b/Engine/Plugins/Experimental/Animation/PoseSearch/Source/Runtime/Private/PoseSearchFeatureChannels.cpp
@@ -61,7 +61,7 @@ T GetValueAtIndex(int32 Sample, const TArray<T>& Values)
 static void CollectBonePositions(TArray<FVector>& BonePositions, IAssetIndexer& Indexer, int8 SchemaBoneIdx)
 {
 	const FAssetIndexingContext& IndexingContext = Indexer.GetIndexingContext();
-	const float FiniteDelta = IndexingContext.Schema->SamplingInterval;
+	const float FiniteDelta = IndexingContext.Schema->GetSamplingInterval();
 	const float SampleTimeStart = FMath::Min(IndexingContext.BeginSampleIdx * FiniteDelta, IndexingContext.MainSampler->GetPlayLength());
 	const int32 NumSamples = IndexingContext.EndSampleIdx - IndexingContext.BeginSampleIdx;
 
@@ -345,7 +345,7 @@ void UPoseSearchFeatureChannel_Position::IndexAsset(UE::PoseSearch::IAssetIndexe
 		int32 VectorIdx = SampleIdx - IndexingContext.BeginSampleIdx;
 		FPoseSearchFeatureVectorBuilder& FeatureVector = IndexingOutput.PoseVectors[VectorIdx];
 
-		const float OriginSampleTime = FMath::Min(SampleIdx * IndexingContext.Schema->SamplingInterval, IndexingContext.MainSampler->GetPlayLength());
+		const float OriginSampleTime = FMath::Min(SampleIdx * IndexingContext.Schema->GetSamplingInterval(), IndexingContext.MainSampler->GetPlayLength());
 		const float SubsampleTime = OriginSampleTime + SampleTimeOffset;
 
 		bool ClampedPresent;
@@ -468,7 +468,7 @@ void UPoseSearchFeatureChannel_Heading::IndexAsset(UE::PoseSearch::IAssetIndexer
 		int32 VectorIdx = SampleIdx - IndexingContext.BeginSampleIdx;
 		FPoseSearchFeatureVectorBuilder& FeatureVector = IndexingOutput.PoseVectors[VectorIdx];
 
-		const float OriginSampleTime = FMath::Min(SampleIdx * IndexingContext.Schema->SamplingInterval, IndexingContext.MainSampler->GetPlayLength());
+		const float OriginSampleTime = FMath::Min(SampleIdx * IndexingContext.Schema->GetSamplingInterval(), IndexingContext.MainSampler->GetPlayLength());
 		const float SubsampleTime = OriginSampleTime + SampleTimeOffset;
 
 		bool ClampedPresent;
@@ -676,7 +676,7 @@ void UPoseSearchFeatureChannel_Pose::CalculatePhases(UE::PoseSearch::IAssetIndex
 	OutPhases.Reset();
 	OutPhases.AddDefaulted(SampledBones.Num());
 	
-	const float FiniteDelta = Indexer.GetIndexingContext().Schema->SamplingInterval;
+	const float FiniteDelta = Indexer.GetIndexingContext().Schema->GetSamplingInterval();
 
 	TArray<float> Signal;
 	TArray<float> SmoothedSignal;
@@ -788,7 +788,7 @@ void UPoseSearchFeatureChannel_Pose::AddPoseFeatures(UE::PoseSearch::IAssetIndex
 	const FAssetIndexingContext& IndexingContext = Indexer.GetIndexingContext();
 	const FAssetSamplingContext* SamplingContext = IndexingContext.SamplingContext;
 
-	const float SampleTime = FMath::Min(SampleIdx * IndexingContext.Schema->SamplingInterval, IndexingContext.MainSampler->GetPlayLength());
+	const float SampleTime = FMath::Min(SampleIdx * IndexingContext.Schema->GetSamplingInterval(), IndexingContext.MainSampler->GetPlayLength());
 
 	int32 DataOffset = ChannelDataOffset;
 	for (int32 ChannelBoneIdx = 0; ChannelBoneIdx != SampledBones.Num(); ++ChannelBoneIdx)
@@ -1462,7 +1462,7 @@ void UPoseSearchFeatureChannel_Trajectory::IndexAssetPrivate(const UE::PoseSearc
 	using FSampleInfo = IAssetIndexer::FSampleInfo;
 
 	const FAssetIndexingContext& IndexingContext = Indexer.GetIndexingContext();
-	float SampleTime = FMath::Min(SampleIdx * IndexingContext.Schema->SamplingInterval, IndexingContext.MainSampler->GetPlayLength());
+	float SampleTime = FMath::Min(SampleIdx * IndexingContext.Schema->GetSamplingInterval(), IndexingContext.MainSampler->GetPlayLength());
 	FSampleInfo Origin = Indexer.GetSampleInfo(SampleTime);
 
 	int32 DataOffset = ChannelDataOffset;
diff --git a/Engine/Plugins/Experimental/Animation/PoseSearch/Source/Runtime/Public/PoseSearch/PoseSearch.h b/Engine/Plugins/Experimental/Animation/PoseSearch/Source/Runtime/Public/PoseSearch/PoseSearch.h
--- a/Engine/Plugins/Experimental/Animation/PoseSearch/Source/Runtime/Public/PoseSearch/PoseSearch.h
+++ b/Engine/Plugins/Experimental/Animation/PoseSearch/Source/Runtime/Public/PoseSearch/PoseSearch.h
@@ -140,7 +140,6 @@ enum class EPoseSearchMirrorOption : int32
 };
 
 
-
 //////////////////////////////////////////////////////////////////////////
 // Common structs
 
@@ -163,7 +162,7 @@ struct POSESEARCH_API FPoseSearchExtrapolationParameters
 };
 
 USTRUCT()
-struct FPoseSearchBlockTransitionParameters
+struct FPoseSearchExcludeFromDatabaseParameters
 {
 	GENERATED_BODY()
 
@@ -174,7 +173,7 @@ struct FPoseSearchBlockTransitionParameters
 	// Excluding the end of sequences help ensure an exact future trajectory, and also prevents the selection of
 	// a sequence which will end too soon to be worth selecting.
 	UPROPERTY(EditAnywhere, Category = "Settings")
-	float SequenceEndInterval = 0.0f;
+	float SequenceEndInterval = 0.3f;
 };
 
 //////////////////////////////////////////////////////////////////////////
@@ -254,8 +253,7 @@ struct FAssetIndexingContext
 	const IAssetSampler* FollowUpSampler = nullptr;
 	bool bMirrored = false;
 	FFloatInterval RequestedSamplingRange = FFloatInterval(0.0f, 0.0f);
-	FPoseSearchBlockTransitionParameters BlockTransitionParameters;
-
+	
 	// Index this asset's data from BeginPoseIdx up to but not including EndPoseIdx
 	int32 BeginSampleIdx = 0;
 	int32 EndSampleIdx = 0;
@@ -327,7 +325,6 @@ public:
 
 } // namespace UE::PoseSearch
 
-
 //////////////////////////////////////////////////////////////////////////
 // Feature channels interface
 
@@ -387,8 +384,6 @@ protected:
 	int32 ChannelCardinality = -1;
 };
 
-
-
 //////////////////////////////////////////////////////////////////////////
 // Schema
 
@@ -414,8 +409,6 @@ private:
 	TArray<FBoneReference> BoneReferences;
 };
 
-
-
 } // namespace UE::PoseSearch
 
 USTRUCT()
@@ -439,17 +432,11 @@ class POSESEARCH_API UPoseSearchSchema : public UDataAsset, public IBoneReferenc
 	GENERATED_BODY()
 
 public:
-
-	static constexpr int32 DefaultSampleRate = 10;
-	static constexpr int32 MaxBoneReferences = MAX_int8;
-	static constexpr int32 MaxChannels = MAX_int8;
-	static constexpr int32 MaxFeatures = MAX_int8;
-
 	UPROPERTY(EditAnywhere, Category = "Schema")
 	TObjectPtr<USkeleton> Skeleton = nullptr;
 
-	UPROPERTY(EditAnywhere, meta = (ClampMin = "1", ClampMax = "60"), Category = "Schema")
-	int32 SampleRate = DefaultSampleRate;
+	UPROPERTY(EditAnywhere, meta = (ClampMin = "1", ClampMax = "240"), Category = "Schema")
+	int32 SampleRate = 10;
 
 	UPROPERTY(EditAnywhere, Instanced, BlueprintReadWrite, Category = "Schema")
 	TArray<TObjectPtr<UPoseSearchFeatureChannel>> Channels;
@@ -461,9 +448,6 @@ public:
 	UPROPERTY(EditAnywhere, Category = "Schema")
 	EPoseSearchDataPreprocessor DataPreprocessor = EPoseSearchDataPreprocessor::Normalize;
 
-	UPROPERTY()
-	float SamplingInterval = 1.0f / DefaultSampleRate;
-
 	UPROPERTY()
 	int32 SchemaCardinality = 0;
 
@@ -496,6 +480,8 @@ public:
 
 	int32 GetNumBones () const { return BoneIndices.Num(); }
 
+	float GetSamplingInterval() const { return 1.0f / SampleRate; }
+
 	// UObject
 	virtual void PreSave(FObjectPreSaveContext ObjectSaveContext) override;
 	virtual void PostLoad() override;
@@ -517,8 +503,6 @@ public:
 	bool BuildQuery(UE::PoseSearch::FSearchContext& SearchContext, FPoseSearchFeatureVectorBuilder& InOutQuery) const;
 };
 
-
-
 //////////////////////////////////////////////////////////////////////////
 // Search index
 
@@ -542,7 +526,6 @@ struct POSESEARCH_API FPoseSearchPoseMetadata
 	float ContinuingPoseCostAddend = 0.0f;
 };
 
-
 /**
 * Information about a source animation asset used by a search index.
 * Some source animation entries may generate multiple FPoseSearchIndexAsset entries.
@@ -715,10 +698,8 @@ struct POSESEARCH_API FPoseSearchDatabaseSequence : public FPoseSearchDatabaseAn
 	UPROPERTY(EditAnywhere, Category="Sequence")
 	TObjectPtr<UAnimSequence> FollowUpSequence = nullptr;
 
-	FFloatInterval GetEffectiveSamplingRange() const;
-
-	virtual UAnimationAsset* GetAnimationAsset() const override;
-	virtual bool IsLooping() const override;
+	virtual UAnimationAsset* GetAnimationAsset() const override { return Sequence; }
+	virtual bool IsLooping() const override { return Sequence->bLoop; }
 };
 
 /** An blend space entry in a UPoseSearchDatabase. */
@@ -809,7 +790,6 @@ public:
 #endif // WITH_EDITORONLY_DATA
 };
 
-
 /**
 * Helper object for writing features into a float buffer according to a feature vector layout.
 * Keeps track of which features are present, allowing the feature vector to be built up piecemeal.
@@ -843,7 +823,6 @@ private:
 	TArray<float> Values;
 };
 
-
 namespace UE::PoseSearch
 {
 
@@ -947,7 +926,6 @@ struct FSearchResult
 
 } // namespace UE::PoseSearch
 
-
 /** A data asset for indexing a collection of animation sequences. */
 UCLASS(Abstract, BlueprintType, Experimental)
 class POSESEARCH_API UPoseSearchSearchableAsset : public UDataAsset
@@ -958,7 +936,6 @@ public:
 	virtual UE::PoseSearch::FSearchResult Search(UE::PoseSearch::FSearchContext& SearchContext) const PURE_VIRTUAL(UPoseSearchSearchableAsset::Search, return UE::PoseSearch::FSearchResult(););
 };
 
-
 /** A data asset for indexing a collection of animation sequences. */
 UCLASS(BlueprintType, Category = "Animation|Pose Search", Experimental, meta = (DisplayName = "Motion Database"))
 class POSESEARCH_API UPoseSearchDatabase : public UPoseSearchSearchableAsset
@@ -973,7 +950,7 @@ public:
 	FPoseSearchExtrapolationParameters ExtrapolationParameters;
 
 	UPROPERTY(EditAnywhere, Category = "Database")
-	FPoseSearchBlockTransitionParameters BlockTransitionParameters = { 0.0f, 0.2f };
+	FPoseSearchExcludeFromDatabaseParameters ExcludeFromDatabaseParameters;
 
 	// Drag and drop animations here to add them in bulk to Sequences
 	UPROPERTY(EditAnywhere, Category = "Database", DisplayName="Drag And Drop Anims Here")
@@ -1050,6 +1027,7 @@ public: // UObject
 private:
 	void CollectSimpleSequences();
 	void CollectSimpleBlendSpaces();
+	void FindValidSequenceIntervals(const FPoseSearchDatabaseSequence& DbSequence, TArray<FFloatRange>& ValidRanges) const;
 
 public:
 	// Populates the FPoseSearchIndex::Assets array by evaluating the data in the Sequences array
@@ -1108,7 +1086,6 @@ protected:
 	UE::PoseSearch::FSearchResult SearchBruteForce(UE::PoseSearch::FSearchContext& SearchContext) const;
 };
 
-
 //////////////////////////////////////////////////////////////////////////
 // Sequence metadata
 
@@ -1144,7 +1121,6 @@ public: // UObject
 
 };
 
-
 //////////////////////////////////////////////////////////////////////////
 // Feature vector reader and builder
 
@@ -1392,7 +1368,6 @@ POSESEARCH_API void CompareFeatureVectors(TConstArrayView<float> A, TConstArrayV
 */
 POSESEARCH_API bool BuildIndex(const UAnimSequence* Sequence, UPoseSearchSequenceMetaData* SequenceMetaData);
 
-
 /**
 * Creates a pose search index for a collection of animations
 * 
@@ -1404,7 +1379,6 @@ POSESEARCH_API bool BuildIndex(const UPoseSearchDatabase* Database, FPoseSearchI
 
 } // namespace UE::PoseSearch
 
-
 UENUM()
 enum class EPoseSearchPostSearchStatus : uint8
 {
