motion matching - exposing more db stats in the reflection panel

diff --git a/Engine/Plugins/Experimental/Animation/PoseSearch/Source/Editor/Private/PoseSearchDatabaseEditor.cpp b/Engine/Plugins/Experimental/Animation/PoseSearch/Source/Editor/Private/PoseSearchDatabaseEditor.cpp
--- a/Engine/Plugins/Experimental/Animation/PoseSearch/Source/Editor/Private/PoseSearchDatabaseEditor.cpp
+++ b/Engine/Plugins/Experimental/Animation/PoseSearch/Source/Editor/Private/PoseSearchDatabaseEditor.cpp
@@ -496,40 +496,25 @@ namespace UE::PoseSearch
 			{
 				if (SelectedItem->SourceAssetType == ESearchIndexAssetType::Sequence)
 				{
-					UPoseSearchDatabaseSequenceReflection* NewSelectionReflection =
-						NewObject<UPoseSearchDatabaseSequenceReflection>();
+					UPoseSearchDatabaseSequenceReflection* NewSelectionReflection = NewObject<UPoseSearchDatabaseSequenceReflection>();
 					NewSelectionReflection->AddToRoot();
-					NewSelectionReflection->Sequence =
-						GetPoseSearchDatabase()->Sequences[SelectedItem->SourceAssetIdx];
+					NewSelectionReflection->Sequence = GetPoseSearchDatabase()->Sequences[SelectedItem->SourceAssetIdx];
 					NewSelectionReflection->SetSourceLink(SelectedItem, AssetTreeWidget);
 					SelectionReflections[SequenceSelectionIndex].Add(NewSelectionReflection);
 				}
 				else if (SelectedItem->SourceAssetType == ESearchIndexAssetType::BlendSpace)
 				{
-					UPoseSearchDatabaseBlendSpaceReflection* NewSelectionReflection =
-						NewObject<UPoseSearchDatabaseBlendSpaceReflection>();
+					UPoseSearchDatabaseBlendSpaceReflection* NewSelectionReflection = NewObject<UPoseSearchDatabaseBlendSpaceReflection>();
 					NewSelectionReflection->AddToRoot();
-					NewSelectionReflection->BlendSpace =
-						GetPoseSearchDatabase()->BlendSpaces[SelectedItem->SourceAssetIdx];
+					NewSelectionReflection->BlendSpace = GetPoseSearchDatabase()->BlendSpaces[SelectedItem->SourceAssetIdx];
 					NewSelectionReflection->SetSourceLink(SelectedItem, AssetTreeWidget);
 					SelectionReflections[BlendSpaceSelectionIndex].Add(NewSelectionReflection);
 				}
 				else
 				{
-					UPoseSearchDatabaseReflection* NewSelectionReflection =
-						NewObject<UPoseSearchDatabaseReflection>();
+					UPoseSearchDatabaseReflection* NewSelectionReflection = NewObject<UPoseSearchDatabaseReflection>();
 					NewSelectionReflection->AddToRoot();
-
-					const FPoseSearchIndex* SearchIndex = GetPoseSearchDatabase()->GetSearchIndex();
-					if (SearchIndex)
-					{
-						NewSelectionReflection->SearchIndex = *SearchIndex;
-					}
-					else
-					{
-						NewSelectionReflection->SearchIndex = FPoseSearchIndex();
-					}
-
+					NewSelectionReflection->Initialize(GetPoseSearchDatabase());
 					NewSelectionReflection->SetSourceLink(SelectedItem, AssetTreeWidget);
 					SelectionReflections[GroupSelectionIndex].Add(NewSelectionReflection);
 				}
diff --git a/Engine/Plugins/Experimental/Animation/PoseSearch/Source/Editor/Private/PoseSearchDatabaseEditorReflection.cpp b/Engine/Plugins/Experimental/Animation/PoseSearch/Source/Editor/Private/PoseSearchDatabaseEditorReflection.cpp
--- a/Engine/Plugins/Experimental/Animation/PoseSearch/Source/Editor/Private/PoseSearchDatabaseEditorReflection.cpp
+++ b/Engine/Plugins/Experimental/Animation/PoseSearch/Source/Editor/Private/PoseSearchDatabaseEditorReflection.cpp
@@ -4,6 +4,8 @@
 #include "SPoseSearchDatabaseAssetList.h"
 #include "PoseSearchDatabaseViewModel.h"
 
+#define LOCTEXT_NAMESPACE "UPoseSearchDatabaseReflection"
+
 #if WITH_EDITOR
 
 void UPoseSearchDatabaseReflectionBase::SetSourceLink(
@@ -45,3 +47,63 @@ void UPoseSearchDatabaseBlendSpaceReflection::PostEditChangeProperty(
 }
 
 #endif // WITH_EDITOR
+
+FText FPoseSearchDatabaseMemoryStats::ToMemoryBudgetText(int32 Size)
+{
+	if (Size < 1024)
+	{
+		return FText::Format(LOCTEXT("FPoseSearchDatabaseMemoryStats_Bytes", "{0} Bytes"), Size);
+	}
+
+	if (Size < 1024 * 1024)
+	{
+		return FText::Format(LOCTEXT("FPoseSearchDatabaseMemoryStats_KiloBytes", "{0} Kb"), Size / float(1024));
+	}
+
+	return FText::Format(LOCTEXT("FPoseSearchDatabaseMemoryStats_MegaBytes", "{0} Mb"), Size / float(1024 * 1024));
+}
+
+void FPoseSearchDatabaseMemoryStats::Initialize(const UPoseSearchDatabase* PoseSearchDatabase)
+{
+	const FPoseSearchIndex* SearchIndexSafe = PoseSearchDatabase->GetSearchIndexSafe();
+	check(SearchIndexSafe);
+
+	const int32 ValuesBytesSize = SearchIndexSafe->Values.GetAllocatedSize();
+	const int32 PCAValuesBytesSize = SearchIndexSafe->PCAValues.GetAllocatedSize();
+	const int32 KDTreeBytesSize = SearchIndexSafe->KDTree.GetAllocatedSize();
+	const int32 PoseMetadataBytesSize = SearchIndexSafe->PoseMetadata.GetAllocatedSize();
+	const int32 AssetsBytesSize = SearchIndexSafe->Assets.GetAllocatedSize();
+
+	const int32 OtherBytesSize = SearchIndexSafe->PCAProjectionMatrix.GetAllocatedSize() +
+							SearchIndexSafe->Mean.GetAllocatedSize() +
+							SearchIndexSafe->WeightsSqrt.GetAllocatedSize();
+	
+	const int32 EstimatedDatabaseBytesSize = ValuesBytesSize + PCAValuesBytesSize + KDTreeBytesSize + PoseMetadataBytesSize + AssetsBytesSize + OtherBytesSize;
+
+	ValuesSize = ToMemoryBudgetText(ValuesBytesSize);
+	PCAValuesSize = ToMemoryBudgetText(PCAValuesBytesSize);
+	KDTreeSize = ToMemoryBudgetText(KDTreeBytesSize);
+	PoseMetadataSize = ToMemoryBudgetText(PoseMetadataBytesSize);
+	AssetsSize = ToMemoryBudgetText(AssetsBytesSize);
+	EstimatedDatabaseSize = ToMemoryBudgetText(EstimatedDatabaseBytesSize);
+}
+
+void UPoseSearchDatabaseReflection::Initialize(const UPoseSearchDatabase* PoseSearchDatabase)
+{
+	if (PoseSearchDatabase)
+	{
+		const FPoseSearchIndex* SearchIndexSafe = PoseSearchDatabase->GetSearchIndexSafe();
+		if (SearchIndexSafe)
+		{
+			SearchIndex = *SearchIndexSafe;
+			MemoryStats.Initialize(PoseSearchDatabase);
+		}
+		else
+		{
+			SearchIndex = FPoseSearchIndex();
+			MemoryStats = FPoseSearchDatabaseMemoryStats();
+		}
+	}
+}
+
+#undef LOCTEXT_NAMESPACE
diff --git a/Engine/Plugins/Experimental/Animation/PoseSearch/Source/Editor/Public/PoseSearchDatabaseEditorReflection.h b/Engine/Plugins/Experimental/Animation/PoseSearch/Source/Editor/Public/PoseSearchDatabaseEditorReflection.h
--- a/Engine/Plugins/Experimental/Animation/PoseSearch/Source/Editor/Public/PoseSearchDatabaseEditorReflection.h
+++ b/Engine/Plugins/Experimental/Animation/PoseSearch/Source/Editor/Public/PoseSearchDatabaseEditorReflection.h
@@ -59,12 +59,44 @@ public:
 #endif
 };
 
+USTRUCT()
+struct FPoseSearchDatabaseMemoryStats
+{
+	GENERATED_BODY()
+
+	UPROPERTY(VisibleAnywhere, Category = "Stats")
+	FText EstimatedDatabaseSize;
+
+	UPROPERTY(VisibleAnywhere, Category = "Stats")
+	FText ValuesSize;
+
+	UPROPERTY(VisibleAnywhere, Category = "Stats")
+	FText PCAValuesSize;
+
+	UPROPERTY(VisibleAnywhere, Category = "Stats")
+	FText KDTreeSize;
+
+	UPROPERTY(VisibleAnywhere, Category = "Stats")
+	FText PoseMetadataSize;
+	
+	UPROPERTY(VisibleAnywhere, Category = "Stats")
+	FText AssetsSize;
+	
+	void Initialize(const UPoseSearchDatabase* PoseSearchDatabase);
+	static FText ToMemoryBudgetText(int32 Size);
+};
+
 UCLASS()
 class UPoseSearchDatabaseReflection : public UPoseSearchDatabaseReflectionBase
 {
 	GENERATED_BODY()
 
 public:
+	void Initialize(const UPoseSearchDatabase* PoseSearchDatabase);
+
+	UPROPERTY(VisibleAnywhere, Category = "Selected Database Search Index")
+	FPoseSearchDatabaseMemoryStats MemoryStats;
+
 	UPROPERTY(VisibleAnywhere, Category = "Selected Database Search Index")
 	FPoseSearchIndex SearchIndex;
 };
diff --git a/Engine/Plugins/Experimental/Animation/PoseSearch/Source/Runtime/Private/KDTree.cpp b/Engine/Plugins/Experimental/Animation/PoseSearch/Source/Runtime/Private/KDTree.cpp
--- a/Engine/Plugins/Experimental/Animation/PoseSearch/Source/Runtime/Private/KDTree.cpp
+++ b/Engine/Plugins/Experimental/Animation/PoseSearch/Source/Runtime/Private/KDTree.cpp
@@ -85,6 +85,22 @@ bool FKDTree::FindNeighbors(KNNResultSet& Result, const float* Query) const
 #endif
 }
 
+SIZE_T FKDTree::GetAllocatedSize() const
+{
+	SIZE_T AllocatedSize = sizeof(DataSrc) + sizeof(Impl);
+
+#if UE_POSE_SEARCH_USE_NANOFLANN
+	if (Impl)
+	{
+		AllocatedSize += sizeof(FKDTreeImplementation);
+		AllocatedSize += Impl->usedMemory(*Impl);
+	}
+#endif
+
+	return AllocatedSize;
+}
+
+
 #if UE_POSE_SEARCH_USE_NANOFLANN
 FArchive& SerializeSubTree(FArchive& Ar, FKDTree& KDTree, FKDTreeImplementation::NodePtr KDTreeNode)
 {
diff --git a/Engine/Plugins/Experimental/Animation/PoseSearch/Source/Runtime/Public/PoseSearch/KDTree.h b/Engine/Plugins/Experimental/Animation/PoseSearch/Source/Runtime/Public/PoseSearch/KDTree.h
--- a/Engine/Plugins/Experimental/Animation/PoseSearch/Source/Runtime/Public/PoseSearch/KDTree.h
+++ b/Engine/Plugins/Experimental/Animation/PoseSearch/Source/Runtime/Public/PoseSearch/KDTree.h
@@ -121,6 +121,7 @@ struct POSESEARCH_API FKDTree
 
 	bool FindNeighbors(KNNResultSet& Result, const float* Query) const;
 	void Construct(int32 Count, int32 Dim, const float* Data, int32 MaxLeafSize);
+	SIZE_T GetAllocatedSize() const;
 
 	DataSource DataSrc;
 	FKDTreeImplementation* Impl = nullptr;
