motion matching - defaulting FDatabaseViewModel::AnimationPreviewMode to EAnimationPreviewMode::OriginalAndMirrored

diff --git a/Engine/Plugins/Experimental/Animation/PoseSearch/Source/Editor/Public/PoseSearchDatabaseViewModel.h b/Engine/Plugins/Experimental/Animation/PoseSearch/Source/Editor/Public/PoseSearchDatabaseViewModel.h
--- a/Engine/Plugins/Experimental/Animation/PoseSearch/Source/Editor/Public/PoseSearchDatabaseViewModel.h
+++ b/Engine/Plugins/Experimental/Animation/PoseSearch/Source/Editor/Public/PoseSearchDatabaseViewModel.h
@@ -135,7 +135,7 @@ namespace UE::PoseSearch
 		EFeaturesDrawMode PoseFeaturesDrawMode = EFeaturesDrawMode::All;
 
 		/** What animations to show in the viewport */
-		EAnimationPreviewMode AnimationPreviewMode = EAnimationPreviewMode::OriginalOnly;
+		EAnimationPreviewMode AnimationPreviewMode = EAnimationPreviewMode::OriginalAndMirrored;
 
 		TArray<TSharedPtr<FDatabaseAssetTreeNode>> SelectedNodes;
 
