fix for warning V568: It's odd that 'sizeof()' operator evaluates the size of a pointer to a class, but not the size of the 'Impl' class object.

diff --git a/Engine/Plugins/Experimental/Animation/PoseSearch/Source/Runtime/Private/KDTree.cpp b/Engine/Plugins/Experimental/Animation/PoseSearch/Source/Runtime/Private/KDTree.cpp
--- a/Engine/Plugins/Experimental/Animation/PoseSearch/Source/Runtime/Private/KDTree.cpp
+++ b/Engine/Plugins/Experimental/Animation/PoseSearch/Source/Runtime/Private/KDTree.cpp
@@ -87,7 +87,7 @@ bool FKDTree::FindNeighbors(KNNResultSet& Result, const float* Query) const
 
 SIZE_T FKDTree::GetAllocatedSize() const
 {
-	SIZE_T AllocatedSize = sizeof(DataSrc) + sizeof(Impl);
+	SIZE_T AllocatedSize = sizeof(FKDTree);
 
 #if UE_POSE_SEARCH_USE_NANOFLANN
 	if (Impl)
