motion matching - fix for log spamming "FAnimNode_BlendStack NumBlends/MaxNumBlends..."

diff --git a/Engine/Plugins/Experimental/Animation/PoseSearch/Source/Runtime/Private/AnimNode_BlendStack.cpp b/Engine/Plugins/Experimental/Animation/PoseSearch/Source/Runtime/Private/AnimNode_BlendStack.cpp
--- a/Engine/Plugins/Experimental/Animation/PoseSearch/Source/Runtime/Private/AnimNode_BlendStack.cpp
+++ b/Engine/Plugins/Experimental/Animation/PoseSearch/Source/Runtime/Private/AnimNode_BlendStack.cpp
@@ -276,7 +276,6 @@ void FAnimNode_BlendStack::Evaluate_AnyThread(FPoseContext& Output)
 		const int32 NumSkeletonBones = RefSkeleton.GetNum();
 		TArray<float> Weights;
 
-		const int BoundaryIndex = RequestedMaxActiveBlends + 1;
 		// evaluating from the second last to the first into EvaluationPoseContext and then blend it with the with the Output (initialized with the last AnimPlayer evaluation)
 		for (int32 i = BlendStackSize - 2; i >= 0; --i)
 		{
@@ -294,7 +293,7 @@ void FAnimNode_BlendStack::Evaluate_AnyThread(FPoseContext& Output)
 			}
 			Output = BlendedPoseContext; // @todo: this should not be necessary either: optimize it away!
 
-			if (i >= BoundaryIndex 
+			if (i >= RequestedMaxActiveBlends
 #if ENABLE_ANIM_DEBUG
 				&& CVarAnimBlendStackPruningEnable.GetValueOnAnyThread()
 #endif // ENABLE_ANIM_DEBUG
@@ -303,8 +302,10 @@ void FAnimNode_BlendStack::Evaluate_AnyThread(FPoseContext& Output)
 				// too many AnimPlayers! we don't have enought available blends to hold them all, so we accumulate the blended poses into Output / BlendedPoseContext, until...
 				AnimPlayers.PopLast();
 				
-				if (i == BoundaryIndex)
+				if (i == RequestedMaxActiveBlends)
 				{
+					check(AnimPlayers.Num() == RequestedMaxActiveBlends + 1);
+
 					// we can store Output / BlendedPoseContext into the last AnimPlayer, that will hold a static pose, no longer an animation playing
 					AnimPlayers[i].StorePoseContext(Output);
 				}
@@ -312,11 +313,9 @@ void FAnimNode_BlendStack::Evaluate_AnyThread(FPoseContext& Output)
 		}
 
 		const int32 ActiveBlends = AnimPlayers.Num() - 1;
-		if (ActiveBlends >= RequestedMaxActiveBlends)
+		if (ActiveBlends > RequestedMaxActiveBlends)
 		{
-			UE_LOG(LogPoseSearch, Display,
-				TEXT("FAnimNode_BlendStack NumBlends/MaxNumBlends %d / %d"),
-				ActiveBlends, RequestedMaxActiveBlends);
+			UE_LOG(LogPoseSearch, Display, TEXT("FAnimNode_BlendStack NumBlends/MaxNumBlends %d / %d"), ActiveBlends, RequestedMaxActiveBlends);
 		}
 	}
 }
