motion matching - blend space support WIP

diff --git a/Engine/Plugins/Experimental/Animation/PoseSearch/Source/Editor/Private/PoseSearchDatabaseViewModel.cpp b/Engine/Plugins/Experimental/Animation/PoseSearch/Source/Editor/Private/PoseSearchDatabaseViewModel.cpp
--- a/Engine/Plugins/Experimental/Animation/PoseSearch/Source/Editor/Private/PoseSearchDatabaseViewModel.cpp
+++ b/Engine/Plugins/Experimental/Animation/PoseSearch/Source/Editor/Private/PoseSearchDatabaseViewModel.cpp
@@ -63,60 +63,64 @@ namespace UE::PoseSearch
 	{
 		RemovePreviewActors();
 
-		if (AnimationPreviewMode == EAnimationPreviewMode::None)
+		MaxPreviewPlayLength = 0.0f;
+
+		if (AnimationPreviewMode == EAnimationPreviewMode::None && SelectedNodes.IsEmpty())
 		{
 			return;
 		}
 
-		if (SelectedNodes.Num() > 0)
+		TSet<int32> AssociatedSequencesAssetIndices;
+		TSet<int32> AssociatedBlendSpacesAssetIndices;
+		for (const TSharedPtr<FDatabaseAssetTreeNode>& SelectedNode : SelectedNodes)
 		{
-			TSet<int32> AssociatedAssetIndices;
-			for (const TSharedPtr<FDatabaseAssetTreeNode>& SelectedNode : SelectedNodes)
+			if (SelectedNode->SourceAssetType == ESearchIndexAssetType::Sequence)
 			{
-				if (SelectedNode->SourceAssetType == ESearchIndexAssetType::Sequence)
-				{
-					AssociatedAssetIndices.Add(SelectedNode->SourceAssetIdx);
-				}
+				AssociatedSequencesAssetIndices.Add(SelectedNode->SourceAssetIdx);
 			}
-			
-			MaxPreviewPlayLength = 0.0f;
+			else if (SelectedNode->SourceAssetType == ESearchIndexAssetType::BlendSpace)
+			{
+				AssociatedBlendSpacesAssetIndices.Add(SelectedNode->SourceAssetIdx);
+			}
+		}
 
-			const FPoseSearchIndex* SearchIndex = PoseSearchDatabase->GetSearchIndex();
-			if (SearchIndex)
+		const FPoseSearchIndex* SearchIndex = PoseSearchDatabase->GetSearchIndexSafe();
+		if (SearchIndex)
+		{
+			for (const FPoseSearchIndexAsset& IndexAsset : SearchIndex->Assets)
 			{
-				for (const FPoseSearchIndexAsset& IndexAsset : SearchIndex->Assets)
+				bool bSpawn = false;
+				if (AnimationPreviewMode == EAnimationPreviewMode::OriginalAndMirrored || !IndexAsset.bMirrored)
 				{
-					const bool bIsAssociatedToSelection =
-						IndexAsset.Type == ESearchIndexAssetType::Sequence &&
-						AssociatedAssetIndices.Contains(IndexAsset.SourceAssetIdx);
-
-					const bool bSpawn =
-						bIsAssociatedToSelection &&
-						(AnimationPreviewMode == EAnimationPreviewMode::OriginalAndMirrored ||
-						 !IndexAsset.bMirrored);
+					bool bIsAssociatedToSelection = false;
+					if (IndexAsset.Type == ESearchIndexAssetType::Sequence)
+					{
+						bSpawn = AssociatedSequencesAssetIndices.Contains(IndexAsset.SourceAssetIdx);
+					}
+					else if (IndexAsset.Type == ESearchIndexAssetType::BlendSpace)
+					{
+						bSpawn = AssociatedBlendSpacesAssetIndices.Contains(IndexAsset.SourceAssetIdx);
+					}
+				}
 
-					if (bSpawn)
+				if (bSpawn)
+				{
+					FDatabasePreviewActor PreviewActor = SpawnPreviewActor(IndexAsset);
+					if (PreviewActor.IsValid())
 					{
-						FDatabasePreviewActor PreviewActor = SpawnPreviewActor(IndexAsset);
-						if (PreviewActor.IsValid())
+						const UAnimationAsset* AnimationAsset = PreviewActor.AnimInstance->GetAnimationAsset();
+						if (AnimationAsset)
 						{
-							const UAnimSequence* Sequence =
-								Cast<UAnimSequence>(PreviewActor.AnimInstance->GetAnimationAsset());
-							if (Sequence)
-							{
-								MaxPreviewPlayLength = FMath::Max(MaxPreviewPlayLength, Sequence->GetPlayLength());
-							}
-
-							PreviewActors.Add(PreviewActor);
+							MaxPreviewPlayLength = FMath::Max(MaxPreviewPlayLength, AnimationAsset->GetPlayLength());
 						}
+						
+						PreviewActors.Add(PreviewActor);
 					}
 				}
-
-				UpdatePreviewActors();
 			}
-		}
 
-		// todo: do blendspaces afterwards
+			UpdatePreviewActors();
+		}
 	}
 
 	void FDatabaseViewModel::BuildSearchIndex()
@@ -164,13 +168,26 @@ namespace UE::PoseSearch
 
 	FDatabasePreviewActor FDatabaseViewModel::SpawnPreviewActor(const FPoseSearchIndexAsset& IndexAsset)
 	{
-		USkeleton* Skeleton = PoseSearchDatabase->Schema->Skeleton;
+		FDatabasePreviewActor PreviewActor;
 
-		FPoseSearchDatabaseSequence DatabaseSequence = PoseSearchDatabase->Sequences[IndexAsset.SourceAssetIdx];
+		if (!PoseSearchDatabase || !PoseSearchDatabase->Schema || !PoseSearchDatabase->Schema->Skeleton)
+		{
+			return PreviewActor;
+		}
 
-		FDatabasePreviewActor PreviewActor;
+		UAnimationAsset* PreviewAsset = nullptr;
+		if (IndexAsset.Type == ESearchIndexAssetType::Sequence)
+		{
+			FPoseSearchDatabaseSequence DatabaseSequence = PoseSearchDatabase->Sequences[IndexAsset.SourceAssetIdx];
+			PreviewAsset = DatabaseSequence.Sequence;
+		}
+		else if (IndexAsset.Type == ESearchIndexAssetType::BlendSpace)
+		{
+			FPoseSearchDatabaseBlendSpace DatabaseBlendSpace = PoseSearchDatabase->BlendSpaces[IndexAsset.SourceAssetIdx];
+			PreviewAsset = DatabaseBlendSpace.BlendSpace;
+		}
 
-		if (!Skeleton || !DatabaseSequence.Sequence)
+		if (!PreviewAsset)
 		{
 			return PreviewActor;
 		}
@@ -193,10 +210,11 @@ namespace UE::PoseSearch
 		PreviewActor.AnimInstance = NewObject<UAnimPreviewInstance>(PreviewActor.Mesh.Get());
 		PreviewActor.Mesh->PreviewInstance = PreviewActor.AnimInstance.Get();
 		PreviewActor.AnimInstance->InitializeAnimation();
+		PreviewActor.AnimInstance->SetBlendSpacePosition(IndexAsset.BlendParameters);
 
-		PreviewActor.Mesh->SetSkeletalMesh(Skeleton->GetPreviewMesh(true));
-		PreviewActor.Mesh->EnablePreview(true, DatabaseSequence.Sequence);
-		PreviewActor.AnimInstance->SetAnimationAsset(DatabaseSequence.Sequence, false, 0.0f);
+		PreviewActor.Mesh->SetSkeletalMesh(PoseSearchDatabase->Schema->Skeleton->GetPreviewMesh(true));
+		PreviewActor.Mesh->EnablePreview(true, PreviewAsset);
+		PreviewActor.AnimInstance->SetAnimationAsset(PreviewAsset, false, 0.0f);
 		if (IndexAsset.bMirrored && PoseSearchDatabase->Schema)
 		{
 			PreviewActor.AnimInstance->SetMirrorDataTable(PoseSearchDatabase->Schema->MirrorDataTable);
@@ -247,17 +265,27 @@ namespace UE::PoseSearch
 		TSharedPtr<FDatabasePreviewScene> PreviewScene = PreviewScenePtr.Pin();
 		for (FDatabasePreviewActor& PreviewActor : GetPreviewActors())
 		{
-			if (const UAnimSequence* Sequence = Cast<UAnimSequence>(PreviewActor.AnimInstance->GetAnimationAsset()))
+			if (const UAnimationAsset* PreviewAsset = PreviewActor.AnimInstance->GetAnimationAsset())
 			{
 				float CurrentTime = 0.0f;
-				FAnimationRuntime::AdvanceTime(false, PlayTime, CurrentTime, Sequence->GetPlayLength());
+				FAnimationRuntime::AdvanceTime(false, PlayTime, CurrentTime, PreviewAsset->GetPlayLength());
 
-				FTransform RootMotion = Sequence->ExtractRootMotionFromRange(0.0f, CurrentTime);
-				if (PreviewActor.AnimInstance->GetMirrorDataTable())
+				if (const UAnimSequence* AnimSequence = Cast<UAnimSequence>(PreviewAsset))
+				{
+					FTransform RootMotion = FTransform::Identity;
+					RootMotion = AnimSequence->ExtractRootMotionFromRange(0.0f, CurrentTime);
+					if (PreviewActor.AnimInstance->GetMirrorDataTable())
+					{
+						RootMotion = MirrorRootMotion(RootMotion, PreviewActor.AnimInstance->GetMirrorDataTable());
+					}
+					PreviewActor.Actor->SetActorTransform(RootMotion);
+				}
+				else if (const UBlendSpace* BlendSpace = Cast<UBlendSpace>(PreviewAsset))
 				{
-					RootMotion = MirrorRootMotion(RootMotion, PreviewActor.AnimInstance->GetMirrorDataTable());
+					// @todo: add BS support
+					//RootMotion = BlendSpace->ExtractRootMotionFromRange(0.0f, CurrentTime);
 				}
-				PreviewActor.Actor->SetActorTransform(RootMotion);
+				
 				PreviewActor.AnimInstance->SetPosition(CurrentTime);
 
 				PreviewActor.CurrentPoseIndex = PoseSearchDatabase->GetPoseIndexFromTime(CurrentTime, &PreviewActor.IndexAsset);
diff --git a/Engine/Plugins/Experimental/Animation/PoseSearch/Source/Editor/Private/SPoseSearchDatabaseAssetList.cpp b/Engine/Plugins/Experimental/Animation/PoseSearch/Source/Editor/Private/SPoseSearchDatabaseAssetList.cpp
--- a/Engine/Plugins/Experimental/Animation/PoseSearch/Source/Editor/Private/SPoseSearchDatabaseAssetList.cpp
+++ b/Engine/Plugins/Experimental/Animation/PoseSearch/Source/Editor/Private/SPoseSearchDatabaseAssetList.cpp
@@ -364,6 +364,7 @@ namespace UE::PoseSearch
 		}
 
 		SkeletonView.Pin()->RefreshTreeView(false, true);
+		ViewModelPtr->BuildSearchIndex();
 	}
 
 	FSlateColor SDatabaseAssetListItem::GetNameTextColorAndOpacity() const
diff --git a/Engine/Plugins/Experimental/Animation/PoseSearch/Source/Runtime/Private/PoseSearch.cpp b/Engine/Plugins/Experimental/Animation/PoseSearch/Source/Runtime/Private/PoseSearch.cpp
--- a/Engine/Plugins/Experimental/Animation/PoseSearch/Source/Runtime/Private/PoseSearch.cpp
+++ b/Engine/Plugins/Experimental/Animation/PoseSearch/Source/Runtime/Private/PoseSearch.cpp
@@ -399,8 +399,7 @@ void UPoseSearchSchema::PostLoad()
 #if WITH_EDITOR
 void UPoseSearchSchema::PostEditChangeProperty(FPropertyChangedEvent& PropertyChangedEvent)
 {
-	SchemaCardinality = -1;
-
+	Finalize();
 	Super::PostEditChangeProperty(PropertyChangedEvent);
 }
 
@@ -432,8 +431,6 @@ bool UPoseSearchSchema::IsValid() const
 
 	bValid &= (BoneReferences.Num() == BoneIndices.Num());
 
-	bValid &= SchemaCardinality > 0;
-	
 	return bValid;
 }
 
@@ -910,7 +907,7 @@ int32 UPoseSearchDatabase::GetNumberOfPrincipalComponents() const
 
 #if WITH_EDITOR
 
-void UPoseSearchDatabase::AddRawSequenceToWriter(UAnimSequence* Sequence, UE::PoseSearch::FDerivedDataKeyBuilder& KeyBuilder)
+static void AddRawSequenceToWriter(UAnimSequence* Sequence, UE::PoseSearch::FDerivedDataKeyBuilder& KeyBuilder)
 {
 	if (Sequence)
 	{
@@ -922,7 +919,7 @@ void UPoseSearchDatabase::AddRawSequenceToWriter(UAnimSequence* Sequence, UE::Po
 	}
 }
 
-void UPoseSearchDatabase::AddPoseSearchNotifiesToWriter(UAnimSequence* Sequence, UE::PoseSearch::FDerivedDataKeyBuilder& KeyBuilder)
+static void AddPoseSearchNotifiesToWriter(UAnimSequence* Sequence, UE::PoseSearch::FDerivedDataKeyBuilder& KeyBuilder)
 {
 	if (!Sequence)
 	{
@@ -951,31 +948,39 @@ void UPoseSearchDatabase::AddPoseSearchNotifiesToWriter(UAnimSequence* Sequence,
 	}
 }
 
-void UPoseSearchDatabase::AddDbSequenceToWriter(FPoseSearchDatabaseSequence& DbSequence, UE::PoseSearch::FDerivedDataKeyBuilder& InOutWriter)
+
+void FPoseSearchDatabaseSequence::BuildDerivedDataKey(UE::PoseSearch::FDerivedDataKeyBuilder& KeyBuilder)
 {
-	AddRawSequenceToWriter(DbSequence.Sequence, InOutWriter);
-	AddRawSequenceToWriter(DbSequence.LeadInSequence, InOutWriter);
-	AddRawSequenceToWriter(DbSequence.FollowUpSequence, InOutWriter);
+	KeyBuilder << bEnabled;
+	KeyBuilder << SamplingRange;
+	KeyBuilder << MirrorOption;
+
+	AddRawSequenceToWriter(Sequence, KeyBuilder);
+	AddRawSequenceToWriter(LeadInSequence, KeyBuilder);
+	AddRawSequenceToWriter(FollowUpSequence, KeyBuilder);
 
-	AddPoseSearchNotifiesToWriter(DbSequence.Sequence, InOutWriter);
+	AddPoseSearchNotifiesToWriter(Sequence, KeyBuilder);
 }
 
-void UPoseSearchDatabase::AddDbBlendSpaceToWriter(FPoseSearchDatabaseBlendSpace& DbBlendSpace, UE::PoseSearch::FDerivedDataKeyBuilder& KeyBuilder)
+void FPoseSearchDatabaseBlendSpace::BuildDerivedDataKey(UE::PoseSearch::FDerivedDataKeyBuilder& KeyBuilder)
 {
-	if (IsValid(DbBlendSpace.BlendSpace))
-	{
-		const TArray<FBlendSample>& BlendSpaceSamples = DbBlendSpace.BlendSpace->GetBlendSamples();
-		for (const FBlendSample& Sample : BlendSpaceSamples)
-		{
-			AddRawSequenceToWriter(Sample.Animation, KeyBuilder);
-			FVector SampleValue = Sample.SampleValue;
-			float RateScale = Sample.RateScale;
- 			KeyBuilder << SampleValue;
- 			KeyBuilder << RateScale;
-		}
+	KeyBuilder << bEnabled;
+	KeyBuilder << MirrorOption;
+	KeyBuilder << bUseGridForSampling;
+	KeyBuilder << NumberOfHorizontalSamples;
+	KeyBuilder << NumberOfVerticalSamples;
 
-		KeyBuilder << DbBlendSpace.BlendSpace->bLoop;
+	const TArray<FBlendSample>& BlendSpaceSamples = BlendSpace->GetBlendSamples();
+	for (const FBlendSample& Sample : BlendSpaceSamples)
+	{
+		AddRawSequenceToWriter(Sample.Animation, KeyBuilder);
+		FVector SampleValue = Sample.SampleValue;
+		float RateScale = Sample.RateScale;
+ 		KeyBuilder << SampleValue;
+ 		KeyBuilder << RateScale;
 	}
+
+	KeyBuilder << BlendSpace->bLoop;
 }
 
 void UPoseSearchDatabase::BuildDerivedDataKey(UE::PoseSearch::FDerivedDataKeyBuilder& KeyBuilder)
@@ -989,15 +994,13 @@ void UPoseSearchDatabase::BuildDerivedDataKey(UE::PoseSearch::FDerivedDataKeyBui
 
 	for (FPoseSearchDatabaseSequence& DbSequence : Sequences)
 	{
-		AddDbSequenceToWriter(DbSequence, KeyBuilder);
+		DbSequence.BuildDerivedDataKey(KeyBuilder);
 	}
 
 	for (FPoseSearchDatabaseBlendSpace& DbBlendSpace : BlendSpaces)
 	{
-		AddDbBlendSpaceToWriter(DbBlendSpace, KeyBuilder);
+		DbBlendSpace.BuildDerivedDataKey(KeyBuilder);
 	}
-
-	// @todo: add database version if POSESEARCHDB_DERIVEDDATA_VER is not enough
 }
 
 #endif // WITH_EDITOR
@@ -1574,6 +1577,8 @@ void UPoseSearchDatabase::BeginCacheForCookedPlatformData(const ITargetPlatform*
 
 bool UPoseSearchDatabase::IsCachedCookedPlatformDataLoaded(const ITargetPlatform* TargetPlatform)
 {
+	check(IsInGameThread());
+
 	if (!PrivateDerivedData)
 	{
 		PrivateDerivedData = new FPoseSearchDatabaseDerivedData();
@@ -2602,7 +2607,7 @@ FColor FDebugDrawParams::GetColor(int32 ColorPreset) const
 	FLinearColor Color = FLinearColor::Red;
 
 	const UPoseSearchSchema* Schema = GetSchema();
-	if (!Schema || Schema->SchemaCardinality <= 0)
+	if (!Schema || !Schema->IsValid())
 	{
 		Color = FLinearColor::Red;
 	}
diff --git a/Engine/Plugins/Experimental/Animation/PoseSearch/Source/Runtime/Private/PoseSearchDerivedData.cpp b/Engine/Plugins/Experimental/Animation/PoseSearch/Source/Runtime/Private/PoseSearchDerivedData.cpp
--- a/Engine/Plugins/Experimental/Animation/PoseSearch/Source/Runtime/Private/PoseSearchDerivedData.cpp
+++ b/Engine/Plugins/Experimental/Animation/PoseSearch/Source/Runtime/Private/PoseSearchDerivedData.cpp
@@ -34,7 +34,14 @@ namespace UE::PoseSearch
 
 void FPoseSearchDatabaseDerivedData::Cache(UPoseSearchDatabase& Database, bool bForceRebuild)
 {
+	check(IsInGameThread());
+
 	CancelCache();
+
+#if UE_POSE_SEARCH_DERIVED_DATA_LOGGING
+	UE_LOG(LogPoseSearch, Log, TEXT("%s: Cache"), *LexToString(PendingDerivedDataKey));
+#endif // UE_POSE_SEARCH_DERIVED_DATA_LOGGING
+
 	if (Database.IsValidForIndexing())
 	{
 		CreateDatabaseBuildTask(Database, bForceRebuild);
@@ -50,8 +57,14 @@ void FPoseSearchDatabaseDerivedData::Cache(UPoseSearchDatabase& Database, bool b
 
 void FPoseSearchDatabaseDerivedData::CancelCache()
 {
+	check(IsInGameThread());
+
 	if (AsyncTask)
 	{
+#if UE_POSE_SEARCH_DERIVED_DATA_LOGGING
+		UE_LOG(LogPoseSearch, Log, TEXT("%s: CancelCache"), *LexToString(PendingDerivedDataKey));
+#endif // UE_POSE_SEARCH_DERIVED_DATA_LOGGING
+
 		AsyncTask->Cancel();
 	}
 
@@ -60,8 +73,14 @@ void FPoseSearchDatabaseDerivedData::CancelCache()
 
 void FPoseSearchDatabaseDerivedData::FinishCache()
 {
+	check(IsInGameThread());
+
 	if (AsyncTask)
 	{
+#if UE_POSE_SEARCH_DERIVED_DATA_LOGGING
+		UE_LOG(LogPoseSearch, Log, TEXT("%s: FinishCache"), *LexToString(PendingDerivedDataKey));
+#endif // UE_POSE_SEARCH_DERIVED_DATA_LOGGING
+
 		AsyncTask->Wait();
 		delete AsyncTask;
 		AsyncTask = nullptr;
@@ -70,6 +89,8 @@ void FPoseSearchDatabaseDerivedData::FinishCache()
 
 void FPoseSearchDatabaseDerivedData::CreateDatabaseBuildTask(UPoseSearchDatabase& Database, bool bForceRebuild)
 {
+	check(IsInGameThread());
+
 	AsyncTask = new UE::PoseSearch::FPoseSearchDatabaseAsyncCacheTask(Database, *this, bForceRebuild);
 }
 
@@ -129,6 +150,10 @@ namespace UE::PoseSearch
 	{
 		using namespace UE::DerivedData;
 
+#if UE_POSE_SEARCH_DERIVED_DATA_LOGGING
+		UE_LOG(LogPoseSearch, Log, TEXT("%s: BeginCache - '%s'"), *LexToString(DerivedData.PendingDerivedDataKey), *Database.GetName());
+#endif // UE_POSE_SEARCH_DERIVED_DATA_LOGGING
+
 		TArray<FCacheGetRequest> CacheRequests;
 		ECachePolicy CachePolicy = ECachePolicy::Default;
 		const FCacheKey CacheKey{ Bucket, DerivedData.PendingDerivedDataKey };
@@ -165,6 +190,10 @@ namespace UE::PoseSearch
 				const bool bIndexReady = BuildIndex(&Database, DerivedData.SearchIndex);
 
 				WriteIndexToCache(NewKey);
+
+#if UE_POSE_SEARCH_DERIVED_DATA_LOGGING
+				UE_LOG(LogPoseSearch, Log, TEXT("%s: BuildAndWrite"), *LexToString(DerivedData.PendingDerivedDataKey));
+#endif // UE_POSE_SEARCH_DERIVED_DATA_LOGGING
 			}
 		});
 	}
@@ -190,6 +219,10 @@ namespace UE::PoseSearch
 	{
 		using namespace UE::DerivedData;
 
+#if UE_POSE_SEARCH_DERIVED_DATA_LOGGING
+		UE_LOG(LogPoseSearch, Log, TEXT("%s: BuildIndexFromCacheRecord"), *LexToString(DerivedData.PendingDerivedDataKey));
+#endif // UE_POSE_SEARCH_DERIVED_DATA_LOGGING
+
 		DerivedData.SearchIndex.Reset();
 		DerivedData.SearchIndex.Schema = Database.Schema;
 
@@ -211,12 +244,15 @@ namespace UE::PoseSearch
 		Database.BuildDerivedDataKey(KeyBuilder);
 		FDerivedDataKeyBuilder::HashDigestType Hash = KeyBuilder.Finalize();
 
+		// Stores a BLAKE3-160 hash, taken from the first 20 bytes of a BLAKE3-256 hash
+		FIoHash IoHash(Hash);
+
 #if UE_POSE_SEARCH_DERIVED_DATA_LOGGING
 		const double TotalTime = FPlatformTime::Seconds() - StartTime;
-		UE_LOG(LogPoseSearch, Log, TEXT("DDC key for '%s': %s (%.0f µs)"), *Database.GetName(), *WriteToString<48>(Hash), TotalTime * 1e6);
+		UE_LOG(LogPoseSearch, Log, TEXT("%s: CreateKey ('%s' - %.0f µs)"), *LexToString(IoHash), *LexToString(Hash), *Database.GetName(), TotalTime * 1e6);
 #endif // UE_POSE_SEARCH_DERIVED_DATA_LOGGING
 
-		return Hash;
+		return IoHash;
 	}
 
 #endif // WITH_EDITOR
diff --git a/Engine/Plugins/Experimental/Animation/PoseSearch/Source/Runtime/Public/PoseSearch/PoseSearch.h b/Engine/Plugins/Experimental/Animation/PoseSearch/Source/Runtime/Public/PoseSearch/PoseSearch.h
--- a/Engine/Plugins/Experimental/Animation/PoseSearch/Source/Runtime/Public/PoseSearch/PoseSearch.h
+++ b/Engine/Plugins/Experimental/Animation/PoseSearch/Source/Runtime/Public/PoseSearch/PoseSearch.h
@@ -664,6 +664,9 @@ struct POSESEARCH_API FPoseSearchDatabaseAnimationAssetBase
 	virtual ~FPoseSearchDatabaseAnimationAssetBase() {}
 	virtual UAnimationAsset* GetAnimationAsset() const { return nullptr; }
 	virtual bool IsLooping() const { return false; }
+#if WITH_EDITOR
+	virtual void BuildDerivedDataKey(UE::PoseSearch::FDerivedDataKeyBuilder& KeyBuilder) {}
+#endif
 };
 
 /** An entry in a UPoseSearchDatabase. */
@@ -700,6 +703,9 @@ struct POSESEARCH_API FPoseSearchDatabaseSequence : public FPoseSearchDatabaseAn
 
 	virtual UAnimationAsset* GetAnimationAsset() const override { return Sequence; }
 	virtual bool IsLooping() const override { return Sequence->bLoop; }
+#if WITH_EDITOR
+	virtual void BuildDerivedDataKey(UE::PoseSearch::FDerivedDataKeyBuilder& KeyBuilder) override;
+#endif
 };
 
 /** An blend space entry in a UPoseSearchDatabase. */
@@ -722,15 +728,17 @@ struct POSESEARCH_API FPoseSearchDatabaseBlendSpace : public FPoseSearchDatabase
 	UPROPERTY(EditAnywhere, Category = "BlendSpace")
 	bool bUseGridForSampling = true;
 
-	UPROPERTY(EditAnywhere, Category = "BlendSpace", meta = (ClampMin = "1", UIMin = "1", UIMax = "25"))
+	UPROPERTY(EditAnywhere, Category = "BlendSpace", meta = (EditCondition = "!bUseGridForSampling", EditConditionHides, ClampMin = "1", UIMin = "1", UIMax = "25"))
 	int32 NumberOfHorizontalSamples = 5;
 
-	UPROPERTY(EditAnywhere, Category = "BlendSpace", meta = (ClampMin = "1", UIMin = "1", UIMax = "25"))
+	UPROPERTY(EditAnywhere, Category = "BlendSpace", meta = (EditCondition = "!bUseGridForSampling", EditConditionHides, ClampMin = "1", UIMin = "1", UIMax = "25"))
 	int32 NumberOfVerticalSamples = 5;
 
 	virtual UAnimationAsset* GetAnimationAsset() const override;
 	virtual bool IsLooping() const override;
-
+#if WITH_EDITOR
+	virtual void BuildDerivedDataKey(UE::PoseSearch::FDerivedDataKeyBuilder& KeyBuilder) override;
+#endif
 public:
 
 	void GetBlendSpaceParameterSampleRanges(
@@ -984,7 +992,6 @@ public:
 
 	virtual ~UPoseSearchDatabase();
 
-	FPoseSearchIndex* GetSearchIndex();
 	const FPoseSearchIndex* GetSearchIndex() const;
 	const FPoseSearchIndex* GetSearchIndexSafe() const;
 	
@@ -1004,36 +1011,27 @@ public:
 	const FString GetSourceAssetName(const FPoseSearchIndexAsset* SearchIndexAsset) const;
 	int32 GetNumberOfPrincipalComponents() const;
 	
-#if WITH_EDITOR
-	void BuildDerivedDataKey(UE::PoseSearch::FDerivedDataKeyBuilder& KeyBuilder);
-private:
-	static void AddRawSequenceToWriter(UAnimSequence* Sequence, UE::PoseSearch::FDerivedDataKeyBuilder& KeyBuilder);
-	static void AddPoseSearchNotifiesToWriter(UAnimSequence* Sequence, UE::PoseSearch::FDerivedDataKeyBuilder& KeyBuilder);
-	static void AddDbSequenceToWriter(FPoseSearchDatabaseSequence& DbSequence, UE::PoseSearch::FDerivedDataKeyBuilder& InOutWriter);
-	static void AddDbBlendSpaceToWriter(FPoseSearchDatabaseBlendSpace& DbBlendSpace, UE::PoseSearch::FDerivedDataKeyBuilder& KeyBuilder);
-#endif // WITH_EDITOR
-
-public: // UObject
+	// UObject
 	virtual void PostLoad() override;
 	virtual void PostSaveRoot(FObjectPostSaveRootContext ObjectSaveContext) override;
 	virtual void Serialize(FArchive& Ar) override;
 
+	// Populates the FPoseSearchIndex::Assets array by evaluating the data in the Sequences array
+	bool TryInitSearchIndexAssets(FPoseSearchIndex& OutSearchIndex) const;
+
 #if WITH_EDITOR
+	void BuildDerivedDataKey(UE::PoseSearch::FDerivedDataKeyBuilder& KeyBuilder);
 	virtual void PostEditChangeProperty(struct FPropertyChangedEvent& PropertyChangedEvent) override;
 	virtual void BeginCacheForCookedPlatformData(const ITargetPlatform* TargetPlatform) override;
 	virtual bool IsCachedCookedPlatformDataLoaded(const ITargetPlatform* TargetPlatform) override;
 #endif // WITH_EDITOR
 
 private:
+	FPoseSearchIndex* GetSearchIndex();
 	void CollectSimpleSequences();
 	void CollectSimpleBlendSpaces();
 	void FindValidSequenceIntervals(const FPoseSearchDatabaseSequence& DbSequence, TArray<FFloatRange>& ValidRanges) const;
-
-public:
-	// Populates the FPoseSearchIndex::Assets array by evaluating the data in the Sequences array
-	bool TryInitSearchIndexAssets(FPoseSearchIndex& OutSearchIndex) const;
-
-private:
+	
 	FPoseSearchDatabaseDerivedData* PrivateDerivedData;
 
 #if WITH_EDITOR
