motion matching - fix for PopulateNonSelectableIdx PoseIdx out of bound

diff --git a/Engine/Plugins/Experimental/Animation/PoseSearch/Source/Runtime/Private/PoseSearch.cpp b/Engine/Plugins/Experimental/Animation/PoseSearch/Source/Runtime/Private/PoseSearch.cpp
--- a/Engine/Plugins/Experimental/Animation/PoseSearch/Source/Runtime/Private/PoseSearch.cpp
+++ b/Engine/Plugins/Experimental/Animation/PoseSearch/Source/Runtime/Private/PoseSearch.cpp
@@ -197,20 +197,23 @@ static TArrayView<size_t> PopulateNonSelectableIdx(TArrayView<size_t> NonSelecta
 
 		for (int32 i = -PoseJumpIndexThreshold; i <= -1; ++i)
 		{
-			const int32 PoseIdx = SearchContext.CurrentResult.PoseIdx + i;
-			const float DeltaTime = i * Database->Schema->GetSamplingInterval();
-
-			// @todo: should we use the quantized time associated to InOutMotionMatchingState.DbPoseIdx instead of InOutMotionMatchingState.AssetPlayerTime?
-			float PoseAssetPlayerTime = SearchContext.CurrentResult.AssetTime + DeltaTime;
+			int32 PoseIdx = SearchContext.CurrentResult.PoseIdx + i;
+			bool bIsPoseInRange = false;
 			if (IsLooping)
 			{
-				while (PoseAssetPlayerTime < CurrentIndexAsset->SamplingInterval.Min)
+				bIsPoseInRange = true;
+
+				while (PoseIdx < CurrentIndexAsset->FirstPoseIdx)
 				{
-					PoseAssetPlayerTime += CurrentIndexAsset->SamplingInterval.Size();
+					PoseIdx += CurrentIndexAsset->NumPoses;
 				}
 			}
+			else if (CurrentIndexAsset->IsPoseInRange(PoseIdx))
+			{
+				bIsPoseInRange = true;
+			}
 
-			if (CurrentIndexAsset->SamplingInterval.Contains(PoseAssetPlayerTime))
+			if (bIsPoseInRange)
 			{
 				if (NonSelectableIdxUsedSize < NonSelectableIdxBuffer.Num())
 				{
@@ -237,20 +240,23 @@ static TArrayView<size_t> PopulateNonSelectableIdx(TArrayView<size_t> NonSelecta
 
 		for (int32 i = 0; i <= PoseJumpIndexThreshold; ++i)
 		{
-			const int32 PoseIdx = SearchContext.CurrentResult.PoseIdx + i;
-			const float DeltaTime = i * Database->Schema->GetSamplingInterval();
-
-			// @todo: should we use the quantized time associated to InOutMotionMatchingState.DbPoseIdx instead of InOutMotionMatchingState.AssetPlayerTime?
-			float PoseAssetPlayerTime = SearchContext.CurrentResult.AssetTime + DeltaTime;
+			int32 PoseIdx = SearchContext.CurrentResult.PoseIdx + i;
+			bool bIsPoseInRange = false;
 			if (IsLooping)
 			{
-				while (PoseAssetPlayerTime > CurrentIndexAsset->SamplingInterval.Max)
+				bIsPoseInRange = true;
+
+				while (PoseIdx >= CurrentIndexAsset->FirstPoseIdx + CurrentIndexAsset->NumPoses)
 				{
-					PoseAssetPlayerTime -= CurrentIndexAsset->SamplingInterval.Size();
+					PoseIdx -= CurrentIndexAsset->NumPoses;
 				}
 			}
+			else if (CurrentIndexAsset->IsPoseInRange(PoseIdx))
+			{
+				bIsPoseInRange = true;
+			}
 
-			if (CurrentIndexAsset->SamplingInterval.Contains(PoseAssetPlayerTime))
+			if (bIsPoseInRange)
 			{
 				if (NonSelectableIdxUsedSize < NonSelectableIdxBuffer.Num())
 				{
