performance improvements, less dangling pointers, fix for bForceInterrupt, fix for SearchBruteForce NonSelectableIdx filtering, fix for crashes while scrubbing rewind debugger and editing MM db or config

diff --git a/Engine/Plugins/Experimental/Animation/PoseSearch/Source/Editor/Private/PoseSearchDatabaseViewModel.cpp b/Engine/Plugins/Experimental/Animation/PoseSearch/Source/Editor/Private/PoseSearchDatabaseViewModel.cpp
--- a/Engine/Plugins/Experimental/Animation/PoseSearch/Source/Editor/Private/PoseSearchDatabaseViewModel.cpp
+++ b/Engine/Plugins/Experimental/Animation/PoseSearch/Source/Editor/Private/PoseSearchDatabaseViewModel.cpp
@@ -87,8 +87,10 @@ namespace UE::PoseSearch
 		const FPoseSearchIndex* SearchIndex = PoseSearchDatabase->GetSearchIndexSafe();
 		if (SearchIndex)
 		{
-			for (const FPoseSearchIndexAsset& IndexAsset : SearchIndex->Assets)
+			for (int32 IndexAssetIndex = 0; IndexAssetIndex < SearchIndex->Assets.Num(); ++IndexAssetIndex)
 			{
+				const FPoseSearchIndexAsset& IndexAsset = SearchIndex->Assets[IndexAssetIndex];
+
 				bool bSpawn = false;
 				if (AnimationPreviewMode == EAnimationPreviewMode::OriginalAndMirrored || !IndexAsset.bMirrored)
 				{
@@ -105,7 +107,7 @@ namespace UE::PoseSearch
 
 				if (bSpawn)
 				{
-					FDatabasePreviewActor PreviewActor = SpawnPreviewActor(IndexAsset);
+					FDatabasePreviewActor PreviewActor = SpawnPreviewActor(IndexAssetIndex);
 					if (PreviewActor.IsValid())
 					{
 						const UAnimationAsset* AnimationAsset = PreviewActor.AnimInstance->GetAnimationAsset();
@@ -166,73 +168,74 @@ namespace UE::PoseSearch
 		PlayTime = MaxPreviewPlayLength;
 	}
 
-	FDatabasePreviewActor FDatabaseViewModel::SpawnPreviewActor(const FPoseSearchIndexAsset& IndexAsset)
+	FDatabasePreviewActor FDatabaseViewModel::SpawnPreviewActor(int32 IndexAssetIndex)
 	{
 		FDatabasePreviewActor PreviewActor;
 
-		if (!PoseSearchDatabase || !PoseSearchDatabase->Schema || !PoseSearchDatabase->Schema->Skeleton)
+		if (PoseSearchDatabase && PoseSearchDatabase->IsValidForIndexing())
 		{
-			return PreviewActor;
-		}
+			if (const FPoseSearchIndex* SearchIndex = PoseSearchDatabase->GetSearchIndexSafe())
+			{
+				const FPoseSearchIndexAsset& IndexAsset = SearchIndex->Assets[IndexAssetIndex];
+				UAnimationAsset* PreviewAsset = nullptr;
+				if (IndexAsset.Type == ESearchIndexAssetType::Sequence)
+				{
+					FPoseSearchDatabaseSequence DatabaseSequence = PoseSearchDatabase->Sequences[IndexAsset.SourceAssetIdx];
+					PreviewAsset = DatabaseSequence.Sequence;
+				}
+				else if (IndexAsset.Type == ESearchIndexAssetType::BlendSpace)
+				{
+					FPoseSearchDatabaseBlendSpace DatabaseBlendSpace = PoseSearchDatabase->BlendSpaces[IndexAsset.SourceAssetIdx];
+					PreviewAsset = DatabaseBlendSpace.BlendSpace;
+				}
 
-		UAnimationAsset* PreviewAsset = nullptr;
-		if (IndexAsset.Type == ESearchIndexAssetType::Sequence)
-		{
-			FPoseSearchDatabaseSequence DatabaseSequence = PoseSearchDatabase->Sequences[IndexAsset.SourceAssetIdx];
-			PreviewAsset = DatabaseSequence.Sequence;
-		}
-		else if (IndexAsset.Type == ESearchIndexAssetType::BlendSpace)
-		{
-			FPoseSearchDatabaseBlendSpace DatabaseBlendSpace = PoseSearchDatabase->BlendSpaces[IndexAsset.SourceAssetIdx];
-			PreviewAsset = DatabaseBlendSpace.BlendSpace;
-		}
+				if (!PreviewAsset)
+				{
+					return PreviewActor;
+				}
 
-		if (!PreviewAsset)
-		{
-			return PreviewActor;
-		}
+				PreviewActor.IndexAssetIndex = IndexAssetIndex;
+				PreviewActor.CurrentPoseIndex = INDEX_NONE;
 
-		PreviewActor.IndexAsset = IndexAsset;
-		PreviewActor.CurrentPoseIndex = INDEX_NONE;
+				// todo: use preview when possible
+				UClass* PreviewClass = AActor::StaticClass();
+				const FTransform SpawnTransform = FTransform::Identity;
 
-		// todo: use preview when possible
-		UClass* PreviewClass = AActor::StaticClass();
-		const FTransform SpawnTransform = FTransform::Identity;
+				FActorSpawnParameters Params;
+				Params.SpawnCollisionHandlingOverride = ESpawnActorCollisionHandlingMethod::AdjustIfPossibleButAlwaysSpawn;
+				PreviewActor.Actor = GetWorld()->SpawnActor<AActor>(PreviewClass, SpawnTransform, Params);
+				PreviewActor.Actor->SetFlags(RF_Transient);
 
-		FActorSpawnParameters Params;
-		Params.SpawnCollisionHandlingOverride = ESpawnActorCollisionHandlingMethod::AdjustIfPossibleButAlwaysSpawn;
-		PreviewActor.Actor = GetWorld()->SpawnActor<AActor>(PreviewClass, SpawnTransform, Params);
-		PreviewActor.Actor->SetFlags(RF_Transient);
+				PreviewActor.Mesh = NewObject<UDebugSkelMeshComponent>(PreviewActor.Actor.Get());
+				PreviewActor.Mesh->RegisterComponentWithWorld(GetWorld());
 
-		PreviewActor.Mesh = NewObject<UDebugSkelMeshComponent>(PreviewActor.Actor.Get());
-		PreviewActor.Mesh->RegisterComponentWithWorld(GetWorld());
+				PreviewActor.AnimInstance = NewObject<UAnimPreviewInstance>(PreviewActor.Mesh.Get());
+				PreviewActor.Mesh->PreviewInstance = PreviewActor.AnimInstance.Get();
+				PreviewActor.AnimInstance->InitializeAnimation();
+				PreviewActor.AnimInstance->SetBlendSpacePosition(IndexAsset.BlendParameters);
 
-		PreviewActor.AnimInstance = NewObject<UAnimPreviewInstance>(PreviewActor.Mesh.Get());
-		PreviewActor.Mesh->PreviewInstance = PreviewActor.AnimInstance.Get();
-		PreviewActor.AnimInstance->InitializeAnimation();
-		PreviewActor.AnimInstance->SetBlendSpacePosition(IndexAsset.BlendParameters);
+				PreviewActor.Mesh->SetSkeletalMesh(PoseSearchDatabase->Schema->Skeleton->GetPreviewMesh(true));
+				PreviewActor.Mesh->EnablePreview(true, PreviewAsset);
+				PreviewActor.AnimInstance->SetAnimationAsset(PreviewAsset, false, 0.0f);
+				if (IndexAsset.bMirrored && PoseSearchDatabase->Schema)
+				{
+					PreviewActor.AnimInstance->SetMirrorDataTable(PoseSearchDatabase->Schema->MirrorDataTable);
+				}
 
-		PreviewActor.Mesh->SetSkeletalMesh(PoseSearchDatabase->Schema->Skeleton->GetPreviewMesh(true));
-		PreviewActor.Mesh->EnablePreview(true, PreviewAsset);
-		PreviewActor.AnimInstance->SetAnimationAsset(PreviewAsset, false, 0.0f);
-		if (IndexAsset.bMirrored && PoseSearchDatabase->Schema)
-		{
-			PreviewActor.AnimInstance->SetMirrorDataTable(PoseSearchDatabase->Schema->MirrorDataTable);
-		}
+				PreviewActor.AnimInstance->PlayAnim(false, 0.0f);
 
-		PreviewActor.AnimInstance->PlayAnim(false, 0.0f);
+				if (!PreviewActor.Actor->GetRootComponent())
+				{
+					PreviewActor.Actor->SetRootComponent(PreviewActor.Mesh.Get());
+				}
 
-		if (!PreviewActor.Actor->GetRootComponent())
-		{
-			PreviewActor.Actor->SetRootComponent(PreviewActor.Mesh.Get());
+				UE_LOG(
+					LogPoseSearchEditor, Log, TEXT("Spawned preview Actor: %s at Loc: %s Rot: %s"),
+					*GetNameSafe(PreviewActor.Actor.Get()),
+					*SpawnTransform.GetLocation().ToString(),
+					*SpawnTransform.Rotator().ToString());
+			}
 		}
-
-		UE_LOG(
-			LogPoseSearchEditor, Log, TEXT("Spawned preview Actor: %s at Loc: %s Rot: %s"),
-			*GetNameSafe(PreviewActor.Actor.Get()),
-			*SpawnTransform.GetLocation().ToString(),
-			*SpawnTransform.Rotator().ToString());
-
 		return PreviewActor;
 	}
 
@@ -262,33 +265,40 @@ namespace UE::PoseSearch
 
 	void FDatabaseViewModel::UpdatePreviewActors()
 	{
-		TSharedPtr<FDatabasePreviewScene> PreviewScene = PreviewScenePtr.Pin();
-		for (FDatabasePreviewActor& PreviewActor : GetPreviewActors())
+		if (const FPoseSearchIndex* SearchIndex = PoseSearchDatabase->GetSearchIndexSafe())
 		{
-			if (const UAnimationAsset* PreviewAsset = PreviewActor.AnimInstance->GetAnimationAsset())
+			TSharedPtr<FDatabasePreviewScene> PreviewScene = PreviewScenePtr.Pin();
+			for (FDatabasePreviewActor& PreviewActor : GetPreviewActors())
 			{
-				float CurrentTime = 0.0f;
-				FAnimationRuntime::AdvanceTime(false, PlayTime, CurrentTime, PreviewAsset->GetPlayLength());
-
-				if (const UAnimSequence* AnimSequence = Cast<UAnimSequence>(PreviewAsset))
+				if (PreviewActor.IndexAssetIndex < SearchIndex->Assets.Num())
 				{
-					FTransform RootMotion = FTransform::Identity;
-					RootMotion = AnimSequence->ExtractRootMotionFromRange(0.0f, CurrentTime);
-					if (PreviewActor.AnimInstance->GetMirrorDataTable())
+					if (const UAnimationAsset* PreviewAsset = PreviewActor.AnimInstance->GetAnimationAsset())
 					{
-						RootMotion = MirrorRootMotion(RootMotion, PreviewActor.AnimInstance->GetMirrorDataTable());
+						float CurrentTime = 0.0f;
+						FAnimationRuntime::AdvanceTime(false, PlayTime, CurrentTime, PreviewAsset->GetPlayLength());
+
+						if (const UAnimSequence* AnimSequence = Cast<UAnimSequence>(PreviewAsset))
+						{
+							FTransform RootMotion = FTransform::Identity;
+							RootMotion = AnimSequence->ExtractRootMotionFromRange(0.0f, CurrentTime);
+							if (PreviewActor.AnimInstance->GetMirrorDataTable())
+							{
+								RootMotion = MirrorRootMotion(RootMotion, PreviewActor.AnimInstance->GetMirrorDataTable());
+							}
+							PreviewActor.Actor->SetActorTransform(RootMotion);
+						}
+						else if (const UBlendSpace* BlendSpace = Cast<UBlendSpace>(PreviewAsset))
+						{
+							// @todo: add BS support
+							//RootMotion = BlendSpace->ExtractRootMotionFromRange(0.0f, CurrentTime);
+						}
+
+						PreviewActor.AnimInstance->SetPosition(CurrentTime);
+
+						const FPoseSearchIndexAsset& SearchIndexAsset = SearchIndex->Assets[PreviewActor.IndexAssetIndex];
+						PreviewActor.CurrentPoseIndex = PoseSearchDatabase->GetPoseIndexFromTime(CurrentTime, SearchIndexAsset);
 					}
-					PreviewActor.Actor->SetActorTransform(RootMotion);
 				}
-				else if (const UBlendSpace* BlendSpace = Cast<UBlendSpace>(PreviewAsset))
-				{
-					// @todo: add BS support
-					//RootMotion = BlendSpace->ExtractRootMotionFromRange(0.0f, CurrentTime);
-				}
-				
-				PreviewActor.AnimInstance->SetPosition(CurrentTime);
-
-				PreviewActor.CurrentPoseIndex = PoseSearchDatabase->GetPoseIndexFromTime(CurrentTime, &PreviewActor.IndexAsset);
 			}
 		}
 	}
@@ -398,18 +408,26 @@ namespace UE::PoseSearch
 
 	void FDatabaseViewModel::ProcessSelectedActor(AActor* Actor)
 	{
-		SelectedActorIndexAsset = nullptr;
-		
-		const FDatabasePreviewActor* SelectedPreviewActor = PreviewActors.FindByPredicate(
-			[Actor](const FDatabasePreviewActor& PreviewActor)
+		SelectedActorIndexAssetIndex = INDEX_NONE;
+		if (const FDatabasePreviewActor* SelectedPreviewActor = PreviewActors.FindByPredicate([Actor](const FDatabasePreviewActor& PreviewActor) { return PreviewActor.Actor == Actor; }))
 		{
-			return PreviewActor.Actor == Actor;
-		});
+			SelectedActorIndexAssetIndex = SelectedPreviewActor->IndexAssetIndex;
+		}
+	}
 
-		if (SelectedPreviewActor)
+	const FPoseSearchIndexAsset* FDatabaseViewModel::GetSelectedActorIndexAsset() const
+	{
+		if (PoseSearchDatabase)
 		{
-			SelectedActorIndexAsset = &SelectedPreviewActor->IndexAsset;
+			if (const FPoseSearchIndex* SearchIndex = PoseSearchDatabase->GetSearchIndexSafe())
+			{
+				if (SearchIndex->Assets.IsValidIndex(SelectedActorIndexAssetIndex))
+				{
+					return &SearchIndex->Assets[SelectedActorIndexAssetIndex];
+				}
+			}
 		}
+		return nullptr;
 	}
 
 	float FDatabaseViewModel::GetMaxPreviewPlayLength() const
diff --git a/Engine/Plugins/Experimental/Animation/PoseSearch/Source/Editor/Private/PoseSearchDebugger.cpp b/Engine/Plugins/Experimental/Animation/PoseSearch/Source/Editor/Private/PoseSearchDebugger.cpp
--- a/Engine/Plugins/Experimental/Animation/PoseSearch/Source/Editor/Private/PoseSearchDebugger.cpp
+++ b/Engine/Plugins/Experimental/Animation/PoseSearch/Source/Editor/Private/PoseSearchDebugger.cpp
@@ -530,7 +530,7 @@ namespace DebuggerDatabaseColumns
 		}
 
 		FText Label;
-		int32 BreakdownCostIndex = -1;
+		int32 BreakdownCostIndex = INDEX_NONE;
 	};
 
 #if WITH_EDITORONLY_DATA
@@ -812,7 +812,7 @@ public:
 	virtual void EndBreakDownSection(const FText& Label) override
 	{
 		check(CurrentLabeIndex == GetOrAddLabelIndex(Label));
-		CurrentLabeIndex = -1;
+		CurrentLabeIndex = INDEX_NONE;
 	}
 
 	virtual bool IsVerbose() const override
@@ -862,7 +862,7 @@ private:
 
 	TArray<FText> Labels;
 
-	int32 CurrentLabeIndex = -1;
+	int32 CurrentLabeIndex = INDEX_NONE;
 
 	bool bIsVerbose = true;
 };
@@ -894,12 +894,11 @@ void SDebuggerDatabaseView::Update(const FTraceMotionMatchingStateMessage& State
 			{
 				for (const FTraceMotionMatchingStatePoseEntry& PoseEntry : DbEntry.PoseEntries)
 				{
-					const FPoseSearchIndexAsset* SearchIndexAsset = SearchIndex->FindAssetForPose(PoseEntry.DbPoseIdx);
-					if (SearchIndexAsset)
+					if (const FPoseSearchIndexAsset* SearchIndexAsset = SearchIndex->GetAssetForPoseSafe(PoseEntry.DbPoseIdx))
 					{
 						TSharedRef<FDebuggerDatabaseRowData>& Row = UnfilteredDatabaseRows.Add_GetRef(MakeShared<FDebuggerDatabaseRowData>());
 
-						const float Time = SearchIndex->GetAssetTime(PoseEntry.DbPoseIdx, SearchIndexAsset);
+						const float Time = SearchIndex->GetAssetTime(PoseEntry.DbPoseIdx);
 
 						Row->PoseIdx = PoseEntry.DbPoseIdx;
 						Row->SourceDatabase = Database;
@@ -912,7 +911,7 @@ void SDebuggerDatabaseView::Update(const FTraceMotionMatchingStateMessage& State
 						Row->PoseCost = PoseEntry.Cost;
 
 						Row->CostVector.SetNum(Database->Schema->SchemaCardinality);
-						TConstArrayView<float> PoseValues = SearchIndex->GetPoseValues(PoseEntry.DbPoseIdx);
+						TConstArrayView<float> PoseValues = SearchIndex->GetPoseValuesSafe(PoseEntry.DbPoseIdx);
 
 						// in case we modify the schema while PIE is paused and displaying the Pose Search Editor, we could end up with a stale State with a DbEntry.QueryVector saved with the previous schema
 						// so the cardinality of DbEntry.QueryVector and PoseValues don't match. In that case we just use PoseValues as query to have all costs set to zero
@@ -921,7 +920,7 @@ void SDebuggerDatabaseView::Update(const FTraceMotionMatchingStateMessage& State
 
 						if (SearchIndexAsset->Type == ESearchIndexAssetType::Sequence)
 						{
-							const FPoseSearchDatabaseSequence& DbSequence = Database->GetSequenceSourceAsset(SearchIndexAsset);
+							const FPoseSearchDatabaseSequence& DbSequence = Database->GetSequenceSourceAsset(*SearchIndexAsset);
 							Row->AssetType = ESearchIndexAssetType::Sequence;
 							Row->AssetName = DbSequence.Sequence->GetName();
 							Row->AssetPath = DbSequence.Sequence->GetPathName();
@@ -932,7 +931,7 @@ void SDebuggerDatabaseView::Update(const FTraceMotionMatchingStateMessage& State
 						}
 						else if (SearchIndexAsset->Type == ESearchIndexAssetType::BlendSpace)
 						{
-							const FPoseSearchDatabaseBlendSpace& DbBlendSpace = Database->GetBlendSpaceSourceAsset(SearchIndexAsset);
+							const FPoseSearchDatabaseBlendSpace& DbBlendSpace = Database->GetBlendSpaceSourceAsset(*SearchIndexAsset);
 							Row->AssetType = ESearchIndexAssetType::BlendSpace;
 							Row->AssetName = DbBlendSpace.BlendSpace->GetName();
 							Row->AssetPath = DbBlendSpace.BlendSpace->GetPathName();
@@ -941,10 +940,6 @@ void SDebuggerDatabaseView::Update(const FTraceMotionMatchingStateMessage& State
 							Row->bLooping = DbBlendSpace.BlendSpace->bLoop;
 							Row->BlendParameters = SearchIndexAsset->BlendParameters;
 						}
-						else
-						{
-							checkNoEntry();
-						}
 					}
 				}
 			}
@@ -1572,9 +1567,9 @@ void SDebuggerDetailsView::UpdateReflection(const FTraceMotionMatchingStateMessa
 	Reflection->AssetPlayerAssetName = "None";
 	if (CurrentSearchIndex)
 	{
-		if (const FPoseSearchIndexAsset* IndexAsset = CurrentSearchIndex->FindAssetForPose(CurrentDbPoseIdx))
+		if (const FPoseSearchIndexAsset* IndexAsset = CurrentSearchIndex->GetAssetForPoseSafe(CurrentDbPoseIdx))
 		{
-			Reflection->AssetPlayerAssetName = CurrentDatabase->GetSourceAssetName(IndexAsset);
+			Reflection->AssetPlayerAssetName = CurrentDatabase->GetSourceAssetName(*IndexAsset);
 		}
 	}
 
@@ -1604,7 +1599,7 @@ void SDebuggerDetailsView::UpdateReflection(const FTraceMotionMatchingStateMessa
 	// Active pose
 	if (CurrentSearchIndex)
 	{
-		Reflection->ActivePoseVector = CurrentSearchIndex->GetPoseValues(CurrentDbPoseIdx);
+		Reflection->ActivePoseVector = CurrentSearchIndex->GetPoseValuesSafe(CurrentDbPoseIdx);
 	}
 
 	auto DebuggerView = ParentDebuggerViewPtr.Pin();
@@ -1619,7 +1614,7 @@ void SDebuggerDetailsView::UpdateReflection(const FTraceMotionMatchingStateMessa
 				const FPoseSearchIndex* SelectedSearchIndex = Selected->SourceDatabase->GetSearchIndex();
 				if (SelectedSearchIndex)
 				{
-					Reflection->SelectedPoseVector = SelectedSearchIndex->GetPoseValues(Selected->PoseIdx);
+					Reflection->SelectedPoseVector = SelectedSearchIndex->GetPoseValuesSafe(Selected->PoseIdx);
 				}
 			}
 
@@ -1638,7 +1633,7 @@ void SDebuggerView::Construct(const FArguments& InArgs, uint64 InAnimInstanceId)
 	check(OnViewClosed.IsBound());
 	
 	AnimInstanceId = InAnimInstanceId;
-	SelectedNodeId = -1;
+	SelectedNodeId = INDEX_NONE;
 
 	ChildSlot
 	[
@@ -2327,21 +2322,22 @@ void FDebuggerViewModel::ShowSelectedSkeleton(const UPoseSearchDatabase* Databas
 	}
 
 	const FPoseSearchIndex* SearchIndex = Database->GetSearchIndex();
-	const FPoseSearchIndexAsset* IndexAsset = SearchIndex ? SearchIndex->FindAssetForPose(DbPoseIdx) : nullptr;
-	if (!IndexAsset)
+	if (!SearchIndex)
 	{
 		return;
 	}
 
+	const FPoseSearchIndexAsset& IndexAsset = SearchIndex->GetAssetForPose(DbPoseIdx);
+	
 	Component->ResetToStart(); 
 	bSelecting = true;
 	
-	Skeletons[SelectedPose].Type = IndexAsset->Type;
+	Skeletons[SelectedPose].Type = IndexAsset.Type;
 	Skeletons[SelectedPose].Time = Time;
-	Skeletons[SelectedPose].bMirrored = IndexAsset->bMirrored;
+	Skeletons[SelectedPose].bMirrored = IndexAsset.bMirrored;
 	Skeletons[SelectedPose].SourceDatabase = Database;
-	Skeletons[SelectedPose].AssetIdx = IndexAsset->SourceAssetIdx;
-	Skeletons[SelectedPose].BlendParameters = IndexAsset->BlendParameters;
+	Skeletons[SelectedPose].AssetIdx = IndexAsset.SourceAssetIdx;
+	Skeletons[SelectedPose].BlendParameters = IndexAsset.BlendParameters;
 }
 
 void FDebuggerViewModel::ClearSelectedSkeleton()
@@ -2416,18 +2412,18 @@ void FDebuggerViewModel::OnUpdateNodeSelection(int32 InNodeId)
 	if (ActiveMotionMatchingState)
 	{
 		const UPoseSearchDatabase* CurrentDatabase = ActiveMotionMatchingState->GetCurrentDatabase();
-		const FPoseSearchIndex* CurrentSearchIndex = CurrentDatabase ? CurrentDatabase->GetSearchIndex() : nullptr;
-		int32 CurrentPoseIdx = ActiveMotionMatchingState->GetCurrentDatabasePoseIndex();
-
-		const FPoseSearchIndexAsset* IndexAsset = CurrentSearchIndex ? CurrentSearchIndex->FindAssetForPose(CurrentPoseIdx) : nullptr;
-
-		if (IndexAsset)
+		const FPoseSearchIndex* CurrentSearchIndex = CurrentDatabase ? CurrentDatabase->GetSearchIndexSafe() : nullptr;
+		if (CurrentSearchIndex)
 		{
-			Skeletons[Asset].Type = IndexAsset->Type;
-			Skeletons[Asset].bMirrored = IndexAsset->bMirrored;
-			Skeletons[Asset].SourceDatabase = CurrentDatabase;
-			Skeletons[Asset].AssetIdx = IndexAsset->SourceAssetIdx;
-			Skeletons[Asset].BlendParameters = IndexAsset->BlendParameters;
+			int32 CurrentPoseIdx = ActiveMotionMatchingState->GetCurrentDatabasePoseIndex();
+			if (const FPoseSearchIndexAsset* IndexAsset = CurrentSearchIndex->GetAssetForPoseSafe(CurrentPoseIdx))
+			{
+				Skeletons[Asset].Type = IndexAsset->Type;
+				Skeletons[Asset].bMirrored = IndexAsset->bMirrored;
+				Skeletons[Asset].SourceDatabase = CurrentDatabase;
+				Skeletons[Asset].AssetIdx = IndexAsset->SourceAssetIdx;
+				Skeletons[Asset].BlendParameters = IndexAsset->BlendParameters;
+			}
 		}
 	}
 
@@ -2443,7 +2439,9 @@ void FDebuggerViewModel::OnDraw(FSkeletonDrawParams& DrawParams)
 {
 	const UPoseSearchDatabase* CurrentDatabase = GetCurrentDatabase();
 	if (!CurrentDatabase)
+	{
 		return;
+	}
 
 	// Returns if it is to be drawn this frame
 	auto SetDrawSkeleton = [this](UPoseSearchMeshComponent* InComponent, bool bDraw)
@@ -2784,24 +2782,21 @@ void FDebuggerViewModel::PlaySelection(int32 PoseIdx, float Time)
 	{
 		return;
 	}
-	const FPoseSearchIndexAsset* IndexAsset = Database->GetSearchIndex()->FindAssetForPose(PoseIdx);
-	if (!IndexAsset)
-	{
-		return;
-	}
+	const FPoseSearchIndexAsset& IndexAsset = Database->GetSearchIndex()->GetAssetForPose(PoseIdx);
 
 	Component->ResetToStart();
 	
-	Skeletons[Asset].Type = IndexAsset->Type;
-	Skeletons[Asset].AssetIdx = IndexAsset->SourceAssetIdx;
+	Skeletons[Asset].Type = IndexAsset.Type;
+	Skeletons[Asset].AssetIdx = IndexAsset.SourceAssetIdx;
 	Skeletons[Asset].Time = Time;
-	Skeletons[Asset].bMirrored = IndexAsset->bMirrored;
-	Skeletons[Asset].BlendParameters = IndexAsset->BlendParameters;
+	Skeletons[Asset].bMirrored = IndexAsset.bMirrored;
+	Skeletons[Asset].BlendParameters = IndexAsset.BlendParameters;
 
 	AssetData.StartTime = Time;
 	AssetData.AccumulatedTime = 0.0f;
 	AssetData.bActive = true;
 }
+
 void FDebuggerViewModel::StopSelection()
 {
 	UPoseSearchMeshComponent* Component = Skeletons[Asset].Component.Get();
diff --git a/Engine/Plugins/Experimental/Animation/PoseSearch/Source/Editor/Public/PoseSearchDatabaseViewModel.h b/Engine/Plugins/Experimental/Animation/PoseSearch/Source/Editor/Public/PoseSearchDatabaseViewModel.h
--- a/Engine/Plugins/Experimental/Animation/PoseSearch/Source/Editor/Public/PoseSearchDatabaseViewModel.h
+++ b/Engine/Plugins/Experimental/Animation/PoseSearch/Source/Editor/Public/PoseSearchDatabaseViewModel.h
@@ -42,7 +42,7 @@ namespace UE::PoseSearch
 		TWeakObjectPtr<AActor> Actor = nullptr;
 		TWeakObjectPtr<UDebugSkelMeshComponent> Mesh = nullptr;
 		TWeakObjectPtr<UAnimPreviewInstance> AnimInstance = nullptr;
-		FPoseSearchIndexAsset IndexAsset; // keeping a copy since database index can be invalidated
+		int32 IndexAssetIndex = INDEX_NONE;
 		int32 CurrentPoseIndex = INDEX_NONE;
 
 		bool IsValid()
@@ -109,7 +109,7 @@ namespace UE::PoseSearch
 		void SetSelectedNodes(const TArrayView<TSharedPtr<FDatabaseAssetTreeNode>>& InSelectedNodes);
 		void ProcessSelectedActor(AActor* Actor);
 		
-		const FPoseSearchIndexAsset* GetSelectedActorIndexAsset() const { return SelectedActorIndexAsset; }
+		const FPoseSearchIndexAsset* GetSelectedActorIndexAsset() const;
 
 		float GetMaxPreviewPlayLength() const;
 		float GetPlayTime() const;
@@ -139,13 +139,13 @@ namespace UE::PoseSearch
 
 		TArray<TSharedPtr<FDatabaseAssetTreeNode>> SelectedNodes;
 
-		const FPoseSearchIndexAsset* SelectedActorIndexAsset = nullptr;
+		int32 SelectedActorIndexAssetIndex = INDEX_NONE;
 
 		UWorld* GetWorld() const;
 
 		UObject* GetPlaybackContext() const;
 
-		FDatabasePreviewActor SpawnPreviewActor(const FPoseSearchIndexAsset& IndexAsset);
+		FDatabasePreviewActor SpawnPreviewActor(int32 IndexAssetIndex);
 
 		void UpdatePreviewActors();
 
diff --git a/Engine/Plugins/Experimental/Animation/PoseSearch/Source/Runtime/Private/AnimNode_MotionMatching.cpp b/Engine/Plugins/Experimental/Animation/PoseSearch/Source/Runtime/Private/AnimNode_MotionMatching.cpp
--- a/Engine/Plugins/Experimental/Animation/PoseSearch/Source/Runtime/Private/AnimNode_MotionMatching.cpp
+++ b/Engine/Plugins/Experimental/Animation/PoseSearch/Source/Runtime/Private/AnimNode_MotionMatching.cpp
@@ -90,15 +90,17 @@ void FAnimNode_MotionMatching::UpdateAssetPlayer(const FAnimationUpdateContext&
 	);
 
 	// If a new pose is requested, blend into the new asset via BlendStackNode
-	const UPoseSearchDatabase* Database = MotionMatchingState.CurrentSearchResult.Database.Get();
-	const FPoseSearchIndexAsset* SearchIndexAsset = MotionMatchingState.GetCurrentSearchIndexAsset();
-
-	if (Database && Database->Schema && SearchIndexAsset && (MotionMatchingState.Flags & EMotionMatchingFlags::JumpedToPose) == EMotionMatchingFlags::JumpedToPose)
+	const FPoseSearchIndexAsset* SearchIndexAsset = MotionMatchingState.CurrentSearchResult.GetSearchIndexAsset();
+	if (SearchIndexAsset)
 	{
-		const FPoseSearchDatabaseAnimationAssetBase& AnimationAssetBase = Database->GetAnimationSourceAsset(SearchIndexAsset);
-		BlendStackNode.BlendTo(	SearchIndexAsset->Type, AnimationAssetBase.GetAnimationAsset(), MotionMatchingState.CurrentSearchResult.AssetTime,
-								AnimationAssetBase.IsLooping(), SearchIndexAsset->bMirrored, Database->Schema->MirrorDataTable.Get(),
-								Settings.MaxActiveBlends, Settings.BlendTime, Settings.BlendProfile, Settings.BlendOption, SearchIndexAsset->BlendParameters);
+		const UPoseSearchDatabase* Database = MotionMatchingState.CurrentSearchResult.Database.Get();
+		if (Database && Database->Schema && EnumHasAnyFlags(MotionMatchingState.Flags, EMotionMatchingFlags::JumpedToPose))
+		{
+			const FPoseSearchDatabaseAnimationAssetBase& AnimationAssetBase = Database->GetAnimationSourceAsset(*SearchIndexAsset);
+			BlendStackNode.BlendTo(SearchIndexAsset->Type, AnimationAssetBase.GetAnimationAsset(), MotionMatchingState.CurrentSearchResult.AssetTime,
+				AnimationAssetBase.IsLooping(), SearchIndexAsset->bMirrored, Database->Schema->MirrorDataTable.Get(),
+				Settings.MaxActiveBlends, Settings.BlendTime, Settings.BlendProfile, Settings.BlendOption, SearchIndexAsset->BlendParameters);
+		}
 	}
 
 	Source.Update(Context);
diff --git a/Engine/Plugins/Experimental/Animation/PoseSearch/Source/Runtime/Private/PoseSearch.cpp b/Engine/Plugins/Experimental/Animation/PoseSearch/Source/Runtime/Private/PoseSearch.cpp
--- a/Engine/Plugins/Experimental/Animation/PoseSearch/Source/Runtime/Private/PoseSearch.cpp
+++ b/Engine/Plugins/Experimental/Animation/PoseSearch/Source/Runtime/Private/PoseSearch.cpp
@@ -187,13 +187,15 @@ FORCEINLINE auto LowerBound(IteratorType First, IteratorType Last, const ValueTy
 static TArrayView<size_t> PopulateNonSelectableIdx(TArrayView<size_t> NonSelectableIdxBuffer, FSearchContext& SearchContext, const UPoseSearchDatabase* Database, TArrayView<const float> QueryValues)
 {
 	check(Database);
+	const FPoseSearchIndex* SearchIndex = Database->GetSearchIndex();
+	check(SearchIndex);
 
 	int32 NonSelectableIdxUsedSize = 0;
-	if (SearchContext.IsCurrentResultFromDatabase(Database) && SearchContext.CurrentResult.IsValid() && SearchContext.PoseJumpThresholdTime > 0.f)
+	const FPoseSearchIndexAsset* CurrentIndexAsset = SearchContext.CurrentResult.GetSearchIndexAsset();
+	if (CurrentIndexAsset && SearchContext.IsCurrentResultFromDatabase(Database) && SearchContext.PoseJumpThresholdTime > 0.f)
 	{
-		const FPoseSearchIndexAsset* CurrentIndexAsset = SearchContext.CurrentResult.SearchIndexAsset;
 		const int32 PoseJumpIndexThreshold = FMath::FloorToInt(SearchContext.PoseJumpThresholdTime / Database->Schema->GetSamplingInterval());
-		const bool IsLooping = Database->IsSourceAssetLooping(CurrentIndexAsset);
+		const bool IsLooping = Database->IsSourceAssetLooping(*CurrentIndexAsset);
 
 		for (int32 i = -PoseJumpIndexThreshold; i <= -1; ++i)
 		{
@@ -220,7 +222,7 @@ static TArrayView<size_t> PopulateNonSelectableIdx(TArrayView<size_t> NonSelecta
 					NonSelectableIdxBuffer[NonSelectableIdxUsedSize++] = PoseIdx;
 
 #if UE_POSE_SEARCH_TRACE_ENABLED
-					const FPoseSearchCost PoseCost = Database->ComparePoses(SearchContext, PoseIdx, EPoseComparisonFlags::None, QueryValues);
+					const FPoseSearchCost PoseCost = SearchIndex->ComparePoses(PoseIdx, SearchContext.QueryMirrorRequest, EPoseComparisonFlags::None, QueryValues);
 					SearchContext.BestCandidates.Add(PoseCost, PoseIdx, Database, EPoseCandidateFlags::DiscardedBy_PoseJumpThresholdTime);
 #endif
 				}
@@ -263,7 +265,7 @@ static TArrayView<size_t> PopulateNonSelectableIdx(TArrayView<size_t> NonSelecta
 					NonSelectableIdxBuffer[NonSelectableIdxUsedSize++] = PoseIdx;
 
 #if UE_POSE_SEARCH_TRACE_ENABLED
-					const FPoseSearchCost PoseCost = Database->ComparePoses(SearchContext, PoseIdx, EPoseComparisonFlags::None, QueryValues);
+					const FPoseSearchCost PoseCost = SearchIndex->ComparePoses(PoseIdx, SearchContext.QueryMirrorRequest, EPoseComparisonFlags::None, QueryValues);
 					SearchContext.BestCandidates.Add(PoseCost, PoseIdx, Database, EPoseCandidateFlags::DiscardedBy_PoseJumpThresholdTime);
 #endif
 				}
@@ -296,13 +298,11 @@ static TArrayView<size_t> PopulateNonSelectableIdx(TArrayView<size_t> NonSelecta
 
 #if UE_POSE_SEARCH_TRACE_ENABLED
 					check(HistoricalPoseIndex.PoseIndex >= 0);
-					const FPoseSearchIndex* SearchIndex = Database->GetSearchIndex();
-					check(SearchIndex);
 					
 					// if we're editing the database and removing assets it's possible that the PoseIndicesHistory contains invalid pose indexes
 					if (HistoricalPoseIndex.PoseIndex < SearchIndex->NumPoses)
 					{
-						const FPoseSearchCost PoseCost = Database->ComparePoses(SearchContext, HistoricalPoseIndex.PoseIndex, EPoseComparisonFlags::None, QueryValues);
+						const FPoseSearchCost PoseCost = SearchIndex->ComparePoses(HistoricalPoseIndex.PoseIndex, SearchContext.QueryMirrorRequest, EPoseComparisonFlags::None, QueryValues);
 						SearchContext.BestCandidates.Add(PoseCost, HistoricalPoseIndex.PoseIndex, Database, EPoseCandidateFlags::DiscardedBy_PoseReselectHistory);
 					}
 #endif
@@ -491,48 +491,44 @@ bool UPoseSearchSchema::BuildQuery(UE::PoseSearch::FSearchContext& SearchContext
 //////////////////////////////////////////////////////////////////////////
 // FPoseSearchIndexAsset
 // 
-const FPoseSearchIndexAsset* FPoseSearchIndex::FindAssetForPose(int32 PoseIdx) const
+const FPoseSearchIndexAsset& FPoseSearchIndex::GetAssetForPose(int32 PoseIdx) const
 {
-	auto Predicate = [PoseIdx](const FPoseSearchIndexAsset& Asset)
-	{
-		return Asset.IsPoseInRange(PoseIdx);
-	};
-	return Assets.FindByPredicate(Predicate);
+	const int32 AssetIndex = PoseMetadata[PoseIdx].AssetIndex;
+	return Assets[AssetIndex];
 }
 
-float FPoseSearchIndex::GetAssetTime(int32 PoseIdx, const FPoseSearchIndexAsset* Asset) const
+const FPoseSearchIndexAsset* FPoseSearchIndex::GetAssetForPoseSafe(int32 PoseIdx) const
 {
-	if (!Asset)
+	if (PoseMetadata.IsValidIndex(PoseIdx))
 	{
-		Asset = FindAssetForPose(PoseIdx);
-		if (!Asset)
+		const int32 AssetIndex = PoseMetadata[PoseIdx].AssetIndex;
+		if (Assets.IsValidIndex(AssetIndex))
 		{
-			UE_LOG(LogPoseSearch, Error, TEXT("Couldn't find asset for pose %i in database"), PoseIdx);
-			return -1.0f;
+			return &Assets[AssetIndex];
 		}
 	}
+	return nullptr;
+}
 
-	if (!Asset->IsPoseInRange(PoseIdx))
-	{
-		UE_LOG(LogPoseSearch, Error, TEXT("Pose %i out of range in database"), PoseIdx);
-		return -1.0f;
-	}
+float FPoseSearchIndex::GetAssetTime(int32 PoseIdx) const
+{
+	const FPoseSearchIndexAsset& Asset = GetAssetForPose(PoseIdx);
 
-	if (Asset->Type == ESearchIndexAssetType::Sequence)
+	if (Asset.Type == ESearchIndexAssetType::Sequence)
 	{
-		const FFloatInterval SamplingRange = Asset->SamplingInterval;
+		const FFloatInterval SamplingRange = Asset.SamplingInterval;
 
-		float AssetTime = FMath::Min(SamplingRange.Min + Schema->GetSamplingInterval() * (PoseIdx - Asset->FirstPoseIdx), SamplingRange.Max);
+		float AssetTime = FMath::Min(SamplingRange.Min + Schema->GetSamplingInterval() * (PoseIdx - Asset.FirstPoseIdx), SamplingRange.Max);
 		return AssetTime;
 	}
 	
-	if (Asset->Type == ESearchIndexAssetType::BlendSpace)
+	if (Asset.Type == ESearchIndexAssetType::BlendSpace)
 	{
-		const FFloatInterval SamplingRange = Asset->SamplingInterval;
+		const FFloatInterval SamplingRange = Asset.SamplingInterval;
 
 		// For BlendSpaces the AssetTime is in the range [0, 1] while the Sampling Range
 		// is in real time (seconds)
-		float AssetTime = FMath::Min(SamplingRange.Min + Schema->GetSamplingInterval() * (PoseIdx - Asset->FirstPoseIdx), SamplingRange.Max) / (Asset->NumPoses * Schema->GetSamplingInterval());
+		float AssetTime = FMath::Min(SamplingRange.Min + Schema->GetSamplingInterval() * (PoseIdx - Asset.FirstPoseIdx), SamplingRange.Max) / (Asset.NumPoses * Schema->GetSamplingInterval());
 		return AssetTime;
 	}
 	
@@ -565,6 +561,16 @@ TConstArrayView<float> FPoseSearchIndex::GetPoseValues(int32 PoseIdx) const
 	return MakeArrayView(&Values[ValueOffset], Schema->SchemaCardinality);
 }
 
+TConstArrayView<float> FPoseSearchIndex::GetPoseValuesSafe(int32 PoseIdx) const
+{
+	if (PoseIdx >= 0 && PoseIdx < NumPoses && Schema && Schema->SchemaCardinality >= 0)
+	{
+		const int32 ValueOffset = PoseIdx * Schema->SchemaCardinality;
+		return MakeArrayView(&Values[ValueOffset], Schema->SchemaCardinality);
+	}
+	return TConstArrayView<float>();
+}
+
 void FPoseSearchIndex::Reset()
 {
 	NumPoses = 0;
@@ -573,35 +579,35 @@ void FPoseSearchIndex::Reset()
 	Schema = nullptr;
 }
 
-float FPoseSearchIndex::ComputeMirrorMismatchAddend(int32 PoseIdx, UE::PoseSearch::FSearchContext& SearchContext) const
+FPoseSearchCost FPoseSearchIndex::ComparePoses(int32 PoseIdx, EPoseSearchBooleanRequest QueryMirrorRequest, UE::PoseSearch::EPoseComparisonFlags PoseComparisonFlags, TConstArrayView<float> QueryValues) const
 {
-	if (SearchContext.QueryMirrorRequest != EPoseSearchBooleanRequest::Indifferent)
+	// base dissimilarity cost representing how the assiciated PoseIdx differ, in a weighted way, from the query pose (QueryValues)
+	const float DissimilarityCost = UE::PoseSearch::CompareFeatureVectors(GetPoseValues(PoseIdx), QueryValues, WeightsSqrt);
+	
+	// cost addend associated to a mismatch in mirror state between query and analyzed PoseIdx
+	float MirrorMismatchAddend = 0.f;
+	if (QueryMirrorRequest != EPoseSearchBooleanRequest::Indifferent)
 	{
-		const FPoseSearchIndexAsset* IndexAsset = FindAssetForPose(PoseIdx);
+		const FPoseSearchIndexAsset& IndexAsset = GetAssetForPose(PoseIdx);
 		const bool bMirroringMismatch =
-			(IndexAsset->bMirrored && SearchContext.QueryMirrorRequest == EPoseSearchBooleanRequest::FalseValue) ||
-			(!IndexAsset->bMirrored && SearchContext.QueryMirrorRequest == EPoseSearchBooleanRequest::TrueValue);
+			(IndexAsset.bMirrored && QueryMirrorRequest == EPoseSearchBooleanRequest::FalseValue) ||
+			(!IndexAsset.bMirrored && QueryMirrorRequest == EPoseSearchBooleanRequest::TrueValue);
 		if (bMirroringMismatch)
 		{
 			check(Schema);
-			return Schema->MirrorMismatchCostBias;
+			MirrorMismatchAddend = Schema->MirrorMismatchCostBias;
 		}
 	}
-	return 0.f;
-}
 
-float FPoseSearchIndex::ComputeNotifyAddend(int32 PoseIdx) const
-{
-	return PoseMetadata[PoseIdx].CostAddend;
-}
+	const FPoseSearchPoseMetadata& PoseIdxMetadata = PoseMetadata[PoseIdx];
+	
+	// cost addend associated to Schema->BaseCostBias or overriden by UAnimNotifyState_PoseSearchModifyCost
+	const float NotifyAddend = PoseIdxMetadata.CostAddend;
 
-float FPoseSearchIndex::ComputeContinuingPoseCostAddend(int32 PoseIdx, UE::PoseSearch::EPoseComparisonFlags PoseComparisonFlags) const
-{
-	if (EnumHasAnyFlags(PoseComparisonFlags, UE::PoseSearch::EPoseComparisonFlags::ContinuingPose))
-	{
-		return PoseMetadata[PoseIdx].ContinuingPoseCostAddend;
-	}
-	return 0.f;
+	// cost addend associated to Schema->ContinuingPoseCostBias or overriden by UAnimNotifyState_PoseSearchOverrideContinuingPoseCostBias
+	const float ContinuingPoseCostAddend = EnumHasAnyFlags(PoseComparisonFlags, UE::PoseSearch::EPoseComparisonFlags::ContinuingPose) ? PoseIdxMetadata.ContinuingPoseCostAddend : 0.f;
+	
+	return FPoseSearchCost(DissimilarityCost, NotifyAddend, MirrorMismatchAddend, ContinuingPoseCostAddend);
 }
 
 //////////////////////////////////////////////////////////////////////////
@@ -657,8 +663,6 @@ UE::PoseSearch::FSearchResult UPoseSearchSequenceMetaData::Search(UE::PoseSearch
 		return Result;
 	}
 
-	FPoseSearchCost BestPoseCost;
-	int32 BestPoseIdx = INDEX_NONE;
 	for (const FPoseSearchIndexAsset& Asset : SearchIndex.Assets)
 	{
 		const int32 EndIndex = Asset.FirstPoseIdx + Asset.NumPoses;
@@ -671,20 +675,19 @@ UE::PoseSearch::FSearchResult UPoseSearchSequenceMetaData::Search(UE::PoseSearch
 				continue;
 			}
 
-			const FPoseSearchCost PoseCost = ComparePoses(PoseIdx, EPoseComparisonFlags::ContinuingPose, QueryValues);
-			if (PoseCost < BestPoseCost)
+			const FPoseSearchCost PoseCost = SearchIndex.ComparePoses(PoseIdx, EPoseSearchBooleanRequest::Indifferent, EPoseComparisonFlags::ContinuingPose, QueryValues);
+			if (PoseCost < Result.PoseCost)
 			{
-				BestPoseCost = PoseCost;
-				BestPoseIdx = PoseIdx;
+				Result.PoseCost = PoseCost;
+				Result.PoseIdx = PoseIdx;
 			}
 		}
 	}
 
-	Result.PoseCost = BestPoseCost;
-	Result.PoseIdx = BestPoseIdx;
-	Result.SearchIndexAsset = SearchIndex.FindAssetForPose(BestPoseIdx);
-	Result.AssetTime = SearchIndex.GetAssetTime(BestPoseIdx, Result.SearchIndexAsset);
-	Result.Database = nullptr;
+	if (Result.PoseIdx != INDEX_NONE)
+	{
+		Result.AssetTime = SearchIndex.GetAssetTime(Result.PoseIdx);
+	}
 
 	DrawFeatureVector(SearchContext.DebugDrawParams, Result.PoseIdx);
 	
@@ -694,23 +697,6 @@ UE::PoseSearch::FSearchResult UPoseSearchSequenceMetaData::Search(UE::PoseSearch
 	return Result;
 }
 
-FPoseSearchCost UPoseSearchSequenceMetaData::ComparePoses(
-	int32 PoseIdx,
-	UE::PoseSearch::EPoseComparisonFlags PoseComparisonFlags,
-	TConstArrayView<float> QueryValues) const
-{
-	using namespace UE::PoseSearch;
-
-	const float DissimilarityCost = CompareFeatureVectors(SearchIndex.GetPoseValues(PoseIdx), QueryValues);
-
-	// @todo: shouldn't we include MirrorMismatchAddend as well?
-	const float MirrorMismatchAddend = 0.f; // SearchIndex.ComputeMirrorMismatchAddend(PoseIdx, SearchContext);
-	const float NotifyAddend = SearchIndex.ComputeNotifyAddend(PoseIdx);
-	const float ContinuingPoseCostAddend = SearchIndex.ComputeContinuingPoseCostAddend(PoseIdx, PoseComparisonFlags);
-
-	return FPoseSearchCost(DissimilarityCost, NotifyAddend, MirrorMismatchAddend, ContinuingPoseCostAddend);
-}
-
 //////////////////////////////////////////////////////////////////////////
 // UPoseSearchDatabase
 
@@ -765,17 +751,11 @@ const FPoseSearchIndex* UPoseSearchDatabase::GetSearchIndexSafe() const
 	return SearchIndex;
 }
 
-float UPoseSearchDatabase::GetAssetTime(int32 PoseIdx, const FPoseSearchIndexAsset* Asset) const
-{
-	const float AssetTime = GetSearchIndex()->GetAssetTime(PoseIdx, Asset);
-	return AssetTime;
-}
-
-int32 UPoseSearchDatabase::GetPoseIndexFromTime(float Time, const FPoseSearchIndexAsset* SearchIndexAsset) const
+int32 UPoseSearchDatabase::GetPoseIndexFromTime(float Time, const FPoseSearchIndexAsset& SearchIndexAsset) const
 {
 	const bool bIsLooping = IsSourceAssetLooping(SearchIndexAsset);
-	const FFloatInterval& Range = SearchIndexAsset->SamplingInterval;
-	const bool bHasPoseIndex = SearchIndexAsset->FirstPoseIdx != INDEX_NONE && SearchIndexAsset->NumPoses > 0 && (bIsLooping || Range.Contains(Time));
+	const FFloatInterval& Range = SearchIndexAsset.SamplingInterval;
+	const bool bHasPoseIndex = SearchIndexAsset.FirstPoseIdx != INDEX_NONE && SearchIndexAsset.NumPoses > 0 && (bIsLooping || Range.Contains(Time));
 	if (bHasPoseIndex)
 	{
 		int32 PoseOffset = FMath::RoundToInt(Schema->SampleRate * (Time - Range.Min));
@@ -784,33 +764,33 @@ int32 UPoseSearchDatabase::GetPoseIndexFromTime(float Time, const FPoseSearchInd
 		{
 			if (bIsLooping)
 			{
-				PoseOffset = (PoseOffset % SearchIndexAsset->NumPoses) + SearchIndexAsset->NumPoses;
+				PoseOffset = (PoseOffset % SearchIndexAsset.NumPoses) + SearchIndexAsset.NumPoses;
 			}
 			else
 			{
 				PoseOffset = 0;
 			}
 		}
-		else if (PoseOffset >= SearchIndexAsset->NumPoses)
+		else if (PoseOffset >= SearchIndexAsset.NumPoses)
 		{
 			if (bIsLooping)
 			{
-				PoseOffset = PoseOffset % SearchIndexAsset->NumPoses;
+				PoseOffset = PoseOffset % SearchIndexAsset.NumPoses;
 			}
 			else
 			{
-				PoseOffset = SearchIndexAsset->NumPoses - 1;
+				PoseOffset = SearchIndexAsset.NumPoses - 1;
 			}
 		}
 
-		int32 PoseIdx = SearchIndexAsset->FirstPoseIdx + PoseOffset;
+		int32 PoseIdx = SearchIndexAsset.FirstPoseIdx + PoseOffset;
 		return PoseIdx;
 	}
 
 	return INDEX_NONE;
 }
 
-bool UPoseSearchDatabase::GetPoseIndicesAndLerpValueFromTime(float Time, const FPoseSearchIndexAsset* SearchIndexAsset, int32& PrevPoseIdx, int32& PoseIdx, int32& NextPoseIdx, float& LerpValue) const
+bool UPoseSearchDatabase::GetPoseIndicesAndLerpValueFromTime(float Time, const FPoseSearchIndexAsset& SearchIndexAsset, int32& PrevPoseIdx, int32& PoseIdx, int32& NextPoseIdx, float& LerpValue) const
 {
 	PoseIdx = GetPoseIndexFromTime(Time, SearchIndexAsset);
 	if (PoseIdx == INDEX_NONE)
@@ -821,7 +801,7 @@ bool UPoseSearchDatabase::GetPoseIndicesAndLerpValueFromTime(float Time, const F
 		return false;
 	}
 
-	const FFloatInterval& Range = SearchIndexAsset->SamplingInterval;
+	const FFloatInterval& Range = SearchIndexAsset.SamplingInterval;
 	const float FloatPoseOffset = Schema->SampleRate * (Time - Range.Min);
 	const int32 PoseOffset = FMath::RoundToInt(FloatPoseOffset);
 	LerpValue = FloatPoseOffset - float(PoseOffset);
@@ -846,42 +826,42 @@ bool UPoseSearchDatabase::GetPoseIndicesAndLerpValueFromTime(float Time, const F
 	return true;
 }
 
-const FPoseSearchDatabaseAnimationAssetBase& UPoseSearchDatabase::GetAnimationSourceAsset(const FPoseSearchIndexAsset* SearchIndexAsset) const
+const FPoseSearchDatabaseAnimationAssetBase& UPoseSearchDatabase::GetAnimationSourceAsset(const FPoseSearchIndexAsset& SearchIndexAsset) const
 {
-	if (SearchIndexAsset->Type == ESearchIndexAssetType::Sequence)
+	if (SearchIndexAsset.Type == ESearchIndexAssetType::Sequence)
 	{
-		return Sequences[SearchIndexAsset->SourceAssetIdx];
+		return Sequences[SearchIndexAsset.SourceAssetIdx];
 	}
 
-	if (SearchIndexAsset->Type == ESearchIndexAssetType::BlendSpace)
+	if (SearchIndexAsset.Type == ESearchIndexAssetType::BlendSpace)
 	{
-		return BlendSpaces[SearchIndexAsset->SourceAssetIdx];
+		return BlendSpaces[SearchIndexAsset.SourceAssetIdx];
 	}
 
 	checkNoEntry();
-	return Sequences[SearchIndexAsset->SourceAssetIdx];
+	return Sequences[SearchIndexAsset.SourceAssetIdx];
 }
 
-const FPoseSearchDatabaseSequence& UPoseSearchDatabase::GetSequenceSourceAsset(const FPoseSearchIndexAsset* SearchIndexAsset) const
+const FPoseSearchDatabaseSequence& UPoseSearchDatabase::GetSequenceSourceAsset(const FPoseSearchIndexAsset& SearchIndexAsset) const
 {
-	check(SearchIndexAsset->Type == ESearchIndexAssetType::Sequence);
-	return Sequences[SearchIndexAsset->SourceAssetIdx];
+	check(SearchIndexAsset.Type == ESearchIndexAssetType::Sequence);
+	return Sequences[SearchIndexAsset.SourceAssetIdx];
 }
 
-const FPoseSearchDatabaseBlendSpace& UPoseSearchDatabase::GetBlendSpaceSourceAsset(const FPoseSearchIndexAsset* SearchIndexAsset) const
+const FPoseSearchDatabaseBlendSpace& UPoseSearchDatabase::GetBlendSpaceSourceAsset(const FPoseSearchIndexAsset& SearchIndexAsset) const
 {
-	check(SearchIndexAsset->Type == ESearchIndexAssetType::BlendSpace);
-	return BlendSpaces[SearchIndexAsset->SourceAssetIdx];
+	check(SearchIndexAsset.Type == ESearchIndexAssetType::BlendSpace);
+	return BlendSpaces[SearchIndexAsset.SourceAssetIdx];
 }
 
-const bool UPoseSearchDatabase::IsSourceAssetLooping(const FPoseSearchIndexAsset* SearchIndexAsset) const
+const bool UPoseSearchDatabase::IsSourceAssetLooping(const FPoseSearchIndexAsset& SearchIndexAsset) const
 {
-	if (SearchIndexAsset->Type == ESearchIndexAssetType::Sequence)
+	if (SearchIndexAsset.Type == ESearchIndexAssetType::Sequence)
 	{
 		return GetSequenceSourceAsset(SearchIndexAsset).Sequence->bLoop;
 	}
 
-	if (SearchIndexAsset->Type == ESearchIndexAssetType::BlendSpace)
+	if (SearchIndexAsset.Type == ESearchIndexAssetType::BlendSpace)
 	{
 		return GetBlendSpaceSourceAsset(SearchIndexAsset).BlendSpace->bLoop;
 	}
@@ -890,14 +870,14 @@ const bool UPoseSearchDatabase::IsSourceAssetLooping(const FPoseSearchIndexAsset
 	return false;
 }
 
-const FString UPoseSearchDatabase::GetSourceAssetName(const FPoseSearchIndexAsset* SearchIndexAsset) const
+const FString UPoseSearchDatabase::GetSourceAssetName(const FPoseSearchIndexAsset& SearchIndexAsset) const
 {
-	if (SearchIndexAsset->Type == ESearchIndexAssetType::Sequence)
+	if (SearchIndexAsset.Type == ESearchIndexAssetType::Sequence)
 	{
 		return GetSequenceSourceAsset(SearchIndexAsset).Sequence->GetName();
 	}
 	
-	if (SearchIndexAsset->Type == ESearchIndexAssetType::BlendSpace)
+	if (SearchIndexAsset.Type == ESearchIndexAssetType::BlendSpace)
 	{
 		return GetBlendSpaceSourceAsset(SearchIndexAsset).BlendSpace->GetName();
 	}
@@ -1607,19 +1587,6 @@ bool UPoseSearchDatabase::IsCachedCookedPlatformDataLoaded(const ITargetPlatform
 
 #endif // WITH_EDITOR
 
-FPoseSearchCost UPoseSearchDatabase::ComparePoses(UE::PoseSearch::FSearchContext& SearchContext, int32 PoseIdx, UE::PoseSearch::EPoseComparisonFlags PoseComparisonFlags, TConstArrayView<float> QueryValues) const
-{
-	const FPoseSearchIndex* SearchIndex = GetSearchIndex();
-	check(SearchIndex);
-
-	const float DissimilarityCost = UE::PoseSearch::CompareFeatureVectors(SearchIndex->GetPoseValues(PoseIdx), QueryValues, SearchIndex->WeightsSqrt);
-	const float MirrorMismatchAddend = SearchIndex->ComputeMirrorMismatchAddend(PoseIdx, SearchContext);
-	const float NotifyAddend = SearchIndex->ComputeNotifyAddend(PoseIdx);
-	const float ContinuingPoseCostAddend = SearchIndex->ComputeContinuingPoseCostAddend(PoseIdx, PoseComparisonFlags);
-
-	return FPoseSearchCost(DissimilarityCost, NotifyAddend, MirrorMismatchAddend, ContinuingPoseCostAddend);
-}
-
 UE::PoseSearch::FSearchResult UPoseSearchDatabase::Search(UE::PoseSearch::FSearchContext& SearchContext) const
 {
 	using namespace UE::PoseSearch;
@@ -1680,7 +1647,6 @@ UE::PoseSearch::FSearchResult UPoseSearchDatabase::SearchPCAKDTree(UE::PoseSearc
 	const uint32 ClampedNumberOfPrincipalComponents = GetNumberOfPrincipalComponents();
 	const uint32 ClampedKDTreeQueryNumNeighbors = FMath::Clamp<uint32>(KDTreeQueryNumNeighbors, 1, SearchIndex->NumPoses);
 
-
 	//stack allocated temporaries
 	TArrayView<size_t> ResultIndexes((size_t*)FMemory_Alloca((ClampedKDTreeQueryNumNeighbors + 1) * sizeof(size_t)), ClampedKDTreeQueryNumNeighbors + 1);
 	TArrayView<float> ResultDistanceSqr((float*)FMemory_Alloca((ClampedKDTreeQueryNumNeighbors + 1) * sizeof(float)), ClampedKDTreeQueryNumNeighbors + 1);
@@ -1733,10 +1699,10 @@ UE::PoseSearch::FSearchResult UPoseSearchDatabase::SearchPCAKDTree(UE::PoseSearc
 	const bool IsCurrentResultFromThisDatabase = SearchContext.IsCurrentResultFromDatabase(this);
 
 	// evaluating the continuing pose only if it hasn't already being evaluated and the related animation can advance
-	if (IsCurrentResultFromThisDatabase && SearchContext.bCanAdvance && !Result.ContinuingPoseCost.IsValid())
+	if (!SearchContext.bForceInterrupt && IsCurrentResultFromThisDatabase && SearchContext.bCanAdvance && !Result.ContinuingPoseCost.IsValid())
 	{
 		Result.PoseIdx = SearchContext.CurrentResult.PoseIdx;
-		Result.PoseCost = ComparePoses(SearchContext, Result.PoseIdx, EPoseComparisonFlags::ContinuingPose, QueryValues);
+		Result.PoseCost = SearchIndex->ComparePoses(Result.PoseIdx, SearchContext.QueryMirrorRequest, EPoseComparisonFlags::ContinuingPose, QueryValues);
 		Result.ContinuingPoseCost = Result.PoseCost;
 
 		if (GetSkipSearchIfPossible())
@@ -1775,13 +1741,13 @@ UE::PoseSearch::FSearchResult UPoseSearchDatabase::SearchPCAKDTree(UE::PoseSearc
 			if (EnumHasAnyFlags(Metadata.Flags, EPoseSearchPoseFlags::BlockTransition))
 			{
 #if UE_POSE_SEARCH_TRACE_ENABLED
-				const FPoseSearchCost PoseCost = ComparePoses(SearchContext, PoseIdx, EPoseComparisonFlags::None, QueryValues);
+				const FPoseSearchCost PoseCost = SearchIndex->ComparePoses(PoseIdx, SearchContext.QueryMirrorRequest, EPoseComparisonFlags::None, QueryValues);
 				SearchContext.BestCandidates.Add(PoseCost, PoseIdx, this, EPoseCandidateFlags::DiscardedBy_BlockTransition);
 #endif
 				continue;
 			}
 
-			const FPoseSearchCost PoseCost = ComparePoses(SearchContext, PoseIdx, EPoseComparisonFlags::None, QueryValues);
+			const FPoseSearchCost PoseCost = SearchIndex->ComparePoses(PoseIdx, SearchContext.QueryMirrorRequest, EPoseComparisonFlags::None, QueryValues);
 			if (PoseCost < Result.PoseCost)
 			{
 				Result.PoseCost = PoseCost;
@@ -1802,8 +1768,7 @@ UE::PoseSearch::FSearchResult UPoseSearchDatabase::SearchPCAKDTree(UE::PoseSearc
 	// finalizing Result properties
 	if (Result.PoseIdx != INDEX_NONE)
 	{
-		Result.SearchIndexAsset = SearchIndex->FindAssetForPose(Result.PoseIdx);
-		Result.AssetTime = SearchIndex->GetAssetTime(Result.PoseIdx, Result.SearchIndexAsset);
+		Result.AssetTime = SearchIndex->GetAssetTime(Result.PoseIdx);
 		Result.Database = this;
 #if WITH_EDITOR
 		Result.SearchIndexHash = GetSearchIndexHash();
@@ -1829,13 +1794,13 @@ UE::PoseSearch::FSearchResult UPoseSearchDatabase::SearchBruteForce(UE::PoseSear
 	TArrayView<const float> QueryValues = Result.ComposedQuery.GetValues();
 
 	const bool IsCurrentResultFromThisDatabase = SearchContext.IsCurrentResultFromDatabase(this);
-	if (IsCurrentResultFromThisDatabase)
+	if (!SearchContext.bForceInterrupt && IsCurrentResultFromThisDatabase)
 	{
 		// evaluating the continuing pose only if it hasn't already being evaluated and the related animation can advance
 		if (SearchContext.bCanAdvance && !Result.ContinuingPoseCost.IsValid())
 		{
 			Result.PoseIdx = SearchContext.CurrentResult.PoseIdx;
-			Result.PoseCost = ComparePoses(SearchContext, Result.PoseIdx, EPoseComparisonFlags::ContinuingPose, QueryValues);
+			Result.PoseCost = SearchIndex->ComparePoses(Result.PoseIdx, SearchContext.QueryMirrorRequest, EPoseComparisonFlags::ContinuingPose, QueryValues);
 			Result.ContinuingPoseCost = Result.PoseCost;
 
 			if (GetSkipSearchIfPossible())
@@ -1855,42 +1820,37 @@ UE::PoseSearch::FSearchResult UPoseSearchDatabase::SearchBruteForce(UE::PoseSear
 	// there's no point in performing the search if CurrentBestTotalCost is already better than that
 	if (SearchContext.GetCurrentBestTotalCost() > SearchIndex->MinCostAddend)
 	{
-		for (const FPoseSearchIndexAsset& Asset : SearchIndex->Assets)
+		for (int32 PoseIdx = 0; PoseIdx < SearchIndex->NumPoses; ++PoseIdx)
 		{
-			const bool CheckForNonSelectableIdx = IsCurrentResultFromThisDatabase && (&Asset == SearchContext.CurrentResult.SearchIndexAsset);
-			const int32 EndIndex = Asset.FirstPoseIdx + Asset.NumPoses;
-			for (int32 PoseIdx = Asset.FirstPoseIdx; PoseIdx < EndIndex; ++PoseIdx)
-			{
-				const FPoseSearchPoseMetadata& Metadata = SearchIndex->PoseMetadata[PoseIdx];
+			const FPoseSearchPoseMetadata& Metadata = SearchIndex->PoseMetadata[PoseIdx];
 
-				if (EnumHasAnyFlags(Metadata.Flags, EPoseSearchPoseFlags::BlockTransition))
-				{
+			if (EnumHasAnyFlags(Metadata.Flags, EPoseSearchPoseFlags::BlockTransition))
+			{
 #if UE_POSE_SEARCH_TRACE_ENABLED
-					const FPoseSearchCost PoseCost = ComparePoses(SearchContext, PoseIdx, EPoseComparisonFlags::None, QueryValues);
-					SearchContext.BestCandidates.Add(PoseCost, PoseIdx, this, EPoseCandidateFlags::DiscardedBy_BlockTransition);
+				const FPoseSearchCost PoseCost = SearchIndex->ComparePoses(PoseIdx, SearchContext.QueryMirrorRequest, EPoseComparisonFlags::None, QueryValues);
+				SearchContext.BestCandidates.Add(PoseCost, PoseIdx, this, EPoseCandidateFlags::DiscardedBy_BlockTransition);
 #endif
-					continue;
-				}
+				continue;
+			}
 
-				if (CheckForNonSelectableIdx && Algo::BinarySearch(NonSelectableIdx, PoseIdx) != INDEX_NONE)
-				{
-					continue;
-				}
+			if (Algo::BinarySearch(NonSelectableIdx, PoseIdx) != INDEX_NONE)
+			{
+				continue;
+			}
 
-				const FPoseSearchCost PoseCost = ComparePoses(SearchContext, PoseIdx, EPoseComparisonFlags::None, QueryValues);
-				if (PoseCost < Result.PoseCost)
-				{
-					Result.PoseCost = PoseCost;
-					Result.PoseIdx = PoseIdx;
-				}
+			const FPoseSearchCost PoseCost = SearchIndex->ComparePoses(PoseIdx, SearchContext.QueryMirrorRequest, EPoseComparisonFlags::None, QueryValues);
+			if (PoseCost < Result.PoseCost)
+			{
+				Result.PoseCost = PoseCost;
+				Result.PoseIdx = PoseIdx;
+			}
 
 #if UE_POSE_SEARCH_TRACE_ENABLED
-				if (PoseSearchMode == EPoseSearchMode::BruteForce)
-				{
-					SearchContext.BestCandidates.Add(PoseCost, PoseIdx, this, EPoseCandidateFlags::Valid_Pose);
-				}
-#endif
+			if (PoseSearchMode == EPoseSearchMode::BruteForce)
+			{
+				SearchContext.BestCandidates.Add(PoseCost, PoseIdx, this, EPoseCandidateFlags::Valid_Pose);
 			}
+#endif
 		}
 
 		if (GetSkipSearchIfPossible() && Result.PoseCost.IsValid())
@@ -1902,8 +1862,7 @@ UE::PoseSearch::FSearchResult UPoseSearchDatabase::SearchBruteForce(UE::PoseSear
 	// finalizing Result properties
 	if (Result.PoseIdx != INDEX_NONE)
 	{
-		Result.SearchIndexAsset = SearchIndex->FindAssetForPose(Result.PoseIdx);
-		Result.AssetTime = SearchIndex->GetAssetTime(Result.PoseIdx, Result.SearchIndexAsset);
+		Result.AssetTime = SearchIndex->GetAssetTime(Result.PoseIdx);
 		Result.Database = this;
 #if WITH_EDITOR
 		Result.SearchIndexHash = GetSearchIndexHash();
@@ -1948,16 +1907,12 @@ UE::PoseSearch::FSearchResult UPoseSearchDatabaseSet::Search(UE::PoseSearch::FSe
 
 			TArrayView<const float> QueryValues = Result.ComposedQuery.GetValues();
 
-			const FPoseSearchIndexAsset* PoseSearchIndexAsset = SearchIndex->FindAssetForPose(SearchContext.CurrentResult.PoseIdx);
-			check(PoseSearchIndexAsset);
-
 			Result.PoseIdx = SearchContext.CurrentResult.PoseIdx;
-			Result.PoseCost = Database->ComparePoses(SearchContext, Result.PoseIdx, EPoseComparisonFlags::ContinuingPose, QueryValues);
+			Result.PoseCost = SearchIndex->ComparePoses(Result.PoseIdx, SearchContext.QueryMirrorRequest, EPoseComparisonFlags::ContinuingPose, QueryValues);
 			Result.ContinuingPoseCost = Result.PoseCost;
 			ContinuingCost = Result.PoseCost;
 
-			Result.SearchIndexAsset = SearchIndex->FindAssetForPose(Result.PoseIdx);
-			Result.AssetTime = SearchIndex->GetAssetTime(Result.PoseIdx, Result.SearchIndexAsset);
+			Result.AssetTime = SearchIndex->GetAssetTime(Result.PoseIdx);
 			Result.Database = Database;
 #if WITH_EDITOR
 			Result.SearchIndexHash = Database->GetSearchIndexHash();
@@ -2210,7 +2165,7 @@ bool FSearchContext::GetOrBuildQuery(const UPoseSearchDatabase* Database, FPoseS
 
 bool FSearchContext::IsCurrentResultFromDatabase(const UPoseSearchDatabase* Database) const
 {
-	return !bForceInterrupt && CurrentResult.IsValid() && CurrentResult.Database == Database;
+	return CurrentResult.IsValid() && CurrentResult.Database == Database;
 }
 
 TConstArrayView<float> FSearchContext::GetCurrentResultPrevPoseVector() const
@@ -2682,48 +2637,50 @@ void FSearchResult::Update(float NewAssetTime)
 	if (!IsValid())
 	{
 		Reset();
-		return;
 	}
-
-	if (SearchIndexAsset->Type == ESearchIndexAssetType::Sequence)
+	else
 	{
-		const FPoseSearchDatabaseSequence& DbSequence = Database->GetSequenceSourceAsset(SearchIndexAsset);
-
-		if (Database->GetPoseIndicesAndLerpValueFromTime(NewAssetTime, SearchIndexAsset, PrevPoseIdx, PoseIdx, NextPoseIdx, LerpValue))
+		const FPoseSearchIndexAsset& SearchIndexAsset = Database->GetSearchIndex()->GetAssetForPose(PoseIdx);
+		if (SearchIndexAsset.Type == ESearchIndexAssetType::Sequence)
 		{
-			AssetTime = NewAssetTime;
+			const FPoseSearchDatabaseSequence& DbSequence = Database->GetSequenceSourceAsset(SearchIndexAsset);
+
+			if (Database->GetPoseIndicesAndLerpValueFromTime(NewAssetTime, SearchIndexAsset, PrevPoseIdx, PoseIdx, NextPoseIdx, LerpValue))
+			{
+				AssetTime = NewAssetTime;
+			}
+			else
+			{
+				Reset();
+			}
 		}
-		else
+		else if (SearchIndexAsset.Type == ESearchIndexAssetType::BlendSpace)
 		{
-			Reset();
-		}
-	}
-	else if (SearchIndexAsset->Type == ESearchIndexAssetType::BlendSpace)
-	{
-		const FPoseSearchDatabaseBlendSpace& DbBlendSpace = Database->GetBlendSpaceSourceAsset(SearchIndexAsset);
+			const FPoseSearchDatabaseBlendSpace& DbBlendSpace = Database->GetBlendSpaceSourceAsset(SearchIndexAsset);
 
-		TArray<FBlendSampleData> BlendSamples;
-		int32 TriangulationIndex = 0;
-		DbBlendSpace.BlendSpace->GetSamplesFromBlendInput(SearchIndexAsset->BlendParameters, BlendSamples, TriangulationIndex, true);
+			TArray<FBlendSampleData> BlendSamples;
+			int32 TriangulationIndex = 0;
+			DbBlendSpace.BlendSpace->GetSamplesFromBlendInput(SearchIndexAsset.BlendParameters, BlendSamples, TriangulationIndex, true);
 
-		const float PlayLength = DbBlendSpace.BlendSpace->GetAnimationLengthFromSampleData(BlendSamples);
+			const float PlayLength = DbBlendSpace.BlendSpace->GetAnimationLengthFromSampleData(BlendSamples);
 
-		// Asset player time for blendspaces is normalized [0, 1] so we need to convert 
-		// to a real time before we advance it
-		const float RealTime = NewAssetTime * PlayLength;
-		if (Database->GetPoseIndicesAndLerpValueFromTime(RealTime, SearchIndexAsset, PrevPoseIdx, PoseIdx, NextPoseIdx, LerpValue))
-		{
-			AssetTime = NewAssetTime;
+			// Asset player time for blendspaces is normalized [0, 1] so we need to convert 
+			// to a real time before we advance it
+			const float RealTime = NewAssetTime * PlayLength;
+			if (Database->GetPoseIndicesAndLerpValueFromTime(RealTime, SearchIndexAsset, PrevPoseIdx, PoseIdx, NextPoseIdx, LerpValue))
+			{
+				AssetTime = NewAssetTime;
+			}
+			else
+			{
+				Reset();
+			}
 		}
 		else
 		{
-			Reset();
+			checkNoEntry();
 		}
 	}
-	else
-	{
-		checkNoEntry();
-	}
 }
 
 bool FSearchResult::IsValid() const
@@ -2744,7 +2701,6 @@ bool FSearchResult::IsValid() const
 void FSearchResult::Reset()
 {
 	PoseIdx = INDEX_NONE;
-	SearchIndexAsset = nullptr;
 	Database = nullptr;
 	ComposedQuery.Reset();
 	AssetTime = 0.0f;
@@ -2754,6 +2710,20 @@ void FSearchResult::Reset()
 #endif // WITH_EDITOR
 }
 
+const FPoseSearchIndexAsset* FSearchResult::GetSearchIndexAsset(bool bMandatory) const
+{
+	if (bMandatory)
+	{
+		check(IsValid());
+	}
+	else if (!IsValid())
+	{
+		return nullptr;
+	}
+
+	return &Database->GetSearchIndex()->GetAssetForPose(PoseIdx);
+}
+
 
 //////////////////////////////////////////////////////////////////////////
 // FAssetSamplerContext
@@ -3992,7 +3962,7 @@ FAssetIndexer::FSampleInfo FAssetIndexer::GetSampleInfo(float SampleTime) const
 		if (SampleTime < 0.0f)
 		{
 			RootMotionPerCycle = RootMotionPerCycle.Inverse();
-			RootDistancePerCycle *= -1;
+			RootDistancePerCycle *= -1.f;
 		}
 
 		// Find the remaining motion deltas after wrapping
@@ -4674,7 +4644,7 @@ void FDatabaseIndexingContext::PrepareIndexers()
 
 		if (SearchIndexAsset.Type == ESearchIndexAssetType::Sequence)
 		{
-			const FPoseSearchDatabaseSequence& DbSequence = Database->GetSequenceSourceAsset(&SearchIndexAsset);
+			const FPoseSearchDatabaseSequence& DbSequence = Database->GetSequenceSourceAsset(SearchIndexAsset);
 			const float SequenceLength = DbSequence.Sequence->GetPlayLength();
 			IndexerContext.MainSampler = GetSequenceSampler(DbSequence.Sequence);
 			IndexerContext.LeadInSampler = SearchIndexAsset.SamplingInterval.Min == 0.0f ? GetSequenceSampler(DbSequence.LeadInSequence) : nullptr;
@@ -4682,7 +4652,7 @@ void FDatabaseIndexingContext::PrepareIndexers()
 		}
 		else if (SearchIndexAsset.Type == ESearchIndexAssetType::BlendSpace)
 		{
-			const FPoseSearchDatabaseBlendSpace& DbBlendSpace = Database->GetBlendSpaceSourceAsset(&SearchIndexAsset);
+			const FPoseSearchDatabaseBlendSpace& DbBlendSpace = Database->GetBlendSpaceSourceAsset(SearchIndexAsset);
 			IndexerContext.MainSampler = GetBlendSpaceSampler(DbBlendSpace.BlendSpace, SearchIndexAsset.BlendParameters);
 		}
 		else
@@ -4746,8 +4716,16 @@ void FDatabaseIndexingContext::JoinIndex()
 		SearchIndexAsset.NumPoses = Output.NumIndexedPoses;
 		SearchIndexAsset.FirstPoseIdx = TotalPoses;
 
+		const int32 PoseMetadataStartIdx = SearchIndex->PoseMetadata.Num();
+		const int32 PoseMetadataEndIdx = PoseMetadataStartIdx + Output.PoseMetadata.Num();
+
 		SearchIndex->Values.Append(Output.FeatureVectorTable.GetData(), Output.FeatureVectorTable.Num());
 		SearchIndex->PoseMetadata.Append(Output.PoseMetadata);
+		
+		for (int32 i = PoseMetadataStartIdx; i < PoseMetadataEndIdx; ++i)
+		{
+			SearchIndex->PoseMetadata[i].AssetIndex = AssetIdx;
+		}
 
 		TotalPoses += Output.NumIndexedPoses;
 		TotalFloats += Output.FeatureVectorTable.Num();
diff --git a/Engine/Plugins/Experimental/Animation/PoseSearch/Source/Runtime/Private/PoseSearchLibrary.cpp b/Engine/Plugins/Experimental/Animation/PoseSearch/Source/Runtime/Private/PoseSearchLibrary.cpp
--- a/Engine/Plugins/Experimental/Animation/PoseSearch/Source/Runtime/Private/PoseSearchLibrary.cpp
+++ b/Engine/Plugins/Experimental/Animation/PoseSearch/Source/Runtime/Private/PoseSearchLibrary.cpp
@@ -39,12 +39,12 @@ bool FMotionMatchingState::CanAdvance(float DeltaTime, bool& bOutAdvanceToFollow
 		return false;
 	}
 
-	const FPoseSearchIndexAsset* SearchIndexAsset = GetCurrentSearchIndexAsset();
+	
+	const FPoseSearchIndexAsset* SearchIndexAsset = CurrentSearchResult.GetSearchIndexAsset(true);
 
 	if (SearchIndexAsset->Type == ESearchIndexAssetType::Sequence)
 	{
-		const FPoseSearchDatabaseSequence& DbSequence = 
-			CurrentSearchResult.Database->GetSequenceSourceAsset(SearchIndexAsset);
+		const FPoseSearchDatabaseSequence& DbSequence = CurrentSearchResult.Database->GetSequenceSourceAsset(*SearchIndexAsset);
 		const float AssetLength = DbSequence.Sequence->GetPlayLength();
 
 		float SteppedTime = CurrentSearchResult.AssetTime;
@@ -81,8 +81,7 @@ bool FMotionMatchingState::CanAdvance(float DeltaTime, bool& bOutAdvanceToFollow
 			{
 				bOutAdvanceToFollowUpAsset = true;
 
-				const FPoseSearchIndexAsset* FollowUpSearchIndexAsset =
-					&CurrentSearchResult.Database->GetSearchIndex()->Assets[FollowUpSearchIndexAssetIdx];
+				const FPoseSearchIndexAsset& FollowUpSearchIndexAsset = CurrentSearchResult.Database->GetSearchIndex()->Assets[FollowUpSearchIndexAssetIdx];
 
 				// Follow up asset time will start slightly before the beginning of the sequence as 
 				// this is essentially what the matching time in the corresponding main sequence is.
@@ -92,8 +91,7 @@ bool FMotionMatchingState::CanAdvance(float DeltaTime, bool& bOutAdvanceToFollow
 
 				// There is no correspoding pose index when we switch due to what is mentioned above
 				// so for now we just take whatever pose index is associated with the first frame.
-				OutFollowUpAsset.PoseIdx = CurrentSearchResult.Database->GetPoseIndexFromTime(FollowUpSearchIndexAsset->SamplingInterval.Min, FollowUpSearchIndexAsset);
-				OutFollowUpAsset.SearchIndexAsset = FollowUpSearchIndexAsset;
+				OutFollowUpAsset.PoseIdx = CurrentSearchResult.Database->GetPoseIndexFromTime(FollowUpSearchIndexAsset.SamplingInterval.Min, FollowUpSearchIndexAsset);
 				OutFollowUpAsset.AssetTime = FollowUpAssetTime;
 				return true;
 			}
@@ -101,8 +99,7 @@ bool FMotionMatchingState::CanAdvance(float DeltaTime, bool& bOutAdvanceToFollow
 	}
 	else if (SearchIndexAsset->Type == ESearchIndexAssetType::BlendSpace)
 	{
-		const FPoseSearchDatabaseBlendSpace& DbBlendSpace = 
-			CurrentSearchResult.Database->GetBlendSpaceSourceAsset(SearchIndexAsset);
+		const FPoseSearchDatabaseBlendSpace& DbBlendSpace = CurrentSearchResult.Database->GetBlendSpaceSourceAsset(*SearchIndexAsset);
 
 		TArray<FBlendSampleData> BlendSamples;
 		int32 TriangulationIndex = 0;
@@ -191,7 +188,7 @@ static void TraceMotionMatchingState(
 	{
 		const uint64 DatabaseId = FTraceMotionMatchingState::GetIdFromObject(Database);
 
-		int32 DbEntryIdx = -1;
+		int32 DbEntryIdx = INDEX_NONE;
 		for (int32 i = 0; i < DatabaseEntries.Num(); ++i)
 		{
 			if (DatabaseEntries[i].DatabaseId == DatabaseId)
@@ -358,7 +355,7 @@ void UpdateMotionMatchingState(
 			SearchContext.History = &PoseHistoryProvider->GetPoseHistory();
 		}
 
-		if (const FPoseSearchIndexAsset* CurrentIndexAsset = InOutMotionMatchingState.GetCurrentSearchIndexAsset())
+		if (const FPoseSearchIndexAsset* CurrentIndexAsset = InOutMotionMatchingState.CurrentSearchResult.GetSearchIndexAsset())
 		{
 			SearchContext.QueryMirrorRequest =
 				CurrentIndexAsset->bMirrored ?
@@ -419,28 +416,15 @@ void UpdateMotionMatchingState(
 #endif
 }
 
-const FPoseSearchIndexAsset* FMotionMatchingState::GetCurrentSearchIndexAsset() const
-{
-	if (CurrentSearchResult.IsValid())
-	{
-		return CurrentSearchResult.SearchIndexAsset;
-	}
-
-	return nullptr;
-}
-
-float FMotionMatchingState::ComputeJumpBlendTime(
-	const UE::PoseSearch::FSearchResult& Result, 
-	const FMotionMatchingSettings& Settings
-) const
+float FMotionMatchingState::ComputeJumpBlendTime(const UE::PoseSearch::FSearchResult& Result, const FMotionMatchingSettings& Settings) const
 {
-	const FPoseSearchIndexAsset* SearchIndexAsset = GetCurrentSearchIndexAsset();
+	const FPoseSearchIndexAsset* SearchIndexAsset = CurrentSearchResult.GetSearchIndexAsset();
 
 	// Use alternate blend time when changing between mirrored and unmirrored
 	float JumpBlendTime = Settings.BlendTime;
 	if ((SearchIndexAsset != nullptr) && (Settings.MirrorChangeBlendTime > 0.0f))
 	{
-		if (Result.SearchIndexAsset->bMirrored != SearchIndexAsset->bMirrored)
+		if (Result.GetSearchIndexAsset(true)->bMirrored != SearchIndexAsset->bMirrored)
 		{
 			JumpBlendTime = Settings.MirrorChangeBlendTime;
 		}
diff --git a/Engine/Plugins/Experimental/Animation/PoseSearch/Source/Runtime/Public/PoseSearch/PoseSearch.h b/Engine/Plugins/Experimental/Animation/PoseSearch/Source/Runtime/Public/PoseSearch/PoseSearch.h
--- a/Engine/Plugins/Experimental/Animation/PoseSearch/Source/Runtime/Public/PoseSearch/PoseSearch.h
+++ b/Engine/Plugins/Experimental/Animation/PoseSearch/Source/Runtime/Public/PoseSearch/PoseSearch.h
@@ -374,14 +374,14 @@ private:
 	friend class ::UPoseSearchSchema;
 
 	UPROPERTY()
-	int32 ChannelIdx = -1;
+	int32 ChannelIdx = INDEX_NONE;
 
 protected:
 	UPROPERTY()
-	int32 ChannelDataOffset = -1;
+	int32 ChannelDataOffset = INDEX_NONE;
 
 	UPROPERTY()
-	int32 ChannelCardinality = -1;
+	int32 ChannelCardinality = INDEX_NONE;
 };
 
 //////////////////////////////////////////////////////////////////////////
@@ -524,6 +524,10 @@ struct POSESEARCH_API FPoseSearchPoseMetadata
 	// @todo: consider float16
 	UPROPERTY()
 	float ContinuingPoseCostAddend = 0.0f;
+
+	// @todo: consider int16
+	UPROPERTY()
+	int32 AssetIndex = INDEX_NONE;
 };
 
 /**
@@ -641,16 +645,16 @@ struct POSESEARCH_API FPoseSearchIndex
 	bool IsEmpty() const;
 
 	TConstArrayView<float> GetPoseValues(int32 PoseIdx) const;
+	TConstArrayView<float> GetPoseValuesSafe(int32 PoseIdx) const;
 
-	const FPoseSearchIndexAsset* FindAssetForPose(int32 PoseIdx) const;
-	float GetAssetTime(int32 PoseIdx, const FPoseSearchIndexAsset* Asset) const;
+	const FPoseSearchIndexAsset& GetAssetForPose(int32 PoseIdx) const;
+	const FPoseSearchIndexAsset* GetAssetForPoseSafe(int32 PoseIdx) const;
+
+	float GetAssetTime(int32 PoseIdx) const;
 
 	void Reset();
 
-	// individual cost addends calculation methods
-	float ComputeMirrorMismatchAddend(int32 PoseIdx, UE::PoseSearch::FSearchContext& SearchContext) const;
-	float ComputeNotifyAddend(int32 PoseIdx) const;
-	float ComputeContinuingPoseCostAddend(int32 PoseIdx, UE::PoseSearch::EPoseComparisonFlags PoseComparisonFlags) const;
+	FPoseSearchCost ComparePoses(int32 PoseIdx, EPoseSearchBooleanRequest QueryMirrorRequest, UE::PoseSearch::EPoseComparisonFlags PoseComparisonFlags, TConstArrayView<float> QueryValues) const;
 };
 
 //////////////////////////////////////////////////////////////////////////
@@ -905,8 +909,6 @@ struct FSearchResult
 	// lerp value to find AssetTime from PrevPoseIdx -> AssetTime -> NextPoseIdx, within range [-0.5, 0.5]
 	float LerpValue = 0.f;
 
-	// @todo: it should be a weak pointer
-	const FPoseSearchIndexAsset* SearchIndexAsset = nullptr;
 	TWeakObjectPtr<const UPoseSearchDatabase> Database = nullptr;
 	FPoseSearchFeatureVectorBuilder ComposedQuery;
 
@@ -930,6 +932,8 @@ struct FSearchResult
 	bool IsValid() const;
 
 	void Reset();
+
+	const FPoseSearchIndexAsset* GetSearchIndexAsset(bool bMandatory = false) const;
 };
 
 } // namespace UE::PoseSearch
@@ -1000,15 +1004,14 @@ public:
 	bool IsValidForSearch() const;
 	bool IsValidPoseIndex(int32 PoseIdx) const;
 
-	float GetAssetTime(int32 PoseIdx, const FPoseSearchIndexAsset* SearchIndexAsset = nullptr) const;
-	int32 GetPoseIndexFromTime(float AssetTime, const FPoseSearchIndexAsset* SearchIndexAsset) const;
-	bool GetPoseIndicesAndLerpValueFromTime(float Time, const FPoseSearchIndexAsset* SearchIndexAsset, int32& PrevPoseIdx, int32& PoseIdx, int32& NextPoseIdx, float& LerpValue) const;
+	int32 GetPoseIndexFromTime(float AssetTime, const FPoseSearchIndexAsset& SearchIndexAsset) const;
+	bool GetPoseIndicesAndLerpValueFromTime(float Time, const FPoseSearchIndexAsset& SearchIndexAsset, int32& PrevPoseIdx, int32& PoseIdx, int32& NextPoseIdx, float& LerpValue) const;
 
-	const FPoseSearchDatabaseAnimationAssetBase& GetAnimationSourceAsset(const FPoseSearchIndexAsset* SearchIndexAsset) const;
-	const FPoseSearchDatabaseSequence& GetSequenceSourceAsset(const FPoseSearchIndexAsset* SearchIndexAsset) const;
-	const FPoseSearchDatabaseBlendSpace& GetBlendSpaceSourceAsset(const FPoseSearchIndexAsset* SearchIndexAsset) const;
-	const bool IsSourceAssetLooping(const FPoseSearchIndexAsset* SearchIndexAsset) const;
-	const FString GetSourceAssetName(const FPoseSearchIndexAsset* SearchIndexAsset) const;
+	const FPoseSearchDatabaseAnimationAssetBase& GetAnimationSourceAsset(const FPoseSearchIndexAsset& SearchIndexAsset) const;
+	const FPoseSearchDatabaseSequence& GetSequenceSourceAsset(const FPoseSearchIndexAsset& SearchIndexAsset) const;
+	const FPoseSearchDatabaseBlendSpace& GetBlendSpaceSourceAsset(const FPoseSearchIndexAsset& SearchIndexAsset) const;
+	const bool IsSourceAssetLooping(const FPoseSearchIndexAsset& SearchIndexAsset) const;
+	const FString GetSourceAssetName(const FPoseSearchIndexAsset& SearchIndexAsset) const;
 	int32 GetNumberOfPrincipalComponents() const;
 	
 	// UObject
@@ -1077,8 +1080,6 @@ public:
 
 	void BuildQuery(UE::PoseSearch::FSearchContext& SearchContext, FPoseSearchFeatureVectorBuilder& OutQuery) const;
 
-	FPoseSearchCost ComparePoses(UE::PoseSearch::FSearchContext& SearchContext, int32 PoseIdx, UE::PoseSearch::EPoseComparisonFlags PoseComparisonFlags, TConstArrayView<float> QueryValues) const;
-
 protected:
 	UE::PoseSearch::FSearchResult SearchPCAKDTree(UE::PoseSearch::FSearchContext& SearchContext) const;
 	UE::PoseSearch::FSearchResult SearchBruteForce(UE::PoseSearch::FSearchContext& SearchContext) const;
@@ -1111,12 +1112,8 @@ public:
 
 	UE::PoseSearch::FSearchResult Search(UE::PoseSearch::FSearchContext& SearchContext) const;
 
-protected:
-	FPoseSearchCost ComparePoses(int32 PoseIdx, UE::PoseSearch::EPoseComparisonFlags PoseComparisonFlags, TConstArrayView<float> QueryValues) const;
-
 public: // UObject
 	virtual void PreSave(FObjectPreSaveContext ObjectSaveContext) override;
-
 };
 
 //////////////////////////////////////////////////////////////////////////
@@ -1226,7 +1223,7 @@ struct FHistoricalPoseIndex
 		return HashCombineFast(::GetTypeHash(Index.PoseIndex), GetTypeHash(Index.DatabaseKey));
 	}
 
-	int32 PoseIndex = -1;
+	int32 PoseIndex = INDEX_NONE;
 	FObjectKey DatabaseKey;
 };
 
diff --git a/Engine/Plugins/Experimental/Animation/PoseSearch/Source/Runtime/Public/PoseSearch/PoseSearchLibrary.h b/Engine/Plugins/Experimental/Animation/PoseSearch/Source/Runtime/Public/PoseSearch/PoseSearchLibrary.h
--- a/Engine/Plugins/Experimental/Animation/PoseSearch/Source/Runtime/Public/PoseSearch/PoseSearchLibrary.h
+++ b/Engine/Plugins/Experimental/Animation/PoseSearch/Source/Runtime/Public/PoseSearch/PoseSearchLibrary.h
@@ -79,8 +79,6 @@ struct POSESEARCH_API FMotionMatchingState
 	// Internally stores the 'jump' to a new pose/sequence index and asset time for evaluation
 	void JumpToPose(const FAnimationUpdateContext& Context, const FMotionMatchingSettings& Settings, const UE::PoseSearch::FSearchResult& Result);
 
-	const FPoseSearchIndexAsset* GetCurrentSearchIndexAsset() const;
-
 	float ComputeJumpBlendTime(const UE::PoseSearch::FSearchResult& Result, const FMotionMatchingSettings& Settings) const;
 
 	UE::PoseSearch::FSearchResult CurrentSearchResult;
