fix for TraceMotionMatchingState when throttling is on

diff --git a/Engine/Plugins/Experimental/Animation/PoseSearch/Source/Editor/Private/PoseSearchDebugger.cpp b/Engine/Plugins/Experimental/Animation/PoseSearch/Source/Editor/Private/PoseSearchDebugger.cpp
--- a/Engine/Plugins/Experimental/Animation/PoseSearch/Source/Editor/Private/PoseSearchDebugger.cpp
+++ b/Engine/Plugins/Experimental/Animation/PoseSearch/Source/Editor/Private/PoseSearchDebugger.cpp
@@ -1152,18 +1152,21 @@ void SDebuggerDatabaseView::OnDatabaseRowSelectionChanged(
 
 void SDebuggerDatabaseView::SortDatabaseRows()
 {
-	if (SortMode == EColumnSortMode::Ascending)
-	{
-		UnfilteredDatabaseRows.Sort(Columns[SortColumn]->GetSortPredicate());
-	}
-	else if (SortMode == EColumnSortMode::Descending)
+	if (!UnfilteredDatabaseRows.IsEmpty())
 	{
-		auto DescendingPredicate = [this](const auto& Lhs, const auto& Rhs) -> bool
+		if (SortMode == EColumnSortMode::Ascending)
+		{
+			UnfilteredDatabaseRows.Sort(Columns[SortColumn]->GetSortPredicate());
+		}
+		else if (SortMode == EColumnSortMode::Descending)
 		{
-			return !Columns[SortColumn]->GetSortPredicate()(Lhs, Rhs);
-		};
+			auto DescendingPredicate = [this](const auto& Lhs, const auto& Rhs) -> bool
+			{
+				return !Columns[SortColumn]->GetSortPredicate()(Lhs, Rhs);
+			};
 
-		UnfilteredDatabaseRows.Sort(DescendingPredicate);
+			UnfilteredDatabaseRows.Sort(DescendingPredicate);
+		}
 	}
 }
 
diff --git a/Engine/Plugins/Experimental/Animation/PoseSearch/Source/Runtime/Private/PoseSearchLibrary.cpp b/Engine/Plugins/Experimental/Animation/PoseSearch/Source/Runtime/Private/PoseSearchLibrary.cpp
--- a/Engine/Plugins/Experimental/Animation/PoseSearch/Source/Runtime/Private/PoseSearchLibrary.cpp
+++ b/Engine/Plugins/Experimental/Animation/PoseSearch/Source/Runtime/Private/PoseSearchLibrary.cpp
@@ -184,7 +184,7 @@ static void TraceMotionMatchingState(
 
 	const float DeltaTime = UpdateContext.GetDeltaTime();
 
-	auto AddUniqueDatabase = [](TArray<FTraceMotionMatchingStateDatabaseEntry>& DatabaseEntries, const UPoseSearchDatabase* Database, const UE::PoseSearch::FSearchContext& SearchContext) -> int32
+	auto AddUniqueDatabase = [](TArray<FTraceMotionMatchingStateDatabaseEntry>& DatabaseEntries, const UPoseSearchDatabase* Database, UE::PoseSearch::FSearchContext& SearchContext) -> int32
 	{
 		const uint64 DatabaseId = FTraceMotionMatchingState::GetIdFromObject(Database);
 
@@ -201,10 +201,11 @@ static void TraceMotionMatchingState(
 		{
 			DbEntryIdx = DatabaseEntries.Add({ DatabaseId });
 
-			const FPoseSearchFeatureVectorBuilder* CachedQuery = SearchContext.GetCachedQuery(Database);
-			check(CachedQuery);
-
-			DatabaseEntries[DbEntryIdx].QueryVector = CachedQuery->GetValues();
+			// if throttling is on, the continuing pose can be valid, but no actual search occurred, so the query will not be cached, and we need to build it
+			FPoseSearchFeatureVectorBuilder FeatureVectorBuilder;
+			FeatureVectorBuilder.Init(Database->Schema);
+			SearchContext.GetOrBuildQuery(Database, FeatureVectorBuilder);
+			DatabaseEntries[DbEntryIdx].QueryVector = FeatureVectorBuilder.GetValues();
 		}
 
 		return DbEntryIdx;
