Additional IsNull fixup

diff --git a/Engine/Plugins/Experimental/Animation/PoseSearch/Source/Runtime/Private/PoseSearch.cpp b/Engine/Plugins/Experimental/Animation/PoseSearch/Source/Runtime/Private/PoseSearch.cpp
--- a/Engine/Plugins/Experimental/Animation/PoseSearch/Source/Runtime/Private/PoseSearch.cpp
+++ b/Engine/Plugins/Experimental/Animation/PoseSearch/Source/Runtime/Private/PoseSearch.cpp
@@ -413,7 +413,7 @@ void UPoseSearchSchema::ComputeCostBreakdowns(UE::PoseSearch::ICostBreakDownData
 {
 	for (const TObjectPtr<UPoseSearchFeatureChannel>& Channel : Channels)
 	{
-		if (!Channel.IsNull())
+		if (Channel)
 		{
 			Channel->ComputeCostBreakdowns(CostBreakDownData, this);
 		}
