motion matching - fix for FFeatureVectorHelper::ComputeMeanDeviations calculation: the mean deviation was calculated using only a very limited dataset (the size of the MeanDeviation vector that has cardinality of the number of features, usually in the order of 20-80 values), instead of the full database dataset (that averages in the order of few thousands values)

diff --git a/Engine/Plugins/Experimental/Animation/PoseSearch/Source/Runtime/Private/PoseSearch.cpp b/Engine/Plugins/Experimental/Animation/PoseSearch/Source/Runtime/Private/PoseSearch.cpp
--- a/Engine/Plugins/Experimental/Animation/PoseSearch/Source/Runtime/Private/PoseSearch.cpp
+++ b/Engine/Plugins/Experimental/Animation/PoseSearch/Source/Runtime/Private/PoseSearch.cpp
@@ -2524,7 +2524,7 @@ FVector2D FFeatureVectorHelper::DecodeVector2DInternal(TArrayView<const float> V
 
 void FFeatureVectorHelper::ComputeMeanDeviations(float MinMeanDeviation, const Eigen::MatrixXd& CenteredPoseMatrix, Eigen::VectorXd& MeanDeviations, int32& DataOffset, int32 Cardinality)
 {
-	const int32 NumPoses = MeanDeviations.size();
+	const int32 NumPoses = CenteredPoseMatrix.cols();
 
 	// Construct a submatrix for the feature and find the average distance to the feature's centroid.
 	// Since we've already mean centered the data, the average distance to the centroid is simply the average norm.
