motion matching - preview BS support

diff --git a/Engine/Plugins/Experimental/Animation/PoseSearch/Source/Editor/Private/PoseSearchDatabaseViewModel.cpp b/Engine/Plugins/Experimental/Animation/PoseSearch/Source/Editor/Private/PoseSearchDatabaseViewModel.cpp
--- a/Engine/Plugins/Experimental/Animation/PoseSearch/Source/Editor/Private/PoseSearchDatabaseViewModel.cpp
+++ b/Engine/Plugins/Experimental/Animation/PoseSearch/Source/Editor/Private/PoseSearchDatabaseViewModel.cpp
@@ -84,7 +84,7 @@ namespace UE::PoseSearch
 			}
 		}
 
-		const FPoseSearchIndex* SearchIndex = PoseSearchDatabase->GetSearchIndexSafe();
+		const FPoseSearchIndex* SearchIndex = PoseSearchDatabase->GetSearchIndexSafe(false);
 		if (SearchIndex)
 		{
 			for (int32 IndexAssetIndex = 0; IndexAssetIndex < SearchIndex->Assets.Num(); ++IndexAssetIndex)
@@ -174,7 +174,7 @@ namespace UE::PoseSearch
 
 		if (PoseSearchDatabase && PoseSearchDatabase->IsValidForIndexing())
 		{
-			if (const FPoseSearchIndex* SearchIndex = PoseSearchDatabase->GetSearchIndexSafe())
+			if (const FPoseSearchIndex* SearchIndex = PoseSearchDatabase->GetSearchIndexSafe(false))
 			{
 				const FPoseSearchIndexAsset& IndexAsset = SearchIndex->Assets[IndexAssetIndex];
 				UAnimationAsset* PreviewAsset = nullptr;
@@ -182,11 +182,31 @@ namespace UE::PoseSearch
 				{
 					FPoseSearchDatabaseSequence DatabaseSequence = PoseSearchDatabase->Sequences[IndexAsset.SourceAssetIdx];
 					PreviewAsset = DatabaseSequence.Sequence;
+
+					FSequenceSampler::FInput Input;
+					Input.ExtrapolationParameters = PoseSearchDatabase->ExtrapolationParameters;
+					Input.Sequence = DatabaseSequence.Sequence;
+
+					PreviewActor.SequenceSampler.Init(Input);
+					PreviewActor.SequenceSampler.Process();
+
 				}
 				else if (IndexAsset.Type == ESearchIndexAssetType::BlendSpace)
 				{
 					FPoseSearchDatabaseBlendSpace DatabaseBlendSpace = PoseSearchDatabase->BlendSpaces[IndexAsset.SourceAssetIdx];
 					PreviewAsset = DatabaseBlendSpace.BlendSpace;
+
+					FAssetSamplingContext SamplingContext;
+					SamplingContext.Init(PoseSearchDatabase->Schema);
+
+					FBlendSpaceSampler::FInput Input;
+					Input.BoneContainer = SamplingContext.BoneContainer;
+					Input.ExtrapolationParameters = PoseSearchDatabase->ExtrapolationParameters;
+					Input.BlendSpace = DatabaseBlendSpace.BlendSpace;
+					Input.BlendParameters = IndexAsset.BlendParameters;
+
+					PreviewActor.BlendSpaceSampler.Init(Input);
+					PreviewActor.BlendSpaceSampler.Process();
 				}
 
 				if (!PreviewAsset)
@@ -212,11 +232,10 @@ namespace UE::PoseSearch
 				PreviewActor.AnimInstance = NewObject<UAnimPreviewInstance>(PreviewActor.Mesh.Get());
 				PreviewActor.Mesh->PreviewInstance = PreviewActor.AnimInstance.Get();
 				PreviewActor.AnimInstance->InitializeAnimation();
-				PreviewActor.AnimInstance->SetBlendSpacePosition(IndexAsset.BlendParameters);
-
 				PreviewActor.Mesh->SetSkeletalMesh(PoseSearchDatabase->Schema->Skeleton->GetPreviewMesh(true));
 				PreviewActor.Mesh->EnablePreview(true, PreviewAsset);
 				PreviewActor.AnimInstance->SetAnimationAsset(PreviewAsset, false, 0.0f);
+				PreviewActor.AnimInstance->SetBlendSpacePosition(IndexAsset.BlendParameters);
 				if (IndexAsset.bMirrored && PoseSearchDatabase->Schema)
 				{
 					PreviewActor.AnimInstance->SetMirrorDataTable(PoseSearchDatabase->Schema->MirrorDataTable);
@@ -260,12 +279,12 @@ namespace UE::PoseSearch
 	{
 		PlayTime += DeltaSeconds * DeltaTimeMultiplier;
 		PlayTime = FMath::Clamp(PlayTime, 0.0f, MaxPreviewPlayLength);
-		UpdatePreviewActors();
+		UpdatePreviewActors(true);
 	}
 
-	void FDatabaseViewModel::UpdatePreviewActors()
+	void FDatabaseViewModel::UpdatePreviewActors(bool bInTickPlayTime)
 	{
-		if (const FPoseSearchIndex* SearchIndex = PoseSearchDatabase->GetSearchIndexSafe())
+		if (const FPoseSearchIndex* SearchIndex = PoseSearchDatabase->GetSearchIndexSafe(false))
 		{
 			TSharedPtr<FDatabasePreviewScene> PreviewScene = PreviewScenePtr.Pin();
 			for (FDatabasePreviewActor& PreviewActor : GetPreviewActors())
@@ -274,29 +293,47 @@ namespace UE::PoseSearch
 				{
 					if (const UAnimationAsset* PreviewAsset = PreviewActor.AnimInstance->GetAnimationAsset())
 					{
-						float CurrentTime = 0.0f;
+						float CurrentTime = 0.f;
+						float CurrentTimeScaled = 0.f;
 						FAnimationRuntime::AdvanceTime(false, PlayTime, CurrentTime, PreviewAsset->GetPlayLength());
 
-						if (const UAnimSequence* AnimSequence = Cast<UAnimSequence>(PreviewAsset))
+						bool bIsSequence = false;
+						bool bIsBlendSpace = false;
+						if (Cast<UAnimSequence>(PreviewAsset))
+						{
+							bIsSequence = true;
+							CurrentTimeScaled = CurrentTime;
+						}
+						else if (Cast<UBlendSpace>(PreviewAsset))
+						{
+							bIsBlendSpace = true;
+							CurrentTimeScaled = PreviewActor.BlendSpaceSampler.GetPlayLength() > UE_KINDA_SMALL_NUMBER ? CurrentTime / PreviewActor.BlendSpaceSampler.GetPlayLength() : 0.f;
+						}
+						
+						const float OldTime = PreviewActor.AnimInstance->GetCurrentTime();
+						if (!FMath::IsNearlyEqual(CurrentTimeScaled, OldTime, UE_KINDA_SMALL_NUMBER))
 						{
 							FTransform RootMotion = FTransform::Identity;
-							RootMotion = AnimSequence->ExtractRootMotionFromRange(0.0f, CurrentTime);
+							if (bIsSequence)
+							{
+								RootMotion = PreviewActor.SequenceSampler.ExtractRootTransform(CurrentTime);
+							}
+							else if (bIsBlendSpace)
+							{
+								RootMotion = PreviewActor.BlendSpaceSampler.ExtractRootTransform(CurrentTime);
+							}
+
 							if (PreviewActor.AnimInstance->GetMirrorDataTable())
 							{
 								RootMotion = MirrorRootMotion(RootMotion, PreviewActor.AnimInstance->GetMirrorDataTable());
 							}
 							PreviewActor.Actor->SetActorTransform(RootMotion);
-						}
-						else if (const UBlendSpace* BlendSpace = Cast<UBlendSpace>(PreviewAsset))
-						{
-							// @todo: add BS support
-							//RootMotion = BlendSpace->ExtractRootMotionFromRange(0.0f, CurrentTime);
-						}
 
-						PreviewActor.AnimInstance->SetPosition(CurrentTime);
+							PreviewActor.AnimInstance->SetPosition(CurrentTimeScaled);
 
-						const FPoseSearchIndexAsset& SearchIndexAsset = SearchIndex->Assets[PreviewActor.IndexAssetIndex];
-						PreviewActor.CurrentPoseIndex = PoseSearchDatabase->GetPoseIndexFromTime(CurrentTime, SearchIndexAsset);
+							const FPoseSearchIndexAsset& SearchIndexAsset = SearchIndex->Assets[PreviewActor.IndexAssetIndex];
+							PreviewActor.CurrentPoseIndex = PoseSearchDatabase->GetPoseIndexFromTime(CurrentTime, SearchIndexAsset);
+						}
 					}
 				}
 			}
@@ -318,8 +355,7 @@ namespace UE::PoseSearch
 		FTransform RootMotion,
 		const UMirrorDataTable* MirrorDataTable)
 	{
-		const FTransform RootReferenceTransform =
-			PoseSearchDatabase->Schema->Skeleton->GetReferenceSkeleton().GetRefBonePose()[0];
+		const FTransform RootReferenceTransform = PoseSearchDatabase->Schema->Skeleton->GetReferenceSkeleton().GetRefBonePose()[0];
 		const FQuat RootReferenceRotation = RootReferenceTransform.GetRotation();
 
 		const EAxis::Type MirrorAxis = MirrorDataTable->MirrorAxis;
@@ -419,7 +455,7 @@ namespace UE::PoseSearch
 	{
 		if (PoseSearchDatabase)
 		{
-			if (const FPoseSearchIndex* SearchIndex = PoseSearchDatabase->GetSearchIndexSafe())
+			if (const FPoseSearchIndex* SearchIndex = PoseSearchDatabase->GetSearchIndexSafe(false))
 			{
 				if (SearchIndex->Assets.IsValidIndex(SelectedActorIndexAssetIndex))
 				{
diff --git a/Engine/Plugins/Experimental/Animation/PoseSearch/Source/Editor/Public/PoseSearchDatabaseViewModel.h b/Engine/Plugins/Experimental/Animation/PoseSearch/Source/Editor/Public/PoseSearchDatabaseViewModel.h
--- a/Engine/Plugins/Experimental/Animation/PoseSearch/Source/Editor/Public/PoseSearchDatabaseViewModel.h
+++ b/Engine/Plugins/Experimental/Animation/PoseSearch/Source/Editor/Public/PoseSearchDatabaseViewModel.h
@@ -45,6 +45,9 @@ namespace UE::PoseSearch
 		int32 IndexAssetIndex = INDEX_NONE;
 		int32 CurrentPoseIndex = INDEX_NONE;
 
+		FSequenceSampler SequenceSampler;
+		FBlendSpaceSampler BlendSpaceSampler;
+
 		bool IsValid()
 		{
 			const bool bIsValid = Actor.IsValid() && Mesh.IsValid() && AnimInstance.IsValid();
@@ -147,7 +150,7 @@ namespace UE::PoseSearch
 
 		FDatabasePreviewActor SpawnPreviewActor(int32 IndexAssetIndex);
 
-		void UpdatePreviewActors();
+		void UpdatePreviewActors(bool bInTickPlayTime = false);
 
 		FTransform MirrorRootMotion(FTransform RootMotion, const class UMirrorDataTable* MirrorDataTable);
 	};
diff --git a/Engine/Plugins/Experimental/Animation/PoseSearch/Source/Runtime/Private/PoseSearch.cpp b/Engine/Plugins/Experimental/Animation/PoseSearch/Source/Runtime/Private/PoseSearch.cpp
--- a/Engine/Plugins/Experimental/Animation/PoseSearch/Source/Runtime/Private/PoseSearch.cpp
+++ b/Engine/Plugins/Experimental/Animation/PoseSearch/Source/Runtime/Private/PoseSearch.cpp
@@ -726,25 +726,27 @@ const FPoseSearchIndex* UPoseSearchDatabase::GetSearchIndex() const
 	return PrivateDerivedData ? &PrivateDerivedData->SearchIndex : nullptr;
 }
 
-const FPoseSearchIndex* UPoseSearchDatabase::GetSearchIndexSafe() const
+const FPoseSearchIndex* UPoseSearchDatabase::GetSearchIndexSafe(bool bVerboseLogging) const
 {
 	const FPoseSearchIndex* SearchIndex = GetSearchIndex();
 
 	if (!SearchIndex || !SearchIndex->IsValid() || SearchIndex->IsEmpty())
 	{
-		if (!Schema)
+		if (bVerboseLogging)
 		{
-			UE_LOG(LogAnimation, Warning, TEXT("UPoseSearchDatabase %s failed to index. Reason: no Schema!"), *GetName());
-		}
-		else if (!Schema->IsValid())
-		{
-			UE_LOG(LogAnimation, Warning, TEXT("UPoseSearchDatabase %s failed to index. Reason: Schema %s is invalid"), *GetName(), *Schema->GetName());
-		}
-		else
-		{
-			UE_LOG(LogAnimation, Warning, TEXT("UPoseSearchDatabase %s failed to index. Reason: is there any unsaved modified asset?"), *GetName());
+			if (!Schema)
+			{
+				UE_LOG(LogAnimation, Warning, TEXT("UPoseSearchDatabase %s failed to index. Reason: no Schema!"), *GetName());
+			}
+			else if (!Schema->IsValid())
+			{
+				UE_LOG(LogAnimation, Warning, TEXT("UPoseSearchDatabase %s failed to index. Reason: Schema %s is invalid"), *GetName(), *Schema->GetName());
+			}
+			else
+			{
+				UE_LOG(LogAnimation, Warning, TEXT("UPoseSearchDatabase %s failed to index. Reason: is there any unsaved modified asset?"), *GetName());
+			}
 		}
-
 		SearchIndex = nullptr;
 	}
 
@@ -1986,40 +1988,14 @@ void FPoseSearchFeatureVectorBuilder::Init(const UPoseSearchSchema* InSchema)
 {
 	check(InSchema && InSchema->IsValid());
 	Schema = InSchema;
-	ResetFeatures();
+	Values.Reset();
+	Values.SetNumZeroed(Schema->SchemaCardinality);
 }
 
 void FPoseSearchFeatureVectorBuilder::Reset()
 {
 	Schema = nullptr;
-	Values.Reset(0);
-}
-
-void FPoseSearchFeatureVectorBuilder::ResetFeatures()
-{
-	Values.Reset(0);
-	Values.SetNumZeroed(Schema->SchemaCardinality);
-}
-
-void FPoseSearchFeatureVectorBuilder::CopyFromSearchIndex(const FPoseSearchIndex& SearchIndex, int32 PoseIdx)
-{
-	check(Schema == SearchIndex.Schema);
-	Values = SearchIndex.GetPoseValues(PoseIdx);
-}
-
-bool FPoseSearchFeatureVectorBuilder::IsInitialized() const
-{
-	return (Schema != nullptr) && (Values.Num() == Schema->SchemaCardinality);
-}
-
-bool FPoseSearchFeatureVectorBuilder::IsInitializedForSchema(const UPoseSearchSchema* InSchema) const
-{
-	return (Schema == InSchema) && IsInitialized();
-}
-
-bool FPoseSearchFeatureVectorBuilder::IsCompatible(const FPoseSearchFeatureVectorBuilder& OtherBuilder) const
-{
-	return IsInitialized() && (Schema == OtherBuilder.Schema);
+	Values.Reset();
 }
 
 namespace UE::PoseSearch
@@ -2921,46 +2897,9 @@ static FTransform ExtrapolateRootMotion(
 
 //////////////////////////////////////////////////////////////////////////
 // FSequenceSampler
-
-struct FSequenceSampler : public IAssetSampler
-{
-public:
-	struct FInput
-	{
-		const UAnimSequence* Sequence = nullptr;
-		bool bLoopable = false;
-		int32 RootDistanceSamplingRate = 60;
-		FPoseSearchExtrapolationParameters ExtrapolationParameters;
-	} Input;
-
-	void Init(const FInput& Input);
-	void Process();
-
-	float GetPlayLength() const override { return Input.Sequence->GetPlayLength(); };
-	bool IsLoopable() const override { return Input.bLoopable; };
-
-	float GetTimeFromRootDistance(float Distance) const override;
-
-	float GetTotalRootDistance() const override { return TotalRootDistance; };
-	FTransform GetTotalRootTransform() const override { return TotalRootTransform; }
-
-	virtual void ExtractPose(const FAnimExtractContext& ExtractionCtx, FAnimationPoseData& OutAnimPoseData) const override;
-	virtual float ExtractRootDistance(float Time) const override;
-	virtual FTransform ExtractRootTransform(float Time) const override;
-	virtual void ExtractPoseSearchNotifyStates(float Time, TArray<class UAnimNotifyState_PoseSearchBase*>& NotifyStates) const override;
-	virtual const UAnimationAsset* GetAsset() const override { return Input.Sequence; }
-
-private:
-	float TotalRootDistance = 0.0f;
-	FTransform TotalRootTransform = FTransform::Identity;
-	TArray<float> AccumulatedRootDistance;
-
-	void ProcessRootDistance();
-};
-
 void FSequenceSampler::Init(const FInput& InInput)
 {
-	check(InInput.Sequence);
+	check(InInput.Sequence.Get());
 
 	Input = InInput;
 }
@@ -2996,6 +2935,11 @@ float FSequenceSampler::GetTimeFromRootDistance(float Distance) const
 	return ClipTime;
 }
 
+bool FSequenceSampler::IsLoopable() const
+{
+	return Input.Sequence->bLoop;
+}
+
 void FSequenceSampler::ExtractPose(const FAnimExtractContext& ExtractionCtx, FAnimationPoseData& OutAnimPoseData) const
 {
 	Input.Sequence->GetAnimationPose(OutAnimPoseData, ExtractionCtx);
@@ -3003,7 +2947,7 @@ void FSequenceSampler::ExtractPose(const FAnimExtractContext& ExtractionCtx, FAn
 
 FTransform FSequenceSampler::ExtractRootTransform(float Time) const
 {
-	if (Input.bLoopable)
+	if (IsLoopable())
 	{
 		FTransform LoopableRootTransform = Input.Sequence->ExtractRootMotion(0.0f, Time, true);
 		return LoopableRootTransform;
@@ -3134,62 +3078,16 @@ void FSequenceSampler::ProcessRootDistance()
 	TotalRootDistance = AccumulatedRootDistance.Last();
 }
 
-//////////////////////////////////////////////////////////////////////////
-// FBlendSpaceSampler
-
-struct FBlendSpaceSampler : public IAssetSampler
+const UAnimationAsset* FSequenceSampler::GetAsset() const
 {
-public:
-	struct FInput
-	{
-		const FAssetSamplingContext* SamplingContext = nullptr;
-		const UBlendSpace* BlendSpace = nullptr;
-		bool bLoopable = false;
-		int32 RootDistanceSamplingRate = 60;
-		int32 RootTransformSamplingRate = 60;
-		FPoseSearchExtrapolationParameters ExtrapolationParameters;
-		FVector BlendParameters;
-	} Input;
-
-	void Init(const FInput& Input);
-
-	void Process();
-
-	float GetPlayLength() const override { return PlayLength; };
-	bool IsLoopable() const override { return Input.bLoopable; };
-
-	float GetTimeFromRootDistance(float Distance) const override;
-
-	float GetTotalRootDistance() const override { return TotalRootDistance; };
-	FTransform GetTotalRootTransform() const override { return TotalRootTransform; }
-
-	virtual void ExtractPose(const FAnimExtractContext& ExtractionCtx, FAnimationPoseData& OutAnimPoseData) const override;
-	virtual float ExtractRootDistance(float Time) const override;
-	virtual FTransform ExtractRootTransform(float Time) const override;
-	virtual void ExtractPoseSearchNotifyStates(float Time, TArray<class UAnimNotifyState_PoseSearchBase*>& NotifyStates) const override;
-
-	virtual const UAnimationAsset* GetAsset() const override { return Input.BlendSpace; }
-
-private:
-	float PlayLength = 0.0f;
-	float TotalRootDistance = 0.0f;
-	FTransform TotalRootTransform = FTransform::Identity;
-	TArray<float> AccumulatedRootDistance;
-	TArray<FTransform> AccumulatedRootTransform;
-	
-	void ProcessPlayLength();
-	void ProcessRootDistance();
-	void ProcessRootTransform();
-
-	// Extracts the pre-computed blend space root transform. ProcessRootTransform must be run first.
-	FTransform ExtractBlendSpaceRootTrackTransform(float Time) const;
-	FTransform ExtractBlendSpaceRootMotion(float StartTime, float DeltaTime, bool bAllowLooping) const;
-	FTransform ExtractBlendSpaceRootMotionFromRange(float StartTrackPosition, float EndTrackPosition) const;
-};
+	return Input.Sequence.Get();
+}
 
+//////////////////////////////////////////////////////////////////////////
+// FBlendSpaceSampler
 void FBlendSpaceSampler::Init(const FInput& InInput)
 {
-	check(InInput.BlendSpace);
+	check(InInput.BlendSpace.Get());
 
 	Input = InInput;
 }
@@ -3229,6 +3127,11 @@ float FBlendSpaceSampler::GetTimeFromRootDistance(float Distance) const
 	return ClipTime;
 }
 
+bool FBlendSpaceSampler::IsLoopable() const
+{
+	return Input.BlendSpace->bLoop;
+}
+
 void FBlendSpaceSampler::ExtractPose(const FAnimExtractContext& ExtractionCtx, FAnimationPoseData& OutAnimPoseData) const
 {
 	TArray<FBlendSampleData> BlendSamples;
@@ -3252,7 +3155,7 @@ void FBlendSpaceSampler::ExtractPose(const FAnimExtractContext& ExtractionCtx, F
 
 FTransform FBlendSpaceSampler::ExtractRootTransform(float Time) const
 {
-	if (Input.bLoopable)
+	if (IsLoopable())
 	{
 		FTransform LoopableRootTransform = ExtractBlendSpaceRootMotion(0.0f, Time, true);
 		return LoopableRootTransform;
@@ -3461,8 +3364,6 @@ void FBlendSpaceSampler::ProcessPlayLength()
 	Input.BlendSpace->GetSamplesFromBlendInput(Input.BlendParameters, BlendSamples, TriangulationIndex, true);
 
 	PlayLength = Input.BlendSpace->GetAnimationLengthFromSampleData(BlendSamples);
-
-	checkf(PlayLength > 0.0f, TEXT("Blendspace has zero play length"));
 }
 
 void FBlendSpaceSampler::ProcessRootTransform()
@@ -3487,7 +3388,7 @@ void FBlendSpaceSampler::ProcessRootTransform()
 
 		FDeltaTimeRecord DeltaTimeRecord;
 		DeltaTimeRecord.Set(PreviousTime, CurrentTime - PreviousTime);
-		FAnimExtractContext ExtractionCtx(static_cast<double>(CurrentTime), true, DeltaTimeRecord, Input.bLoopable);
+		FAnimExtractContext ExtractionCtx(static_cast<double>(CurrentTime), true, DeltaTimeRecord, IsLoopable());
 
 		for (int32 BlendSampleIdex = 0; BlendSampleIdex < BlendSamples.Num(); BlendSampleIdex++)
 		{
@@ -3506,8 +3407,8 @@ void FBlendSpaceSampler::ProcessRootTransform()
 		Anim::FStackAttributeContainer StackAttributeContainer;
 		FAnimationPoseData AnimPoseData(Pose, BlendedCurve, StackAttributeContainer);
 
-		Pose.SetBoneContainer(&Input.SamplingContext->BoneContainer);
-		BlendedCurve.InitFrom(Input.SamplingContext->BoneContainer);
+		Pose.SetBoneContainer(&Input.BoneContainer);
+		BlendedCurve.InitFrom(Input.BoneContainer);
 
 		Input.BlendSpace->GetAnimationPose(BlendSamples, ExtractionCtx, AnimPoseData);
 
@@ -3568,6 +3469,11 @@ void FBlendSpaceSampler::ProcessRootDistance()
 	TotalRootDistance = AccumulatedRootDistance.Last();
 }
 
+const UAnimationAsset* FBlendSpaceSampler::GetAsset() const
+{
+	return Input.BlendSpace.Get();
+}
+
 //////////////////////////////////////////////////////////////////////////
 // FAssetIndexer helpers
 
@@ -4461,7 +4367,6 @@ bool BuildIndex(const UAnimSequence* Sequence, UPoseSearchSequenceMetaData* Sequ
 	FSequenceSampler::FInput SamplerInput;
 	SamplerInput.ExtrapolationParameters = SequenceMetaData->ExtrapolationParameters;
 	SamplerInput.Sequence = Sequence;
-	SamplerInput.bLoopable = false;
 	Sampler.Init(SamplerInput);
 	Sampler.Process();
 
@@ -4531,7 +4436,6 @@ void FDatabaseIndexingContext::PrepareSamplers()
 			FSequenceSampler::FInput Input;
 			Input.ExtrapolationParameters = Database->ExtrapolationParameters;
 			Input.Sequence = Sequence;
-			Input.bLoopable = bLoopable;
 			SequenceSamplers[SequenceSamplerIdx].Init(Input);
 		}
 	};
@@ -4597,10 +4501,9 @@ void FDatabaseIndexingContext::PrepareSamplers()
 						BlendSpaceSamplerMap.Add({ DbBlendSpace.BlendSpace, BlendParameters }, BlendSpaceSamplerIdx);
 
 						FBlendSpaceSampler::FInput Input;
-						Input.SamplingContext = &SamplingContext;
+						Input.BoneContainer = SamplingContext.BoneContainer;
 						Input.ExtrapolationParameters = Database->ExtrapolationParameters;
 						Input.BlendSpace = DbBlendSpace.BlendSpace;
-						Input.bLoopable = DbBlendSpace.BlendSpace->bLoop;
 						Input.BlendParameters = BlendParameters;
 
 						BlendSpaceSamplers[BlendSpaceSamplerIdx].Init(Input);
diff --git a/Engine/Plugins/Experimental/Animation/PoseSearch/Source/Runtime/Public/PoseSearch/PoseSearch.h b/Engine/Plugins/Experimental/Animation/PoseSearch/Source/Runtime/Public/PoseSearch/PoseSearch.h
--- a/Engine/Plugins/Experimental/Animation/PoseSearch/Source/Runtime/Public/PoseSearch/PoseSearch.h
+++ b/Engine/Plugins/Experimental/Animation/PoseSearch/Source/Runtime/Public/PoseSearch/PoseSearch.h
@@ -241,6 +241,90 @@ public:
 	virtual const UAnimationAsset* GetAsset() const = 0;
 };
 
+struct POSESEARCH_API FSequenceSampler : public IAssetSampler
+{
+public:
+	struct FInput
+	{
+		TWeakObjectPtr<const UAnimSequence> Sequence;
+		int32 RootDistanceSamplingRate = 60;
+		FPoseSearchExtrapolationParameters ExtrapolationParameters;
+	} Input;
+
+	void Init(const FInput& Input);
+	void Process();
+
+	float GetPlayLength() const override { return Input.Sequence->GetPlayLength(); }
+	bool IsLoopable() const override;
+
+	float GetTimeFromRootDistance(float Distance) const override;
+
+	float GetTotalRootDistance() const override { return TotalRootDistance; }
+	FTransform GetTotalRootTransform() const override { return TotalRootTransform; }
+
+	virtual void ExtractPose(const FAnimExtractContext& ExtractionCtx, FAnimationPoseData& OutAnimPoseData) const override;
+	virtual float ExtractRootDistance(float Time) const override;
+	virtual FTransform ExtractRootTransform(float Time) const override;
+	virtual void ExtractPoseSearchNotifyStates(float Time, TArray<class UAnimNotifyState_PoseSearchBase*>& NotifyStates) const override;
+	virtual const UAnimationAsset* GetAsset() const override;
+
+private:
+	float TotalRootDistance = 0.0f;
+	FTransform TotalRootTransform = FTransform::Identity;
+	TArray<float> AccumulatedRootDistance;
+
+	void ProcessRootDistance();
+};
+
+struct POSESEARCH_API FBlendSpaceSampler : public IAssetSampler
+{
+public:
+	struct FInput
+	{
+		FBoneContainer BoneContainer;
+		TWeakObjectPtr<const UBlendSpace> BlendSpace;
+		int32 RootDistanceSamplingRate = 60;
+		int32 RootTransformSamplingRate = 60;
+		FPoseSearchExtrapolationParameters ExtrapolationParameters;
+		FVector BlendParameters;
+	} Input;
+
+	void Init(const FInput& Input);
+
+	void Process();
+
+	float GetPlayLength() const override { return PlayLength; }
+	bool IsLoopable() const override;
+
+	float GetTimeFromRootDistance(float Distance) const override;
+
+	float GetTotalRootDistance() const override { return TotalRootDistance; }
+	FTransform GetTotalRootTransform() const override { return TotalRootTransform; }
+
+	virtual void ExtractPose(const FAnimExtractContext& ExtractionCtx, FAnimationPoseData& OutAnimPoseData) const override;
+	virtual float ExtractRootDistance(float Time) const override;
+	virtual FTransform ExtractRootTransform(float Time) const override;
+	virtual void ExtractPoseSearchNotifyStates(float Time, TArray<class UAnimNotifyState_PoseSearchBase*>& NotifyStates) const override;
+
+	virtual const UAnimationAsset* GetAsset() const override;
+
+private:
+	float PlayLength = 0.0f;
+	float TotalRootDistance = 0.0f;
+	FTransform TotalRootTransform = FTransform::Identity;
+	TArray<float> AccumulatedRootDistance;
+	TArray<FTransform> AccumulatedRootTransform;
+
+	void ProcessPlayLength();
+	void ProcessRootDistance();
+	void ProcessRootTransform();
+
+	// Extracts the pre-computed blend space root transform. ProcessRootTransform must be run first.
+	FTransform ExtractBlendSpaceRootTrackTransform(float Time) const;
+	FTransform ExtractBlendSpaceRootMotion(float StartTime, float DeltaTime, bool bAllowLooping) const;
+	FTransform ExtractBlendSpaceRootMotionFromRange(float StartTrackPosition, float EndTrackPosition) const;
+};
+
 /**
  * Inputs for asset indexing
  */
@@ -730,13 +814,13 @@ struct POSESEARCH_API FPoseSearchDatabaseBlendSpace : public FPoseSearchDatabase
 	// If to use the blendspace grid locations as parameter sample locations.
 	// When enabled, NumberOfHorizontalSamples and NumberOfVerticalSamples are ignored.
 	UPROPERTY(EditAnywhere, Category = "BlendSpace")
-	bool bUseGridForSampling = true;
+	bool bUseGridForSampling = false;
 
 	UPROPERTY(EditAnywhere, Category = "BlendSpace", meta = (EditCondition = "!bUseGridForSampling", EditConditionHides, ClampMin = "1", UIMin = "1", UIMax = "25"))
-	int32 NumberOfHorizontalSamples = 5;
+	int32 NumberOfHorizontalSamples = 9;
 
 	UPROPERTY(EditAnywhere, Category = "BlendSpace", meta = (EditCondition = "!bUseGridForSampling", EditConditionHides, ClampMin = "1", UIMin = "1", UIMax = "25"))
-	int32 NumberOfVerticalSamples = 5;
+	int32 NumberOfVerticalSamples = 2;
 
 	virtual UAnimationAsset* GetAnimationAsset() const override;
 	virtual bool IsLooping() const override;
@@ -803,11 +887,10 @@ public:
 };
 
 /**
-* Helper object for writing features into a float buffer according to a feature vector layout.
-* Keeps track of which features are present, allowing the feature vector to be built up piecemeal.
+* float buffer of features according to a UPoseSearchSchema layout.
 * FFeatureVectorBuilder is used to build search queries at runtime and for adding samples during search index construction.
 */
-USTRUCT(BlueprintType, Category = "Animation|Pose Search")
+USTRUCT()
 struct POSESEARCH_API FPoseSearchFeatureVectorBuilder
 {
 	GENERATED_BODY()
@@ -815,23 +898,15 @@ struct POSESEARCH_API FPoseSearchFeatureVectorBuilder
 public:
 	void Init(const UPoseSearchSchema* Schema);
 	void Reset();
-	void ResetFeatures();
 
 	const UPoseSearchSchema* GetSchema() const { return Schema.Get(); }
 
 	TArray<float>& EditValues() { return Values; }
 	TArrayView<const float> GetValues() const { return Values; }
 
-	void CopyFromSearchIndex(const FPoseSearchIndex& SearchIndex, int32 PoseIdx);
-
-	bool IsInitialized() const;
-	bool IsInitializedForSchema(const UPoseSearchSchema* Schema) const;
-	bool IsCompatible(const FPoseSearchFeatureVectorBuilder& OtherBuilder) const;
-
 private:
 	UPROPERTY(Transient)
 	TWeakObjectPtr<const UPoseSearchSchema> Schema = nullptr;
-
 	TArray<float> Values;
 };
 
@@ -997,7 +1072,7 @@ public:
 	virtual ~UPoseSearchDatabase();
 
 	const FPoseSearchIndex* GetSearchIndex() const;
-	const FPoseSearchIndex* GetSearchIndexSafe() const;
+	const FPoseSearchIndex* GetSearchIndexSafe(bool bVerboseLogging = true) const;
 	
 	bool GetSkipSearchIfPossible() const;
 	bool IsValidForIndexing() const;
