motion matching
- fix for parallel indexing (multiple copies of FBoneContainer to avoid concurrent access on mutable properties)
- add WIP PoseSearchDatabaseEditorCommands Detailed Features
- removed BuildSearchIndex on FDatabaseEditor::OnFinishedChangingSelectionProperties
- improved BlendParameterForSampleRanges calculationg for wrapping BS axes

diff --git a/Engine/Plugins/Experimental/Animation/PoseSearch/Source/Editor/Private/PoseSearchDatabaseEditor.cpp b/Engine/Plugins/Experimental/Animation/PoseSearch/Source/Editor/Private/PoseSearchDatabaseEditor.cpp
--- a/Engine/Plugins/Experimental/Animation/PoseSearch/Source/Editor/Private/PoseSearchDatabaseEditor.cpp
+++ b/Engine/Plugins/Experimental/Animation/PoseSearch/Source/Editor/Private/PoseSearchDatabaseEditor.cpp
@@ -472,7 +472,8 @@ namespace UE::PoseSearch
 
 	void FDatabaseEditor::OnFinishedChangingSelectionProperties(const FPropertyChangedEvent& PropertyChangedEvent)
 	{
-		ViewModel->BuildSearchIndex();
+		// @todo: BuildSearchIndex is already called, but it'd be better to perform it from here
+		//ViewModel->BuildSearchIndex();
 	}
 
 	void FDatabaseEditor::OnAssetTreeSelectionChanged(
diff --git a/Engine/Plugins/Experimental/Animation/PoseSearch/Source/Editor/Private/PoseSearchDatabaseEditorCommands.cpp b/Engine/Plugins/Experimental/Animation/PoseSearch/Source/Editor/Private/PoseSearchDatabaseEditorCommands.cpp
--- a/Engine/Plugins/Experimental/Animation/PoseSearch/Source/Editor/Private/PoseSearchDatabaseEditorCommands.cpp
+++ b/Engine/Plugins/Experimental/Animation/PoseSearch/Source/Editor/Private/PoseSearchDatabaseEditorCommands.cpp
@@ -16,6 +16,9 @@ namespace UE::PoseSearch
 		UI_COMMAND(
 			ShowPoseFeaturesAll, "All Features", "Show all features",
 			EUserInterfaceActionType::RadioButton, FInputChord());
+		UI_COMMAND(
+			ShowPoseFeaturesDetailed, "Detailed Features", "Show all features with mode details",
+			EUserInterfaceActionType::RadioButton, FInputChord());
 
 		UI_COMMAND(
 			ShowAnimationNone, "None", "Don't show animations",
diff --git a/Engine/Plugins/Experimental/Animation/PoseSearch/Source/Editor/Private/PoseSearchDatabasePreviewScene.cpp b/Engine/Plugins/Experimental/Animation/PoseSearch/Source/Editor/Private/PoseSearchDatabasePreviewScene.cpp
--- a/Engine/Plugins/Experimental/Animation/PoseSearch/Source/Editor/Private/PoseSearchDatabasePreviewScene.cpp
+++ b/Engine/Plugins/Experimental/Animation/PoseSearch/Source/Editor/Private/PoseSearchDatabasePreviewScene.cpp
@@ -60,8 +60,7 @@ namespace UE::PoseSearch
 		const FDatabaseViewModel* ViewModel = GetEditor()->GetViewModel();
 		const UPoseSearchDatabase* Database = ViewModel->GetPoseSearchDatabase();
 
-		if (Database && Database->IsValidForSearch() &&
-			ViewModel->IsPoseFeaturesDrawMode(EFeaturesDrawMode::All))
+		if (Database && Database->IsValidForSearch() && ViewModel->IsPoseFeaturesDrawMode(EFeaturesDrawMode::All | EFeaturesDrawMode::Detailed))
 		{
 			for (const FDatabasePreviewActor& PreviewActor : ViewModel->GetPreviewActors())
 			{
@@ -75,6 +74,11 @@ namespace UE::PoseSearch
 					DrawParams.PointSize = 5.f;
 					DrawParams.Mesh = PreviewActor.Mesh;
 
+					if (ViewModel->IsPoseFeaturesDrawMode(EFeaturesDrawMode::Detailed))
+					{
+						EnumAddFlags(DrawParams.Flags, UE::PoseSearch::EDebugDrawFlags::DrawBoneNames);
+					}
+
 					EnumAddFlags(DrawParams.Flags, UE::PoseSearch::EDebugDrawFlags::DrawFast);
 					DrawFeatureVector(DrawParams, PreviewActor.CurrentPoseIndex);
 				}
diff --git a/Engine/Plugins/Experimental/Animation/PoseSearch/Source/Editor/Private/PoseSearchDatabaseViewModel.cpp b/Engine/Plugins/Experimental/Animation/PoseSearch/Source/Editor/Private/PoseSearchDatabaseViewModel.cpp
--- a/Engine/Plugins/Experimental/Animation/PoseSearch/Source/Editor/Private/PoseSearchDatabaseViewModel.cpp
+++ b/Engine/Plugins/Experimental/Animation/PoseSearch/Source/Editor/Private/PoseSearchDatabaseViewModel.cpp
@@ -16,12 +16,54 @@
 #include "Modules/ModuleManager.h"
 #include "Camera/CameraComponent.h"
 #include "Components/CapsuleComponent.h"
-
+#include "GameFramework/PlayerController.h"
 
 namespace UE::PoseSearch
 {
 	constexpr float StepDeltaTime = 1.0f / 30.0f;
 
+	// FDatabasePreviewActor
+	bool FDatabasePreviewActor::IsValid() const
+	{
+		const bool bIsValid = Actor.IsValid() && Mesh.IsValid() && AnimInstance.IsValid();
+		return  bIsValid;
+	}
+
+	void FDatabasePreviewActor::Process()
+	{
+		if (Type == ESearchIndexAssetType::Sequence)
+		{
+			SequenceSampler.Process();
+		}
+		else if (Type == ESearchIndexAssetType::BlendSpace)
+		{
+			BlendSpaceSampler.Process();
+		}
+	}
+
+	const IAssetSampler* FDatabasePreviewActor::GetSampler() const
+	{
+		switch (Type)
+		{
+		case ESearchIndexAssetType::Sequence:
+			return &SequenceSampler;
+		case ESearchIndexAssetType::BlendSpace:
+			return &BlendSpaceSampler;
+		}
+		return nullptr;
+	}
+
+	float FDatabasePreviewActor::GetScaledTime(float Time) const
+	{
+		float ScaledTime = Time;
+		if (Type == ESearchIndexAssetType::BlendSpace)
+		{
+			ScaledTime = BlendSpaceSampler.GetPlayLength() > UE_KINDA_SMALL_NUMBER ? Time / BlendSpaceSampler.GetPlayLength() : 0.f;
+		}
+		return ScaledTime;
+	}
+
+	// FDatabaseViewModel
 	FDatabaseViewModel::FDatabaseViewModel()
 		: PoseSearchDatabase(nullptr)
 	{
@@ -87,6 +129,9 @@ namespace UE::PoseSearch
 		const FPoseSearchIndex* SearchIndex = PoseSearchDatabase->GetSearchIndexSafe(false);
 		if (SearchIndex)
 		{
+			FBoneContainer BoneContainer;
+			BoneContainer.InitializeTo(PoseSearchDatabase->Schema->BoneIndicesWithParents, FCurveEvaluationOption(false), *PoseSearchDatabase->Schema->Skeleton);
+
 			for (int32 IndexAssetIndex = 0; IndexAssetIndex < SearchIndex->Assets.Num(); ++IndexAssetIndex)
 			{
 				const FPoseSearchIndexAsset& IndexAsset = SearchIndex->Assets[IndexAssetIndex];
@@ -107,7 +152,7 @@ namespace UE::PoseSearch
 
 				if (bSpawn)
 				{
-					FDatabasePreviewActor PreviewActor = SpawnPreviewActor(IndexAssetIndex);
+					FDatabasePreviewActor PreviewActor = SpawnPreviewActor(IndexAssetIndex, BoneContainer);
 					if (PreviewActor.IsValid())
 					{
 						const UAnimationAsset* AnimationAsset = PreviewActor.AnimInstance->GetAnimationAsset();
@@ -121,6 +166,8 @@ namespace UE::PoseSearch
 				}
 			}
 
+			ParallelFor(PreviewActors.Num(), [this](int32 PreviewActorIndex) { PreviewActors[PreviewActorIndex].Process(); }, ParallelForFlags);
+
 			UpdatePreviewActors();
 		}
 	}
@@ -168,7 +215,7 @@ namespace UE::PoseSearch
 		PlayTime = MaxPreviewPlayLength;
 	}
 
-	FDatabasePreviewActor FDatabaseViewModel::SpawnPreviewActor(int32 IndexAssetIndex)
+	FDatabasePreviewActor FDatabaseViewModel::SpawnPreviewActor(int32 IndexAssetIndex, const FBoneContainer& BoneContainer)
 	{
 		FDatabasePreviewActor PreviewActor;
 
@@ -182,48 +229,42 @@ namespace UE::PoseSearch
 				{
 					FPoseSearchDatabaseSequence DatabaseSequence = PoseSearchDatabase->Sequences[IndexAsset.SourceAssetIdx];
 					PreviewAsset = DatabaseSequence.Sequence;
+					if (!PreviewAsset)
+					{
+						return PreviewActor;
+					}
 
 					FSequenceSampler::FInput Input;
 					Input.ExtrapolationParameters = PoseSearchDatabase->ExtrapolationParameters;
 					Input.Sequence = DatabaseSequence.Sequence;
 
 					PreviewActor.SequenceSampler.Init(Input);
-					PreviewActor.SequenceSampler.Process();
-
 				}
 				else if (IndexAsset.Type == ESearchIndexAssetType::BlendSpace)
 				{
 					FPoseSearchDatabaseBlendSpace DatabaseBlendSpace = PoseSearchDatabase->BlendSpaces[IndexAsset.SourceAssetIdx];
 					PreviewAsset = DatabaseBlendSpace.BlendSpace;
-
-					FAssetSamplingContext SamplingContext;
-					SamplingContext.Init(PoseSearchDatabase->Schema);
+					if (!PreviewAsset)
+					{
+						return PreviewActor;
+					}
 
 					FBlendSpaceSampler::FInput Input;
-					Input.BoneContainer = SamplingContext.BoneContainer;
+					Input.BoneContainer = BoneContainer;
 					Input.ExtrapolationParameters = PoseSearchDatabase->ExtrapolationParameters;
 					Input.BlendSpace = DatabaseBlendSpace.BlendSpace;
 					Input.BlendParameters = IndexAsset.BlendParameters;
 
 					PreviewActor.BlendSpaceSampler.Init(Input);
-					PreviewActor.BlendSpaceSampler.Process();
-				}
-
-				if (!PreviewAsset)
-				{
-					return PreviewActor;
 				}
 
+				PreviewActor.Type = IndexAsset.Type;
 				PreviewActor.IndexAssetIndex = IndexAssetIndex;
 				PreviewActor.CurrentPoseIndex = INDEX_NONE;
 
-				// todo: use preview when possible
-				UClass* PreviewClass = AActor::StaticClass();
-				const FTransform SpawnTransform = FTransform::Identity;
-
 				FActorSpawnParameters Params;
 				Params.SpawnCollisionHandlingOverride = ESpawnActorCollisionHandlingMethod::AdjustIfPossibleButAlwaysSpawn;
-				PreviewActor.Actor = GetWorld()->SpawnActor<AActor>(PreviewClass, SpawnTransform, Params);
+				PreviewActor.Actor = GetWorld()->SpawnActor<AActor>(AActor::StaticClass(), FTransform::Identity, Params);
 				PreviewActor.Actor->SetFlags(RF_Transient);
 
 				PreviewActor.Mesh = NewObject<UDebugSkelMeshComponent>(PreviewActor.Actor.Get());
@@ -248,11 +289,7 @@ namespace UE::PoseSearch
 					PreviewActor.Actor->SetRootComponent(PreviewActor.Mesh.Get());
 				}
 
-				UE_LOG(
-					LogPoseSearchEditor, Log, TEXT("Spawned preview Actor: %s at Loc: %s Rot: %s"),
-					*GetNameSafe(PreviewActor.Actor.Get()),
-					*SpawnTransform.GetLocation().ToString(),
-					*SpawnTransform.Rotator().ToString());
+				UE_LOG(LogPoseSearchEditor, Log, TEXT("Spawned preview Actor: %s"), *GetNameSafe(PreviewActor.Actor.Get()));
 			}
 		}
 		return PreviewActor;
@@ -293,46 +330,27 @@ namespace UE::PoseSearch
 				{
 					if (const UAnimationAsset* PreviewAsset = PreviewActor.AnimInstance->GetAnimationAsset())
 					{
-						float CurrentTime = 0.f;
-						float CurrentTimeScaled = 0.f;
-						FAnimationRuntime::AdvanceTime(false, PlayTime, CurrentTime, PreviewAsset->GetPlayLength());
-
-						bool bIsSequence = false;
-						bool bIsBlendSpace = false;
-						if (Cast<UAnimSequence>(PreviewAsset))
+						if (const IAssetSampler* Sampler = PreviewActor.GetSampler())
 						{
-							bIsSequence = true;
-							CurrentTimeScaled = CurrentTime;
-						}
-						else if (Cast<UBlendSpace>(PreviewAsset))
-						{
-							bIsBlendSpace = true;
-							CurrentTimeScaled = PreviewActor.BlendSpaceSampler.GetPlayLength() > UE_KINDA_SMALL_NUMBER ? CurrentTime / PreviewActor.BlendSpaceSampler.GetPlayLength() : 0.f;
-						}
-						
-						const float OldTime = PreviewActor.AnimInstance->GetCurrentTime();
-						if (!FMath::IsNearlyEqual(CurrentTimeScaled, OldTime, UE_KINDA_SMALL_NUMBER))
-						{
-							FTransform RootMotion = FTransform::Identity;
-							if (bIsSequence)
-							{
-								RootMotion = PreviewActor.SequenceSampler.ExtractRootTransform(CurrentTime);
-							}
-							else if (bIsBlendSpace)
-							{
-								RootMotion = PreviewActor.BlendSpaceSampler.ExtractRootTransform(CurrentTime);
-							}
+							float CurrentTime = 0.f;
+							FAnimationRuntime::AdvanceTime(false, PlayTime, CurrentTime, PreviewAsset->GetPlayLength());
 
-							if (PreviewActor.AnimInstance->GetMirrorDataTable())
+							const float CurrentScaledTime = PreviewActor.GetScaledTime(CurrentTime);
+							const float OldScaledTime = PreviewActor.AnimInstance->GetCurrentTime();
+							if (!FMath::IsNearlyEqual(CurrentScaledTime, OldScaledTime, UE_KINDA_SMALL_NUMBER))
 							{
-								RootMotion = MirrorRootMotion(RootMotion, PreviewActor.AnimInstance->GetMirrorDataTable());
-							}
-							PreviewActor.Actor->SetActorTransform(RootMotion);
+								FTransform RootMotion = Sampler->ExtractRootTransform(CurrentTime);
+								if (PreviewActor.AnimInstance->GetMirrorDataTable())
+								{
+									RootMotion = MirrorRootMotion(RootMotion, PreviewActor.AnimInstance->GetMirrorDataTable());
+								}
+								PreviewActor.Actor->SetActorTransform(RootMotion);
 
-							PreviewActor.AnimInstance->SetPosition(CurrentTimeScaled);
+								PreviewActor.AnimInstance->SetPosition(CurrentScaledTime);
 
-							const FPoseSearchIndexAsset& SearchIndexAsset = SearchIndex->Assets[PreviewActor.IndexAssetIndex];
-							PreviewActor.CurrentPoseIndex = PoseSearchDatabase->GetPoseIndexFromTime(CurrentTime, SearchIndexAsset);
+								const FPoseSearchIndexAsset& SearchIndexAsset = SearchIndex->Assets[PreviewActor.IndexAssetIndex];
+								PreviewActor.CurrentPoseIndex = PoseSearchDatabase->GetPoseIndexFromTime(CurrentTime, SearchIndexAsset);
+							}
 						}
 					}
 				}
@@ -376,7 +394,7 @@ namespace UE::PoseSearch
 
 	bool FDatabaseViewModel::IsPoseFeaturesDrawMode(EFeaturesDrawMode DrawMode) const
 	{
-		return PoseFeaturesDrawMode == DrawMode;
+		return EnumHasAnyFlags(PoseFeaturesDrawMode, DrawMode);
 	}
 
 	void FDatabaseViewModel::OnSetAnimationPreviewMode(EAnimationPreviewMode PreviewMode)
@@ -390,7 +408,7 @@ namespace UE::PoseSearch
 
 	bool FDatabaseViewModel::IsAnimationPreviewMode(EAnimationPreviewMode PreviewMode) const
 	{
-		return AnimationPreviewMode == PreviewMode;
+		return EnumHasAnyFlags(AnimationPreviewMode, PreviewMode);
 	}
 
 	void FDatabaseViewModel::AddSequenceToDatabase(UAnimSequence* AnimSequence)
diff --git a/Engine/Plugins/Experimental/Animation/PoseSearch/Source/Editor/Private/PoseSearchDebugger.cpp b/Engine/Plugins/Experimental/Animation/PoseSearch/Source/Editor/Private/PoseSearchDebugger.cpp
--- a/Engine/Plugins/Experimental/Animation/PoseSearch/Source/Editor/Private/PoseSearchDebugger.cpp
+++ b/Engine/Plugins/Experimental/Animation/PoseSearch/Source/Editor/Private/PoseSearchDebugger.cpp
@@ -2751,11 +2751,12 @@ const USkinnedMeshComponent* FDebuggerViewModel::GetMeshComponent() const
 
 void FDebuggerViewModel::FillCompactPoseAndComponentRefRotations()
 {
+	bool bResetMirrorBonesAndRotations = true;
 	const UPoseSearchDatabase* Database = GetCurrentDatabase();
 	if (Database)
 	{
 		UMirrorDataTable* MirrorDataTable = Database->Schema->MirrorDataTable;
-		if (MirrorDataTable != nullptr)
+		if (MirrorDataTable != nullptr && Skeletons[ActivePose].Component->RequiredBones.IsValid())
 		{
 			if (CompactPoseMirrorBones.Num() == 0 || ComponentSpaceRefRotations.Num() == 0)
 			{
@@ -2764,13 +2765,15 @@ void FDebuggerViewModel::FillCompactPoseAndComponentRefRotations()
 					CompactPoseMirrorBones,
 					ComponentSpaceRefRotations);
 			}
-
-			return;
+			bResetMirrorBonesAndRotations = false;
 		}
 	}
 
-	CompactPoseMirrorBones.Reset();
-	ComponentSpaceRefRotations.Reset();
+	if (bResetMirrorBonesAndRotations)
+	{
+		CompactPoseMirrorBones.Reset();
+		ComponentSpaceRefRotations.Reset();
+	}
 }
 
 void FDebuggerViewModel::PlaySelection(int32 PoseIdx, float Time)
diff --git a/Engine/Plugins/Experimental/Animation/PoseSearch/Source/Editor/Private/SPoseSearchDatabaseAssetList.cpp b/Engine/Plugins/Experimental/Animation/PoseSearch/Source/Editor/Private/SPoseSearchDatabaseAssetList.cpp
--- a/Engine/Plugins/Experimental/Animation/PoseSearch/Source/Editor/Private/SPoseSearchDatabaseAssetList.cpp
+++ b/Engine/Plugins/Experimental/Animation/PoseSearch/Source/Editor/Private/SPoseSearchDatabaseAssetList.cpp
@@ -924,6 +924,7 @@ namespace UE::PoseSearch
 			}
 		}
 
+		// @todo: investigate if we should call a TreeView->ClearSelection() before TreeView->SetItemSelection
 		TreeView->SetItemSelection(NewSelectedNodes, true, ESelectInfo::Direct);
 	}
 
diff --git a/Engine/Plugins/Experimental/Animation/PoseSearch/Source/Editor/Private/SPoseSearchDatabaseViewport.cpp b/Engine/Plugins/Experimental/Animation/PoseSearch/Source/Editor/Private/SPoseSearchDatabaseViewport.cpp
--- a/Engine/Plugins/Experimental/Animation/PoseSearch/Source/Editor/Private/SPoseSearchDatabaseViewport.cpp
+++ b/Engine/Plugins/Experimental/Animation/PoseSearch/Source/Editor/Private/SPoseSearchDatabaseViewport.cpp
@@ -59,6 +59,18 @@ namespace UE::PoseSearch
 				&FDatabaseViewModel::IsPoseFeaturesDrawMode,
 				EFeaturesDrawMode::All));
 
+		CommandList->MapAction(
+			Commands.ShowPoseFeaturesDetailed,
+			FExecuteAction::CreateSP(
+				ViewModelRef,
+				&FDatabaseViewModel::OnSetPoseFeaturesDrawMode,
+				EFeaturesDrawMode::Detailed),
+			FCanExecuteAction(),
+			FIsActionChecked::CreateSP(
+				ViewModelRef,
+				&FDatabaseViewModel::IsPoseFeaturesDrawMode,
+				EFeaturesDrawMode::Detailed));
+
 		CommandList->MapAction(
 			Commands.ShowAnimationNone,
 			FExecuteAction::CreateSP(
diff --git a/Engine/Plugins/Experimental/Animation/PoseSearch/Source/Editor/Private/SPoseSearchDatabaseViewportToolbar.cpp b/Engine/Plugins/Experimental/Animation/PoseSearch/Source/Editor/Private/SPoseSearchDatabaseViewportToolbar.cpp
--- a/Engine/Plugins/Experimental/Animation/PoseSearch/Source/Editor/Private/SPoseSearchDatabaseViewportToolbar.cpp
+++ b/Engine/Plugins/Experimental/Animation/PoseSearch/Source/Editor/Private/SPoseSearchDatabaseViewportToolbar.cpp
@@ -40,6 +40,7 @@ namespace UE::PoseSearch
 				{
 					SubMenuBuilder.AddMenuEntry(Commands.ShowPoseFeaturesNone);
 					SubMenuBuilder.AddMenuEntry(Commands.ShowPoseFeaturesAll);
+					SubMenuBuilder.AddMenuEntry(Commands.ShowPoseFeaturesDetailed);
 				}
 				SubMenuBuilder.EndSection();
 			})
diff --git a/Engine/Plugins/Experimental/Animation/PoseSearch/Source/Editor/Public/PoseSearchDatabaseEditorCommands.h b/Engine/Plugins/Experimental/Animation/PoseSearch/Source/Editor/Public/PoseSearchDatabaseEditorCommands.h
--- a/Engine/Plugins/Experimental/Animation/PoseSearch/Source/Editor/Public/PoseSearchDatabaseEditorCommands.h
+++ b/Engine/Plugins/Experimental/Animation/PoseSearch/Source/Editor/Public/PoseSearchDatabaseEditorCommands.h
@@ -28,6 +28,7 @@ namespace UE::PoseSearch
 
 		TSharedPtr<FUICommandInfo> ShowPoseFeaturesNone;
 		TSharedPtr<FUICommandInfo> ShowPoseFeaturesAll;
+		TSharedPtr<FUICommandInfo> ShowPoseFeaturesDetailed;
 
 		TSharedPtr<FUICommandInfo> ShowAnimationNone;
 		TSharedPtr<FUICommandInfo> ShowAnimationOriginalOnly;
diff --git a/Engine/Plugins/Experimental/Animation/PoseSearch/Source/Editor/Public/PoseSearchDatabaseViewModel.h b/Engine/Plugins/Experimental/Animation/PoseSearch/Source/Editor/Public/PoseSearchDatabaseViewModel.h
--- a/Engine/Plugins/Experimental/Animation/PoseSearch/Source/Editor/Public/PoseSearchDatabaseViewModel.h
+++ b/Engine/Plugins/Experimental/Animation/PoseSearch/Source/Editor/Public/PoseSearchDatabaseViewModel.h
@@ -25,8 +25,10 @@ namespace UE::PoseSearch
 	enum class EFeaturesDrawMode : uint8
 	{
 		None,
-		All
+		All,
+		Detailed
 	};
+	ENUM_CLASS_FLAGS(EFeaturesDrawMode);
 
 	enum class EAnimationPreviewMode : uint8
 	{
@@ -34,7 +36,7 @@ namespace UE::PoseSearch
 		OriginalOnly,
 		OriginalAndMirrored
 	};
-
+	ENUM_CLASS_FLAGS(EAnimationPreviewMode);
 
 	struct FDatabasePreviewActor
 	{
@@ -47,12 +49,12 @@ namespace UE::PoseSearch
 
 		FSequenceSampler SequenceSampler;
 		FBlendSpaceSampler BlendSpaceSampler;
+		ESearchIndexAssetType Type = ESearchIndexAssetType::Invalid;
 
-		bool IsValid()
-		{
-			const bool bIsValid = Actor.IsValid() && Mesh.IsValid() && AnimInstance.IsValid();
-			return  bIsValid;
-		}
+		bool IsValid() const;
+		void Process();
+		const IAssetSampler* GetSampler() const;
+		float GetScaledTime(float Time) const;
 	};
 
 	class FDatabaseViewModel : public TSharedFromThis<FDatabaseViewModel>, public FGCObject
@@ -66,9 +68,7 @@ namespace UE::PoseSearch
 		virtual void AddReferencedObjects(FReferenceCollector& Collector) override;
 		virtual FString GetReferencerName() const override { return TEXT("FPoseSearchDatabaseViewModel"); }
 
-		void Initialize(
-			UPoseSearchDatabase* InPoseSearchDatabase,
-			const TSharedRef<FDatabasePreviewScene>& InPreviewScene);
+		void Initialize(UPoseSearchDatabase* InPoseSearchDatabase, const TSharedRef<FDatabasePreviewScene>& InPreviewScene);
 
 		void RemovePreviewActors();
 		void ResetPreviewActors();
@@ -130,7 +130,7 @@ namespace UE::PoseSearch
 
 		/** Actors to be displayed in the preview viewport */
 		TArray<FDatabasePreviewActor> PreviewActors;
-
+		
 		/** From zero to the play length of the longest preview */
 		float MaxPreviewPlayLength = 0.0f;
 
@@ -148,7 +148,7 @@ namespace UE::PoseSearch
 
 		UObject* GetPlaybackContext() const;
 
-		FDatabasePreviewActor SpawnPreviewActor(int32 IndexAssetIndex);
+		FDatabasePreviewActor SpawnPreviewActor(int32 IndexAssetIndex, const FBoneContainer& BoneContainer);
 
 		void UpdatePreviewActors(bool bInTickPlayTime = false);
 
diff --git a/Engine/Plugins/Experimental/Animation/PoseSearch/Source/Runtime/Private/PoseSearch.cpp b/Engine/Plugins/Experimental/Animation/PoseSearch/Source/Runtime/Private/PoseSearch.cpp
--- a/Engine/Plugins/Experimental/Animation/PoseSearch/Source/Runtime/Private/PoseSearch.cpp
+++ b/Engine/Plugins/Experimental/Animation/PoseSearch/Source/Runtime/Private/PoseSearch.cpp
@@ -57,22 +57,6 @@ namespace UE::PoseSearch
 //////////////////////////////////////////////////////////////////////////
 // Constants and utilities
 
-// @todo: set UE_POSE_SEARCH_FORCE_SINGLE_THREAD to 0
-// Temporarily disable parallel indexing due to a bug
-// A single BoneContainer is being used while indexing assets, but some
-// const BoneContainer methods change mutable lookup tables.
-#define UE_POSE_SEARCH_FORCE_SINGLE_THREAD 1
-
-#ifndef UE_POSE_SEARCH_FORCE_SINGLE_THREAD
-#define UE_POSE_SEARCH_FORCE_SINGLE_THREAD 0
-#endif
-
-#if UE_POSE_SEARCH_FORCE_SINGLE_THREAD
-	constexpr EParallelForFlags ParallelForFlags = EParallelForFlags::ForceSingleThread;
-#else
-	constexpr EParallelForFlags ParallelForFlags = EParallelForFlags::None;
-#endif // UE_POSE_SEARCH_FORCE_SINGLE_THREAD
-
 static inline float ArraySum(TConstArrayView<float> View, int32 StartIndex, int32 Offset)
 {
 	float Sum = 0.f;
@@ -1156,46 +1140,47 @@ UAnimationAsset* FPoseSearchDatabaseBlendSpace::GetAnimationAsset() const
 
 bool FPoseSearchDatabaseBlendSpace::IsLooping() const
 {
+	check(BlendSpace);
+
 	return BlendSpace->bLoop;
 }
 
-void FPoseSearchDatabaseBlendSpace::GetBlendSpaceParameterSampleRanges(
-	int32& HorizontalBlendNum,
-	int32& VerticalBlendNum,
-	float& HorizontalBlendMin,
-	float& HorizontalBlendMax,
-	float& VerticalBlendMin,
-	float& VerticalBlendMax) const
+void FPoseSearchDatabaseBlendSpace::GetBlendSpaceParameterSampleRanges(int32& HorizontalBlendNum, int32& VerticalBlendNum) const
 {
+	check(BlendSpace);
+
 	HorizontalBlendNum = bUseGridForSampling ? BlendSpace->GetBlendParameter(0).GridNum + 1 : FMath::Max(NumberOfHorizontalSamples, 1);
-	VerticalBlendNum = bUseGridForSampling ? BlendSpace->GetBlendParameter(1).GridNum + 1 : FMath::Max(NumberOfVerticalSamples, 1);
+	VerticalBlendNum = BlendSpace->IsA<UBlendSpace1D>() ? 1 : bUseGridForSampling ? BlendSpace->GetBlendParameter(1).GridNum + 1 : FMath::Max(NumberOfVerticalSamples, 1);
 
 	check(HorizontalBlendNum >= 1 && VerticalBlendNum >= 1);
+}
+
+FVector FPoseSearchDatabaseBlendSpace::BlendParameterForSampleRanges(int32 HorizontalBlendIndex, int32 VerticalBlendIndex) const
+{
+	check(BlendSpace);
 
-	HorizontalBlendMin = BlendSpace->GetBlendParameter(0).Min;
-	HorizontalBlendMax = BlendSpace->GetBlendParameter(0).Max;
+	const bool bWrapInputOnHorizontalAxis = BlendSpace->GetBlendParameter(0).bWrapInput;
+	const bool bWrapInputOnVerticalAxis = BlendSpace->GetBlendParameter(1).bWrapInput;
 
-	VerticalBlendMin = BlendSpace->GetBlendParameter(1).Min;
-	VerticalBlendMax = BlendSpace->GetBlendParameter(1).Max;
+	int32 HorizontalBlendNum, VerticalBlendNum;
+	GetBlendSpaceParameterSampleRanges(HorizontalBlendNum, VerticalBlendNum);
 
-	if (BlendSpace->IsA<UBlendSpace1D>())
+	if (bWrapInputOnHorizontalAxis)
 	{
-		VerticalBlendNum = 1;
-		VerticalBlendMin = 0.0;
-		VerticalBlendMax = 0.0;
+		++HorizontalBlendNum;
 	}
-}
 
-static FVector BlendParameterForSampleRanges(
-	int32 HorizontalBlendIndex,
-	int32 VerticalBlendIndex,
-	int32 HorizontalBlendNum,
-	int32 VerticalBlendNum,
-	float HorizontalBlendMin,
-	float HorizontalBlendMax,
-	float VerticalBlendMin,
-	float VerticalBlendMax)
-{
+	if (bWrapInputOnVerticalAxis)
+	{
+		++VerticalBlendNum;
+	}
+
+	const float HorizontalBlendMin = BlendSpace->GetBlendParameter(0).Min;
+	const float HorizontalBlendMax = BlendSpace->GetBlendParameter(0).Max;
+
+	const float VerticalBlendMin = BlendSpace->GetBlendParameter(1).Min;
+	const float VerticalBlendMax = BlendSpace->GetBlendParameter(1).Max;
+
 	return FVector(
 		HorizontalBlendNum > 1 ? 
 			HorizontalBlendMin + (HorizontalBlendMax - HorizontalBlendMin) * 
@@ -1278,29 +1263,15 @@ bool UPoseSearchDatabase::TryInitSearchIndexAssets(FPoseSearchIndex& OutSearchIn
 			BlendSpace.MirrorOption == EPoseSearchMirrorOption::UnmirroredAndMirrored;
 
 		int32 HorizontalBlendNum, VerticalBlendNum;
-		float HorizontalBlendMin, HorizontalBlendMax, VerticalBlendMin, VerticalBlendMax;
-
-		BlendSpace.GetBlendSpaceParameterSampleRanges(
-			HorizontalBlendNum,
-			VerticalBlendNum,
-			HorizontalBlendMin,
-			HorizontalBlendMax,
-			VerticalBlendMin,
-			VerticalBlendMax);
+		BlendSpace.GetBlendSpaceParameterSampleRanges(HorizontalBlendNum, VerticalBlendNum);
 
+		const bool bWrapInputOnHorizontalAxis = BlendSpace.BlendSpace->GetBlendParameter(0).bWrapInput;
+		const bool bWrapInputOnVerticalAxis = BlendSpace.BlendSpace->GetBlendParameter(1).bWrapInput;
 		for (int32 HorizontalIndex = 0; HorizontalIndex < HorizontalBlendNum; HorizontalIndex++)
 		{
 			for (int32 VerticalIndex = 0; VerticalIndex < VerticalBlendNum; VerticalIndex++)
 			{
-				FVector BlendParameters = BlendParameterForSampleRanges(
-					HorizontalIndex,
-					VerticalIndex,
-					HorizontalBlendNum,
-					VerticalBlendNum,
-					HorizontalBlendMin,
-					HorizontalBlendMax,
-					VerticalBlendMin,
-					VerticalBlendMax);
+				const FVector BlendParameters = BlendSpace.BlendParameterForSampleRanges(HorizontalIndex, VerticalIndex);
 						
 				int32 TriangulationIndex = 0;
 				BlendSpace.BlendSpace->GetSamplesFromBlendInput(BlendParameters, BlendSamples, TriangulationIndex, true);
@@ -2704,30 +2675,21 @@ const FPoseSearchIndexAsset* FSearchResult::GetSearchIndexAsset(bool bMandatory)
 //////////////////////////////////////////////////////////////////////////
 // FAssetSamplerContext
 
-static void FillCompactPoseAndComponentRefRotations(FAssetSamplingContext& Context)
+void FAssetSamplingContext::Init(const UMirrorDataTable* InMirrorDataTable, const FBoneContainer& BoneContainer)
 {
-	if (Context.MirrorDataTable)
+	MirrorDataTable = InMirrorDataTable;
+
+	if (InMirrorDataTable)
 	{
-		const UMirrorDataTable* MirrorDataTablePtr = Context.MirrorDataTable.Get();
-		MirrorDataTablePtr->FillCompactPoseAndComponentRefRotations(
-			Context.BoneContainer,
-			Context.CompactPoseMirrorBones,
-			Context.ComponentSpaceRefRotations);
+		InMirrorDataTable->FillCompactPoseAndComponentRefRotations(BoneContainer, CompactPoseMirrorBones, ComponentSpaceRefRotations);
 	}
 	else
 	{
-		Context.CompactPoseMirrorBones.Reset();
-		Context.ComponentSpaceRefRotations.Reset();
+		CompactPoseMirrorBones.Reset();
+		ComponentSpaceRefRotations.Reset();
 	}
 }
 
-void FAssetSamplingContext::Init(const UPoseSearchSchema* Schema)
-{
-	MirrorDataTable = Schema->MirrorDataTable;
-	BoneContainer.InitializeTo(Schema->BoneIndicesWithParents, FCurveEvaluationOption(false), *Schema->Skeleton);
-	FillCompactPoseAndComponentRefRotations(*this);
-}
-
 FTransform FAssetSamplingContext::MirrorTransform(const FTransform& InTransform) const
 {
 	const EAxis::Type MirrorAxis = MirrorDataTable->MirrorAxis;
@@ -3550,7 +3512,7 @@ public:
 	} Output;
 
 	void Reset();
-	void Init(const FAssetIndexingContext& IndexingContext);
+	void Init(const FAssetIndexingContext& IndexingContext, const FBoneContainer& InBoneContainer);
 	bool Process();
 
 public: // IAssetIndexer
@@ -3563,6 +3525,7 @@ public: // IAssetIndexer
 	FTransform GetTransformAndCacheResults(float SampleTime, float OriginTime, int8 SchemaBoneIdx, bool& Clamped) override;
 
 private:
+	FBoneContainer BoneContainer;
 	FAssetIndexingContext IndexingContext;
 	TArray<FPoseSearchFeatureVectorBuilder> FeatureVectorBuilders;
 	FPoseSearchPoseMetadata Metadata;
@@ -3598,19 +3561,19 @@ void FAssetIndexer::Reset()
 	Output.AllFeaturesNotAdded.Reset();
 }
 
-void FAssetIndexer::Init(const FAssetIndexingContext& InIndexingContext)
+void FAssetIndexer::Init(const FAssetIndexingContext& InIndexingContext, const FBoneContainer& InBoneContainer)
 {
 	check(InIndexingContext.Schema);
 	check(InIndexingContext.Schema->IsValid());
 	check(InIndexingContext.MainSampler);
 
+	BoneContainer = InBoneContainer;
 	IndexingContext = InIndexingContext;
 
 	Reset();
 
 	Output.FirstIndexedSample = FMath::FloorToInt(IndexingContext.RequestedSamplingRange.Min * IndexingContext.Schema->SampleRate);
-	Output.LastIndexedSample = 
-		FMath::Max(0, FMath::CeilToInt(IndexingContext.RequestedSamplingRange.Max * IndexingContext.Schema->SampleRate));
+	Output.LastIndexedSample = FMath::Max(0, FMath::CeilToInt(IndexingContext.RequestedSamplingRange.Max * IndexingContext.Schema->SampleRate));
 	Output.NumIndexedPoses = Output.LastIndexedSample - Output.FirstIndexedSample + 1;
 	
 	Output.FeatureVectorTable.SetNumZeroed(IndexingContext.Schema->SchemaCardinality * Output.NumIndexedPoses);
@@ -3962,19 +3925,19 @@ FTransform FAssetIndexer::GetTransformAndCacheResults(float SampleTime, float Or
 		Entry->SampleTime = SampleTime;
 		Entry->OriginTime = OriginTime;
 
-		if (!SamplingContext->BoneContainer.IsValid())
+		if (!BoneContainer.IsValid())
 		{
 			UE_LOG(LogPoseSearch,
 				Warning, 
 				TEXT("Invalid BoneContainer encountered in FAssetIndexer::GetTransformAndCacheResults. Asset: %s. Schema: %s. BoneContainerAsset: %s. NumBoneIndices: %d"),
 				*GetNameSafe(IndexingContext.MainSampler->GetAsset()),
 				*GetNameSafe(IndexingContext.Schema),
-				*GetNameSafe(SamplingContext->BoneContainer.GetAsset()),
-				SamplingContext->BoneContainer.GetCompactPoseNumBones());
+				*GetNameSafe(BoneContainer.GetAsset()),
+				BoneContainer.GetCompactPoseNumBones());
 		}
 
-		Entry->Pose.SetBoneContainer(&SamplingContext->BoneContainer);
-		Entry->UnusedCurve.InitFrom(SamplingContext->BoneContainer);
+		Entry->Pose.SetBoneContainer(&BoneContainer);
+		Entry->UnusedCurve.InitFrom(BoneContainer);
 
 		IAssetIndexer::FSampleInfo Origin = GetSampleInfo(OriginTime);
 		IAssetIndexer::FSampleInfo Sample = GetSampleInfoRelative(SampleTime, Origin);
@@ -4005,7 +3968,7 @@ FTransform FAssetIndexer::GetTransformAndCacheResults(float SampleTime, float Or
 	}
 
 	const FBoneReference& BoneReference = IndexingContext.Schema->BoneReferences[SchemaBoneIdx];
-	FCompactPoseBoneIndex CompactBoneIndex = SamplingContext->BoneContainer.MakeCompactPoseIndex(FMeshPoseBoneIndex(BoneReference.BoneIndex));
+	FCompactPoseBoneIndex CompactBoneIndex = BoneContainer.MakeCompactPoseIndex(FMeshPoseBoneIndex(BoneReference.BoneIndex));
 
 	const FTransform BoneTransform = Entry->ComponentSpacePose.GetComponentSpaceTransform(CompactBoneIndex) * MirrorTransform(Entry->RootTransform);
 	Clamped = Entry->Clamped;
@@ -4360,8 +4323,10 @@ bool BuildIndex(const UAnimSequence* Sequence, UPoseSearchSequenceMetaData* Sequ
 		return false;
 	}
 
+	FBoneContainer BoneContainer;
+	BoneContainer.InitializeTo(SequenceMetaData->Schema->BoneIndicesWithParents, FCurveEvaluationOption(false), *SequenceMetaData->Schema->Skeleton);
 	FAssetSamplingContext SamplingContext;
-	SamplingContext.Init(SequenceMetaData->Schema);
+	SamplingContext.Init(SequenceMetaData->Schema->MirrorDataTable, BoneContainer);
 
 	FSequenceSampler Sampler;
 	FSequenceSampler::FInput SamplerInput;
@@ -4376,7 +4341,8 @@ bool BuildIndex(const UAnimSequence* Sequence, UPoseSearchSequenceMetaData* Sequ
 	IndexerContext.MainSampler = &Sampler;
 	IndexerContext.Schema = SequenceMetaData->Schema;
 	IndexerContext.RequestedSamplingRange = GetEffectiveSamplingRange(Sequence, SequenceMetaData->SamplingRange);
-	Indexer.Init(IndexerContext);
+
+	Indexer.Init(IndexerContext, BoneContainer);
 	if (!Indexer.Process())
 	{
 		return false;
@@ -4413,19 +4379,18 @@ struct FDatabaseIndexingContext
 
 	TArray<FAssetIndexer> Indexers;
 
-	void PrepareSamplers();
-	void PrepareIndexers();
+	void PrepareSamplers(const FBoneContainer& BoneContainer);
+	void PrepareIndexers(const FBoneContainer& BoneContainer);
 	bool IndexAssets();
 	void JoinIndex();
 	float CalculateMinCostAddend() const;
 };
 
-void FDatabaseIndexingContext::PrepareSamplers()
+void FDatabaseIndexingContext::PrepareSamplers(const FBoneContainer& BoneContainer)
 {
-	SamplingContext.Init(Database->Schema);
+	SamplingContext.Init(Database->Schema->MirrorDataTable, BoneContainer);
 
 	// Prepare samplers for all sequences
-
 	auto AddSequenceSampler = [&](const UAnimSequence* Sequence, bool bLoopable)
 	{
 		if (!SequenceSamplerMap.Contains(Sequence))
@@ -4457,43 +4422,20 @@ void FDatabaseIndexingContext::PrepareSamplers()
 			AddSequenceSampler(DbSequence.FollowUpSequence, DbSequence.FollowUpSequence->bLoop);
 		}
 	}
-
-	ParallelFor(
-		SequenceSamplers.Num(),
-		[this](int32 SamplerIdx) { SequenceSamplers[SamplerIdx].Process(); },
-		ParallelForFlags
-	);
-
+		
 	// Prepare samplers for all blend spaces
-
 	for (const FPoseSearchDatabaseBlendSpace& DbBlendSpace : Database->BlendSpaces)
 	{
 		if (DbBlendSpace.BlendSpace)
 		{
 			int32 HorizontalBlendNum, VerticalBlendNum;
-			float HorizontalBlendMin, HorizontalBlendMax, VerticalBlendMin, VerticalBlendMax;
-
-			DbBlendSpace.GetBlendSpaceParameterSampleRanges(
-				HorizontalBlendNum,
-				VerticalBlendNum,
-				HorizontalBlendMin,
-				HorizontalBlendMax,
-				VerticalBlendMin,
-				VerticalBlendMax);
+			DbBlendSpace.GetBlendSpaceParameterSampleRanges(HorizontalBlendNum, VerticalBlendNum);
 
 			for (int32 HorizontalIndex = 0; HorizontalIndex < HorizontalBlendNum; HorizontalIndex++)
 			{
 				for (int32 VerticalIndex = 0; VerticalIndex < VerticalBlendNum; VerticalIndex++)
 				{
-					FVector BlendParameters = BlendParameterForSampleRanges(
-						HorizontalIndex,
-						VerticalIndex,
-						HorizontalBlendNum,
-						VerticalBlendNum,
-						HorizontalBlendMin,
-						HorizontalBlendMax,
-						VerticalBlendMin,
-						VerticalBlendMax);
+					const FVector BlendParameters = DbBlendSpace.BlendParameterForSampleRanges(HorizontalIndex, VerticalIndex);
 
 					if (!BlendSpaceSamplerMap.Contains({ DbBlendSpace.BlendSpace, BlendParameters }))
 					{
@@ -4501,7 +4443,7 @@ void FDatabaseIndexingContext::PrepareSamplers()
 						BlendSpaceSamplerMap.Add({ DbBlendSpace.BlendSpace, BlendParameters }, BlendSpaceSamplerIdx);
 
 						FBlendSpaceSampler::FInput Input;
-						Input.BoneContainer = SamplingContext.BoneContainer;
+						Input.BoneContainer = BoneContainer;
 						Input.ExtrapolationParameters = Database->ExtrapolationParameters;
 						Input.BlendSpace = DbBlendSpace.BlendSpace;
 						Input.BlendParameters = BlendParameters;
@@ -4513,14 +4455,22 @@ void FDatabaseIndexingContext::PrepareSamplers()
 		}
 	}
 
-	ParallelFor(
-		BlendSpaceSamplers.Num(),
-		[this](int32 SamplerIdx) { BlendSpaceSamplers[SamplerIdx].Process(); },
-		ParallelForFlags
-	);
+	TArray<IAssetSampler*, TInlineAllocator<512>> AssetSampler;
+	AssetSampler.SetNumUninitialized(SequenceSamplers.Num() + BlendSpaceSamplers.Num());
+	
+	for (int i = 0; i < SequenceSamplers.Num(); ++i)
+	{
+		AssetSampler[i] = &SequenceSamplers[i];
+	}
+	for (int i = 0; i < BlendSpaceSamplers.Num(); ++i)
+	{
+		AssetSampler[i + SequenceSamplers.Num()] = &BlendSpaceSamplers[i];
+	}
+
+	ParallelFor(AssetSampler.Num(), [AssetSampler](int32 SamplerIdx) { AssetSampler[SamplerIdx]->Process(); }, ParallelForFlags);
 }
 
-void FDatabaseIndexingContext::PrepareIndexers()
+void FDatabaseIndexingContext::PrepareIndexers(const FBoneContainer& BoneContainer)
 {
 	Indexers.Reserve(Database->GetSearchIndex()->Assets.Num());
 
@@ -4535,6 +4485,7 @@ void FDatabaseIndexingContext::PrepareIndexers()
 	};
 
 	Indexers.Reserve(Database->GetSearchIndex()->Assets.Num());
+
 	for (int32 AssetIdx = 0; AssetIdx != SearchIndex->Assets.Num(); ++AssetIdx)
 	{
 		const FPoseSearchIndexAsset& SearchIndexAsset = SearchIndex->Assets[AssetIdx];
@@ -4564,18 +4515,14 @@ void FDatabaseIndexingContext::PrepareIndexers()
 		}
 
 		FAssetIndexer& Indexer = Indexers.AddDefaulted_GetRef();
-		Indexer.Init(IndexerContext);
+		Indexer.Init(IndexerContext, BoneContainer);
 	}
 }
 
 bool FDatabaseIndexingContext::IndexAssets()
 {
 	// Index asset data
-	ParallelFor(
-		Indexers.Num(),
-		[this](int32 AssetIdx) { Indexers[AssetIdx].Process(); },
-		ParallelForFlags
-	);
+	ParallelFor(Indexers.Num(), [this](int32 AssetIdx) { Indexers[AssetIdx].Process(); }, ParallelForFlags);
 	return true;
 }
 
@@ -4666,8 +4613,11 @@ bool BuildIndex(const UPoseSearchDatabase* Database, FPoseSearchIndex& OutSearch
 	DbIndexingContext.Database = Database;
 	DbIndexingContext.SearchIndex = &OutSearchIndex;
 
-	DbIndexingContext.PrepareSamplers();
-	DbIndexingContext.PrepareIndexers();
+	FBoneContainer BoneContainer;
+	BoneContainer.InitializeTo(Database->Schema->BoneIndicesWithParents, FCurveEvaluationOption(false), *Database->Schema->Skeleton);
+
+	DbIndexingContext.PrepareSamplers(BoneContainer);
+	DbIndexingContext.PrepareIndexers(BoneContainer);
 	bSuccess &= DbIndexingContext.IndexAssets();
 	DbIndexingContext.JoinIndex();
 
diff --git a/Engine/Plugins/Experimental/Animation/PoseSearch/Source/Runtime/Private/PoseSearchFeatureChannels.cpp b/Engine/Plugins/Experimental/Animation/PoseSearch/Source/Runtime/Private/PoseSearchFeatureChannels.cpp
--- a/Engine/Plugins/Experimental/Animation/PoseSearch/Source/Runtime/Private/PoseSearchFeatureChannels.cpp
+++ b/Engine/Plugins/Experimental/Animation/PoseSearch/Source/Runtime/Private/PoseSearchFeatureChannels.cpp
@@ -338,7 +338,6 @@ void UPoseSearchFeatureChannel_Position::IndexAsset(UE::PoseSearch::IAssetIndexe
 	using namespace UE::PoseSearch;
 
 	const FAssetIndexingContext& IndexingContext = Indexer.GetIndexingContext();
-	const FAssetSamplingContext* SamplingContext = IndexingContext.SamplingContext;
 
 	for (int32 SampleIdx = IndexingContext.BeginSampleIdx; SampleIdx != IndexingContext.EndSampleIdx; ++SampleIdx)
 	{
@@ -461,7 +460,6 @@ void UPoseSearchFeatureChannel_Heading::IndexAsset(UE::PoseSearch::IAssetIndexer
 	using namespace UE::PoseSearch;
 
 	const FAssetIndexingContext& IndexingContext = Indexer.GetIndexingContext();
-	const FAssetSamplingContext* SamplingContext = IndexingContext.SamplingContext;
 
 	for (int32 SampleIdx = IndexingContext.BeginSampleIdx; SampleIdx != IndexingContext.EndSampleIdx; ++SampleIdx)
 	{
diff --git a/Engine/Plugins/Experimental/Animation/PoseSearch/Source/Runtime/Public/PoseSearch/PoseSearch.h b/Engine/Plugins/Experimental/Animation/PoseSearch/Source/Runtime/Public/PoseSearch/PoseSearch.h
--- a/Engine/Plugins/Experimental/Animation/PoseSearch/Source/Runtime/Public/PoseSearch/PoseSearch.h
+++ b/Engine/Plugins/Experimental/Animation/PoseSearch/Source/Runtime/Public/PoseSearch/PoseSearch.h
@@ -42,6 +42,18 @@ class UAnimNotifyState_PoseSearchBase;
 class UPoseSearchSchema;
 class FBlake3;
 
+#ifndef UE_POSE_SEARCH_FORCE_SINGLE_THREAD
+#define UE_POSE_SEARCH_FORCE_SINGLE_THREAD 0
+#endif
+
+#if UE_POSE_SEARCH_FORCE_SINGLE_THREAD
+constexpr EParallelForFlags ParallelForFlags = EParallelForFlags::ForceSingleThread;
+#else
+constexpr EParallelForFlags ParallelForFlags = EParallelForFlags::None;
+#endif // UE_POSE_SEARCH_FORCE_SINGLE_THREAD
+
+
+
 namespace UE::PoseSearch
 {
 class FPoseHistory;
@@ -139,7 +151,6 @@ enum class EPoseSearchMirrorOption : int32
 	Invalid = Num UMETA(Hidden)
 };
 
-
 //////////////////////////////////////////////////////////////////////////
 // Common structs
 
@@ -187,10 +198,8 @@ struct POSESEARCH_API FAssetSamplingContext
 	// Time delta used for computing pose derivatives
 	static constexpr float FiniteDelta = 1 / 60.0f;
 
-	FBoneContainer BoneContainer;
-
 	// Mirror data table pointer copied from Schema for convenience
-	TObjectPtr<UMirrorDataTable> MirrorDataTable = nullptr;
+	TObjectPtr<const UMirrorDataTable> MirrorDataTable = nullptr;
 
 	// Compact pose format of Mirror Bone Map
 	TCustomBoneIndexArray<FCompactPoseBoneIndex, FCompactPoseBoneIndex> CompactPoseMirrorBones;
@@ -199,7 +208,7 @@ struct POSESEARCH_API FAssetSamplingContext
 	// Only initialized and used when a mirroring table is specified
 	TCustomBoneIndexArray<FQuat, FCompactPoseBoneIndex> ComponentSpaceRefRotations;
 
-	void Init(const UPoseSearchSchema* Schema);
+	void Init(const UMirrorDataTable* InMirrorDataTable, const FBoneContainer& BoneContainer);
 	FTransform MirrorTransform(const FTransform& Transform) const;
 };
 
@@ -239,6 +248,8 @@ public:
 	virtual void ExtractPoseSearchNotifyStates(float Time, TArray<UAnimNotifyState_PoseSearchBase*>& NotifyStates) const = 0;
 
 	virtual const UAnimationAsset* GetAsset() const = 0;
+
+	virtual void Process() = 0;
 };
 
 struct POSESEARCH_API FSequenceSampler : public IAssetSampler
@@ -252,15 +263,15 @@ public:
 	} Input;
 
 	void Init(const FInput& Input);
-	void Process();
+	virtual void Process() override;
 
-	float GetPlayLength() const override { return Input.Sequence->GetPlayLength(); }
-	bool IsLoopable() const override;
+	virtual float GetPlayLength() const override { return Input.Sequence->GetPlayLength(); }
+	virtual bool IsLoopable() const override;
 
-	float GetTimeFromRootDistance(float Distance) const override;
+	virtual float GetTimeFromRootDistance(float Distance) const override;
 
-	float GetTotalRootDistance() const override { return TotalRootDistance; }
-	FTransform GetTotalRootTransform() const override { return TotalRootTransform; }
+	virtual float GetTotalRootDistance() const override { return TotalRootDistance; }
+	virtual FTransform GetTotalRootTransform() const override { return TotalRootTransform; }
 
 	virtual void ExtractPose(const FAnimExtractContext& ExtractionCtx, FAnimationPoseData& OutAnimPoseData) const override;
 	virtual float ExtractRootDistance(float Time) const override;
@@ -290,16 +301,15 @@ public:
 	} Input;
 
 	void Init(const FInput& Input);
+	virtual void Process() override;
 
-	void Process();
+	virtual float GetPlayLength() const override { return PlayLength; }
+	virtual bool IsLoopable() const override;
 
-	float GetPlayLength() const override { return PlayLength; }
-	bool IsLoopable() const override;
+	virtual float GetTimeFromRootDistance(float Distance) const override;
 
-	float GetTimeFromRootDistance(float Distance) const override;
-
-	float GetTotalRootDistance() const override { return TotalRootDistance; }
-	FTransform GetTotalRootTransform() const override { return TotalRootTransform; }
+	virtual float GetTotalRootDistance() const override { return TotalRootDistance; }
+	virtual FTransform GetTotalRootTransform() const override { return TotalRootTransform; }
 
 	virtual void ExtractPose(const FAnimExtractContext& ExtractionCtx, FAnimationPoseData& OutAnimPoseData) const override;
 	virtual float ExtractRootDistance(float Time) const override;
@@ -827,15 +837,9 @@ struct POSESEARCH_API FPoseSearchDatabaseBlendSpace : public FPoseSearchDatabase
 #if WITH_EDITOR
 	virtual void BuildDerivedDataKey(UE::PoseSearch::FDerivedDataKeyBuilder& KeyBuilder) override;
 #endif
-public:
 
-	void GetBlendSpaceParameterSampleRanges(
-		int32& HorizontalBlendNum,
-		int32& VerticalBlendNum,
-		float& HorizontalBlendMin,
-		float& HorizontalBlendMax,
-		float& VerticalBlendMin,
-		float& VerticalBlendMax) const;
+	void GetBlendSpaceParameterSampleRanges(int32& HorizontalBlendNum, int32& VerticalBlendNum) const;
+	FVector BlendParameterForSampleRanges(int32 HorizontalBlendIndex, int32 VerticalBlendIndex) const;
 };
 
 USTRUCT()
