motion matching - fix for improperly setup schema (no channels)

diff --git a/Engine/Plugins/Experimental/Animation/PoseSearch/Source/Runtime/Private/PoseSearch.cpp b/Engine/Plugins/Experimental/Animation/PoseSearch/Source/Runtime/Private/PoseSearch.cpp
--- a/Engine/Plugins/Experimental/Animation/PoseSearch/Source/Runtime/Private/PoseSearch.cpp
+++ b/Engine/Plugins/Experimental/Animation/PoseSearch/Source/Runtime/Private/PoseSearch.cpp
@@ -540,7 +540,7 @@ bool FPoseSearchIndex::IsEmpty() const
 
 TConstArrayView<float> FPoseSearchIndex::GetPoseValues(int32 PoseIdx) const
 {
-	check(PoseIdx >= 0 && PoseIdx < NumPoses && Schema && Schema->SchemaCardinality >= 0);
+	check(PoseIdx >= 0 && PoseIdx < NumPoses && Schema && Schema->SchemaCardinality > 0);
 	const int32 ValueOffset = PoseIdx * Schema->SchemaCardinality;
 	return MakeArrayView(&Values[ValueOffset], Schema->SchemaCardinality);
 }
@@ -2507,7 +2507,7 @@ bool FDebugDrawParams::CanDraw() const
 		return false;
 	}
 
-	return SearchIndex->IsValid() && !SearchIndex->IsEmpty();
+	return SearchIndex->IsValid() && !SearchIndex->IsEmpty() && SearchIndex->Schema->SchemaCardinality > 0;
 }
 
 FColor FDebugDrawParams::GetColor(int32 ColorPreset) const
@@ -3611,15 +3611,18 @@ bool FAssetIndexer::Process()
 	// @todo: this step can be avoided now since Schema.SchemaCardinality is known and the data can be preallocated before Channel->IndexAsset
 	
 	// Merge spans of feature vectors into contiguous buffer
-	for (int32 VectorIdx = 0; VectorIdx != NumSamplesInRange; ++VectorIdx)
+	if (IndexingContext.Schema->SchemaCardinality > 0)
 	{
-		const int32 SampleIdx = VectorIdx + IndexingContext.BeginSampleIdx;
-		const int32 PoseIdx = SampleIdx - Output.FirstIndexedSample;
-		const int32 FirstValueIdx = PoseIdx * IndexingContext.Schema->SchemaCardinality;
-		TArrayView<float> WriteValues = MakeArrayView(&Output.FeatureVectorTable[FirstValueIdx], IndexingContext.Schema->SchemaCardinality);
-		TArrayView<const float> ReadValues = FeatureVectorBuilders[VectorIdx].GetValues();
-		check(WriteValues.Num() == ReadValues.Num());
-		FMemory::Memcpy(WriteValues.GetData(), ReadValues.GetData(), WriteValues.Num() * WriteValues.GetTypeSize());
+		for (int32 VectorIdx = 0; VectorIdx != NumSamplesInRange; ++VectorIdx)
+		{
+			const int32 SampleIdx = VectorIdx + IndexingContext.BeginSampleIdx;
+			const int32 PoseIdx = SampleIdx - Output.FirstIndexedSample;
+			const int32 FirstValueIdx = PoseIdx * IndexingContext.Schema->SchemaCardinality;
+			TArrayView<float> WriteValues = MakeArrayView(&Output.FeatureVectorTable[FirstValueIdx], IndexingContext.Schema->SchemaCardinality);
+			TArrayView<const float> ReadValues = FeatureVectorBuilders[VectorIdx].GetValues();
+			check(WriteValues.Num() == ReadValues.Num());
+			FMemory::Memcpy(WriteValues.GetData(), ReadValues.GetData(), WriteValues.Num() * WriteValues.GetTypeSize());
+		}
 	}
 
 	// Generate pose metadata
@@ -4130,95 +4133,106 @@ static void PreprocessSearchIndexWeights(FPoseSearchIndex& SearchIndex, const UP
 #endif // WITH_EDITORONLY_DATA
 }
 
-// it calcualtes Mean and PCAProjectionMatrix
+// it calculates Mean, PCAValues, and PCAProjectionMatrix
 static void PreprocessSearchIndexPCAData(FPoseSearchIndex& SearchIndex, const UPoseSearchDatabase* Database)
 {
 	// binding SearchIndex.Values and SearchIndex.PCAValues Eigen row major matrix maps
-	const int32 NumDimensions = Database->Schema->SchemaCardinality;
+	check(Database && Database->Schema);
 	const int32 NumPoses = SearchIndex.NumPoses;
-
-	const RowMajorVectorMapConst MapWeightsSqrt(SearchIndex.WeightsSqrt.GetData(), 1, NumDimensions);
-	const RowMajorMatrixMapConst MapValues(SearchIndex.Values.GetData(), NumPoses, NumDimensions);
-	const RowMajorMatrix WeightedValues = MapValues.array().rowwise() * MapWeightsSqrt.array();
 	const uint32 NumberOfPrincipalComponents = Database->GetNumberOfPrincipalComponents();
+	const int32 NumDimensions = Database->Schema->SchemaCardinality;
 
 	SearchIndex.PCAValues.Reset();
+	SearchIndex.Mean.Reset();
+	SearchIndex.PCAProjectionMatrix.Reset();
+
 	SearchIndex.PCAValues.AddZeroed(NumPoses * NumberOfPrincipalComponents);
-	RowMajorMatrixMap MapPCAValues(SearchIndex.PCAValues.GetData(), NumPoses, NumberOfPrincipalComponents);
+	SearchIndex.Mean.AddZeroed(NumDimensions);
+	SearchIndex.PCAProjectionMatrix.AddZeroed(NumDimensions * NumberOfPrincipalComponents);
+
+#if WITH_EDITORONLY_DATA
+	SearchIndex.PCAExplainedVariance = 0.f;
+#endif
 
-	// calculating the mean
-	SearchIndex.Mean.SetNumZeroed(NumDimensions);
-	RowMajorVectorMap Mean(SearchIndex.Mean.GetData(), 1, NumDimensions);
-	Mean = WeightedValues.colwise().mean();
+	if (NumDimensions > 0)
+	{
+		const RowMajorVectorMapConst MapWeightsSqrt(SearchIndex.WeightsSqrt.GetData(), 1, NumDimensions);
+		const RowMajorMatrixMapConst MapValues(SearchIndex.Values.GetData(), NumPoses, NumDimensions);
+		const RowMajorMatrix WeightedValues = MapValues.array().rowwise() * MapWeightsSqrt.array();
+		RowMajorMatrixMap MapPCAValues(SearchIndex.PCAValues.GetData(), NumPoses, NumberOfPrincipalComponents);
 
-	// use the mean to center the data points
-	const RowMajorMatrix CenteredValues = WeightedValues.rowwise() - Mean;
+		// calculating the mean
+		RowMajorVectorMap Mean(SearchIndex.Mean.GetData(), 1, NumDimensions);
+		Mean = WeightedValues.colwise().mean();
 
-	// estimating the covariance matrix (with dimensionality of NumDimensions, NumDimensions)
-	// formula: https://en.wikipedia.org/wiki/Covariance_matrix#Estimation
-	// details: https://en.wikipedia.org/wiki/Estimation_of_covariance_matrices
-	const ColMajorMatrix CovariantMatrix = (CenteredValues.transpose() * CenteredValues) / float(NumPoses - 1);
-	const Eigen::SelfAdjointEigenSolver<ColMajorMatrix> EigenSolver(CovariantMatrix);
+		// use the mean to center the data points
+		const RowMajorMatrix CenteredValues = WeightedValues.rowwise() - Mean;
 
-	check(EigenSolver.info() == Eigen::Success);
+		// estimating the covariance matrix (with dimensionality of NumDimensions, NumDimensions)
+		// formula: https://en.wikipedia.org/wiki/Covariance_matrix#Estimation
+		// details: https://en.wikipedia.org/wiki/Estimation_of_covariance_matrices
+		const ColMajorMatrix CovariantMatrix = (CenteredValues.transpose() * CenteredValues) / float(NumPoses - 1);
+		const Eigen::SelfAdjointEigenSolver<ColMajorMatrix> EigenSolver(CovariantMatrix);
 
-	// validating EigenSolver results
-	const ColMajorMatrix EigenVectors = EigenSolver.eigenvectors().real();
+		check(EigenSolver.info() == Eigen::Success);
 
-	if (Database->PoseSearchMode == EPoseSearchMode::PCAKDTree_Validate && NumberOfPrincipalComponents == NumDimensions)
-	{
-		const RowMajorVector ReciprocalWeightsSqrt = MapWeightsSqrt.cwiseInverse();
-		const RowMajorMatrix ProjectedValues = CenteredValues * EigenVectors;
-		for (Eigen::Index RowIndex = 0; RowIndex < MapValues.rows(); ++RowIndex)
+		// validating EigenSolver results
+		const ColMajorMatrix EigenVectors = EigenSolver.eigenvectors().real();
+
+		if (Database->PoseSearchMode == EPoseSearchMode::PCAKDTree_Validate && NumberOfPrincipalComponents == NumDimensions)
 		{
-			const RowMajorVector WeightedReconstructedPoint = ProjectedValues.row(RowIndex) * EigenVectors.transpose() + Mean;
-			const RowMajorVector ReconstructedPoint = WeightedReconstructedPoint.array() * ReciprocalWeightsSqrt.array();
-			const float Error = (ReconstructedPoint - MapValues.row(RowIndex)).squaredNorm();
-			check(Error < UE_KINDA_SMALL_NUMBER);
+			const RowMajorVector ReciprocalWeightsSqrt = MapWeightsSqrt.cwiseInverse();
+			const RowMajorMatrix ProjectedValues = CenteredValues * EigenVectors;
+			for (Eigen::Index RowIndex = 0; RowIndex < MapValues.rows(); ++RowIndex)
+			{
+				const RowMajorVector WeightedReconstructedPoint = ProjectedValues.row(RowIndex) * EigenVectors.transpose() + Mean;
+				const RowMajorVector ReconstructedPoint = WeightedReconstructedPoint.array() * ReciprocalWeightsSqrt.array();
+				const float Error = (ReconstructedPoint - MapValues.row(RowIndex)).squaredNorm();
+				check(Error < UE_KINDA_SMALL_NUMBER);
+			}
 		}
-	}
 
-	// sorting EigenVectors by EigenValues, so we pick the most significant ones to compose our PCA projection matrix.
-	const RowMajorVector EigenValues = EigenSolver.eigenvalues().real();
-	TArray<size_t> Indexer;
-	Indexer.Reserve(NumDimensions);
-	for (size_t DimensionIndex = 0; DimensionIndex < NumDimensions; ++DimensionIndex)
-	{
-		Indexer.Push(DimensionIndex);
-	}
-	Indexer.Sort([&EigenValues](size_t a, size_t b)
-	{
-		return EigenValues[a] > EigenValues[b];
-	});
+		// sorting EigenVectors by EigenValues, so we pick the most significant ones to compose our PCA projection matrix.
+		const RowMajorVector EigenValues = EigenSolver.eigenvalues().real();
+		TArray<size_t> Indexer;
+		Indexer.Reserve(NumDimensions);
+		for (size_t DimensionIndex = 0; DimensionIndex < NumDimensions; ++DimensionIndex)
+		{
+			Indexer.Push(DimensionIndex);
+		}
+		Indexer.Sort([&EigenValues](size_t a, size_t b)
+		{
+			return EigenValues[a] > EigenValues[b];
+		});
 
-	// composing the PCA projection matrix with the PCANumComponents most significant EigenVectors
-	SearchIndex.PCAProjectionMatrix.SetNumZeroed(NumDimensions * NumberOfPrincipalComponents);
-	ColMajorMatrixMap PCAProjectionMatrix(SearchIndex.PCAProjectionMatrix.GetData(), NumDimensions, NumberOfPrincipalComponents);
-	float AccumulatedVariance = 0.f;
-	for (size_t PCAComponentIndex = 0; PCAComponentIndex < NumberOfPrincipalComponents; ++PCAComponentIndex)
-	{
-		PCAProjectionMatrix.col(PCAComponentIndex) = EigenVectors.col(Indexer[PCAComponentIndex]);
-		AccumulatedVariance += EigenValues[Indexer[PCAComponentIndex]];
-	}
+		// composing the PCA projection matrix with the PCANumComponents most significant EigenVectors
+		ColMajorMatrixMap PCAProjectionMatrix(SearchIndex.PCAProjectionMatrix.GetData(), NumDimensions, NumberOfPrincipalComponents);
+		float AccumulatedVariance = 0.f;
+		for (size_t PCAComponentIndex = 0; PCAComponentIndex < NumberOfPrincipalComponents; ++PCAComponentIndex)
+		{
+			PCAProjectionMatrix.col(PCAComponentIndex) = EigenVectors.col(Indexer[PCAComponentIndex]);
+			AccumulatedVariance += EigenValues[Indexer[PCAComponentIndex]];
+		}
 
 #if WITH_EDITORONLY_DATA
-	// calculating the total variance knowing that eigen values measure variance along the principal components:
-	const float TotalVariance = EigenValues.sum();
-	// and explained variance as ratio between AccumulatedVariance and TotalVariance: https://ro-che.info/articles/2017-12-11-pca-explained-variance
-	SearchIndex.PCAExplainedVariance = TotalVariance > UE_KINDA_SMALL_NUMBER ? AccumulatedVariance / TotalVariance : 0.f;
+		// calculating the total variance knowing that eigen values measure variance along the principal components:
+		const float TotalVariance = EigenValues.sum();
+		// and explained variance as ratio between AccumulatedVariance and TotalVariance: https://ro-che.info/articles/2017-12-11-pca-explained-variance
+		SearchIndex.PCAExplainedVariance = TotalVariance > UE_KINDA_SMALL_NUMBER ? AccumulatedVariance / TotalVariance : 0.f;
 #endif // WITH_EDITORONLY_DATA
 
-	MapPCAValues = CenteredValues * PCAProjectionMatrix;
+		MapPCAValues = CenteredValues * PCAProjectionMatrix;
 
-	if (Database->PoseSearchMode == EPoseSearchMode::PCAKDTree_Validate && NumberOfPrincipalComponents == NumDimensions)
-	{
-		const RowMajorVector ReciprocalWeightsSqrt = MapWeightsSqrt.cwiseInverse();
-		for (Eigen::Index RowIndex = 0; RowIndex < MapValues.rows(); ++RowIndex)
+		if (Database->PoseSearchMode == EPoseSearchMode::PCAKDTree_Validate && NumberOfPrincipalComponents == NumDimensions)
 		{
-			const RowMajorVector WeightedReconstructedValues = MapPCAValues.row(RowIndex) * PCAProjectionMatrix.transpose() + Mean;
-			const RowMajorVector ReconstructedValues = WeightedReconstructedValues.array() * ReciprocalWeightsSqrt.array();
-			const float Error = (ReconstructedValues - MapValues.row(RowIndex)).squaredNorm();
-			check(Error < UE_KINDA_SMALL_NUMBER);
+			const RowMajorVector ReciprocalWeightsSqrt = MapWeightsSqrt.cwiseInverse();
+			for (Eigen::Index RowIndex = 0; RowIndex < MapValues.rows(); ++RowIndex)
+			{
+				const RowMajorVector WeightedReconstructedValues = MapPCAValues.row(RowIndex) * PCAProjectionMatrix.transpose() + Mean;
+				const RowMajorVector ReconstructedValues = WeightedReconstructedValues.array() * ReciprocalWeightsSqrt.array();
+				const float Error = (ReconstructedValues - MapValues.row(RowIndex)).squaredNorm();
+				check(Error < UE_KINDA_SMALL_NUMBER);
+			}
 		}
 	}
 }
diff --git a/Engine/Plugins/Experimental/Animation/PoseSearch/Source/Runtime/Private/PoseSearchDerivedData.cpp b/Engine/Plugins/Experimental/Animation/PoseSearch/Source/Runtime/Private/PoseSearchDerivedData.cpp
--- a/Engine/Plugins/Experimental/Animation/PoseSearch/Source/Runtime/Private/PoseSearchDerivedData.cpp
+++ b/Engine/Plugins/Experimental/Animation/PoseSearch/Source/Runtime/Private/PoseSearchDerivedData.cpp
@@ -307,9 +307,6 @@ namespace UE::PoseSearch
 		Ar << Index.Mean;
 		Ar << Index.PCAProjectionMatrix;
 
-		check(Index.PCAProjectionMatrix.Num() > 0 && Index.Mean.Num() > 0);
-		check(Index.PCAProjectionMatrix.Num() % Index.Mean.Num() == 0);
-		const int NumberOfPrincipalComponents = Index.PCAProjectionMatrix.Num() / Index.Mean.Num();
 		Serialize(Ar, Index.KDTree, Index.PCAValues.GetData());
 
 		return Ar;
diff --git a/Engine/Plugins/Experimental/Animation/PoseSearch/Source/Runtime/Public/PoseSearch/KDTree.h b/Engine/Plugins/Experimental/Animation/PoseSearch/Source/Runtime/Public/PoseSearch/KDTree.h
--- a/Engine/Plugins/Experimental/Animation/PoseSearch/Source/Runtime/Public/PoseSearch/KDTree.h
+++ b/Engine/Plugins/Experimental/Animation/PoseSearch/Source/Runtime/Public/PoseSearch/KDTree.h
@@ -17,7 +17,7 @@ struct POSESEARCH_API FKDTree
 	struct DataSource
 	{
 		DataSource(int32 pointCount, int32 pointDim, const float* data)
-		: PointCount(pointCount)
+		: PointCount(pointDim > 0 ? pointCount : 0)
 		, PointDim(pointDim)
 		, Data(data)
 		{
