motion matching - fix for users not be able to add Channels to an MM Database Config Schema

diff --git a/Engine/Plugins/Experimental/Animation/PoseSearch/Source/Runtime/Private/PoseSearch.cpp b/Engine/Plugins/Experimental/Animation/PoseSearch/Source/Runtime/Private/PoseSearch.cpp
--- a/Engine/Plugins/Experimental/Animation/PoseSearch/Source/Runtime/Private/PoseSearch.cpp
+++ b/Engine/Plugins/Experimental/Animation/PoseSearch/Source/Runtime/Private/PoseSearch.cpp
@@ -349,21 +349,25 @@ int32 FSchemaInitializer::AddBoneReference(const FBoneReference& BoneReference)
 
 } // namespace UE::PoseSearch
 
-void UPoseSearchSchema::Finalize()
+void UPoseSearchSchema::Finalize(bool bRemoveEmptyChannels)
 {
 	using namespace UE::PoseSearch;
 
-	// Discard null channels
-	Channels.RemoveAll([](TObjectPtr<UPoseSearchFeatureChannel>& Channel) { return !Channel; });
+	if (bRemoveEmptyChannels)
+	{
+		Channels.RemoveAll([](TObjectPtr<UPoseSearchFeatureChannel>& Channel) { return !Channel; });
+	}
 
 	BoneReferences.Reset();
 
 	FSchemaInitializer Initializer;
 	for (int32 ChannelIdx = 0; ChannelIdx != Channels.Num(); ++ChannelIdx)
 	{
-		TObjectPtr<UPoseSearchFeatureChannel>& Channel = Channels[ChannelIdx];
-		Initializer.CurrentChannelIdx = ChannelIdx;
-		Channel->InitializeSchema(Initializer);
+		if (Channels[ChannelIdx].Get())
+		{
+			Initializer.CurrentChannelIdx = ChannelIdx;
+			Channels[ChannelIdx]->InitializeSchema(Initializer);
+		}
 	}
 
 	SchemaCardinality = Initializer.GetCurrentChannelDataOffset();
@@ -389,7 +393,7 @@ void UPoseSearchSchema::PostLoad()
 #if WITH_EDITOR
 void UPoseSearchSchema::PostEditChangeProperty(FPropertyChangedEvent& PropertyChangedEvent)
 {
-	Finalize();
+	Finalize(false);
 	Super::PostEditChangeProperty(PropertyChangedEvent);
 }
 
diff --git a/Engine/Plugins/Experimental/Animation/PoseSearch/Source/Runtime/Public/PoseSearch/PoseSearch.h b/Engine/Plugins/Experimental/Animation/PoseSearch/Source/Runtime/Public/PoseSearch/PoseSearch.h
--- a/Engine/Plugins/Experimental/Animation/PoseSearch/Source/Runtime/Public/PoseSearch/PoseSearch.h
+++ b/Engine/Plugins/Experimental/Animation/PoseSearch/Source/Runtime/Public/PoseSearch/PoseSearch.h
@@ -590,7 +590,7 @@ public:
 #endif
 
 private:
-	void Finalize();
+	void Finalize(bool bRemoveEmptyChannels = true);
 	void ResolveBoneReferences();
 
 public:
