motion matching
- introducing IPoseFilter to channels to be able to perform filtering / pruning after candidate selection
- UPoseSearchFeatureChannel_FilterCrashingLegs prototype to minimize crashing legs while locomoting
- tuned some MM default values (KDtree on by default, 30Hz database sampling rate)
- small code refactoring
- less allocations while indexing the db

diff --git a/Engine/Plugins/Experimental/Animation/PoseSearch/Source/Editor/Private/PoseSearchDebugger.cpp b/Engine/Plugins/Experimental/Animation/PoseSearch/Source/Editor/Private/PoseSearchDebugger.cpp
--- a/Engine/Plugins/Experimental/Animation/PoseSearch/Source/Editor/Private/PoseSearchDebugger.cpp
+++ b/Engine/Plugins/Experimental/Animation/PoseSearch/Source/Editor/Private/PoseSearchDebugger.cpp
@@ -663,6 +663,12 @@ namespace DebuggerDatabaseColumns
 					TextBuilder.AppendLine(LOCTEXT("DiscardedBy_BlockTransition", "BlockTransition"));
 				}
 
+				if (EnumHasAnyFlags(Row->PoseCandidateFlags, EPoseCandidateFlags::DiscardedBy_PoseFilter))
+				{
+					AddDelimiter(TextBuilder, bNeedDelimiter);
+					TextBuilder.AppendLine(LOCTEXT("DiscardedBy_PoseFilter", "PoseFilter"));
+				}
+
 				return TextBuilder.ToText();
 			}
 
diff --git a/Engine/Plugins/Experimental/Animation/PoseSearch/Source/Runtime/Private/PoseSearch.cpp b/Engine/Plugins/Experimental/Animation/PoseSearch/Source/Runtime/Private/PoseSearch.cpp
--- a/Engine/Plugins/Experimental/Animation/PoseSearch/Source/Runtime/Private/PoseSearch.cpp
+++ b/Engine/Plugins/Experimental/Animation/PoseSearch/Source/Runtime/Private/PoseSearch.cpp
@@ -168,7 +168,7 @@ FORCEINLINE auto LowerBound(IteratorType First, IteratorType Last, const ValueTy
 	return LowerBound(First, Last, Value, FIdentityFunctor(), TLess<>());
 }
 
-static TArrayView<size_t> PopulateNonSelectableIdx(TArrayView<size_t> NonSelectableIdxBuffer, FSearchContext& SearchContext, const UPoseSearchDatabase* Database, TArrayView<const float> QueryValues)
+static TArrayView<size_t> PopulateNonSelectableIdx(TArrayView<size_t> NonSelectableIdxBuffer, FSearchContext& SearchContext, const UPoseSearchDatabase* Database, TConstArrayView<float> QueryValues)
 {
 	check(Database);
 	const FPoseSearchIndex* SearchIndex = Database->GetSearchIndex();
@@ -307,6 +307,99 @@ static TArrayView<size_t> PopulateNonSelectableIdx(TArrayView<size_t> NonSelecta
 	return NonSelectableIdx;
 }
 
+struct FPoseFilters
+{
+	FPoseFilters(const UPoseSearchSchema* Schema, TConstArrayView<size_t> NonSelectableIdx, EPoseSearchPoseFlags OverallFlags)
+	{
+		NonSelectableIdxPoseFilter.NonSelectableIdx = NonSelectableIdx;
+
+		if (EnumHasAnyFlags(OverallFlags, EPoseSearchPoseFlags::BlockTransition))
+		{
+			AllPoseFilters.Add(&BlockTransitionPoseFilter);
+		}
+
+		if (NonSelectableIdxPoseFilter.IsPoseFilterActive())
+		{
+			AllPoseFilters.Add(&NonSelectableIdxPoseFilter);
+		}
+
+		for (const IPoseFilter* ChannelPoseFilter : Schema->Channels)
+		{
+			if (ChannelPoseFilter->IsPoseFilterActive())
+			{
+				AllPoseFilters.Add(ChannelPoseFilter);
+			}
+		}
+	}
+
+	bool AreFiltersValid(const FPoseSearchIndex* SearchIndex, TConstArrayView<float> QueryValues, int32 PoseIdx, const FPoseSearchPoseMetadata& Metadata
+#if UE_POSE_SEARCH_TRACE_ENABLED
+		, UE::PoseSearch::FSearchContext& SearchContext, const UPoseSearchDatabase* Database
+#endif
+	) const
+	{
+		TConstArrayView<float> PoseValues = SearchIndex->GetPoseValues(PoseIdx);
+		for (const IPoseFilter* PoseFilter : AllPoseFilters)
+		{
+			if (!PoseFilter->IsPoseValid(PoseValues, QueryValues, PoseIdx, Metadata))
+			{
+#if UE_POSE_SEARCH_TRACE_ENABLED
+				if (PoseFilter == &NonSelectableIdxPoseFilter)
+				{
+					// candidate already added to SearchContext.BestCandidates by PopulateNonSelectableIdx
+				}
+				else if (PoseFilter == &BlockTransitionPoseFilter)
+				{
+					const FPoseSearchCost PoseCost = SearchIndex->ComparePoses(PoseIdx, SearchContext.QueryMirrorRequest, EPoseComparisonFlags::None, QueryValues);
+					SearchContext.BestCandidates.Add(PoseCost, PoseIdx, Database, EPoseCandidateFlags::DiscardedBy_BlockTransition);
+				}
+				else
+				{
+					const FPoseSearchCost PoseCost = SearchIndex->ComparePoses(PoseIdx, SearchContext.QueryMirrorRequest, EPoseComparisonFlags::None, QueryValues);
+					SearchContext.BestCandidates.Add(PoseCost, PoseIdx, Database, EPoseCandidateFlags::DiscardedBy_PoseFilter);
+				}
+#endif
+				return false;
+			}
+		}
+		return true;
+	};
+
+private:
+	struct FNonSelectableIdxPoseFilter : public IPoseFilter
+	{
+		virtual bool IsPoseFilterActive() const override
+		{
+			return !NonSelectableIdx.IsEmpty();
+		}
+
+		virtual bool IsPoseValid(TConstArrayView<float> PoseValues, TConstArrayView<float> QueryValues, int32 PoseIdx, const FPoseSearchPoseMetadata& Metadata) const override
+		{
+			return Algo::BinarySearch(NonSelectableIdx, PoseIdx) == INDEX_NONE;
+		}
+
+		TConstArrayView<size_t> NonSelectableIdx;
+	};
+
+	struct FBlockTransitionPoseFilter : public IPoseFilter
+	{
+		virtual bool IsPoseFilterActive() const override
+		{
+			return true;
+		}
+
+		virtual bool IsPoseValid(TConstArrayView<float> PoseValues, TConstArrayView<float> QueryValues, int32 PoseIdx, const FPoseSearchPoseMetadata& Metadata) const override
+		{
+			return !EnumHasAnyFlags(Metadata.Flags, EPoseSearchPoseFlags::BlockTransition);
+		}
+	};
+
+	FNonSelectableIdxPoseFilter NonSelectableIdxPoseFilter;
+	FBlockTransitionPoseFilter BlockTransitionPoseFilter;
+
+	TArray<const IPoseFilter*, TInlineAllocator<64>> AllPoseFilters;
+};
+
 } // namespace UE::PoseSearch
 
 
@@ -331,7 +424,7 @@ void UPoseSearchFeatureChannel::ComputeMeanDeviations(const Eigen::MatrixXd& Cen
 	using namespace UE::PoseSearch;
 
 	int32 DataOffset = ChannelDataOffset;
-	FFeatureVectorHelper::ComputeMeanDeviations(MinimumMeanDeviation, CenteredPoseMatrix, MeanDeviations, DataOffset, ChannelCardinality);
+	FFeatureVectorHelper::ComputeMeanDeviations(GetMinimumMeanDeviation(), CenteredPoseMatrix, MeanDeviations, DataOffset, ChannelCardinality);
 
 	check(DataOffset == ChannelDataOffset + ChannelCardinality);
 }
@@ -644,7 +737,7 @@ UE::PoseSearch::FSearchResult UPoseSearchSequenceMetaData::Search(UE::PoseSearch
 	}
 
 	Schema->BuildQuery(SearchContext, Result.ComposedQuery);
-	TArrayView<const float> QueryValues = Result.ComposedQuery.GetValues();
+	TConstArrayView<float> QueryValues = Result.ComposedQuery.GetValues();
 
 	if (!ensure(QueryValues.Num() == SearchIndex.Schema->SchemaCardinality))
 	{
@@ -1644,7 +1737,7 @@ UE::PoseSearch::FSearchResult UPoseSearchDatabase::SearchPCAKDTree(UE::PoseSearc
 		for (int32 PoseIdx = 0; PoseIdx < SearchIndex->NumPoses; ++PoseIdx)
 		{
 			FKDTree::KNNResultSet ResultSet(ClampedKDTreeQueryNumNeighbors, ResultIndexes, ResultDistanceSqr);
-			TArrayView<const float> PoseValues = SearchIndex->GetPoseValues(PoseIdx);
+			TConstArrayView<float> PoseValues = SearchIndex->GetPoseValues(PoseIdx);
 
 			const RowMajorVectorMapConst Mean(SearchIndex->Mean.GetData(), 1, NumDimensions);
 			const ColMajorMatrixMapConst PCAProjectionMatrix(SearchIndex->PCAProjectionMatrix.GetData(), NumDimensions, ClampedNumberOfPrincipalComponents);
@@ -1671,7 +1764,7 @@ UE::PoseSearch::FSearchResult UPoseSearchDatabase::SearchPCAKDTree(UE::PoseSearc
 
 	SearchContext.GetOrBuildQuery(this, Result.ComposedQuery);
 
-	TArrayView<const float> QueryValues = Result.ComposedQuery.GetValues();
+	TConstArrayView<float> QueryValues = Result.ComposedQuery.GetValues();
 
 	const bool IsCurrentResultFromThisDatabase = SearchContext.IsCurrentResultFromDatabase(this);
 
@@ -1709,31 +1802,28 @@ UE::PoseSearch::FSearchResult UPoseSearchDatabase::SearchPCAKDTree(UE::PoseSearc
 
 		SearchIndex->KDTree.FindNeighbors(ResultSet, ProjectedQueryValues.data());
 
+		// NonSelectableIdx are already filtered out inside the kdtree search
+		const FPoseFilters PoseFilters(SearchIndex->Schema, TConstArrayView<size_t>(), SearchIndex->OverallFlags);
 		for (size_t ResultIndex = 0; ResultIndex < ResultSet.Num(); ++ResultIndex)
 		{
 			const int32 PoseIdx = ResultIndexes[ResultIndex];
-
-			const FPoseSearchPoseMetadata& Metadata = SearchIndex->PoseMetadata[PoseIdx];
-
-			if (EnumHasAnyFlags(Metadata.Flags, EPoseSearchPoseFlags::BlockTransition))
-			{
+			if (PoseFilters.AreFiltersValid(SearchIndex, QueryValues, PoseIdx, SearchIndex->PoseMetadata[PoseIdx]
 #if UE_POSE_SEARCH_TRACE_ENABLED
-				const FPoseSearchCost PoseCost = SearchIndex->ComparePoses(PoseIdx, SearchContext.QueryMirrorRequest, EPoseComparisonFlags::None, QueryValues);
-				SearchContext.BestCandidates.Add(PoseCost, PoseIdx, this, EPoseCandidateFlags::DiscardedBy_BlockTransition);
+				, SearchContext, this
 #endif
-				continue;
-			}
-
-			const FPoseSearchCost PoseCost = SearchIndex->ComparePoses(PoseIdx, SearchContext.QueryMirrorRequest, EPoseComparisonFlags::None, QueryValues);
-			if (PoseCost < Result.PoseCost)
+			))
 			{
-				Result.PoseCost = PoseCost;
-				Result.PoseIdx = PoseIdx;
-			}
+				const FPoseSearchCost PoseCost = SearchIndex->ComparePoses(PoseIdx, SearchContext.QueryMirrorRequest, EPoseComparisonFlags::None, QueryValues);
+				if (PoseCost < Result.PoseCost)
+				{
+					Result.PoseCost = PoseCost;
+					Result.PoseIdx = PoseIdx;
+				}
 
 #if UE_POSE_SEARCH_TRACE_ENABLED
-			SearchContext.BestCandidates.Add(PoseCost, PoseIdx, this, EPoseCandidateFlags::Valid_Pose);
+				SearchContext.BestCandidates.Add(PoseCost, PoseIdx, this, EPoseCandidateFlags::Valid_Pose);
 #endif
+			}
 		}
 
 		if (GetSkipSearchIfPossible() && Result.PoseCost.IsValid())
@@ -1768,7 +1858,7 @@ UE::PoseSearch::FSearchResult UPoseSearchDatabase::SearchBruteForce(UE::PoseSear
 	check(SearchIndex);
 
 	SearchContext.GetOrBuildQuery(this, Result.ComposedQuery);
-	TArrayView<const float> QueryValues = Result.ComposedQuery.GetValues();
+	TConstArrayView<float> QueryValues = Result.ComposedQuery.GetValues();
 
 	const bool IsCurrentResultFromThisDatabase = SearchContext.IsCurrentResultFromDatabase(this);
 	if (!SearchContext.bForceInterrupt && IsCurrentResultFromThisDatabase)
@@ -1797,37 +1887,29 @@ UE::PoseSearch::FSearchResult UPoseSearchDatabase::SearchBruteForce(UE::PoseSear
 	// there's no point in performing the search if CurrentBestTotalCost is already better than that
 	if (SearchContext.GetCurrentBestTotalCost() > SearchIndex->MinCostAddend)
 	{
+		const FPoseFilters PoseFilters(SearchIndex->Schema, NonSelectableIdx, SearchIndex->OverallFlags);
 		for (int32 PoseIdx = 0; PoseIdx < SearchIndex->NumPoses; ++PoseIdx)
 		{
-			const FPoseSearchPoseMetadata& Metadata = SearchIndex->PoseMetadata[PoseIdx];
-
-			if (EnumHasAnyFlags(Metadata.Flags, EPoseSearchPoseFlags::BlockTransition))
-			{
+			if (PoseFilters.AreFiltersValid(SearchIndex, QueryValues, PoseIdx, SearchIndex->PoseMetadata[PoseIdx]
 #if UE_POSE_SEARCH_TRACE_ENABLED
-				const FPoseSearchCost PoseCost = SearchIndex->ComparePoses(PoseIdx, SearchContext.QueryMirrorRequest, EPoseComparisonFlags::None, QueryValues);
-				SearchContext.BestCandidates.Add(PoseCost, PoseIdx, this, EPoseCandidateFlags::DiscardedBy_BlockTransition);
+				, SearchContext, this
 #endif
-				continue;
-			}
-
-			if (Algo::BinarySearch(NonSelectableIdx, PoseIdx) != INDEX_NONE)
+			))
 			{
-				continue;
-			}
-
-			const FPoseSearchCost PoseCost = SearchIndex->ComparePoses(PoseIdx, SearchContext.QueryMirrorRequest, EPoseComparisonFlags::None, QueryValues);
-			if (PoseCost < Result.PoseCost)
-			{
-				Result.PoseCost = PoseCost;
-				Result.PoseIdx = PoseIdx;
-			}
+				const FPoseSearchCost PoseCost = SearchIndex->ComparePoses(PoseIdx, SearchContext.QueryMirrorRequest, EPoseComparisonFlags::None, QueryValues);
+				if (PoseCost < Result.PoseCost)
+				{
+					Result.PoseCost = PoseCost;
+					Result.PoseIdx = PoseIdx;
+				}
 
 #if UE_POSE_SEARCH_TRACE_ENABLED
-			if (PoseSearchMode == EPoseSearchMode::BruteForce)
-			{
-				SearchContext.BestCandidates.Add(PoseCost, PoseIdx, this, EPoseCandidateFlags::Valid_Pose);
-			}
+				if (PoseSearchMode == EPoseSearchMode::BruteForce)
+				{
+					SearchContext.BestCandidates.Add(PoseCost, PoseIdx, this, EPoseCandidateFlags::Valid_Pose);
+				}
 #endif
+			}
 		}
 
 		if (GetSkipSearchIfPossible() && Result.PoseCost.IsValid())
@@ -1882,7 +1964,7 @@ UE::PoseSearch::FSearchResult UPoseSearchDatabaseSet::Search(UE::PoseSearch::FSe
 		{
 			SearchContext.GetOrBuildQuery(Database, Result.ComposedQuery);
 
-			TArrayView<const float> QueryValues = Result.ComposedQuery.GetValues();
+			TConstArrayView<float> QueryValues = Result.ComposedQuery.GetValues();
 
 			Result.PoseIdx = SearchContext.CurrentResult.PoseIdx;
 			Result.PoseCost = SearchIndex->ComparePoses(Result.PoseIdx, SearchContext.QueryMirrorRequest, EPoseComparisonFlags::ContinuingPose, QueryValues);
@@ -2003,7 +2085,7 @@ void FPoseIndicesHistory::Update(const FSearchResult& SearchResult, float DeltaT
 	}
 }
 
-FTransform FSearchContext::TryGetTransformAndCacheResults(float SampleTime, const UPoseSearchSchema* Schema, int8 SchemaBoneIdx, bool& Error)
+FTransform FSearchContext::TryGetTransformAndCacheResults(float SampleTime, const UPoseSearchSchema* Schema, int8 SchemaBoneIdx, bool& bError)
 {
 	check(History && Schema);
 
@@ -2018,7 +2100,7 @@ FTransform FSearchContext::TryGetTransformAndCacheResults(float SampleTime, cons
 
 	if (Entry)
 	{
-		Error = false;
+		bError = false;
 		return Entry->Transform;
 	}
 
@@ -2040,11 +2122,11 @@ FTransform FSearchContext::TryGetTransformAndCacheResults(float SampleTime, cons
 				NewEntry.Transform = SampledComponentPose[NewEntryBoneIndexType];
 			}
 
-			Error = false;
+			bError = false;
 			return SampledComponentPose[BoneIndexType];
 		}
 
-		Error = true;
+		bError = true;
 		return FTransform::Identity;
 	}
 	
@@ -2056,11 +2138,11 @@ FTransform FSearchContext::TryGetTransformAndCacheResults(float SampleTime, cons
 		NewEntry.BoneIndexType = BoneIndexType;
 		NewEntry.Transform = SampledRootTransform;
 
-		Error = false;
+		bError = false;
 		return SampledRootTransform;
 	}
 	
-	Error = true;
+	bError = true;
 	return FTransform::Identity;
 }
 
@@ -2345,7 +2427,7 @@ void FFeatureVectorHelper::EncodeQuat(TArrayView<float> Values, int32& DataOffse
 	DataOffset += EncodeQuatCardinality;
 }
 
-void FFeatureVectorHelper::EncodeQuat(TArrayView<float> Values, int32& DataOffset, TArrayView<const float> PrevValues, TArrayView<const float> CurValues, TArrayView<const float> NextValues, float LerpValue)
+void FFeatureVectorHelper::EncodeQuat(TArrayView<float> Values, int32& DataOffset, TConstArrayView<float> PrevValues, TConstArrayView<float> CurValues, TConstArrayView<float> NextValues, float LerpValue)
 {
 	FQuat Quat = DecodeQuatInternal(CurValues, DataOffset);
 
@@ -2361,19 +2443,19 @@ void FFeatureVectorHelper::EncodeQuat(TArrayView<float> Values, int32& DataOffse
 			Quat = FQuat::Slerp(Quat, DecodeQuatInternal(NextValues, DataOffset), LerpValue);
 		}
 	}
+	
 	// @todo: do we need to add options for cubic interpolation?
-
 	EncodeQuat(Values, DataOffset, Quat);
 }
 
-FQuat FFeatureVectorHelper::DecodeQuat(TArrayView<const float> Values, int32& DataOffset)
+FQuat FFeatureVectorHelper::DecodeQuat(TConstArrayView<float> Values, int32& DataOffset)
 {
 	const FQuat Quat = DecodeQuatInternal(Values, DataOffset);
 	DataOffset += EncodeQuatCardinality;
 	return Quat;
 }
 
-FQuat FFeatureVectorHelper::DecodeQuatInternal(TArrayView<const float> Values, int32 DataOffset)
+FQuat FFeatureVectorHelper::DecodeQuatInternal(TConstArrayView<float> Values, int32 DataOffset)
 {
 	const FVector X(Values[DataOffset + 0], Values[DataOffset + 1], Values[DataOffset + 2]);
 	const FVector Y(Values[DataOffset + 3], Values[DataOffset + 4], Values[DataOffset + 5]);
@@ -2395,7 +2477,7 @@ void FFeatureVectorHelper::EncodeVector(TArrayView<float> Values, int32& DataOff
 	DataOffset += EncodeVectorCardinality;
 }
 
-void FFeatureVectorHelper::EncodeVector(TArrayView<float> Values, int32& DataOffset, TArrayView<const float> PrevValues, TArrayView<const float> CurValues, TArrayView<const float> NextValues, float LerpValue, bool bNormalize)
+void FFeatureVectorHelper::EncodeVector(TArrayView<float> Values, int32& DataOffset, TConstArrayView<float> PrevValues, TConstArrayView<float> CurValues, TConstArrayView<float> NextValues, float LerpValue, bool bNormalize)
 {
 	FVector Vector = DecodeVectorInternal(CurValues, DataOffset);
 
@@ -2411,8 +2493,8 @@ void FFeatureVectorHelper::EncodeVector(TArrayView<float> Values, int32& DataOff
 			Vector = FMath::Lerp(Vector, DecodeVectorInternal(NextValues, DataOffset), LerpValue);
 		}
 	}
+	
 	// @todo: do we need to add options for cubic interpolation?
-
 	if (bNormalize)
 	{
 		Vector = Vector.GetSafeNormal(UE_SMALL_NUMBER, FVector::XAxisVector);
@@ -2421,14 +2503,14 @@ void FFeatureVectorHelper::EncodeVector(TArrayView<float> Values, int32& DataOff
 	EncodeVector(Values, DataOffset, Vector);
 }
 
-FVector FFeatureVectorHelper::DecodeVector(TArrayView<const float> Values, int32& DataOffset)
+FVector FFeatureVectorHelper::DecodeVector(TConstArrayView<float> Values, int32& DataOffset)
 {
 	const FVector Vector = DecodeVectorInternal(Values, DataOffset);
 	DataOffset += EncodeVectorCardinality;
 	return Vector;
 }
 
-FVector FFeatureVectorHelper::DecodeVectorInternal(TArrayView<const float> Values, int32 DataOffset)
+FVector FFeatureVectorHelper::DecodeVectorInternal(TConstArrayView<float> Values, int32 DataOffset)
 {
 	return FVector(Values[DataOffset + 0], Values[DataOffset + 1], Values[DataOffset + 2]);
 }
@@ -2440,7 +2522,7 @@ void FFeatureVectorHelper::EncodeVector2D(TArrayView<float> Values, int32& DataO
 	DataOffset += EncodeVector2DCardinality;
 }
 
-void FFeatureVectorHelper::EncodeVector2D(TArrayView<float> Values, int32& DataOffset, TArrayView<const float> PrevValues, TArrayView<const float> CurValues, TArrayView<const float> NextValues, float LerpValue)
+void FFeatureVectorHelper::EncodeVector2D(TArrayView<float> Values, int32& DataOffset, TConstArrayView<float> PrevValues, TConstArrayView<float> CurValues, TConstArrayView<float> NextValues, float LerpValue)
 {
 	FVector2D Vector2D = DecodeVector2DInternal(CurValues, DataOffset);
 
@@ -2456,23 +2538,62 @@ void FFeatureVectorHelper::EncodeVector2D(TArrayView<float> Values, int32& DataO
 			Vector2D = FMath::Lerp(Vector2D, DecodeVector2DInternal(NextValues, DataOffset), LerpValue);
 		}
 	}
+	
 	// @todo: do we need to add options for cubic interpolation?
-
 	EncodeVector2D(Values, DataOffset, Vector2D);
 }
 
-FVector2D FFeatureVectorHelper::DecodeVector2D(TArrayView<const float> Values, int32& DataOffset)
+FVector2D FFeatureVectorHelper::DecodeVector2D(TConstArrayView<float> Values, int32& DataOffset)
 {
 	const FVector2D Vector2D = DecodeVector2DInternal(Values, DataOffset);
 	DataOffset += EncodeVector2DCardinality;
 	return Vector2D;
 }
 
-FVector2D FFeatureVectorHelper::DecodeVector2DInternal(TArrayView<const float> Values, int32 DataOffset)
+FVector2D FFeatureVectorHelper::DecodeVector2DInternal(TConstArrayView<float> Values, int32 DataOffset)
 {
 	return FVector2D(Values[DataOffset + 0], Values[DataOffset + 1]);
 }
 
+void FFeatureVectorHelper::EncodeFloat(TArrayView<float> Values, int32& DataOffset, const float Value)
+{
+	Values[DataOffset + 0] = Value;
+	DataOffset += EncodeFloatCardinality;
+}
+
+void FFeatureVectorHelper::EncodeFloat(TArrayView<float> Values, int32& DataOffset, TConstArrayView<float> PrevValues, TConstArrayView<float> CurValues, TConstArrayView<float> NextValues, float LerpValue)
+{
+	float Value = DecodeFloatInternal(CurValues, DataOffset);
+
+	// linear interpolation
+	if (!FMath::IsNearlyZero(LerpValue))
+	{
+		if (LerpValue < 0.f)
+		{
+			Value = FMath::Lerp(Value, DecodeFloatInternal(PrevValues, DataOffset), -LerpValue);
+		}
+		else
+		{
+			Value = FMath::Lerp(Value, DecodeFloatInternal(NextValues, DataOffset), LerpValue);
+		}
+	}
+
+	// @todo: do we need to add options for cubic interpolation?
+	EncodeFloat(Values, DataOffset, Value);
+}
+
+float FFeatureVectorHelper::DecodeFloat(TConstArrayView<float> Values, int32& DataOffset)
+{
+	const float Value = DecodeFloatInternal(Values, DataOffset);
+	DataOffset += EncodeFloatCardinality;
+	return Value;
+}
+
+float FFeatureVectorHelper::DecodeFloatInternal(TConstArrayView<float> Values, int32 DataOffset)
+{
+	return Values[DataOffset];
+}
+
 void FFeatureVectorHelper::ComputeMeanDeviations(float MinMeanDeviation, const Eigen::MatrixXd& CenteredPoseMatrix, Eigen::VectorXd& MeanDeviations, int32& DataOffset, int32 Cardinality)
 {
 	const int32 NumPoses = CenteredPoseMatrix.cols();
@@ -2733,7 +2854,7 @@ void ICostBreakDownData::AddEntireBreakDownSection(const FText& Label, const UPo
 // Uses distance delta between NextRootDistanceIndex and NextRootDistanceIndex - 1 and extrapolates it to ExtrapolationTime
 static float ExtrapolateAccumulatedRootDistance(
 	int32 SamplingRate,
-	TArrayView<const float> AccumulatedRootDistance,
+	TConstArrayView<float> AccumulatedRootDistance,
 	int32 NextRootDistanceIndex, 
 	float ExtrapolationTime,
 	const FPoseSearchExtrapolationParameters& ExtrapolationParameters)
@@ -2753,7 +2874,7 @@ static float ExtrapolateAccumulatedRootDistance(
 
 static float ExtractAccumulatedRootDistance(
 	int32 SamplingRate,
-	TArrayView<const float> AccumulatedRootDistance,
+	TConstArrayView<float> AccumulatedRootDistance,
 	float PlayLength,
 	float Time,
 	const FPoseSearchExtrapolationParameters& ExtrapolationParameters)
@@ -3510,6 +3631,7 @@ public:
 		int32 FirstIndexedSample = 0;
 		int32 LastIndexedSample = 0;
 		int32 NumIndexedPoses = 0;
+
 		TArray<float> FeatureVectorTable;
 		TArray<FPoseSearchPoseMetadata> PoseMetadata;
 		TBitArray<> AllFeaturesNotAdded;
@@ -3529,13 +3651,8 @@ public: // IAssetIndexer
 	FTransform GetTransformAndCacheResults(float SampleTime, float OriginTime, int8 SchemaBoneIdx, bool& Clamped) override;
 
 private:
-	FBoneContainer BoneContainer;
-	FAssetIndexingContext IndexingContext;
-	TArray<FPoseSearchFeatureVectorBuilder> FeatureVectorBuilders;
-	FPoseSearchPoseMetadata Metadata;
+	FPoseSearchPoseMetadata GetMetadata(int32 SampleIdx) const;
 	
-	void AddMetadata(int32 SampleIdx);
-
 	struct CachedEntry
 	{
 		float SampleTime;
@@ -3551,6 +3668,8 @@ private:
 		FAnimationPoseData AnimPoseData = { Pose, UnusedCurve, UnusedAtrribute };
 	};
 
+	FBoneContainer BoneContainer;
+	FAssetIndexingContext IndexingContext;
 	TArray<CachedEntry> CachedEntries;
 };
 
@@ -3581,14 +3700,7 @@ void FAssetIndexer::Init(const FAssetIndexingContext& InIndexingContext, const F
 	Output.NumIndexedPoses = Output.LastIndexedSample - Output.FirstIndexedSample + 1;
 	
 	Output.FeatureVectorTable.SetNumZeroed(IndexingContext.Schema->SchemaCardinality * Output.NumIndexedPoses);
-
 	Output.PoseMetadata.SetNum(Output.NumIndexedPoses);
-
-	FeatureVectorBuilders.SetNum(Output.NumIndexedPoses);
-	for (FPoseSearchFeatureVectorBuilder& Builder : FeatureVectorBuilders)
-	{
-		Builder.Init(IndexingContext.Schema);
-	}
 }
 
 bool FAssetIndexer::Process()
@@ -3601,40 +3713,23 @@ bool FAssetIndexer::Process()
 
 	IndexingContext.BeginSampleIdx = Output.FirstIndexedSample;
 	IndexingContext.EndSampleIdx = Output.LastIndexedSample + 1;
-	const int32 NumSamplesInRange = IndexingContext.EndSampleIdx - IndexingContext.BeginSampleIdx;
 
-	FAssetIndexingOutput AssetIndexingOutput{FeatureVectorBuilders};
-
-	// Index each channel
-	for (int32 ChannelIdx = 0; ChannelIdx != IndexingContext.Schema->Channels.Num(); ++ChannelIdx)
-	{
-		const TObjectPtr<UPoseSearchFeatureChannel>& Channel = IndexingContext.Schema->Channels[ChannelIdx];
-		Channel->IndexAsset(*this, AssetIndexingOutput);
-	}
-
-	// @todo: this step can be avoided now since Schema.SchemaCardinality is known and the data can be preallocated before Channel->IndexAsset
-	
-	// Merge spans of feature vectors into contiguous buffer
 	if (IndexingContext.Schema->SchemaCardinality > 0)
 	{
-		for (int32 VectorIdx = 0; VectorIdx != NumSamplesInRange; ++VectorIdx)
+		// Index each channel
+		FAssetIndexingOutput AssetIndexingOutput(IndexingContext.Schema->SchemaCardinality, Output.FeatureVectorTable);
+		for (int32 ChannelIdx = 0; ChannelIdx != IndexingContext.Schema->Channels.Num(); ++ChannelIdx) 
 		{
-			const int32 SampleIdx = VectorIdx + IndexingContext.BeginSampleIdx;
-			const int32 PoseIdx = SampleIdx - Output.FirstIndexedSample;
-			const int32 FirstValueIdx = PoseIdx * IndexingContext.Schema->SchemaCardinality;
-			TArrayView<float> WriteValues = MakeArrayView(&Output.FeatureVectorTable[FirstValueIdx], IndexingContext.Schema->SchemaCardinality);
-			TArrayView<const float> ReadValues = FeatureVectorBuilders[VectorIdx].GetValues();
-			check(WriteValues.Num() == ReadValues.Num());
-			FMemory::Memcpy(WriteValues.GetData(), ReadValues.GetData(), WriteValues.Num() * WriteValues.GetTypeSize());
+			IndexingContext.Schema->Channels[ChannelIdx]->IndexAsset(*this, AssetIndexingOutput);
 		}
 	}
 
 	// Generate pose metadata
 	for (int32 SampleIdx = IndexingContext.BeginSampleIdx; SampleIdx != IndexingContext.EndSampleIdx; ++SampleIdx)
 	{
-		AddMetadata(SampleIdx);
 		const int32 PoseIdx = SampleIdx - Output.FirstIndexedSample;
-		Output.PoseMetadata[PoseIdx] = Metadata;
+		FPoseSearchPoseMetadata& OutputPoseMetadata = Output.PoseMetadata[PoseIdx];
+		OutputPoseMetadata = GetMetadata(SampleIdx);
 	}
 
 	return true;
@@ -3883,12 +3978,12 @@ FTransform FAssetIndexer::MirrorTransform(const FTransform& Transform) const
 	return IndexingContext.bMirrored ? IndexingContext.SamplingContext->MirrorTransform(Transform) : Transform;
 }
 
-void FAssetIndexer::AddMetadata(int32 SampleIdx)
+FPoseSearchPoseMetadata FAssetIndexer::GetMetadata(int32 SampleIdx) const
 {
 	const float SequenceLength = IndexingContext.MainSampler->GetPlayLength();
 	const float SampleTime = FMath::Min(SampleIdx * IndexingContext.Schema->GetSamplingInterval(), SequenceLength);
 
-	Metadata = FPoseSearchPoseMetadata();
+	FPoseSearchPoseMetadata Metadata;
 	Metadata.CostAddend = IndexingContext.Schema->BaseCostBias;
 	Metadata.ContinuingPoseCostAddend = IndexingContext.Schema->ContinuingPoseCostBias;
 	
@@ -3913,6 +4008,7 @@ void FAssetIndexer::AddMetadata(int32 SampleIdx)
 			Metadata.ContinuingPoseCostAddend = ContinuingPoseCostBias->CostAddend;
 		}
 	}
+	return Metadata;
 }
 
 FTransform FAssetIndexer::GetTransformAndCacheResults(float SampleTime, float OriginTime, int8 SchemaBoneIdx, bool& Clamped)
@@ -3986,7 +4082,7 @@ FTransform FAssetIndexer::GetTransformAndCacheResults(float SampleTime, float Or
 //////////////////////////////////////////////////////////////////////////
 // PoseSearch API
 
-void DrawFeatureVector(const FDebugDrawParams& DrawParams, TArrayView<const float> PoseVector)
+void DrawFeatureVector(const FDebugDrawParams& DrawParams, TConstArrayView<float> PoseVector)
 {
 	if (DrawParams.CanDraw())
 	{
@@ -4379,6 +4475,12 @@ bool BuildIndex(const UAnimSequence* Sequence, UPoseSearchSequenceMetaData* Sequ
 	SequenceMetaData->SearchIndex.Assets.Add(SearchIndexAsset);
 	SequenceMetaData->SearchIndex.PoseMetadata = Indexer.Output.PoseMetadata;
 
+	SequenceMetaData->SearchIndex.OverallFlags = EPoseSearchPoseFlags::None;
+	for (const FPoseSearchPoseMetadata& PoseMetadata : Indexer.Output.PoseMetadata)
+	{
+		SequenceMetaData->SearchIndex.OverallFlags = PoseMetadata.Flags;
+	}
+
 	// todo: do we need to PreprocessSearchIndex?
 	// PreprocessSearchIndex(&SequenceMetaData->SearchIndex, Database);
 	return true;
@@ -4575,6 +4677,7 @@ void FDatabaseIndexingContext::JoinIndex()
 	SearchIndex->Values.Reset();
 	SearchIndex->PoseMetadata.Reset();
 	SearchIndex->PCAValues.Reset();
+	SearchIndex->OverallFlags = EPoseSearchPoseFlags::None;
 
 	for (int32 AssetIdx = 0; AssetIdx != SearchIndex->Assets.Num(); ++AssetIdx)
 	{
@@ -4593,6 +4696,7 @@ void FDatabaseIndexingContext::JoinIndex()
 		for (int32 i = PoseMetadataStartIdx; i < PoseMetadataEndIdx; ++i)
 		{
 			SearchIndex->PoseMetadata[i].AssetIndex = AssetIdx;
+			SearchIndex->OverallFlags |= SearchIndex->PoseMetadata[i].Flags;
 		}
 
 		TotalPoses += Output.NumIndexedPoses;
diff --git a/Engine/Plugins/Experimental/Animation/PoseSearch/Source/Runtime/Private/PoseSearchDerivedData.cpp b/Engine/Plugins/Experimental/Animation/PoseSearch/Source/Runtime/Private/PoseSearchDerivedData.cpp
--- a/Engine/Plugins/Experimental/Animation/PoseSearch/Source/Runtime/Private/PoseSearchDerivedData.cpp
+++ b/Engine/Plugins/Experimental/Animation/PoseSearch/Source/Runtime/Private/PoseSearchDerivedData.cpp
@@ -274,6 +274,7 @@ namespace UE::PoseSearch
 		Ar << NumValues;
 		Ar << NumPCAValues;
 		Ar << NumAssets;
+		Ar << Index.OverallFlags;
 
 		if (Ar.IsLoading())
 		{
diff --git a/Engine/Plugins/Experimental/Animation/PoseSearch/Source/Runtime/Private/PoseSearchFeatureChannels.cpp b/Engine/Plugins/Experimental/Animation/PoseSearch/Source/Runtime/Private/PoseSearchFeatureChannels.cpp
--- a/Engine/Plugins/Experimental/Animation/PoseSearch/Source/Runtime/Private/PoseSearchFeatureChannels.cpp
+++ b/Engine/Plugins/Experimental/Animation/PoseSearch/Source/Runtime/Private/PoseSearchFeatureChannels.cpp
@@ -71,8 +71,8 @@ static void CollectBonePositions(TArray<FVector>& BonePositions, IAssetIndexer&
 	for (int32 SampleIdx = 0; SampleIdx != NumSamples; ++SampleIdx)
 	{
 		const float SampleTime = SampleTimeStart + SampleIdx * FiniteDelta;
-		bool Unused;
-		BonePositions[SampleIdx] = Indexer.GetTransformAndCacheResults(SampleTime, SampleTimeStart, SchemaBoneIdx, Unused).GetTranslation();
+		bool bUnused;
+		BonePositions[SampleIdx] = Indexer.GetTransformAndCacheResults(SampleTime, SampleTimeStart, SchemaBoneIdx, bUnused).GetTranslation();
 	}
 }
 
@@ -341,8 +341,7 @@ void UPoseSearchFeatureChannel_Position::IndexAsset(UE::PoseSearch::IAssetIndexe
 
 	for (int32 SampleIdx = IndexingContext.BeginSampleIdx; SampleIdx != IndexingContext.EndSampleIdx; ++SampleIdx)
 	{
-		int32 VectorIdx = SampleIdx - IndexingContext.BeginSampleIdx;
-		FPoseSearchFeatureVectorBuilder& FeatureVector = IndexingOutput.PoseVectors[VectorIdx];
+		const int32 VectorIdx = SampleIdx - IndexingContext.BeginSampleIdx;
 
 		const float OriginSampleTime = FMath::Min(SampleIdx * IndexingContext.Schema->GetSamplingInterval(), IndexingContext.MainSampler->GetPlayLength());
 		const float SubsampleTime = OriginSampleTime + SampleTimeOffset;
@@ -350,7 +349,7 @@ void UPoseSearchFeatureChannel_Position::IndexAsset(UE::PoseSearch::IAssetIndexe
 		bool ClampedPresent;
 		const FTransform BoneTransformsPresent = Indexer.GetTransformAndCacheResults(SubsampleTime, bUseSampleTimeOffsetRootBone ? SubsampleTime : OriginSampleTime, SchemaBoneIdx, ClampedPresent);
 		int32 DataOffset = ChannelDataOffset;
-		FFeatureVectorHelper::EncodeVector(FeatureVector.EditValues(), DataOffset, BoneTransformsPresent.GetTranslation());
+		FFeatureVectorHelper::EncodeVector(IndexingOutput.GetPoseVector(VectorIdx), DataOffset, BoneTransformsPresent.GetTranslation());
 	}
 }
 
@@ -368,26 +367,27 @@ bool UPoseSearchFeatureChannel_Position::BuildQuery(UE::PoseSearch::FSearchConte
 			int32 DataOffset = ChannelDataOffset;
 			FFeatureVectorHelper::EncodeVector(InOutQuery.EditValues(), DataOffset, SearchContext.GetCurrentResultPrevPoseVector(), SearchContext.GetCurrentResultPoseVector(), SearchContext.GetCurrentResultNextPoseVector(), LerpValue);
 		}
+		// else leave the InOutQuery set to zero since the SearchContext.History is invalid and it'll fail if we continue
 		return bSkip;
 	}
 
-	bool AnyError = false;
-	FTransform Transform = SearchContext.TryGetTransformAndCacheResults(SampleTimeOffset, InOutQuery.GetSchema(), SchemaBoneIdx, AnyError);
+	bool bAnyError = false;
+	FTransform Transform = SearchContext.TryGetTransformAndCacheResults(SampleTimeOffset, InOutQuery.GetSchema(), SchemaBoneIdx, bAnyError);
 
 	if (!bUseSampleTimeOffsetRootBone)
 	{
-		const FTransform RootTransform = SearchContext.TryGetTransformAndCacheResults(0.f, InOutQuery.GetSchema(), FSearchContext::SchemaRootBoneIdx, AnyError);
-		const FTransform RootTransformPrev = SearchContext.TryGetTransformAndCacheResults(SampleTimeOffset, InOutQuery.GetSchema(), FSearchContext::SchemaRootBoneIdx, AnyError);
+		const FTransform RootTransform = SearchContext.TryGetTransformAndCacheResults(0.f, InOutQuery.GetSchema(), FSearchContext::SchemaRootBoneIdx, bAnyError);
+		const FTransform RootTransformPrev = SearchContext.TryGetTransformAndCacheResults(SampleTimeOffset, InOutQuery.GetSchema(), FSearchContext::SchemaRootBoneIdx, bAnyError);
 		Transform = Transform * (RootTransformPrev * RootTransform.Inverse());
 	}
 
 	int32 DataOffset = ChannelDataOffset;
 	FFeatureVectorHelper::EncodeVector(InOutQuery.EditValues(), DataOffset, Transform.GetTranslation());
 	check(DataOffset == ChannelDataOffset + ChannelCardinality);
-	return !AnyError;
+	return !bAnyError;
 }
 
-void UPoseSearchFeatureChannel_Position::DebugDraw(const UE::PoseSearch::FDebugDrawParams& DrawParams, TArrayView<const float> PoseVector) const
+void UPoseSearchFeatureChannel_Position::DebugDraw(const UE::PoseSearch::FDebugDrawParams& DrawParams, TConstArrayView<float> PoseVector) const
 {
 	using namespace UE::PoseSearch;
 
@@ -463,8 +463,7 @@ void UPoseSearchFeatureChannel_Heading::IndexAsset(UE::PoseSearch::IAssetIndexer
 
 	for (int32 SampleIdx = IndexingContext.BeginSampleIdx; SampleIdx != IndexingContext.EndSampleIdx; ++SampleIdx)
 	{
-		int32 VectorIdx = SampleIdx - IndexingContext.BeginSampleIdx;
-		FPoseSearchFeatureVectorBuilder& FeatureVector = IndexingOutput.PoseVectors[VectorIdx];
+		const int32 VectorIdx = SampleIdx - IndexingContext.BeginSampleIdx;
 
 		const float OriginSampleTime = FMath::Min(SampleIdx * IndexingContext.Schema->GetSamplingInterval(), IndexingContext.MainSampler->GetPlayLength());
 		const float SubsampleTime = OriginSampleTime + SampleTimeOffset;
@@ -472,7 +471,7 @@ void UPoseSearchFeatureChannel_Heading::IndexAsset(UE::PoseSearch::IAssetIndexer
 		bool ClampedPresent;
 		const FTransform BoneTransformsPresent = Indexer.GetTransformAndCacheResults(SubsampleTime, bUseSampleTimeOffsetRootBone ? SubsampleTime : OriginSampleTime, SchemaBoneIdx, ClampedPresent);
 		int32 DataOffset = ChannelDataOffset;
-		FFeatureVectorHelper::EncodeVector(FeatureVector.EditValues(), DataOffset, GetAxis(BoneTransformsPresent.GetRotation()));
+		FFeatureVectorHelper::EncodeVector(IndexingOutput.GetPoseVector(VectorIdx), DataOffset, GetAxis(BoneTransformsPresent.GetRotation()));
 		check(DataOffset == ChannelDataOffset + ChannelCardinality);
 	}
 }
@@ -492,16 +491,17 @@ bool UPoseSearchFeatureChannel_Heading::BuildQuery(UE::PoseSearch::FSearchContex
 			FFeatureVectorHelper::EncodeVector(InOutQuery.EditValues(), DataOffset, SearchContext.GetCurrentResultPrevPoseVector(), SearchContext.GetCurrentResultPoseVector(), SearchContext.GetCurrentResultNextPoseVector(), LerpValue, true);
 			check(DataOffset == ChannelDataOffset + ChannelCardinality);
 		}
+		// else leave the InOutQuery set to zero since the SearchContext.History is invalid and it'll fail if we continue
 		return bSkip;
 	}
 
-	bool AnyError = false;
-	FTransform Transform = SearchContext.TryGetTransformAndCacheResults(SampleTimeOffset, InOutQuery.GetSchema(), SchemaBoneIdx, AnyError);
+	bool bAnyError = false;
+	FTransform Transform = SearchContext.TryGetTransformAndCacheResults(SampleTimeOffset, InOutQuery.GetSchema(), SchemaBoneIdx, bAnyError);
 
 	if (!bUseSampleTimeOffsetRootBone)
 	{
-		const FTransform RootTransform = SearchContext.TryGetTransformAndCacheResults(0.f, InOutQuery.GetSchema(), FSearchContext::SchemaRootBoneIdx, AnyError);
-		const FTransform RootTransformPrev = SearchContext.TryGetTransformAndCacheResults(SampleTimeOffset, InOutQuery.GetSchema(), FSearchContext::SchemaRootBoneIdx, AnyError);
+		const FTransform RootTransform = SearchContext.TryGetTransformAndCacheResults(0.f, InOutQuery.GetSchema(), FSearchContext::SchemaRootBoneIdx, bAnyError);
+		const FTransform RootTransformPrev = SearchContext.TryGetTransformAndCacheResults(SampleTimeOffset, InOutQuery.GetSchema(), FSearchContext::SchemaRootBoneIdx, bAnyError);
 		Transform = Transform * (RootTransformPrev * RootTransform.Inverse());
 	}
 
@@ -509,10 +509,10 @@ bool UPoseSearchFeatureChannel_Heading::BuildQuery(UE::PoseSearch::FSearchContex
 	FFeatureVectorHelper::EncodeVector(InOutQuery.EditValues(), DataOffset, GetAxis(Transform.GetRotation()));
 	check(DataOffset == ChannelDataOffset + ChannelCardinality);
 
-	return !AnyError;
+	return !bAnyError;
 }
 
-void UPoseSearchFeatureChannel_Heading::DebugDraw(const UE::PoseSearch::FDebugDrawParams& DrawParams, TArrayView<const float> PoseVector) const
+void UPoseSearchFeatureChannel_Heading::DebugDraw(const UE::PoseSearch::FDebugDrawParams& DrawParams, TConstArrayView<float> PoseVector) const
 {
 	using namespace UE::PoseSearch;
 
@@ -718,8 +718,7 @@ void UPoseSearchFeatureChannel_Pose::IndexAsset(UE::PoseSearch::IAssetIndexer& I
 	for (int32 SampleIdx = IndexingContext.BeginSampleIdx; SampleIdx != IndexingContext.EndSampleIdx; ++SampleIdx)
 	{
 		const int32 VectorIdx = SampleIdx - IndexingContext.BeginSampleIdx;
-		FPoseSearchFeatureVectorBuilder& FeatureVector = IndexingOutput.PoseVectors[VectorIdx];
-		AddPoseFeatures(Indexer, SampleIdx, FeatureVector, Phases);
+		AddPoseFeatures(Indexer, SampleIdx, IndexingOutput.GetPoseVector(VectorIdx), Phases);
 	}
 }
 
@@ -737,28 +736,28 @@ void UPoseSearchFeatureChannel_Pose::ComputeMeanDeviations(const Eigen::MatrixXd
 		{
 			for (int32 SubsampleIdx = 0; SubsampleIdx != SampleTimes.Num(); ++SubsampleIdx)
 			{
-				FFeatureVectorHelper::ComputeMeanDeviations(MinimumMeanDeviation, CenteredPoseMatrix, MeanDeviations, DataOffset, FFeatureVectorHelper::EncodeVectorCardinality);
+				FFeatureVectorHelper::ComputeMeanDeviations(GetMinimumMeanDeviation(), CenteredPoseMatrix, MeanDeviations, DataOffset, FFeatureVectorHelper::EncodeVectorCardinality);
 			}
 		}
 		if (EnumHasAnyFlags(SampledBone.Flags, EPoseSearchBoneFlags::Rotation))
 		{
 			for (int32 SubsampleIdx = 0; SubsampleIdx != SampleTimes.Num(); ++SubsampleIdx)
 			{
-				FFeatureVectorHelper::ComputeMeanDeviations(MinimumMeanDeviation, CenteredPoseMatrix, MeanDeviations, DataOffset, FFeatureVectorHelper::EncodeQuatCardinality);
+				FFeatureVectorHelper::ComputeMeanDeviations(GetMinimumMeanDeviation(), CenteredPoseMatrix, MeanDeviations, DataOffset, FFeatureVectorHelper::EncodeQuatCardinality);
 			}
 		}
 		if (EnumHasAnyFlags(SampledBone.Flags, EPoseSearchBoneFlags::Velocity))
 		{
 			for (int32 SubsampleIdx = 0; SubsampleIdx != SampleTimes.Num(); ++SubsampleIdx)
 			{
-				FFeatureVectorHelper::ComputeMeanDeviations(MinimumMeanDeviation, CenteredPoseMatrix, MeanDeviations, DataOffset, FFeatureVectorHelper::EncodeVectorCardinality);
+				FFeatureVectorHelper::ComputeMeanDeviations(GetMinimumMeanDeviation(), CenteredPoseMatrix, MeanDeviations, DataOffset, FFeatureVectorHelper::EncodeVectorCardinality);
 			}
 		}
 		if (EnumHasAnyFlags(SampledBone.Flags, EPoseSearchBoneFlags::Phase))
 		{
 			for (int32 SubsampleIdx = 0; SubsampleIdx != SampleTimes.Num(); ++SubsampleIdx)
 			{
-				FFeatureVectorHelper::ComputeMeanDeviations(MinimumMeanDeviation, CenteredPoseMatrix, MeanDeviations, DataOffset, FFeatureVectorHelper::EncodeVector2DCardinality);
+				FFeatureVectorHelper::ComputeMeanDeviations(GetMinimumMeanDeviation(), CenteredPoseMatrix, MeanDeviations, DataOffset, FFeatureVectorHelper::EncodeVector2DCardinality);
 			}
 		}
 	}
@@ -766,7 +765,7 @@ void UPoseSearchFeatureChannel_Pose::ComputeMeanDeviations(const Eigen::MatrixXd
 	check(DataOffset == ChannelDataOffset + ChannelCardinality);
 }
 
-void UPoseSearchFeatureChannel_Pose::AddPoseFeatures(UE::PoseSearch::IAssetIndexer& Indexer, int32 SampleIdx, FPoseSearchFeatureVectorBuilder& FeatureVector, const TArray<TArray<FVector2D>>& Phases) const
+void UPoseSearchFeatureChannel_Pose::AddPoseFeatures(UE::PoseSearch::IAssetIndexer& Indexer, int32 SampleIdx, TArrayView<float> FeatureVector, const TArray<TArray<FVector2D>>& Phases) const
 {
 	// This function samples the instantaneous pose at time t as well as the pose's velocity and acceleration at time t.
 	// Symmetric finite differences are used to approximate derivatives:
@@ -805,7 +804,7 @@ void UPoseSearchFeatureChannel_Pose::AddPoseFeatures(UE::PoseSearch::IAssetIndex
 
 				bool ClampedPresent;
 				const FTransform BoneTransformsPresent = Indexer.GetTransformAndCacheResults(SubsampleTime, UseCharacterSpaceVelocities ? SubsampleTime : SampleTime, SchemaBoneIdx[ChannelBoneIdx], ClampedPresent);
-				FFeatureVectorHelper::EncodeVector(FeatureVector.EditValues(), DataOffset, BoneTransformsPresent.GetTranslation());
+				FFeatureVectorHelper::EncodeVector(FeatureVector, DataOffset, BoneTransformsPresent.GetTranslation());
 			}
 		}
 
@@ -817,7 +816,7 @@ void UPoseSearchFeatureChannel_Pose::AddPoseFeatures(UE::PoseSearch::IAssetIndex
 
 				bool ClampedPresent;
 				const FTransform BoneTransformsPresent = Indexer.GetTransformAndCacheResults(SubsampleTime, UseCharacterSpaceVelocities ? SubsampleTime : SampleTime, SchemaBoneIdx[ChannelBoneIdx], ClampedPresent);
-				FFeatureVectorHelper::EncodeQuat(FeatureVector.EditValues(), DataOffset, BoneTransformsPresent.GetRotation());
+				FFeatureVectorHelper::EncodeQuat(FeatureVector, DataOffset, BoneTransformsPresent.GetRotation());
 			}
 		}
 
@@ -850,7 +849,7 @@ void UPoseSearchFeatureChannel_Pose::AddPoseFeatures(UE::PoseSearch::IAssetIndex
 					LinearVelocity = (BoneTransformsFuture.GetTranslation() - BoneTransformsPast.GetTranslation()) / (SamplingContext->FiniteDelta * 2.0f);
 				}
 
-				FFeatureVectorHelper::EncodeVector(FeatureVector.EditValues(), DataOffset, LinearVelocity);
+				FFeatureVectorHelper::EncodeVector(FeatureVector, DataOffset, LinearVelocity);
 			}
 		}
 
@@ -860,7 +859,7 @@ void UPoseSearchFeatureChannel_Pose::AddPoseFeatures(UE::PoseSearch::IAssetIndex
 			{
 				// @todo: support for SubsampleIdx
 				const int32 VectorIdx = SampleIdx - IndexingContext.BeginSampleIdx;
-				FFeatureVectorHelper::EncodeVector2D(FeatureVector.EditValues(), DataOffset, Phases[ChannelBoneIdx][VectorIdx]);
+				FFeatureVectorHelper::EncodeVector2D(FeatureVector, DataOffset, Phases[ChannelBoneIdx][VectorIdx]);
 			}
 		}
 	}
@@ -916,6 +915,7 @@ bool UPoseSearchFeatureChannel_Pose::BuildQuery(UE::PoseSearch::FSearchContext&
 				}
 			}
 		}
+		// else leave the InOutQuery set to zero since the SearchContext.History is invalid and it'll fail if we continue
 		return bSkip;
 	}
 
@@ -925,10 +925,10 @@ bool UPoseSearchFeatureChannel_Pose::BuildQuery(UE::PoseSearch::FSearchContext&
 		FTransform Previous;
 		bool Valid = false;
 	};
-	TArray<CachedTransforms> CachedTransforms;
+	TArray<CachedTransforms, TInlineAllocator<32>> CachedTransforms;
 	CachedTransforms.AddUninitialized(SampleTimes.Num() * SampledBones.Num());
 
-	bool AnyError = false;
+	bool bAnyError = false;
 	for (int32 SubsampleIdx = 0; SubsampleIdx != SampleTimes.Num(); ++SubsampleIdx)
 	{
 		// Stop when we've reached future samples
@@ -941,7 +941,7 @@ bool UPoseSearchFeatureChannel_Pose::BuildQuery(UE::PoseSearch::FSearchContext&
 		for (int32 SampledBoneIdx = 0; SampledBoneIdx != SampledBones.Num(); ++SampledBoneIdx)
 		{
 			const int32 CachedTransformsIndex = SubsampleIdx * SampledBones.Num() + SampledBoneIdx;
-			CachedTransforms[CachedTransformsIndex].Current = SearchContext.TryGetTransformAndCacheResults(SampleTime, InOutQuery.GetSchema(), SchemaBoneIdx[SampledBoneIdx], AnyError);
+			CachedTransforms[CachedTransformsIndex].Current = SearchContext.TryGetTransformAndCacheResults(SampleTime, InOutQuery.GetSchema(), SchemaBoneIdx[SampledBoneIdx], bAnyError);
 			const FPoseSearchBone& SampledBone = SampledBones[SampledBoneIdx];
 
 			if (EnumHasAnyFlags(SampledBone.Flags, EPoseSearchBoneFlags::Velocity))
@@ -949,12 +949,12 @@ bool UPoseSearchFeatureChannel_Pose::BuildQuery(UE::PoseSearch::FSearchContext&
 				check(SearchContext.History);
 				const float HistorySameplInterval = SearchContext.History->GetSampleTimeInterval();
 
-				CachedTransforms[CachedTransformsIndex].Previous = SearchContext.TryGetTransformAndCacheResults(SampleTime - HistorySameplInterval, InOutQuery.GetSchema(), SchemaBoneIdx[SampledBoneIdx], AnyError);
+				CachedTransforms[CachedTransformsIndex].Previous = SearchContext.TryGetTransformAndCacheResults(SampleTime - HistorySameplInterval, InOutQuery.GetSchema(), SchemaBoneIdx[SampledBoneIdx], bAnyError);
 
 				if (!UE::PoseSearch::UseCharacterSpaceVelocities)
 				{
-					const FTransform RootTransform = SearchContext.TryGetTransformAndCacheResults(SampleTime, InOutQuery.GetSchema(), FSearchContext::SchemaRootBoneIdx, AnyError);
-					const FTransform RootTransformPrev = SearchContext.TryGetTransformAndCacheResults(SampleTime - HistorySameplInterval, InOutQuery.GetSchema(), FSearchContext::SchemaRootBoneIdx, AnyError);
+					const FTransform RootTransform = SearchContext.TryGetTransformAndCacheResults(SampleTime, InOutQuery.GetSchema(), FSearchContext::SchemaRootBoneIdx, bAnyError);
+					const FTransform RootTransformPrev = SearchContext.TryGetTransformAndCacheResults(SampleTime - HistorySameplInterval, InOutQuery.GetSchema(), FSearchContext::SchemaRootBoneIdx, bAnyError);
 
 					// animation space velocity
 					CachedTransforms[CachedTransformsIndex].Previous = CachedTransforms[CachedTransformsIndex].Previous * (RootTransformPrev * RootTransform.Inverse());
@@ -964,7 +964,7 @@ bool UPoseSearchFeatureChannel_Pose::BuildQuery(UE::PoseSearch::FSearchContext&
 		}
 	}
 
-	if (AnyError)
+	if (bAnyError)
 	{
 		return false;
 	}
@@ -1043,7 +1043,7 @@ bool UPoseSearchFeatureChannel_Pose::BuildQuery(UE::PoseSearch::FSearchContext&
 	return true;
 }
 
-void UPoseSearchFeatureChannel_Pose::DebugDraw(const UE::PoseSearch::FDebugDrawParams& DrawParams, TArrayView<const float> PoseVector) const
+void UPoseSearchFeatureChannel_Pose::DebugDraw(const UE::PoseSearch::FDebugDrawParams& DrawParams, TConstArrayView<float> PoseVector) const
 {
 	using namespace UE::PoseSearch;
 
@@ -1067,7 +1067,7 @@ void UPoseSearchFeatureChannel_Pose::DebugDraw(const UE::PoseSearch::FDebugDrawP
 	{
 		const FPoseSearchBone& SampledBone = SampledBones[ChannelBoneIdx];
 
-		TArray<FVector> BonePos;
+		TArray<FVector, TInlineAllocator<32>> BonePos;
 		BonePos.AddUninitialized(NumSubsamples);
 		if (EnumHasAnyFlags(SampledBone.Flags, EPoseSearchBoneFlags::Position))
 		{
@@ -1378,7 +1378,7 @@ void UPoseSearchFeatureChannel_Trajectory::IndexAsset(UE::PoseSearch::IAssetInde
 	for (int32 SampleIdx = IndexingContext.BeginSampleIdx; SampleIdx != IndexingContext.EndSampleIdx; ++SampleIdx)
 	{
 		const int32 VectorIdx = SampleIdx - IndexingContext.BeginSampleIdx;
-		IndexAssetPrivate(Indexer, SampleIdx, IndexingOutput.PoseVectors[VectorIdx]);
+		IndexAssetPrivate(Indexer, SampleIdx, IndexingOutput.GetPoseVector(VectorIdx));
 	}
 }
 
@@ -1392,20 +1392,20 @@ void UPoseSearchFeatureChannel_Trajectory::ComputeMeanDeviations(const Eigen::Ma
 	{
 		if (EnumHasAnyFlags(Sample.Flags, EPoseSearchTrajectoryFlags::Position))
 		{
-			FFeatureVectorHelper::ComputeMeanDeviations(MinimumMeanDeviation, CenteredPoseMatrix, MeanDeviations, DataOffset, FFeatureVectorHelper::EncodeVectorCardinality);
+			FFeatureVectorHelper::ComputeMeanDeviations(GetMinimumMeanDeviation(), CenteredPoseMatrix, MeanDeviations, DataOffset, FFeatureVectorHelper::EncodeVectorCardinality);
 		}
 		if (EnumHasAnyFlags(Sample.Flags, EPoseSearchTrajectoryFlags::PositionXY))
 		{
-			FFeatureVectorHelper::ComputeMeanDeviations(MinimumMeanDeviation, CenteredPoseMatrix, MeanDeviations, DataOffset, FFeatureVectorHelper::EncodeVector2DCardinality);
+			FFeatureVectorHelper::ComputeMeanDeviations(GetMinimumMeanDeviation(), CenteredPoseMatrix, MeanDeviations, DataOffset, FFeatureVectorHelper::EncodeVector2DCardinality);
 		}
 
 		if (EnumHasAnyFlags(Sample.Flags, EPoseSearchTrajectoryFlags::Velocity))
 		{
-			FFeatureVectorHelper::ComputeMeanDeviations(MinimumMeanDeviation, CenteredPoseMatrix, MeanDeviations, DataOffset, FFeatureVectorHelper::EncodeVectorCardinality);
+			FFeatureVectorHelper::ComputeMeanDeviations(GetMinimumMeanDeviation(), CenteredPoseMatrix, MeanDeviations, DataOffset, FFeatureVectorHelper::EncodeVectorCardinality);
 		}
 		if (EnumHasAnyFlags(Sample.Flags, EPoseSearchTrajectoryFlags::VelocityXY))
 		{
-			FFeatureVectorHelper::ComputeMeanDeviations(MinimumMeanDeviation, CenteredPoseMatrix, MeanDeviations, DataOffset, FFeatureVectorHelper::EncodeVector2DCardinality);
+			FFeatureVectorHelper::ComputeMeanDeviations(GetMinimumMeanDeviation(), CenteredPoseMatrix, MeanDeviations, DataOffset, FFeatureVectorHelper::EncodeVector2DCardinality);
 		}
 
 		if (EnumHasAnyFlags(Sample.Flags, EPoseSearchTrajectoryFlags::VelocityDirection))
@@ -1447,7 +1447,7 @@ float UPoseSearchFeatureChannel_Trajectory::GetSampleTime(const UE::PoseSearch::
 	return 0.0f;
 }
 
-void UPoseSearchFeatureChannel_Trajectory::IndexAssetPrivate(const UE::PoseSearch::IAssetIndexer& Indexer, int32 SampleIdx, FPoseSearchFeatureVectorBuilder& FeatureVector) const
+void UPoseSearchFeatureChannel_Trajectory::IndexAssetPrivate(const UE::PoseSearch::IAssetIndexer& Indexer, int32 SampleIdx, TArrayView<float> FeatureVector) const
 {
 	// This function samples the instantaneous trajectory at time t as well as the trajectory's velocity and acceleration at time t.
 	// Symmetric finite differences are used to approximate derivatives:
@@ -1502,38 +1502,38 @@ void UPoseSearchFeatureChannel_Trajectory::IndexAssetPrivate(const UE::PoseSearc
 
 		if (EnumHasAnyFlags(Sample.Flags, EPoseSearchTrajectoryFlags::Position))
 		{
-			FFeatureVectorHelper::EncodeVector(FeatureVector.EditValues(), DataOffset, MirroredRootPresent.GetTranslation());
+			FFeatureVectorHelper::EncodeVector(FeatureVector, DataOffset, MirroredRootPresent.GetTranslation());
 		}
 		if (EnumHasAnyFlags(Sample.Flags, EPoseSearchTrajectoryFlags::PositionXY))
 		{
-			FFeatureVectorHelper::EncodeVector2D(FeatureVector.EditValues(), DataOffset, FVector2D(MirroredRootPresent.GetTranslation().X, MirroredRootPresent.GetTranslation().Y));
+			FFeatureVectorHelper::EncodeVector2D(FeatureVector, DataOffset, FVector2D(MirroredRootPresent.GetTranslation().X, MirroredRootPresent.GetTranslation().Y));
 		}
 
 		if (EnumHasAnyFlags(Sample.Flags, EPoseSearchTrajectoryFlags::Velocity))
 		{
-			FFeatureVectorHelper::EncodeVector(FeatureVector.EditValues(), DataOffset, LinearVelocity);
+			FFeatureVectorHelper::EncodeVector(FeatureVector, DataOffset, LinearVelocity);
 		}
 		if (EnumHasAnyFlags(Sample.Flags, EPoseSearchTrajectoryFlags::VelocityXY))
 		{
-			FFeatureVectorHelper::EncodeVector2D(FeatureVector.EditValues(), DataOffset, FVector2D(LinearVelocity.X, LinearVelocity.Y));
+			FFeatureVectorHelper::EncodeVector2D(FeatureVector, DataOffset, FVector2D(LinearVelocity.X, LinearVelocity.Y));
 		}
 
 		if (EnumHasAnyFlags(Sample.Flags, EPoseSearchTrajectoryFlags::VelocityDirection))
 		{
-			FFeatureVectorHelper::EncodeVector(FeatureVector.EditValues(), DataOffset, LinearVelocityDirection);
+			FFeatureVectorHelper::EncodeVector(FeatureVector, DataOffset, LinearVelocityDirection);
 		}
 		if (EnumHasAnyFlags(Sample.Flags, EPoseSearchTrajectoryFlags::VelocityDirectionXY))
 		{
-			FFeatureVectorHelper::EncodeVector2D(FeatureVector.EditValues(), DataOffset, FVector2D(LinearVelocityDirection.X, LinearVelocityDirection.Y));
+			FFeatureVectorHelper::EncodeVector2D(FeatureVector, DataOffset, FVector2D(LinearVelocityDirection.X, LinearVelocityDirection.Y));
 		}
 
 		if (EnumHasAnyFlags(Sample.Flags, EPoseSearchTrajectoryFlags::FacingDirection))
 		{
-			FFeatureVectorHelper::EncodeVector(FeatureVector.EditValues(), DataOffset, FacingDirection);
+			FFeatureVectorHelper::EncodeVector(FeatureVector, DataOffset, FacingDirection);
 		}
 		if (EnumHasAnyFlags(Sample.Flags, EPoseSearchTrajectoryFlags::FacingDirectionXY))
 		{
-			FFeatureVectorHelper::EncodeVector2D(FeatureVector.EditValues(), DataOffset, FVector2D(FacingDirection.X, FacingDirection.Y));
+			FFeatureVectorHelper::EncodeVector2D(FeatureVector, DataOffset, FVector2D(FacingDirection.X, FacingDirection.Y));
 		}
 	}
 	check(DataOffset == ChannelDataOffset + ChannelCardinality);
@@ -1627,10 +1627,10 @@ struct TrajectoryPositionReconstructor
 		float Offset;
 	};
 
-	TArray<PositionAndOffsetSample> PositionAndOffsetSamples;
+	TArray<PositionAndOffsetSample, TInlineAllocator<32>> PositionAndOffsetSamples;
 	bool bInitialized = false;
 
-	void Init(const UPoseSearchFeatureChannel_Trajectory& TrajectoryChannel, TArrayView<const float> PoseVector, const FTransform& RootTransform)
+	void Init(const UPoseSearchFeatureChannel_Trajectory& TrajectoryChannel, TConstArrayView<float> PoseVector, const FTransform& RootTransform)
 	{
 		PositionAndOffsetSamples.Reserve(TrajectoryChannel.Samples.Num() + 1);
 		bool bAddZeroOffsetSample = true;
@@ -1714,7 +1714,7 @@ struct TrajectoryPositionReconstructor
 		check(DataOffset == TrajectoryChannel.GetChannelDataOffset() + TrajectoryChannel.GetChannelCardinality());
 	}
 
-	FVector GetReconstructedTrajectoryPos(const UPoseSearchFeatureChannel_Trajectory& TrajectoryChannel, TArrayView<const float> PoseVector, const FTransform& RootTransform, float SampleOffset)
+	FVector GetReconstructedTrajectoryPos(const UPoseSearchFeatureChannel_Trajectory& TrajectoryChannel, TConstArrayView<float> PoseVector, const FTransform& RootTransform, float SampleOffset)
 	{
 		if (!bInitialized)
 		{
@@ -1753,7 +1753,7 @@ struct TrajectoryPositionReconstructor
 	}
 };
 
-void UPoseSearchFeatureChannel_Trajectory::DebugDraw(const UE::PoseSearch::FDebugDrawParams& DrawParams, TArrayView<const float> PoseVector) const
+void UPoseSearchFeatureChannel_Trajectory::DebugDraw(const UE::PoseSearch::FDebugDrawParams& DrawParams, TConstArrayView<float> PoseVector) const
 {
 	using namespace UE::PoseSearch;
 
@@ -1770,8 +1770,8 @@ void UPoseSearchFeatureChannel_Trajectory::DebugDraw(const UE::PoseSearch::FDebu
 	int32 DataOffset = ChannelDataOffset;
 	int32 SampleIdx = 0;
 	TrajectoryPositionReconstructor TrajectoryPositionReconstructor;
-	TArray<FVector> TrajSplinePos;
-	TArray<FColor> TrajSplineColor;
+	TArray<FVector, TInlineAllocator<32>> TrajSplinePos;
+	TArray<FColor, TInlineAllocator<32>> TrajSplineColor;
 	for (const FPoseSearchTrajectorySample& Sample : Samples)
 	{
 		bool bIsTrajectoryPosValid = false;
@@ -2104,4 +2104,165 @@ void UPoseSearchFeatureChannel_Trajectory::ComputeCostBreakdowns(UE::PoseSearch:
 
 #endif // WITH_EDITOR
 
+//////////////////////////////////////////////////////////////////////////
+// UPoseSearchFeatureChannel_FilterCrashingLegs
+static void ComputeThighsSideAndForward(const FVector& RightThighPos, const FVector& LeftThighPos, FVector& ThighsSide, FVector& ThighsForward)
+{
+	ThighsSide = (RightThighPos - LeftThighPos).GetSafeNormal2D(UE_KINDA_SMALL_NUMBER, FVector::LeftVector);
+	ThighsForward = (ThighsSide ^ FVector::UpVector).GetSafeNormal2D(UE_KINDA_SMALL_NUMBER, FVector::ForwardVector);
+}
+
+static float ComputeCrashingLegsValue(const FVector& RightThighPos, const FVector& LeftThighPos, const FVector& RightFootPos, const FVector& LeftFootPos)
+{
+	FVector ThighsSide;
+	FVector ThighsForward;
+	ComputeThighsSideAndForward(RightThighPos, LeftThighPos, ThighsSide, ThighsForward);
+
+	const FVector FeetDir = (RightFootPos - LeftFootPos).GetSafeNormal2D(UE_KINDA_SMALL_NUMBER, FVector::LeftVector);
+	const float SideDot = FeetDir | ThighsSide;
+	const float ForwardDot = FeetDir | ThighsForward;
+	const float CrashingLegsValue = FMath::Atan2(ForwardDot, SideDot) / UE_PI;
+	return CrashingLegsValue;
+}
+
+void UPoseSearchFeatureChannel_FilterCrashingLegs::InitializeSchema(UE::PoseSearch::FSchemaInitializer& Initializer)
+{
+	Super::InitializeSchema(Initializer);
+	ChannelCardinality = UE::PoseSearch::FFeatureVectorHelper::EncodeFloatCardinality;
+	Initializer.SetCurrentChannelDataOffset(ChannelDataOffset + ChannelCardinality);
+
+	LeftThighIdx = Initializer.AddBoneReference(LeftThigh);
+	RightThighIdx = Initializer.AddBoneReference(RightThigh);
+	LeftFootIdx = Initializer.AddBoneReference(LeftFoot);
+	RightFootIdx = Initializer.AddBoneReference(RightFoot);
+}
+
+void UPoseSearchFeatureChannel_FilterCrashingLegs::FillWeights(TArray<float>& Weights) const
+{
+	check(ChannelCardinality == UE::PoseSearch::FFeatureVectorHelper::EncodeFloatCardinality);
+	Weights[ChannelDataOffset] = Weight;
+}
+
+void UPoseSearchFeatureChannel_FilterCrashingLegs::IndexAsset(UE::PoseSearch::IAssetIndexer& Indexer, UE::PoseSearch::FAssetIndexingOutput& IndexingOutput) const
+{
+	using namespace UE::PoseSearch;
+
+	const FAssetIndexingContext& IndexingContext = Indexer.GetIndexingContext();
+
+	for (int32 SampleIdx = IndexingContext.BeginSampleIdx; SampleIdx != IndexingContext.EndSampleIdx; ++SampleIdx)
+	{
+		const float SubsampleTime = SampleIdx * IndexingContext.Schema->GetSamplingInterval();
+
+		bool bUnused;
+		const FTransform RightThighTransform = Indexer.GetTransformAndCacheResults(SubsampleTime, SubsampleTime, RightThighIdx, bUnused);
+		const FTransform LeftThighTransform = Indexer.GetTransformAndCacheResults(SubsampleTime, SubsampleTime, LeftThighIdx, bUnused);
+		const FTransform RightFootTransform = Indexer.GetTransformAndCacheResults(SubsampleTime, SubsampleTime, RightFootIdx, bUnused);
+		const FTransform LeftFootTransform = Indexer.GetTransformAndCacheResults(SubsampleTime, SubsampleTime, LeftFootIdx, bUnused);
+
+		const float CrashingLegsValue = ComputeCrashingLegsValue(RightThighTransform.GetTranslation(), LeftThighTransform.GetTranslation(), RightFootTransform.GetTranslation(), LeftFootTransform.GetTranslation());
+		
+		int32 DataOffset = ChannelDataOffset;
+		const int32 VectorIdx = SampleIdx - IndexingContext.BeginSampleIdx;
+		FFeatureVectorHelper::EncodeFloat(IndexingOutput.GetPoseVector(VectorIdx), DataOffset, CrashingLegsValue);
+	}
+}
+
+bool UPoseSearchFeatureChannel_FilterCrashingLegs::BuildQuery(UE::PoseSearch::FSearchContext& SearchContext, FPoseSearchFeatureVectorBuilder& InOutQuery) const
+{
+	using namespace UE::PoseSearch;
+
+	const bool bIsCurrentResultValid = SearchContext.CurrentResult.IsValid();
+	const bool bSkip = InputQueryPose != EInputQueryPose::UseCharacterPose && bIsCurrentResultValid && SearchContext.CurrentResult.Database->Schema == InOutQuery.GetSchema();
+	
+	int32 DataOffset = ChannelDataOffset;
+	if (bSkip || !SearchContext.History)
+	{
+		if (bIsCurrentResultValid)
+		{
+			const float LerpValue = InputQueryPose == EInputQueryPose::UseInterpolatedContinuingPose ? SearchContext.CurrentResult.LerpValue : 0.f;
+			FFeatureVectorHelper::EncodeFloat(InOutQuery.EditValues(), DataOffset, SearchContext.GetCurrentResultPrevPoseVector(), SearchContext.GetCurrentResultPoseVector(), SearchContext.GetCurrentResultNextPoseVector(), LerpValue);
+		}
+		// else leave the InOutQuery set to zero since the SearchContext.History is invalid and it'll fail if we continue
+		return bSkip;
+	}
+
+	const float SampleTime = 0.f;
+	bool bAnyError = false;
+	const FTransform LeftThighTransform = SearchContext.TryGetTransformAndCacheResults(SampleTime, InOutQuery.GetSchema(), LeftThighIdx, bAnyError);
+	const FTransform RightThighTransform = SearchContext.TryGetTransformAndCacheResults(SampleTime, InOutQuery.GetSchema(), RightThighIdx, bAnyError);
+	const FTransform LeftFootTransform = SearchContext.TryGetTransformAndCacheResults(SampleTime, InOutQuery.GetSchema(), LeftFootIdx, bAnyError);
+	const FTransform RightFootTransform = SearchContext.TryGetTransformAndCacheResults(SampleTime, InOutQuery.GetSchema(), RightFootIdx, bAnyError);
+
+	const float CrashingLegsValue = ComputeCrashingLegsValue(RightThighTransform.GetTranslation(), LeftThighTransform.GetTranslation(), RightFootTransform.GetTranslation(), LeftFootTransform.GetTranslation());
+	FFeatureVectorHelper::EncodeFloat(InOutQuery.EditValues(), DataOffset, CrashingLegsValue);
+	return !bAnyError;
+}
+
+void UPoseSearchFeatureChannel_FilterCrashingLegs::DebugDraw(const UE::PoseSearch::FDebugDrawParams& DrawParams, TConstArrayView<float> PoseVector) const
+{
+	using namespace UE::PoseSearch;
+
+	int32 DataOffset = ChannelDataOffset;
+	const float CrashingLegsValue = FFeatureVectorHelper::DecodeFloat(PoseVector, DataOffset);
+
+	const float LifeTime = DrawParams.DefaultLifeTime;
+	const uint8 DepthPriority = ESceneDepthPriorityGroup::SDPG_Foreground + 2;
+	const bool bPersistent = EnumHasAnyFlags(DrawParams.Flags, EDebugDrawFlags::Persistent);
+
+	const FTransform LeftThighTransform = DrawParams.Mesh != nullptr ? DrawParams.Mesh->GetSocketTransform(LeftThigh.BoneName) : DrawParams.RootTransform;
+	const FTransform RightThighTransform = DrawParams.Mesh != nullptr ? DrawParams.Mesh->GetSocketTransform(RightThigh.BoneName) : DrawParams.RootTransform;
+	const FTransform LeftFootTransform = DrawParams.Mesh != nullptr ? DrawParams.Mesh->GetSocketTransform(LeftFoot.BoneName) : DrawParams.RootTransform;
+	const FTransform RightFootTransform = DrawParams.Mesh != nullptr ? DrawParams.Mesh->GetSocketTransform(RightFoot.BoneName) : DrawParams.RootTransform;
+	
+	const float FeetDistance = (RightFootTransform.GetTranslation() - LeftFootTransform.GetTranslation()).Length();
+
+	FVector ThighsSide;
+	FVector ThighsForward;
+	ComputeThighsSideAndForward(RightThighTransform.GetTranslation(), LeftThighTransform.GetTranslation(), ThighsSide, ThighsForward);
+	const FVector CrossingLegsVector = (FMath::Cos(CrashingLegsValue * UE_PI) * ThighsSide + FMath::Sin(CrashingLegsValue * UE_PI) * ThighsForward) * FeetDistance;
+
+	auto LerpColor = [](FColor A, FColor B, float T) -> FColor
+	{
+		return FColor(
+			FMath::RoundToInt(float(A.R) * (1.f - T) + float(B.R) * T),
+			FMath::RoundToInt(float(A.G) * (1.f - T) + float(B.G) * T),
+			FMath::RoundToInt(float(A.B) * (1.f - T) + float(B.B) * T),
+			FMath::RoundToInt(float(A.A) * (1.f - T) + float(B.A) * T));
+	};
+
+	// when CrossingLegsValue is greater than .5 or lesser than -.5 we draw in Red or Orange
+	const float LerpValue = FMath::Min(CrashingLegsValue * 2.f, 1.f);
+	const FColor Color = LerpValue >= 0.f ?
+		LerpColor(FColor::Green, FColor::Red, LerpValue) :
+		LerpColor(FColor::Green, FColor::Orange, -LerpValue);
+
+	DrawDebugLine(DrawParams.World, LeftFootTransform.GetTranslation(), LeftFootTransform.GetTranslation() + CrossingLegsVector, Color, bPersistent, LifeTime, DepthPriority);
+	DrawDebugLine(DrawParams.World, RightFootTransform.GetTranslation(), RightFootTransform.GetTranslation() - CrossingLegsVector, Color, bPersistent, LifeTime, DepthPriority);
+}
+
+// IPoseFilter interface
+bool UPoseSearchFeatureChannel_FilterCrashingLegs::IsPoseFilterActive() const
+{
+	return true;
+}
+
+bool UPoseSearchFeatureChannel_FilterCrashingLegs::IsPoseValid(TConstArrayView<float> PoseValues, TConstArrayView<float> QueryValues, int32 PoseIdx, const FPoseSearchPoseMetadata& Metadata) const
+{
+	using namespace UE::PoseSearch;
+
+	int32 DataOffset = ChannelDataOffset;
+	const float PoseValue = FFeatureVectorHelper::DecodeFloat(PoseValues, DataOffset);
+
+	DataOffset = ChannelDataOffset;
+	const float QueryValue = FFeatureVectorHelper::DecodeFloat(QueryValues, DataOffset);
+
+	bool bIsPoseValid = true;
+	if (FMath::Abs(QueryValue - PoseValue) > AllowedTolerance)
+	{
+		bIsPoseValid = false;
+	}
+
+	return bIsPoseValid;
+}
+
 #undef LOCTEXT_NAMESPACE
\ No newline at end of file
diff --git a/Engine/Plugins/Experimental/Animation/PoseSearch/Source/Runtime/Public/PoseSearch/PoseSearch.h b/Engine/Plugins/Experimental/Animation/PoseSearch/Source/Runtime/Public/PoseSearch/PoseSearch.h
--- a/Engine/Plugins/Experimental/Animation/PoseSearch/Source/Runtime/Public/PoseSearch/PoseSearch.h
+++ b/Engine/Plugins/Experimental/Animation/PoseSearch/Source/Runtime/Public/PoseSearch/PoseSearch.h
@@ -358,9 +358,20 @@ struct FAssetIndexingContext
  */
 struct FAssetIndexingOutput
 {
-	// Channel data should be written to this array of feature vector builders
-	// Size is EndPoseIdx - BeginPoseIdx and PoseVectors[0] contains data for BeginPoseIdx
-	const TArrayView<FPoseSearchFeatureVectorBuilder> PoseVectors;
+	FAssetIndexingOutput(const int32 InSchemaCardinality, TArrayView<float> InFeatureVectorTable)
+	: SchemaCardinality(InSchemaCardinality)
+	, FeatureVectorTable(InFeatureVectorTable)
+	{
+	}
+
+	TArrayView<float> GetPoseVector(int32 VectorIdx)
+	{
+		return MakeArrayView(&FeatureVectorTable[VectorIdx * SchemaCardinality], SchemaCardinality);
+	}
+
+private:
+	const int32 SchemaCardinality;
+	TArrayView<float> FeatureVectorTable;
 };
 
 class POSESEARCH_API IAssetIndexer
@@ -419,11 +430,23 @@ public:
 
 } // namespace UE::PoseSearch
 
+class POSESEARCH_API IPoseFilter
+{
+public:
+	virtual ~IPoseFilter() {}
+
+	// if true this filter will be evaluated
+	virtual bool IsPoseFilterActive() const { return false; }
+	
+	// if it returns false the pose candidate will be discarded
+	virtual bool IsPoseValid(TConstArrayView<float> PoseValues, TConstArrayView<float> QueryValues, int32 PoseIdx, const FPoseSearchPoseMetadata& Metadata) const { return true; }
+};
+
 //////////////////////////////////////////////////////////////////////////
 // Feature channels interface
 
 UCLASS(Abstract, BlueprintType, EditInlineNew)
-class POSESEARCH_API UPoseSearchFeatureChannel : public UObject, public IBoneReferenceSkeletonProvider
+class POSESEARCH_API UPoseSearchFeatureChannel : public UObject, public IBoneReferenceSkeletonProvider, public IPoseFilter
 {
 	GENERATED_BODY()
 
@@ -450,15 +473,14 @@ public:
 	virtual bool BuildQuery(UE::PoseSearch::FSearchContext& SearchContext, FPoseSearchFeatureVectorBuilder& InOutQuery) const PURE_VIRTUAL(UPoseSearchFeatureChannel::BuildQuery, return false;);
 
 	// Draw this channel's data for the given pose vector
-	virtual void DebugDraw(const UE::PoseSearch::FDebugDrawParams& DrawParams, TArrayView<const float> PoseVector) const PURE_VIRTUAL(UPoseSearchFeatureChannel::DebugDraw, );
+	virtual void DebugDraw(const UE::PoseSearch::FDebugDrawParams& DrawParams, TConstArrayView<float> PoseVector) const PURE_VIRTUAL(UPoseSearchFeatureChannel::DebugDraw, );
 
 #if WITH_EDITOR
 	virtual void ComputeCostBreakdowns(UE::PoseSearch::ICostBreakDownData& CostBreakDownData, const UPoseSearchSchema* Schema) const;
 #endif
 
 	// Used during data normalization. If a feature has less than this amount of deviation from the mean across all poses in a database, then it will not be normalized.
-	UPROPERTY(EditAnywhere, Category = "Settings", meta = (ClampMin = "0.0001"))
-	float MinimumMeanDeviation = 1.0f;
+	virtual float GetMinimumMeanDeviation() const { return 0.1f; }
 
 private:
 	// IBoneReferenceSkeletonProvider interface
@@ -530,7 +552,7 @@ public:
 	TObjectPtr<USkeleton> Skeleton = nullptr;
 
 	UPROPERTY(EditAnywhere, meta = (ClampMin = "1", ClampMax = "240"), Category = "Schema")
-	int32 SampleRate = 10;
+	int32 SampleRate = 30;
 
 	UPROPERTY(EditAnywhere, Instanced, BlueprintReadWrite, Category = "Schema")
 	TArray<TObjectPtr<UPoseSearchFeatureChannel>> Channels;
@@ -724,6 +746,9 @@ struct POSESEARCH_API FPoseSearchIndex
 	UPROPERTY()
 	TArray<FPoseSearchPoseMetadata> PoseMetadata;
 
+	UPROPERTY()
+	EPoseSearchPoseFlags OverallFlags = EPoseSearchPoseFlags::None;
+
 	UPROPERTY()
 	TObjectPtr<const UPoseSearchSchema> Schema = nullptr;
 
@@ -906,7 +931,7 @@ public:
 	const UPoseSearchSchema* GetSchema() const { return Schema.Get(); }
 
 	TArray<float>& EditValues() { return Values; }
-	TArrayView<const float> GetValues() const { return Values; }
+	TConstArrayView<float> GetValues() const { return Values; }
 
 private:
 	UPROPERTY(Transient)
@@ -1058,20 +1083,20 @@ public:
 	TArray<FPoseSearchDatabaseBlendSpace> BlendSpaces;
 
 	UPROPERTY(EditAnywhere, Category = "Performance")
-	EPoseSearchMode PoseSearchMode = EPoseSearchMode::BruteForce;
+	EPoseSearchMode PoseSearchMode = EPoseSearchMode::PCAKDTree;
 
 	UPROPERTY(EditAnywhere, Category = "Performance", meta = (EditCondition = "PoseSearchMode != EPoseSearchMode::BruteForce", EditConditionHides, ClampMin = "1", ClampMax = "64", UIMin = "1", UIMax = "64"))
 	int32 NumberOfPrincipalComponents = 4;
 
 	UPROPERTY(EditAnywhere, Category = "Performance", meta = (EditCondition = "PoseSearchMode != EPoseSearchMode::BruteForce", EditConditionHides, ClampMin = "1", ClampMax = "256", UIMin = "1", UIMax = "256"))
-	int32 KDTreeMaxLeafSize = 8;
+	int32 KDTreeMaxLeafSize = 16;
 	
 	UPROPERTY(EditAnywhere, Category = "Performance", meta = (EditCondition = "PoseSearchMode != EPoseSearchMode::BruteForce", EditConditionHides, ClampMin = "1", ClampMax = "600", UIMin = "1", UIMax = "600"))
-	int32 KDTreeQueryNumNeighbors = 100;
+	int32 KDTreeQueryNumNeighbors = 200;
 
 	// if true, this database search will be skipped if cannot decrease the pose cost, and poses will not be listed into the PoseSearchDebugger
 	UPROPERTY(EditAnywhere, Category = "Performance")
-	bool bSkipSearchIfPossible = false;
+	bool bSkipSearchIfPossible = true;
 
 	virtual ~UPoseSearchDatabase();
 
@@ -1214,8 +1239,9 @@ enum class EPoseCandidateFlags : uint8
 	DiscardedBy_PoseJumpThresholdTime = 1 << 3,
 	DiscardedBy_PoseReselectHistory = 1 << 4,
 	DiscardedBy_BlockTransition = 1 << 5,
+	DiscardedBy_PoseFilter = 1 << 6,
 
-	AnyDiscardedMask = DiscardedBy_PoseJumpThresholdTime | DiscardedBy_PoseReselectHistory | DiscardedBy_BlockTransition,
+	AnyDiscardedMask = DiscardedBy_PoseJumpThresholdTime | DiscardedBy_PoseReselectHistory | DiscardedBy_BlockTransition | DiscardedBy_PoseFilter,
 };
 ENUM_CLASS_FLAGS(EPoseCandidateFlags);
 
@@ -1225,18 +1251,23 @@ class POSESEARCH_API FFeatureVectorHelper
 public:
 	enum { EncodeQuatCardinality = 6 };
 	static void EncodeQuat(TArrayView<float> Values, int32& DataOffset, const FQuat& Quat);
-	static void EncodeQuat(TArrayView<float> Values, int32& DataOffset, TArrayView<const float> PrevValues, TArrayView<const float> CurValues, TArrayView<const float> NextValues, float LerpValue);
-	static FQuat DecodeQuat(TArrayView<const float> Values, int32& DataOffset);
+	static void EncodeQuat(TArrayView<float> Values, int32& DataOffset, TConstArrayView<float> PrevValues, TConstArrayView<float> CurValues, TConstArrayView<float> NextValues, float LerpValue);
+	static FQuat DecodeQuat(TConstArrayView<float> Values, int32& DataOffset);
 
 	enum { EncodeVectorCardinality = 3 };
 	static void EncodeVector(TArrayView<float> Values, int32& DataOffset, const FVector& Vector);
-	static void EncodeVector(TArrayView<float> Values, int32& DataOffset, TArrayView<const float> PrevValues, TArrayView<const float> CurValues, TArrayView<const float> NextValues, float LerpValue, bool bNormalize = false);
-	static FVector DecodeVector(TArrayView<const float> Values, int32& DataOffset);
+	static void EncodeVector(TArrayView<float> Values, int32& DataOffset, TConstArrayView<float> PrevValues, TConstArrayView<float> CurValues, TConstArrayView<float> NextValues, float LerpValue, bool bNormalize = false);
+	static FVector DecodeVector(TConstArrayView<float> Values, int32& DataOffset);
 
 	enum { EncodeVector2DCardinality = 2 };
 	static void EncodeVector2D(TArrayView<float> Values, int32& DataOffset, const FVector2D& Vector2D);
-	static void EncodeVector2D(TArrayView<float> Values, int32& DataOffset, TArrayView<const float> PrevValues, TArrayView<const float> CurValues, TArrayView<const float> NextValues, float LerpValue);
-	static FVector2D DecodeVector2D(TArrayView<const float> Values, int32& DataOffset);
+	static void EncodeVector2D(TArrayView<float> Values, int32& DataOffset, TConstArrayView<float> PrevValues, TConstArrayView<float> CurValues, TConstArrayView<float> NextValues, float LerpValue);
+	static FVector2D DecodeVector2D(TConstArrayView<float> Values, int32& DataOffset);
+
+	enum { EncodeFloatCardinality = 1 };
+	static void EncodeFloat(TArrayView<float> Values, int32& DataOffset, const float Value);
+	static void EncodeFloat(TArrayView<float> Values, int32& DataOffset, TConstArrayView<float> PrevValues, TConstArrayView<float> CurValues, TConstArrayView<float> NextValues, float LerpValue);
+	static float DecodeFloat(TConstArrayView<float> Values, int32& DataOffset);
 
 	// populates MeanDeviations[DataOffset] ... MeanDeviations[DataOffset + Cardinality] with a single value the mean deviation calculated from a centered matrix
 	static void ComputeMeanDeviations(float MinMeanDeviation, const Eigen::MatrixXd& CenteredPoseMatrix, Eigen::VectorXd& MeanDeviations, int32& DataOffset, int32 Cardinality);
@@ -1245,9 +1276,11 @@ public:
 	static void SetMeanDeviations(float Deviation, Eigen::VectorXd& MeanDeviations, int32& DataOffset, int32 Cardinality);
 
 private:
-	static FQuat DecodeQuatInternal(TArrayView<const float> Values, int32 DataOffset);
-	static FVector DecodeVectorInternal(TArrayView<const float> Values, int32 DataOffset);
-	static FVector2D DecodeVector2DInternal(TArrayView<const float> Values, int32 DataOffset);
+	static FQuat DecodeQuatInternal(TConstArrayView<float> Values, int32 DataOffset);
+	static FVector DecodeVectorInternal(TConstArrayView<float> Values, int32 DataOffset);
+	static FVector2D DecodeVector2DInternal(TConstArrayView<float> Values, int32 DataOffset);
+	static float DecodeFloatInternal(TConstArrayView<float> Values, int32 DataOffset);
+
 };
 
 /**
@@ -1370,7 +1403,7 @@ private:
 		FPoseSearchFeatureVectorBuilder FeatureVectorBuilder;
 	};
 
-	TArray<FCachedQuery> CachedQueries;
+	TArray<FCachedQuery, TInlineAllocator<8>> CachedQueries;
 
 	float CurrentBestTotalCost = MAX_flt;
 
@@ -1427,7 +1460,7 @@ public:
 //////////////////////////////////////////////////////////////////////////
 // Main PoseSearch API
 
-POSESEARCH_API void DrawFeatureVector(const FDebugDrawParams& DrawParams, TArrayView<const float> PoseVector);
+POSESEARCH_API void DrawFeatureVector(const FDebugDrawParams& DrawParams, TConstArrayView<float> PoseVector);
 POSESEARCH_API void DrawFeatureVector(const FDebugDrawParams& DrawParams, int32 PoseIdx);
 POSESEARCH_API void DrawSearchIndex(const FDebugDrawParams& DrawParams);
 POSESEARCH_API void CompareFeatureVectors(TConstArrayView<float> A, TConstArrayView<float> B, TConstArrayView<float> WeightsSqrt, TArrayView<float> Result);
diff --git a/Engine/Plugins/Experimental/Animation/PoseSearch/Source/Runtime/Public/PoseSearch/PoseSearchFeatureChannels.h b/Engine/Plugins/Experimental/Animation/PoseSearch/Source/Runtime/Public/PoseSearch/PoseSearchFeatureChannels.h
--- a/Engine/Plugins/Experimental/Animation/PoseSearch/Source/Runtime/Public/PoseSearch/PoseSearchFeatureChannels.h
+++ b/Engine/Plugins/Experimental/Animation/PoseSearch/Source/Runtime/Public/PoseSearch/PoseSearchFeatureChannels.h
@@ -42,7 +42,7 @@ public:
 	float SampleTimeOffset = 0.f;
 
 	UPROPERTY()
-	int8 SchemaBoneIdx;
+	int8 SchemaBoneIdx = 0;
 
 	UPROPERTY(EditAnywhere, Category = "Settings")
 	int32 ColorPresetIndex = 0;
@@ -60,7 +60,7 @@ public:
 	virtual void FillWeights(TArray<float>& Weights) const override;
 	virtual void IndexAsset(UE::PoseSearch::IAssetIndexer& Indexer, UE::PoseSearch::FAssetIndexingOutput& IndexingOutput) const override;
 	virtual bool BuildQuery(UE::PoseSearch::FSearchContext& SearchContext, FPoseSearchFeatureVectorBuilder& InOutQuery) const override;
-	virtual void DebugDraw(const UE::PoseSearch::FDebugDrawParams& DrawParams, TArrayView<const float> PoseVector) const override;
+	virtual void DebugDraw(const UE::PoseSearch::FDebugDrawParams& DrawParams, TConstArrayView<float> PoseVector) const override;
 };
 
 //////////////////////////////////////////////////////////////////////////
@@ -96,7 +96,7 @@ public:
 	EHeadingAxis HeadingAxis = EHeadingAxis::X;	
 
 	UPROPERTY()
-	int8 SchemaBoneIdx;
+	int8 SchemaBoneIdx = 0;
 
 	UPROPERTY(EditAnywhere, Category = "Settings")
 	int32 ColorPresetIndex = 0;
@@ -114,7 +114,7 @@ public:
 	virtual void FillWeights(TArray<float>& Weights) const override;
 	virtual void IndexAsset(UE::PoseSearch::IAssetIndexer& Indexer, UE::PoseSearch::FAssetIndexingOutput& IndexingOutput) const override;
 	virtual bool BuildQuery(UE::PoseSearch::FSearchContext& SearchContext, FPoseSearchFeatureVectorBuilder& InOutQuery) const override;
-	virtual void DebugDraw(const UE::PoseSearch::FDebugDrawParams& DrawParams, TArrayView<const float> PoseVector) const override;
+	virtual void DebugDraw(const UE::PoseSearch::FDebugDrawParams& DrawParams, TConstArrayView<float> PoseVector) const override;
 
 	FVector GetAxis(const FQuat& Rotation) const;
 };
@@ -181,14 +181,14 @@ public:
 	virtual void IndexAsset(UE::PoseSearch::IAssetIndexer& Indexer, UE::PoseSearch::FAssetIndexingOutput& IndexingOutput) const override;
 	virtual void ComputeMeanDeviations(const Eigen::MatrixXd& CenteredPoseMatrix, Eigen::VectorXd& MeanDeviations) const override;
 	virtual bool BuildQuery(UE::PoseSearch::FSearchContext& SearchContext, FPoseSearchFeatureVectorBuilder& InOutQuery) const override;
-	virtual void DebugDraw(const UE::PoseSearch::FDebugDrawParams& DrawParams, TArrayView<const float> PoseVector) const override;
+	virtual void DebugDraw(const UE::PoseSearch::FDebugDrawParams& DrawParams, TConstArrayView<float> PoseVector) const override;
 
 #if WITH_EDITOR
 	virtual void ComputeCostBreakdowns(UE::PoseSearch::ICostBreakDownData& CostBreakDownData, const UPoseSearchSchema* Schema) const override;
 #endif
 
 protected:
-	void AddPoseFeatures(UE::PoseSearch::IAssetIndexer& Indexer, int32 SampleIdx, FPoseSearchFeatureVectorBuilder& FeatureVector, const TArray<TArray<FVector2D>>& Phases) const;
+	void AddPoseFeatures(UE::PoseSearch::IAssetIndexer& Indexer, int32 SampleIdx, TArrayView<float> FeatureVector, const TArray<TArray<FVector2D>>& Phases) const;
 	void CalculatePhases(UE::PoseSearch::IAssetIndexer& Indexer, UE::PoseSearch::FAssetIndexingOutput& IndexingOutput, TArray<TArray<FVector2D>>& OutPhases) const;
 };
 
@@ -262,13 +262,73 @@ public:
 	virtual void IndexAsset(UE::PoseSearch::IAssetIndexer& Indexer, UE::PoseSearch::FAssetIndexingOutput& IndexingOutput) const override;
 	virtual void ComputeMeanDeviations(const Eigen::MatrixXd& CenteredPoseMatrix, Eigen::VectorXd& MeanDeviations) const override;
 	virtual bool BuildQuery(UE::PoseSearch::FSearchContext& SearchContext, FPoseSearchFeatureVectorBuilder& InOutQuery) const override;
-	virtual void DebugDraw(const UE::PoseSearch::FDebugDrawParams& DrawParams, TArrayView<const float> PoseVector) const override;
+	virtual void DebugDraw(const UE::PoseSearch::FDebugDrawParams& DrawParams, TConstArrayView<float> PoseVector) const override;
 
 #if WITH_EDITOR
 	virtual void ComputeCostBreakdowns(UE::PoseSearch::ICostBreakDownData& CostBreakDownData, const UPoseSearchSchema* Schema) const override;
 #endif
 
 protected:
-	void IndexAssetPrivate(const UE::PoseSearch::IAssetIndexer& Indexer, int32 SampleIdx, FPoseSearchFeatureVectorBuilder& FeatureVector) const;
+	void IndexAssetPrivate(const UE::PoseSearch::IAssetIndexer& Indexer, int32 SampleIdx, TArrayView<float> FeatureVector) const;
 	float GetSampleTime(const UE::PoseSearch::IAssetIndexer& Indexer, float Offset, float SampleTime, float RootDistance) const;
+};
+
+//////////////////////////////////////////////////////////////////////////
+// UPoseSearchFeatureChannel_FilterCrashingLegs
+// the idea is to calculate the angle between the direction from LeftThigh position to RightThigh position and the direction from LeftFoot position to RightFoot position, and divide it by PI to have values in range [-1,1]
+// the number (called 'CrashingLegsValue' calculated in ComputeCrashingLegsValue) is gonna be
+// 0 if the feet are aligned with the thighs (for example in an idle standing position)
+// 0.5 if the right foot is exactly in front of the left foot (for example when a character is running  following a line)
+// -0.5 if the left foot is exactly in front of the right foot
+// close to 1 or -1 if the feet (and so the legs) are completely crossed
+// at runtime we'll match the CrashingLegsValue and also filter by discarding pose candidates that don't respect the 'AllowedTolerance' between query and database values (happening in IsPoseValid)
+UCLASS(BlueprintType, EditInlineNew)
+class POSESEARCH_API UPoseSearchFeatureChannel_FilterCrashingLegs : public UPoseSearchFeatureChannel
+{
+	GENERATED_BODY()
+
+public:
+	UPROPERTY(EditAnywhere, Category = "Settings")
+	FBoneReference LeftThigh;
+	
+	UPROPERTY(EditAnywhere, Category = "Settings")
+	FBoneReference RightThigh;
+
+	UPROPERTY(EditAnywhere, Category = "Settings")
+	FBoneReference LeftFoot;
+	
+	UPROPERTY(EditAnywhere, Category = "Settings")
+	FBoneReference RightFoot;
+
+	UPROPERTY(EditAnywhere, Category = "Settings")
+	float Weight = 0.2f;
+
+	UPROPERTY()
+	int8 LeftThighIdx;
+
+	UPROPERTY()
+	int8 RightThighIdx;
+
+	UPROPERTY()
+	int8 LeftFootIdx;
+
+	UPROPERTY()
+	int8 RightFootIdx;
+
+	UPROPERTY(EditAnywhere, Category = "Settings")
+	EInputQueryPose InputQueryPose = EInputQueryPose::UseContinuingPose;
+
+	UPROPERTY(EditAnywhere, Category = "Settings")
+	float AllowedTolerance = 0.3f;
+
+	// UPoseSearchFeatureChannel interface
+	virtual void InitializeSchema(UE::PoseSearch::FSchemaInitializer& Initializer) override;
+	virtual void FillWeights(TArray<float>& Weights) const override;
+	virtual void IndexAsset(UE::PoseSearch::IAssetIndexer& Indexer, UE::PoseSearch::FAssetIndexingOutput& IndexingOutput) const override;
+	virtual bool BuildQuery(UE::PoseSearch::FSearchContext& SearchContext, FPoseSearchFeatureVectorBuilder& InOutQuery) const override;
+	virtual void DebugDraw(const UE::PoseSearch::FDebugDrawParams& DrawParams, TConstArrayView<float> PoseVector) const override;
+
+	// IPoseFilter interface
+	virtual bool IsPoseFilterActive() const override;
+	virtual bool IsPoseValid(TConstArrayView<float> PoseValues, TConstArrayView<float> QueryValues, int32 PoseIdx, const FPoseSearchPoseMetadata& Metadata) const override;
 };
\ No newline at end of file
diff --git a/Engine/Plugins/Experimental/Animation/PoseSearch/Source/Runtime/Public/PoseSearch/PoseSearchLibrary.h b/Engine/Plugins/Experimental/Animation/PoseSearch/Source/Runtime/Public/PoseSearch/PoseSearchLibrary.h
--- a/Engine/Plugins/Experimental/Animation/PoseSearch/Source/Runtime/Public/PoseSearch/PoseSearchLibrary.h
+++ b/Engine/Plugins/Experimental/Animation/PoseSearch/Source/Runtime/Public/PoseSearch/PoseSearchLibrary.h
@@ -47,7 +47,7 @@ struct POSESEARCH_API FMotionMatchingSettings
 	
 	// Don't jump to poses that are less than this many seconds away
 	UPROPERTY(EditAnywhere, BlueprintReadWrite, Category=Settings, meta=(ClampMin="0"))
-	float PoseJumpThresholdTime = 1.f;
+	float PoseJumpThresholdTime = 0.f;
 
 	// Don't jump to poses that has been selected previously within this many seconds in the past
 	UPROPERTY(EditAnywhere, BlueprintReadWrite, Category = Settings, meta = (ClampMin = "0"))
@@ -55,7 +55,7 @@ struct POSESEARCH_API FMotionMatchingSettings
 
 	// Minimum amount of time to wait between pose search queries
 	UPROPERTY(EditAnywhere, BlueprintReadWrite, Category=Settings, meta=(ClampMin="0"))
-	float SearchThrottleTime = 0.1f;
+	float SearchThrottleTime = 0.f;
 };
 
 USTRUCT(BlueprintType, Category="Animation|Pose Search")
