motion matching - dynamic PlayRate

it's about time scaling the animation playback rate proportionally to the difference in "estimated" velocity between query and result: if we request a trajectory that goes at 4 m/s and the best candidates returns 3 m/s we try to play the animation 1.33 times faster (but then clamp to a min max PlayRate)

the idea is to help out stride warping not to make huge steps

diff --git a/Engine/Plugins/Experimental/Animation/PoseSearch/Source/Runtime/Private/AnimNode_BlendStack.cpp b/Engine/Plugins/Experimental/Animation/PoseSearch/Source/Runtime/Private/AnimNode_BlendStack.cpp
--- a/Engine/Plugins/Experimental/Animation/PoseSearch/Source/Runtime/Private/AnimNode_BlendStack.cpp
+++ b/Engine/Plugins/Experimental/Animation/PoseSearch/Source/Runtime/Private/AnimNode_BlendStack.cpp
@@ -15,7 +15,7 @@ TAutoConsoleVariable<int32> CVarAnimBlendStackPruningEnable(TEXT("a.AnimNode.Ble
 
 /////////////////////////////////////////////////////
 // FPoseSearchAnimPlayer
-void FPoseSearchAnimPlayer::Initialize(ESearchIndexAssetType InAssetType, UAnimationAsset* AnimationAsset, float AccumulatedTime, bool bLoop, bool bMirrored, UMirrorDataTable* MirrorDataTable, float BlendTime, const UBlendProfile* BlendProfile, EAlphaBlendOption InBlendOption, FVector BlendParameters)
+void FPoseSearchAnimPlayer::Initialize(ESearchIndexAssetType InAssetType, UAnimationAsset* AnimationAsset, float AccumulatedTime, bool bLoop, bool bMirrored, UMirrorDataTable* MirrorDataTable, float BlendTime, const UBlendProfile* BlendProfile, EAlphaBlendOption InBlendOption, FVector BlendParameters, float PlayRate)
 {
 	check(AnimationAsset);
 
@@ -57,7 +57,7 @@ void FPoseSearchAnimPlayer::Initialize(ESearchIndexAssetType InAssetType, UAnima
 		SequencePlayerNode.SetAccumulatedTime(AccumulatedTime);
 		SequencePlayerNode.SetSequence(Sequence);
 		SequencePlayerNode.SetLoopAnimation(bLoop);
-		SequencePlayerNode.SetPlayRate(1.0f);
+		SequencePlayerNode.SetPlayRate(PlayRate);
 	}
 	else if (AssetType == ESearchIndexAssetType::BlendSpace)
 	{
@@ -67,7 +67,7 @@ void FPoseSearchAnimPlayer::Initialize(ESearchIndexAssetType InAssetType, UAnima
 		BlendSpacePlayerNode.SetAccumulatedTime(AccumulatedTime);
 		BlendSpacePlayerNode.SetBlendSpace(BlendSpace);
 		BlendSpacePlayerNode.SetLoop(bLoop);
-		BlendSpacePlayerNode.SetPlayRate(1.0f);
+		BlendSpacePlayerNode.SetPlayRate(PlayRate);
 		BlendSpacePlayerNode.SetPosition(BlendParameters);
 	}
 	else 
@@ -77,6 +77,18 @@ void FPoseSearchAnimPlayer::Initialize(ESearchIndexAssetType InAssetType, UAnima
 	UpdateSourceLinkNode();
 }
 
+void FPoseSearchAnimPlayer::UpdatePlayRate(float PlayRate)
+{
+	if (AssetType == ESearchIndexAssetType::Sequence)
+	{
+		SequencePlayerNode.SetPlayRate(PlayRate);
+	}
+	else if (AssetType == ESearchIndexAssetType::BlendSpace)
+	{
+		BlendSpacePlayerNode.SetPlayRate(PlayRate);
+	}
+}
+
 void FPoseSearchAnimPlayer::StorePoseContext(const FPoseContext& PoseContext)
 {
 	AssetType = ESearchIndexAssetType::Invalid;
@@ -361,11 +373,19 @@ float FAnimNode_BlendStack::GetAccumulatedTime() const
 	return AnimPlayers.IsEmpty() ? 0.f : AnimPlayers.First().GetAccumulatedTime();
 }
 
-void FAnimNode_BlendStack::BlendTo(ESearchIndexAssetType AssetType, UAnimationAsset* AnimationAsset, float AccumulatedTime, bool bLoop, bool bMirrored, UMirrorDataTable* MirrorDataTable, int32 MaxActiveBlends, float BlendTime, const UBlendProfile* BlendProfile, EAlphaBlendOption BlendOption, FVector BlendParameters)
+void FAnimNode_BlendStack::BlendTo(ESearchIndexAssetType AssetType, UAnimationAsset* AnimationAsset, float AccumulatedTime, bool bLoop, bool bMirrored, UMirrorDataTable* MirrorDataTable, int32 MaxActiveBlends, float BlendTime, const UBlendProfile* BlendProfile, EAlphaBlendOption BlendOption, FVector BlendParameters, float PlayRate)
 {
 	RequestedMaxActiveBlends = MaxActiveBlends;
 	AnimPlayers.PushFirst(FPoseSearchAnimPlayer());
-	AnimPlayers.First().Initialize(AssetType, AnimationAsset, AccumulatedTime, bLoop, bMirrored, MirrorDataTable, BlendTime, BlendProfile, BlendOption, BlendParameters);
+	AnimPlayers.First().Initialize(AssetType, AnimationAsset, AccumulatedTime, bLoop, bMirrored, MirrorDataTable, BlendTime, BlendProfile, BlendOption, BlendParameters, PlayRate);
+}
+
+void FAnimNode_BlendStack::UpdatePlayRate(float PlayRate)
+{
+	if (!AnimPlayers.IsEmpty())
+	{
+		AnimPlayers.First().UpdatePlayRate(PlayRate);
+	}
 }
 
 void FAnimNode_BlendStack::GatherDebugData(FNodeDebugData& DebugData)
diff --git a/Engine/Plugins/Experimental/Animation/PoseSearch/Source/Runtime/Private/AnimNode_MotionMatching.cpp b/Engine/Plugins/Experimental/Animation/PoseSearch/Source/Runtime/Private/AnimNode_MotionMatching.cpp
--- a/Engine/Plugins/Experimental/Animation/PoseSearch/Source/Runtime/Private/AnimNode_MotionMatching.cpp
+++ b/Engine/Plugins/Experimental/Animation/PoseSearch/Source/Runtime/Private/AnimNode_MotionMatching.cpp
@@ -99,9 +99,10 @@ void FAnimNode_MotionMatching::UpdateAssetPlayer(const FAnimationUpdateContext&
 			const FPoseSearchDatabaseAnimationAssetBase& AnimationAssetBase = Database->GetAnimationSourceAsset(*SearchIndexAsset);
 			BlendStackNode.BlendTo(SearchIndexAsset->Type, AnimationAssetBase.GetAnimationAsset(), MotionMatchingState.CurrentSearchResult.AssetTime,
 				AnimationAssetBase.IsLooping(), SearchIndexAsset->bMirrored, Database->Schema->MirrorDataTable.Get(),
-				Settings.MaxActiveBlends, Settings.BlendTime, Settings.BlendProfile, Settings.BlendOption, SearchIndexAsset->BlendParameters);
+				Settings.MaxActiveBlends, Settings.BlendTime, Settings.BlendProfile, Settings.BlendOption, SearchIndexAsset->BlendParameters, MotionMatchingState.WantedPlayRate);
 		}
 	}
+	BlendStackNode.UpdatePlayRate(MotionMatchingState.WantedPlayRate);
 
 	Source.Update(Context);
 }
diff --git a/Engine/Plugins/Experimental/Animation/PoseSearch/Source/Runtime/Private/PoseSearchFeatureChannels.cpp b/Engine/Plugins/Experimental/Animation/PoseSearch/Source/Runtime/Private/PoseSearchFeatureChannels.cpp
--- a/Engine/Plugins/Experimental/Animation/PoseSearch/Source/Runtime/Private/PoseSearchFeatureChannels.cpp
+++ b/Engine/Plugins/Experimental/Animation/PoseSearch/Source/Runtime/Private/PoseSearchFeatureChannels.cpp
@@ -2104,6 +2104,84 @@ void UPoseSearchFeatureChannel_Trajectory::ComputeCostBreakdowns(UE::PoseSearch:
 
 #endif // WITH_EDITOR
 
+bool UPoseSearchFeatureChannel_Trajectory::GetEstimatedSpeedRatio(TConstArrayView<float> QueryVector, TConstArrayView<float> PoseVector, float& EstimatedSpeedRatio) const
+{
+	float EstimatedQuerySpeed = 0.f;
+	float EstimatedPoseSpeed = 0.f;
+
+	int32 QueryDataOffset = ChannelDataOffset;
+	int32 PoseDataOffset = ChannelDataOffset;
+
+	bool bValidEstimate = false;
+	for (const FPoseSearchTrajectorySample& Sample : Samples)
+	{
+		// @todo: decode positions and use them to estimate velocities in case UPoseSearchFeatureChannel_Trajectory doens't contain EPoseSearchTrajectoryFlags::Velocity or EPoseSearchTrajectoryFlags::VelocityXY samples
+		if (EnumHasAnyFlags(Sample.Flags, EPoseSearchTrajectoryFlags::Position))
+		{
+			QueryDataOffset += UE::PoseSearch::FFeatureVectorHelper::EncodeVectorCardinality;
+			PoseDataOffset += UE::PoseSearch::FFeatureVectorHelper::EncodeVectorCardinality;
+		}
+		if (EnumHasAnyFlags(Sample.Flags, EPoseSearchTrajectoryFlags::PositionXY))
+		{
+			QueryDataOffset += UE::PoseSearch::FFeatureVectorHelper::EncodeVector2DCardinality;
+			PoseDataOffset += UE::PoseSearch::FFeatureVectorHelper::EncodeVector2DCardinality;
+		}
+
+		if (EnumHasAnyFlags(Sample.Flags, EPoseSearchTrajectoryFlags::Velocity))
+		{
+			const FVector QueryVelocity = UE::PoseSearch::FFeatureVectorHelper::DecodeVector(QueryVector, QueryDataOffset);
+			const FVector PoseVelocity = UE::PoseSearch::FFeatureVectorHelper::DecodeVector(PoseVector, PoseDataOffset);
+			EstimatedQuerySpeed += QueryVelocity.Length();
+			EstimatedPoseSpeed += PoseVelocity.Length();
+			bValidEstimate = true;
+		}
+		if (EnumHasAnyFlags(Sample.Flags, EPoseSearchTrajectoryFlags::VelocityXY))
+		{
+			const FVector2D QueryVelocity = UE::PoseSearch::FFeatureVectorHelper::DecodeVector2D(QueryVector, QueryDataOffset);
+			const FVector2D PoseVelocity = UE::PoseSearch::FFeatureVectorHelper::DecodeVector2D(PoseVector, PoseDataOffset);
+			EstimatedQuerySpeed += QueryVelocity.Length();
+			EstimatedPoseSpeed += PoseVelocity.Length();
+			bValidEstimate = true;
+		}
+
+		if (EnumHasAnyFlags(Sample.Flags, EPoseSearchTrajectoryFlags::VelocityDirection))
+		{
+			QueryDataOffset += UE::PoseSearch::FFeatureVectorHelper::EncodeVectorCardinality;
+			PoseDataOffset += UE::PoseSearch::FFeatureVectorHelper::EncodeVectorCardinality;
+		}
+		if (EnumHasAnyFlags(Sample.Flags, EPoseSearchTrajectoryFlags::VelocityDirectionXY))
+		{
+			QueryDataOffset += UE::PoseSearch::FFeatureVectorHelper::EncodeVector2DCardinality;
+			PoseDataOffset += UE::PoseSearch::FFeatureVectorHelper::EncodeVector2DCardinality;
+		}
+		if (EnumHasAnyFlags(Sample.Flags, EPoseSearchTrajectoryFlags::FacingDirection))
+		{
+			QueryDataOffset += UE::PoseSearch::FFeatureVectorHelper::EncodeVectorCardinality;
+			PoseDataOffset += UE::PoseSearch::FFeatureVectorHelper::EncodeVectorCardinality;
+		}
+		if (EnumHasAnyFlags(Sample.Flags, EPoseSearchTrajectoryFlags::FacingDirectionXY))
+		{
+			QueryDataOffset += UE::PoseSearch::FFeatureVectorHelper::EncodeVector2DCardinality;
+			PoseDataOffset += UE::PoseSearch::FFeatureVectorHelper::EncodeVector2DCardinality;
+		}
+	}
+
+	check(QueryDataOffset == ChannelDataOffset + ChannelCardinality);
+	check(PoseDataOffset == ChannelDataOffset + ChannelCardinality);
+
+	if (bValidEstimate && EstimatedPoseSpeed > UE_KINDA_SMALL_NUMBER)
+	{
+		EstimatedSpeedRatio = EstimatedQuerySpeed / EstimatedPoseSpeed;
+	}
+	else
+	{
+		EstimatedSpeedRatio = 1.f;
+	}
+
+	return bValidEstimate;
+}
+
+
 //////////////////////////////////////////////////////////////////////////
 // UPoseSearchFeatureChannel_FilterCrashingLegs
 static void ComputeThighsSideAndForward(const FVector& RightThighPos, const FVector& LeftThighPos, FVector& ThighsSide, FVector& ThighsForward)
diff --git a/Engine/Plugins/Experimental/Animation/PoseSearch/Source/Runtime/Private/PoseSearchLibrary.cpp b/Engine/Plugins/Experimental/Animation/PoseSearch/Source/Runtime/Private/PoseSearchLibrary.cpp
--- a/Engine/Plugins/Experimental/Animation/PoseSearch/Source/Runtime/Private/PoseSearchLibrary.cpp
+++ b/Engine/Plugins/Experimental/Animation/PoseSearch/Source/Runtime/Private/PoseSearchLibrary.cpp
@@ -10,6 +10,7 @@
 #include "PoseSearch/AnimNode_MotionMatching.h"
 #include "Trace/PoseSearchTraceLogger.h"
 #include "MotionTrajectoryLibrary.h"
+#include "PoseSearch/PoseSearchFeatureChannels.h"
 
 #define LOCTEXT_NAMESPACE "PoseSearchLibrary"
 
@@ -22,6 +23,7 @@ void FMotionMatchingState::Reset()
 	// Set the elapsed time to INFINITY to trigger a search right away
 	ElapsedPoseJumpTime = INFINITY;
 	PoseIndicesHistory.Reset();
+	WantedPlayRate = 1.f;
 }
 
 void FMotionMatchingState::AdjustAssetTime(float AssetTime)
@@ -146,11 +148,6 @@ static void RequestInertialBlend(const FAnimationUpdateContext& Context, float B
 
 void FMotionMatchingState::JumpToPose(const FAnimationUpdateContext& Context, const FMotionMatchingSettings& Settings, const UE::PoseSearch::FSearchResult& Result)
 {
-	// Remember which pose and sequence we're playing from the database
-	CurrentSearchResult = Result;
-
-	ElapsedPoseJumpTime = 0.0f;
-	
 	// requesting inertial blending only if blendstack is disabled
 	if (Settings.MaxActiveBlends <= 0)
 	{
@@ -158,6 +155,11 @@ void FMotionMatchingState::JumpToPose(const FAnimationUpdateContext& Context, co
 		RequestInertialBlend(Context, JumpBlendTime);
 	}
 
+	// Remember which pose and sequence we're playing from the database
+	CurrentSearchResult = Result;
+
+	ElapsedPoseJumpTime = 0.0f;
+
 	Flags |= EMotionMatchingFlags::JumpedToPose;
 }
 
@@ -403,8 +405,9 @@ void UpdateMotionMatchingState(
 		InOutMotionMatchingState.ElapsedPoseJumpTime += DeltaTime;
 	}
 
+	InOutMotionMatchingState.UpdateWantedPlayRate(SearchContext, Settings);
 	InOutMotionMatchingState.PoseIndicesHistory.Update(InOutMotionMatchingState.CurrentSearchResult, DeltaTime, Settings.PoseReselectHistory);
-	
+
 #if UE_POSE_SEARCH_TRACE_ENABLED
 	// Record debugger details
 	TraceMotionMatchingState(
@@ -417,13 +420,49 @@ void UpdateMotionMatchingState(
 #endif
 }
 
+void FMotionMatchingState::UpdateWantedPlayRate(const UE::PoseSearch::FSearchContext& SearchContext, const FMotionMatchingSettings& Settings)
+{
+	if (CurrentSearchResult.IsValid())
+	{
+		if (!FMath::IsNearlyEqual(Settings.PlayRateMin, 1.f, UE_KINDA_SMALL_NUMBER) || !FMath::IsNearlyEqual(Settings.PlayRateMax, 1.f, UE_KINDA_SMALL_NUMBER))
+		{
+			if (const FPoseSearchFeatureVectorBuilder* PoseSearchFeatureVectorBuilder = SearchContext.GetCachedQuery(CurrentSearchResult.Database.Get()))
+			{
+				if (const UPoseSearchFeatureChannel_Trajectory* TrajectoryChannel = CurrentSearchResult.Database->Schema->FindFirstChannelOfType<UPoseSearchFeatureChannel_Trajectory>())
+				{
+					TConstArrayView<float> QueryData = PoseSearchFeatureVectorBuilder->GetValues();
+					TConstArrayView<float> ResultData = CurrentSearchResult.Database->GetSearchIndex()->GetPoseValues(CurrentSearchResult.PoseIdx);
+					float EstimatedSpeedRatio = 1.f;
+					if (TrajectoryChannel->GetEstimatedSpeedRatio(QueryData, ResultData, EstimatedSpeedRatio))
+					{
+						check(Settings.PlayRateMin <= Settings.PlayRateMax);
+						WantedPlayRate = FMath::Clamp(EstimatedSpeedRatio, Settings.PlayRateMin, Settings.PlayRateMax);
+					}
+					else
+					{
+						UE_LOG(LogPoseSearch, Warning,
+							TEXT("Couldn't update the WantedPlayRate in FMotionMatchingState::UpdateWantedPlayRate, because Schema '%s' UPoseSearchFeatureChannel_Trajectory::GetEstimatedSpeedRatio failed to compute, because of missing EPoseSearchTrajectoryFlags::Velocity samples"),
+							*GetNameSafe(CurrentSearchResult.Database->Schema));
+					}
+				}
+				else
+				{
+					UE_LOG(LogPoseSearch, Warning,
+						TEXT("Couldn't update the WantedPlayRate in FMotionMatchingState::UpdateWantedPlayRate, because Schema '%s' couldn't find a UPoseSearchFeatureChannel_Trajectory channel"),
+						*GetNameSafe(CurrentSearchResult.Database->Schema));
+				}
+			}
+		}
+	}
+}
+
 float FMotionMatchingState::ComputeJumpBlendTime(const UE::PoseSearch::FSearchResult& Result, const FMotionMatchingSettings& Settings) const
 {
 	const FPoseSearchIndexAsset* SearchIndexAsset = CurrentSearchResult.GetSearchIndexAsset();
 
 	// Use alternate blend time when changing between mirrored and unmirrored
 	float JumpBlendTime = Settings.BlendTime;
-	if ((SearchIndexAsset != nullptr) && (Settings.MirrorChangeBlendTime > 0.0f))
+	if (SearchIndexAsset && Settings.MirrorChangeBlendTime > 0.0f)
 	{
 		if (Result.GetSearchIndexAsset(true)->bMirrored != SearchIndexAsset->bMirrored)
 		{
diff --git a/Engine/Plugins/Experimental/Animation/PoseSearch/Source/Runtime/Public/PoseSearch/AnimNode_BlendStack.h b/Engine/Plugins/Experimental/Animation/PoseSearch/Source/Runtime/Public/PoseSearch/AnimNode_BlendStack.h
--- a/Engine/Plugins/Experimental/Animation/PoseSearch/Source/Runtime/Public/PoseSearch/AnimNode_BlendStack.h
+++ b/Engine/Plugins/Experimental/Animation/PoseSearch/Source/Runtime/Public/PoseSearch/AnimNode_BlendStack.h
@@ -16,7 +16,8 @@ struct FPoseSearchAnimPlayer
 {
 	GENERATED_BODY()
 	
-	void Initialize(ESearchIndexAssetType InAssetType, UAnimationAsset* AnimationAsset, float AccumulatedTime, bool bLoop, bool bMirrored, UMirrorDataTable* MirrorDataTable, float BlendTime, const UBlendProfile* BlendProfile, EAlphaBlendOption InBlendOption, FVector BlendParameters);
+	void Initialize(ESearchIndexAssetType InAssetType, UAnimationAsset* AnimationAsset, float AccumulatedTime, bool bLoop, bool bMirrored, UMirrorDataTable* MirrorDataTable, float BlendTime, const UBlendProfile* BlendProfile, EAlphaBlendOption InBlendOption, FVector BlendParameters, float PlayRate);
+	void UpdatePlayRate(float PlayRate);
 	void Evaluate_AnyThread(FPoseContext& Output);
 	void Update_AnyThread(const FAnimationUpdateContext& Context, float BlendWeight);
 	float GetAccumulatedTime() const;
@@ -72,7 +73,8 @@ struct POSESEARCH_API FAnimNode_BlendStack : public FAnimNode_AssetPlayerBase
 	virtual void GatherDebugData(FNodeDebugData& DebugData) override;
 	// End of FAnimNode_Base interface
 
-	void BlendTo(ESearchIndexAssetType AssetType, UAnimationAsset* AnimationAsset, float AccumulatedTime = 0.f, bool bLoop = false, bool bMirrored = false, UMirrorDataTable* MirrorDataTable = nullptr, int32 MaxActiveBlends = 3, float BlendTime = 0.3f, const UBlendProfile* BlendProfile = nullptr, EAlphaBlendOption BlendOption = EAlphaBlendOption::Linear, FVector BlendParameters = FVector::Zero());
+	void BlendTo(ESearchIndexAssetType AssetType, UAnimationAsset* AnimationAsset, float AccumulatedTime = 0.f, bool bLoop = false, bool bMirrored = false, UMirrorDataTable* MirrorDataTable = nullptr, int32 MaxActiveBlends = 3, float BlendTime = 0.3f, const UBlendProfile* BlendProfile = nullptr, EAlphaBlendOption BlendOption = EAlphaBlendOption::Linear, FVector BlendParameters = FVector::Zero(), float PlayRate = 1.f);
+	void UpdatePlayRate(float PlayRate);
 
 	// FAnimNode_AssetPlayerBase interface
 	virtual float GetAccumulatedTime() const override;
diff --git a/Engine/Plugins/Experimental/Animation/PoseSearch/Source/Runtime/Public/PoseSearch/PoseSearch.h b/Engine/Plugins/Experimental/Animation/PoseSearch/Source/Runtime/Public/PoseSearch/PoseSearch.h
--- a/Engine/Plugins/Experimental/Animation/PoseSearch/Source/Runtime/Public/PoseSearch/PoseSearch.h
+++ b/Engine/Plugins/Experimental/Animation/PoseSearch/Source/Runtime/Public/PoseSearch/PoseSearch.h
@@ -598,6 +598,20 @@ public:
 
 	float GetSamplingInterval() const { return 1.0f / SampleRate; }
 
+	template<typename ChannelType>
+	const ChannelType* FindFirstChannelOfType() const
+	{
+		for (const TObjectPtr<UPoseSearchFeatureChannel>& ChannelPtr : Channels)
+		{
+			if (const ChannelType* Channel = Cast<const ChannelType>(ChannelPtr.Get()))
+			{
+				return Channel;
+			}
+		}
+		return nullptr;
+	}
+
+
 	// UObject
 	virtual void PreSave(FObjectPreSaveContext ObjectSaveContext) override;
 	virtual void PostLoad() override;
diff --git a/Engine/Plugins/Experimental/Animation/PoseSearch/Source/Runtime/Public/PoseSearch/PoseSearchFeatureChannels.h b/Engine/Plugins/Experimental/Animation/PoseSearch/Source/Runtime/Public/PoseSearch/PoseSearchFeatureChannels.h
--- a/Engine/Plugins/Experimental/Animation/PoseSearch/Source/Runtime/Public/PoseSearch/PoseSearchFeatureChannels.h
+++ b/Engine/Plugins/Experimental/Animation/PoseSearch/Source/Runtime/Public/PoseSearch/PoseSearchFeatureChannels.h
@@ -268,6 +268,8 @@ public:
 	virtual void ComputeCostBreakdowns(UE::PoseSearch::ICostBreakDownData& CostBreakDownData, const UPoseSearchSchema* Schema) const override;
 #endif
 
+	bool GetEstimatedSpeedRatio(TConstArrayView<float> QueryVector, TConstArrayView<float> PoseVector, float& EstimatedSpeedRatio) const;
+
 protected:
 	void IndexAssetPrivate(const UE::PoseSearch::IAssetIndexer& Indexer, int32 SampleIdx, TArrayView<float> FeatureVector) const;
 	float GetSampleTime(const UE::PoseSearch::IAssetIndexer& Indexer, float Offset, float SampleTime, float RootDistance) const;
diff --git a/Engine/Plugins/Experimental/Animation/PoseSearch/Source/Runtime/Public/PoseSearch/PoseSearchLibrary.h b/Engine/Plugins/Experimental/Animation/PoseSearch/Source/Runtime/Public/PoseSearch/PoseSearchLibrary.h
--- a/Engine/Plugins/Experimental/Animation/PoseSearch/Source/Runtime/Public/PoseSearch/PoseSearchLibrary.h
+++ b/Engine/Plugins/Experimental/Animation/PoseSearch/Source/Runtime/Public/PoseSearch/PoseSearchLibrary.h
@@ -56,6 +56,12 @@ struct POSESEARCH_API FMotionMatchingSettings
 	// Minimum amount of time to wait between pose search queries
 	UPROPERTY(EditAnywhere, BlueprintReadWrite, Category=Settings, meta=(ClampMin="0"))
 	float SearchThrottleTime = 0.f;
+
+	UPROPERTY(EditAnywhere, BlueprintReadWrite, Category = Settings, meta = (ClampMin = "0.5", ClampMax = "1.0", UIMin = "0.5", UIMax = "1.0"))
+	float PlayRateMin = 1.f;
+
+	UPROPERTY(EditAnywhere, BlueprintReadWrite, Category = Settings, meta = (ClampMin = "1.0", ClampMax = "2.0", UIMin = "1.0", UIMax = "2.0"))
+	float PlayRateMax = 1.f;
 };
 
 USTRUCT(BlueprintType, Category="Animation|Pose Search")
@@ -81,12 +87,18 @@ struct POSESEARCH_API FMotionMatchingState
 
 	float ComputeJumpBlendTime(const UE::PoseSearch::FSearchResult& Result, const FMotionMatchingSettings& Settings) const;
 
+	void UpdateWantedPlayRate(const UE::PoseSearch::FSearchContext& SearchContext, const FMotionMatchingSettings& Settings);
+
 	UE::PoseSearch::FSearchResult CurrentSearchResult;
 
 	// Time since the last pose jump
 	UPROPERTY(Transient)
 	float ElapsedPoseJumpTime = 0.f;
 
+	// wanted PlayRate to have the selected animation playing at the estimated requested speed from the query
+	UPROPERTY(Transient)
+	float WantedPlayRate = 1.f;
+
 	// Evaluation flags relevant to the state of motion matching
 	UPROPERTY(VisibleAnywhere, BlueprintReadOnly, Category=State)
 	EMotionMatchingFlags Flags = EMotionMatchingFlags::None;
