motion matching - fix for kdtree searches incorrectly reporting filtered poses when full searches were skipped

diff --git a/Engine/Plugins/Experimental/Animation/PoseSearch/Source/Runtime/Private/PoseSearch.cpp b/Engine/Plugins/Experimental/Animation/PoseSearch/Source/Runtime/Private/PoseSearch.cpp
--- a/Engine/Plugins/Experimental/Animation/PoseSearch/Source/Runtime/Private/PoseSearch.cpp
+++ b/Engine/Plugins/Experimental/Animation/PoseSearch/Source/Runtime/Private/PoseSearch.cpp
@@ -168,13 +168,13 @@ FORCEINLINE auto LowerBound(IteratorType First, IteratorType Last, const ValueTy
 	return LowerBound(First, Last, Value, FIdentityFunctor(), TLess<>());
 }
 
-static TArrayView<size_t> PopulateNonSelectableIdx(TArrayView<size_t> NonSelectableIdxBuffer, FSearchContext& SearchContext, const UPoseSearchDatabase* Database, TConstArrayView<float> QueryValues)
+typedef TArray<size_t, TInlineAllocator<128>> FNonSelectableIdx;
+static void PopulateNonSelectableIdx(FNonSelectableIdx& NonSelectableIdx, FSearchContext& SearchContext, const UPoseSearchDatabase* Database, TConstArrayView<float> QueryValues)
 {
 	check(Database);
 	const FPoseSearchIndex* SearchIndex = Database->GetSearchIndex();
 	check(SearchIndex);
 
-	int32 NonSelectableIdxUsedSize = 0;
 	const FPoseSearchIndexAsset* CurrentIndexAsset = SearchContext.CurrentResult.GetSearchIndexAsset();
 	if (CurrentIndexAsset && SearchContext.IsCurrentResultFromDatabase(Database) && SearchContext.PoseJumpThresholdTime > 0.f)
 	{
@@ -201,22 +201,12 @@ static TArrayView<size_t> PopulateNonSelectableIdx(TArrayView<size_t> NonSelecta
 
 			if (bIsPoseInRange)
 			{
-				if (NonSelectableIdxUsedSize < NonSelectableIdxBuffer.Num())
-				{
-					NonSelectableIdxBuffer[NonSelectableIdxUsedSize++] = PoseIdx;
+				NonSelectableIdx.Add(PoseIdx);
 
 #if UE_POSE_SEARCH_TRACE_ENABLED
-					const FPoseSearchCost PoseCost = SearchIndex->ComparePoses(PoseIdx, SearchContext.QueryMirrorRequest, EPoseComparisonFlags::None, QueryValues);
-					SearchContext.BestCandidates.Add(PoseCost, PoseIdx, Database, EPoseCandidateFlags::DiscardedBy_PoseJumpThresholdTime);
+				const FPoseSearchCost PoseCost = SearchIndex->ComparePoses(PoseIdx, SearchContext.QueryMirrorRequest, EPoseComparisonFlags::None, QueryValues);
+				SearchContext.BestCandidates.Add(PoseCost, PoseIdx, Database, EPoseCandidateFlags::DiscardedBy_PoseJumpThresholdTime);
 #endif
-				}
-				else
-				{
-					UE_LOG(LogPoseSearch, Warning, TEXT("PopulateNonSelectableIdx couldn't add all the NonSelectableIdx"));
-					TArrayView<size_t> NonSelectableIdx(NonSelectableIdxBuffer.GetData(), NonSelectableIdxUsedSize);
-					NonSelectableIdx.Sort();
-					return NonSelectableIdx;
-				}
 			}
 			else
 			{
@@ -244,22 +234,12 @@ static TArrayView<size_t> PopulateNonSelectableIdx(TArrayView<size_t> NonSelecta
 
 			if (bIsPoseInRange)
 			{
-				if (NonSelectableIdxUsedSize < NonSelectableIdxBuffer.Num())
-				{
-					NonSelectableIdxBuffer[NonSelectableIdxUsedSize++] = PoseIdx;
+				NonSelectableIdx.Add(PoseIdx);
 
 #if UE_POSE_SEARCH_TRACE_ENABLED
-					const FPoseSearchCost PoseCost = SearchIndex->ComparePoses(PoseIdx, SearchContext.QueryMirrorRequest, EPoseComparisonFlags::None, QueryValues);
-					SearchContext.BestCandidates.Add(PoseCost, PoseIdx, Database, EPoseCandidateFlags::DiscardedBy_PoseJumpThresholdTime);
+				const FPoseSearchCost PoseCost = SearchIndex->ComparePoses(PoseIdx, SearchContext.QueryMirrorRequest, EPoseComparisonFlags::None, QueryValues);
+				SearchContext.BestCandidates.Add(PoseCost, PoseIdx, Database, EPoseCandidateFlags::DiscardedBy_PoseJumpThresholdTime);
 #endif
-				}
-				else
-				{
-					UE_LOG(LogPoseSearch, Warning, TEXT("PopulateNonSelectableIdx couldn't add all the NonSelectableIdx"));
-					TArrayView<size_t> NonSelectableIdx(NonSelectableIdxBuffer.GetData(), NonSelectableIdxUsedSize);
-					NonSelectableIdx.Sort();
-					return NonSelectableIdx;
-				}
 			}
 			else
 			{
@@ -276,35 +256,23 @@ static TArrayView<size_t> PopulateNonSelectableIdx(TArrayView<size_t> NonSelecta
 			const FHistoricalPoseIndex& HistoricalPoseIndex = It.Key();
 			if (HistoricalPoseIndex.DatabaseKey == DatabaseKey)
 			{
-				if (NonSelectableIdxUsedSize < NonSelectableIdxBuffer.Num())
-				{
-					NonSelectableIdxBuffer[NonSelectableIdxUsedSize++] = HistoricalPoseIndex.PoseIndex;
+				NonSelectableIdx.Add(HistoricalPoseIndex.PoseIndex);
 
 #if UE_POSE_SEARCH_TRACE_ENABLED
-					check(HistoricalPoseIndex.PoseIndex >= 0);
+				check(HistoricalPoseIndex.PoseIndex >= 0);
 					
-					// if we're editing the database and removing assets it's possible that the PoseIndicesHistory contains invalid pose indexes
-					if (HistoricalPoseIndex.PoseIndex < SearchIndex->NumPoses)
-					{
-						const FPoseSearchCost PoseCost = SearchIndex->ComparePoses(HistoricalPoseIndex.PoseIndex, SearchContext.QueryMirrorRequest, EPoseComparisonFlags::None, QueryValues);
-						SearchContext.BestCandidates.Add(PoseCost, HistoricalPoseIndex.PoseIndex, Database, EPoseCandidateFlags::DiscardedBy_PoseReselectHistory);
-					}
-#endif
-				}
-				else
+				// if we're editing the database and removing assets it's possible that the PoseIndicesHistory contains invalid pose indexes
+				if (HistoricalPoseIndex.PoseIndex < SearchIndex->NumPoses)
 				{
-					UE_LOG(LogPoseSearch, Warning, TEXT("PopulateNonSelectableIdx couldn't add all the NonSelectableIdx"));
-					TArrayView<size_t> NonSelectableIdx(NonSelectableIdxBuffer.GetData(), NonSelectableIdxUsedSize);
-					NonSelectableIdx.Sort();
-					return NonSelectableIdx;
+					const FPoseSearchCost PoseCost = SearchIndex->ComparePoses(HistoricalPoseIndex.PoseIndex, SearchContext.QueryMirrorRequest, EPoseComparisonFlags::None, QueryValues);
+					SearchContext.BestCandidates.Add(PoseCost, HistoricalPoseIndex.PoseIndex, Database, EPoseCandidateFlags::DiscardedBy_PoseReselectHistory);
 				}
+#endif
 			}
 		}
 	}
 
-	TArrayView<size_t> NonSelectableIdx(NonSelectableIdxBuffer.GetData(), NonSelectableIdxUsedSize);
 	NonSelectableIdx.Sort();
-	return NonSelectableIdx;
 }
 
 struct FPoseFilters
@@ -941,12 +909,12 @@ const bool UPoseSearchDatabase::IsSourceAssetLooping(const FPoseSearchIndexAsset
 {
 	if (SearchIndexAsset.Type == ESearchIndexAssetType::Sequence)
 	{
-		return GetSequenceSourceAsset(SearchIndexAsset).Sequence->bLoop;
+		return Sequences[SearchIndexAsset.SourceAssetIdx].Sequence->bLoop;
 	}
 
 	if (SearchIndexAsset.Type == ESearchIndexAssetType::BlendSpace)
 	{
-		return GetBlendSpaceSourceAsset(SearchIndexAsset).BlendSpace->bLoop;
+		return BlendSpaces[SearchIndexAsset.SourceAssetIdx].BlendSpace->bLoop;
 	}
 	
 	checkNoEntry();
@@ -957,12 +925,12 @@ const FString UPoseSearchDatabase::GetSourceAssetName(const FPoseSearchIndexAsse
 {
 	if (SearchIndexAsset.Type == ESearchIndexAssetType::Sequence)
 	{
-		return GetSequenceSourceAsset(SearchIndexAsset).Sequence->GetName();
+		return Sequences[SearchIndexAsset.SourceAssetIdx].Sequence->GetName();
 	}
 	
 	if (SearchIndexAsset.Type == ESearchIndexAssetType::BlendSpace)
 	{
-		return GetBlendSpaceSourceAsset(SearchIndexAsset).BlendSpace->GetName();
+		return BlendSpaces[SearchIndexAsset.SourceAssetIdx].BlendSpace->GetName();
 	}
 	
 	checkNoEntry();
@@ -1724,10 +1692,6 @@ UE::PoseSearch::FSearchResult UPoseSearchDatabase::SearchPCAKDTree(UE::PoseSearc
 	RowMajorVectorMap CenteredQueryValues((float*)FMemory_Alloca(NumDimensions * sizeof(float)), 1, NumDimensions);
 	RowMajorVectorMap ProjectedQueryValues((float*)FMemory_Alloca(ClampedNumberOfPrincipalComponents * sizeof(float)), 1, ClampedNumberOfPrincipalComponents);
 	
-	constexpr int NonSelectableIdxDataSize = 128;
-	size_t* NonSelectableIdxData((size_t*)FMemory_Alloca(NonSelectableIdxDataSize * sizeof(size_t)));
-	TArrayView<size_t> NonSelectableIdxBuffer(NonSelectableIdxData, NonSelectableIdxDataSize);
-
 	// KDTree in PCA space search
 	if (PoseSearchMode == EPoseSearchMode::PCAKDTree_Validate)
 	{
@@ -1785,8 +1749,8 @@ UE::PoseSearch::FSearchResult UPoseSearchDatabase::SearchPCAKDTree(UE::PoseSearc
 	// there's no point in performing the search if CurrentBestTotalCost is already better than that
 	if (SearchContext.GetCurrentBestTotalCost() > SearchIndex->MinCostAddend)
 	{
-		TConstArrayView<size_t> NonSelectableIdx = PopulateNonSelectableIdx(NonSelectableIdxBuffer, SearchContext, this, QueryValues);
-		const RowMajorVectorMapConst MapWeightsSqrt(SearchIndex->WeightsSqrt.GetData(), 1, NumDimensions);
+		FNonSelectableIdx NonSelectableIdx;
+		PopulateNonSelectableIdx(NonSelectableIdx, SearchContext, this, QueryValues);
 		FKDTree::KNNResultSet ResultSet(ClampedKDTreeQueryNumNeighbors, ResultIndexes, ResultDistanceSqr, NonSelectableIdx);
 
 		check(QueryValues.Num() == NumDimensions);
@@ -1796,6 +1760,7 @@ UE::PoseSearch::FSearchResult UPoseSearchDatabase::SearchPCAKDTree(UE::PoseSearc
 
 		// transforming query values into PCA space to query the KDTree
 		const RowMajorVectorMapConst QueryValuesMap(QueryValues.GetData(), 1, NumDimensions);
+		const RowMajorVectorMapConst MapWeightsSqrt(SearchIndex->WeightsSqrt.GetData(), 1, NumDimensions);
 		WeightedQueryValues = QueryValuesMap.array() * MapWeightsSqrt.array();
 		CenteredQueryValues.noalias() = WeightedQueryValues - Mean;
 		ProjectedQueryValues.noalias() = CenteredQueryValues * PCAProjectionMatrix;
@@ -1831,6 +1796,14 @@ UE::PoseSearch::FSearchResult UPoseSearchDatabase::SearchPCAKDTree(UE::PoseSearc
 			SearchContext.UpdateCurrentBestCost(Result.PoseCost);
 		}
 	}
+	else
+	{
+#if UE_POSE_SEARCH_TRACE_ENABLED
+		// calling just for reporting non selectable poses
+		FNonSelectableIdx NonSelectableIdx;
+		PopulateNonSelectableIdx(NonSelectableIdx, SearchContext, this, QueryValues);
+#endif
+	}
 
 	// finalizing Result properties
 	if (Result.PoseIdx != INDEX_NONE)
@@ -1877,16 +1850,14 @@ UE::PoseSearch::FSearchResult UPoseSearchDatabase::SearchBruteForce(UE::PoseSear
 		}
 	}
 
-	constexpr int NonSelectableIdxDataSize = 128;
-	size_t* NonSelectableIdxData((size_t*)FMemory_Alloca(NonSelectableIdxDataSize * sizeof(size_t)));
-	TArrayView<size_t> NonSelectableIdxBuffer(NonSelectableIdxData, NonSelectableIdxDataSize);
-	TConstArrayView<size_t> NonSelectableIdx = PopulateNonSelectableIdx(NonSelectableIdxBuffer, SearchContext, this, QueryValues);
-	check(Algo::IsSorted(NonSelectableIdx));
-
 	// since any PoseCost calculated here is at least SearchIndex->MinCostAddend,
 	// there's no point in performing the search if CurrentBestTotalCost is already better than that
 	if (SearchContext.GetCurrentBestTotalCost() > SearchIndex->MinCostAddend)
 	{
+		FNonSelectableIdx NonSelectableIdx;
+		PopulateNonSelectableIdx(NonSelectableIdx, SearchContext, this, QueryValues);
+		check(Algo::IsSorted(NonSelectableIdx));
+
 		const FPoseFilters PoseFilters(SearchIndex->Schema, NonSelectableIdx, SearchIndex->OverallFlags);
 		for (int32 PoseIdx = 0; PoseIdx < SearchIndex->NumPoses; ++PoseIdx)
 		{
@@ -1917,6 +1888,14 @@ UE::PoseSearch::FSearchResult UPoseSearchDatabase::SearchBruteForce(UE::PoseSear
 			SearchContext.UpdateCurrentBestCost(Result.PoseCost);
 		}
 	}
+	else
+	{
+#if UE_POSE_SEARCH_TRACE_ENABLED
+		// calling just for reporting non selectable poses
+		FNonSelectableIdx NonSelectableIdx;
+		PopulateNonSelectableIdx(NonSelectableIdx, SearchContext, this, QueryValues);
+#endif
+	}
 
 	// finalizing Result properties
 	if (Result.PoseIdx != INDEX_NONE)
@@ -2715,8 +2694,6 @@ void FSearchResult::Update(float NewAssetTime)
 		const FPoseSearchIndexAsset& SearchIndexAsset = Database->GetSearchIndex()->GetAssetForPose(PoseIdx);
 		if (SearchIndexAsset.Type == ESearchIndexAssetType::Sequence)
 		{
-			const FPoseSearchDatabaseSequence& DbSequence = Database->GetSequenceSourceAsset(SearchIndexAsset);
-
 			if (Database->GetPoseIndicesAndLerpValueFromTime(NewAssetTime, SearchIndexAsset, PrevPoseIdx, PoseIdx, NextPoseIdx, LerpValue))
 			{
 				AssetTime = NewAssetTime;
diff --git a/Engine/Plugins/Experimental/Animation/PoseSearch/Source/Runtime/Private/PoseSearchLibrary.cpp b/Engine/Plugins/Experimental/Animation/PoseSearch/Source/Runtime/Private/PoseSearchLibrary.cpp
--- a/Engine/Plugins/Experimental/Animation/PoseSearch/Source/Runtime/Private/PoseSearchLibrary.cpp
+++ b/Engine/Plugins/Experimental/Animation/PoseSearch/Source/Runtime/Private/PoseSearchLibrary.cpp
@@ -179,11 +179,6 @@ static void TraceMotionMatchingState(
 		return;
 	}
 
-	if (!MotionMatchingState.CurrentSearchResult.IsValid())
-	{
-		return;
-	}
-
 	const float DeltaTime = UpdateContext.GetDeltaTime();
 
 	auto AddUniqueDatabase = [](TArray<FTraceMotionMatchingStateDatabaseEntry>& DatabaseEntries, const UPoseSearchDatabase* Database, UE::PoseSearch::FSearchContext& SearchContext) -> int32
