Motion Matching Editor:
- Added new icons to specify states for an entry in an PoseSearchDatabase
- Added Drag and Drop suggestion text when opening a blank PoseSearchDatabase
- Added UI break between the asset name and asset info icons.

diff --git a/Engine/Plugins/Experimental/Animation/PoseSearch/Source/Editor/Private/SPoseSearchDatabaseAssetList.cpp b/Engine/Plugins/Experimental/Animation/PoseSearch/Source/Editor/Private/SPoseSearchDatabaseAssetList.cpp
--- a/Engine/Plugins/Experimental/Animation/PoseSearch/Source/Editor/Private/SPoseSearchDatabaseAssetList.cpp
+++ b/Engine/Plugins/Experimental/Animation/PoseSearch/Source/Editor/Private/SPoseSearchDatabaseAssetList.cpp
@@ -11,6 +11,8 @@
 #include "DragAndDrop/AssetDragDropOp.h"
 #include "Misc/FeedbackContext.h"
 #include "AssetSelection.h"
+#include "ClassIconFinder.h"
+#include "DetailColumnSizeData.h"
 
 #include "Widgets/Text/SRichTextBlock.h"
 #include "Widgets/Input/SCheckBox.h"
@@ -21,11 +23,14 @@
 #include "Styling/AppStyle.h"
 
 #include "ScopedTransaction.h"
+#include "Styling/StyleColors.h"
 
 #define LOCTEXT_NAMESPACE "PoseSearchDatabaseAssetList"
 
 namespace UE::PoseSearch
 {
+	static constexpr FLinearColor DisabledColor = FLinearColor(1.0f, 1.0f, 1.0f, 0.25f);
+	
 	FDatabaseAssetTreeNode::FDatabaseAssetTreeNode(
 		int32 InSourceAssetIdx,
 		ESearchIndexAssetType InSourceAssetType,
@@ -50,6 +55,81 @@ namespace UE::PoseSearch
 			InHierarchy);
 	}
 
+	bool FDatabaseAssetTreeNode::IsRootMotionEnabled() const
+	{
+		if (const UPoseSearchDatabase* Database = EditorViewModel.Pin()->GetPoseSearchDatabase())
+		{
+			if (SourceAssetType == ESearchIndexAssetType::Sequence)
+			{
+				if (const UAnimSequence* Sequence = Database->Sequences[SourceAssetIdx].Sequence)
+				{
+					return Sequence->HasRootMotion();
+				}
+			}
+			else if (SourceAssetType == ESearchIndexAssetType::BlendSpace)
+			{
+				if (const UBlendSpace* BlendSpace = Database->BlendSpaces[SourceAssetIdx].BlendSpace)
+				{
+					bool bIsRootMotionUsedInBlendSpace = false;
+			
+					BlendSpace->ForEachImmutableSample([&bIsRootMotionUsedInBlendSpace](const FBlendSample & Sample)
+					{
+						const TObjectPtr<UAnimSequence> Sequence = Sample.Animation;
+				
+						if (IsValid(Sequence) && Sequence->HasRootMotion())
+						{
+							bIsRootMotionUsedInBlendSpace = true;
+						}
+					});
+			
+					return bIsRootMotionUsedInBlendSpace;
+				}
+			}
+		}
+		
+		return false;
+	}
+
+	bool FDatabaseAssetTreeNode::IsLooping() const
+	{
+		if (const UPoseSearchDatabase* Database = EditorViewModel.Pin()->GetPoseSearchDatabase())
+		{
+			if (SourceAssetType == ESearchIndexAssetType::Sequence)
+			{
+				if (const UAnimSequence* Sequence = Database->Sequences[SourceAssetIdx].Sequence)
+				{
+					return Sequence->bLoop;
+				}
+			}
+			else if (SourceAssetType == ESearchIndexAssetType::BlendSpace)
+			{
+				if (const UBlendSpace* BlendSpace = Database->BlendSpaces[SourceAssetIdx].BlendSpace)
+				{
+					return BlendSpace->bLoop;
+				}
+			}
+		}
+
+		return false;
+	}
+
+	EPoseSearchMirrorOption FDatabaseAssetTreeNode::GetMirrorOption() const
+	{
+		if (const UPoseSearchDatabase* Database = EditorViewModel.Pin()->GetPoseSearchDatabase())
+		{
+			if (SourceAssetType == ESearchIndexAssetType::Sequence)
+			{
+				return Database->Sequences[SourceAssetIdx].MirrorOption;
+			}
+			else if (SourceAssetType == ESearchIndexAssetType::BlendSpace)
+			{
+				return Database->BlendSpaces[SourceAssetIdx].MirrorOption;
+			}
+		}
+
+		return EPoseSearchMirrorOption::Invalid;
+	}
+
 	void SDatabaseAssetListItem::Construct(
 		const FArguments& InArgs,
 		const TSharedRef<FDatabaseViewModel>& InEditorViewModel,
@@ -159,9 +239,12 @@ namespace UE::PoseSearch
 	TSharedRef<SWidget> SDatabaseAssetListItem::GenerateItemWidget()
 	{
 		TSharedPtr<FDatabaseAssetTreeNode> Node = WeakAssetTreeNode.Pin();
-
+		TSharedPtr<FDatabaseViewModel> ViewModel = EditorViewModel.Pin();
+		UPoseSearchDatabase* Database = ViewModel->GetPoseSearchDatabase();
+		
 		TSharedPtr<SWidget> ItemWidget;
-
+		const FDetailColumnSizeData& ColumnSizeData = SkeletonView.Pin()->GetColumnSizeData();
+		
 		if (Node->SourceAssetType == ESearchIndexAssetType::Invalid)
 		{
 			// it's a group
@@ -199,64 +282,145 @@ namespace UE::PoseSearch
 		}
 		else
 		{
+			// Item Icon (Sequence or Blendspace)
 			TSharedPtr<SImage> ItemIconWidget;
 			if (Node->SourceAssetType == ESearchIndexAssetType::Sequence)
 			{
+				FSlateIconFinder::FindIconBrushForClass(UBlendSpace::StaticClass());
 				SAssignNew(ItemIconWidget, SImage)
-				.Image(FAppStyle::Get().GetBrush("Icons.Minus"));
+				.Image(FSlateIconFinder::FindIconBrushForClass(UAnimSequence::StaticClass()));
 			}
 			else
 			{
 				SAssignNew(ItemIconWidget, SImage)
-				.Image(FAppStyle::Get().GetBrush("Icons.Plus"));
+				.Image(FSlateIconFinder::FindIconBrushForClass(UBlendSpace::StaticClass()));
 			}
 
-			// it's an asset (sequence or blendspace)
+			// Setup table row to display 
 			SAssignNew(ItemWidget, SHorizontalBox)
-			+ SHorizontalBox::Slot()
-			.MaxWidth(18)
-			.AutoWidth()
-			.HAlign(HAlign_Left)
-			.VAlign(VAlign_Center)
-			[
-				ItemIconWidget.ToSharedRef()
-			]
-			+ SHorizontalBox::Slot()
+			+SHorizontalBox::Slot()
 			.FillWidth(1.0f)
-			.VAlign(VAlign_Center)
 			[
-				SNew(STextBlock)
-				.Text(this, &SDatabaseAssetListItem::GetName)
-				.ColorAndOpacity(this, &SDatabaseAssetListItem::GetNameTextColorAndOpacity)
-			]
-			+ SHorizontalBox::Slot()
-			.MaxWidth(18)
-			.AutoWidth()
-			.HAlign(HAlign_Right)
-			.VAlign(VAlign_Center)
-			[
-				SNew(SImage)
-				.Image(FAppStyle::Get().GetBrush("Icons.EyeDropper"))
-				.Visibility_Raw(this, &SDatabaseAssetListItem::GetSelectedActorIconVisbility)
+				SNew(SSplitter)
+				.Style(FAppStyle::Get(), "FoliageEditMode.Splitter")
+				.PhysicalSplitterHandleSize(1.0f)
+				.HitDetectionSplitterHandleSize(5.0f)
+				.HighlightedHandleIndex(ColumnSizeData.GetHoveredSplitterIndex())
+				.MinimumSlotHeight(0.5f)
+				
+				// Asset Name with type icon
+				+SSplitter::Slot()
+				.Value(ColumnSizeData.GetNameColumnWidth())
+				.MinSize(0.3f)
+				.OnSlotResized(ColumnSizeData.GetOnNameColumnResized())
+				[
+					SNew(SHorizontalBox)
+					.Clipping(EWidgetClipping::ClipToBounds)
+					+ SHorizontalBox::Slot()
+					.MaxWidth(18)
+					.AutoWidth()
+					.Padding(0.0f, 0.0f, 5.0f, 0.0f)
+					.HAlign(HAlign_Left)
+					.VAlign(VAlign_Center)
+					[
+						ItemIconWidget.ToSharedRef()
+					]
+					+ SHorizontalBox::Slot()
+					.FillWidth(1.0f)
+					.VAlign(VAlign_Center)
+					[
+						SNew(STextBlock)
+						.Text(this, &SDatabaseAssetListItem::GetName)
+						.ColorAndOpacity(this, &SDatabaseAssetListItem::GetNameTextColorAndOpacity)
+					]
+				]
+				
+				// Display information via icons
+				+SSplitter::Slot()
+				.Value(ColumnSizeData.GetValueColumnWidth())
+				.MinSize(0.3f)
+				.OnSlotResized(ColumnSizeData.GetOnValueColumnResized())
+				[
+					// Asset Info.
+
+					// Looping
+					SNew(SHorizontalBox)
+					.Clipping(EWidgetClipping::ClipToBounds)
+					+ SHorizontalBox::Slot()
+					.AutoWidth()
+					.Padding(4.0f, 1.0f)
+					.HAlign(HAlign_Center)
+					.VAlign(VAlign_Center)
+					[
+						SNew(SImage)
+						.Image(FAppStyle::Get().GetBrush("Graph.Node.Loop"))
+						.ColorAndOpacity(this, &SDatabaseAssetListItem::GetLoopingColorAndOpacity)
+						.ToolTipText(this, &SDatabaseAssetListItem::GetLoopingToolTip)
+					]
+
+					// Root Motion
+					+ SHorizontalBox::Slot()
+					.Padding(1.0f, 1.0f)
+					.AutoWidth()
+					.HAlign(HAlign_Center)
+					.VAlign(VAlign_Center)
+					[
+						SNew(SImage)
+						.Image(FAppStyle::Get().GetBrush("AnimGraph.Attribute.RootMotionDelta.Icon"))
+						.ColorAndOpacity(this, &SDatabaseAssetListItem::GetRootMotionColorAndOpacity)
+						.ToolTipText(this, &SDatabaseAssetListItem::GetRootMotionOptionToolTip)
+					]
+					
+					// Mirror Type
+					+ SHorizontalBox::Slot()
+					.Padding(1.0f, 1.0f)
+					.AutoWidth()
+					.HAlign(HAlign_Center)
+					.VAlign(VAlign_Center)
+					[
+						SNew(SImage)
+						.Image(this, &SDatabaseAssetListItem::GetMirrorOptionSlateBrush)
+						.ToolTipText(this, &SDatabaseAssetListItem::GetMirrorOptionToolTip)
+					]
+				]
 			]
+			
 			+SHorizontalBox::Slot()
-			.MaxWidth(16)
 			.AutoWidth()
 			.HAlign(HAlign_Right)
 			.VAlign(VAlign_Center)
 			[
-				SNew(SCheckBox)
-				.IsChecked(this, &SDatabaseAssetListItem::GetAssetEnabledChecked)
-				.OnCheckStateChanged(const_cast<SDatabaseAssetListItem*>(this), &SDatabaseAssetListItem::OnAssetIsEnabledChanged)
-				.ToolTipText(this, &SDatabaseAssetListItem::GetAssetEnabledToolTip)
-				.CheckedImage(FAppStyle::Get().GetBrush("Icons.Visible"))
-				.CheckedHoveredImage(FAppStyle::Get().GetBrush("Icons.Visible"))
-				.CheckedPressedImage(FAppStyle::Get().GetBrush("Icons.Visible"))
-				.UncheckedImage(FAppStyle::Get().GetBrush("Icons.Hidden"))
-				.UncheckedHoveredImage(FAppStyle::Get().GetBrush("Icons.Hidden"))
-				.UncheckedPressedImage(FAppStyle::Get().GetBrush("Icons.Hidden"))
+				SNew(SHorizontalBox)
+				+ SHorizontalBox::Slot()
+				.MaxWidth(18)
+				.Padding(4.0f, 0.0f, 0.0f, 0.0f)
+				.AutoWidth()
+				.HAlign(HAlign_Right)
+				.VAlign(VAlign_Center)
+				[
+					SNew(SImage)
+					.Image(FAppStyle::Get().GetBrush("Icons.EyeDropper"))
+					.Visibility_Raw(this, &SDatabaseAssetListItem::GetSelectedActorIconVisbility)
+				]
+				+ SHorizontalBox::Slot()
+				.MaxWidth(16)
+				.Padding(4.0f, 0.0f)
+				.AutoWidth()
+				.HAlign(HAlign_Right)
+				.VAlign(VAlign_Center)
+				[
+					SNew(SCheckBox)
+					.IsChecked(this, &SDatabaseAssetListItem::GetAssetEnabledChecked)
+					.OnCheckStateChanged(const_cast<SDatabaseAssetListItem*>(this), &SDatabaseAssetListItem::OnAssetIsEnabledChanged)
+					.ToolTipText(this, &SDatabaseAssetListItem::GetAssetEnabledToolTip)
+					.CheckedImage(FAppStyle::Get().GetBrush("Icons.Visible"))
+					.CheckedHoveredImage(FAppStyle::Get().GetBrush("Icons.Visible"))
+					.CheckedPressedImage(FAppStyle::Get().GetBrush("Icons.Visible"))
+					.UncheckedImage(FAppStyle::Get().GetBrush("Icons.Hidden"))
+					.UncheckedHoveredImage(FAppStyle::Get().GetBrush("Icons.Hidden"))
+					.UncheckedPressedImage(FAppStyle::Get().GetBrush("Icons.Hidden"))
+				]
 			];
-
 		}
 
 		return ItemWidget.ToSharedRef();
@@ -369,9 +533,60 @@ namespace UE::PoseSearch
 
 	FSlateColor SDatabaseAssetListItem::GetNameTextColorAndOpacity() const
 	{
-		return GetAssetEnabledChecked() == ECheckBoxState::Checked ? FLinearColor::White : FLinearColor(1.0f, 1.0f, 1.0f, 0.3f);
+		return GetAssetEnabledChecked() == ECheckBoxState::Checked ? FLinearColor::White : DisabledColor;
+	}
+
+	FSlateColor SDatabaseAssetListItem::GetLoopingColorAndOpacity() const
+	{
+		const TSharedPtr<FDatabaseAssetTreeNode> Node = WeakAssetTreeNode.Pin();
+		return Node->IsLooping() ? FLinearColor::White : DisabledColor;
+	}
+
+	FText SDatabaseAssetListItem::GetLoopingToolTip() const
+	{
+		const TSharedPtr<FDatabaseAssetTreeNode> Node = WeakAssetTreeNode.Pin();
+		return Node->IsLooping() ? LOCTEXT("NodeLoopEnabledToolTip", "Looping") : LOCTEXT("NodeLoopDisabledToolTip", "Not looping");
 	}
 
+	FSlateColor SDatabaseAssetListItem::GetRootMotionColorAndOpacity() const
+	{
+		const TSharedPtr<FDatabaseAssetTreeNode> Node = WeakAssetTreeNode.Pin();
+		return Node->IsRootMotionEnabled() ? FLinearColor::White : DisabledColor;
+	}
+
+	FText SDatabaseAssetListItem::GetRootMotionOptionToolTip() const
+	{
+		const TSharedPtr<FDatabaseAssetTreeNode> Node = WeakAssetTreeNode.Pin();
+		return Node->IsRootMotionEnabled() ? LOCTEXT("NodeRootMotionEnabledToolTip", "Root motion enabled") : LOCTEXT("NodeRootMotionDisabledToolTip", "No root motion enabled");
+
+	}
+	const FSlateBrush* SDatabaseAssetListItem::GetMirrorOptionSlateBrush() const
+	{
+		const TSharedPtr<FDatabaseAssetTreeNode> Node = WeakAssetTreeNode.Pin();
+
+		// TODO: Update icons when appropriate assets become available.
+		switch (Node->GetMirrorOption())
+		{
+			case EPoseSearchMirrorOption::UnmirroredOnly: 
+				return FAppStyle::Get().GetBrush("Icons.Minus");
+			
+			case EPoseSearchMirrorOption::MirroredOnly: 
+				return FAppStyle::Get().GetBrush("Icons.Plus");
+			
+			case EPoseSearchMirrorOption::UnmirroredAndMirrored:
+				return FAppStyle::Get().GetBrush("Icons.X");
+			
+			default:
+				return nullptr;
+		}
+	}
+
+	FText SDatabaseAssetListItem::GetMirrorOptionToolTip() const
+	{
+		const TSharedPtr<FDatabaseAssetTreeNode> Node = WeakAssetTreeNode.Pin();
+		return FText::FromString(LOCTEXT("ToolTipMirrorOption", "Mirror Option: ").ToString() + UEnum::GetDisplayValueAsText(Node->GetMirrorOption()).ToString());
+	}
+	
 	FText SDatabaseAssetListItem::GetAssetEnabledToolTip() const
 	{
 		if (GetAssetEnabledChecked() == ECheckBoxState::Checked)
@@ -391,6 +606,8 @@ namespace UE::PoseSearch
 		TSharedRef<FDatabaseViewModel> InEditorViewModel)
 	{
 		EditorViewModel = InEditorViewModel;
+		
+		ColumnSizeData.SetValueColumnWidth(0.6f);
 
 		CreateCommandList();
 
@@ -428,19 +645,35 @@ namespace UE::PoseSearch
 				.Padding(2.0f)
 				.BorderImage(FAppStyle::GetBrush("SCSEditor.TreePanel"))
 				[
-					SAssignNew(TreeView, STreeView<TSharedPtr<FDatabaseAssetTreeNode>>)
-					.TreeItemsSource(&RootNodes)
-					.SelectionMode(ESelectionMode::Multi)
-					.OnGenerateRow(this, &SDatabaseAssetTree::MakeTableRowWidget)
-					.OnGetChildren(this, &SDatabaseAssetTree::HandleGetChildrenForTree)
-					.OnContextMenuOpening(this, &SDatabaseAssetTree::CreateContextMenu)
-					.HighlightParentNodesForSelection(false)
-					.OnSelectionChanged_Lambda([this](TSharedPtr<FDatabaseAssetTreeNode> Item, ESelectInfo::Type Type)
-						{
-							TArray<TSharedPtr<FDatabaseAssetTreeNode>> SelectedItems = TreeView->GetSelectedItems();
-							OnSelectionChanged.Broadcast(SelectedItems, Type);
-						})
-					.ItemHeight(24)
+					SNew(SOverlay)
+					+SOverlay::Slot()
+					[
+						SAssignNew(TreeView, STreeView<TSharedPtr<FDatabaseAssetTreeNode>>)
+						.TreeItemsSource(&RootNodes)
+						.SelectionMode(ESelectionMode::Multi)
+						.OnGenerateRow(this, &SDatabaseAssetTree::MakeTableRowWidget)
+						.OnGetChildren(this, &SDatabaseAssetTree::HandleGetChildrenForTree)
+						.OnContextMenuOpening(this, &SDatabaseAssetTree::CreateContextMenu)
+						.HighlightParentNodesForSelection(false)
+						.OnSelectionChanged_Lambda([this](TSharedPtr<FDatabaseAssetTreeNode> Item, ESelectInfo::Type Type)
+							{
+								TArray<TSharedPtr<FDatabaseAssetTreeNode>> SelectedItems = TreeView->GetSelectedItems();
+								OnSelectionChanged.Broadcast(SelectedItems, Type);
+							})
+						.ItemHeight(24)
+					]
+					+SOverlay::Slot()
+					[
+						SAssignNew(TreeViewDragAndDropSuggestion, SVerticalBox)
+						+ SVerticalBox::Slot()
+						.HAlign(HAlign_Center)
+						.VAlign(VAlign_Center)
+						[
+							SNew(STextBlock)
+							.Text(FText::FromString(TEXT("Drag and drop Animation Sequences or Blendspaces")))
+							.Font(FAppStyle::Get().GetFontStyle("DetailsView.CategoryFontStyle"))
+						]
+					]
 				]
 			]
 		];
@@ -640,6 +873,9 @@ namespace UE::PoseSearch
 		{
 			TreeView->SetItemSelection(PreviouslySelectedNodes, false, ESelectInfo::Direct);
 		}
+		
+		// Show drag and drop suggestion if tree is empty
+		TreeViewDragAndDropSuggestion->SetVisibility(BlendspaceIndexArray.IsEmpty() && SequenceIndexArray.IsEmpty() ? EVisibility::Visible : EVisibility::Hidden);
 	}
 
 	TSharedRef<ITableRow> SDatabaseAssetTree::MakeTableRowWidget(
@@ -713,35 +949,40 @@ namespace UE::PoseSearch
 		if (NumAssets > 0)
 		{
 			GWarn->BeginSlowTask(LOCTEXT("LoadingAssets", "Loading Asset(s)"), true);
-			for (int32 DroppedAssetIdx = 0; DroppedAssetIdx < NumAssets; ++DroppedAssetIdx)
-			{
-				const FAssetData& AssetData = DroppedAssetData[DroppedAssetIdx];
 
-				if (!AssetData.IsAssetLoaded())
+			{
+				const FScopedTransaction Transaction(LOCTEXT("AddSequencesOrBlendspaces", "Add Sequence(s) and/or Blendspace(s) to Pose Search Database"));
+				
+				for (int32 DroppedAssetIdx = 0; DroppedAssetIdx < NumAssets; ++DroppedAssetIdx)
 				{
-					GWarn->StatusUpdate(
-						DroppedAssetIdx,
-						NumAssets,
-						FText::Format(
-							LOCTEXT("LoadingAsset", "Loading Asset {0}"),
-							FText::FromName(AssetData.AssetName)));
-				}
+					const FAssetData& AssetData = DroppedAssetData[DroppedAssetIdx];
 
-				UClass* AssetClass = AssetData.GetClass();
-				UObject* Asset = AssetData.GetAsset();
+					if (!AssetData.IsAssetLoaded())
+					{
+						GWarn->StatusUpdate(
+							DroppedAssetIdx,
+							NumAssets,
+							FText::Format(
+								LOCTEXT("LoadingAsset", "Loading Asset {0}"),
+								FText::FromName(AssetData.AssetName)));
+					}
 
-				if (AssetClass->IsChildOf(UAnimSequence::StaticClass()))
-				{
-					ViewModel->AddSequenceToDatabase(Cast<UAnimSequence>(Asset));
-					++AddedAssets;
-				}
-				else if (AssetClass->IsChildOf(UBlendSpace::StaticClass()))
-				{
-					ViewModel->AddBlendSpaceToDatabase(Cast<UBlendSpace>(Asset));
-					++AddedAssets;
+					UClass* AssetClass = AssetData.GetClass();
+					UObject* Asset = AssetData.GetAsset();
+					
+					if (AssetClass->IsChildOf(UAnimSequence::StaticClass()))
+					{
+						ViewModel->AddSequenceToDatabase(Cast<UAnimSequence>(Asset));
+						++AddedAssets;
+					}
+					else if (AssetClass->IsChildOf(UBlendSpace::StaticClass()))
+					{
+						ViewModel->AddBlendSpaceToDatabase(Cast<UBlendSpace>(Asset));
+						++AddedAssets;
+					}
 				}
 			}
-
+			
 			GWarn->EndSlowTask();
 		}
 
@@ -956,6 +1197,8 @@ namespace UE::PoseSearch
 
 	void SDatabaseAssetTree::OnDeleteNodes()
 	{
+		const FScopedTransaction Transaction(LOCTEXT("DeletePoseSearchDatabaseNodes", "Delete selected items from Pose Search Database"));
+
 		TArray<TSharedPtr<FDatabaseAssetTreeNode>> SelectedNodes = TreeView->GetSelectedItems();
 		if (!SelectedNodes.IsEmpty())
 		{
@@ -989,6 +1232,8 @@ namespace UE::PoseSearch
 
 	void SDatabaseAssetTree::OnEnableNodes()
 	{
+		const FScopedTransaction Transaction(LOCTEXT("EnablePoseSearchDatabaseNodes", "Enable selected items from Pose Search Database"));
+
 		TArray<TSharedPtr<FDatabaseAssetTreeNode>> SelectedNodes = TreeView->GetSelectedItems();
 		if (!SelectedNodes.IsEmpty())
 		{
@@ -1010,6 +1255,8 @@ namespace UE::PoseSearch
 
 	void SDatabaseAssetTree::OnDisableNodes()
 	{
+		const FScopedTransaction Transaction(LOCTEXT("DisablePoseSearchDatabaseNodes", "Disable selected items from Pose Search Database"));
+
 		TArray<TSharedPtr<FDatabaseAssetTreeNode>> SelectedNodes = TreeView->GetSelectedItems();
 		if (!SelectedNodes.IsEmpty())
 		{
diff --git a/Engine/Plugins/Experimental/Animation/PoseSearch/Source/Editor/Public/SPoseSearchDatabaseAssetList.h b/Engine/Plugins/Experimental/Animation/PoseSearch/Source/Editor/Public/SPoseSearchDatabaseAssetList.h
--- a/Engine/Plugins/Experimental/Animation/PoseSearch/Source/Editor/Public/SPoseSearchDatabaseAssetList.h
+++ b/Engine/Plugins/Experimental/Animation/PoseSearch/Source/Editor/Public/SPoseSearchDatabaseAssetList.h
@@ -3,6 +3,7 @@
 #pragma once
 
 #include "CoreMinimal.h"
+#include "DetailColumnSizeData.h"
 #include "PoseSearch/PoseSearch.h"
 #include "Widgets/Views/STreeView.h"
 #include "EditorUndoClient.h"
@@ -29,6 +30,10 @@ namespace UE::PoseSearch
 			TSharedRef<FUICommandList> InCommandList,
 			TSharedPtr<SDatabaseAssetTree> InHierarchy);
 
+		bool IsRootMotionEnabled() const;
+		bool IsLooping() const;
+		EPoseSearchMirrorOption GetMirrorOption() const;
+		
 		int32 SourceAssetIdx;
 		ESearchIndexAssetType SourceAssetType;
 		TSharedPtr<FDatabaseAssetTreeNode> Parent;
@@ -63,6 +68,13 @@ namespace UE::PoseSearch
 		EVisibility GetSelectedActorIconVisbility() const;
 
 		FSlateColor GetNameTextColorAndOpacity() const;
+		FSlateColor GetLoopingColorAndOpacity() const;
+		FText GetLoopingToolTip() const;
+		FSlateColor GetRootMotionColorAndOpacity() const;
+		FText GetRootMotionOptionToolTip() const;
+		const FSlateBrush* GetMirrorOptionSlateBrush() const;
+		FText GetMirrorOptionToolTip() const;
+		
 		FText GetAssetEnabledToolTip() const;
 		ECheckBoxState GetAssetEnabledChecked() const;
 		void OnAssetIsEnabledChanged(ECheckBoxState NewCheckboxState);
@@ -91,6 +103,8 @@ namespace UE::PoseSearch
 		void RefreshTreeView(bool bIsInitialSetup = false, bool bRecoverSelection = false);
 		void FinalizeTreeChanges(bool bRecoverSelection = false);
 
+		FDetailColumnSizeData& GetColumnSizeData() { return ColumnSizeData; }
+		
 	protected:
 		TWeakPtr<FDatabaseViewModel> EditorViewModel;
 
@@ -99,6 +113,7 @@ namespace UE::PoseSearch
 
 		/** tree view widget */
 		TSharedPtr<STreeView<TSharedPtr<FDatabaseAssetTreeNode>>> TreeView;
+		TSharedPtr<SWidget> TreeViewDragAndDropSuggestion;
 		TArray<TSharedPtr<FDatabaseAssetTreeNode>> RootNodes;
 		TArray<TSharedPtr<FDatabaseAssetTreeNode>> AllNodes;
 
@@ -148,6 +163,8 @@ namespace UE::PoseSearch
 
 		friend SDatabaseAssetListItem;
 
+		FDetailColumnSizeData ColumnSizeData;
+		
 	protected:
 		// Called when an item is selected/deselected
 		DECLARE_MULTICAST_DELEGATE_TwoParams(
