motion matching - wrapping debug draw code with ENABLE_DRAW_DEBUG, adding more ExcludeFromHash and NeverInHash meta to UPROPERTY

diff --git a/Engine/Plugins/Experimental/Animation/PoseSearch/Source/Runtime/Private/PoseSearch.cpp b/Engine/Plugins/Experimental/Animation/PoseSearch/Source/Runtime/Private/PoseSearch.cpp
--- a/Engine/Plugins/Experimental/Animation/PoseSearch/Source/Runtime/Private/PoseSearch.cpp
+++ b/Engine/Plugins/Experimental/Animation/PoseSearch/Source/Runtime/Private/PoseSearch.cpp
@@ -738,10 +738,12 @@ UE::PoseSearch::FSearchResult UPoseSearchSequenceMetaData::Search(UE::PoseSearch
 		Result.AssetTime = SearchIndex.GetAssetTime(Result.PoseIdx);
 	}
 
+#if ENABLE_DRAW_DEBUG
 	DrawFeatureVector(SearchContext.DebugDrawParams, Result.PoseIdx);
 	
 	EnumAddFlags(SearchContext.DebugDrawParams.Flags, EDebugDrawFlags::DrawQuery);
 	DrawFeatureVector(SearchContext.DebugDrawParams, QueryValues);
+#endif // ENABLE_DRAW_DEBUG
 
 	return Result;
 }
@@ -2600,6 +2602,7 @@ void FFeatureVectorHelper::SetMeanDeviations(float Deviation, Eigen::VectorXd& M
 // FDebugDrawParams
 bool FDebugDrawParams::CanDraw() const
 {
+#if ENABLE_DRAW_DEBUG
 	if (!World)
 	{
 		return false;
@@ -2612,10 +2615,14 @@ bool FDebugDrawParams::CanDraw() const
 	}
 
 	return SearchIndex->IsValid() && !SearchIndex->IsEmpty() && SearchIndex->Schema->SchemaCardinality > 0;
+#else // ENABLE_DRAW_DEBUG
+	return false;
+#endif // ENABLE_DRAW_DEBUG
 }
 
 FColor FDebugDrawParams::GetColor(int32 ColorPreset) const
 {
+#if ENABLE_DRAW_DEBUG
 	FLinearColor Color = FLinearColor::Red;
 
 	const UPoseSearchSchema* Schema = GetSchema();
@@ -2647,6 +2654,9 @@ FColor FDebugDrawParams::GetColor(int32 ColorPreset) const
 	}
 
 	return Color.ToFColor(true);
+#else // ENABLE_DRAW_DEBUG
+	return FColor::Black;
+#endif // ENABLE_DRAW_DEBUG
 }
 
 const FPoseSearchIndex* FDebugDrawParams::GetSearchIndex() const
@@ -4061,6 +4071,7 @@ FTransform FAssetIndexer::GetTransformAndCacheResults(float SampleTime, float Or
 
 void DrawFeatureVector(const FDebugDrawParams& DrawParams, TConstArrayView<float> PoseVector)
 {
+#if ENABLE_DRAW_DEBUG
 	if (DrawParams.CanDraw())
 	{
 		const UPoseSearchSchema* Schema = DrawParams.GetSchema();
@@ -4077,19 +4088,23 @@ void DrawFeatureVector(const FDebugDrawParams& DrawParams, TConstArrayView<float
 			}
 		}
 	}
+#endif // ENABLE_DRAW_DEBUG
 }
 
 void DrawFeatureVector(const FDebugDrawParams& DrawParams, int32 PoseIdx)
 {
+#if ENABLE_DRAW_DEBUG
 	// if we're editing the schema while in PIE with Rewind Debugger active, PoseIdx could be out of bound / stale
 	if (DrawParams.CanDraw() && PoseIdx >= 0 && PoseIdx < DrawParams.GetSearchIndex()->NumPoses)
 	{
 		DrawFeatureVector(DrawParams, DrawParams.GetSearchIndex()->GetPoseValues(PoseIdx));
 	}
+#endif // ENABLE_DRAW_DEBUG
 }
 
 void DrawSearchIndex(const FDebugDrawParams& DrawParams)
 {
+#if ENABLE_DRAW_DEBUG
 	if (DrawParams.CanDraw())
 	{
 		const FPoseSearchIndex* SearchIndex = DrawParams.GetSearchIndex();
@@ -4099,6 +4114,7 @@ void DrawSearchIndex(const FDebugDrawParams& DrawParams)
 			DrawFeatureVector(DrawParams, PoseIdx);
 		}
 	}
+#endif // ENABLE_DRAW_DEBUG
 }
 
 static Eigen::VectorXd ComputeChannelsDeviations(const FPoseSearchIndex* SearchIndex)
diff --git a/Engine/Plugins/Experimental/Animation/PoseSearch/Source/Runtime/Private/PoseSearchFeatureChannels.cpp b/Engine/Plugins/Experimental/Animation/PoseSearch/Source/Runtime/Private/PoseSearchFeatureChannels.cpp
--- a/Engine/Plugins/Experimental/Animation/PoseSearch/Source/Runtime/Private/PoseSearchFeatureChannels.cpp
+++ b/Engine/Plugins/Experimental/Animation/PoseSearch/Source/Runtime/Private/PoseSearchFeatureChannels.cpp
@@ -389,6 +389,7 @@ bool UPoseSearchFeatureChannel_Position::BuildQuery(UE::PoseSearch::FSearchConte
 
 void UPoseSearchFeatureChannel_Position::DebugDraw(const UE::PoseSearch::FDebugDrawParams& DrawParams, TConstArrayView<float> PoseVector) const
 {
+#if ENABLE_DRAW_DEBUG
 	using namespace UE::PoseSearch;
 
 	const UPoseSearchSchema* Schema = DrawParams.GetSchema();
@@ -416,6 +417,7 @@ void UPoseSearchFeatureChannel_Position::DebugDraw(const UE::PoseSearch::FDebugD
 	{
 		DrawDebugString(DrawParams.World, BonePos + FVector(0.0, 0.0, 10.0), Schema->BoneReferences[SchemaBoneIdx].BoneName.ToString(), nullptr, Color, LifeTime, false, 1.0f);
 	}
+#endif // ENABLE_DRAW_DEBUG
 }
 
 //////////////////////////////////////////////////////////////////////////
@@ -514,6 +516,7 @@ bool UPoseSearchFeatureChannel_Heading::BuildQuery(UE::PoseSearch::FSearchContex
 
 void UPoseSearchFeatureChannel_Heading::DebugDraw(const UE::PoseSearch::FDebugDrawParams& DrawParams, TConstArrayView<float> PoseVector) const
 {
+#if ENABLE_DRAW_DEBUG
 	using namespace UE::PoseSearch;
 
 	const UPoseSearchSchema* Schema = DrawParams.GetSchema();
@@ -541,6 +544,7 @@ void UPoseSearchFeatureChannel_Heading::DebugDraw(const UE::PoseSearch::FDebugDr
 		const float AdjustedThickness = EnumHasAnyFlags(DrawParams.Flags, EDebugDrawFlags::DrawFast) ? 0.0f : DrawDebugLineThickness;
 		DrawDebugLine(DrawParams.World, BonePos, BonePos + BoneHeading, Color, bPersistent, LifeTime, DepthPriority, AdjustedThickness);
 	}
+#endif // ENABLE_DRAW_DEBUG
 }
 
 //////////////////////////////////////////////////////////////////////////
@@ -1045,6 +1049,7 @@ bool UPoseSearchFeatureChannel_Pose::BuildQuery(UE::PoseSearch::FSearchContext&
 
 void UPoseSearchFeatureChannel_Pose::DebugDraw(const UE::PoseSearch::FDebugDrawParams& DrawParams, TConstArrayView<float> PoseVector) const
 {
+#if ENABLE_DRAW_DEBUG
 	using namespace UE::PoseSearch;
 
 	const UPoseSearchSchema* Schema = DrawParams.GetSchema();
@@ -1173,6 +1178,7 @@ void UPoseSearchFeatureChannel_Pose::DebugDraw(const UE::PoseSearch::FDebugDrawP
 	}
 
 	check(DataOffset == ChannelDataOffset + ChannelCardinality);
+#endif // ENABLE_DRAW_DEBUG
 }
 
 #if WITH_EDITOR
@@ -1755,6 +1761,7 @@ struct TrajectoryPositionReconstructor
 
 void UPoseSearchFeatureChannel_Trajectory::DebugDraw(const UE::PoseSearch::FDebugDrawParams& DrawParams, TConstArrayView<float> PoseVector) const
 {
+#if ENABLE_DRAW_DEBUG
 	using namespace UE::PoseSearch;
 
 	const float LifeTime = DrawParams.DefaultLifeTime;
@@ -2038,6 +2045,7 @@ void UPoseSearchFeatureChannel_Trajectory::DebugDraw(const UE::PoseSearch::FDebu
 	DrawCentripetalCatmullRomSpline(DrawParams.World, TrajSplinePos, TrajSplineColor, 0.5f, 8.f, bPersistent, LifeTime, DepthPriority, 0.f);
 
 	check(DataOffset == ChannelDataOffset + ChannelCardinality);
+#endif // ENABLE_DRAW_DEBUG
 }
 
 #if WITH_EDITOR
@@ -2278,6 +2286,7 @@ bool UPoseSearchFeatureChannel_FilterCrashingLegs::BuildQuery(UE::PoseSearch::FS
 
 void UPoseSearchFeatureChannel_FilterCrashingLegs::DebugDraw(const UE::PoseSearch::FDebugDrawParams& DrawParams, TConstArrayView<float> PoseVector) const
 {
+#if ENABLE_DRAW_DEBUG
 	using namespace UE::PoseSearch;
 
 	int32 DataOffset = ChannelDataOffset;
@@ -2316,6 +2325,7 @@ void UPoseSearchFeatureChannel_FilterCrashingLegs::DebugDraw(const UE::PoseSearc
 
 	DrawDebugLine(DrawParams.World, LeftFootTransform.GetTranslation(), LeftFootTransform.GetTranslation() + CrossingLegsVector, Color, bPersistent, LifeTime, DepthPriority);
 	DrawDebugLine(DrawParams.World, RightFootTransform.GetTranslation(), RightFootTransform.GetTranslation() - CrossingLegsVector, Color, bPersistent, LifeTime, DepthPriority);
+#endif // ENABLE_DRAW_DEBUG
 }
 
 // IPoseFilter interface
diff --git a/Engine/Plugins/Experimental/Animation/PoseSearch/Source/Runtime/Public/PoseSearch/PoseSearch.h b/Engine/Plugins/Experimental/Animation/PoseSearch/Source/Runtime/Public/PoseSearch/PoseSearch.h
--- a/Engine/Plugins/Experimental/Animation/PoseSearch/Source/Runtime/Public/PoseSearch/PoseSearch.h
+++ b/Engine/Plugins/Experimental/Animation/PoseSearch/Source/Runtime/Public/PoseSearch/PoseSearch.h
@@ -200,7 +200,7 @@ struct POSESEARCH_API FAssetSamplingContext
 	static constexpr float FiniteDelta = 1 / 60.0f;
 
 	// Mirror data table pointer copied from Schema for convenience
-	TObjectPtr<const UMirrorDataTable> MirrorDataTable = nullptr;
+	TObjectPtr<const UMirrorDataTable> MirrorDataTable;
 
 	// Compact pose format of Mirror Bone Map
 	TCustomBoneIndexArray<FCompactPoseBoneIndex, FCompactPoseBoneIndex> CompactPoseMirrorBones;
@@ -533,10 +533,10 @@ struct FPoseSearchSchemaColorPreset
 {
 	GENERATED_BODY()
 
-	UPROPERTY(EditAnywhere, Category = "Colors")
+	UPROPERTY(EditAnywhere, Category = "Colors", meta = (ExcludeFromHash))
 	FLinearColor Query = FLinearColor::Blue;
 
-	UPROPERTY(EditAnywhere, Category = "Colors")
+	UPROPERTY(EditAnywhere, Category = "Colors", meta = (ExcludeFromHash))
 	FLinearColor Result = FLinearColor::Yellow;
 };
 
@@ -550,7 +550,7 @@ class POSESEARCH_API UPoseSearchSchema : public UDataAsset, public IBoneReferenc
 
 public:
 	UPROPERTY(EditAnywhere, Category = "Schema")
-	TObjectPtr<USkeleton> Skeleton = nullptr;
+	TObjectPtr<USkeleton> Skeleton;
 
 	UPROPERTY(EditAnywhere, meta = (ClampMin = "1", ClampMax = "240"), Category = "Schema")
 	int32 SampleRate = 30;
@@ -590,7 +590,7 @@ public:
 	UPROPERTY(EditAnywhere, Category = "Schema")
 	float MirrorMismatchCostBias = 0.f;
 
-	UPROPERTY(EditAnywhere, Category = "Schema")
+	UPROPERTY(EditAnywhere, Category = "Schema", meta = (ExcludeFromHash))
 	TArray<FPoseSearchSchemaColorPreset> ColorPresets;
 	
 	bool IsValid () const;
@@ -688,26 +688,26 @@ public:
 
 	// Default to Sequence for now for backward compatibility but
 	// at some point we might want to change this to Invalid.
-	UPROPERTY()
+	UPROPERTY(meta = (NeverInHash))
 	ESearchIndexAssetType Type = ESearchIndexAssetType::Sequence;
 
 	// Index of the source asset in search index's container (i.e. UPoseSearchDatabase)
-	UPROPERTY()
+	UPROPERTY(meta = (NeverInHash))
 	int32 SourceAssetIdx = INDEX_NONE;
 
-	UPROPERTY()
+	UPROPERTY(meta = (NeverInHash))
 	bool bMirrored = false;
 
-	UPROPERTY()
+	UPROPERTY(meta = (NeverInHash))
 	FVector BlendParameters = FVector::Zero();
 
-	UPROPERTY()
+	UPROPERTY(meta = (NeverInHash))
 	FFloatInterval SamplingInterval;
 
-	UPROPERTY()
+	UPROPERTY(meta = (NeverInHash))
 	int32 FirstPoseIdx = INDEX_NONE;
 
-	UPROPERTY()
+	UPROPERTY(meta = (NeverInHash))
 	int32 NumPoses = 0;
 
 	bool IsPoseInRange(int32 PoseIdx) const
@@ -725,29 +725,29 @@ struct POSESEARCH_API FPoseSearchIndex
 {
 	GENERATED_BODY()
 
-	UPROPERTY(Category = Info, VisibleAnywhere)
+	UPROPERTY(Category = Info, VisibleAnywhere, meta = (NeverInHash))
 	int32 NumPoses = 0;
 
-	UPROPERTY()
+	UPROPERTY(meta = (NeverInHash))
 	TArray<float> Values;
 
-	UPROPERTY()
+	UPROPERTY(meta = (NeverInHash))
 	TArray<float> PCAValues;
 
 #if WITH_EDITORONLY_DATA
-	UPROPERTY(Category = Info, VisibleAnywhere)
+	UPROPERTY(Category = Info, VisibleAnywhere, meta = (NeverInHash))
 	float PCAExplainedVariance = 0.f;
 	
-	UPROPERTY(Category = Info, VisibleAnywhere)
+	UPROPERTY(Category = Info, VisibleAnywhere, meta = (NeverInHash))
 	TArray<float> Deviation;
 #endif // WITH_EDITORONLY_DATA
 
 	UE::PoseSearch::FKDTree KDTree;
 
-	UPROPERTY(Category = Info, VisibleAnywhere)
+	UPROPERTY(Category = Info, VisibleAnywhere, meta = (NeverInHash))
 	TArray<float> PCAProjectionMatrix;
 
-	UPROPERTY(Category = Info, VisibleAnywhere)
+	UPROPERTY(Category = Info, VisibleAnywhere, meta = (NeverInHash))
 	TArray<float> Mean;
 
 	// we store weights square roots to reduce numerical errors when CompareFeatureVectors 
@@ -755,23 +755,23 @@ struct POSESEARCH_API FPoseSearchIndex
 	// instead of
 	// ((VA - VB).square() * VW).sum()
 	// since (VA - VB).square() could lead to big numbers, and VW being multiplied by the variance of the dataset
-	UPROPERTY(Category = Info, VisibleAnywhere) 
+	UPROPERTY(Category = Info, VisibleAnywhere, meta = (NeverInHash))
 	TArray<float> WeightsSqrt;
 
-	UPROPERTY()
+	UPROPERTY(meta = (NeverInHash))
 	TArray<FPoseSearchPoseMetadata> PoseMetadata;
 
-	UPROPERTY()
+	UPROPERTY(meta = (NeverInHash))
 	EPoseSearchPoseFlags OverallFlags = EPoseSearchPoseFlags::None;
 
-	UPROPERTY()
-	TObjectPtr<const UPoseSearchSchema> Schema = nullptr;
+	UPROPERTY(meta = (NeverInHash))
+	TObjectPtr<const UPoseSearchSchema> Schema;
 
-	UPROPERTY()
+	UPROPERTY(meta = (NeverInHash))
 	TArray<FPoseSearchIndexAsset> Assets;
 
 	// minimum of the database metadata CostAddend: it represents the minimum cost of any search for the associated database (we'll skip the search in case the search result total cost is already less than MinCostAddend)
-	UPROPERTY(Category = Info, VisibleAnywhere)
+	UPROPERTY(Category = Info, VisibleAnywhere, meta = (NeverInHash))
 	float MinCostAddend = -MAX_FLT;
 
 	bool IsValid() const;
@@ -814,7 +814,7 @@ struct POSESEARCH_API FPoseSearchDatabaseSequence : public FPoseSearchDatabaseAn
 	GENERATED_BODY()
 
 	UPROPERTY(EditAnywhere, Category="Sequence")
-	TObjectPtr<UAnimSequence> Sequence = nullptr;
+	TObjectPtr<UAnimSequence> Sequence;
 
 	UPROPERTY(EditAnywhere, Category = "Sequence")
 	bool bEnabled = true;
@@ -830,14 +830,14 @@ struct POSESEARCH_API FPoseSearchDatabaseSequence : public FPoseSearchDatabaseAn
 	// for one shot anims with past sampling. When past sampling is used without a lead in sequence,
 	// the sampling range of the main sequence will be clamped if necessary.
 	UPROPERTY(EditAnywhere, Category="Sequence")
-	TObjectPtr<UAnimSequence> LeadInSequence = nullptr;
+	TObjectPtr<UAnimSequence> LeadInSequence;
 
 	// Used for sampling future pose information at the end of the main sequence.
 	// This setting is intended for transitions between cycles. It is optional and only used
 	// for one shot anims with future sampling. When future sampling is used without a follow up sequence,
 	// the sampling range of the main sequence will be clamped if necessary.
 	UPROPERTY(EditAnywhere, Category="Sequence")
-	TObjectPtr<UAnimSequence> FollowUpSequence = nullptr;
+	TObjectPtr<UAnimSequence> FollowUpSequence;
 
 	virtual UAnimationAsset* GetAnimationAsset() const override { return Sequence; }
 	virtual bool IsLooping() const override { return Sequence->bLoop; }
@@ -853,7 +853,7 @@ struct POSESEARCH_API FPoseSearchDatabaseBlendSpace : public FPoseSearchDatabase
 	GENERATED_BODY()
 
 	UPROPERTY(EditAnywhere, Category = "BlendSpace")
-	TObjectPtr<UBlendSpace> BlendSpace = nullptr;
+	TObjectPtr<UBlendSpace> BlendSpace;
 
 	UPROPERTY(EditAnywhere, Category = "BlendSpace")
 	bool bEnabled = true;
@@ -1028,7 +1028,7 @@ struct FSearchResult
 	// lerp value to find AssetTime from PrevPoseIdx -> AssetTime -> NextPoseIdx, within range [-0.5, 0.5]
 	float LerpValue = 0.f;
 
-	TWeakObjectPtr<const UPoseSearchDatabase> Database = nullptr;
+	TWeakObjectPtr<const UPoseSearchDatabase> Database;
 	FPoseSearchFeatureVectorBuilder ComposedQuery;
 
 	// cost of the current pose with the query from database in the result, if possible
@@ -1075,7 +1075,7 @@ class POSESEARCH_API UPoseSearchDatabase : public UPoseSearchSearchableAsset
 public:
 	// Motion Database Config asset to use with this database.
 	UPROPERTY(EditAnywhere, BlueprintReadOnly, Category="Database", DisplayName="Config")
-	TObjectPtr<const UPoseSearchSchema> Schema = nullptr;
+	TObjectPtr<const UPoseSearchSchema> Schema;
 
 	UPROPERTY(EditAnywhere, Category = "Database")
 	FPoseSearchExtrapolationParameters ExtrapolationParameters;
@@ -1215,15 +1215,15 @@ class POSESEARCH_API UPoseSearchSequenceMetaData : public UAnimMetaData
 public:
 
 	UPROPERTY(EditAnywhere, Category = "Settings")
-	TObjectPtr<const UPoseSearchSchema> Schema = nullptr;
+	TObjectPtr<const UPoseSearchSchema> Schema;
 
 	UPROPERTY(EditAnywhere, Category = "Settings")
-	FFloatInterval SamplingRange = FFloatInterval(0.0f, 0.0f);
+	FFloatInterval SamplingRange = FFloatInterval(0.f, 0.f);
 
 	UPROPERTY(EditAnywhere, Category = "Settings")
 	FPoseSearchExtrapolationParameters ExtrapolationParameters;
 
-	UPROPERTY()
+	UPROPERTY(meta = (ExcludeFromHash))
 	FPoseSearchIndex SearchIndex;
 
 	bool IsValidForIndexing() const;
@@ -1518,7 +1518,7 @@ struct POSESEARCH_API FPoseSearchDatabaseSetEntry
 
 public:
 	UPROPERTY(EditAnywhere, Category = Settings)
-	TObjectPtr<UPoseSearchSearchableAsset> Searchable = nullptr;
+	TObjectPtr<UPoseSearchSearchableAsset> Searchable;
 
 	UPROPERTY(EditAnywhere, Category = Settings)
 	FGameplayTag Tag;
diff --git a/Engine/Plugins/Experimental/Animation/PoseSearch/Source/Runtime/Public/PoseSearch/PoseSearchDerivedDataKey.h b/Engine/Plugins/Experimental/Animation/PoseSearch/Source/Runtime/Public/PoseSearch/PoseSearchDerivedDataKey.h
--- a/Engine/Plugins/Experimental/Animation/PoseSearch/Source/Runtime/Public/PoseSearch/PoseSearchDerivedDataKey.h
+++ b/Engine/Plugins/Experimental/Animation/PoseSearch/Source/Runtime/Public/PoseSearch/PoseSearchDerivedDataKey.h
@@ -25,6 +25,7 @@ public:
 	using HashBuilderType = HashBuilder;
 
 	inline static const FName ExcludeFromHashName = FName(TEXT("ExcludeFromHash"));
+	inline static const FName NeverInHashName = FName(TEXT("NeverInHash"));
 
 	/**
 	* Default constructor.
@@ -57,8 +58,14 @@ public:
 		{
 			return true;
 		}
-		const bool bExclude = InProperty->HasMetaData(ExcludeFromHashName);
-		return bExclude;
+		
+		if (InProperty->HasMetaData(ExcludeFromHashName))
+		{
+			return true;
+		}
+		
+		check(!InProperty->HasMetaData(NeverInHashName));
+		return false;
 	}
 
 	virtual void Serialize(void* Data, int64 Length) override
diff --git a/Engine/Plugins/Experimental/Animation/PoseSearch/Source/Runtime/Public/PoseSearch/PoseSearchFeatureChannels.h b/Engine/Plugins/Experimental/Animation/PoseSearch/Source/Runtime/Public/PoseSearch/PoseSearchFeatureChannels.h
--- a/Engine/Plugins/Experimental/Animation/PoseSearch/Source/Runtime/Public/PoseSearch/PoseSearchFeatureChannels.h
+++ b/Engine/Plugins/Experimental/Animation/PoseSearch/Source/Runtime/Public/PoseSearch/PoseSearchFeatureChannels.h
@@ -44,7 +44,7 @@ public:
 	UPROPERTY()
 	int8 SchemaBoneIdx = 0;
 
-	UPROPERTY(EditAnywhere, Category = "Settings")
+	UPROPERTY(EditAnywhere, Category = "Settings", meta = (ExcludeFromHash))
 	int32 ColorPresetIndex = 0;
 
 	UPROPERTY(EditAnywhere, Category = "Settings")
@@ -98,7 +98,7 @@ public:
 	UPROPERTY()
 	int8 SchemaBoneIdx = 0;
 
-	UPROPERTY(EditAnywhere, Category = "Settings")
+	UPROPERTY(EditAnywhere, Category = "Settings", meta = (ExcludeFromHash))
 	int32 ColorPresetIndex = 0;
 
 	UPROPERTY(EditAnywhere, Category = "Settings")
@@ -234,7 +234,7 @@ struct POSESEARCH_API FPoseSearchTrajectorySample
 	UPROPERTY(EditAnywhere, Category = Config)
 	float Weight = 1.f;
 
-	UPROPERTY(EditAnywhere, Category = Config)
+	UPROPERTY(EditAnywhere, Category = Config, meta = (ExcludeFromHash))
 	int32 ColorPresetIndex = 0;
 };
 
