Added option in Pose Search Database Viewport to view current root motion speed of the root joint for a selected clip

#rn minor

diff --git a/Engine/Plugins/Experimental/Animation/PoseSearch/Source/Editor/Private/PoseSearchDatabaseEditorCommands.cpp b/Engine/Plugins/Experimental/Animation/PoseSearch/Source/Editor/Private/PoseSearchDatabaseEditorCommands.cpp
--- a/Engine/Plugins/Experimental/Animation/PoseSearch/Source/Editor/Private/PoseSearchDatabaseEditorCommands.cpp
+++ b/Engine/Plugins/Experimental/Animation/PoseSearch/Source/Editor/Private/PoseSearchDatabaseEditorCommands.cpp
@@ -30,6 +30,7 @@ namespace UE::PoseSearch
 			ShowAnimationOriginalAndMirrored, "Original and mirrored", "Show original and mirrored animations",
 			EUserInterfaceActionType::RadioButton, FInputChord());
 
+		UI_COMMAND(ShowDisplayRootMotionSpeed, "Display Root Motion Speed", "Show root motion speed for the selected item", EUserInterfaceActionType::ToggleButton, FInputChord());
 	}
 }
 
diff --git a/Engine/Plugins/Experimental/Animation/PoseSearch/Source/Editor/Private/PoseSearchDatabaseViewModel.cpp b/Engine/Plugins/Experimental/Animation/PoseSearch/Source/Editor/Private/PoseSearchDatabaseViewModel.cpp
--- a/Engine/Plugins/Experimental/Animation/PoseSearch/Source/Editor/Private/PoseSearchDatabaseViewModel.cpp
+++ b/Engine/Plugins/Experimental/Animation/PoseSearch/Source/Editor/Private/PoseSearchDatabaseViewModel.cpp
@@ -411,6 +411,16 @@ namespace UE::PoseSearch
 		return EnumHasAnyFlags(AnimationPreviewMode, PreviewMode);
 	}
 
+	void FDatabaseViewModel::OnToggleDisplayRootMotionSpeed()
+	{
+		DisplayRootMotionSpeed = !DisplayRootMotionSpeed;
+	}
+
+	bool FDatabaseViewModel::CanDisplayRootMotionSpeed() const
+	{
+		return SelectedNodes.Num() <= 1;
+	}
+
 	void FDatabaseViewModel::AddSequenceToDatabase(UAnimSequence* AnimSequence)
 	{
 		FPoseSearchDatabaseSequence& NewDbSequence = PoseSearchDatabase->Sequences.AddDefaulted_GetRef();
diff --git a/Engine/Plugins/Experimental/Animation/PoseSearch/Source/Editor/Private/SPoseSearchDatabaseViewport.cpp b/Engine/Plugins/Experimental/Animation/PoseSearch/Source/Editor/Private/SPoseSearchDatabaseViewport.cpp
--- a/Engine/Plugins/Experimental/Animation/PoseSearch/Source/Editor/Private/SPoseSearchDatabaseViewport.cpp
+++ b/Engine/Plugins/Experimental/Animation/PoseSearch/Source/Editor/Private/SPoseSearchDatabaseViewport.cpp
@@ -108,6 +108,12 @@ namespace UE::PoseSearch
 				ViewModelRef,
 				&FDatabaseViewModel::IsAnimationPreviewMode,
 				EAnimationPreviewMode::OriginalAndMirrored));
+
+		CommandList->MapAction(
+			Commands.ShowDisplayRootMotionSpeed,
+			FExecuteAction::CreateSP(ViewModelRef, &FDatabaseViewModel::OnToggleDisplayRootMotionSpeed),
+			FCanExecuteAction::CreateSP(ViewModelRef, &FDatabaseViewModel::CanDisplayRootMotionSpeed),
+			FIsActionChecked::CreateSP(ViewModelRef, &FDatabaseViewModel::IsDisplayRootMotionSpeedChecked));
 	}
 
 	TSharedRef<FEditorViewportClient> SDatabaseViewport::MakeEditorViewportClient()
diff --git a/Engine/Plugins/Experimental/Animation/PoseSearch/Source/Editor/Private/SPoseSearchDatabaseViewportToolbar.cpp b/Engine/Plugins/Experimental/Animation/PoseSearch/Source/Editor/Private/SPoseSearchDatabaseViewportToolbar.cpp
--- a/Engine/Plugins/Experimental/Animation/PoseSearch/Source/Editor/Private/SPoseSearchDatabaseViewportToolbar.cpp
+++ b/Engine/Plugins/Experimental/Animation/PoseSearch/Source/Editor/Private/SPoseSearchDatabaseViewportToolbar.cpp
@@ -63,6 +63,22 @@ namespace UE::PoseSearch
 			})
 			);
 
+			
+			ShowMenuBuilder.AddSubMenu(
+				LOCTEXT("ShowMenu_DebugSubMenu", "Debug"),
+				LOCTEXT("ShowMenu_DebugSubMenuToolTip", "Debug Options"),
+				FNewMenuDelegate::CreateLambda([](FMenuBuilder& SubMenuBuilder)
+			{
+				const FDatabaseEditorCommands& Commands = FDatabaseEditorCommands::Get();
+
+				SubMenuBuilder.BeginSection("Debug", LOCTEXT("ShowMenu_DebugLabel", "Debug"));
+				{
+					SubMenuBuilder.AddMenuEntry(Commands.ShowDisplayRootMotionSpeed);
+				}
+				SubMenuBuilder.EndSection();
+			})
+			);
+			
 		}
 
 		return ShowMenuBuilder.MakeWidget();
diff --git a/Engine/Plugins/Experimental/Animation/PoseSearch/Source/Editor/Public/PoseSearchDatabaseEditorCommands.h b/Engine/Plugins/Experimental/Animation/PoseSearch/Source/Editor/Public/PoseSearchDatabaseEditorCommands.h
--- a/Engine/Plugins/Experimental/Animation/PoseSearch/Source/Editor/Public/PoseSearchDatabaseEditorCommands.h
+++ b/Engine/Plugins/Experimental/Animation/PoseSearch/Source/Editor/Public/PoseSearchDatabaseEditorCommands.h
@@ -33,7 +33,8 @@ namespace UE::PoseSearch
 		TSharedPtr<FUICommandInfo> ShowAnimationNone;
 		TSharedPtr<FUICommandInfo> ShowAnimationOriginalOnly;
 		TSharedPtr<FUICommandInfo> ShowAnimationOriginalAndMirrored;
-
+		
+		TSharedPtr<FUICommandInfo> ShowDisplayRootMotionSpeed;
 	};
 }
 
diff --git a/Engine/Plugins/Experimental/Animation/PoseSearch/Source/Editor/Public/PoseSearchDatabaseViewModel.h b/Engine/Plugins/Experimental/Animation/PoseSearch/Source/Editor/Public/PoseSearchDatabaseViewModel.h
--- a/Engine/Plugins/Experimental/Animation/PoseSearch/Source/Editor/Public/PoseSearchDatabaseViewModel.h
+++ b/Engine/Plugins/Experimental/Animation/PoseSearch/Source/Editor/Public/PoseSearchDatabaseViewModel.h
@@ -97,6 +97,10 @@ namespace UE::PoseSearch
 		void OnSetAnimationPreviewMode(EAnimationPreviewMode PreviewMode);
 		bool IsAnimationPreviewMode(EAnimationPreviewMode PreviewMode) const;
 
+		void OnToggleDisplayRootMotionSpeed();
+		bool CanDisplayRootMotionSpeed() const;
+		bool IsDisplayRootMotionSpeedChecked() const { return DisplayRootMotionSpeed; };
+		
 		void AddSequenceToDatabase(UAnimSequence* AnimSequence);
 		void AddBlendSpaceToDatabase(UBlendSpace* BlendSpace);
 
@@ -140,6 +144,9 @@ namespace UE::PoseSearch
 		/** What animations to show in the viewport */
 		EAnimationPreviewMode AnimationPreviewMode = EAnimationPreviewMode::OriginalAndMirrored;
 
+		/** Is animation debug draw enabled */
+		bool DisplayRootMotionSpeed = false;
+		
 		TArray<TSharedPtr<FDatabaseAssetTreeNode>> SelectedNodes;
 
 		int32 SelectedActorIndexAssetIndex = INDEX_NONE;
