motion matching - kdtree deserialization fix

diff --git a/Engine/Plugins/Experimental/Animation/PoseSearch/Source/Runtime/Private/KDTree.cpp b/Engine/Plugins/Experimental/Animation/PoseSearch/Source/Runtime/Private/KDTree.cpp
--- a/Engine/Plugins/Experimental/Animation/PoseSearch/Source/Runtime/Private/KDTree.cpp
+++ b/Engine/Plugins/Experimental/Animation/PoseSearch/Source/Runtime/Private/KDTree.cpp
@@ -74,6 +74,7 @@ void FKDTree::Construct(int32 Count, int32 Dim, const float* Data, int32 MaxLeaf
 bool FKDTree::FindNeighbors(KNNResultSet& Result, const float* Query) const
 {
 #if UE_POSE_SEARCH_USE_NANOFLANN
+	check(Query && Impl->root_node);
 	const nanoflann::SearchParams SearchParams(
 		32,			// Ignored parameter (Kept for compatibility with the FLANN interface).
 		0.f,		// search for eps-approximate neighbours (default: 0)
@@ -102,7 +103,7 @@ SIZE_T FKDTree::GetAllocatedSize() const
 
 
 #if UE_POSE_SEARCH_USE_NANOFLANN
-FArchive& SerializeSubTree(FArchive& Ar, FKDTree& KDTree, FKDTreeImplementation::NodePtr KDTreeNode)
+FArchive& SerializeSubTree(FArchive& Ar, FKDTree& KDTree, FKDTreeImplementation::NodePtr& KDTreeNode)
 {
 	check(KDTree.Impl);
 
