Motion Matching Editor - Database Stats Tab

diff --git a/Engine/Plugins/Experimental/Animation/PoseSearch/Source/Editor/Private/PoseSearchDatabaseEditor.cpp b/Engine/Plugins/Experimental/Animation/PoseSearch/Source/Editor/Private/PoseSearchDatabaseEditor.cpp
--- a/Engine/Plugins/Experimental/Animation/PoseSearch/Source/Editor/Private/PoseSearchDatabaseEditor.cpp
+++ b/Engine/Plugins/Experimental/Animation/PoseSearch/Source/Editor/Private/PoseSearchDatabaseEditor.cpp
@@ -24,7 +24,7 @@
 namespace UE::PoseSearch
 {
 	const FName PoseSearchDatabaseEditorAppName = FName(TEXT("PoseSearchDatabaseEditorApp"));
-	constexpr int32 SelectionDetailsCount = 3; // sequence, blendspace, group
+	constexpr int32 SelectionDetailsCount = 2; // sequence, blendspace
 	constexpr double ViewRangeSlack = 0.25;
 
 	// Tab identifiers
@@ -35,6 +35,7 @@ namespace UE::PoseSearch
 		static const FName PreviewSettingsID;
 		static const FName AssetTreeViewID;
 		static const FName SelectionDetailsID;
+		static const FName StatisticsOverview;
 	};
 
 	const FName FDatabaseEditorTabs::AssetDetailsID(TEXT("PoseSearchDatabaseEditorAssetDetailsTabID"));
@@ -42,7 +43,8 @@ namespace UE::PoseSearch
 	const FName FDatabaseEditorTabs::PreviewSettingsID(TEXT("PoseSearchDatabaseEditorPreviewSettingsTabID"));
 	const FName FDatabaseEditorTabs::AssetTreeViewID(TEXT("PoseSearchDatabaseEditorAssetTreeViewTabID"));
 	const FName FDatabaseEditorTabs::SelectionDetailsID(TEXT("PoseSearchDatabaseEditorSelectionDetailsID"));
-
+	const FName FDatabaseEditorTabs::StatisticsOverview(TEXT("PoseSearchDatabaseEditorStatisticsOverviewID"));
+	
 	FDatabaseEditor::FDatabaseEditor()
 	{
 	}
@@ -135,32 +137,37 @@ namespace UE::PoseSearch
 		ViewModel->Initialize(DatabaseAsset, PreviewScene.ToSharedRef());
 
 		// Create viewport widget
-		FDatabasePreviewRequiredArgs PreviewArgs(
-			StaticCastSharedRef<FDatabaseEditor>(AsShared()),
-			PreviewScene.ToSharedRef());
-		PreviewWidget = SNew(SDatabasePreview, PreviewArgs)
-			.SliderScrubTime_Lambda([this]() { return ViewModel->GetPlayTime(); })
-			.SliderViewRange_Lambda([this]() 
-			{ 
-				return TRange<double>(-ViewRangeSlack, ViewModel->GetMaxPreviewPlayLength() + ViewRangeSlack);
-			})
-			.OnSliderScrubPositionChanged_Lambda([this](float NewScrubPosition, bool bScrubbing)
-			{
-				ViewModel->SetPlayTime(NewScrubPosition, !bScrubbing);
-			})
-			.OnBackwardEnd_Raw(this, &FDatabaseEditor::PreviewBackwardEnd)
-			.OnBackwardStep_Raw(this, &FDatabaseEditor::PreviewBackwardStep)
-			.OnBackward_Raw(this, &FDatabaseEditor::PreviewBackward)
-			.OnPause_Raw(this, &FDatabaseEditor::PreviewPause)
-			.OnForward_Raw(this, &FDatabaseEditor::PreviewForward)
-			.OnForwardStep_Raw(this, &FDatabaseEditor::PreviewForwardStep)
-			.OnForwardEnd_Raw(this, &FDatabaseEditor::PreviewForwardEnd);
+		{
+			FDatabasePreviewRequiredArgs PreviewArgs(
+				StaticCastSharedRef<FDatabaseEditor>(AsShared()),
+				PreviewScene.ToSharedRef());
+			
+			PreviewWidget = SNew(SDatabasePreview, PreviewArgs)
+				.SliderScrubTime_Lambda([this]() { return ViewModel->GetPlayTime(); })
+				.SliderViewRange_Lambda([this]() 
+				{ 
+					return TRange<double>(-ViewRangeSlack, ViewModel->GetMaxPreviewPlayLength() + ViewRangeSlack);
+				})
+				.OnSliderScrubPositionChanged_Lambda([this](float NewScrubPosition, bool bScrubbing)
+				{
+					ViewModel->SetPlayTime(NewScrubPosition, !bScrubbing);
+				})
+				.OnBackwardEnd_Raw(this, &FDatabaseEditor::PreviewBackwardEnd)
+				.OnBackwardStep_Raw(this, &FDatabaseEditor::PreviewBackwardStep)
+				.OnBackward_Raw(this, &FDatabaseEditor::PreviewBackward)
+				.OnPause_Raw(this, &FDatabaseEditor::PreviewPause)
+				.OnForward_Raw(this, &FDatabaseEditor::PreviewForward)
+				.OnForwardStep_Raw(this, &FDatabaseEditor::PreviewForwardStep)
+				.OnForwardEnd_Raw(this, &FDatabaseEditor::PreviewForwardEnd);
+		}
 
+		// Create asset tree widget
 		AssetTreeWidget = SNew(SDatabaseAssetTree, ViewModel.ToSharedRef());
 		AssetTreeWidget->RegisterOnSelectionChanged(
 			SDatabaseAssetTree::FOnSelectionChanged::CreateSP(
 				this,
 				&FDatabaseEditor::OnAssetTreeSelectionChanged));
+		
 		if (IsValid(DatabaseAsset))
 		{
 			DatabaseAsset->RegisterOnAssetChange(
@@ -172,38 +179,77 @@ namespace UE::PoseSearch
 					AssetTreeWidget.Get(),
 					&SDatabaseAssetTree::RefreshTreeView, false, false));
 		}
-
+		
 		// Create Asset Details widget
-		FPropertyEditorModule& PropertyModule = 
-			FModuleManager::LoadModuleChecked<FPropertyEditorModule>("PropertyEditor");
+		FPropertyEditorModule& PropertyModule = FModuleManager::LoadModuleChecked<FPropertyEditorModule>("PropertyEditor");
 
-		FDetailsViewArgs DatabaseDetailsArgs;
+		// Database details widget
 		{
-			DatabaseDetailsArgs.bHideSelectionTip = true;
-			DatabaseDetailsArgs.NotifyHook = this;
-		}
-
-		EditingAssetWidget = PropertyModule.CreateDetailView(DatabaseDetailsArgs);
-		EditingAssetWidget->SetObject(DatabaseAsset);
+			FDetailsViewArgs DatabaseDetailsArgs;
+			{
+				DatabaseDetailsArgs.bHideSelectionTip = true;
+				DatabaseDetailsArgs.NotifyHook = this;
+			}
 
-		FDetailsViewArgs SelectionDetailsArgs;
+			EditingAssetWidget = PropertyModule.CreateDetailView(DatabaseDetailsArgs);
+			EditingAssetWidget->SetObject(DatabaseAsset);
+		}
+		
+		// Selection details widgets
 		{
-			SelectionDetailsArgs.bHideSelectionTip = true;
-			SelectionDetailsArgs.NotifyHook = this;
-			SelectionDetailsArgs.bShowScrollBar = false;
+			FDetailsViewArgs SelectionDetailsArgs;
+			{
+				SelectionDetailsArgs.bHideSelectionTip = true;
+				SelectionDetailsArgs.NotifyHook = this;
+				SelectionDetailsArgs.bShowScrollBar = false;
+			}
+			
+			SelectionWidgets.Reset(SelectionDetailsCount);
+			for (int32 i = 0; i < SelectionDetailsCount; ++i)
+			{
+				TSharedPtr<IDetailsView>& SelectionWidget = 
+					SelectionWidgets.Add_GetRef(PropertyModule.CreateDetailView(SelectionDetailsArgs));
+				SelectionWidget->SetObject(nullptr);
+				SelectionWidget->OnFinishedChangingProperties().AddSP(
+					this,
+					&FDatabaseEditor::OnFinishedChangingSelectionProperties);
+			}
 		}
 
-		SelectionWidgets.Reset(SelectionDetailsCount);
-		for (int32 i = 0; i < SelectionDetailsCount; ++i)
+		// Statistics details widgets
 		{
-			TSharedPtr<IDetailsView>& SelectionWidget = 
-				SelectionWidgets.Add_GetRef(PropertyModule.CreateDetailView(SelectionDetailsArgs));
-			SelectionWidget->SetObject(nullptr);
-			SelectionWidget->OnFinishedChangingProperties().AddSP(
-				this,
-				&FDatabaseEditor::OnFinishedChangingSelectionProperties);
-		}
+			FDetailsViewArgs StatisticsOverviewDetailsArgs;
+			{
+				StatisticsOverviewDetailsArgs.bHideSelectionTip = true;
+				StatisticsOverviewDetailsArgs.NotifyHook = this;
+				StatisticsOverviewDetailsArgs.bAllowSearch = false;
+			}
+			
+			StatisticsOverviewWidget = PropertyModule.CreateDetailView(StatisticsOverviewDetailsArgs);
+			StatisticsOverviewWidget->SetObject(nullptr);
 
+			// Ensure statistics information get updated/populated
+			if (DatabaseAsset)
+			{
+				const auto UpdateStatisticsInformationLambda = [this]()
+				{
+					UPoseSearchDatabaseStatistics* Statistics = NewObject<UPoseSearchDatabaseStatistics>();
+					Statistics->AddToRoot();
+					Statistics->Initialize(GetPoseSearchDatabase());
+					StatisticsOverviewWidget->SetObject(Statistics);
+				};
+
+				// Init statistics
+				UpdateStatisticsInformationLambda();
+
+				// Ensure any database changes are reflected
+				DatabaseAsset->RegisterOnAssetChange(UPoseSearchDatabase::FOnAssetChange::CreateLambda(UpdateStatisticsInformationLambda));
+				DatabaseAsset->RegisterOnGroupChange(UPoseSearchDatabase::FOnGroupChange::CreateLambda(UpdateStatisticsInformationLambda));
+				DatabaseAsset->RegisterOnDerivedDataRebuild(UPoseSearchDatabase::FOnDerivedDataRebuild::CreateLambda(UpdateStatisticsInformationLambda));
+				// TODO: Register to delegate that gets triggered after the DerivedData has been rebuilt.
+			}
+		}
+		
 		// Define Editor Layout
 		const TSharedRef<FTabManager::FLayout> StandaloneDefaultLayout =
 		FTabManager::NewLayout("Standalone_PoseSearchDatabaseEditor_Layout_v0.07")
@@ -251,6 +297,12 @@ namespace UE::PoseSearch
 								->AddTab(FDatabaseEditorTabs::SelectionDetailsID, ETabState::OpenedTab)
 								->AddTab(FDatabaseEditorTabs::PreviewSettingsID, ETabState::OpenedTab)
 							)
+							->Split
+							(
+								FTabManager::NewStack()
+								->SetSizeCoefficient(0.5f)
+								->AddTab(FDatabaseEditorTabs::StatisticsOverview, ETabState::OpenedTab)
+							)
 						)
 					)
 				)
@@ -349,6 +401,13 @@ namespace UE::PoseSearch
 			.SetDisplayName(LOCTEXT("SelectionDetailsTab", "Selection Details"))
 			.SetGroup(WorkspaceMenuCategoryRef)
 			.SetIcon(FSlateIcon(FAppStyle::GetAppStyleSetName(), "GraphEditor.EventGraph_16x"));
+
+		InTabManager->RegisterTabSpawner(
+		FDatabaseEditorTabs::StatisticsOverview,
+		FOnSpawnTab::CreateSP(this, &FDatabaseEditor::SpawnTab_StatisticsOverview))
+		.SetDisplayName(LOCTEXT("StatisticsOverviewTab", "Statistics Overview"))
+		.SetGroup(WorkspaceMenuCategoryRef)
+		.SetIcon(FSlateIcon(FAppStyle::GetAppStyleSetName(), "GraphEditor.EventGraph_16x"));
 	}
 
 	void FDatabaseEditor::UnregisterTabSpawners(const TSharedRef<FTabManager>& InTabManager)
@@ -360,6 +419,7 @@ namespace UE::PoseSearch
 		InTabManager->UnregisterTabSpawner(FDatabaseEditorTabs::PreviewSettingsID);
 		InTabManager->UnregisterTabSpawner(FDatabaseEditorTabs::AssetTreeViewID);
 		InTabManager->UnregisterTabSpawner(FDatabaseEditorTabs::SelectionDetailsID);
+		InTabManager->UnregisterTabSpawner(FDatabaseEditorTabs::StatisticsOverview);
 	}
 
 	FName FDatabaseEditor::GetToolkitFName() const
@@ -470,6 +530,21 @@ namespace UE::PoseSearch
 			];
 	}
 
+	TSharedRef<SDockTab> FDatabaseEditor::SpawnTab_StatisticsOverview(const FSpawnTabArgs& Args) const
+	{
+		check(Args.GetTabId() == FDatabaseEditorTabs::StatisticsOverview);
+		
+		return SNew(SDockTab)
+		.Label(LOCTEXT("StatisticsOverview_Title", "Statistics Overview"))
+		[
+			SNew(SScrollBox)
+			+SScrollBox::Slot()
+			[
+				StatisticsOverviewWidget.ToSharedRef()
+			]
+		];
+	}
+
 	void FDatabaseEditor::OnFinishedChangingSelectionProperties(const FPropertyChangedEvent& PropertyChangedEvent)
 	{
 		// @todo: BuildSearchIndex is already called, but it'd be better to perform it from here
@@ -484,8 +559,7 @@ namespace UE::PoseSearch
 		{
 			SequenceSelectionIndex = 0,
 			BlendSpaceSelectionIndex = 1,
-			GroupSelectionIndex = 2,
-			SelectionTypeCount = 3
+			SelectionTypeCount = 2
 		};
 
 		TArray<TArray<TWeakObjectPtr<UObject>>> SelectionReflections;
@@ -511,14 +585,6 @@ namespace UE::PoseSearch
 					NewSelectionReflection->SetSourceLink(SelectedItem, AssetTreeWidget);
 					SelectionReflections[BlendSpaceSelectionIndex].Add(NewSelectionReflection);
 				}
-				else
-				{
-					UPoseSearchDatabaseReflection* NewSelectionReflection = NewObject<UPoseSearchDatabaseReflection>();
-					NewSelectionReflection->AddToRoot();
-					NewSelectionReflection->Initialize(GetPoseSearchDatabase());
-					NewSelectionReflection->SetSourceLink(SelectedItem, AssetTreeWidget);
-					SelectionReflections[GroupSelectionIndex].Add(NewSelectionReflection);
-				}
 			}
 		}
 
diff --git a/Engine/Plugins/Experimental/Animation/PoseSearch/Source/Editor/Private/PoseSearchDatabaseEditorReflection.cpp b/Engine/Plugins/Experimental/Animation/PoseSearch/Source/Editor/Private/PoseSearchDatabaseEditorReflection.cpp
--- a/Engine/Plugins/Experimental/Animation/PoseSearch/Source/Editor/Private/PoseSearchDatabaseEditorReflection.cpp
+++ b/Engine/Plugins/Experimental/Animation/PoseSearch/Source/Editor/Private/PoseSearchDatabaseEditorReflection.cpp
@@ -48,60 +48,58 @@ void UPoseSearchDatabaseBlendSpaceReflection::PostEditChangeProperty(
 
 #endif // WITH_EDITOR
 
-FText FPoseSearchDatabaseMemoryStats::ToMemoryBudgetText(int32 Size)
+void UPoseSearchDatabaseStatistics::Initialize(const UPoseSearchDatabase* PoseSearchDatabase)
 {
-	if (Size < 1024)
-	{
-		return FText::Format(LOCTEXT("FPoseSearchDatabaseMemoryStats_Bytes", "{0} Bytes"), Size);
-	}
-
-	if (Size < 1024 * 1024)
-	{
-		return FText::Format(LOCTEXT("FPoseSearchDatabaseMemoryStats_KiloBytes", "{0} Kb"), Size / float(1024));
-	}
-
-	return FText::Format(LOCTEXT("FPoseSearchDatabaseMemoryStats_MegaBytes", "{0} Mb"), Size / float(1024 * 1024));
-}
-
-void FPoseSearchDatabaseMemoryStats::Initialize(const UPoseSearchDatabase* PoseSearchDatabase)
-{
-	const FPoseSearchIndex* SearchIndexSafe = PoseSearchDatabase->GetSearchIndexSafe();
-	check(SearchIndexSafe);
-
-	const int32 ValuesBytesSize = SearchIndexSafe->Values.GetAllocatedSize();
-	const int32 PCAValuesBytesSize = SearchIndexSafe->PCAValues.GetAllocatedSize();
-	const int32 KDTreeBytesSize = SearchIndexSafe->KDTree.GetAllocatedSize();
-	const int32 PoseMetadataBytesSize = SearchIndexSafe->PoseMetadata.GetAllocatedSize();
-	const int32 AssetsBytesSize = SearchIndexSafe->Assets.GetAllocatedSize();
-
-	const int32 OtherBytesSize = SearchIndexSafe->PCAProjectionMatrix.GetAllocatedSize() +
-							SearchIndexSafe->Mean.GetAllocatedSize() +
-							SearchIndexSafe->WeightsSqrt.GetAllocatedSize();
+	static FText TimeFormat = LOCTEXT("TimeFormat", "{0} {0}|plural(one=Second,other=Seconds)");
 	
-	const int32 EstimatedDatabaseBytesSize = ValuesBytesSize + PCAValuesBytesSize + KDTreeBytesSize + PoseMetadataBytesSize + AssetsBytesSize + OtherBytesSize;
-
-	ValuesSize = ToMemoryBudgetText(ValuesBytesSize);
-	PCAValuesSize = ToMemoryBudgetText(PCAValuesBytesSize);
-	KDTreeSize = ToMemoryBudgetText(KDTreeBytesSize);
-	PoseMetadataSize = ToMemoryBudgetText(PoseMetadataBytesSize);
-	AssetsSize = ToMemoryBudgetText(AssetsBytesSize);
-	EstimatedDatabaseSize = ToMemoryBudgetText(EstimatedDatabaseBytesSize);
-}
-
-void UPoseSearchDatabaseReflection::Initialize(const UPoseSearchDatabase* PoseSearchDatabase)
-{
 	if (PoseSearchDatabase)
 	{
-		const FPoseSearchIndex* SearchIndexSafe = PoseSearchDatabase->GetSearchIndexSafe();
-		if (SearchIndexSafe)
+		if (const FPoseSearchIndex* SearchIndexSafe = PoseSearchDatabase->GetSearchIndexSafe())
 		{
-			SearchIndex = *SearchIndexSafe;
-			MemoryStats.Initialize(PoseSearchDatabase);
-		}
-		else
-		{
-			SearchIndex = FPoseSearchIndex();
-			MemoryStats = FPoseSearchDatabaseMemoryStats();
+			// General information
+	
+			AnimationSequences = PoseSearchDatabase->Sequences.Num();
+			
+			TotalAnimationPosesInFrames = SearchIndexSafe->NumPoses;
+			TotalAnimationPosesInTime = FText::Format(TimeFormat, static_cast<double>(SearchIndexSafe->NumPoses) / PoseSearchDatabase->Schema->SampleRate);
+			
+			{
+				uint32 NumOfSearchablePoses = 0;
+				for (const FPoseSearchPoseMetadata & PoseMetadata : SearchIndexSafe->PoseMetadata)
+				{
+					NumOfSearchablePoses += PoseMetadata.Flags != EPoseSearchPoseFlags::BlockTransition;
+				}
+				
+				SearchableFrames = NumOfSearchablePoses;
+				SearchableTime = FText::Format(TimeFormat, static_cast<double>(NumOfSearchablePoses) / PoseSearchDatabase->Schema->SampleRate);
+			}
+			
+			// Velocity information
+	
+			// TODO: Set values once they can be queried from the PoseSearchIndex.
+			
+			// Principal Component Analysis
+			
+			ExplainedVariance = SearchIndexSafe->PCAExplainedVariance;
+			
+			// Memory information
+			
+			{
+				const uint32 ValuesBytesSize = SearchIndexSafe->Values.GetAllocatedSize();
+				const uint32 PCAValuesBytesSize = SearchIndexSafe->PCAValues.GetAllocatedSize();
+				const uint32 KDTreeBytesSize = SearchIndexSafe->KDTree.GetAllocatedSize();
+				const uint32 PoseMetadataBytesSize = SearchIndexSafe->PoseMetadata.GetAllocatedSize();
+				const uint32 AssetsBytesSize = SearchIndexSafe->Assets.GetAllocatedSize();
+				const uint32 OtherBytesSize = SearchIndexSafe->PCAProjectionMatrix.GetAllocatedSize() + SearchIndexSafe->Mean.GetAllocatedSize() + SearchIndexSafe->WeightsSqrt.GetAllocatedSize();
+				const uint32 EstimatedDatabaseBytesSize = ValuesBytesSize + PCAValuesBytesSize + KDTreeBytesSize + PoseMetadataBytesSize + AssetsBytesSize + OtherBytesSize;
+				
+				ValuesSize = FText::AsMemory(ValuesBytesSize);
+				PCAValuesSize = FText::AsMemory(PCAValuesBytesSize);
+				KDTreeSize = FText::AsMemory(KDTreeBytesSize);
+				PoseMetadataSize = FText::AsMemory(PoseMetadataBytesSize);
+				AssetsSize = FText::AsMemory(AssetsBytesSize);
+				EstimatedDatabaseSize = FText::AsMemory(EstimatedDatabaseBytesSize);
+			}
 		}
 	}
 }
diff --git a/Engine/Plugins/Experimental/Animation/PoseSearch/Source/Editor/Public/PoseSearchDatabaseEditor.h b/Engine/Plugins/Experimental/Animation/PoseSearch/Source/Editor/Public/PoseSearchDatabaseEditor.h
--- a/Engine/Plugins/Experimental/Animation/PoseSearch/Source/Editor/Public/PoseSearchDatabaseEditor.h
+++ b/Engine/Plugins/Experimental/Animation/PoseSearch/Source/Editor/Public/PoseSearchDatabaseEditor.h
@@ -59,7 +59,8 @@ namespace UE::PoseSearch
 		TSharedRef<SDockTab> SpawnTab_PreviewSettings(const FSpawnTabArgs& Args);
 		TSharedRef<SDockTab> SpawnTab_AssetTreeView(const FSpawnTabArgs& Args);
 		TSharedRef<SDockTab> SpawnTab_SelectionDetails(const FSpawnTabArgs& Args);
-
+		TSharedRef<SDockTab> SpawnTab_StatisticsOverview(const FSpawnTabArgs& Args) const;
+		
 		void BindCommands();
 		void ExtendToolbar();
 		void FillToolbar(FToolBarBuilder& ToolbarBuilder);
@@ -77,6 +78,8 @@ namespace UE::PoseSearch
 
 		TArray<TSharedPtr<IDetailsView>> SelectionWidgets;
 
+		TSharedPtr<IDetailsView> StatisticsOverviewWidget;
+		
 		TSharedPtr<FDatabasePreviewScene> PreviewScene;
 
 		TSharedPtr<FDatabaseViewModel> ViewModel;
diff --git a/Engine/Plugins/Experimental/Animation/PoseSearch/Source/Editor/Public/PoseSearchDatabaseEditorReflection.h b/Engine/Plugins/Experimental/Animation/PoseSearch/Source/Editor/Public/PoseSearchDatabaseEditorReflection.h
--- a/Engine/Plugins/Experimental/Animation/PoseSearch/Source/Editor/Public/PoseSearchDatabaseEditorReflection.h
+++ b/Engine/Plugins/Experimental/Animation/PoseSearch/Source/Editor/Public/PoseSearchDatabaseEditorReflection.h
@@ -87,16 +87,68 @@ struct FPoseSearchDatabaseMemoryStats
 };
 
 UCLASS()
-class UPoseSearchDatabaseReflection : public UPoseSearchDatabaseReflectionBase
+class UPoseSearchDatabaseStatistics : public UObject
 {
 	GENERATED_BODY()
 
 public:
-	void Initialize(const UPoseSearchDatabase* PoseSearchDatabase);
+	
+	// General information
 
-	UPROPERTY(VisibleAnywhere, Category = "Selected Database Search Index")
-	FPoseSearchDatabaseMemoryStats MemoryStats;
+	UPROPERTY(VisibleAnywhere, Category = "General Information")
+	uint32 AnimationSequences;
 
-	UPROPERTY(VisibleAnywhere, Category = "Selected Database Search Index")
-	FPoseSearchIndex SearchIndex;
-};
+	UPROPERTY(VisibleAnywhere, Category = "General Information")
+	uint32 TotalAnimationPosesInFrames;
+
+	UPROPERTY(VisibleAnywhere, Category = "General Information")
+	FText TotalAnimationPosesInTime;
+
+	UPROPERTY(VisibleAnywhere, Category = "General Information")
+	uint32 SearchableFrames;
+
+	UPROPERTY(VisibleAnywhere, Category = "General Information")
+	FText SearchableTime;
+
+	// Velocity Information
+
+	UPROPERTY(VisibleAnywhere, Category = "Velocity Information")
+	double AverageVelocity;
+
+	UPROPERTY(VisibleAnywhere, Category = "Velocity Information")
+	double MaxVelocity;
+
+	UPROPERTY(VisibleAnywhere, Category = "Velocity Information")
+	double AverageAcceleration;
+
+	UPROPERTY(VisibleAnywhere, Category = "Velocity Information")
+	double MaxAcceleration;
+
+	// Principal Component Analysis (PCA) Information
+
+	UPROPERTY(VisibleAnywhere, Category = "Principal Component Analysis (PCA) Information")
+	float ExplainedVariance;
+
+	// Memory information
+	
+	UPROPERTY(VisibleAnywhere, Category = "Memory Information")
+	FText EstimatedDatabaseSize;
+
+	UPROPERTY(VisibleAnywhere, Category = "Memory Information")
+	FText ValuesSize;
+
+	UPROPERTY(VisibleAnywhere, Category = "Memory Information")
+	FText PCAValuesSize;
+
+	UPROPERTY(VisibleAnywhere, Category = "Memory Information")
+	FText KDTreeSize;
+
+	UPROPERTY(VisibleAnywhere, Category = "Memory Information")
+	FText PoseMetadataSize;
+	
+	UPROPERTY(VisibleAnywhere, Category = "Memory Information")
+	FText AssetsSize;
+
+	/** Initialize statistics given a database */
+	void Initialize(const UPoseSearchDatabase* PoseSearchDatabase);
+};
\ No newline at end of file
