motion matching
- UPoseSearchDatabase concurrent indexing fix
- added FAsyncPoseSearchDatabasesManagement to manage async db indexing
- fixed non deterministic FDerivedDataKeyBuilder behavior

diff --git a/Engine/Plugins/Experimental/Animation/PoseSearch/Source/Editor/Private/PoseSearchCustomization.cpp b/Engine/Plugins/Experimental/Animation/PoseSearch/Source/Editor/Private/PoseSearchCustomization.cpp
--- a/Engine/Plugins/Experimental/Animation/PoseSearch/Source/Editor/Private/PoseSearchCustomization.cpp
+++ b/Engine/Plugins/Experimental/Animation/PoseSearch/Source/Editor/Private/PoseSearchCustomization.cpp
@@ -94,8 +94,6 @@ void FPoseSearchDatabaseDetails::CustomizeDetails(IDetailLayoutBuilder& DetailBu
 	TArray<TSharedPtr<IPropertyHandle>> HiddenHandles;
 	HiddenHandles.Add(DetailBuilder.GetProperty(GET_MEMBER_NAME_CHECKED(UPoseSearchDatabase, Sequences)));
 	HiddenHandles.Add(DetailBuilder.GetProperty(GET_MEMBER_NAME_CHECKED(UPoseSearchDatabase, BlendSpaces)));
-	HiddenHandles.Add(DetailBuilder.GetProperty(GET_MEMBER_NAME_CHECKED(UPoseSearchDatabase, SimpleSequences)));
-	HiddenHandles.Add(DetailBuilder.GetProperty(GET_MEMBER_NAME_CHECKED(UPoseSearchDatabase, SimpleBlendSpaces)));
 
 	for (TSharedPtr<IPropertyHandle> PropertyHandle : HiddenHandles)
 	{
diff --git a/Engine/Plugins/Experimental/Animation/PoseSearch/Source/Editor/Private/PoseSearchDatabaseEditor.cpp b/Engine/Plugins/Experimental/Animation/PoseSearch/Source/Editor/Private/PoseSearchDatabaseEditor.cpp
--- a/Engine/Plugins/Experimental/Animation/PoseSearch/Source/Editor/Private/PoseSearchDatabaseEditor.cpp
+++ b/Engine/Plugins/Experimental/Animation/PoseSearch/Source/Editor/Private/PoseSearchDatabaseEditor.cpp
@@ -45,20 +45,6 @@ namespace UE::PoseSearch
 	const FName FDatabaseEditorTabs::SelectionDetailsID(TEXT("PoseSearchDatabaseEditorSelectionDetailsID"));
 	const FName FDatabaseEditorTabs::StatisticsOverview(TEXT("PoseSearchDatabaseEditorStatisticsOverviewID"));
 	
-	FDatabaseEditor::FDatabaseEditor()
-	{
-	}
-
-	FDatabaseEditor::~FDatabaseEditor()
-	{
-		UPoseSearchDatabase* DatabaseAsset = ViewModel->GetPoseSearchDatabase();
-		if (IsValid(DatabaseAsset))
-		{
-			DatabaseAsset->UnregisterOnAssetChange(AssetTreeWidget.Get());
-			DatabaseAsset->UnregisterOnGroupChange(AssetTreeWidget.Get());
-		}
-	}
-
 	const UPoseSearchDatabase* FDatabaseEditor::GetPoseSearchDatabase() const
 	{
 		return ViewModel.IsValid() ? ViewModel->GetPoseSearchDatabase() : nullptr;
@@ -168,18 +154,6 @@ namespace UE::PoseSearch
 				this,
 				&FDatabaseEditor::OnAssetTreeSelectionChanged));
 		
-		if (IsValid(DatabaseAsset))
-		{
-			DatabaseAsset->RegisterOnAssetChange(
-				UPoseSearchDatabase::FOnDerivedDataRebuild::CreateSP(
-					AssetTreeWidget.Get(),
-					&SDatabaseAssetTree::RefreshTreeView, false, false));
-			DatabaseAsset->RegisterOnGroupChange(
-				UPoseSearchDatabase::FOnDerivedDataRebuild::CreateSP(
-					AssetTreeWidget.Get(),
-					&SDatabaseAssetTree::RefreshTreeView, false, false));
-		}
-		
 		// Create Asset Details widget
 		FPropertyEditorModule& PropertyModule = FModuleManager::LoadModuleChecked<FPropertyEditorModule>("PropertyEditor");
 
@@ -243,8 +217,6 @@ namespace UE::PoseSearch
 				UpdateStatisticsInformationLambda();
 
 				// Ensure any database changes are reflected
-				DatabaseAsset->RegisterOnAssetChange(UPoseSearchDatabase::FOnAssetChange::CreateLambda(UpdateStatisticsInformationLambda));
-				DatabaseAsset->RegisterOnGroupChange(UPoseSearchDatabase::FOnGroupChange::CreateLambda(UpdateStatisticsInformationLambda));
 				DatabaseAsset->RegisterOnDerivedDataRebuild(UPoseSearchDatabase::FOnDerivedDataRebuild::CreateLambda(UpdateStatisticsInformationLambda));
 				// TODO: Register to delegate that gets triggered after the DerivedData has been rebuilt.
 			}
@@ -547,8 +519,6 @@ namespace UE::PoseSearch
 
 	void FDatabaseEditor::OnFinishedChangingSelectionProperties(const FPropertyChangedEvent& PropertyChangedEvent)
 	{
-		// @todo: BuildSearchIndex is already called, but it'd be better to perform it from here
-		//ViewModel->BuildSearchIndex();
 	}
 
 	void FDatabaseEditor::OnAssetTreeSelectionChanged(
diff --git a/Engine/Plugins/Experimental/Animation/PoseSearch/Source/Editor/Private/PoseSearchDatabaseViewModel.cpp b/Engine/Plugins/Experimental/Animation/PoseSearch/Source/Editor/Private/PoseSearchDatabaseViewModel.cpp
--- a/Engine/Plugins/Experimental/Animation/PoseSearch/Source/Editor/Private/PoseSearchDatabaseViewModel.cpp
+++ b/Engine/Plugins/Experimental/Animation/PoseSearch/Source/Editor/Private/PoseSearchDatabaseViewModel.cpp
@@ -4,7 +4,7 @@
 #include "PoseSearchDatabasePreviewScene.h"
 #include "SPoseSearchDatabaseAssetList.h"
 #include "PoseSearchEditor.h"
-#include "PoseSearch/PoseSearch.h"
+#include "PoseSearch/PoseSearchDerivedData.h"
 #include "Modules/ModuleManager.h"
 #include "AnimPreviewInstance.h"
 #include "Animation/DebugSkelMeshComponent.h"
@@ -174,8 +174,7 @@ namespace UE::PoseSearch
 
 	void FDatabaseViewModel::BuildSearchIndex()
 	{
-		PoseSearchDatabase->NotifyAssetChange();
-		PoseSearchDatabase->BeginCacheDerivedData();
+		UE::PoseSearch::FAsyncPoseSearchDatabasesManagement::Get().RequestAsyncBuildIndex(*PoseSearchDatabase, false, true);
 	}
 
 	void FDatabaseViewModel::PreviewBackwardEnd()
@@ -362,6 +361,7 @@ namespace UE::PoseSearch
 	{
 		for (auto PreviewActor : PreviewActors)
 		{
+			// @todo: PreviewActor.Actor is a TWeakObjectPtr so it can be null.
 			PreviewActor.Actor->Destroy();
 		}
 		PreviewActors.Reset();
diff --git a/Engine/Plugins/Experimental/Animation/PoseSearch/Source/Editor/Private/PoseSearchDebugger.cpp b/Engine/Plugins/Experimental/Animation/PoseSearch/Source/Editor/Private/PoseSearchDebugger.cpp
--- a/Engine/Plugins/Experimental/Animation/PoseSearch/Source/Editor/Private/PoseSearchDebugger.cpp
+++ b/Engine/Plugins/Experimental/Animation/PoseSearch/Source/Editor/Private/PoseSearchDebugger.cpp
@@ -899,8 +899,7 @@ void SDebuggerDatabaseView::Update(const FTraceMotionMatchingStateMessage& State
 		const UPoseSearchDatabase* Database = FTraceMotionMatchingState::GetObjectFromId<UPoseSearchDatabase>(DbEntry.DatabaseId);
 		if (Database)
 		{
-			const FPoseSearchIndex* SearchIndex = Database->GetSearchIndex();
-			if (SearchIndex && SearchIndex->IsValid())
+			if (const FPoseSearchIndex* SearchIndex = Database->GetSearchIndexSafe(false))
 			{
 				for (const FTraceMotionMatchingStatePoseEntry& PoseEntry : DbEntry.PoseEntries)
 				{
@@ -908,7 +907,7 @@ void SDebuggerDatabaseView::Update(const FTraceMotionMatchingStateMessage& State
 					{
 						TSharedRef<FDebuggerDatabaseRowData>& Row = UnfilteredDatabaseRows.Add_GetRef(MakeShared<FDebuggerDatabaseRowData>());
 
-						const float Time = SearchIndex->GetAssetTime(PoseEntry.DbPoseIdx);
+						const float Time = SearchIndex->GetAssetTime(PoseEntry.DbPoseIdx, Database->Schema->GetSamplingInterval());
 
 						Row->PoseIdx = PoseEntry.DbPoseIdx;
 						Row->SourceDatabase = Database;
@@ -964,8 +963,7 @@ void SDebuggerDatabaseView::Update(const FTraceMotionMatchingStateMessage& State
 			const UPoseSearchDatabase* Database = FTraceMotionMatchingState::GetObjectFromId<UPoseSearchDatabase>(DbEntry.DatabaseId);
 			if (Database)
 			{
-				const FPoseSearchIndex* SearchIndex = Database->GetSearchIndex();
-				if (SearchIndex && SearchIndex->IsValid())
+				if (Database->GetSearchIndexSafe(false))
 				{
 					Database->Schema->ComputeCostBreakdowns(CostBreakDownData);
 				}
@@ -1029,10 +1027,8 @@ void SDebuggerDatabaseView::Update(const FTraceMotionMatchingStateMessage& State
 		OldLabels = CostBreakDownData.GetLabels();
 
 		// recreating and binding the columns
-
 		Columns.Reset();
 
-		// @TODO: Support runtime reordering of these indices
 		// Construct all column types
 		int32 ColumnIdx = 0;
 		AddColumn(MakeShared<FDatabaseName>(ColumnIdx++));
@@ -1600,8 +1596,7 @@ void SDebuggerDetailsView::UpdateReflection(const FTraceMotionMatchingStateMessa
 		const UPoseSearchDatabase* Database = FTraceMotionMatchingState::GetObjectFromId<UPoseSearchDatabase>(DbEntry.DatabaseId);
 		if (Database && Database == CurrentDatabase)
 		{
-			const FPoseSearchIndex* SearchIndex = Database->GetSearchIndex();
-			if (SearchIndex && SearchIndex->IsValid())
+			if (Database->GetSearchIndexSafe(false))
 			{
 				Reflection->QueryPoseVector = DbEntry.QueryVector;
 				break;
@@ -1875,8 +1870,7 @@ void SDebuggerView::DrawFeatures(
 				const UPoseSearchDatabase* Database = FTraceMotionMatchingState::GetObjectFromId<UPoseSearchDatabase>(DbEntry.DatabaseId);
 				if (Database && Database == CurrentDatabase)
 				{
-					const FPoseSearchIndex* SearchIndex = Database->GetSearchIndex();
-					if (SearchIndex && SearchIndex->IsValid() && DbEntry.QueryVector.Num() == Database->Schema->SchemaCardinality)
+					if (Database->GetSearchIndexSafe(false) && DbEntry.QueryVector.Num() == Database->Schema->SchemaCardinality)
 					{
 						// Set shared state
 						FDebugDrawParams DrawParams;
diff --git a/Engine/Plugins/Experimental/Animation/PoseSearch/Source/Editor/Public/PoseSearchDatabaseEditor.h b/Engine/Plugins/Experimental/Animation/PoseSearch/Source/Editor/Public/PoseSearchDatabaseEditor.h
--- a/Engine/Plugins/Experimental/Animation/PoseSearch/Source/Editor/Public/PoseSearchDatabaseEditor.h
+++ b/Engine/Plugins/Experimental/Animation/PoseSearch/Source/Editor/Public/PoseSearchDatabaseEditor.h
@@ -21,9 +21,6 @@ namespace UE::PoseSearch
 	{
 	public:
 
-		FDatabaseEditor();
-		virtual ~FDatabaseEditor();
-
 		void InitAssetEditor(
 			const EToolkitMode::Type Mode,
 			const TSharedPtr<IToolkitHost>& InitToolkitHost,
diff --git a/Engine/Plugins/Experimental/Animation/PoseSearch/Source/Runtime/PoseSearch.Build.cs b/Engine/Plugins/Experimental/Animation/PoseSearch/Source/Runtime/PoseSearch.Build.cs
--- a/Engine/Plugins/Experimental/Animation/PoseSearch/Source/Runtime/PoseSearch.Build.cs
+++ b/Engine/Plugins/Experimental/Animation/PoseSearch/Source/Runtime/PoseSearch.Build.cs
@@ -40,7 +40,8 @@ public class PoseSearch : ModuleRules
 		{
 			PrivateDependencyModuleNames.AddRange(
 				new string[] {
-					"DerivedDataCache"
+					"DerivedDataCache",
+					"UnrealEd"
 				}
 			);
 		}
diff --git a/Engine/Plugins/Experimental/Animation/PoseSearch/Source/Runtime/Private/AnimNode_MotionMatching.cpp b/Engine/Plugins/Experimental/Animation/PoseSearch/Source/Runtime/Private/AnimNode_MotionMatching.cpp
--- a/Engine/Plugins/Experimental/Animation/PoseSearch/Source/Runtime/Private/AnimNode_MotionMatching.cpp
+++ b/Engine/Plugins/Experimental/Animation/PoseSearch/Source/Runtime/Private/AnimNode_MotionMatching.cpp
@@ -73,6 +73,14 @@ void FAnimNode_MotionMatching::UpdateAssetPlayer(const FAnimationUpdateContext&
 	}
 	else
 	{
+#if WITH_EDITOR
+		// in case we're still indexing MotionMatchingState.CurrentSearchResult.Database we Reset the MotionMatchingState
+		if (MotionMatchingState.CurrentSearchResult.IsValid() && !MotionMatchingState.CurrentSearchResult.Database->GetSearchIndexSafe(true))
+		{
+			MotionMatchingState.Reset();
+		}
+#endif // WITH_EDITOR
+
 		// We adjust the motion matching state asset time to the current player node's asset time. This is done 
 		// because the player node may have ticked more or less time than we expected due to variable dt or the 
 		// dynamic playback rate adjustment and as such the motion matching state does not update by itself
diff --git a/Engine/Plugins/Experimental/Animation/PoseSearch/Source/Runtime/Private/KDTree.cpp b/Engine/Plugins/Experimental/Animation/PoseSearch/Source/Runtime/Private/KDTree.cpp
--- a/Engine/Plugins/Experimental/Animation/PoseSearch/Source/Runtime/Private/KDTree.cpp
+++ b/Engine/Plugins/Experimental/Animation/PoseSearch/Source/Runtime/Private/KDTree.cpp
@@ -16,7 +16,7 @@ namespace UE::PoseSearch
 {
 
 #if UE_POSE_SEARCH_USE_NANOFLANN
-using FKDTreeImplementationBase = nanoflann::KDTreeSingleIndexAdaptor<nanoflann::L2_Simple_Adaptor<float, FKDTree::DataSource>, FKDTree::DataSource>;
+using FKDTreeImplementationBase = nanoflann::KDTreeSingleIndexAdaptor<nanoflann::L2_Simple_Adaptor<float, FKDTree::FDataSource>, FKDTree::FDataSource>;
 struct FKDTreeImplementation : FKDTreeImplementationBase
 {
 	using FKDTreeImplementationBase::FKDTreeImplementationBase;
@@ -24,20 +24,20 @@ struct FKDTreeImplementation : FKDTreeImplementationBase
 #endif
 
 FKDTree::FKDTree(int32 Count, int32 Dim, const float* Data, int32 MaxLeafSize)
-: DataSrc(Count, Dim, Data)
+: DataSource(Count, Dim, Data)
 , Impl(nullptr)
 {
 #if UE_POSE_SEARCH_USE_NANOFLANN
-	Impl = new FKDTreeImplementation(Dim, DataSrc, nanoflann::KDTreeSingleIndexAdaptorParams(MaxLeafSize));
+	Impl = new FKDTreeImplementation(Dim, DataSource, nanoflann::KDTreeSingleIndexAdaptorParams(MaxLeafSize));
 #endif
 }
 
 FKDTree::FKDTree()
-: DataSrc(0, 0, nullptr)
+: DataSource(0, 0, nullptr)
 , Impl(nullptr)
 {
 #if UE_POSE_SEARCH_USE_NANOFLANN
-	Impl = new FKDTreeImplementation(0, DataSrc, nanoflann::KDTreeSingleIndexAdaptorParams(0));
+	Impl = new FKDTreeImplementation(0, DataSource, nanoflann::KDTreeSingleIndexAdaptorParams(0));
 #endif
 }
 
@@ -45,23 +45,93 @@ FKDTree::~FKDTree()
 {
 #if UE_POSE_SEARCH_USE_NANOFLANN
 	delete Impl;
+	Impl = nullptr;
 #endif
+	DataSource = FDataSource();
 }
 
-FKDTree::FKDTree(const FKDTree& r)
-: DataSrc(r.DataSrc.PointCount, r.DataSrc.PointDim, r.DataSrc.Data)
-, Impl(nullptr)
+#if UE_POSE_SEARCH_USE_NANOFLANN
+void CopySubTree(FKDTree& KDTree, FKDTreeImplementation::NodePtr& ThisNode, const FKDTreeImplementation::NodePtr& OtherNode)
+{
+	check(KDTree.Impl);
+
+	ThisNode = KDTree.Impl->pool.template allocate<FKDTreeImplementation::Node>();
+
+	ThisNode->node_type = OtherNode->node_type;
+
+	if (OtherNode->child1 != nullptr)
+	{
+		CopySubTree(KDTree, ThisNode->child1, OtherNode->child1);
+	}
+	else
+	{
+		ThisNode->child1 = nullptr;
+	}
+
+	if (OtherNode->child2 != nullptr)
+	{
+		CopySubTree(KDTree, ThisNode->child2, OtherNode->child2);
+	}
+	else
+	{
+		ThisNode->child2 = nullptr;
+	}
+}
+
+#endif
+
+FKDTree::FKDTree(const FKDTree& Other)
 {
 #if UE_POSE_SEARCH_USE_NANOFLANN
-	check(r.Impl);
-	Impl = new FKDTreeImplementation(r.DataSrc.PointDim, DataSrc, nanoflann::KDTreeSingleIndexAdaptorParams(r.Impl->m_leaf_max_size));
+	if (this != &Other)
+	{
+		check(Other.Impl);
+		Impl = new FKDTreeImplementation(0, DataSource, nanoflann::KDTreeSingleIndexAdaptorParams(0));
+
+		DataSource = Other.DataSource;
+
+		check(Other.Impl->m_size < UINT_MAX);
+		Impl->m_size = Other.Impl->m_size;
+
+		if (Impl->m_size > 0)
+		{
+			Impl->dim = Other.Impl->dim;
+
+			const uint32 root_bbox_size = Other.Impl->root_bbox.size();
+			check(root_bbox_size < UINT_MAX);
+			Impl->root_bbox.resize(root_bbox_size);
+
+			for (uint32 i = 0; i < root_bbox_size; ++i)
+			{
+				Impl->root_bbox[i] = Other.Impl->root_bbox[i];
+			}
+
+			check(Other.Impl->m_leaf_max_size < UINT_MAX);
+			const uint32 KDTreeLeafMaxSize = Other.Impl->m_leaf_max_size;
+			Impl->m_leaf_max_size = KDTreeLeafMaxSize;
+
+			const uint32 VAccSize = Other.Impl->vAcc.size();
+			check(VAccSize < UINT_MAX);
+			Impl->vAcc.resize(VAccSize);
+			
+			for (uint32 i = 0; i < VAccSize; ++i)
+			{
+				Impl->vAcc[i] = Other.Impl->vAcc[i];
+			}
+			
+			CopySubTree(*this, Impl->root_node, Other.Impl->root_node);
+		}
+	}
 #endif
 }
 
-FKDTree& FKDTree::operator=(const FKDTree& r)
+FKDTree& FKDTree::operator=(const FKDTree& Other)
 {
-	this->~FKDTree();
-	new(this)FKDTree(r);
+	if (this != &Other)
+	{
+		this->~FKDTree();
+		new(this)FKDTree(Other);
+	}
 	return *this;
 }
 
@@ -161,9 +231,9 @@ FArchive& Serialize(FArchive& Ar, FKDTree& KDTree, const float* KDTreeData)
 
 		if (Ar.IsLoading())
 		{
-			KDTree.DataSrc.Data = KDTreeData;
-			KDTree.DataSrc.PointDim = KDTree.Impl->dim;
-			KDTree.DataSrc.PointCount = KDTree.Impl->m_size;
+			KDTree.DataSource.Data = KDTreeData;
+			KDTree.DataSource.PointDim = KDTree.Impl->dim;
+			KDTree.DataSource.PointCount = KDTree.Impl->m_size;
 
 			KDTree.Impl->root_bbox.resize(root_bbox_size);
 		}
diff --git a/Engine/Plugins/Experimental/Animation/PoseSearch/Source/Runtime/Private/PoseSearch.cpp b/Engine/Plugins/Experimental/Animation/PoseSearch/Source/Runtime/Private/PoseSearch.cpp
--- a/Engine/Plugins/Experimental/Animation/PoseSearch/Source/Runtime/Private/PoseSearch.cpp
+++ b/Engine/Plugins/Experimental/Animation/PoseSearch/Source/Runtime/Private/PoseSearch.cpp
@@ -3,8 +3,8 @@
 #include "PoseSearch/PoseSearch.h"
 #include "AssetRegistry/AssetData.h"
 #include "PoseSearch/PoseSearchAnimNotifies.h"
-
-#include "PoseSearchDerivedData.h"
+#include "PoseSearch/PoseSearchDerivedDataKey.h"
+#include "PoseSearch/PoseSearchDerivedData.h"
 #include "PoseSearchEigenHelper.h"
 
 #include "Algo/BinarySearch.h"
@@ -33,12 +33,9 @@
 #include "Animation/AnimRootMotionProvider.h"
 #include "Animation/BlendSpace1D.h"
 
-
 #if WITH_EDITOR
-#include "AssetRegistry/AssetRegistryModule.h"
-#include "AssetRegistry/ARFilter.h"
-#include "PoseSearch/PoseSearchDerivedDataKey.h"
-#endif // WITH_EDITOR
+#include "DerivedDataRequestOwner.h"
+#endif
 
 IMPLEMENT_ANIMGRAPH_MESSAGE(UE::PoseSearch::IPoseHistoryProvider);
 
@@ -205,7 +202,7 @@ static void PopulateNonSelectableIdx(FNonSelectableIdx& NonSelectableIdx, FSearc
 				NonSelectableIdx.Add(PoseIdx);
 
 #if UE_POSE_SEARCH_TRACE_ENABLED
-				const FPoseSearchCost PoseCost = SearchIndex->ComparePoses(PoseIdx, SearchContext.QueryMirrorRequest, EPoseComparisonFlags::None, QueryValues);
+				const FPoseSearchCost PoseCost = SearchIndex->ComparePoses(PoseIdx, SearchContext.QueryMirrorRequest, EPoseComparisonFlags::None, Database->Schema->MirrorMismatchCostBias, QueryValues);
 				SearchContext.BestCandidates.Add(PoseCost, PoseIdx, Database, EPoseCandidateFlags::DiscardedBy_PoseJumpThresholdTime);
 #endif
 			}
@@ -238,7 +235,7 @@ static void PopulateNonSelectableIdx(FNonSelectableIdx& NonSelectableIdx, FSearc
 				NonSelectableIdx.Add(PoseIdx);
 
 #if UE_POSE_SEARCH_TRACE_ENABLED
-				const FPoseSearchCost PoseCost = SearchIndex->ComparePoses(PoseIdx, SearchContext.QueryMirrorRequest, EPoseComparisonFlags::None, QueryValues);
+				const FPoseSearchCost PoseCost = SearchIndex->ComparePoses(PoseIdx, SearchContext.QueryMirrorRequest, EPoseComparisonFlags::None, Database->Schema->MirrorMismatchCostBias, QueryValues);
 				SearchContext.BestCandidates.Add(PoseCost, PoseIdx, Database, EPoseCandidateFlags::DiscardedBy_PoseJumpThresholdTime);
 #endif
 			}
@@ -265,7 +262,7 @@ static void PopulateNonSelectableIdx(FNonSelectableIdx& NonSelectableIdx, FSearc
 				// if we're editing the database and removing assets it's possible that the PoseIndicesHistory contains invalid pose indexes
 				if (HistoricalPoseIndex.PoseIndex < SearchIndex->NumPoses)
 				{
-					const FPoseSearchCost PoseCost = SearchIndex->ComparePoses(HistoricalPoseIndex.PoseIndex, SearchContext.QueryMirrorRequest, EPoseComparisonFlags::None, QueryValues);
+					const FPoseSearchCost PoseCost = SearchIndex->ComparePoses(HistoricalPoseIndex.PoseIndex, SearchContext.QueryMirrorRequest, EPoseComparisonFlags::None, Database->Schema->MirrorMismatchCostBias, QueryValues);
 					SearchContext.BestCandidates.Add(PoseCost, HistoricalPoseIndex.PoseIndex, Database, EPoseCandidateFlags::DiscardedBy_PoseReselectHistory);
 				}
 #endif
@@ -319,12 +316,12 @@ struct FPoseFilters
 				}
 				else if (PoseFilter == &BlockTransitionPoseFilter)
 				{
-					const FPoseSearchCost PoseCost = SearchIndex->ComparePoses(PoseIdx, SearchContext.QueryMirrorRequest, EPoseComparisonFlags::None, QueryValues);
+					const FPoseSearchCost PoseCost = SearchIndex->ComparePoses(PoseIdx, SearchContext.QueryMirrorRequest, EPoseComparisonFlags::None, Database->Schema->MirrorMismatchCostBias, QueryValues);
 					SearchContext.BestCandidates.Add(PoseCost, PoseIdx, Database, EPoseCandidateFlags::DiscardedBy_BlockTransition);
 				}
 				else
 				{
-					const FPoseSearchCost PoseCost = SearchIndex->ComparePoses(PoseIdx, SearchContext.QueryMirrorRequest, EPoseComparisonFlags::None, QueryValues);
+					const FPoseSearchCost PoseCost = SearchIndex->ComparePoses(PoseIdx, SearchContext.QueryMirrorRequest, EPoseComparisonFlags::None, Database->Schema->MirrorMismatchCostBias, QueryValues);
 					SearchContext.BestCandidates.Add(PoseCost, PoseIdx, Database, EPoseCandidateFlags::DiscardedBy_PoseFilter);
 				}
 #endif
@@ -369,6 +366,51 @@ private:
 	TArray<const IPoseFilter*, TInlineAllocator<64>> AllPoseFilters;
 };
 
+static void FindValidSequenceIntervals(const FPoseSearchDatabaseSequence& DbSequence, const FPoseSearchExcludeFromDatabaseParameters& ExcludeFromDatabaseParameters, TArray<FFloatRange>& ValidRanges)
+{
+	const UAnimSequence* Sequence = DbSequence.Sequence;
+	check(DbSequence.Sequence);
+
+	const float SequenceLength = DbSequence.Sequence->GetPlayLength();
+
+	const FFloatInterval EffectiveSamplingInterval = UE::PoseSearch::GetEffectiveSamplingRange(DbSequence.Sequence, DbSequence.SamplingRange);
+	FFloatRange EffectiveSamplingRange = FFloatRange::Inclusive(EffectiveSamplingInterval.Min, EffectiveSamplingInterval.Max);
+	if (!DbSequence.IsLooping())
+	{
+		const FFloatRange ExcludeFromDatabaseRange(ExcludeFromDatabaseParameters.SequenceStartInterval, SequenceLength - ExcludeFromDatabaseParameters.SequenceEndInterval);
+		EffectiveSamplingRange = FFloatRange::Intersection(EffectiveSamplingRange, ExcludeFromDatabaseRange);
+	}
+
+	// start from a single interval defined by the database sequence sampling range
+	ValidRanges.Empty();
+	ValidRanges.Add(EffectiveSamplingRange);
+
+	FAnimNotifyContext NotifyContext;
+	Sequence->GetAnimNotifies(0.0f, SequenceLength, NotifyContext);
+
+	for (const FAnimNotifyEventReference& EventReference : NotifyContext.ActiveNotifies)
+	{
+		if (const FAnimNotifyEvent* NotifyEvent = EventReference.GetNotify())
+		{
+			if (const UAnimNotifyState_PoseSearchExcludeFromDatabase* ExclusionNotifyState = Cast<const UAnimNotifyState_PoseSearchExcludeFromDatabase>(NotifyEvent->NotifyStateClass))
+			{
+				FFloatRange ExclusionRange = FFloatRange::Inclusive(NotifyEvent->GetTriggerTime(), NotifyEvent->GetEndTriggerTime());
+
+				// Split every valid range based on the exclusion range just found. Because this might increase the 
+				// number of ranges in ValidRanges, the algorithm iterates from end to start.
+				for (int RangeIdx = ValidRanges.Num() - 1; RangeIdx >= 0; --RangeIdx)
+				{
+					FFloatRange EvaluatedRange = ValidRanges[RangeIdx];
+					ValidRanges.RemoveAt(RangeIdx);
+
+					TArray<FFloatRange> Diff = FFloatRange::Difference(EvaluatedRange, ExclusionRange);
+					ValidRanges.Append(Diff);
+				}
+			}
+		}
+	}
+}
+
 } // namespace UE::PoseSearch
 
 
@@ -539,8 +581,39 @@ bool UPoseSearchSchema::BuildQuery(UE::PoseSearch::FSearchContext& SearchContext
 }
 
 //////////////////////////////////////////////////////////////////////////
-// FPoseSearchIndexAsset
-// 
+// FPoseSearchIndex
+
+FPoseSearchIndex::FPoseSearchIndex(const FPoseSearchIndex& Other)
+	: NumPoses(Other.NumPoses)
+	, Values(Other.Values)
+	, PCAValues(Other.PCAValues)
+#if WITH_EDITORONLY_DATA
+	, PCAExplainedVariance(Other.PCAExplainedVariance)
+	, Deviation(Other.Deviation)
+#endif // WITH_EDITORONLY_DATA
+	, KDTree(Other.KDTree)
+	, PCAProjectionMatrix(Other.PCAProjectionMatrix)
+	, Mean(Other.Mean)
+	, WeightsSqrt(Other.WeightsSqrt)
+	, PoseMetadata(Other.PoseMetadata)
+	, OverallFlags(Other.OverallFlags)
+	, Assets(Other.Assets)
+	, MinCostAddend(Other.MinCostAddend)
+{
+	check(!PCAValues.IsEmpty() || KDTree.DataSource.PointCount == 0);
+	KDTree.DataSource.Data = PCAValues.IsEmpty() ? nullptr : PCAValues.GetData();
+}
+
+FPoseSearchIndex& FPoseSearchIndex::operator=(const FPoseSearchIndex& Other)
+{
+	if (this != &Other)
+	{
+		this->~FPoseSearchIndex();
+		new(this)FPoseSearchIndex(Other);
+	}
+	return *this;
+}
+
 const FPoseSearchIndexAsset& FPoseSearchIndex::GetAssetForPose(int32 PoseIdx) const
 {
 	const int32 AssetIndex = PoseMetadata[PoseIdx].AssetIndex;
@@ -560,7 +633,7 @@ const FPoseSearchIndexAsset* FPoseSearchIndex::GetAssetForPoseSafe(int32 PoseIdx
 	return nullptr;
 }
 
-float FPoseSearchIndex::GetAssetTime(int32 PoseIdx) const
+float FPoseSearchIndex::GetAssetTime(int32 PoseIdx, float SamplingInterval) const
 {
 	const FPoseSearchIndexAsset& Asset = GetAssetForPose(PoseIdx);
 
@@ -568,7 +641,7 @@ float FPoseSearchIndex::GetAssetTime(int32 PoseIdx) const
 	{
 		const FFloatInterval SamplingRange = Asset.SamplingInterval;
 
-		float AssetTime = FMath::Min(SamplingRange.Min + Schema->GetSamplingInterval() * (PoseIdx - Asset.FirstPoseIdx), SamplingRange.Max);
+		float AssetTime = FMath::Min(SamplingRange.Min + SamplingInterval * (PoseIdx - Asset.FirstPoseIdx), SamplingRange.Max);
 		return AssetTime;
 	}
 	
@@ -578,7 +651,7 @@ float FPoseSearchIndex::GetAssetTime(int32 PoseIdx) const
 
 		// For BlendSpaces the AssetTime is in the range [0, 1] while the Sampling Range
 		// is in real time (seconds)
-		float AssetTime = FMath::Min(SamplingRange.Min + Schema->GetSamplingInterval() * (PoseIdx - Asset.FirstPoseIdx), SamplingRange.Max) / (Asset.NumPoses * Schema->GetSamplingInterval());
+		float AssetTime = FMath::Min(SamplingRange.Min + SamplingInterval * (PoseIdx - Asset.FirstPoseIdx), SamplingRange.Max) / (Asset.NumPoses * SamplingInterval);
 		return AssetTime;
 	}
 	
@@ -590,14 +663,6 @@ float FPoseSearchIndex::GetAssetTime(int32 PoseIdx) const
 //////////////////////////////////////////////////////////////////////////
 // FPoseSearchIndex
 
-bool FPoseSearchIndex::IsValid() const
-{
-	const bool bSchemaValid = Schema && Schema->IsValid();
-	const bool bSearchIndexValid = bSchemaValid && (NumPoses * Schema->SchemaCardinality == Values.Num());
-
-	return bSearchIndexValid;
-}
-
 bool FPoseSearchIndex::IsEmpty() const
 {
 	const bool bEmpty = Assets.Num() == 0 || NumPoses == 0;
@@ -606,32 +671,32 @@ bool FPoseSearchIndex::IsEmpty() const
 
 TConstArrayView<float> FPoseSearchIndex::GetPoseValues(int32 PoseIdx) const
 {
-	check(PoseIdx >= 0 && PoseIdx < NumPoses && Schema && Schema->SchemaCardinality > 0);
-	const int32 ValueOffset = PoseIdx * Schema->SchemaCardinality;
-	return MakeArrayView(&Values[ValueOffset], Schema->SchemaCardinality);
+	const int32 SchemaCardinality = WeightsSqrt.Num();
+	check(PoseIdx >= 0 && PoseIdx < NumPoses && SchemaCardinality > 0);
+	const int32 ValueOffset = PoseIdx * SchemaCardinality;
+	return MakeArrayView(&Values[ValueOffset], SchemaCardinality);
 }
 
 TConstArrayView<float> FPoseSearchIndex::GetPoseValuesSafe(int32 PoseIdx) const
 {
-	if (PoseIdx >= 0 && PoseIdx < NumPoses && Schema && Schema->SchemaCardinality >= 0)
+	if (PoseIdx >= 0 && PoseIdx < NumPoses)
 	{
-		const int32 ValueOffset = PoseIdx * Schema->SchemaCardinality;
-		return MakeArrayView(&Values[ValueOffset], Schema->SchemaCardinality);
+		const int32 SchemaCardinality = WeightsSqrt.Num();
+		const int32 ValueOffset = PoseIdx * SchemaCardinality;
+		return MakeArrayView(&Values[ValueOffset], SchemaCardinality);
 	}
 	return TConstArrayView<float>();
 }
 
 void FPoseSearchIndex::Reset()
 {
-	NumPoses = 0;
-	Assets.Reset();
-	Values.Reset();
-	Schema = nullptr;
+	FPoseSearchIndex Default;
+	*this = Default;
 }
 
-FPoseSearchCost FPoseSearchIndex::ComparePoses(int32 PoseIdx, EPoseSearchBooleanRequest QueryMirrorRequest, UE::PoseSearch::EPoseComparisonFlags PoseComparisonFlags, TConstArrayView<float> QueryValues) const
+FPoseSearchCost FPoseSearchIndex::ComparePoses(int32 PoseIdx, EPoseSearchBooleanRequest QueryMirrorRequest, UE::PoseSearch::EPoseComparisonFlags PoseComparisonFlags, float MirrorMismatchCostBias, TConstArrayView<float> QueryValues) const
 {
-	// base dissimilarity cost representing how the assiciated PoseIdx differ, in a weighted way, from the query pose (QueryValues)
+	// base dissimilarity cost representing how the associated PoseIdx differ, in a weighted way, from the query pose (QueryValues)
 	const float DissimilarityCost = UE::PoseSearch::CompareFeatureVectors(GetPoseValues(PoseIdx), QueryValues, WeightsSqrt);
 	
 	// cost addend associated to a mismatch in mirror state between query and analyzed PoseIdx
@@ -644,8 +709,7 @@ FPoseSearchCost FPoseSearchIndex::ComparePoses(int32 PoseIdx, EPoseSearchBoolean
 			(!IndexAsset.bMirrored && QueryMirrorRequest == EPoseSearchBooleanRequest::TrueValue);
 		if (bMirroringMismatch)
 		{
-			check(Schema);
-			MirrorMismatchAddend = Schema->MirrorMismatchCostBias;
+			MirrorMismatchAddend = MirrorMismatchCostBias;
 		}
 	}
 
@@ -660,6 +724,139 @@ FPoseSearchCost FPoseSearchIndex::ComparePoses(int32 PoseIdx, EPoseSearchBoolean
 	return FPoseSearchCost(DissimilarityCost, NotifyAddend, MirrorMismatchAddend, ContinuingPoseCostAddend);
 }
 
+void FPoseSearchIndex::InitSearchIndexAssets(TConstArrayView<FPoseSearchDatabaseSequence> Sequences, TConstArrayView<FPoseSearchDatabaseBlendSpace> BlendSpaces, const FPoseSearchExcludeFromDatabaseParameters& ExcludeFromDatabaseParameters)
+{
+	using namespace UE::PoseSearch;
+
+	Assets.Empty();
+
+	TArray<FFloatRange> ValidRanges;
+
+	for (int32 SequenceIdx = 0; SequenceIdx < Sequences.Num(); ++SequenceIdx)
+	{
+		const FPoseSearchDatabaseSequence& Sequence = Sequences[SequenceIdx];
+
+		if (Sequence.bEnabled)
+		{
+			const bool bAddUnmirrored = Sequence.MirrorOption == EPoseSearchMirrorOption::UnmirroredOnly || Sequence.MirrorOption == EPoseSearchMirrorOption::UnmirroredAndMirrored;
+			const bool bAddMirrored = Sequence.MirrorOption == EPoseSearchMirrorOption::MirroredOnly || Sequence.MirrorOption == EPoseSearchMirrorOption::UnmirroredAndMirrored;
+
+			ValidRanges.Reset();
+			FindValidSequenceIntervals(Sequence, ExcludeFromDatabaseParameters, ValidRanges);
+			for (const FFloatRange& Range : ValidRanges)
+			{
+				if (bAddUnmirrored)
+				{
+					Assets.Add(FPoseSearchIndexAsset(ESearchIndexAssetType::Sequence, SequenceIdx, false, FFloatInterval(Range.GetLowerBoundValue(), Range.GetUpperBoundValue())));
+				}
+
+				if (bAddMirrored)
+				{
+					Assets.Add(FPoseSearchIndexAsset(ESearchIndexAssetType::Sequence, SequenceIdx, true, FFloatInterval(Range.GetLowerBoundValue(), Range.GetUpperBoundValue())));
+				}
+			}
+		}
+	}
+
+	TArray<FBlendSampleData> BlendSamples;
+
+	for (int32 BlendSpaceIdx = 0; BlendSpaceIdx < BlendSpaces.Num(); ++BlendSpaceIdx)
+	{
+		const FPoseSearchDatabaseBlendSpace& BlendSpace = BlendSpaces[BlendSpaceIdx];
+
+		if (BlendSpace.bEnabled)
+		{
+			const bool bAddUnmirrored = BlendSpace.MirrorOption == EPoseSearchMirrorOption::UnmirroredOnly || BlendSpace.MirrorOption == EPoseSearchMirrorOption::UnmirroredAndMirrored;
+			const bool bAddMirrored = BlendSpace.MirrorOption == EPoseSearchMirrorOption::MirroredOnly || BlendSpace.MirrorOption == EPoseSearchMirrorOption::UnmirroredAndMirrored;
+
+			int32 HorizontalBlendNum, VerticalBlendNum;
+			BlendSpace.GetBlendSpaceParameterSampleRanges(HorizontalBlendNum, VerticalBlendNum);
+
+			const bool bWrapInputOnHorizontalAxis = BlendSpace.BlendSpace->GetBlendParameter(0).bWrapInput;
+			const bool bWrapInputOnVerticalAxis = BlendSpace.BlendSpace->GetBlendParameter(1).bWrapInput;
+			for (int32 HorizontalIndex = 0; HorizontalIndex < HorizontalBlendNum; HorizontalIndex++)
+			{
+				for (int32 VerticalIndex = 0; VerticalIndex < VerticalBlendNum; VerticalIndex++)
+				{
+					const FVector BlendParameters = BlendSpace.BlendParameterForSampleRanges(HorizontalIndex, VerticalIndex);
+
+					int32 TriangulationIndex = 0;
+					BlendSpace.BlendSpace->GetSamplesFromBlendInput(BlendParameters, BlendSamples, TriangulationIndex, true);
+
+					float PlayLength = BlendSpace.BlendSpace->GetAnimationLengthFromSampleData(BlendSamples);
+
+					if (bAddUnmirrored)
+					{
+						Assets.Add(FPoseSearchIndexAsset(ESearchIndexAssetType::BlendSpace, BlendSpaceIdx, false, FFloatInterval(0.0f, PlayLength), BlendParameters));
+					}
+
+					if (bAddMirrored)
+					{
+						Assets.Add(FPoseSearchIndexAsset(ESearchIndexAssetType::BlendSpace, BlendSpaceIdx, true, FFloatInterval(0.0f, PlayLength), BlendParameters));
+					}
+				}
+			}
+		}
+	}
+}
+
+
+FArchive& operator<<(FArchive& Ar, FPoseSearchIndex& Index)
+{
+	int32 NumValues = 0;
+	int32 NumPCAValues = 0;
+	int32 NumAssets = 0;
+
+	if (Ar.IsSaving())
+	{
+		NumValues = Index.Values.Num();
+		NumPCAValues = Index.PCAValues.Num();
+		NumAssets = Index.Assets.Num();
+	}
+
+	Ar << Index.NumPoses;
+	Ar << NumValues;
+	Ar << NumPCAValues;
+	Ar << NumAssets;
+	Ar << Index.OverallFlags;
+
+	if (Ar.IsLoading())
+	{
+		Index.Values.SetNumUninitialized(NumValues);
+		Index.PCAValues.SetNumUninitialized(NumPCAValues);
+		Index.PoseMetadata.SetNumUninitialized(Index.NumPoses);
+		Index.Assets.SetNumUninitialized(NumAssets);
+	}
+
+	if (Index.Values.Num() > 0)
+	{
+		Ar.Serialize(&Index.Values[0], Index.Values.Num() * Index.Values.GetTypeSize());
+	}
+
+	if (Index.PCAValues.Num() > 0)
+	{
+		Ar.Serialize(&Index.PCAValues[0], Index.PCAValues.Num() * Index.PCAValues.GetTypeSize());
+	}
+
+	if (Index.PoseMetadata.Num() > 0)
+	{
+		Ar.Serialize(&Index.PoseMetadata[0], Index.PoseMetadata.Num() * Index.PoseMetadata.GetTypeSize());
+	}
+
+	if (Index.Assets.Num() > 0)
+	{
+		Ar.Serialize(&Index.Assets[0], Index.Assets.Num() * Index.Assets.GetTypeSize());
+	}
+
+	Ar << Index.WeightsSqrt;
+	Ar << Index.Mean;
+	Ar << Index.PCAProjectionMatrix;
+
+	Serialize(Ar, Index.KDTree, Index.PCAValues.GetData());
+
+	return Ar;
+}
+
 //////////////////////////////////////////////////////////////////////////
 // UPoseSearchSequenceMetaData
 
@@ -691,7 +888,7 @@ bool UPoseSearchSequenceMetaData::IsValidForIndexing() const
 
 bool UPoseSearchSequenceMetaData::IsValidForSearch() const
 {
-	return IsValidForIndexing() && SearchIndex.IsValid() && !SearchIndex.IsEmpty();
+	return IsValidForIndexing() && !SearchIndex.IsEmpty();
 }
 
 UE::PoseSearch::FSearchResult UPoseSearchSequenceMetaData::Search(UE::PoseSearch::FSearchContext& SearchContext) const
@@ -700,7 +897,7 @@ UE::PoseSearch::FSearchResult UPoseSearchSequenceMetaData::Search(UE::PoseSearch
 
 	FSearchResult Result;
 
-	if (!ensure(SearchIndex.IsValid() && !SearchIndex.IsEmpty()))
+	if (!ensure(Schema->IsValid() && !SearchIndex.IsEmpty()))
 	{
 		return Result;
 	}
@@ -708,7 +905,7 @@ UE::PoseSearch::FSearchResult UPoseSearchSequenceMetaData::Search(UE::PoseSearch
 	Schema->BuildQuery(SearchContext, Result.ComposedQuery);
 	TConstArrayView<float> QueryValues = Result.ComposedQuery.GetValues();
 
-	if (!ensure(QueryValues.Num() == SearchIndex.Schema->SchemaCardinality))
+	if (!ensure(QueryValues.Num() == Schema->SchemaCardinality))
 	{
 		return Result;
 	}
@@ -725,7 +922,7 @@ UE::PoseSearch::FSearchResult UPoseSearchSequenceMetaData::Search(UE::PoseSearch
 				continue;
 			}
 
-			const FPoseSearchCost PoseCost = SearchIndex.ComparePoses(PoseIdx, EPoseSearchBooleanRequest::Indifferent, EPoseComparisonFlags::ContinuingPose, QueryValues);
+			const FPoseSearchCost PoseCost = SearchIndex.ComparePoses(PoseIdx, EPoseSearchBooleanRequest::Indifferent, EPoseComparisonFlags::ContinuingPose, Schema->MirrorMismatchCostBias, QueryValues);
 			if (PoseCost < Result.PoseCost)
 			{
 				Result.PoseCost = PoseCost;
@@ -736,7 +933,7 @@ UE::PoseSearch::FSearchResult UPoseSearchSequenceMetaData::Search(UE::PoseSearch
 
 	if (Result.PoseIdx != INDEX_NONE)
 	{
-		Result.AssetTime = SearchIndex.GetAssetTime(Result.PoseIdx);
+		Result.AssetTime = SearchIndex.GetAssetTime(Result.PoseIdx, Schema->GetSamplingInterval());
 	}
 
 #if ENABLE_DRAW_DEBUG
@@ -754,54 +951,61 @@ UE::PoseSearch::FSearchResult UPoseSearchSequenceMetaData::Search(UE::PoseSearch
 
 UPoseSearchDatabase::~UPoseSearchDatabase()
 {
-	if (PrivateDerivedData)
-	{
 #if WITH_EDITOR
-		PrivateDerivedData->CancelCache();
+	UE::PoseSearch::FAsyncPoseSearchDatabasesManagement::Get().WaitOnExistingBuildIndex(*this, false);
 #endif
-		delete PrivateDerivedData;
-	}
 }
 
 FPoseSearchIndex* UPoseSearchDatabase::GetSearchIndex()
 {
-	if (PrivateDerivedData == nullptr)
-	{
-		return nullptr;
-	}
-
-	return &PrivateDerivedData->SearchIndex;
+	return &PoseSearchIndex;
 }
 
 const FPoseSearchIndex* UPoseSearchDatabase::GetSearchIndex() const
 {
-	return PrivateDerivedData ? &PrivateDerivedData->SearchIndex : nullptr;
+	return &PoseSearchIndex;
 }
 
 const FPoseSearchIndex* UPoseSearchDatabase::GetSearchIndexSafe(bool bVerboseLogging) const
 {
-	const FPoseSearchIndex* SearchIndex = GetSearchIndex();
+	if (!Schema)
+	{
+		if (bVerboseLogging)
+		{
+			UE_LOG(LogAnimation, Warning, TEXT("UPoseSearchDatabase %s failed to index. Reason: no Schema!"), *GetName());
+		}
+		return nullptr;
+	}
 
-	if (!SearchIndex || !SearchIndex->IsValid() || SearchIndex->IsEmpty())
+	if (!Schema->IsValid())
 	{
 		if (bVerboseLogging)
 		{
-			if (!Schema)
-			{
-				UE_LOG(LogAnimation, Warning, TEXT("UPoseSearchDatabase %s failed to index. Reason: no Schema!"), *GetName());
-			}
-			else if (!Schema->IsValid())
-			{
-				UE_LOG(LogAnimation, Warning, TEXT("UPoseSearchDatabase %s failed to index. Reason: Schema %s is invalid"), *GetName(), *Schema->GetName());
-			}
-			else
-			{
-				UE_LOG(LogAnimation, Warning, TEXT("UPoseSearchDatabase %s failed to index. Reason: is there any unsaved modified asset?"), *GetName());
-			}
+			UE_LOG(LogAnimation, Warning, TEXT("UPoseSearchDatabase %s failed to index. Reason: Schema %s is invalid"), *GetName(), *Schema->GetName());
+		}
+		return nullptr;
+	}
+
+	const FPoseSearchIndex* SearchIndex = GetSearchIndex();
+	if (!SearchIndex || SearchIndex->IsEmpty())
+	{
+		if (bVerboseLogging)
+		{
+			UE_LOG(LogAnimation, Warning, TEXT("UPoseSearchDatabase %s failed to index. Reason: is there any unsaved modified asset?"), *GetName());
 		}
-		SearchIndex = nullptr;
+		return nullptr;
 	}
 
+	// @todo: perhaps use FAsyncPoseSearchDatabasesManagement to understand if this UPoseSearchDatabase indexing task is running, instead of using this logic
+	if (!SearchIndex || SearchIndex->IsEmpty() || SearchIndex->WeightsSqrt.Num() != Schema->SchemaCardinality)
+	{
+		if (bVerboseLogging)
+		{
+			UE_LOG(LogAnimation, Warning, TEXT("UPoseSearchDatabase %s SearchIndex mismatch. Indexing..."), *GetName());
+		}
+		return nullptr;
+	}
+	
 	return SearchIndex;
 }
 
@@ -988,12 +1192,17 @@ static void AddPoseSearchNotifiesToWriter(UAnimSequence* Sequence, UE::PoseSearc
 	}
 }
 
-
-void FPoseSearchDatabaseSequence::BuildDerivedDataKey(UE::PoseSearch::FDerivedDataKeyBuilder& KeyBuilder)
+void FPoseSearchDatabaseSequence::BuildDerivedDataKey(UE::PoseSearch::FDerivedDataKeyBuilder& KeyBuilder) const
 {
-	KeyBuilder << bEnabled;
-	KeyBuilder << SamplingRange;
-	KeyBuilder << MirrorOption;
+	check(KeyBuilder.IsSaving());
+
+	bool bEnabledCopy = bEnabled;
+	FFloatInterval SamplingRangeCopy = SamplingRange;
+	EPoseSearchMirrorOption MirrorOptionCopy = MirrorOption;
+
+	KeyBuilder << bEnabledCopy;
+	KeyBuilder << SamplingRangeCopy;
+	KeyBuilder << MirrorOptionCopy;
 
 	AddRawSequenceToWriter(Sequence, KeyBuilder);
 	AddRawSequenceToWriter(LeadInSequence, KeyBuilder);
@@ -1002,13 +1211,20 @@ void FPoseSearchDatabaseSequence::BuildDerivedDataKey(UE::PoseSearch::FDerivedDa
 	AddPoseSearchNotifiesToWriter(Sequence, KeyBuilder);
 }
 
-void FPoseSearchDatabaseBlendSpace::BuildDerivedDataKey(UE::PoseSearch::FDerivedDataKeyBuilder& KeyBuilder)
+void FPoseSearchDatabaseBlendSpace::BuildDerivedDataKey(UE::PoseSearch::FDerivedDataKeyBuilder& KeyBuilder) const
 {
-	KeyBuilder << bEnabled;
-	KeyBuilder << MirrorOption;
-	KeyBuilder << bUseGridForSampling;
-	KeyBuilder << NumberOfHorizontalSamples;
-	KeyBuilder << NumberOfVerticalSamples;
+	check(KeyBuilder.IsSaving());
+	bool bEnabledCopy = bEnabled;
+	EPoseSearchMirrorOption MirrorOptionCopy = MirrorOption;
+	bool bUseGridForSamplingCopy = bUseGridForSampling;
+	int32 NumberOfHorizontalSamplesCopy = NumberOfHorizontalSamples;
+	int32 NumberOfVerticalSamplesCopy = NumberOfVerticalSamples;
+
+	KeyBuilder << bEnabledCopy;
+	KeyBuilder << MirrorOptionCopy;
+	KeyBuilder << bUseGridForSamplingCopy;
+	KeyBuilder << NumberOfHorizontalSamplesCopy;
+	KeyBuilder << NumberOfVerticalSamplesCopy;
 
 	const TArray<FBlendSample>& BlendSpaceSamples = BlendSpace->GetBlendSamples();
 	for (const FBlendSample& Sample : BlendSpaceSamples)
@@ -1023,7 +1239,7 @@ void FPoseSearchDatabaseBlendSpace::BuildDerivedDataKey(UE::PoseSearch::FDerived
 	KeyBuilder << BlendSpace->bLoop;
 }
 
-void UPoseSearchDatabase::BuildDerivedDataKey(UE::PoseSearch::FDerivedDataKeyBuilder& KeyBuilder)
+void UPoseSearchDatabase::BuildDerivedDataKey(UE::PoseSearch::FDerivedDataKeyBuilder& KeyBuilder) const
 {
 	KeyBuilder.Update(this);
 
@@ -1032,12 +1248,12 @@ void UPoseSearchDatabase::BuildDerivedDataKey(UE::PoseSearch::FDerivedDataKeyBui
 		KeyBuilder.Update(Schema.Get());
 	}
 
-	for (FPoseSearchDatabaseSequence& DbSequence : Sequences)
+	for (const FPoseSearchDatabaseSequence& DbSequence : Sequences)
 	{
 		DbSequence.BuildDerivedDataKey(KeyBuilder);
 	}
 
-	for (FPoseSearchDatabaseBlendSpace& DbBlendSpace : BlendSpaces)
+	for (const FPoseSearchDatabaseBlendSpace& DbBlendSpace : BlendSpaces)
 	{
 		DbBlendSpace.BuildDerivedDataKey(KeyBuilder);
 	}
@@ -1072,7 +1288,7 @@ bool UPoseSearchDatabase::IsValidForIndexing() const
 			const USkeleton* SeqSkeleton = DbSequence.Sequence->GetSkeleton();
 			if (!SeqSkeleton)
 			{
-				bValid = false;
+				bValid = false; 
 				break;
 			}
 		}
@@ -1099,14 +1315,13 @@ bool UPoseSearchDatabase::IsValidForIndexing() const
 
 bool UPoseSearchDatabase::IsValidForSearch() const
 {
-	const FPoseSearchIndex* SearchIndex = GetSearchIndex();
-	bool bIsValid = IsValidForIndexing() && SearchIndex && SearchIndex->IsValid() && !SearchIndex->IsEmpty();
+	bool bIsValid = IsValidForIndexing() && GetSearchIndexSafe(false);
 
 #if WITH_EDITOR
-	const bool bIsCurrentDerivedData = 
-		PrivateDerivedData &&
-		PrivateDerivedData->PendingDerivedDataKey == PrivateDerivedData->DerivedDataKey.Hash;
-	bIsValid = bIsValid && bIsCurrentDerivedData;
+	if (bIsValid && UE::PoseSearch::FAsyncPoseSearchDatabasesManagement::Get().IsBuildingIndex(*this))
+	{
+		bIsValid = false;
+	}
 #endif // WITH_EDITOR
 
 	return bIsValid;
@@ -1118,89 +1333,6 @@ bool UPoseSearchDatabase::IsValidPoseIndex(int32 PoseIdx) const
 	return SearchIndex ? SearchIndex->IsValidPoseIndex(PoseIdx) : false;
 }
 
-void UPoseSearchDatabase::CollectSimpleSequences()
-{
-	for (auto& SimpleSequence: SimpleSequences)
-	{
-		auto Predicate = [&SimpleSequence](FPoseSearchDatabaseSequence& DbSequence) -> bool
-		{
-			return DbSequence.Sequence == SimpleSequence;
-		};
-
-		if (!Sequences.ContainsByPredicate(Predicate))
-		{
-			FPoseSearchDatabaseSequence& DbSequence = Sequences.AddDefaulted_GetRef();
-			DbSequence.Sequence = SimpleSequence;
-		}
-	}
-
-	SimpleSequences.Reset();
-}
-
-void UPoseSearchDatabase::CollectSimpleBlendSpaces()
-{
-	for (auto& SimpleBlendSpace : SimpleBlendSpaces)
-	{
-		auto Predicate = [&SimpleBlendSpace](FPoseSearchDatabaseBlendSpace& DbBlendSpace) -> bool
-		{
-			return DbBlendSpace.BlendSpace == SimpleBlendSpace;
-		};
-
-		if (!BlendSpaces.ContainsByPredicate(Predicate))
-		{
-			FPoseSearchDatabaseBlendSpace& DbBlendSpace = BlendSpaces.AddDefaulted_GetRef();
-			DbBlendSpace.BlendSpace = SimpleBlendSpace;
-		}
-	}
-
-	SimpleBlendSpaces.Reset();
-}
-
-void UPoseSearchDatabase::FindValidSequenceIntervals(const FPoseSearchDatabaseSequence& DbSequence, TArray<FFloatRange>& ValidRanges) const
-{
-	const UAnimSequence* Sequence = DbSequence.Sequence;
-	check(DbSequence.Sequence);
-
-	const float SequenceLength = DbSequence.Sequence->GetPlayLength();
-
-	const FFloatInterval EffectiveSamplingInterval = UE::PoseSearch::GetEffectiveSamplingRange(DbSequence.Sequence, DbSequence.SamplingRange);
-	FFloatRange EffectiveSamplingRange = FFloatRange::Inclusive(EffectiveSamplingInterval.Min, EffectiveSamplingInterval.Max);
-	if (!DbSequence.IsLooping())
-	{
-		const FFloatRange ExcludeFromDatabaseRange(ExcludeFromDatabaseParameters.SequenceStartInterval, SequenceLength - ExcludeFromDatabaseParameters.SequenceEndInterval);
-		EffectiveSamplingRange = FFloatRange::Intersection(EffectiveSamplingRange, ExcludeFromDatabaseRange);
-	}
-
-	// start from a single interval defined by the database sequence sampling range
-	ValidRanges.Empty();
-	ValidRanges.Add(EffectiveSamplingRange);
-
-	FAnimNotifyContext NotifyContext;
-	Sequence->GetAnimNotifies(0.0f, SequenceLength, NotifyContext);
-
-	for (const FAnimNotifyEventReference& EventReference : NotifyContext.ActiveNotifies)
-	{
-		if (const FAnimNotifyEvent* NotifyEvent = EventReference.GetNotify())
-		{
-			if (const UAnimNotifyState_PoseSearchExcludeFromDatabase* ExclusionNotifyState = Cast<const UAnimNotifyState_PoseSearchExcludeFromDatabase>(NotifyEvent->NotifyStateClass))
-			{
-				FFloatRange ExclusionRange = FFloatRange::Inclusive(NotifyEvent->GetTriggerTime(), NotifyEvent->GetEndTriggerTime());
-
-				// Split every valid range based on the exclusion range just found. Because this might increase the 
-				// number of ranges in ValidRanges, the algorithm iterates from end to start.
-				for (int RangeIdx = ValidRanges.Num() - 1; RangeIdx >= 0; --RangeIdx)
-				{
-					FFloatRange EvaluatedRange = ValidRanges[RangeIdx];
-					ValidRanges.RemoveAt(RangeIdx);
-
-					TArray<FFloatRange> Diff = FFloatRange::Difference(EvaluatedRange, ExclusionRange);
-					ValidRanges.Append(Diff);
-				}
-			}
-		}
-	}
-}
-
 UAnimationAsset* FPoseSearchDatabaseBlendSpace::GetAnimationAsset() const
 {
 	return BlendSpace.Get();
@@ -1252,148 +1384,19 @@ FVector FPoseSearchDatabaseBlendSpace::BlendParameterForSampleRanges(int32 Horiz
 	return FVector(
 		HorizontalBlendNum > 1 ? 
 			HorizontalBlendMin + (HorizontalBlendMax - HorizontalBlendMin) * 
-			((float)HorizontalBlendIndex) / (HorizontalBlendNum - 1) : 
-		HorizontalBlendMin,
-		VerticalBlendNum > 1 ? 
-			VerticalBlendMin + (VerticalBlendMax - VerticalBlendMin) * 
-			((float)VerticalBlendIndex) / (VerticalBlendNum - 1) : 
-		VerticalBlendMin,
-		0.0f);
-}
-
-bool UPoseSearchDatabase::TryInitSearchIndexAssets(FPoseSearchIndex& OutSearchIndex) const
-{
-	OutSearchIndex.Assets.Empty();
-	
-	bool bAnyMirrored = false;
-	
-	TArray<FFloatRange> ValidRanges;
-
-	for (int32 SequenceIdx = 0; SequenceIdx < Sequences.Num(); ++SequenceIdx)
-	{
-		const FPoseSearchDatabaseSequence& Sequence = Sequences[SequenceIdx];
-
-		if (!Sequence.bEnabled)
-		{
-			continue;
-		}
-
-		bool bAddUnmirrored = 
-			Sequence.MirrorOption == EPoseSearchMirrorOption::UnmirroredOnly ||
-			Sequence.MirrorOption == EPoseSearchMirrorOption::UnmirroredAndMirrored;
-		bool bAddMirrored =
-			Sequence.MirrorOption == EPoseSearchMirrorOption::MirroredOnly ||
-			Sequence.MirrorOption == EPoseSearchMirrorOption::UnmirroredAndMirrored;
-
-		ValidRanges.Reset();
-		FindValidSequenceIntervals(Sequence, ValidRanges);
-		for (const FFloatRange& Range : ValidRanges)
-		{
-			if (bAddUnmirrored)
-			{
-				OutSearchIndex.Assets.Add(
-					FPoseSearchIndexAsset(
-						ESearchIndexAssetType::Sequence,
-						SequenceIdx,
-						false,
-						FFloatInterval(Range.GetLowerBoundValue(), Range.GetUpperBoundValue())));
-			}
-
-			if (bAddMirrored)
-			{
-				OutSearchIndex.Assets.Add(
-					FPoseSearchIndexAsset(
-						ESearchIndexAssetType::Sequence,
-						SequenceIdx,
-						true,
-						FFloatInterval(Range.GetLowerBoundValue(), Range.GetUpperBoundValue())));
-				bAnyMirrored = true;
-			}
-		}
-	}
-
-	TArray<FBlendSampleData> BlendSamples;
-
-	for (int32 BlendSpaceIdx = 0; BlendSpaceIdx < BlendSpaces.Num(); ++BlendSpaceIdx)
-	{
-		const FPoseSearchDatabaseBlendSpace& BlendSpace = BlendSpaces[BlendSpaceIdx];
-
-		if (!BlendSpace.bEnabled)
-		{
-			continue;
-		}
-
-		bool bAddUnmirrored =
-			BlendSpace.MirrorOption == EPoseSearchMirrorOption::UnmirroredOnly ||
-			BlendSpace.MirrorOption == EPoseSearchMirrorOption::UnmirroredAndMirrored;
-		bool bAddMirrored =
-			BlendSpace.MirrorOption == EPoseSearchMirrorOption::MirroredOnly ||
-			BlendSpace.MirrorOption == EPoseSearchMirrorOption::UnmirroredAndMirrored;
-
-		int32 HorizontalBlendNum, VerticalBlendNum;
-		BlendSpace.GetBlendSpaceParameterSampleRanges(HorizontalBlendNum, VerticalBlendNum);
-
-		const bool bWrapInputOnHorizontalAxis = BlendSpace.BlendSpace->GetBlendParameter(0).bWrapInput;
-		const bool bWrapInputOnVerticalAxis = BlendSpace.BlendSpace->GetBlendParameter(1).bWrapInput;
-		for (int32 HorizontalIndex = 0; HorizontalIndex < HorizontalBlendNum; HorizontalIndex++)
-		{
-			for (int32 VerticalIndex = 0; VerticalIndex < VerticalBlendNum; VerticalIndex++)
-			{
-				const FVector BlendParameters = BlendSpace.BlendParameterForSampleRanges(HorizontalIndex, VerticalIndex);
-						
-				int32 TriangulationIndex = 0;
-				BlendSpace.BlendSpace->GetSamplesFromBlendInput(BlendParameters, BlendSamples, TriangulationIndex, true);
-
-				float PlayLength = BlendSpace.BlendSpace->GetAnimationLengthFromSampleData(BlendSamples);
-
-				if (bAddUnmirrored)
-				{
-					OutSearchIndex.Assets.Add(
-						FPoseSearchIndexAsset(
-							ESearchIndexAssetType::BlendSpace,
-							BlendSpaceIdx,
-							false,
-							FFloatInterval(0.0f, PlayLength),
-							BlendParameters));
-				}
-
-				if (bAddMirrored)
-				{
-					OutSearchIndex.Assets.Add(
-						FPoseSearchIndexAsset(
-							ESearchIndexAssetType::BlendSpace,
-							BlendSpaceIdx,
-							true,
-							FFloatInterval(0.0f, PlayLength),
-							BlendParameters));
-					bAnyMirrored = true;
-				}
-			}
-		}
-	}
-
-	if (bAnyMirrored && !Schema->MirrorDataTable)
-	{
-		UE_LOG(
-			LogPoseSearch, 
-			Error, 
-			TEXT("Database %s is asking for mirrored sequences but MirrorDataBase is null in %s"),
-			*GetNameSafe(this), 
-			*GetNameSafe(Schema));
-		OutSearchIndex.Assets.Empty();
-		return false;
-	}
-
-	return true;
+			((float)HorizontalBlendIndex) / (HorizontalBlendNum - 1) : 
+		HorizontalBlendMin,
+		VerticalBlendNum > 1 ? 
+			VerticalBlendMin + (VerticalBlendMax - VerticalBlendMin) * 
+			((float)VerticalBlendIndex) / (VerticalBlendNum - 1) : 
+		VerticalBlendMin,
+		0.0f);
 }
 
 void UPoseSearchDatabase::PostLoad()
 {
 #if WITH_EDITOR
-	if (!PrivateDerivedData)
-	{
-		BeginCacheDerivedData();
-	}
+	UE::PoseSearch::FAsyncPoseSearchDatabasesManagement::Get().RequestAsyncBuildIndex(*this);
 #endif
 
 	Super::PostLoad();
@@ -1409,225 +1412,60 @@ void UPoseSearchDatabase::UnregisterOnDerivedDataRebuild(void* Unregister)
 	OnDerivedDataRebuild.RemoveAll(Unregister);
 }
 
-void UPoseSearchDatabase::NotifyDerivedDataBuildStarted()
+void UPoseSearchDatabase::NotifyDerivedDataBuildStarted() const
 {
 	OnDerivedDataRebuild.Broadcast();
 }
 
-void UPoseSearchDatabase::RegisterOnAssetChange(const FOnAssetChange& Delegate)
-{
-	OnAssetChange.Add(Delegate);
-}
-
-void UPoseSearchDatabase::UnregisterOnAssetChange(void* Unregister)
-{
-	OnAssetChange.RemoveAll(Unregister);
-}
-
-void UPoseSearchDatabase::NotifyAssetChange()
-{
-	OnAssetChange.Broadcast();
-}
-
-void UPoseSearchDatabase::RegisterOnGroupChange(const FOnGroupChange& Delegate)
-{
-	OnGroupChange.Add(Delegate);
-}
-
-void UPoseSearchDatabase::UnregisterOnGroupChange(void* Unregister)
-{
-	OnGroupChange.RemoveAll(Unregister);
-}
-
-void UPoseSearchDatabase::NotifyGroupChange()
+void UPoseSearchDatabase::BeginCacheForCookedPlatformData(const ITargetPlatform* TargetPlatform)
 {
-	OnGroupChange.Broadcast();
+	Super::BeginCacheForCookedPlatformData(TargetPlatform);
+	UE::PoseSearch::FAsyncPoseSearchDatabasesManagement::Get().RequestAsyncBuildIndex(*this, false, true);
 }
 
-void UPoseSearchDatabase::BeginCacheDerivedData()
+bool UPoseSearchDatabase::IsCachedCookedPlatformDataLoaded(const ITargetPlatform* TargetPlatform)
 {
-	bool bPerformCache = true;
-			
-	using namespace UE::DerivedData;
-	if (PrivateDerivedData)
-	{
-		const FIoHash ExistingDerivedDataHash = PrivateDerivedData->PendingDerivedDataKey;
-		if (!ExistingDerivedDataHash.IsZero())
-		{
-			const FIoHash CurrentHash = UE::PoseSearch::FPoseSearchDatabaseAsyncCacheTask::CreateKey(*this);
-			if (ExistingDerivedDataHash == CurrentHash)
-			{
-				bPerformCache = false;
-			}
-		}
-	}
+	check(IsInGameThread());
 
-	if (bPerformCache)
+	if (!UE::PoseSearch::FAsyncPoseSearchDatabasesManagement::Get().IsBuildingIndex(*this))
 	{
-		if (!PrivateDerivedData)
+		if (IsValidForSearch())
 		{
-			PrivateDerivedData = new FPoseSearchDatabaseDerivedData();
+			return true;
 		}
 
-		PrivateDerivedData->Cache(*this, false);
-	}
-}
-
-FIoHash UPoseSearchDatabase::GetSearchIndexHash() const
-{
-	if (!PrivateDerivedData)
-	{
-		return FIoHash::Zero;
-	}
-
-	return PrivateDerivedData->DerivedDataKey.Hash;
-}
-
-bool UPoseSearchDatabase::IsDerivedDataBuildPending() const
-{
-	if (!PrivateDerivedData)
-	{
-		return true;
+		UE::PoseSearch::FAsyncPoseSearchDatabasesManagement::Get().RequestAsyncBuildIndex(*this, false, true);
 	}
 
-	return PrivateDerivedData->DerivedDataKey.Hash != PrivateDerivedData->PendingDerivedDataKey;
+	return false;
 }
 #endif // WITH_EDITOR
 
-
 void UPoseSearchDatabase::PostSaveRoot(FObjectPostSaveRootContext ObjectSaveContext)
 {
-	Super::PostSaveRoot(ObjectSaveContext);
-
 #if WITH_EDITOR
 	if (!IsTemplate() && !ObjectSaveContext.IsProceduralSave())
 	{
-		if (IsValidForIndexing())
-		{
-			if (!PrivateDerivedData)
-			{
-				PrivateDerivedData = new FPoseSearchDatabaseDerivedData();
-			}
-
-			PrivateDerivedData->Cache(*this, true);
-		}
+		UE::PoseSearch::FAsyncPoseSearchDatabasesManagement::Get().RequestAsyncBuildIndex(*this, true, true);
 	}
 #endif
+
+	Super::PostSaveRoot(ObjectSaveContext);
 }
 
 void UPoseSearchDatabase::Serialize(FArchive& Ar)
 {
-	using namespace UE::PoseSearch;
-
- 	Super::Serialize(Ar);
+	Super::Serialize(Ar);
 
 	if (Ar.IsFilterEditorOnly())
 	{
-		if (Ar.IsLoading())
-		{
-			if (!PrivateDerivedData)
-			{
-				PrivateDerivedData = new FPoseSearchDatabaseDerivedData();
-				PrivateDerivedData->SearchIndex.Schema = Schema;
-			}
-		}
-		
-		if (Ar.IsLoading() || (Ar.IsCooking() && IsDerivedDataValid()))
-		{
-			FPoseSearchIndex* SearchIndex = GetSearchIndex();
-			Ar << *SearchIndex;
-		}
-		else
-		{
-			UE_LOG(
-				LogPoseSearch,
-				Warning,
-				TEXT("Invalid condition in UPoseSearchDatabase::Serialize for asset %s. It might be that the asset has a missing or invalid schema. IsLoading: %d, IsCooking: %d, IsDerivedDataValid: %d"),
-				*GetNameSafe(this),
-				Ar.IsLoading(), Ar.IsCooking(), IsDerivedDataValid());
-		}
-	}
-}
-
-bool UPoseSearchDatabase::IsDerivedDataValid()
-{
-	const FPoseSearchIndex* SearchIndex = GetSearchIndex();
-	bool bIsValid = SearchIndex && SearchIndex->IsValid();
-	return bIsValid;
-}
-
-#if WITH_EDITOR
-void UPoseSearchDatabase::PostEditChangeProperty(struct FPropertyChangedEvent& PropertyChangedEvent)
-{
-	Super::PostEditChangeProperty(PropertyChangedEvent);
-
-	bool bAssetChange = false;
-
-	if (PropertyChangedEvent.Property->GetFName() == GET_MEMBER_NAME_CHECKED(UPoseSearchDatabase, SimpleSequences))
-	{
-		if (!SimpleSequences.IsEmpty())
-		{
-			CollectSimpleSequences();
-			bAssetChange = true;
-		}
-	}
-
-	if (PropertyChangedEvent.Property->GetFName() == GET_MEMBER_NAME_CHECKED(UPoseSearchDatabase, SimpleBlendSpaces))
-	{
-		if (!SimpleBlendSpaces.IsEmpty())
+		if (Ar.IsLoading() || Ar.IsCooking())
 		{
-			CollectSimpleBlendSpaces();
-			bAssetChange = true;
+			Ar << PoseSearchIndex;
 		}
 	}
-
-	if (PropertyChangedEvent.MemberProperty->GetFName() == GET_MEMBER_NAME_CHECKED(UPoseSearchDatabase, Sequences) ||
-		PropertyChangedEvent.MemberProperty->GetFName() == GET_MEMBER_NAME_CHECKED(UPoseSearchDatabase, BlendSpaces))
-	{
-		bAssetChange = true;
-	}
-
-	if (bAssetChange)
-	{
-		NotifyAssetChange();
-	}
-
-	BeginCacheDerivedData();
-}
-
-void UPoseSearchDatabase::BeginCacheForCookedPlatformData(const ITargetPlatform* TargetPlatform)
-{
-	Super::BeginCacheForCookedPlatformData(TargetPlatform);
-
-	BeginCacheDerivedData();
-}
-
-bool UPoseSearchDatabase::IsCachedCookedPlatformDataLoaded(const ITargetPlatform* TargetPlatform)
-{
-	check(IsInGameThread());
-
-	if (!PrivateDerivedData)
-	{
-		PrivateDerivedData = new FPoseSearchDatabaseDerivedData();
-		PrivateDerivedData->Cache(*this, true);
-		return false;
-	}
-
-	if (PrivateDerivedData->AsyncTask && PrivateDerivedData->AsyncTask->Poll())
-	{
-		PrivateDerivedData->FinishCache();
-	}
-
-	if (PrivateDerivedData->AsyncTask)
-	{
-		return false;
-	}
-
-	return true;
 }
 
-#endif // WITH_EDITOR
-
 UE::PoseSearch::FSearchResult UPoseSearchDatabase::Search(UE::PoseSearch::FSearchContext& SearchContext) const
 {
 	using namespace UE::PoseSearch;
@@ -1635,7 +1473,7 @@ UE::PoseSearch::FSearchResult UPoseSearchDatabase::Search(UE::PoseSearch::FSearc
 	FSearchResult Result;
 
 #if WITH_EDITOR
-	if (IsDerivedDataBuildPending())
+	if (UE::PoseSearch::FAsyncPoseSearchDatabasesManagement::Get().IsBuildingIndex(*this))
 	{
 		return Result;
 	}
@@ -1739,7 +1577,7 @@ UE::PoseSearch::FSearchResult UPoseSearchDatabase::SearchPCAKDTree(UE::PoseSearc
 	if (!SearchContext.bForceInterrupt && IsCurrentResultFromThisDatabase && SearchContext.bCanAdvance && !Result.ContinuingPoseCost.IsValid())
 	{
 		Result.PoseIdx = SearchContext.CurrentResult.PoseIdx;
-		Result.PoseCost = SearchIndex->ComparePoses(Result.PoseIdx, SearchContext.QueryMirrorRequest, EPoseComparisonFlags::ContinuingPose, QueryValues);
+		Result.PoseCost = SearchIndex->ComparePoses(Result.PoseIdx, SearchContext.QueryMirrorRequest, EPoseComparisonFlags::ContinuingPose, Schema->MirrorMismatchCostBias, QueryValues);
 		Result.ContinuingPoseCost = Result.PoseCost;
 
 		if (GetSkipSearchIfPossible())
@@ -1771,7 +1609,7 @@ UE::PoseSearch::FSearchResult UPoseSearchDatabase::SearchPCAKDTree(UE::PoseSearc
 		SearchIndex->KDTree.FindNeighbors(ResultSet, ProjectedQueryValues.data());
 
 		// NonSelectableIdx are already filtered out inside the kdtree search
-		const FPoseFilters PoseFilters(SearchIndex->Schema, TConstArrayView<size_t>(), SearchIndex->OverallFlags);
+		const FPoseFilters PoseFilters(Schema, TConstArrayView<size_t>(), SearchIndex->OverallFlags);
 		for (size_t ResultIndex = 0; ResultIndex < ResultSet.Num(); ++ResultIndex)
 		{
 			const int32 PoseIdx = ResultIndexes[ResultIndex];
@@ -1781,7 +1619,7 @@ UE::PoseSearch::FSearchResult UPoseSearchDatabase::SearchPCAKDTree(UE::PoseSearc
 #endif
 			))
 			{
-				const FPoseSearchCost PoseCost = SearchIndex->ComparePoses(PoseIdx, SearchContext.QueryMirrorRequest, EPoseComparisonFlags::None, QueryValues);
+				const FPoseSearchCost PoseCost = SearchIndex->ComparePoses(PoseIdx, SearchContext.QueryMirrorRequest, EPoseComparisonFlags::None, Schema->MirrorMismatchCostBias, QueryValues);
 				if (PoseCost < Result.PoseCost)
 				{
 					Result.PoseCost = PoseCost;
@@ -1811,11 +1649,8 @@ UE::PoseSearch::FSearchResult UPoseSearchDatabase::SearchPCAKDTree(UE::PoseSearc
 	// finalizing Result properties
 	if (Result.PoseIdx != INDEX_NONE)
 	{
-		Result.AssetTime = SearchIndex->GetAssetTime(Result.PoseIdx);
+		Result.AssetTime = SearchIndex->GetAssetTime(Result.PoseIdx, Schema->GetSamplingInterval());
 		Result.Database = this;
-#if WITH_EDITOR
-		Result.SearchIndexHash = GetSearchIndexHash();
-#endif // WITH_EDITOR
 	}
 
 	return Result;
@@ -1843,7 +1678,7 @@ UE::PoseSearch::FSearchResult UPoseSearchDatabase::SearchBruteForce(UE::PoseSear
 		if (SearchContext.bCanAdvance && !Result.ContinuingPoseCost.IsValid())
 		{
 			Result.PoseIdx = SearchContext.CurrentResult.PoseIdx;
-			Result.PoseCost = SearchIndex->ComparePoses(Result.PoseIdx, SearchContext.QueryMirrorRequest, EPoseComparisonFlags::ContinuingPose, QueryValues);
+			Result.PoseCost = SearchIndex->ComparePoses(Result.PoseIdx, SearchContext.QueryMirrorRequest, EPoseComparisonFlags::ContinuingPose, Schema->MirrorMismatchCostBias, QueryValues);
 			Result.ContinuingPoseCost = Result.PoseCost;
 
 			if (GetSkipSearchIfPossible())
@@ -1861,7 +1696,7 @@ UE::PoseSearch::FSearchResult UPoseSearchDatabase::SearchBruteForce(UE::PoseSear
 		PopulateNonSelectableIdx(NonSelectableIdx, SearchContext, this, QueryValues);
 		check(Algo::IsSorted(NonSelectableIdx));
 
-		const FPoseFilters PoseFilters(SearchIndex->Schema, NonSelectableIdx, SearchIndex->OverallFlags);
+		const FPoseFilters PoseFilters(Schema, NonSelectableIdx, SearchIndex->OverallFlags);
 		for (int32 PoseIdx = 0; PoseIdx < SearchIndex->NumPoses; ++PoseIdx)
 		{
 			if (PoseFilters.AreFiltersValid(SearchIndex, QueryValues, PoseIdx, SearchIndex->PoseMetadata[PoseIdx]
@@ -1870,7 +1705,7 @@ UE::PoseSearch::FSearchResult UPoseSearchDatabase::SearchBruteForce(UE::PoseSear
 #endif
 			))
 			{
-				const FPoseSearchCost PoseCost = SearchIndex->ComparePoses(PoseIdx, SearchContext.QueryMirrorRequest, EPoseComparisonFlags::None, QueryValues);
+				const FPoseSearchCost PoseCost = SearchIndex->ComparePoses(PoseIdx, SearchContext.QueryMirrorRequest, EPoseComparisonFlags::None, Schema->MirrorMismatchCostBias, QueryValues);
 				if (PoseCost < Result.PoseCost)
 				{
 					Result.PoseCost = PoseCost;
@@ -1903,11 +1738,8 @@ UE::PoseSearch::FSearchResult UPoseSearchDatabase::SearchBruteForce(UE::PoseSear
 	// finalizing Result properties
 	if (Result.PoseIdx != INDEX_NONE)
 	{
-		Result.AssetTime = SearchIndex->GetAssetTime(Result.PoseIdx);
+		Result.AssetTime = SearchIndex->GetAssetTime(Result.PoseIdx, Schema->GetSamplingInterval());
 		Result.Database = this;
-#if WITH_EDITOR
-		Result.SearchIndexHash = GetSearchIndexHash();
-#endif // WITH_EDITOR
 	}
 
 #if WITH_EDITORONLY_DATA
@@ -1949,15 +1781,12 @@ UE::PoseSearch::FSearchResult UPoseSearchDatabaseSet::Search(UE::PoseSearch::FSe
 			TConstArrayView<float> QueryValues = Result.ComposedQuery.GetValues();
 
 			Result.PoseIdx = SearchContext.CurrentResult.PoseIdx;
-			Result.PoseCost = SearchIndex->ComparePoses(Result.PoseIdx, SearchContext.QueryMirrorRequest, EPoseComparisonFlags::ContinuingPose, QueryValues);
+			Result.PoseCost = SearchIndex->ComparePoses(Result.PoseIdx, SearchContext.QueryMirrorRequest, EPoseComparisonFlags::ContinuingPose, Database->Schema->MirrorMismatchCostBias, QueryValues);
 			Result.ContinuingPoseCost = Result.PoseCost;
 			ContinuingCost = Result.PoseCost;
 
-			Result.AssetTime = SearchIndex->GetAssetTime(Result.PoseIdx);
+			Result.AssetTime = SearchIndex->GetAssetTime(Result.PoseIdx, Database->Schema->GetSamplingInterval());
 			Result.Database = Database;
-#if WITH_EDITOR
-			Result.SearchIndexHash = Database->GetSearchIndexHash();
-#endif // WITH_EDITOR
 
 			if (Database->GetSkipSearchIfPossible())
 			{
@@ -2097,11 +1926,7 @@ FTransform FSearchContext::TryGetTransformAndCacheResults(float SampleTime, cons
 			// adding bunch of entries, without caring about adding eventual duplicates
 			for (const FBoneIndexType NewEntryBoneIndexType : Schema->BoneIndicesWithParents)
 			{
-				// @todo: maybe add them with a single allocation with CachedEntries.AddDefaulted(Schema->BoneIndicesWithParents.Num())
-				FCachedEntry& NewEntry = CachedEntries[CachedEntries.AddDefaulted()];
-				NewEntry.SampleTime = SampleTime;
-				NewEntry.BoneIndexType = NewEntryBoneIndexType;
-				NewEntry.Transform = SampledComponentPose[NewEntryBoneIndexType];
+				CachedEntries.Emplace(SampleTime, SampledComponentPose[NewEntryBoneIndexType], NewEntryBoneIndexType);
 			}
 
 			bError = false;
@@ -2615,7 +2440,7 @@ bool FDebugDrawParams::CanDraw() const
 		return false;
 	}
 
-	return SearchIndex->IsValid() && !SearchIndex->IsEmpty() && SearchIndex->Schema->SchemaCardinality > 0;
+	return true;
 #else // ENABLE_DRAW_DEBUG
 	return false;
 #endif // ENABLE_DRAW_DEBUG
@@ -2664,7 +2489,7 @@ const FPoseSearchIndex* FDebugDrawParams::GetSearchIndex() const
 {
 	if (Database)
 	{
-		return Database->GetSearchIndex();
+		return Database->GetSearchIndexSafe(false);
 	}
 
 	if (SequenceMetaData)
@@ -2745,17 +2570,7 @@ void FSearchResult::Update(float NewAssetTime)
 
 bool FSearchResult::IsValid() const
 {
-	bool bIsValid =
-		PoseIdx != INDEX_NONE &&
-		Database.IsValid();
-
-#if WITH_EDITOR
-	bIsValid = bIsValid &&
-		!Database->IsDerivedDataBuildPending() &&
-		Database->GetSearchIndexHash() == SearchIndexHash;
-#endif // WITH_EDITOR
-
-	return bIsValid;
+	return PoseIdx != INDEX_NONE && Database.IsValid();
 }
 
 void FSearchResult::Reset()
@@ -2764,10 +2579,6 @@ void FSearchResult::Reset()
 	Database = nullptr;
 	ComposedQuery.Reset();
 	AssetTime = 0.0f;
-
-#if WITH_EDITOR
-	SearchIndexHash = FIoHash::Zero;
-#endif // WITH_EDITOR
 }
 
 const FPoseSearchIndexAsset* FSearchResult::GetSearchIndexAsset(bool bMandatory) const
@@ -3574,7 +3385,8 @@ static FSamplingParam WrapOrClampSamplingParam(bool bCanWrap, float SamplingPara
 
 	Result.WrappedParam = SamplingParam;
 
-	if (bCanWrap)
+	const bool bIsSamplingParamExtentKindaSmall = SamplingParamExtent <= UE_KINDA_SMALL_NUMBER;
+	if (!bIsSamplingParamExtentKindaSmall && bCanWrap)
 	{
 		if (SamplingParam < 0.0f)
 		{
@@ -3595,14 +3407,14 @@ static FSamplingParam WrapOrClampSamplingParam(bool bCanWrap, float SamplingPara
 		}
 	}
 
-	float ParamClamped = FMath::Clamp(Result.WrappedParam, 0.0f, SamplingParamExtent);
+	const float ParamClamped = FMath::Clamp(Result.WrappedParam, 0.0f, SamplingParamExtent);
 	if (ParamClamped != Result.WrappedParam)
 	{
-		check(!bCanWrap);
+		check(bIsSamplingParamExtentKindaSmall || !bCanWrap);
 		Result.Extrapolation = Result.WrappedParam - ParamClamped;
 		Result.WrappedParam = ParamClamped;
 	}
-	
+
 	return Result;
 }
 
@@ -3818,13 +3630,10 @@ FAssetIndexer::FSampleInfo FAssetIndexer::GetSampleInfo(float SampleTime) const
 	float RootDistanceLast = 0.0f;
 	float RootDistanceInitial = 0.0f;
 
-	auto CanWrapTimeSamples = [](const IAssetSampler* Sampler) -> bool
-	{
-		return Sampler->IsLoopable();
-	};
+	check(IndexingContext.MainSampler);
 
-	float MainPlayLength = IndexingContext.MainSampler->GetPlayLength();
-	bool bMainCanWrap = CanWrapTimeSamples(IndexingContext.MainSampler);
+	const float MainPlayLength = IndexingContext.MainSampler->GetPlayLength();
+	const bool bMainCanWrap = IndexingContext.MainSampler->IsLoopable();
 
 	FSamplingParam SamplingParam;
 	if (!bMainCanWrap)
@@ -3834,8 +3643,8 @@ FAssetIndexer::FSampleInfo FAssetIndexer::GetSampleInfo(float SampleTime) const
 		{
 			const IAssetSampler* ClipSampler = IndexingContext.LeadInSampler;
 
-			bool bLeadInCanWrap = CanWrapTimeSamples(IndexingContext.LeadInSampler);
-			float LeadRelativeTime = SampleTime + ClipSampler->GetPlayLength();
+			const bool bLeadInCanWrap = IndexingContext.LeadInSampler->IsLoopable();
+			const float LeadRelativeTime = SampleTime + ClipSampler->GetPlayLength();
 			SamplingParam = WrapOrClampSamplingParam(bLeadInCanWrap, ClipSampler->GetPlayLength(), LeadRelativeTime);
 
 			Sample.Clip = IndexingContext.LeadInSampler;
@@ -3861,8 +3670,8 @@ FAssetIndexer::FSampleInfo FAssetIndexer::GetSampleInfo(float SampleTime) const
 		{
 			const IAssetSampler* ClipSampler = IndexingContext.FollowUpSampler;
 
-			bool bFollowUpCanWrap = CanWrapTimeSamples(IndexingContext.FollowUpSampler);
-			float FollowRelativeTime = SampleTime - MainPlayLength;
+			const bool bFollowUpCanWrap = IndexingContext.FollowUpSampler->IsLoopable();
+			const float FollowRelativeTime = SampleTime - MainPlayLength;
 			SamplingParam = WrapOrClampSamplingParam(bFollowUpCanWrap, ClipSampler->GetPlayLength(), FollowRelativeTime);
 
 			Sample.Clip = IndexingContext.FollowUpSampler;
@@ -4118,7 +3927,7 @@ void DrawSearchIndex(const FDebugDrawParams& DrawParams)
 #endif // ENABLE_DRAW_DEBUG
 }
 
-static Eigen::VectorXd ComputeChannelsDeviations(const FPoseSearchIndex* SearchIndex)
+static Eigen::VectorXd ComputeChannelsDeviations(const FPoseSearchIndex& SearchIndex, const UPoseSearchSchema* Schema)
 {
 	// This function performs a modified z-score normalization where features are normalized
 	// by mean absolute deviation rather than standard deviation. Both methods are preferable
@@ -4135,10 +3944,10 @@ static Eigen::VectorXd ComputeChannelsDeviations(const FPoseSearchIndex* SearchI
 
 	using namespace Eigen;
 
-	check(SearchIndex->IsValid());
+	check(Schema->IsValid());
 
-	const int32 NumPoses = SearchIndex->NumPoses;
-	const int32 NumDimensions = SearchIndex->Schema->SchemaCardinality;
+	const int32 NumPoses = SearchIndex.NumPoses;
+	const int32 NumDimensions = Schema->SchemaCardinality;
 
 	// Compute per Channel average distances
 	VectorXd MeanDeviations(NumDimensions);
@@ -4148,7 +3957,7 @@ static Eigen::VectorXd ComputeChannelsDeviations(const FPoseSearchIndex* SearchI
 	{
 		// Map input buffer
 		auto PoseMatrixSourceMap = RowMajorMatrixMapConst(
-			SearchIndex->Values.GetData(),
+			SearchIndex.Values.GetData(),
 			NumPoses,		// rows
 			NumDimensions	// cols
 		);
@@ -4164,9 +3973,9 @@ static Eigen::VectorXd ComputeChannelsDeviations(const FPoseSearchIndex* SearchI
 		VectorXd SampleMean = PoseMatrix.rowwise().mean();
 		PoseMatrix = PoseMatrix.colwise() - SampleMean;
 
-		for (int ChannelIdx = 0; ChannelIdx != SearchIndex->Schema->Channels.Num(); ++ChannelIdx)
+		for (int ChannelIdx = 0; ChannelIdx != Schema->Channels.Num(); ++ChannelIdx)
 		{
-			const UPoseSearchFeatureChannel* Channel = SearchIndex->Schema->Channels[ChannelIdx].Get();
+			const UPoseSearchFeatureChannel* Channel = Schema->Channels[ChannelIdx].Get();
 			Channel->ComputeMeanDeviations(PoseMatrix, MeanDeviations);
 		}
 	}
@@ -4174,18 +3983,18 @@ static Eigen::VectorXd ComputeChannelsDeviations(const FPoseSearchIndex* SearchI
 	return MeanDeviations;
 }
 
-static void PreprocessSearchIndexWeights(FPoseSearchIndex& SearchIndex, const UPoseSearchDatabase* Database)
+static void PreprocessSearchIndexWeights(FPoseSearchIndex& SearchIndex, const UPoseSearchSchema* Schema)
 {
-	const int32 NumDimensions = Database->Schema->SchemaCardinality;
+	const int32 NumDimensions = Schema->SchemaCardinality;
 	SearchIndex.WeightsSqrt.Init(1.f, NumDimensions);
 
-	for (int ChannelIdx = 0; ChannelIdx != Database->Schema->Channels.Num(); ++ChannelIdx)
+	for (int ChannelIdx = 0; ChannelIdx != Schema->Channels.Num(); ++ChannelIdx)
 	{
-		const UPoseSearchFeatureChannel* Channel = Database->Schema->Channels[ChannelIdx].Get();
+		const UPoseSearchFeatureChannel* Channel = Schema->Channels[ChannelIdx].Get();
 		Channel->FillWeights(SearchIndex.WeightsSqrt);
 	}
 
-	Eigen::VectorXd ChannelsMeanDeviations = ComputeChannelsDeviations(Database->GetSearchIndex());
+	Eigen::VectorXd ChannelsMeanDeviations = ComputeChannelsDeviations(SearchIndex, Schema);
 	TArray<float> Deviation;
 	Deviation.Init(1.f, NumDimensions);
 	for (int32 Dimension = 0; Dimension != NumDimensions; ++Dimension)
@@ -4194,7 +4003,7 @@ static void PreprocessSearchIndexWeights(FPoseSearchIndex& SearchIndex, const UP
 		Deviation[Dimension] = ChannelsMeanDeviation;
 	}
 
-	EPoseSearchDataPreprocessor DataPreprocessor = Database->GetSearchIndex()->Schema->DataPreprocessor;
+	EPoseSearchDataPreprocessor DataPreprocessor = Schema->DataPreprocessor;
 
 	if (DataPreprocessor == EPoseSearchDataPreprocessor::Normalize)
 	{
@@ -4228,13 +4037,10 @@ static void PreprocessSearchIndexWeights(FPoseSearchIndex& SearchIndex, const UP
 }
 
 // it calculates Mean, PCAValues, and PCAProjectionMatrix
-static void PreprocessSearchIndexPCAData(FPoseSearchIndex& SearchIndex, const UPoseSearchDatabase* Database)
+static void PreprocessSearchIndexPCAData(FPoseSearchIndex& SearchIndex, int32 NumDimensions, uint32 NumberOfPrincipalComponents, EPoseSearchMode PoseSearchMode)
 {
 	// binding SearchIndex.Values and SearchIndex.PCAValues Eigen row major matrix maps
-	check(Database && Database->Schema);
 	const int32 NumPoses = SearchIndex.NumPoses;
-	const uint32 NumberOfPrincipalComponents = Database->GetNumberOfPrincipalComponents();
-	const int32 NumDimensions = Database->Schema->SchemaCardinality;
 
 	SearchIndex.PCAValues.Reset();
 	SearchIndex.Mean.Reset();
@@ -4273,7 +4079,7 @@ static void PreprocessSearchIndexPCAData(FPoseSearchIndex& SearchIndex, const UP
 		// validating EigenSolver results
 		const ColMajorMatrix EigenVectors = EigenSolver.eigenvectors().real();
 
-		if (Database->PoseSearchMode == EPoseSearchMode::PCAKDTree_Validate && NumberOfPrincipalComponents == NumDimensions)
+		if (PoseSearchMode == EPoseSearchMode::PCAKDTree_Validate && NumberOfPrincipalComponents == NumDimensions)
 		{
 			const RowMajorVector ReciprocalWeightsSqrt = MapWeightsSqrt.cwiseInverse();
 			const RowMajorMatrix ProjectedValues = CenteredValues * EigenVectors;
@@ -4317,7 +4123,7 @@ static void PreprocessSearchIndexPCAData(FPoseSearchIndex& SearchIndex, const UP
 
 		MapPCAValues = CenteredValues * PCAProjectionMatrix;
 
-		if (Database->PoseSearchMode == EPoseSearchMode::PCAKDTree_Validate && NumberOfPrincipalComponents == NumDimensions)
+		if (PoseSearchMode == EPoseSearchMode::PCAKDTree_Validate && NumberOfPrincipalComponents == NumDimensions)
 		{
 			const RowMajorVector ReciprocalWeightsSqrt = MapWeightsSqrt.cwiseInverse();
 			for (Eigen::Index RowIndex = 0; RowIndex < MapValues.rows(); ++RowIndex)
@@ -4331,13 +4137,12 @@ static void PreprocessSearchIndexPCAData(FPoseSearchIndex& SearchIndex, const UP
 	}
 }
 
-static void PreprocessSearchIndexKDTree(FPoseSearchIndex& SearchIndex, const UPoseSearchDatabase* Database)
+static void PreprocessSearchIndexKDTree(FPoseSearchIndex& SearchIndex, int32 NumDimensions, uint32 NumberOfPrincipalComponents, EPoseSearchMode PoseSearchMode, int32 KDTreeMaxLeafSize, int32 KDTreeQueryNumNeighbors)
 {
 	const int32 NumPoses = SearchIndex.NumPoses;
-	const uint32 NumberOfPrincipalComponents = Database->GetNumberOfPrincipalComponents();
-	SearchIndex.KDTree.Construct(NumPoses, NumberOfPrincipalComponents, SearchIndex.PCAValues.GetData(), Database->KDTreeMaxLeafSize);
+	SearchIndex.KDTree.Construct(NumPoses, NumberOfPrincipalComponents, SearchIndex.PCAValues.GetData(), KDTreeMaxLeafSize);
 
-	if (Database->PoseSearchMode == EPoseSearchMode::PCAKDTree_Validate)
+	if (PoseSearchMode == EPoseSearchMode::PCAKDTree_Validate)
 	{
 		// testing the KDTree is returning the proper searches for all the points in pca space
 		int32 NumberOfFailingPoints = 0;
@@ -4345,9 +4150,9 @@ static void PreprocessSearchIndexKDTree(FPoseSearchIndex& SearchIndex, const UPo
 		{
 			TArray<size_t> ResultIndexes;
 			TArray<float> ResultDistanceSqr;
-			ResultIndexes.SetNum(Database->KDTreeQueryNumNeighbors + 1);
-			ResultDistanceSqr.SetNum(Database->KDTreeQueryNumNeighbors + 1);
-			FKDTree::KNNResultSet ResultSet(Database->KDTreeQueryNumNeighbors, ResultIndexes, ResultDistanceSqr);
+			ResultIndexes.SetNum(KDTreeQueryNumNeighbors + 1);
+			ResultDistanceSqr.SetNum(KDTreeQueryNumNeighbors + 1);
+			FKDTree::KNNResultSet ResultSet(KDTreeQueryNumNeighbors, ResultIndexes, ResultDistanceSqr);
 			SearchIndex.KDTree.FindNeighbors(ResultSet, &SearchIndex.PCAValues[PointIndex * NumberOfPrincipalComponents]);
 
 			size_t ResultIndex = 0;
@@ -4369,14 +4174,13 @@ static void PreprocessSearchIndexKDTree(FPoseSearchIndex& SearchIndex, const UPo
 
 		// testing the KDTree is returning the proper searches for all the original points transformed in pca space
 		NumberOfFailingPoints = 0;
-		const int32 NumDimensions = Database->Schema->SchemaCardinality;
 		for (size_t PointIndex = 0; PointIndex < NumPoses; ++PointIndex)
 		{
 			TArray<size_t> ResultIndexes;
 			TArray<float> ResultDistanceSqr;
-			ResultIndexes.SetNum(Database->KDTreeQueryNumNeighbors + 1);
-			ResultDistanceSqr.SetNum(Database->KDTreeQueryNumNeighbors + 1);
-			FKDTree::KNNResultSet ResultSet(Database->KDTreeQueryNumNeighbors, ResultIndexes, ResultDistanceSqr);
+			ResultIndexes.SetNum(KDTreeQueryNumNeighbors + 1);
+			ResultDistanceSqr.SetNum(KDTreeQueryNumNeighbors + 1);
+			FKDTree::KNNResultSet ResultSet(KDTreeQueryNumNeighbors, ResultIndexes, ResultDistanceSqr);
 
 			const RowMajorVectorMapConst MapValues(&SearchIndex.Values[PointIndex * NumDimensions], 1, NumDimensions);
 			const RowMajorVectorMapConst MapWeightsSqrt(SearchIndex.WeightsSqrt.GetData(), 1, NumDimensions);
@@ -4408,11 +4212,11 @@ static void PreprocessSearchIndexKDTree(FPoseSearchIndex& SearchIndex, const UPo
 	}
 }
 
-static void PreprocessSearchIndex(FPoseSearchIndex& SearchIndex, const UPoseSearchDatabase* Database)
+static void PreprocessSearchIndex(FPoseSearchIndex& SearchIndex, const UPoseSearchSchema* Schema, uint32 NumberOfPrincipalComponents, EPoseSearchMode PoseSearchMode, int32 KDTreeMaxLeafSize, int32 KDTreeQueryNumNeighbors)
 {
-	PreprocessSearchIndexWeights(SearchIndex, Database);
-	PreprocessSearchIndexPCAData(SearchIndex, Database);
-	PreprocessSearchIndexKDTree(SearchIndex, Database);
+	PreprocessSearchIndexWeights(SearchIndex, Schema);
+	PreprocessSearchIndexPCAData(SearchIndex, Schema->SchemaCardinality, NumberOfPrincipalComponents, PoseSearchMode);
+	PreprocessSearchIndexKDTree(SearchIndex, Schema->SchemaCardinality, NumberOfPrincipalComponents, PoseSearchMode, KDTreeMaxLeafSize, KDTreeQueryNumNeighbors);
 }
 
 bool BuildIndex(const UAnimSequence* Sequence, UPoseSearchSequenceMetaData* SequenceMetaData)
@@ -4465,7 +4269,6 @@ bool BuildIndex(const UAnimSequence* Sequence, UPoseSearchSequenceMetaData* Sequ
 
 	SequenceMetaData->SearchIndex.Values = Indexer.Output.FeatureVectorTable;
 	SequenceMetaData->SearchIndex.NumPoses = Indexer.Output.NumIndexedPoses;
-	SequenceMetaData->SearchIndex.Schema = SequenceMetaData->Schema;
 	SequenceMetaData->SearchIndex.Assets.Add(SearchIndexAsset);
 	SequenceMetaData->SearchIndex.PoseMetadata = Indexer.Output.PoseMetadata;
 
@@ -4482,63 +4285,56 @@ bool BuildIndex(const UAnimSequence* Sequence, UPoseSearchSequenceMetaData* Sequ
 
 struct FDatabaseIndexingContext
 {
-	const UPoseSearchDatabase* Database = nullptr;
 	FPoseSearchIndex* SearchIndex = nullptr;
 
 	FAssetSamplingContext SamplingContext;
 	TArray<FSequenceSampler> SequenceSamplers;
-	TMap<const UAnimSequence*, int32> SequenceSamplerMap;
 	TArray<FBlendSpaceSampler> BlendSpaceSamplers;
-	TMap<TPair<const UBlendSpace*, FVector>, int32> BlendSpaceSamplerMap;
 
 	TArray<FAssetIndexer> Indexers;
 
-	void PrepareSamplers(const FBoneContainer& BoneContainer);
-	void PrepareIndexers(const FBoneContainer& BoneContainer);
+	void Prepare(const UPoseSearchSchema* Schema, const FPoseSearchExtrapolationParameters& ExtrapolationParameters, TConstArrayView<FPoseSearchDatabaseSequence> Sequences, TConstArrayView<FPoseSearchDatabaseBlendSpace> BlendSpaces);
 	bool IndexAssets();
 	void JoinIndex();
 	float CalculateMinCostAddend() const;
 };
 
-void FDatabaseIndexingContext::PrepareSamplers(const FBoneContainer& BoneContainer)
+void FDatabaseIndexingContext::Prepare(const UPoseSearchSchema* Schema, const FPoseSearchExtrapolationParameters& ExtrapolationParameters, TConstArrayView<FPoseSearchDatabaseSequence> Sequences, TConstArrayView<FPoseSearchDatabaseBlendSpace> BlendSpaces)
 {
-	SamplingContext.Init(Database->Schema->MirrorDataTable, BoneContainer);
+	check(Schema);
+
+	FBoneContainer BoneContainer;
+	BoneContainer.InitializeTo(Schema->BoneIndicesWithParents, FCurveEvaluationOption(false), *Schema->Skeleton);
+
+	TMap<const UAnimSequence*, int32> SequenceSamplerMap;
+	TMap<TPair<const UBlendSpace*, FVector>, int32> BlendSpaceSamplerMap;
+
+	SamplingContext.Init(Schema->MirrorDataTable, BoneContainer);
 
 	// Prepare samplers for all sequences
-	auto AddSequenceSampler = [&](const UAnimSequence* Sequence, bool bLoopable)
+	auto AddSequenceSampler = [&](const UAnimSequence* Sequence)
 	{
-		if (!SequenceSamplerMap.Contains(Sequence))
+		if (Sequence && !SequenceSamplerMap.Contains(Sequence))
 		{
 			int32 SequenceSamplerIdx = SequenceSamplers.AddDefaulted();
 			SequenceSamplerMap.Add(Sequence, SequenceSamplerIdx);
 
 			FSequenceSampler::FInput Input;
-			Input.ExtrapolationParameters = Database->ExtrapolationParameters;
+			Input.ExtrapolationParameters = ExtrapolationParameters;
 			Input.Sequence = Sequence;
 			SequenceSamplers[SequenceSamplerIdx].Init(Input);
 		}
 	};
 
-	for (const FPoseSearchDatabaseSequence& DbSequence : Database->Sequences)
+	for (const FPoseSearchDatabaseSequence& DbSequence : Sequences)
 	{
-		if (DbSequence.Sequence)
-		{
-			AddSequenceSampler(DbSequence.Sequence, DbSequence.Sequence->bLoop);
-		}
-
-		if (DbSequence.LeadInSequence)
-		{
-			AddSequenceSampler(DbSequence.LeadInSequence, DbSequence.LeadInSequence->bLoop);
-		}
-
-		if (DbSequence.FollowUpSequence)
-		{
-			AddSequenceSampler(DbSequence.FollowUpSequence, DbSequence.FollowUpSequence->bLoop);
-		}
+		AddSequenceSampler(DbSequence.Sequence);
+		AddSequenceSampler(DbSequence.LeadInSequence);
+		AddSequenceSampler(DbSequence.FollowUpSequence);
 	}
 		
 	// Prepare samplers for all blend spaces
-	for (const FPoseSearchDatabaseBlendSpace& DbBlendSpace : Database->BlendSpaces)
+	for (const FPoseSearchDatabaseBlendSpace& DbBlendSpace : BlendSpaces)
 	{
 		if (DbBlendSpace.BlendSpace)
 		{
@@ -4558,7 +4354,7 @@ void FDatabaseIndexingContext::PrepareSamplers(const FBoneContainer& BoneContain
 
 						FBlendSpaceSampler::FInput Input;
 						Input.BoneContainer = BoneContainer;
-						Input.ExtrapolationParameters = Database->ExtrapolationParameters;
+						Input.ExtrapolationParameters = ExtrapolationParameters;
 						Input.BlendSpace = DbBlendSpace.BlendSpace;
 						Input.BlendParameters = BlendParameters;
 
@@ -4582,11 +4378,9 @@ void FDatabaseIndexingContext::PrepareSamplers(const FBoneContainer& BoneContain
 	}
 
 	ParallelFor(AssetSampler.Num(), [AssetSampler](int32 SamplerIdx) { AssetSampler[SamplerIdx]->Process(); }, ParallelForFlags);
-}
 
-void FDatabaseIndexingContext::PrepareIndexers(const FBoneContainer& BoneContainer)
-{
-	Indexers.Reserve(Database->GetSearchIndex()->Assets.Num());
+	// prepare indexers
+	Indexers.Reserve(SearchIndex->Assets.Num());
 
 	auto GetSequenceSampler = [&](const UAnimSequence* Sequence) -> const FSequenceSampler*
 	{
@@ -4598,7 +4392,7 @@ void FDatabaseIndexingContext::PrepareIndexers(const FBoneContainer& BoneContain
 		return BlendSpace ? &BlendSpaceSamplers[BlendSpaceSamplerMap[{BlendSpace, BlendParameters}]] : nullptr;
 	};
 
-	Indexers.Reserve(Database->GetSearchIndex()->Assets.Num());
+	Indexers.Reserve(SearchIndex->Assets.Num());
 
 	for (int32 AssetIdx = 0; AssetIdx != SearchIndex->Assets.Num(); ++AssetIdx)
 	{
@@ -4606,13 +4400,13 @@ void FDatabaseIndexingContext::PrepareIndexers(const FBoneContainer& BoneContain
 
 		FAssetIndexingContext IndexerContext;
 		IndexerContext.SamplingContext = &SamplingContext;
-		IndexerContext.Schema = Database->Schema;
+		IndexerContext.Schema = Schema;
 		IndexerContext.RequestedSamplingRange = SearchIndexAsset.SamplingInterval;
 		IndexerContext.bMirrored = SearchIndexAsset.bMirrored;
 
 		if (SearchIndexAsset.Type == ESearchIndexAssetType::Sequence)
 		{
-			const FPoseSearchDatabaseSequence& DbSequence = Database->GetSequenceSourceAsset(SearchIndexAsset);
+			const FPoseSearchDatabaseSequence& DbSequence = Sequences[SearchIndexAsset.SourceAssetIdx];
 			const float SequenceLength = DbSequence.Sequence->GetPlayLength();
 			IndexerContext.MainSampler = GetSequenceSampler(DbSequence.Sequence);
 			IndexerContext.LeadInSampler = SearchIndexAsset.SamplingInterval.Min == 0.0f ? GetSequenceSampler(DbSequence.LeadInSequence) : nullptr;
@@ -4620,7 +4414,7 @@ void FDatabaseIndexingContext::PrepareIndexers(const FBoneContainer& BoneContain
 		}
 		else if (SearchIndexAsset.Type == ESearchIndexAssetType::BlendSpace)
 		{
-			const FPoseSearchDatabaseBlendSpace& DbBlendSpace = Database->GetBlendSpaceSourceAsset(SearchIndexAsset);
+			const FPoseSearchDatabaseBlendSpace& DbBlendSpace = BlendSpaces[SearchIndexAsset.SourceAssetIdx];
 			IndexerContext.MainSampler = GetBlendSpaceSampler(DbBlendSpace.BlendSpace, SearchIndexAsset.BlendParameters);
 		}
 		else
@@ -4698,47 +4492,73 @@ void FDatabaseIndexingContext::JoinIndex()
 	}
 	
 	SearchIndex->NumPoses = TotalPoses;
-	SearchIndex->Schema = Database->Schema;
 	SearchIndex->MinCostAddend = CalculateMinCostAddend();
-
-	check(TotalFloats == TotalPoses * Database->Schema->SchemaCardinality);
 }
 
-bool BuildIndex(const UPoseSearchDatabase* Database, FPoseSearchIndex& OutSearchIndex)
+bool BuildIndex(const UPoseSearchDatabase& Database, FPoseSearchIndex& OutSearchIndex
+#if WITH_EDITOR
+	, UE::DerivedData::IRequestOwner& Owner
+#endif
+)
 {
-	check(Database);
-
-	bool bSuccess = true;
-
-	if (!Database->IsValidForIndexing())
+	if (!Database.IsValidForIndexing())
 	{
-		UE_LOG(LogPoseSearch, Warning, TEXT("Database '%f' is invalid for indexing"), *Database->GetName());
+		OutSearchIndex.Reset();
 		return false;
 	}
 
-	UE_LOG(LogPoseSearch, Log, TEXT("Building PoseSearch index for database %s"), *GetNameSafe(Database));
-
-	OutSearchIndex.Schema = Database->Schema;
-
-	if (!Database->TryInitSearchIndexAssets(OutSearchIndex))
+#if WITH_EDITOR
+	if (Owner.IsCanceled())
 	{
+		OutSearchIndex.Reset();
 		return false;
 	}
+#endif
+
+	OutSearchIndex.InitSearchIndexAssets(Database.Sequences, Database.BlendSpaces, Database.ExcludeFromDatabaseParameters);
 
 	FDatabaseIndexingContext DbIndexingContext;
-	DbIndexingContext.Database = Database;
 	DbIndexingContext.SearchIndex = &OutSearchIndex;
+	DbIndexingContext.Prepare(Database.Schema, Database.ExtrapolationParameters, Database.Sequences, Database.BlendSpaces);
+	
+#if WITH_EDITOR
+	if (Owner.IsCanceled())
+	{
+		OutSearchIndex.Reset();
+		return false;
+	}
+#endif
 
-	FBoneContainer BoneContainer;
-	BoneContainer.InitializeTo(Database->Schema->BoneIndicesWithParents, FCurveEvaluationOption(false), *Database->Schema->Skeleton);
+	const bool bSuccess = DbIndexingContext.IndexAssets();
+
+#if WITH_EDITOR
+	if (Owner.IsCanceled())
+	{
+		OutSearchIndex.Reset();
+		return false;
+	}
+#endif
 
-	DbIndexingContext.PrepareSamplers(BoneContainer);
-	DbIndexingContext.PrepareIndexers(BoneContainer);
-	bSuccess &= DbIndexingContext.IndexAssets();
 	DbIndexingContext.JoinIndex();
 
-	PreprocessSearchIndex(OutSearchIndex, Database);
-	
+#if WITH_EDITOR
+	if (Owner.IsCanceled())
+	{
+		OutSearchIndex.Reset();
+		return false;
+	}
+#endif
+
+	PreprocessSearchIndex(OutSearchIndex, Database.Schema, Database.GetNumberOfPrincipalComponents(), Database.PoseSearchMode, Database.KDTreeMaxLeafSize, Database.KDTreeQueryNumNeighbors);
+
+#if WITH_EDITOR
+	if (Owner.IsCanceled())
+	{
+		OutSearchIndex.Reset();
+		return false;
+	}
+#endif
+
 	return bSuccess;
 }
 
@@ -4753,20 +4573,11 @@ public: // IModuleInterface
 
 public: // IPoseSearchProvider
 	virtual UE::Anim::IPoseSearchProvider::FSearchResult Search(const FAnimationBaseContext& GraphContext, const UAnimSequenceBase* Sequence) override;
-
-private:
-#if WITH_EDITOR
-	void OnObjectSaved(UObject* SavedObject, FObjectPreSaveContext SaveContext);
-#endif // WITH_EDITOR
 };
 
 void FModule::StartupModule()
 {
 	IModularFeatures::Get().RegisterModularFeature(UE::Anim::IPoseSearchProvider::ModularFeatureName, this);
-
-#if WITH_EDITOR
-	FCoreUObjectDelegates::OnObjectPreSave.AddRaw(this, &FModule::OnObjectSaved);
-#endif // WITH_EDITOR
 }
 
 void FModule::ShutdownModule()
@@ -4807,95 +4618,6 @@ UE::Anim::IPoseSearchProvider::FSearchResult FModule::Search(const FAnimationBas
 	return ProviderResult;
 }
 
-#if WITH_EDITOR
-
-void GetPoseSearchDatabaseAssetDataList(TArray<FAssetData>& OutPoseSearchDatabaseAssetDataList)
-{
-	FAssetRegistryModule& AssetRegistryModule =
-		FModuleManager::LoadModuleChecked<FAssetRegistryModule>("AssetRegistry");
-	FARFilter Filter;
-	Filter.bRecursiveClasses = true;
-	Filter.ClassPaths.Add(UPoseSearchDatabase::StaticClass()->GetClassPathName());
-
-	OutPoseSearchDatabaseAssetDataList.Reset();
-	AssetRegistryModule.Get().GetAssets(Filter, OutPoseSearchDatabaseAssetDataList);
-}
-
-void ForEachPoseSearchDatabase(bool bLoadAssets, TFunctionRef<void(UPoseSearchDatabase&)> InFunction)
-{
-	TArray<FAssetData> PoseSearchDatabaseAssetDataList;
-	GetPoseSearchDatabaseAssetDataList(PoseSearchDatabaseAssetDataList);
-	for (const auto& PoseSearchDbAssetData : PoseSearchDatabaseAssetDataList)
-	{
-		if (UPoseSearchDatabase* PoseSearchDb = 
-			Cast<UPoseSearchDatabase>(PoseSearchDbAssetData.FastGetAsset(bLoadAssets)))
-		{
-			InFunction(*PoseSearchDb);
-		}
-	}
-}
-
-void FModule::OnObjectSaved(UObject* SavedObject, FObjectPreSaveContext SaveContext)
-{
-	if (UAnimSequence* SavedSequence = Cast<UAnimSequence>(SavedObject))
-	{
-		ForEachPoseSearchDatabase(false, [SavedSequence](UPoseSearchDatabase& PoseSearchDb)
-		{
-			bool bSequenceFound =
-				PoseSearchDb.Sequences.ContainsByPredicate([SavedSequence](FPoseSearchDatabaseSequence& DbSequence)
-			{
-				bool bIsMatch =
-					SavedSequence == DbSequence.Sequence ||
-					SavedSequence == DbSequence.LeadInSequence ||
-					SavedSequence == DbSequence.FollowUpSequence;
-				return bIsMatch;
-			});
-
-			if (bSequenceFound)
-			{
-				PoseSearchDb.BeginCacheDerivedData();
-			}
-		});
-	}
-	else if (UBlendSpace* SavedBlendSpace = Cast<UBlendSpace>(SavedObject))
-	{
-		ForEachPoseSearchDatabase(false, [SavedBlendSpace](UPoseSearchDatabase& PoseSearchDb)
-		{
-			bool bBlendSpaceFound = PoseSearchDb.BlendSpaces.ContainsByPredicate(
-					[SavedBlendSpace](FPoseSearchDatabaseBlendSpace& DbBlendSpace)
-			{
-				return SavedBlendSpace == DbBlendSpace.BlendSpace;
-			});
-
-			if (bBlendSpaceFound)
-			{
-				PoseSearchDb.BeginCacheDerivedData();
-			}
-		});
-	}
-	else if (UPoseSearchSchema* SavedSchema = Cast<UPoseSearchSchema>(SavedObject))
-	{
-		ForEachPoseSearchDatabase(false, [SavedSchema](UPoseSearchDatabase& PoseSearchDb)
-		{
-			if (PoseSearchDb.Schema == SavedSchema)
-			{
-				PoseSearchDb.BeginCacheDerivedData();
-			}
-		});
-	}
-	else if (USkeleton* SavedSkeleton = Cast<USkeleton>(SavedObject))
-	{
-		ForEachPoseSearchDatabase(false, [SavedSkeleton](UPoseSearchDatabase& PoseSearchDb)
-		{
-			if (PoseSearchDb.Schema && PoseSearchDb.Schema->Skeleton == SavedSkeleton)
-			{
-				PoseSearchDb.BeginCacheDerivedData();
-			}
-		});
-	}
-}
-#endif // WITH_EDITOR
-
 } // namespace UE::PoseSearch
 
 #undef LOCTEXT_NAMESPACE
diff --git a/Engine/Plugins/Experimental/Animation/PoseSearch/Source/Runtime/Private/PoseSearchDerivedData.cpp b/Engine/Plugins/Experimental/Animation/PoseSearch/Source/Runtime/Private/PoseSearchDerivedData.cpp
--- a/Engine/Plugins/Experimental/Animation/PoseSearch/Source/Runtime/Private/PoseSearchDerivedData.cpp
+++ b/Engine/Plugins/Experimental/Animation/PoseSearch/Source/Runtime/Private/PoseSearchDerivedData.cpp
@@ -1,316 +1,474 @@
 // Copyright Epic Games, Inc. All Rights Reserved.
 
-#include "PoseSearchDerivedData.h"
-#include "PoseSearch/PoseSearchAnimNotifies.h"
-
-#include "Animation/AnimSequence.h"
-#include "Animation/BlendSpace.h"
-
-#include "UObject/Class.h"
+#include "PoseSearch/PoseSearchDerivedData.h"
 
 #if WITH_EDITOR
-#include "Serialization/BulkDataRegistry.h"
+#include "Animation/BlendSpace.h"
+#include "AssetRegistry/ARFilter.h"
+#include "AssetRegistry/AssetRegistryModule.h"
+#include "DerivedDataCache.h"
+#include "DerivedDataRequestOwner.h"
+#include "Misc/CoreDelegates.h"
 #include "PoseSearch/PoseSearchDerivedDataKey.h"
-#endif // WITH_EDITOR
+#include "ProfilingDebugging/CookStats.h"
+#include "Serialization/BulkDataRegistry.h"
+#include "UObject/NoExportTypes.h"
 
-#if UE_BUILD_DEBUG && WITH_EDITOR
-#define UE_POSE_SEARCH_DERIVED_DATA_LOGGING 1
-#endif
+namespace UE::PoseSearch
+{
+	static const UE::DerivedData::FValueId Id(UE::DerivedData::FValueId::FromName("Data"));
+	static const UE::DerivedData::FCacheBucket Bucket("PoseSearchDatabase");
 
-#ifndef UE_POSE_SEARCH_DERIVED_DATA_LOGGING
-#define UE_POSE_SEARCH_DERIVED_DATA_LOGGING 0
+#if ENABLE_COOK_STATS
+	static FCookStats::FDDCResourceUsageStats UsageStats;
+	static FCookStatsManager::FAutoRegisterCallback RegisterCookStats([](FCookStatsManager::AddStatFuncRef AddStat)
+		{
+			UsageStats.LogStats(AddStat, TEXT("MotionMatching.Usage"), TEXT(""));
+		});
 #endif
 
+	struct FPoseSearchDatabaseAsyncCacheTask
+	{
+		FPoseSearchDatabaseAsyncCacheTask(UPoseSearchDatabase& InDatabase);
+		~FPoseSearchDatabaseAsyncCacheTask();
 
-#if WITH_EDITOR
-namespace UE::PoseSearch
-{
-	const UE::DerivedData::FValueId FPoseSearchDatabaseAsyncCacheTask::Id = UE::DerivedData::FValueId::FromName("Data");
-	const UE::DerivedData::FCacheBucket FPoseSearchDatabaseAsyncCacheTask::Bucket("PoseSearchDatabase");
-}
-#endif // WITH_EDITOR
+		void Cancel();
+		void Wait();
+		bool Poll() const;
 
-#if WITH_EDITOR
+		UPoseSearchDatabase& GetDatabase() { return Database; }
+		const UPoseSearchDatabase& GetDatabase() const { return Database; }
+		static FIoHash CreateKey(UPoseSearchDatabase& Database);
 
-void FPoseSearchDatabaseDerivedData::Cache(UPoseSearchDatabase& Database, bool bForceRebuild)
-{
-	check(IsInGameThread());
+		void MoveSearchIndexToDatabase();
 
-	CancelCache();
+	private:
+		FPoseSearchDatabaseAsyncCacheTask(const FPoseSearchDatabaseAsyncCacheTask& Other) = delete;
+		FPoseSearchDatabaseAsyncCacheTask(FPoseSearchDatabaseAsyncCacheTask&& Other) = delete;
+		FPoseSearchDatabaseAsyncCacheTask& operator=(const FPoseSearchDatabaseAsyncCacheTask& Other) = delete;
+		FPoseSearchDatabaseAsyncCacheTask& operator=(FPoseSearchDatabaseAsyncCacheTask&& Other) = delete;
 
-#if UE_POSE_SEARCH_DERIVED_DATA_LOGGING
-	UE_LOG(LogPoseSearch, Log, TEXT("%s: Cache"), *LexToString(PendingDerivedDataKey));
-#endif // UE_POSE_SEARCH_DERIVED_DATA_LOGGING
+		void OnGetComplete(UE::DerivedData::FCacheGetResponse&& Response);
 
-	if (Database.IsValidForIndexing())
-	{
-		CreateDatabaseBuildTask(Database, bForceRebuild);
-	}
-	else
-	{
-		SearchIndex.Reset();
-		SearchIndex.Schema = Database.Schema;
-		DerivedDataKey = { UE::DerivedData::FCacheBucket(), FIoHash::Zero };
-		PendingDerivedDataKey = FIoHash::Zero;
-	}
-}
+		UPoseSearchDatabase& Database;
+		FPoseSearchIndex SearchIndex;
+		UE::DerivedData::FRequestOwner Owner;
+	};
 
-void FPoseSearchDatabaseDerivedData::CancelCache()
-{
-	check(IsInGameThread());
+	class FPoseSearchDatabaseAsyncCacheTasks : public TArray<TUniquePtr<FPoseSearchDatabaseAsyncCacheTask>> {};
 
-	if (AsyncTask)
+	FPoseSearchDatabaseAsyncCacheTask::FPoseSearchDatabaseAsyncCacheTask(UPoseSearchDatabase& InDatabase)
+		: Database(InDatabase)
+		, SearchIndex()
+		, Owner(UE::DerivedData::EPriority::Normal)
 	{
-#if UE_POSE_SEARCH_DERIVED_DATA_LOGGING
-		UE_LOG(LogPoseSearch, Log, TEXT("%s: CancelCache"), *LexToString(PendingDerivedDataKey));
-#endif // UE_POSE_SEARCH_DERIVED_DATA_LOGGING
+		using namespace UE::DerivedData;
+		const FIoHash PendingDerivedDataKey = CreateKey(InDatabase);
 
-		AsyncTask->Cancel();
-	}
+		InDatabase.NotifyDerivedDataBuildStarted();
 
-	FinishCache();
-}
+		UE_LOG(LogPoseSearch, Log, TEXT("%s - %s BeginCache"), *LexToString(PendingDerivedDataKey), *Database.GetName());
 
-void FPoseSearchDatabaseDerivedData::FinishCache()
-{
-	check(IsInGameThread());
+		TArray<FCacheGetRequest> CacheRequests;
+		const FCacheKey CacheKey{ Bucket, PendingDerivedDataKey };
+		CacheRequests.Add({ { Database.GetPathName() }, CacheKey, ECachePolicy::Default });
+		GetCache().Get(CacheRequests, Owner, [this](FCacheGetResponse&& Response)
+			{
+				OnGetComplete(MoveTemp(Response));
+			});
+	}
 
-	if (AsyncTask)
+	FPoseSearchDatabaseAsyncCacheTask::~FPoseSearchDatabaseAsyncCacheTask()
 	{
-#if UE_POSE_SEARCH_DERIVED_DATA_LOGGING
-		UE_LOG(LogPoseSearch, Log, TEXT("%s: FinishCache"), *LexToString(PendingDerivedDataKey));
-#endif // UE_POSE_SEARCH_DERIVED_DATA_LOGGING
-
-		AsyncTask->Wait();
-		delete AsyncTask;
-		AsyncTask = nullptr;
+		// @todo: add InDatabase.NotifyDerivedDataBuildEnded();
 	}
-}
-
-void FPoseSearchDatabaseDerivedData::CreateDatabaseBuildTask(UPoseSearchDatabase& Database, bool bForceRebuild)
-{
-	check(IsInGameThread());
 
-	AsyncTask = new UE::PoseSearch::FPoseSearchDatabaseAsyncCacheTask(Database, *this, bForceRebuild);
-}
+	void FPoseSearchDatabaseAsyncCacheTask::MoveSearchIndexToDatabase()
+	{
+		check(IsInGameThread());
+		check(Poll());
+		if (!Owner.IsCanceled())
+		{
+			// @todo: implement FPoseSearchIndex move ctor and assignment operator and use a MoveTemp(SearchIndex) here
+			Database.PoseSearchIndex = SearchIndex;
+		}
+	}
 
-#endif // WITH_EDITOR
+	void FPoseSearchDatabaseAsyncCacheTask::Cancel()
+	{
+		Owner.Cancel();
+	}
 
-namespace UE::PoseSearch
-{
+	void FPoseSearchDatabaseAsyncCacheTask::Wait()
+	{
+		Owner.Wait();
+	}
 
-#if WITH_EDITOR
+	bool FPoseSearchDatabaseAsyncCacheTask::Poll() const
+	{
+		return Owner.Poll();
+	}
 
-	FPoseSearchDatabaseAsyncCacheTask::FPoseSearchDatabaseAsyncCacheTask(
-		UPoseSearchDatabase& InDatabase,
-		FPoseSearchDatabaseDerivedData& InDerivedData,
-		bool bForceRebuild)
-		: Owner(UE::DerivedData::EPriority::Normal)
-		, DerivedData(InDerivedData)
-		, Database(InDatabase)
+	void FPoseSearchDatabaseAsyncCacheTask::OnGetComplete(UE::DerivedData::FCacheGetResponse&& Response)
 	{
 		using namespace UE::DerivedData;
 
-		FIoHash DerivedDataKey = CreateKey(Database);
-		DerivedData.PendingDerivedDataKey = DerivedDataKey;
+		const FCacheKey Key = Response.Record.GetKey();
+		if (Response.Status == EStatus::Ok)
+		{
+			COOK_STAT(auto Timer = UsageStats.TimeAsyncWait());
 
-		InDatabase.NotifyDerivedDataBuildStarted();
+			// we found the cached data associated to the PendingDerivedDataKey: we'll deserialized into SearchIndex
+			SearchIndex.Reset();
+			FSharedBuffer RawData = Response.Record.GetValue(Id).GetData().Decompress();
+			FMemoryReaderView Reader(RawData);
+			Reader << SearchIndex;
 
-		if (bForceRebuild)
+			UE_LOG(LogPoseSearch, Log, TEXT("%s - %s BuildIndex From Cache"), *LexToString(Key.Hash), *Database.GetName());
+
+			COOK_STAT(Timer.AddHit(RawData.GetSize()));
+		}
+		else if (Response.Status == EStatus::Error)
 		{
-			// when the build is forced, the derived data key is zeroed so the comparison with the pending key fails, 
-			// informing other systems that data is being rebuilt
-			DerivedData.DerivedDataKey.Hash = FIoHash::Zero;
-			BuildAndWrite({ Bucket, DerivedDataKey });
+			// we didn't find the cached data associated to the PendingDerivedDataKey: we'll BuildIndex to update SearchIndex and "Put" the data over the DDC
+			Owner.LaunchTask(TEXT("PoseSearchDatabaseBuild"), [this, Key]
+				{
+					COOK_STAT(auto Timer = UsageStats.TimeSyncWork());
+
+					int32 BytesProcessed = 0;
+					if (BuildIndex(Database, SearchIndex, Owner))
+					{
+						UE_LOG(LogPoseSearch, Log, TEXT("%s - %s BuildIndex Succeeded"), *LexToString(Key.Hash), *Database.GetName());
+
+						TArray<uint8> RawBytes;
+						FMemoryWriter Writer(RawBytes);
+						Writer << SearchIndex;
+						FSharedBuffer RawData = MakeSharedBufferFromArray(MoveTemp(RawBytes));
+						BytesProcessed = RawData.GetSize();
+
+						FCacheRecordBuilder Builder(Key);
+						Builder.AddValue(Id, RawData);
+						GetCache().Put({ { { Database.GetPathName() }, Builder.Build() } }, Owner, [this, Key](FCachePutResponse&& Response)
+							{
+								if (Response.Status == EStatus::Error)
+								{
+									UE_LOG(LogPoseSearch, Log, TEXT("%s - %s Failed to store DDC"), *LexToString(Key.Hash), *Database.GetName());
+								}
+							});
+					}
+					else
+					{
+						if (Owner.IsCanceled())
+						{
+							UE_LOG(LogPoseSearch, Log, TEXT("%s - %s BuildIndex Cancelled"), *LexToString(Key.Hash), *Database.GetName());
+						}
+						else
+						{
+							UE_LOG(LogPoseSearch, Log, TEXT("%s - %s BuildIndex Failed"), *LexToString(Key.Hash), *Database.GetName());
+						}
+						SearchIndex.Reset();
+					}
+
+					COOK_STAT(Timer.AddMiss(BytesProcessed));
+				});
 		}
-		else
+		else if(Response.Status == EStatus::Canceled)
 		{
-			BeginCache();
+			SearchIndex.Reset();
+			UE_LOG(LogPoseSearch, Log, TEXT("%s - %s BuildIndex Cancelled"), *LexToString(Key.Hash), *Database.GetName());
 		}
 	}
 
-	bool FPoseSearchDatabaseAsyncCacheTask::Cancel()
+	FIoHash FPoseSearchDatabaseAsyncCacheTask::CreateKey(UPoseSearchDatabase& Database)
 	{
-		Owner.Cancel();
-		return true;
+		using UE::PoseSearch::FDerivedDataKeyBuilder;
+
+		FDerivedDataKeyBuilder KeyBuilder;
+		FGuid VersionGuid = FDevSystemGuids::GetSystemGuid(FDevSystemGuids::Get().POSESEARCHDB_DERIVEDDATA_VER);
+		KeyBuilder << VersionGuid;
+		Database.BuildDerivedDataKey(KeyBuilder);
+
+		// Stores a BLAKE3-160 hash, taken from the first 20 bytes of a BLAKE3-256 hash
+		FIoHash IoHash(KeyBuilder.Finalize());
+		return IoHash;
 	}
 
-	void FPoseSearchDatabaseAsyncCacheTask::Wait()
+	//////////////////////////////////////////////////////////////////////////
+	// FAsyncPoseSearchDatabasesManagement
+	FAsyncPoseSearchDatabasesManagement& FAsyncPoseSearchDatabasesManagement::Get()
 	{
-		Owner.Wait();
+		static FAsyncPoseSearchDatabasesManagement SingletonInstance;
+		return SingletonInstance;
 	}
 
-	bool FPoseSearchDatabaseAsyncCacheTask::Poll() const
+	FAsyncPoseSearchDatabasesManagement::FAsyncPoseSearchDatabasesManagement()
+		: Tasks(*(new FPoseSearchDatabaseAsyncCacheTasks()))
 	{
-		return Owner.Poll();
+		OnObjectPreSaveHandle = FCoreUObjectDelegates::OnObjectPreSave.AddRaw(this, &FAsyncPoseSearchDatabasesManagement::OnObjectPreSave); 
+		OnPreObjectPropertyChangedHandle = FCoreUObjectDelegates::OnPreObjectPropertyChanged.AddRaw(this, &FAsyncPoseSearchDatabasesManagement::OnPreObjectPropertyChanged);
+		OnObjectPropertyChangedHandle = FCoreUObjectDelegates::OnObjectPropertyChanged.AddRaw(this, &FAsyncPoseSearchDatabasesManagement::OnObjectPropertyChanged);
+
+		FCoreDelegates::OnPreExit.AddRaw(this, &FAsyncPoseSearchDatabasesManagement::Shutdown);
 	}
 
+	FAsyncPoseSearchDatabasesManagement::~FAsyncPoseSearchDatabasesManagement()
+	{
+		FCoreDelegates::OnPreExit.RemoveAll(this);
+		Shutdown();
+
+		delete &Tasks;
+	}
 
-	void FPoseSearchDatabaseAsyncCacheTask::BeginCache()
+	void GetPoseSearchDatabaseAssetDataList(TArray<FAssetData>& OutPoseSearchDatabaseAssetDataList)
 	{
-		using namespace UE::DerivedData;
+		FAssetRegistryModule& AssetRegistryModule = FModuleManager::LoadModuleChecked<FAssetRegistryModule>("AssetRegistry");
+		FARFilter Filter;
+		Filter.bRecursiveClasses = true;
+		Filter.ClassPaths.Add(UPoseSearchDatabase::StaticClass()->GetClassPathName());
 
-#if UE_POSE_SEARCH_DERIVED_DATA_LOGGING
-		UE_LOG(LogPoseSearch, Log, TEXT("%s: BeginCache - '%s'"), *LexToString(DerivedData.PendingDerivedDataKey), *Database.GetName());
-#endif // UE_POSE_SEARCH_DERIVED_DATA_LOGGING
+		OutPoseSearchDatabaseAssetDataList.Reset();
+		AssetRegistryModule.Get().GetAssets(Filter, OutPoseSearchDatabaseAssetDataList);
+	}
 
-		TArray<FCacheGetRequest> CacheRequests;
-		ECachePolicy CachePolicy = ECachePolicy::Default;
-		const FCacheKey CacheKey{ Bucket, DerivedData.PendingDerivedDataKey };
-		CacheRequests.Add({ {Database.GetPathName()}, CacheKey, CachePolicy });
-		GetCache().Get(CacheRequests, Owner, [this](FCacheGetResponse&& Response)
-		{
-			OnGetComplete(MoveTemp(Response));
-		});
+	FPoseSearchDatabaseAsyncCacheTask& FAsyncPoseSearchDatabasesManagement::GetTask(int32 TaskIndex)
+	{
+		return *Tasks[TaskIndex].Get();
 	}
 
-	void FPoseSearchDatabaseAsyncCacheTask::OnGetComplete(UE::DerivedData::FCacheGetResponse&& Response)
+	const FPoseSearchDatabaseAsyncCacheTask& FAsyncPoseSearchDatabasesManagement::GetTask(int32 TaskIndex) const
 	{
-		using namespace UE::DerivedData;
+		return *Tasks[TaskIndex].Get();
+	}
 
-		if (Response.Status == EStatus::Ok)
+	void FAsyncPoseSearchDatabasesManagement::RemoveTask(int32 TaskIndex)
+	{
+		GetTask(TaskIndex).MoveSearchIndexToDatabase();
+		Tasks.RemoveAtSwap(TaskIndex, 1, false);
+	}
+
+	void FAsyncPoseSearchDatabasesManagement::ForEachPoseSearchDatabase(bool bUseTasksDatabases, TFunctionRef<void(UPoseSearchDatabase&)> InFunction)
+	{
+		if (bUseTasksDatabases)
 		{
-			BuildIndexFromCacheRecord(MoveTemp(Response.Record));
-			DerivedData.DerivedDataKey = Response.Record.GetKey();
+			for (int32 TaskIndex = 0; TaskIndex < Tasks.Num(); ++TaskIndex)
+			{
+				UPoseSearchDatabase& PoseSearchDb = GetTask(TaskIndex).GetDatabase();
+				InFunction(PoseSearchDb);
+			}
 		}
-		else if (Response.Status == EStatus::Error)
+		else
 		{
-			BuildAndWrite(Response.Record.GetKey());
+			TArray<FAssetData> PoseSearchDatabaseAssetDataList;
+			GetPoseSearchDatabaseAssetDataList(PoseSearchDatabaseAssetDataList);
+			for (const auto& PoseSearchDbAssetData : PoseSearchDatabaseAssetDataList)
+			{
+				if (UPoseSearchDatabase* PoseSearchDb = Cast<UPoseSearchDatabase>(PoseSearchDbAssetData.FastGetAsset(false)))
+				{
+					InFunction(*PoseSearchDb);
+				}
+			}
 		}
 	}
 
-	void FPoseSearchDatabaseAsyncCacheTask::BuildAndWrite(const UE::DerivedData::FCacheKey& NewKey)
+	void FAsyncPoseSearchDatabasesManagement::ExecuteIfObjectIsReferencedByDatabase(UObject* Object, bool bUseTasksDatabases, TFunctionRef<void(UPoseSearchDatabase&)> InFunction)
 	{
-		GetRequestOwner().LaunchTask(TEXT("PoseSearchDatabaseBuild"), [this, NewKey]
+		if (UAnimSequence* Sequence = Cast<UAnimSequence>(Object))
 		{
-			if (!GetRequestOwner().IsCanceled())
-			{
-				DerivedData.SearchIndex.Reset();
-				DerivedData.SearchIndex.Schema = Database.Schema;
-				const bool bIndexReady = BuildIndex(&Database, DerivedData.SearchIndex);
-
-				WriteIndexToCache(NewKey);
-
-#if UE_POSE_SEARCH_DERIVED_DATA_LOGGING
-				UE_LOG(LogPoseSearch, Log, TEXT("%s: BuildAndWrite"), *LexToString(DerivedData.PendingDerivedDataKey));
-#endif // UE_POSE_SEARCH_DERIVED_DATA_LOGGING
-			}
-		});
+			ForEachPoseSearchDatabase(bUseTasksDatabases, [Sequence, InFunction](UPoseSearchDatabase& PoseSearchDb)
+				{
+					const bool bSequenceFound = PoseSearchDb.Sequences.ContainsByPredicate([Sequence](FPoseSearchDatabaseSequence& DbSequence)
+						{
+							return Sequence == DbSequence.Sequence || Sequence == DbSequence.LeadInSequence || Sequence == DbSequence.FollowUpSequence;
+						});
+
+					if (bSequenceFound)
+					{
+						InFunction(PoseSearchDb);
+					}
+				});
+		}
+		else if (UBlendSpace* BlendSpace = Cast<UBlendSpace>(Object))
+		{
+			ForEachPoseSearchDatabase(bUseTasksDatabases, [BlendSpace, InFunction](UPoseSearchDatabase& PoseSearchDb)
+				{
+					const bool bBlendSpaceFound = PoseSearchDb.BlendSpaces.ContainsByPredicate([BlendSpace](FPoseSearchDatabaseBlendSpace& DbBlendSpace)
+						{
+							return BlendSpace == DbBlendSpace.BlendSpace;
+						});
+
+					if (bBlendSpaceFound)
+					{
+						InFunction(PoseSearchDb);
+					}
+				});
+		}
+		else if (UPoseSearchSchema* Schema = Cast<UPoseSearchSchema>(Object))
+		{
+			ForEachPoseSearchDatabase(bUseTasksDatabases, [Schema, InFunction](UPoseSearchDatabase& PoseSearchDb)
+				{
+					if (PoseSearchDb.Schema == Schema)
+					{
+						InFunction(PoseSearchDb);
+					}
+				});
+		}
+		else if (USkeleton* Skeleton = Cast<USkeleton>(Object))
+		{
+			ForEachPoseSearchDatabase(bUseTasksDatabases, [Skeleton, InFunction](UPoseSearchDatabase& PoseSearchDb)
+				{
+					if (PoseSearchDb.Schema && PoseSearchDb.Schema->Skeleton == Skeleton)
+					{
+						InFunction(PoseSearchDb);
+					}
+				});
+		}
+		else if (UPoseSearchDatabase* PoseSearchDb = Cast<UPoseSearchDatabase>(Object))
+		{
+			InFunction(*PoseSearchDb);
+		}
 	}
 
-	void FPoseSearchDatabaseAsyncCacheTask::WriteIndexToCache(const UE::DerivedData::FCacheKey& NewKey)
+	// @todo: probably overkilling listening to OnObjectPreSave to RequestAsyncBuildIndex, since we already perform it during OnObjectPropertyChanged
+	void FAsyncPoseSearchDatabasesManagement::OnObjectPreSave(UObject* SavedObject, FObjectPreSaveContext SaveContext)
 	{
-		using namespace UE::DerivedData;
-
-		TArray<uint8> RawBytes;
-		FMemoryWriter Writer(RawBytes);
-		Writer << DerivedData.SearchIndex;
-		FSharedBuffer RawData = MakeSharedBufferFromArray(MoveTemp(RawBytes));
+		ExecuteIfObjectIsReferencedByDatabase(SavedObject, false, [this](UPoseSearchDatabase& Database)
+			{
+				RequestAsyncBuildIndex(Database, false, true);
+			});
+	}
 
-		FCacheRecordBuilder Builder(NewKey);
-		Builder.AddValue(Id, RawData);
+	// we're listening to OnPreObjectPropertyChanged to cancel any pending Task indexing databases to avoid multi threading issues
+	void FAsyncPoseSearchDatabasesManagement::OnPreObjectPropertyChanged(UObject* Object, const class FEditPropertyChain& PropChain)
+	{
+		ExecuteIfObjectIsReferencedByDatabase(Object, true, [this](UPoseSearchDatabase& Database)
+			{
+				// cancelling the async indexing request for Database
+				WaitOnExistingBuildIndex(Database, false);
+			});
+	}
 
-		Owner.KeepAlive();
-		GetCache().Put({ { { Database.GetPathName() }, Builder.Build() } }, Owner);
-		DerivedData.DerivedDataKey = NewKey;
+	// @todo: investigate if it's possible to move the indexing request then databases index gets accessed
+	void FAsyncPoseSearchDatabasesManagement::OnObjectPropertyChanged(UObject* Object, struct FPropertyChangedEvent& Event)
+	{
+		ExecuteIfObjectIsReferencedByDatabase(Object, false, [this](UPoseSearchDatabase& Database)
+			{
+				// requesting a new async indexing for Database (it should already have been cancelled by OnPreObjectPropertyChanged)
+				RequestAsyncBuildIndex(Database, false, false);
+			});
 	}
 
-	void FPoseSearchDatabaseAsyncCacheTask::BuildIndexFromCacheRecord(UE::DerivedData::FCacheRecord&& CacheRecord)
+	void FAsyncPoseSearchDatabasesManagement::Shutdown()
 	{
-		using namespace UE::DerivedData;
+		check(IsInGameThread());
 
-#if UE_POSE_SEARCH_DERIVED_DATA_LOGGING
-		UE_LOG(LogPoseSearch, Log, TEXT("%s: BuildIndexFromCacheRecord"), *LexToString(DerivedData.PendingDerivedDataKey));
-#endif // UE_POSE_SEARCH_DERIVED_DATA_LOGGING
+		FCoreUObjectDelegates::OnObjectPreSave.Remove(OnObjectPreSaveHandle);
+		OnObjectPreSaveHandle.Reset();
 
-		DerivedData.SearchIndex.Reset();
-		DerivedData.SearchIndex.Schema = Database.Schema;
+		FCoreUObjectDelegates::OnPreObjectPropertyChanged.Remove(OnPreObjectPropertyChangedHandle);
+		OnPreObjectPropertyChangedHandle.Reset();
 
-		FSharedBuffer RawData = CacheRecord.GetValue(Id).GetData().Decompress();
-		FMemoryReaderView Reader(RawData);
-		Reader << DerivedData.SearchIndex;
+		FCoreUObjectDelegates::OnObjectPropertyChanged.Remove(OnObjectPropertyChangedHandle);
+		OnObjectPropertyChangedHandle.Reset();
 	}
 
-	FIoHash FPoseSearchDatabaseAsyncCacheTask::CreateKey(UPoseSearchDatabase& Database)
+	void FAsyncPoseSearchDatabasesManagement::Tick(float DeltaTime)
 	{
-		using UE::PoseSearch::FDerivedDataKeyBuilder;
+		check(IsInGameThread());
 
-#if UE_POSE_SEARCH_DERIVED_DATA_LOGGING
-		const double StartTime = FPlatformTime::Seconds();
-#endif
-		FDerivedDataKeyBuilder KeyBuilder;
-		FGuid VersionGuid = FDevSystemGuids::GetSystemGuid(FDevSystemGuids::Get().POSESEARCHDB_DERIVEDDATA_VER);
-		KeyBuilder << VersionGuid;
-		Database.BuildDerivedDataKey(KeyBuilder);
-		FDerivedDataKeyBuilder::HashDigestType Hash = KeyBuilder.Finalize();
-
-		// Stores a BLAKE3-160 hash, taken from the first 20 bytes of a BLAKE3-256 hash
-		FIoHash IoHash(Hash);
-
-#if UE_POSE_SEARCH_DERIVED_DATA_LOGGING
-		const double TotalTime = FPlatformTime::Seconds() - StartTime;
-		UE_LOG(LogPoseSearch, Log, TEXT("%s: CreateKey ('%s' - %.0f µs)"), *LexToString(IoHash), *LexToString(Hash), *Database.GetName(), TotalTime * 1e6);
-#endif // UE_POSE_SEARCH_DERIVED_DATA_LOGGING
+		const int32 NumTasks = Tasks.Num();
+		if (NumTasks > 0)
+		{
+			// iterating backwards because of the possible RemoveAtSwap 
+			for (int32 TaskIndex = NumTasks - 1; TaskIndex >= 0; --TaskIndex)
+			{
+				FPoseSearchDatabaseAsyncCacheTask& Task = GetTask(TaskIndex);
+				if (Task.Poll())
+				{
+					RemoveTask(TaskIndex);
+				}
+			}
+		}
+	}
 
-		return IoHash;
+	void FAsyncPoseSearchDatabasesManagement::TickCook(float DeltaTime, bool bCookCompete)
+	{
+		Tick(DeltaTime);
 	}
 
-#endif // WITH_EDITOR
+	TStatId FAsyncPoseSearchDatabasesManagement::GetStatId() const
+	{
+		RETURN_QUICK_DECLARE_CYCLE_STAT(FAsyncPoseSearchDatabasesManagement, STATGROUP_Tickables);
+	}
 
-	FArchive& operator<<(FArchive& Ar, FPoseSearchIndex& Index)
+	void FAsyncPoseSearchDatabasesManagement::AddReferencedObjects(FReferenceCollector& Collector)
 	{
-		int32 NumValues = 0;
-		int32 NumPCAValues = 0;
-		int32 NumAssets = 0;
+	}
 
-		if (Ar.IsSaving())
-		{
-			NumValues = Index.Values.Num();
-			NumPCAValues = Index.PCAValues.Num();
-			NumAssets = Index.Assets.Num();
-		}
+	void FAsyncPoseSearchDatabasesManagement::RequestAsyncBuildIndex(UPoseSearchDatabase& Database, bool bWaitForCompletion, bool bCancelPreviousTask)
+	{
+		check(IsInGameThread());
 
-		Ar << Index.NumPoses;
-		Ar << NumValues;
-		Ar << NumPCAValues;
-		Ar << NumAssets;
-		Ar << Index.OverallFlags;
+		WaitOnExistingBuildIndex(Database, !bCancelPreviousTask);
 
-		if (Ar.IsLoading())
+#if DO_CHECK
+		for (int32 TaskIndex = 0; TaskIndex < Tasks.Num(); ++TaskIndex)
 		{
-			Index.Values.SetNumUninitialized(NumValues);
-			Index.PCAValues.SetNumUninitialized(NumPCAValues);
-			Index.PoseMetadata.SetNumUninitialized(Index.NumPoses);
-			Index.Assets.SetNumUninitialized(NumAssets);
+			const FPoseSearchDatabaseAsyncCacheTask& Task = GetTask(TaskIndex);
+			if (&Task.GetDatabase() == &Database)
+			{
+				// making sure we have no tasks associated to Database
+				checkNoEntry();
+			}
 		}
 
-		if (Index.Values.Num() > 0)
-		{
-			Ar.Serialize(&Index.Values[0], Index.Values.Num() * Index.Values.GetTypeSize());
-		}
+#endif // DO_CHECK
+		Tasks.Emplace(MakeUnique<FPoseSearchDatabaseAsyncCacheTask>(Database));
 
-		if (Index.PCAValues.Num() > 0)
+		if (bWaitForCompletion)
 		{
-			Ar.Serialize(&Index.PCAValues[0], Index.PCAValues.Num() * Index.PCAValues.GetTypeSize());
+			WaitOnExistingBuildIndex(Database, true);
 		}
+	}
+
+	void FAsyncPoseSearchDatabasesManagement::WaitOnExistingBuildIndex(const UPoseSearchDatabase& Database, bool bWantResults)
+	{
+		check(IsInGameThread());
 
-		if (Index.PoseMetadata.Num() > 0)
+		// iterating backwards because of the possible RemoveAtSwap
+		for (int32 TaskIndex = Tasks.Num() - 1; TaskIndex >= 0; --TaskIndex)
 		{
-			Ar.Serialize(&Index.PoseMetadata[0], Index.PoseMetadata.Num() * Index.PoseMetadata.GetTypeSize());
+			FPoseSearchDatabaseAsyncCacheTask& Task = GetTask(TaskIndex);
+			if (&Task.GetDatabase() == &Database)
+			{
+				if (bWantResults)
+				{
+					Task.Wait();
+				}
+				else
+				{
+					Task.Cancel();
+				}
+
+				RemoveTask(TaskIndex);
+			}
 		}
+	}
 
-		if (Index.Assets.Num() > 0)
+	bool FAsyncPoseSearchDatabasesManagement::IsBuildingIndex(const UPoseSearchDatabase& Database) const
+	{
+		for (int32 TaskIndex = 0; TaskIndex < Tasks.Num(); ++TaskIndex)
 		{
-			Ar.Serialize(&Index.Assets[0], Index.Assets.Num() * Index.Assets.GetTypeSize());
+			const FPoseSearchDatabaseAsyncCacheTask& Task = GetTask(TaskIndex);
+			if (&Task.GetDatabase() == &Database)
+			{
+				if (!Task.Poll())
+				{
+					return true;
+				}
+			}
 		}
-
-		Ar << Index.WeightsSqrt;
-		Ar << Index.Mean;
-		Ar << Index.PCAProjectionMatrix;
-
-		Serialize(Ar, Index.KDTree, Index.PCAValues.GetData());
-
-		return Ar;
+		return false;
 	}
 
-}
\ No newline at end of file
+} // namespace UE::PoseSearch
+#endif // WITH_EDITOR
diff --git a/Engine/Plugins/Experimental/Animation/PoseSearch/Source/Runtime/Private/PoseSearchDerivedData.h b/Engine/Plugins/Experimental/Animation/PoseSearch/Source/Runtime/Private/PoseSearchDerivedData.h
deleted file mode 100644
diff --git a/Engine/Plugins/Experimental/Animation/PoseSearch/Source/Runtime/Public/PoseSearch/KDTree.h b/Engine/Plugins/Experimental/Animation/PoseSearch/Source/Runtime/Public/PoseSearch/KDTree.h
--- a/Engine/Plugins/Experimental/Animation/PoseSearch/Source/Runtime/Public/PoseSearch/KDTree.h
+++ b/Engine/Plugins/Experimental/Animation/PoseSearch/Source/Runtime/Public/PoseSearch/KDTree.h
@@ -14,15 +14,28 @@ struct FKDTreeImplementation;
 
 struct POSESEARCH_API FKDTree
 {
-	struct DataSource
+	struct FDataSource
 	{
-		DataSource(int32 pointCount, int32 pointDim, const float* data)
+		FDataSource(int32 pointCount, int32 pointDim, const float* data)
 		: PointCount(pointDim > 0 ? pointCount : 0)
 		, PointDim(pointDim)
 		, Data(data)
 		{
 		}
 
+		~FDataSource()
+		{
+			PointCount = 0;
+			PointDim = 0;
+			Data = nullptr;
+		}
+
+		FDataSource() = default;
+		FDataSource(const FDataSource& Other) = default;
+		FDataSource(FDataSource&& Other) = default;
+		FDataSource& operator=(const FDataSource& Other) = default;
+		FDataSource& operator=(FDataSource&& Other) = default;
+
 		// Must return the number of data points
 		inline size_t kdtree_get_point_count() const { return PointCount; }
 
@@ -112,18 +125,19 @@ struct POSESEARCH_API FKDTree
 
 	FKDTree(int32 Count, int32 Dim, const float* Data, int32 MaxLeafSize);
 	FKDTree();
-	FKDTree(const FKDTree& r);
-	FKDTree(FKDTree&& r) = delete;
+	FKDTree(const FKDTree& Other);
+	FKDTree(FKDTree&& Other) = delete;
 	
 	~FKDTree();
 
-	FKDTree& operator=(const FKDTree& r);
+	FKDTree& operator=(const FKDTree& Other);
+	FKDTree& operator=(FKDTree&& Other) = delete;
 
 	bool FindNeighbors(KNNResultSet& Result, const float* Query) const;
 	void Construct(int32 Count, int32 Dim, const float* Data, int32 MaxLeafSize);
 	SIZE_T GetAllocatedSize() const;
 
-	DataSource DataSrc;
+	FDataSource DataSource;
 	FKDTreeImplementation* Impl = nullptr;
 };
 
diff --git a/Engine/Plugins/Experimental/Animation/PoseSearch/Source/Runtime/Public/PoseSearch/PoseSearch.h b/Engine/Plugins/Experimental/Animation/PoseSearch/Source/Runtime/Public/PoseSearch/PoseSearch.h
--- a/Engine/Plugins/Experimental/Animation/PoseSearch/Source/Runtime/Public/PoseSearch/PoseSearch.h
+++ b/Engine/Plugins/Experimental/Animation/PoseSearch/Source/Runtime/Public/PoseSearch/PoseSearch.h
@@ -23,7 +23,6 @@
 #include "GameplayTagContainer.h"
 #include "Interfaces/Interface_BoneReferenceSkeletonProvider.h"
 #include "PoseSearch/KDTree.h"
-#include "PoseSearch/PoseSearchDerivedDataKey.h"
 #include "ObjectTrace.h"
 #include "PoseSearch.generated.h"
 
@@ -40,10 +39,8 @@ class UBlendSpace;
 struct FCompactPose;
 struct FPoseContext;
 struct FReferenceSkeleton;
-struct FPoseSearchDatabaseDerivedData;
 class UAnimNotifyState_PoseSearchBase;
 class UPoseSearchSchema;
-class FBlake3;
 
 #ifndef UE_POSE_SEARCH_FORCE_SINGLE_THREAD
 #define UE_POSE_SEARCH_FORCE_SINGLE_THREAD 0
@@ -60,13 +57,20 @@ constexpr EParallelForFlags ParallelForFlags = EParallelForFlags::None;
 namespace UE::PoseSearch
 {
 class FPoseHistory;
-struct FPoseSearchDatabaseAsyncCacheTask;
 struct FDebugDrawParams;
 struct FSchemaInitializer;
 struct FQueryBuildingContext;
 struct FSearchContext;
+class FDerivedDataKeyBuilder;
 } // namespace UE::PoseSearch
 
+#if WITH_EDITOR
+namespace UE::DerivedData
+{
+	class IRequestOwner;
+} // namespace UE::DerivedData
+#endif
+
 // eigen forward declaration
 namespace Eigen
 {
@@ -551,6 +555,7 @@ class POSESEARCH_API UPoseSearchSchema : public UDataAsset, public IBoneReferenc
 	GENERATED_BODY()
 
 public:
+	// @todo: used only for indexing: cache it somewhere else
 	UPROPERTY(EditAnywhere, Category = "Schema")
 	TObjectPtr<USkeleton> Skeleton;
 
@@ -570,6 +575,7 @@ public:
 	UPROPERTY()
 	int32 SchemaCardinality = 0;
 
+	// @todo: used only for indexing: cache it somewhere else
 	UPROPERTY()
 	TArray<FBoneReference> BoneReferences;
 
@@ -766,9 +772,6 @@ struct POSESEARCH_API FPoseSearchIndex
 	UPROPERTY(meta = (NeverInHash))
 	EPoseSearchPoseFlags OverallFlags = EPoseSearchPoseFlags::None;
 
-	UPROPERTY(meta = (NeverInHash))
-	TObjectPtr<const UPoseSearchSchema> Schema;
-
 	UPROPERTY(meta = (NeverInHash))
 	TArray<FPoseSearchIndexAsset> Assets;
 
@@ -776,7 +779,12 @@ struct POSESEARCH_API FPoseSearchIndex
 	UPROPERTY(Category = Info, VisibleAnywhere, meta = (NeverInHash))
 	float MinCostAddend = -MAX_FLT;
 
-	bool IsValid() const;
+	FPoseSearchIndex() = default;
+	FPoseSearchIndex(const FPoseSearchIndex& Other); // custom copy constructor to deal with the KDTree DataSrc
+	FPoseSearchIndex(FPoseSearchIndex&& Other) = delete;
+	FPoseSearchIndex& operator=(const FPoseSearchIndex& Other); // custom equal operator to deal with the KDTree DataSrc
+	FPoseSearchIndex& operator=(FPoseSearchIndex&& Other) = delete;
+
 	bool IsValidPoseIndex(int32 PoseIdx) const { return PoseIdx < NumPoses; }
 	bool IsEmpty() const;
 
@@ -786,11 +794,14 @@ struct POSESEARCH_API FPoseSearchIndex
 	const FPoseSearchIndexAsset& GetAssetForPose(int32 PoseIdx) const;
 	const FPoseSearchIndexAsset* GetAssetForPoseSafe(int32 PoseIdx) const;
 
-	float GetAssetTime(int32 PoseIdx) const;
+	float GetAssetTime(int32 PoseIdx, float SamplingInterval) const;
 
 	void Reset();
+	void InitSearchIndexAssets(TConstArrayView<FPoseSearchDatabaseSequence> Sequences, TConstArrayView<FPoseSearchDatabaseBlendSpace> BlendSpaces, const FPoseSearchExcludeFromDatabaseParameters& ExcludeFromDatabaseParameters);
 
-	FPoseSearchCost ComparePoses(int32 PoseIdx, EPoseSearchBooleanRequest QueryMirrorRequest, UE::PoseSearch::EPoseComparisonFlags PoseComparisonFlags, TConstArrayView<float> QueryValues) const;
+	FPoseSearchCost ComparePoses(int32 PoseIdx, EPoseSearchBooleanRequest QueryMirrorRequest, UE::PoseSearch::EPoseComparisonFlags PoseComparisonFlags, float MirrorMismatchCostBias, TConstArrayView<float> QueryValues) const;
+
+	friend FArchive& operator<<(FArchive& Ar, FPoseSearchIndex& Index);
 };
 
 //////////////////////////////////////////////////////////////////////////
@@ -805,7 +816,7 @@ struct POSESEARCH_API FPoseSearchDatabaseAnimationAssetBase
 	virtual UAnimationAsset* GetAnimationAsset() const { return nullptr; }
 	virtual bool IsLooping() const { return false; }
 #if WITH_EDITOR
-	virtual void BuildDerivedDataKey(UE::PoseSearch::FDerivedDataKeyBuilder& KeyBuilder) {}
+	virtual void BuildDerivedDataKey(UE::PoseSearch::FDerivedDataKeyBuilder& KeyBuilder) const {}
 #endif
 };
 
@@ -844,7 +855,7 @@ struct POSESEARCH_API FPoseSearchDatabaseSequence : public FPoseSearchDatabaseAn
 	virtual UAnimationAsset* GetAnimationAsset() const override { return Sequence; }
 	virtual bool IsLooping() const override { return Sequence->bLoop; }
 #if WITH_EDITOR
-	virtual void BuildDerivedDataKey(UE::PoseSearch::FDerivedDataKeyBuilder& KeyBuilder) override;
+	virtual void BuildDerivedDataKey(UE::PoseSearch::FDerivedDataKeyBuilder& KeyBuilder) const override;
 #endif
 };
 
@@ -877,7 +888,7 @@ struct POSESEARCH_API FPoseSearchDatabaseBlendSpace : public FPoseSearchDatabase
 	virtual UAnimationAsset* GetAnimationAsset() const override;
 	virtual bool IsLooping() const override;
 #if WITH_EDITOR
-	virtual void BuildDerivedDataKey(UE::PoseSearch::FDerivedDataKeyBuilder& KeyBuilder) override;
+	virtual void BuildDerivedDataKey(UE::PoseSearch::FDerivedDataKeyBuilder& KeyBuilder) const override;
 #endif
 
 	void GetBlendSpaceParameterSampleRanges(int32& HorizontalBlendNum, int32& VerticalBlendNum) const;
@@ -1038,9 +1049,6 @@ struct FSearchResult
 
 	float AssetTime = 0.0f;
 
-#if WITH_EDITOR
-	FIoHash SearchIndexHash = FIoHash::Zero;
-#endif // WITH_EDITOR
 #if WITH_EDITORONLY_DATA
 	FPoseSearchCost BruteForcePoseCost;
 #endif // WITH_EDITORONLY_DATA
@@ -1085,17 +1093,9 @@ public:
 	UPROPERTY(EditAnywhere, Category = "Database")
 	FPoseSearchExcludeFromDatabaseParameters ExcludeFromDatabaseParameters;
 
-	// Drag and drop animations here to add them in bulk to Sequences
-	UPROPERTY(EditAnywhere, Category = "Database", DisplayName="Drag And Drop Anims Here")
-	TArray<TObjectPtr<UAnimSequence>> SimpleSequences;
-
 	UPROPERTY(EditAnywhere, Category="Database")
 	TArray<FPoseSearchDatabaseSequence> Sequences;
 
-	// Drag and drop blendspaces here to add them in bulk to Blend Spaces
-	UPROPERTY(EditAnywhere, Category = "Database", DisplayName = "Drag And Drop Blend Spaces Here")
-	TArray<TObjectPtr<UBlendSpace>> SimpleBlendSpaces;
-
 	UPROPERTY(EditAnywhere, Category = "Database")
 	TArray<FPoseSearchDatabaseBlendSpace> BlendSpaces;
 
@@ -1115,6 +1115,9 @@ public:
 	UPROPERTY(EditAnywhere, Category = "Performance")
 	bool bSkipSearchIfPossible = true;
 
+	UPROPERTY(Transient, meta = (ExcludeFromHash))
+	FPoseSearchIndex PoseSearchIndex;
+
 	virtual ~UPoseSearchDatabase();
 
 	const FPoseSearchIndex* GetSearchIndex() const;
@@ -1134,74 +1137,34 @@ public:
 	const bool IsSourceAssetLooping(const FPoseSearchIndexAsset& SearchIndexAsset) const;
 	const FString GetSourceAssetName(const FPoseSearchIndexAsset& SearchIndexAsset) const;
 	int32 GetNumberOfPrincipalComponents() const;
-	
-	// UObject
+	void BuildQuery(UE::PoseSearch::FSearchContext& SearchContext, FPoseSearchFeatureVectorBuilder& OutQuery) const;
+
+	// Begin UObject
 	virtual void PostLoad() override;
 	virtual void PostSaveRoot(FObjectPostSaveRootContext ObjectSaveContext) override;
 	virtual void Serialize(FArchive& Ar) override;
-
-	// Populates the FPoseSearchIndex::Assets array by evaluating the data in the Sequences array
-	bool TryInitSearchIndexAssets(FPoseSearchIndex& OutSearchIndex) const;
+	// End UObject
+	
+	// Begin UPoseSearchSearchableAsset
+	virtual UE::PoseSearch::FSearchResult Search(UE::PoseSearch::FSearchContext& SearchContext) const override;
+	// End UPoseSearchSearchableAsset
 
 #if WITH_EDITOR
-	void BuildDerivedDataKey(UE::PoseSearch::FDerivedDataKeyBuilder& KeyBuilder);
-	virtual void PostEditChangeProperty(struct FPropertyChangedEvent& PropertyChangedEvent) override;
+	void BuildDerivedDataKey(UE::PoseSearch::FDerivedDataKeyBuilder& KeyBuilder) const;
 	virtual void BeginCacheForCookedPlatformData(const ITargetPlatform* TargetPlatform) override;
 	virtual bool IsCachedCookedPlatformDataLoaded(const ITargetPlatform* TargetPlatform) override;
-#endif // WITH_EDITOR
-
 private:
-	FPoseSearchIndex* GetSearchIndex();
-	void CollectSimpleSequences();
-	void CollectSimpleBlendSpaces();
-	void FindValidSequenceIntervals(const FPoseSearchDatabaseSequence& DbSequence, TArray<FFloatRange>& ValidRanges) const;
-	
-	FPoseSearchDatabaseDerivedData* PrivateDerivedData;
-
-#if WITH_EDITOR
 	DECLARE_MULTICAST_DELEGATE(FOnDerivedDataRebuildMulticaster);
 	FOnDerivedDataRebuildMulticaster OnDerivedDataRebuild;
-
-	DECLARE_MULTICAST_DELEGATE(FOnAssetChangeMulticaster);
-	FOnDerivedDataRebuildMulticaster OnAssetChange;
-
-	DECLARE_MULTICAST_DELEGATE(FOnGroupChangeMulticaster);
-	FOnDerivedDataRebuildMulticaster OnGroupChange;
-#endif // WITH_EDITOR
-
 public:
-#if WITH_EDITOR
-
 	typedef FOnDerivedDataRebuildMulticaster::FDelegate FOnDerivedDataRebuild;
 	void RegisterOnDerivedDataRebuild(const FOnDerivedDataRebuild& Delegate);
 	void UnregisterOnDerivedDataRebuild(void* Unregister);
-	void NotifyDerivedDataBuildStarted();
-
-	typedef FOnAssetChangeMulticaster::FDelegate FOnAssetChange;
-	void RegisterOnAssetChange(const FOnAssetChange& Delegate);
-	void UnregisterOnAssetChange(void* Unregister);
-	void NotifyAssetChange();
-
-	typedef FOnGroupChangeMulticaster::FDelegate FOnGroupChange;
-	void RegisterOnGroupChange(const FOnGroupChange& Delegate);
-	void UnregisterOnGroupChange(void* Unregister);
-	void NotifyGroupChange();
-
-	void BeginCacheDerivedData();
-
-	FIoHash GetSearchIndexHash() const;
-
-	bool IsDerivedDataBuildPending() const;
+	void NotifyDerivedDataBuildStarted() const;
 #endif // WITH_EDITOR
 
-	bool IsDerivedDataValid();
-
-public:
-	virtual UE::PoseSearch::FSearchResult Search(UE::PoseSearch::FSearchContext& SearchContext) const override;
-
-	void BuildQuery(UE::PoseSearch::FSearchContext& SearchContext, FPoseSearchFeatureVectorBuilder& OutQuery) const;
-
-protected:
+private:
+	FPoseSearchIndex* GetSearchIndex();
 	UE::PoseSearch::FSearchResult SearchPCAKDTree(UE::PoseSearch::FSearchContext& SearchContext) const;
 	UE::PoseSearch::FSearchResult SearchBruteForce(UE::PoseSearch::FSearchContext& SearchContext) const;
 };
@@ -1402,17 +1365,23 @@ struct POSESEARCH_API FSearchContext
 private:
 	struct FCachedEntry
 	{
+		FCachedEntry(float InSampleTime = 0.f, const FTransform& InTransform = FTransform::Identity, FBoneIndexType InBoneIndexType = -1)
+			: SampleTime(InSampleTime)
+			, Transform(InTransform)
+			, BoneIndexType(InBoneIndexType)
+		{
+		}
+
 		float SampleTime = 0.f;
 
 		// associated transform to BoneIndexType in ComponentSpace (except for the root bone stored in global space)
-		FTransform Transform;
+		FTransform Transform = FTransform::Identity;
 
 		// if -1 it represents the root bone
 		FBoneIndexType BoneIndexType = -1;
 	};
 
-	// @todo: make it a fixed size array (or hash map if we end up having many CachedEntry) to avoid allocations
-	TArray<FCachedEntry> CachedEntries;
+	TArray<FCachedEntry, TInlineAllocator<64>> CachedEntries;
 	
 	struct FCachedQuery
 	{
@@ -1499,7 +1468,11 @@ POSESEARCH_API bool BuildIndex(const UAnimSequence* Sequence, UPoseSearchSequenc
 * 
 * @return Whether the index was built successfully
 */
-POSESEARCH_API bool BuildIndex(const UPoseSearchDatabase* Database, FPoseSearchIndex& OutSearchIndex);
+POSESEARCH_API bool BuildIndex(const UPoseSearchDatabase& Database, FPoseSearchIndex& OutSearchIndex
+#if WITH_EDITOR
+	, UE::DerivedData::IRequestOwner& Owner
+#endif
+);
 
 } // namespace UE::PoseSearch
 
diff --git a/Engine/Plugins/Experimental/Animation/PoseSearch/Source/Runtime/Public/PoseSearch/PoseSearchDerivedData.h b/Engine/Plugins/Experimental/Animation/PoseSearch/Source/Runtime/Public/PoseSearch/PoseSearchDerivedData.h
new file mode 100644
--- /dev/null
+++ b/Engine/Plugins/Experimental/Animation/PoseSearch/Source/Runtime/Public/PoseSearch/PoseSearchDerivedData.h
@@ -0,0 +1,65 @@
+// Copyright Epic Games, Inc. All Rights Reserved.
+
+#pragma once
+
+#include "PoseSearch/PoseSearch.h"
+
+#if WITH_EDITOR
+#include "TickableEditorObject.h"
+#include "UObject/ObjectSaveContext.h"
+
+namespace UE::PoseSearch
+{
+	struct FPoseSearchDatabaseAsyncCacheTask;
+	class FPoseSearchDatabaseAsyncCacheTasks;
+
+	class POSESEARCH_API FAsyncPoseSearchDatabasesManagement : public FTickableEditorObject, public FTickableCookObject, public FGCObject
+	{
+	public:
+		~FAsyncPoseSearchDatabasesManagement();
+		static FAsyncPoseSearchDatabasesManagement& Get();
+		void RequestAsyncBuildIndex(UPoseSearchDatabase& Database, bool bWaitForCompletion = false, bool bCancelPreviousTask = false);
+		void WaitOnExistingBuildIndex(const UPoseSearchDatabase& Database, bool bWantResults);
+		bool IsBuildingIndex(const UPoseSearchDatabase& Database) const;
+
+	private:
+		FAsyncPoseSearchDatabasesManagement();
+
+		void ForEachPoseSearchDatabase(bool bUseTasksDatabases, TFunctionRef<void(UPoseSearchDatabase&)> InFunction);
+		void ExecuteIfObjectIsReferencedByDatabase(UObject* Object, bool bUseTasksDatabases, TFunctionRef<void(UPoseSearchDatabase&)> InFunction);
+
+		void OnObjectPreSave(UObject* SavedObject, FObjectPreSaveContext SaveContext);
+		void OnPreObjectPropertyChanged(UObject* Object, const class FEditPropertyChain& PropChain);
+		void OnObjectPropertyChanged(UObject* Object, struct FPropertyChangedEvent& Event);
+
+		FPoseSearchDatabaseAsyncCacheTask& GetTask(int32 TaskIndex);
+		const FPoseSearchDatabaseAsyncCacheTask& GetTask(int32 TaskIndex) const;
+		void RemoveTask(int32 TaskIndex);
+
+		void Shutdown();
+		void StartQueuedTasks(int32 MaxActiveTasks);
+
+		// Begin FTickableEditorObject
+		virtual void Tick(float DeltaTime) override;
+		virtual ETickableTickType GetTickableTickType() const override { return ETickableTickType::Always; }
+		virtual TStatId GetStatId() const override;
+		// End FTickableEditorObject
+
+		// Begin FTickableCookObject
+		virtual void TickCook(float DeltaTime, bool bCookCompete) override;
+		// End FTickableCookObject
+
+		// Begin FGCObject
+		void AddReferencedObjects(FReferenceCollector& Collector) override;
+		virtual FString GetReferencerName() const override { return TEXT("FAsyncPoseSearchDatabaseManagement"); }
+		// End FGCObject
+
+		FPoseSearchDatabaseAsyncCacheTasks& Tasks;
+
+		FDelegateHandle OnObjectPreSaveHandle;
+		FDelegateHandle OnObjectPropertyChangedHandle;
+		FDelegateHandle OnPreObjectPropertyChangedHandle;
+	};
+} // namespace UE::PoseSearch
+
+#endif // WITH_EDITOR
\ No newline at end of file
diff --git a/Engine/Plugins/Experimental/Animation/PoseSearch/Source/Runtime/Public/PoseSearch/PoseSearchDerivedDataKey.h b/Engine/Plugins/Experimental/Animation/PoseSearch/Source/Runtime/Public/PoseSearch/PoseSearchDerivedDataKey.h
--- a/Engine/Plugins/Experimental/Animation/PoseSearch/Source/Runtime/Public/PoseSearch/PoseSearchDerivedDataKey.h
+++ b/Engine/Plugins/Experimental/Animation/PoseSearch/Source/Runtime/Public/PoseSearch/PoseSearchDerivedDataKey.h
@@ -15,32 +15,39 @@
 namespace UE::PoseSearch
 {
 
+#ifndef UE_POSE_SEARCH_DERIVED_DATA_LOGGING
+	#define UE_POSE_SEARCH_DERIVED_DATA_LOGGING 0
+#endif
+
 //////////////////////////////////////////////////////////////////////////
-// TObjectHasher
-template <typename HashBuilder, typename HashDigest>
-class POSESEARCH_API TObjectHasher : public FArchiveUObject
+// FDerivedDataKeyBuilder
+class POSESEARCH_API FDerivedDataKeyBuilder : public FArchiveUObject
 {
 public:
-	using HashDigestType = HashDigest;
-	using HashBuilderType = HashBuilder;
+	using Super = FArchiveUObject;
+	using HashDigestType = FBlake3Hash;
+	using HashBuilderType = FBlake3;
 
 	inline static const FName ExcludeFromHashName = FName(TEXT("ExcludeFromHash"));
 	inline static const FName NeverInHashName = FName(TEXT("NeverInHash"));
 
-	/**
-	* Default constructor.
-	*/
-	TObjectHasher()
+	FDerivedDataKeyBuilder()
 	{
 		ArIgnoreOuterRef = true;
 
-		// Set TObjectHasher to be a saving archive instead of a reference collector.
+		// Set FDerivedDataKeyBuilder to be a saving archive instead of a reference collector.
 		// Reference collection causes FSoftObjectPtrs to be serialized by their weak pointer,
 		// which doesn't give a stable hash.  Serializing these to a saving archive will
 		// use a string reference instead, which is a more meaningful hash value.
 		SetIsSaving(true);
 	}
 
+	virtual void Seek(int64 InPos) override
+	{
+		checkf(InPos == Tell(), TEXT("A hash cannot be computed when serialization relies on seeking."));
+		FArchiveUObject::Seek(InPos);
+	}
+
 	//~ Begin FArchive Interface
 	virtual bool ShouldSkipProperty(const FProperty* InProperty) const override
 	{
@@ -49,7 +56,7 @@ public:
 			return false;
 		}
 
-		if (FArchiveUObject::ShouldSkipProperty(InProperty))
+		if (Super::ShouldSkipProperty(InProperty))
 		{
 			return true;
 		}
@@ -65,48 +72,62 @@ public:
 		}
 		
 		check(!InProperty->HasMetaData(NeverInHashName));
+
+#if UE_POSE_SEARCH_DERIVED_DATA_LOGGING
+		UE_LOG(LogPoseSearch, Log, TEXT("    %s"), *InProperty->GetFullName());
+#endif
 		return false;
 	}
 
 	virtual void Serialize(void* Data, int64 Length) override
 	{
-		Hasher.Update(reinterpret_cast<uint8*>(Data), Length);
+		const uint8* HasherData = reinterpret_cast<uint8*>(Data);
+
+#if UE_POSE_SEARCH_DERIVED_DATA_LOGGING
+		FString RawBytesString = BytesToString(HasherData, Length);
+		UE_LOG(LogPoseSearch, Log, TEXT("    %s"), *RawBytesString);
+#endif
+		Hasher.Update(HasherData, Length);
 	}
 
+	// @todo: do we need those Serialize(s)?
+	//virtual void SerializeBits(void* Bits, int64 LengthBits) override
+	//virtual void SerializeInt(uint32& Value, uint32 Max) override
+	//virtual void SerializeIntPacked(uint32& Value) override
+
+	using Super::IsSaving;
+	using Super::operator<<; // For visibility of the overloads we don't override
+
 	virtual FArchive& operator<<(FName& Name) override
 	{
 		// Don't include the name of the object being serialized, since that isn't technically part of the object's state
 		if (!ObjectBeingSerialized || (Name != ObjectBeingSerialized->GetFName()))
 		{
-			uint32 NameHash = GetTypeHash(Name);
-			Hasher.Update(&NameHash, sizeof(NameHash));
+			// we cannot use GetTypeHash(Name) since it's bound to be non deterministic between editor restarts, so we convert the name into an FString and let the Serialize(void* Data, int64 Length) deal with it
+			FString NameString = Name.ToString();
+			*this << NameString;
 		}
-
 		return *this;
 	}
 
 	virtual FArchive& operator<<(class UObject*& Object) override
 	{
-		FArchive& Ar = *this;
-
 		if (!RootObject || !Object || !Object->IsIn(RootObject))
 		{
 			auto UniqueName = GetPathNameSafe(Object);
-			Ar << UniqueName;
+			*this << UniqueName;
 		}
 		else
 		{
 			ObjectsToSerialize.Enqueue(Object);
 		}
-
-		return Ar;
+		return *this;
 	}
 
-	using FArchiveUObject::operator<<; // For visibility of the overloads we don't override
 
 	virtual FString GetArchiveName() const override
 	{
-		return TEXT("TObjectHasher");
+		return TEXT("FDerivedDataKeyBuilder");
 	}
 	//~ End FArchive Interface
 
@@ -131,10 +152,7 @@ public:
 				{
 					// Serialize it
 					ObjectBeingSerialized = Object;
-					if (!CustomSerialize(Object))
-					{
-						const_cast<UObject*>(Object)->Serialize(*this);
-					}
+					const_cast<UObject*>(Object)->Serialize(*this);
 					ObjectBeingSerialized = nullptr;
 				}
 			}
@@ -149,18 +167,13 @@ public:
 		Update(Object, Object);
 	}
 
-	HashDigest Finalize()
+	HashDigestType Finalize()
 	{
 		return Hasher.Finalize();
 	}
 
 protected:
-	virtual bool CustomSerialize(const UObject* Object)
-	{
-		return false;
-	}
-
-	HashBuilder Hasher;
+	HashBuilderType Hasher;
 
 	/** Queue of object references awaiting serialization */
 	TQueue<const UObject*> ObjectsToSerialize;
@@ -175,10 +188,6 @@ protected:
 	const UObject* RootObject = nullptr;
 };
 
-using FObjectHasherBlake3 = TObjectHasher<FBlake3, FBlake3Hash>;
-
-using FDerivedDataKeyBuilder = FObjectHasherBlake3;
-
 } // namespace UE::PoseSearch
 
 #endif // WITH_EDITORONLY_DATA
\ No newline at end of file
