motion matching - hooking up DDC events to the editor

diff --git a/Engine/Plugins/Experimental/Animation/PoseSearch/Source/Editor/Private/PoseSearchDatabaseEditor.cpp b/Engine/Plugins/Experimental/Animation/PoseSearch/Source/Editor/Private/PoseSearchDatabaseEditor.cpp
--- a/Engine/Plugins/Experimental/Animation/PoseSearch/Source/Editor/Private/PoseSearchDatabaseEditor.cpp
+++ b/Engine/Plugins/Experimental/Animation/PoseSearch/Source/Editor/Private/PoseSearchDatabaseEditor.cpp
@@ -205,12 +205,15 @@ namespace UE::PoseSearch
 			// Ensure statistics information get updated/populated
 			if (DatabaseAsset)
 			{
-				const auto UpdateStatisticsInformationLambda = [this]()
+				const auto UpdateStatisticsInformationLambda = [this](UPoseSearchDatabase::EDerivedDataBuildState State = UPoseSearchDatabase::EDerivedDataBuildState::Ended)
 				{
-					UPoseSearchDatabaseStatistics* Statistics = NewObject<UPoseSearchDatabaseStatistics>();
-					Statistics->AddToRoot();
-					Statistics->Initialize(GetPoseSearchDatabase());
-					StatisticsOverviewWidget->SetObject(Statistics);
+					if (State == UPoseSearchDatabase::EDerivedDataBuildState::Ended)
+					{
+						UPoseSearchDatabaseStatistics* Statistics = NewObject<UPoseSearchDatabaseStatistics>();
+						Statistics->AddToRoot();
+						Statistics->Initialize(GetPoseSearchDatabase());
+						StatisticsOverviewWidget->SetObject(Statistics);
+					}
 				};
 
 				// Init statistics
@@ -218,7 +221,6 @@ namespace UE::PoseSearch
 
 				// Ensure any database changes are reflected
 				DatabaseAsset->RegisterOnDerivedDataRebuild(UPoseSearchDatabase::FOnDerivedDataRebuild::CreateLambda(UpdateStatisticsInformationLambda));
-				// TODO: Register to delegate that gets triggered after the DerivedData has been rebuilt.
 			}
 		}
 		
diff --git a/Engine/Plugins/Experimental/Animation/PoseSearch/Source/Editor/Private/PoseSearchDatabaseViewModel.cpp b/Engine/Plugins/Experimental/Animation/PoseSearch/Source/Editor/Private/PoseSearchDatabaseViewModel.cpp
--- a/Engine/Plugins/Experimental/Animation/PoseSearch/Source/Editor/Private/PoseSearchDatabaseViewModel.cpp
+++ b/Engine/Plugins/Experimental/Animation/PoseSearch/Source/Editor/Private/PoseSearchDatabaseViewModel.cpp
@@ -79,19 +79,22 @@ namespace UE::PoseSearch
 		Collector.AddReferencedObject(PoseSearchDatabase);
 	}
 
-	void FDatabaseViewModel::Initialize(
-		UPoseSearchDatabase* InPoseSearchDatabase,
-		const TSharedRef<FDatabasePreviewScene>& InPreviewScene)
+	void FDatabaseViewModel::Initialize(UPoseSearchDatabase* InPoseSearchDatabase, const TSharedRef<FDatabasePreviewScene>& InPreviewScene)
 	{
 		PoseSearchDatabase = InPoseSearchDatabase;
 		PreviewScenePtr = InPreviewScene;
 
 		RemovePreviewActors();
 
-		PoseSearchDatabase->RegisterOnDerivedDataRebuild(
-			UPoseSearchDatabase::FOnDerivedDataRebuild::CreateSP(
-				this,
-				&FDatabaseViewModel::ResetPreviewActors));
+		const auto ResetPreviewActorsLambda = [this](UPoseSearchDatabase::EDerivedDataBuildState State)
+		{
+			if (State != UPoseSearchDatabase::EDerivedDataBuildState::Cancelled)
+			{
+				ResetPreviewActors();
+			}
+		};
+
+		PoseSearchDatabase->RegisterOnDerivedDataRebuild(UPoseSearchDatabase::FOnDerivedDataRebuild::CreateLambda(ResetPreviewActorsLambda));
 	}
 
 	void FDatabaseViewModel::ResetPreviewActors()
diff --git a/Engine/Plugins/Experimental/Animation/PoseSearch/Source/Runtime/Private/PoseSearch.cpp b/Engine/Plugins/Experimental/Animation/PoseSearch/Source/Runtime/Private/PoseSearch.cpp
--- a/Engine/Plugins/Experimental/Animation/PoseSearch/Source/Runtime/Private/PoseSearch.cpp
+++ b/Engine/Plugins/Experimental/Animation/PoseSearch/Source/Runtime/Private/PoseSearch.cpp
@@ -1412,9 +1412,9 @@ void UPoseSearchDatabase::UnregisterOnDerivedDataRebuild(void* Unregister)
 	OnDerivedDataRebuild.RemoveAll(Unregister);
 }
 
-void UPoseSearchDatabase::NotifyDerivedDataBuildStarted() const
+void UPoseSearchDatabase::NotifyDerivedDataRebuild(EDerivedDataBuildState State) const
 {
-	OnDerivedDataRebuild.Broadcast();
+	OnDerivedDataRebuild.Broadcast(State);
 }
 
 void UPoseSearchDatabase::BeginCacheForCookedPlatformData(const ITargetPlatform* TargetPlatform)
diff --git a/Engine/Plugins/Experimental/Animation/PoseSearch/Source/Runtime/Private/PoseSearchDerivedData.cpp b/Engine/Plugins/Experimental/Animation/PoseSearch/Source/Runtime/Private/PoseSearchDerivedData.cpp
--- a/Engine/Plugins/Experimental/Animation/PoseSearch/Source/Runtime/Private/PoseSearchDerivedData.cpp
+++ b/Engine/Plugins/Experimental/Animation/PoseSearch/Source/Runtime/Private/PoseSearchDerivedData.cpp
@@ -40,8 +40,6 @@ namespace UE::PoseSearch
 		const UPoseSearchDatabase& GetDatabase() const { return Database; }
 		static FIoHash CreateKey(UPoseSearchDatabase& Database);
 
-		void MoveSearchIndexToDatabase();
-
 	private:
 		FPoseSearchDatabaseAsyncCacheTask(const FPoseSearchDatabaseAsyncCacheTask& Other) = delete;
 		FPoseSearchDatabaseAsyncCacheTask(FPoseSearchDatabaseAsyncCacheTask&& Other) = delete;
@@ -63,9 +61,9 @@ namespace UE::PoseSearch
 		, Owner(UE::DerivedData::EPriority::Normal)
 	{
 		using namespace UE::DerivedData;
-		const FIoHash PendingDerivedDataKey = CreateKey(InDatabase);
+		const FIoHash PendingDerivedDataKey = CreateKey(Database);
 
-		InDatabase.NotifyDerivedDataBuildStarted();
+		Database.NotifyDerivedDataRebuild(UPoseSearchDatabase::EDerivedDataBuildState::Prestarted);
 
 		UE_LOG(LogPoseSearch, Log, TEXT("%s - %s BeginCache"), *LexToString(PendingDerivedDataKey), *Database.GetName());
 
@@ -79,11 +77,6 @@ namespace UE::PoseSearch
 	}
 
 	FPoseSearchDatabaseAsyncCacheTask::~FPoseSearchDatabaseAsyncCacheTask()
-	{
-		// @todo: add InDatabase.NotifyDerivedDataBuildEnded();
-	}
-
-	void FPoseSearchDatabaseAsyncCacheTask::MoveSearchIndexToDatabase()
 	{
 		check(IsInGameThread());
 		check(Poll());
@@ -91,6 +84,12 @@ namespace UE::PoseSearch
 		{
 			// @todo: implement FPoseSearchIndex move ctor and assignment operator and use a MoveTemp(SearchIndex) here
 			Database.PoseSearchIndex = SearchIndex;
+
+			Database.NotifyDerivedDataRebuild(UPoseSearchDatabase::EDerivedDataBuildState::Ended);
+		}
+		else
+		{
+			Database.NotifyDerivedDataRebuild(UPoseSearchDatabase::EDerivedDataBuildState::Cancelled);
 		}
 	}
 
@@ -242,7 +241,6 @@ namespace UE::PoseSearch
 
 	void FAsyncPoseSearchDatabasesManagement::RemoveTask(int32 TaskIndex)
 	{
-		GetTask(TaskIndex).MoveSearchIndexToDatabase();
 		Tasks.RemoveAtSwap(TaskIndex, 1, false);
 	}
 
diff --git a/Engine/Plugins/Experimental/Animation/PoseSearch/Source/Runtime/Public/PoseSearch/PoseSearch.h b/Engine/Plugins/Experimental/Animation/PoseSearch/Source/Runtime/Public/PoseSearch/PoseSearch.h
--- a/Engine/Plugins/Experimental/Animation/PoseSearch/Source/Runtime/Public/PoseSearch/PoseSearch.h
+++ b/Engine/Plugins/Experimental/Animation/PoseSearch/Source/Runtime/Public/PoseSearch/PoseSearch.h
@@ -1153,14 +1153,21 @@ public:
 	void BuildDerivedDataKey(UE::PoseSearch::FDerivedDataKeyBuilder& KeyBuilder) const;
 	virtual void BeginCacheForCookedPlatformData(const ITargetPlatform* TargetPlatform) override;
 	virtual bool IsCachedCookedPlatformDataLoaded(const ITargetPlatform* TargetPlatform) override;
+
+	enum class EDerivedDataBuildState
+	{
+		Prestarted,
+		Ended,
+		Cancelled
+	};
 private:
-	DECLARE_MULTICAST_DELEGATE(FOnDerivedDataRebuildMulticaster);
+	DECLARE_MULTICAST_DELEGATE_OneParam(FOnDerivedDataRebuildMulticaster, EDerivedDataBuildState State);
 	FOnDerivedDataRebuildMulticaster OnDerivedDataRebuild;
 public:
 	typedef FOnDerivedDataRebuildMulticaster::FDelegate FOnDerivedDataRebuild;
 	void RegisterOnDerivedDataRebuild(const FOnDerivedDataRebuild& Delegate);
 	void UnregisterOnDerivedDataRebuild(void* Unregister);
-	void NotifyDerivedDataBuildStarted() const;
+	void NotifyDerivedDataRebuild(EDerivedDataBuildState State) const;
 #endif // WITH_EDITOR
 
 private:
