motion matching
- FAsyncPoseSearchDatabasesManagement key generation improvements
- lazy load MM databases

diff --git a/Engine/Plugins/Experimental/Animation/PoseSearch/Source/Editor/Private/PoseSearchDatabaseEditor.cpp b/Engine/Plugins/Experimental/Animation/PoseSearch/Source/Editor/Private/PoseSearchDatabaseEditor.cpp
--- a/Engine/Plugins/Experimental/Animation/PoseSearch/Source/Editor/Private/PoseSearchDatabaseEditor.cpp
+++ b/Engine/Plugins/Experimental/Animation/PoseSearch/Source/Editor/Private/PoseSearchDatabaseEditor.cpp
@@ -9,6 +9,7 @@
 #include "PoseSearchDatabaseEditorReflection.h"
 #include "PoseSearchEditor.h"
 #include "PoseSearch/PoseSearch.h"
+#include "PoseSearch/PoseSearchDerivedData.h"
 #include "GameFramework/WorldSettings.h"
 #include "AdvancedPreviewSceneModule.h"
 #include "IStructureDetailsView.h"
@@ -205,19 +206,19 @@ namespace UE::PoseSearch
 			// Ensure statistics information get updated/populated
 			if (DatabaseAsset)
 			{
-				const auto UpdateStatisticsInformationLambda = [this](UPoseSearchDatabase::EDerivedDataBuildState State = UPoseSearchDatabase::EDerivedDataBuildState::Ended)
+				const auto UpdateStatisticsInformationLambda = [this]()
 				{
-					if (State == UPoseSearchDatabase::EDerivedDataBuildState::Ended)
-					{
-						UPoseSearchDatabaseStatistics* Statistics = NewObject<UPoseSearchDatabaseStatistics>();
-						Statistics->AddToRoot();
-						Statistics->Initialize(GetPoseSearchDatabase());
-						StatisticsOverviewWidget->SetObject(Statistics);
-					}
+					UPoseSearchDatabaseStatistics* Statistics = NewObject<UPoseSearchDatabaseStatistics>();
+					Statistics->AddToRoot();
+					Statistics->Initialize(GetPoseSearchDatabase());
+					StatisticsOverviewWidget->SetObject(Statistics);
 				};
 
 				// Init statistics
-				UpdateStatisticsInformationLambda();
+				if (FAsyncPoseSearchDatabasesManagement::RequestAsyncBuildIndex(GetPoseSearchDatabase(), ERequestAsyncBuildFlag::ContinueRequest))
+				{
+					UpdateStatisticsInformationLambda();
+				}
 
 				// Ensure any database changes are reflected
 				DatabaseAsset->RegisterOnDerivedDataRebuild(UPoseSearchDatabase::FOnDerivedDataRebuild::CreateLambda(UpdateStatisticsInformationLambda));
@@ -226,7 +227,7 @@ namespace UE::PoseSearch
 		
 		// Define Editor Layout
 		const TSharedRef<FTabManager::FLayout> StandaloneDefaultLayout =
-		FTabManager::NewLayout("Standalone_PoseSearchDatabaseEditor_Layout_v0.07")
+		FTabManager::NewLayout("Standalone_PoseSearchDatabaseEditor_Layout_v0.08")
 			->AddArea
 			(
 				FTabManager::NewPrimaryArea()->SetOrientation(Orient_Vertical)
diff --git a/Engine/Plugins/Experimental/Animation/PoseSearch/Source/Editor/Private/PoseSearchDatabaseEditorReflection.cpp b/Engine/Plugins/Experimental/Animation/PoseSearch/Source/Editor/Private/PoseSearchDatabaseEditorReflection.cpp
--- a/Engine/Plugins/Experimental/Animation/PoseSearch/Source/Editor/Private/PoseSearchDatabaseEditorReflection.cpp
+++ b/Engine/Plugins/Experimental/Animation/PoseSearch/Source/Editor/Private/PoseSearchDatabaseEditorReflection.cpp
@@ -54,52 +54,51 @@ void UPoseSearchDatabaseStatistics::Initialize(const UPoseSearchDatabase* PoseSe
 	
 	if (PoseSearchDatabase)
 	{
-		if (const FPoseSearchIndex* SearchIndexSafe = PoseSearchDatabase->GetSearchIndexSafe())
-		{
-			// General information
+		const FPoseSearchIndex& SearchIndex = PoseSearchDatabase->GetSearchIndex();
+		// General information
 	
-			AnimationSequences = PoseSearchDatabase->Sequences.Num();
+		AnimationSequences = PoseSearchDatabase->Sequences.Num();
 			
-			TotalAnimationPosesInFrames = SearchIndexSafe->NumPoses;
-			TotalAnimationPosesInTime = FText::Format(TimeFormat, static_cast<double>(SearchIndexSafe->NumPoses) / PoseSearchDatabase->Schema->SampleRate);
+		const int32 SampleRate = FMath::Max(1, PoseSearchDatabase->Schema->SampleRate);
+		TotalAnimationPosesInFrames = SearchIndex.NumPoses;
+		TotalAnimationPosesInTime = FText::Format(TimeFormat, static_cast<double>(SearchIndex.NumPoses) / SampleRate);
 			
+		{
+			uint32 NumOfSearchablePoses = 0;
+			for (const FPoseSearchPoseMetadata & PoseMetadata : SearchIndex.PoseMetadata)
 			{
-				uint32 NumOfSearchablePoses = 0;
-				for (const FPoseSearchPoseMetadata & PoseMetadata : SearchIndexSafe->PoseMetadata)
-				{
-					NumOfSearchablePoses += PoseMetadata.Flags != EPoseSearchPoseFlags::BlockTransition;
-				}
-				
-				SearchableFrames = NumOfSearchablePoses;
-				SearchableTime = FText::Format(TimeFormat, static_cast<double>(NumOfSearchablePoses) / PoseSearchDatabase->Schema->SampleRate);
+				NumOfSearchablePoses += PoseMetadata.Flags != EPoseSearchPoseFlags::BlockTransition;
 			}
+				
+			SearchableFrames = NumOfSearchablePoses;
+			SearchableTime = FText::Format(TimeFormat, static_cast<double>(NumOfSearchablePoses) / SampleRate);
+		}
 			
-			// Velocity information
+		// Velocity information
 	
-			// TODO: Set values once they can be queried from the PoseSearchIndex.
+		// TODO: Set values once they can be queried from the PoseSearchIndex.
 			
-			// Principal Component Analysis
+		// Principal Component Analysis
 			
-			ExplainedVariance = SearchIndexSafe->PCAExplainedVariance;
+		ExplainedVariance = SearchIndex.PCAExplainedVariance;
 			
-			// Memory information
+		// Memory information
 			
-			{
-				const uint32 ValuesBytesSize = SearchIndexSafe->Values.GetAllocatedSize();
-				const uint32 PCAValuesBytesSize = SearchIndexSafe->PCAValues.GetAllocatedSize();
-				const uint32 KDTreeBytesSize = SearchIndexSafe->KDTree.GetAllocatedSize();
-				const uint32 PoseMetadataBytesSize = SearchIndexSafe->PoseMetadata.GetAllocatedSize();
-				const uint32 AssetsBytesSize = SearchIndexSafe->Assets.GetAllocatedSize();
-				const uint32 OtherBytesSize = SearchIndexSafe->PCAProjectionMatrix.GetAllocatedSize() + SearchIndexSafe->Mean.GetAllocatedSize() + SearchIndexSafe->WeightsSqrt.GetAllocatedSize();
-				const uint32 EstimatedDatabaseBytesSize = ValuesBytesSize + PCAValuesBytesSize + KDTreeBytesSize + PoseMetadataBytesSize + AssetsBytesSize + OtherBytesSize;
+		{
+			const uint32 ValuesBytesSize = SearchIndex.Values.GetAllocatedSize();
+			const uint32 PCAValuesBytesSize = SearchIndex.PCAValues.GetAllocatedSize();
+			const uint32 KDTreeBytesSize = SearchIndex.KDTree.GetAllocatedSize();
+			const uint32 PoseMetadataBytesSize = SearchIndex.PoseMetadata.GetAllocatedSize();
+			const uint32 AssetsBytesSize = SearchIndex.Assets.GetAllocatedSize();
+			const uint32 OtherBytesSize = SearchIndex.PCAProjectionMatrix.GetAllocatedSize() + SearchIndex.Mean.GetAllocatedSize() + SearchIndex.WeightsSqrt.GetAllocatedSize();
+			const uint32 EstimatedDatabaseBytesSize = ValuesBytesSize + PCAValuesBytesSize + KDTreeBytesSize + PoseMetadataBytesSize + AssetsBytesSize + OtherBytesSize;
 				
-				ValuesSize = FText::AsMemory(ValuesBytesSize);
-				PCAValuesSize = FText::AsMemory(PCAValuesBytesSize);
-				KDTreeSize = FText::AsMemory(KDTreeBytesSize);
-				PoseMetadataSize = FText::AsMemory(PoseMetadataBytesSize);
-				AssetsSize = FText::AsMemory(AssetsBytesSize);
-				EstimatedDatabaseSize = FText::AsMemory(EstimatedDatabaseBytesSize);
-			}
+			ValuesSize = FText::AsMemory(ValuesBytesSize);
+			PCAValuesSize = FText::AsMemory(PCAValuesBytesSize);
+			KDTreeSize = FText::AsMemory(KDTreeBytesSize);
+			PoseMetadataSize = FText::AsMemory(PoseMetadataBytesSize);
+			AssetsSize = FText::AsMemory(AssetsBytesSize);
+			EstimatedDatabaseSize = FText::AsMemory(EstimatedDatabaseBytesSize);
 		}
 	}
 }
diff --git a/Engine/Plugins/Experimental/Animation/PoseSearch/Source/Editor/Private/PoseSearchDatabasePreviewScene.cpp b/Engine/Plugins/Experimental/Animation/PoseSearch/Source/Editor/Private/PoseSearchDatabasePreviewScene.cpp
--- a/Engine/Plugins/Experimental/Animation/PoseSearch/Source/Editor/Private/PoseSearchDatabasePreviewScene.cpp
+++ b/Engine/Plugins/Experimental/Animation/PoseSearch/Source/Editor/Private/PoseSearchDatabasePreviewScene.cpp
@@ -6,6 +6,7 @@
 #include "PoseSearchDatabaseViewModel.h"
 
 #include "PoseSearch/PoseSearch.h"
+#include "PoseSearch/PoseSearchDerivedData.h"
 
 #include "Animation/DebugSkelMeshComponent.h"
 #include "GameFramework/WorldSettings.h"
@@ -61,11 +62,12 @@ namespace UE::PoseSearch
 		const FDatabaseViewModel* ViewModel = GetEditor()->GetViewModel();
 		const UPoseSearchDatabase* Database = ViewModel->GetPoseSearchDatabase();
 
-		if (Database && Database->IsValidForSearch() && ViewModel->IsPoseFeaturesDrawMode(EFeaturesDrawMode::All | EFeaturesDrawMode::Detailed))
+		if (ViewModel->IsPoseFeaturesDrawMode(EFeaturesDrawMode::All | EFeaturesDrawMode::Detailed) && !ViewModel->GetPreviewActors().IsEmpty() &&
+			FAsyncPoseSearchDatabasesManagement::RequestAsyncBuildIndex(Database, ERequestAsyncBuildFlag::ContinueRequest))
 		{
 			for (const FDatabasePreviewActor& PreviewActor : ViewModel->GetPreviewActors())
 			{
-				if (Database->IsValidPoseIndex(PreviewActor.CurrentPoseIndex))
+				if (Database->GetSearchIndex().IsValidPoseIndex(PreviewActor.CurrentPoseIndex))
 				{
 					UE::PoseSearch::FDebugDrawParams DrawParams;
 					DrawParams.RootTransform = PreviewActor.Mesh->GetComponentTransform();
diff --git a/Engine/Plugins/Experimental/Animation/PoseSearch/Source/Editor/Private/PoseSearchDatabaseViewModel.cpp b/Engine/Plugins/Experimental/Animation/PoseSearch/Source/Editor/Private/PoseSearchDatabaseViewModel.cpp
--- a/Engine/Plugins/Experimental/Animation/PoseSearch/Source/Editor/Private/PoseSearchDatabaseViewModel.cpp
+++ b/Engine/Plugins/Experimental/Animation/PoseSearch/Source/Editor/Private/PoseSearchDatabaseViewModel.cpp
@@ -86,15 +86,7 @@ namespace UE::PoseSearch
 
 		RemovePreviewActors();
 
-		const auto ResetPreviewActorsLambda = [this](UPoseSearchDatabase::EDerivedDataBuildState State)
-		{
-			if (State != UPoseSearchDatabase::EDerivedDataBuildState::Cancelled)
-			{
-				ResetPreviewActors();
-			}
-		};
-
-		PoseSearchDatabase->RegisterOnDerivedDataRebuild(UPoseSearchDatabase::FOnDerivedDataRebuild::CreateLambda(ResetPreviewActorsLambda));
+		PoseSearchDatabase->RegisterOnDerivedDataRebuild(UPoseSearchDatabase::FOnDerivedDataRebuild::CreateSP(this, &FDatabaseViewModel::ResetPreviewActors));
 	}
 
 	void FDatabaseViewModel::ResetPreviewActors()
@@ -129,15 +121,15 @@ namespace UE::PoseSearch
 			}
 		}
 
-		const FPoseSearchIndex* SearchIndex = PoseSearchDatabase->GetSearchIndexSafe(false);
-		if (SearchIndex)
+		if (FAsyncPoseSearchDatabasesManagement::RequestAsyncBuildIndex(PoseSearchDatabase, ERequestAsyncBuildFlag::ContinueRequest))
 		{
 			FBoneContainer BoneContainer;
 			BoneContainer.InitializeTo(PoseSearchDatabase->Schema->BoneIndicesWithParents, FCurveEvaluationOption(false), *PoseSearchDatabase->Schema->Skeleton);
 
-			for (int32 IndexAssetIndex = 0; IndexAssetIndex < SearchIndex->Assets.Num(); ++IndexAssetIndex)
+			const FPoseSearchIndex& SearchIndex = PoseSearchDatabase->GetSearchIndex();
+			for (int32 IndexAssetIndex = 0; IndexAssetIndex < SearchIndex.Assets.Num(); ++IndexAssetIndex)
 			{
-				const FPoseSearchIndexAsset& IndexAsset = SearchIndex->Assets[IndexAssetIndex];
+				const FPoseSearchIndexAsset& IndexAsset = SearchIndex.Assets[IndexAssetIndex];
 
 				bool bSpawn = false;
 				if (AnimationPreviewMode == EAnimationPreviewMode::OriginalAndMirrored || !IndexAsset.bMirrored)
@@ -177,7 +169,8 @@ namespace UE::PoseSearch
 
 	void FDatabaseViewModel::BuildSearchIndex()
 	{
-		UE::PoseSearch::FAsyncPoseSearchDatabasesManagement::Get().RequestAsyncBuildIndex(*PoseSearchDatabase, false, true);
+		using namespace UE::PoseSearch;
+		FAsyncPoseSearchDatabasesManagement::RequestAsyncBuildIndex(PoseSearchDatabase, ERequestAsyncBuildFlag::NewRequest);
 	}
 
 	void FDatabaseViewModel::PreviewBackwardEnd()
@@ -221,79 +214,75 @@ namespace UE::PoseSearch
 	{
 		FDatabasePreviewActor PreviewActor;
 
-		if (PoseSearchDatabase && PoseSearchDatabase->IsValidForIndexing())
+		check(PoseSearchDatabase);
+		const FPoseSearchIndex& SearchIndex = PoseSearchDatabase->GetSearchIndex();
+		const FPoseSearchIndexAsset& IndexAsset = SearchIndex.Assets[IndexAssetIndex];
+		UAnimationAsset* PreviewAsset = nullptr;
+		if (IndexAsset.Type == ESearchIndexAssetType::Sequence)
 		{
-			if (const FPoseSearchIndex* SearchIndex = PoseSearchDatabase->GetSearchIndexSafe(false))
+			FPoseSearchDatabaseSequence DatabaseSequence = PoseSearchDatabase->Sequences[IndexAsset.SourceAssetIdx];
+			PreviewAsset = DatabaseSequence.Sequence;
+			if (!PreviewAsset)
 			{
-				const FPoseSearchIndexAsset& IndexAsset = SearchIndex->Assets[IndexAssetIndex];
-				UAnimationAsset* PreviewAsset = nullptr;
-				if (IndexAsset.Type == ESearchIndexAssetType::Sequence)
-				{
-					FPoseSearchDatabaseSequence DatabaseSequence = PoseSearchDatabase->Sequences[IndexAsset.SourceAssetIdx];
-					PreviewAsset = DatabaseSequence.Sequence;
-					if (!PreviewAsset)
-					{
-						return PreviewActor;
-					}
+				return PreviewActor;
+			}
 
-					FSequenceSampler::FInput Input;
-					Input.ExtrapolationParameters = PoseSearchDatabase->ExtrapolationParameters;
-					Input.Sequence = DatabaseSequence.Sequence;
+			FSequenceSampler::FInput Input;
+			Input.ExtrapolationParameters = PoseSearchDatabase->ExtrapolationParameters;
+			Input.Sequence = DatabaseSequence.Sequence;
 
-					PreviewActor.SequenceSampler.Init(Input);
-				}
-				else if (IndexAsset.Type == ESearchIndexAssetType::BlendSpace)
-				{
-					FPoseSearchDatabaseBlendSpace DatabaseBlendSpace = PoseSearchDatabase->BlendSpaces[IndexAsset.SourceAssetIdx];
-					PreviewAsset = DatabaseBlendSpace.BlendSpace;
-					if (!PreviewAsset)
-					{
-						return PreviewActor;
-					}
-
-					FBlendSpaceSampler::FInput Input;
-					Input.BoneContainer = BoneContainer;
-					Input.ExtrapolationParameters = PoseSearchDatabase->ExtrapolationParameters;
-					Input.BlendSpace = DatabaseBlendSpace.BlendSpace;
-					Input.BlendParameters = IndexAsset.BlendParameters;
+			PreviewActor.SequenceSampler.Init(Input);
+		}
+		else if (IndexAsset.Type == ESearchIndexAssetType::BlendSpace)
+		{
+			FPoseSearchDatabaseBlendSpace DatabaseBlendSpace = PoseSearchDatabase->BlendSpaces[IndexAsset.SourceAssetIdx];
+			PreviewAsset = DatabaseBlendSpace.BlendSpace;
+			if (!PreviewAsset)
+			{
+				return PreviewActor;
+			}
 
-					PreviewActor.BlendSpaceSampler.Init(Input);
-				}
+			FBlendSpaceSampler::FInput Input;
+			Input.BoneContainer = BoneContainer;
+			Input.ExtrapolationParameters = PoseSearchDatabase->ExtrapolationParameters;
+			Input.BlendSpace = DatabaseBlendSpace.BlendSpace;
+			Input.BlendParameters = IndexAsset.BlendParameters;
 
-				PreviewActor.Type = IndexAsset.Type;
-				PreviewActor.IndexAssetIndex = IndexAssetIndex;
-				PreviewActor.CurrentPoseIndex = INDEX_NONE;
-
-				FActorSpawnParameters Params;
-				Params.SpawnCollisionHandlingOverride = ESpawnActorCollisionHandlingMethod::AdjustIfPossibleButAlwaysSpawn;
-				PreviewActor.Actor = GetWorld()->SpawnActor<AActor>(AActor::StaticClass(), FTransform::Identity, Params);
-				PreviewActor.Actor->SetFlags(RF_Transient);
-
-				PreviewActor.Mesh = NewObject<UDebugSkelMeshComponent>(PreviewActor.Actor.Get());
-				PreviewActor.Mesh->RegisterComponentWithWorld(GetWorld());
-
-				PreviewActor.AnimInstance = NewObject<UAnimPreviewInstance>(PreviewActor.Mesh.Get());
-				PreviewActor.Mesh->PreviewInstance = PreviewActor.AnimInstance.Get();
-				PreviewActor.AnimInstance->InitializeAnimation();
-				PreviewActor.Mesh->SetSkeletalMesh(PoseSearchDatabase->Schema->Skeleton->GetPreviewMesh(true));
-				PreviewActor.Mesh->EnablePreview(true, PreviewAsset);
-				PreviewActor.AnimInstance->SetAnimationAsset(PreviewAsset, false, 0.0f);
-				PreviewActor.AnimInstance->SetBlendSpacePosition(IndexAsset.BlendParameters);
-				if (IndexAsset.bMirrored && PoseSearchDatabase->Schema)
-				{
-					PreviewActor.AnimInstance->SetMirrorDataTable(PoseSearchDatabase->Schema->MirrorDataTable);
-				}
+			PreviewActor.BlendSpaceSampler.Init(Input);
+		}
 
-				PreviewActor.AnimInstance->PlayAnim(false, 0.0f);
+		PreviewActor.Type = IndexAsset.Type;
+		PreviewActor.IndexAssetIndex = IndexAssetIndex;
+		PreviewActor.CurrentPoseIndex = INDEX_NONE;
+
+		FActorSpawnParameters Params;
+		Params.SpawnCollisionHandlingOverride = ESpawnActorCollisionHandlingMethod::AdjustIfPossibleButAlwaysSpawn;
+		PreviewActor.Actor = GetWorld()->SpawnActor<AActor>(AActor::StaticClass(), FTransform::Identity, Params);
+		PreviewActor.Actor->SetFlags(RF_Transient);
+
+		PreviewActor.Mesh = NewObject<UDebugSkelMeshComponent>(PreviewActor.Actor.Get());
+		PreviewActor.Mesh->RegisterComponentWithWorld(GetWorld());
+
+		PreviewActor.AnimInstance = NewObject<UAnimPreviewInstance>(PreviewActor.Mesh.Get());
+		PreviewActor.Mesh->PreviewInstance = PreviewActor.AnimInstance.Get();
+		PreviewActor.AnimInstance->InitializeAnimation();
+		PreviewActor.Mesh->SetSkeletalMesh(PoseSearchDatabase->Schema->Skeleton->GetPreviewMesh(true));
+		PreviewActor.Mesh->EnablePreview(true, PreviewAsset);
+		PreviewActor.AnimInstance->SetAnimationAsset(PreviewAsset, false, 0.0f);
+		PreviewActor.AnimInstance->SetBlendSpacePosition(IndexAsset.BlendParameters);
+		if (IndexAsset.bMirrored && PoseSearchDatabase->Schema)
+		{
+			PreviewActor.AnimInstance->SetMirrorDataTable(PoseSearchDatabase->Schema->MirrorDataTable);
+		}
 
-				if (!PreviewActor.Actor->GetRootComponent())
-				{
-					PreviewActor.Actor->SetRootComponent(PreviewActor.Mesh.Get());
-				}
+		PreviewActor.AnimInstance->PlayAnim(false, 0.0f);
 
-				UE_LOG(LogPoseSearchEditor, Log, TEXT("Spawned preview Actor: %s"), *GetNameSafe(PreviewActor.Actor.Get()));
-			}
+		if (!PreviewActor.Actor->GetRootComponent())
+		{
+			PreviewActor.Actor->SetRootComponent(PreviewActor.Mesh.Get());
 		}
+
+		UE_LOG(LogPoseSearchEditor, Log, TEXT("Spawned preview Actor: %s"), *GetNameSafe(PreviewActor.Actor.Get()));
 		return PreviewActor;
 	}
 
@@ -318,41 +307,43 @@ namespace UE::PoseSearch
 	{
 		PlayTime += DeltaSeconds * DeltaTimeMultiplier;
 		PlayTime = FMath::Clamp(PlayTime, 0.0f, MaxPreviewPlayLength);
-		UpdatePreviewActors(true);
+
+		if (FAsyncPoseSearchDatabasesManagement::RequestAsyncBuildIndex(PoseSearchDatabase, ERequestAsyncBuildFlag::ContinueRequest))
+		{
+			UpdatePreviewActors(true);
+		}
 	}
 
 	void FDatabaseViewModel::UpdatePreviewActors(bool bInTickPlayTime)
 	{
-		if (const FPoseSearchIndex* SearchIndex = PoseSearchDatabase->GetSearchIndexSafe(false))
+		const FPoseSearchIndex& SearchIndex = PoseSearchDatabase->GetSearchIndex();
+		TSharedPtr<FDatabasePreviewScene> PreviewScene = PreviewScenePtr.Pin();
+		for (FDatabasePreviewActor& PreviewActor : GetPreviewActors())
 		{
-			TSharedPtr<FDatabasePreviewScene> PreviewScene = PreviewScenePtr.Pin();
-			for (FDatabasePreviewActor& PreviewActor : GetPreviewActors())
+			if (PreviewActor.IndexAssetIndex < SearchIndex.Assets.Num())
 			{
-				if (PreviewActor.IndexAssetIndex < SearchIndex->Assets.Num())
+				if (const UAnimationAsset* PreviewAsset = PreviewActor.AnimInstance->GetAnimationAsset())
 				{
-					if (const UAnimationAsset* PreviewAsset = PreviewActor.AnimInstance->GetAnimationAsset())
+					if (const IAssetSampler* Sampler = PreviewActor.GetSampler())
 					{
-						if (const IAssetSampler* Sampler = PreviewActor.GetSampler())
-						{
-							float CurrentTime = 0.f;
-							FAnimationRuntime::AdvanceTime(false, PlayTime, CurrentTime, PreviewAsset->GetPlayLength());
+						float CurrentTime = 0.f;
+						FAnimationRuntime::AdvanceTime(false, PlayTime, CurrentTime, PreviewAsset->GetPlayLength());
 
-							const float CurrentScaledTime = PreviewActor.GetScaledTime(CurrentTime);
-							const float OldScaledTime = PreviewActor.AnimInstance->GetCurrentTime();
-							if (!FMath::IsNearlyEqual(CurrentScaledTime, OldScaledTime, UE_KINDA_SMALL_NUMBER))
+						const float CurrentScaledTime = PreviewActor.GetScaledTime(CurrentTime);
+						const float OldScaledTime = PreviewActor.AnimInstance->GetCurrentTime();
+						if (!FMath::IsNearlyEqual(CurrentScaledTime, OldScaledTime, UE_KINDA_SMALL_NUMBER))
+						{
+							FTransform RootMotion = Sampler->ExtractRootTransform(CurrentTime);
+							if (PreviewActor.AnimInstance->GetMirrorDataTable())
 							{
-								FTransform RootMotion = Sampler->ExtractRootTransform(CurrentTime);
-								if (PreviewActor.AnimInstance->GetMirrorDataTable())
-								{
-									RootMotion = MirrorRootMotion(RootMotion, PreviewActor.AnimInstance->GetMirrorDataTable());
-								}
-								PreviewActor.Actor->SetActorTransform(RootMotion);
+								RootMotion = MirrorRootMotion(RootMotion, PreviewActor.AnimInstance->GetMirrorDataTable());
+							}
+							PreviewActor.Actor->SetActorTransform(RootMotion);
 
-								PreviewActor.AnimInstance->SetPosition(CurrentScaledTime);
+							PreviewActor.AnimInstance->SetPosition(CurrentScaledTime);
 
-								const FPoseSearchIndexAsset& SearchIndexAsset = SearchIndex->Assets[PreviewActor.IndexAssetIndex];
-								PreviewActor.CurrentPoseIndex = PoseSearchDatabase->GetPoseIndexFromTime(CurrentTime, SearchIndexAsset);
-							}
+							const FPoseSearchIndexAsset& SearchIndexAsset = SearchIndex.Assets[PreviewActor.IndexAssetIndex];
+							PreviewActor.CurrentPoseIndex = PoseSearchDatabase->GetPoseIndexFromTime(CurrentTime, SearchIndexAsset);
 						}
 					}
 				}
@@ -484,14 +475,12 @@ namespace UE::PoseSearch
 
 	const FPoseSearchIndexAsset* FDatabaseViewModel::GetSelectedActorIndexAsset() const
 	{
-		if (PoseSearchDatabase)
+		if (FAsyncPoseSearchDatabasesManagement::RequestAsyncBuildIndex(PoseSearchDatabase, ERequestAsyncBuildFlag::ContinueRequest))
 		{
-			if (const FPoseSearchIndex* SearchIndex = PoseSearchDatabase->GetSearchIndexSafe(false))
+			const FPoseSearchIndex& SearchIndex = PoseSearchDatabase->GetSearchIndex();
+			if (SearchIndex.Assets.IsValidIndex(SelectedActorIndexAssetIndex))
 			{
-				if (SearchIndex->Assets.IsValidIndex(SelectedActorIndexAssetIndex))
-				{
-					return &SearchIndex->Assets[SelectedActorIndexAssetIndex];
-				}
+				return &SearchIndex.Assets[SelectedActorIndexAssetIndex];
 			}
 		}
 		return nullptr;
diff --git a/Engine/Plugins/Experimental/Animation/PoseSearch/Source/Editor/Private/PoseSearchDebugger.cpp b/Engine/Plugins/Experimental/Animation/PoseSearch/Source/Editor/Private/PoseSearchDebugger.cpp
--- a/Engine/Plugins/Experimental/Animation/PoseSearch/Source/Editor/Private/PoseSearchDebugger.cpp
+++ b/Engine/Plugins/Experimental/Animation/PoseSearch/Source/Editor/Private/PoseSearchDebugger.cpp
@@ -15,6 +15,7 @@
 #include "IRewindDebugger.h"
 #include "ObjectTrace.h"
 #include "PoseSearch/PoseSearch.h"
+#include "PoseSearch/PoseSearchDerivedData.h"
 #include "Preferences/PersonaOptions.h"
 #include "PropertyEditorModule.h"
 #include "Styling/AppStyle.h"
@@ -897,58 +898,56 @@ void SDebuggerDatabaseView::Update(const FTraceMotionMatchingStateMessage& State
 	for (const FTraceMotionMatchingStateDatabaseEntry& DbEntry : State.DatabaseEntries)
 	{
 		const UPoseSearchDatabase* Database = FTraceMotionMatchingState::GetObjectFromId<UPoseSearchDatabase>(DbEntry.DatabaseId);
-		if (Database)
+		if (FAsyncPoseSearchDatabasesManagement::RequestAsyncBuildIndex(Database, ERequestAsyncBuildFlag::ContinueRequest))
 		{
-			if (const FPoseSearchIndex* SearchIndex = Database->GetSearchIndexSafe(false))
+			const FPoseSearchIndex& SearchIndex = Database->GetSearchIndex();
+			for (const FTraceMotionMatchingStatePoseEntry& PoseEntry : DbEntry.PoseEntries)
 			{
-				for (const FTraceMotionMatchingStatePoseEntry& PoseEntry : DbEntry.PoseEntries)
+				if (const FPoseSearchIndexAsset* SearchIndexAsset = SearchIndex.GetAssetForPoseSafe(PoseEntry.DbPoseIdx))
 				{
-					if (const FPoseSearchIndexAsset* SearchIndexAsset = SearchIndex->GetAssetForPoseSafe(PoseEntry.DbPoseIdx))
+					TSharedRef<FDebuggerDatabaseRowData>& Row = UnfilteredDatabaseRows.Add_GetRef(MakeShared<FDebuggerDatabaseRowData>());
+
+					const float Time = SearchIndex.GetAssetTime(PoseEntry.DbPoseIdx, Database->Schema->GetSamplingInterval());
+
+					Row->PoseIdx = PoseEntry.DbPoseIdx;
+					Row->SourceDatabase = Database;
+					Row->DatabaseName = Database->GetName();
+					Row->DatabasePath = Database->GetPathName();
+					Row->PoseCandidateFlags = PoseEntry.PoseCandidateFlags;
+					Row->DbAssetIdx = SearchIndexAsset->SourceAssetIdx;
+					Row->AssetTime = Time;
+					Row->bMirrored = SearchIndexAsset->bMirrored;
+					Row->PoseCost = PoseEntry.Cost;
+
+					Row->CostVector.SetNum(Database->Schema->SchemaCardinality);
+					TConstArrayView<float> PoseValues = SearchIndex.GetPoseValuesSafe(PoseEntry.DbPoseIdx);
+
+					// in case we modify the schema while PIE is paused and displaying the Pose Search Editor, we could end up with a stale State with a DbEntry.QueryVector saved with the previous schema
+					// so the cardinality of DbEntry.QueryVector and PoseValues don't match. In that case we just use PoseValues as query to have all costs set to zero
+					const bool bIsQueryVectorValid = DbEntry.QueryVector.Num() == PoseValues.Num();
+					CompareFeatureVectors(PoseValues, bIsQueryVectorValid ? DbEntry.QueryVector : PoseValues, SearchIndex.WeightsSqrt, Row->CostVector);
+
+					if (SearchIndexAsset->Type == ESearchIndexAssetType::Sequence)
 					{
-						TSharedRef<FDebuggerDatabaseRowData>& Row = UnfilteredDatabaseRows.Add_GetRef(MakeShared<FDebuggerDatabaseRowData>());
-
-						const float Time = SearchIndex->GetAssetTime(PoseEntry.DbPoseIdx, Database->Schema->GetSamplingInterval());
-
-						Row->PoseIdx = PoseEntry.DbPoseIdx;
-						Row->SourceDatabase = Database;
-						Row->DatabaseName = Database->GetName();
-						Row->DatabasePath = Database->GetPathName();
-						Row->PoseCandidateFlags = PoseEntry.PoseCandidateFlags;
-						Row->DbAssetIdx = SearchIndexAsset->SourceAssetIdx;
-						Row->AssetTime = Time;
-						Row->bMirrored = SearchIndexAsset->bMirrored;
-						Row->PoseCost = PoseEntry.Cost;
-
-						Row->CostVector.SetNum(Database->Schema->SchemaCardinality);
-						TConstArrayView<float> PoseValues = SearchIndex->GetPoseValuesSafe(PoseEntry.DbPoseIdx);
-
-						// in case we modify the schema while PIE is paused and displaying the Pose Search Editor, we could end up with a stale State with a DbEntry.QueryVector saved with the previous schema
-						// so the cardinality of DbEntry.QueryVector and PoseValues don't match. In that case we just use PoseValues as query to have all costs set to zero
-						const bool bIsQueryVectorValid = DbEntry.QueryVector.Num() == PoseValues.Num();
-						CompareFeatureVectors(PoseValues, bIsQueryVectorValid ? DbEntry.QueryVector : PoseValues, SearchIndex->WeightsSqrt, Row->CostVector);
-
-						if (SearchIndexAsset->Type == ESearchIndexAssetType::Sequence)
-						{
-							const FPoseSearchDatabaseSequence& DbSequence = Database->GetSequenceSourceAsset(*SearchIndexAsset);
-							Row->AssetType = ESearchIndexAssetType::Sequence;
-							Row->AssetName = DbSequence.Sequence->GetName();
-							Row->AssetPath = DbSequence.Sequence->GetPathName();
-							Row->AnimFrame = DbSequence.Sequence->GetFrameAtTime(Time);
-							Row->AnimPercentage = Time / DbSequence.Sequence->GetPlayLength();
-							Row->bLooping = DbSequence.Sequence->bLoop;
-							Row->BlendParameters = FVector::Zero();
-						}
-						else if (SearchIndexAsset->Type == ESearchIndexAssetType::BlendSpace)
-						{
-							const FPoseSearchDatabaseBlendSpace& DbBlendSpace = Database->GetBlendSpaceSourceAsset(*SearchIndexAsset);
-							Row->AssetType = ESearchIndexAssetType::BlendSpace;
-							Row->AssetName = DbBlendSpace.BlendSpace->GetName();
-							Row->AssetPath = DbBlendSpace.BlendSpace->GetPathName();
-							Row->AnimFrame = 0; // There is no frame index associated with a blendspace
-							Row->AnimPercentage = 0.0f;
-							Row->bLooping = DbBlendSpace.BlendSpace->bLoop;
-							Row->BlendParameters = SearchIndexAsset->BlendParameters;
-						}
+						const FPoseSearchDatabaseSequence& DbSequence = Database->GetSequenceSourceAsset(*SearchIndexAsset);
+						Row->AssetType = ESearchIndexAssetType::Sequence;
+						Row->AssetName = DbSequence.Sequence->GetName();
+						Row->AssetPath = DbSequence.Sequence->GetPathName();
+						Row->AnimFrame = DbSequence.Sequence->GetFrameAtTime(Time);
+						Row->AnimPercentage = Time / DbSequence.Sequence->GetPlayLength();
+						Row->bLooping = DbSequence.Sequence->bLoop;
+						Row->BlendParameters = FVector::Zero();
+					}
+					else if (SearchIndexAsset->Type == ESearchIndexAssetType::BlendSpace)
+					{
+						const FPoseSearchDatabaseBlendSpace& DbBlendSpace = Database->GetBlendSpaceSourceAsset(*SearchIndexAsset);
+						Row->AssetType = ESearchIndexAssetType::BlendSpace;
+						Row->AssetName = DbBlendSpace.BlendSpace->GetName();
+						Row->AssetPath = DbBlendSpace.BlendSpace->GetPathName();
+						Row->AnimFrame = 0; // There is no frame index associated with a blendspace
+						Row->AnimPercentage = 0.0f;
+						Row->bLooping = DbBlendSpace.BlendSpace->bLoop;
+						Row->BlendParameters = SearchIndexAsset->BlendParameters;
 					}
 				}
 			}
@@ -961,12 +960,9 @@ void SDebuggerDatabaseView::Update(const FTraceMotionMatchingStateMessage& State
 		for (const FTraceMotionMatchingStateDatabaseEntry& DbEntry : State.DatabaseEntries)
 		{
 			const UPoseSearchDatabase* Database = FTraceMotionMatchingState::GetObjectFromId<UPoseSearchDatabase>(DbEntry.DatabaseId);
-			if (Database)
+			if (FAsyncPoseSearchDatabasesManagement::RequestAsyncBuildIndex(Database, ERequestAsyncBuildFlag::ContinueRequest))
 			{
-				if (Database->GetSearchIndexSafe(false))
-				{
-					Database->Schema->ComputeCostBreakdowns(CostBreakDownData);
-				}
+				Database->Schema->ComputeCostBreakdowns(CostBreakDownData);
 			}
 		}
 
@@ -1566,48 +1562,42 @@ void SDebuggerDetailsView::UpdateReflection(const FTraceMotionMatchingStateMessa
 	check(Reflection);
 
 	const UPoseSearchDatabase* CurrentDatabase = State.GetCurrentDatabase();
-	const FPoseSearchIndex* CurrentSearchIndex = CurrentDatabase ? CurrentDatabase->GetSearchIndex() : nullptr;
-	int32 CurrentDbPoseIdx = State.GetCurrentDatabasePoseIndex();
+	if (FAsyncPoseSearchDatabasesManagement::RequestAsyncBuildIndex(CurrentDatabase, ERequestAsyncBuildFlag::ContinueRequest))
+	{
+		const FPoseSearchIndex& CurrentSearchIndex = CurrentDatabase->GetSearchIndex();
+		int32 CurrentDbPoseIdx = State.GetCurrentDatabasePoseIndex();
 
-	Reflection->CurrentDatabaseName = CurrentDatabase ? CurrentDatabase->GetName() : "None";
-	Reflection->ElapsedPoseJumpTime = State.ElapsedPoseJumpTime;
-	Reflection->bFollowUpAnimation = EnumHasAnyFlags(State.Flags, FTraceMotionMatchingState::EFlags::FollowupAnimation);
+		Reflection->CurrentDatabaseName = CurrentDatabase->GetName();
+		Reflection->ElapsedPoseJumpTime = State.ElapsedPoseJumpTime;
+		Reflection->bFollowUpAnimation = EnumHasAnyFlags(State.Flags, FTraceMotionMatchingState::EFlags::FollowupAnimation);
 
-	Reflection->AssetPlayerAssetName = "None";
-	if (CurrentSearchIndex)
-	{
-		if (const FPoseSearchIndexAsset* IndexAsset = CurrentSearchIndex->GetAssetForPoseSafe(CurrentDbPoseIdx))
+		Reflection->AssetPlayerAssetName = "None";
+		if (const FPoseSearchIndexAsset* IndexAsset = CurrentSearchIndex.GetAssetForPoseSafe(CurrentDbPoseIdx))
 		{
 			Reflection->AssetPlayerAssetName = CurrentDatabase->GetSourceAssetName(*IndexAsset);
 		}
-	}
 
-	Reflection->AssetPlayerTime = State.AssetPlayerTime;
-	Reflection->LastDeltaTime = State.DeltaTime;
-	Reflection->SimLinearVelocity = State.SimLinearVelocity;
-	Reflection->SimAngularVelocity = State.SimAngularVelocity;
-	Reflection->AnimLinearVelocity = State.AnimLinearVelocity;
-	Reflection->AnimAngularVelocity = State.AnimAngularVelocity;
-	 
-	// Query pose
-	Reflection->QueryPoseVector.Reset();
-	for (const FTraceMotionMatchingStateDatabaseEntry& DbEntry : State.DatabaseEntries)
-	{
-		const UPoseSearchDatabase* Database = FTraceMotionMatchingState::GetObjectFromId<UPoseSearchDatabase>(DbEntry.DatabaseId);
-		if (Database && Database == CurrentDatabase)
+		Reflection->AssetPlayerTime = State.AssetPlayerTime;
+		Reflection->LastDeltaTime = State.DeltaTime;
+		Reflection->SimLinearVelocity = State.SimLinearVelocity;
+		Reflection->SimAngularVelocity = State.SimAngularVelocity;
+		Reflection->AnimLinearVelocity = State.AnimLinearVelocity;
+		Reflection->AnimAngularVelocity = State.AnimAngularVelocity;
+
+		// Query pose
+		Reflection->QueryPoseVector.Reset();
+		for (const FTraceMotionMatchingStateDatabaseEntry& DbEntry : State.DatabaseEntries)
 		{
-			if (Database->GetSearchIndexSafe(false))
+			const UPoseSearchDatabase* Database = FTraceMotionMatchingState::GetObjectFromId<UPoseSearchDatabase>(DbEntry.DatabaseId);
+			if (Database && Database == CurrentDatabase)
 			{
 				Reflection->QueryPoseVector = DbEntry.QueryVector;
 				break;
 			}
 		}
-	}
 
-	// Active pose
-	if (CurrentSearchIndex)
-	{
-		Reflection->ActivePoseVector = CurrentSearchIndex->GetPoseValuesSafe(CurrentDbPoseIdx);
+		// Active pose
+		Reflection->ActivePoseVector = CurrentSearchIndex.GetPoseValuesSafe(CurrentDbPoseIdx);
 	}
 
 	auto DebuggerView = ParentDebuggerViewPtr.Pin();
@@ -1617,15 +1607,11 @@ void SDebuggerDetailsView::UpdateReflection(const FTraceMotionMatchingStateMessa
 		if (!SelectedRows.IsEmpty())
 		{
 			const TSharedRef<FDebuggerDatabaseRowData>& Selected = SelectedRows[0];
-			if (Selected->SourceDatabase.IsValid())
+			if (FAsyncPoseSearchDatabasesManagement::RequestAsyncBuildIndex(Selected->SourceDatabase.Get(), ERequestAsyncBuildFlag::ContinueRequest))
 			{
-				const FPoseSearchIndex* SelectedSearchIndex = Selected->SourceDatabase->GetSearchIndex();
-				if (SelectedSearchIndex)
-				{
-					Reflection->SelectedPoseVector = SelectedSearchIndex->GetPoseValuesSafe(Selected->PoseIdx);
-				}
+				const FPoseSearchIndex& SelectedSearchIndex = Selected->SourceDatabase->GetSearchIndex();
+				Reflection->SelectedPoseVector = SelectedSearchIndex.GetPoseValuesSafe(Selected->PoseIdx);
 			}
-
 			Reflection->CostVector = Selected->CostVector;
 		}
 	}
@@ -1868,23 +1854,21 @@ void SDebuggerView::DrawFeatures(
 			for (const FTraceMotionMatchingStateDatabaseEntry& DbEntry : State.DatabaseEntries)
 			{
 				const UPoseSearchDatabase* Database = FTraceMotionMatchingState::GetObjectFromId<UPoseSearchDatabase>(DbEntry.DatabaseId);
-				if (Database && Database == CurrentDatabase)
+				if (Database && Database == CurrentDatabase && DbEntry.QueryVector.Num() == Database->Schema->SchemaCardinality && 
+					FAsyncPoseSearchDatabasesManagement::RequestAsyncBuildIndex(CurrentDatabase, ERequestAsyncBuildFlag::ContinueRequest))
 				{
-					if (Database->GetSearchIndexSafe(false) && DbEntry.QueryVector.Num() == Database->Schema->SchemaCardinality)
-					{
-						// Set shared state
-						FDebugDrawParams DrawParams;
-						DrawParams.World = &DebuggerWorld;
-						DrawParams.RootTransform = Transform;
-						DrawParams.DefaultLifeTime = 0.0f; // Single frame render
-						DrawParams.Mesh = Mesh;
-						DrawParams.Database = CurrentDatabase;
-						SetDrawFlags(DrawParams, Reflection ? Reflection->QueryDrawOptions : FPoseSearchDebuggerFeatureDrawOptions());
-						EnumAddFlags(DrawParams.Flags, EDebugDrawFlags::DrawQuery);
-						DrawFeatureVector(DrawParams, DbEntry.QueryVector);
-						EnumRemoveFlags(DrawParams.Flags, EDebugDrawFlags::DrawQuery);
-						break;
-					}
+					// Set shared state
+					FDebugDrawParams DrawParams;
+					DrawParams.World = &DebuggerWorld;
+					DrawParams.RootTransform = Transform;
+					DrawParams.DefaultLifeTime = 0.0f; // Single frame render
+					DrawParams.Mesh = Mesh;
+					DrawParams.Database = CurrentDatabase;
+					SetDrawFlags(DrawParams, Reflection ? Reflection->QueryDrawOptions : FPoseSearchDebuggerFeatureDrawOptions());
+					EnumAddFlags(DrawParams.Flags, EDebugDrawFlags::DrawQuery);
+					DrawFeatureVector(DrawParams, DbEntry.QueryVector);
+					EnumRemoveFlags(DrawParams.Flags, EDebugDrawFlags::DrawQuery);
+					break;
 				}
 			}
 		}
@@ -2317,24 +2301,19 @@ const UPoseSearchSearchableAsset* FDebuggerViewModel::GetSearchableAsset() const
 
 void FDebuggerViewModel::ShowSelectedSkeleton(const UPoseSearchDatabase* Database, int32 DbPoseIdx, float Time)
 {
-	if (!Database)
-	{
-		return;
-	}
-
 	UPoseSearchMeshComponent* Component = Skeletons[SelectedPose].Component.Get();
 	if (!Component)
 	{
 		return;
 	}
 
-	const FPoseSearchIndex* SearchIndex = Database->GetSearchIndex();
-	if (!SearchIndex)
+	if (!FAsyncPoseSearchDatabasesManagement::RequestAsyncBuildIndex(Database, ERequestAsyncBuildFlag::ContinueRequest))
 	{
 		return;
 	}
 
-	const FPoseSearchIndexAsset& IndexAsset = SearchIndex->GetAssetForPose(DbPoseIdx);
+	const FPoseSearchIndex& SearchIndex = Database->GetSearchIndex();
+	const FPoseSearchIndexAsset& IndexAsset = SearchIndex.GetAssetForPose(DbPoseIdx);
 	
 	Component->ResetToStart(); 
 	bSelecting = true;
@@ -2419,11 +2398,11 @@ void FDebuggerViewModel::OnUpdateNodeSelection(int32 InNodeId)
 	if (ActiveMotionMatchingState)
 	{
 		const UPoseSearchDatabase* CurrentDatabase = ActiveMotionMatchingState->GetCurrentDatabase();
-		const FPoseSearchIndex* CurrentSearchIndex = CurrentDatabase ? CurrentDatabase->GetSearchIndexSafe() : nullptr;
-		if (CurrentSearchIndex)
+		if (FAsyncPoseSearchDatabasesManagement::RequestAsyncBuildIndex(CurrentDatabase, ERequestAsyncBuildFlag::ContinueRequest))
 		{
+			const FPoseSearchIndex& CurrentSearchIndex = CurrentDatabase->GetSearchIndex();
 			int32 CurrentPoseIdx = ActiveMotionMatchingState->GetCurrentDatabasePoseIndex();
-			if (const FPoseSearchIndexAsset* IndexAsset = CurrentSearchIndex->GetAssetForPoseSafe(CurrentPoseIdx))
+			if (const FPoseSearchIndexAsset* IndexAsset = CurrentSearchIndex.GetAssetForPoseSafe(CurrentPoseIdx))
 			{
 				Skeletons[Asset].Type = IndexAsset->Type;
 				Skeletons[Asset].bMirrored = IndexAsset->bMirrored;
@@ -2782,18 +2761,19 @@ void FDebuggerViewModel::FillCompactPoseAndComponentRefRotations()
 
 void FDebuggerViewModel::PlaySelection(int32 PoseIdx, float Time)
 {
-	const UPoseSearchDatabase* Database = GetCurrentDatabase();
-	if (!Database)
-	{
-		return;
-	}
 	UPoseSearchMeshComponent* Component = Skeletons[Asset].Component.Get();
 	if (!Component)
 	{
 		return;
 	}
-	const FPoseSearchIndexAsset& IndexAsset = Database->GetSearchIndex()->GetAssetForPose(PoseIdx);
 
+	const UPoseSearchDatabase* Database = GetCurrentDatabase();
+	if (!FAsyncPoseSearchDatabasesManagement::RequestAsyncBuildIndex(Database, ERequestAsyncBuildFlag::ContinueRequest))
+	{
+		return;
+	}
+	
+	const FPoseSearchIndexAsset& IndexAsset = Database->GetSearchIndex().GetAssetForPose(PoseIdx);
 	Component->ResetToStart();
 	
 	Skeletons[Asset].Type = IndexAsset.Type;
diff --git a/Engine/Plugins/Experimental/Animation/PoseSearch/Source/Runtime/Private/AnimNode_MotionMatching.cpp b/Engine/Plugins/Experimental/Animation/PoseSearch/Source/Runtime/Private/AnimNode_MotionMatching.cpp
--- a/Engine/Plugins/Experimental/Animation/PoseSearch/Source/Runtime/Private/AnimNode_MotionMatching.cpp
+++ b/Engine/Plugins/Experimental/Animation/PoseSearch/Source/Runtime/Private/AnimNode_MotionMatching.cpp
@@ -4,13 +4,14 @@
 
 #include "Animation/AnimInstanceProxy.h"
 #include "Animation/AnimNode_Inertialization.h"
+#include "Animation/AnimRootMotionProvider.h"
 #include "Animation/AnimSequence.h"
 #include "Animation/BlendSpace.h"
 #include "Animation/MotionTrajectoryTypes.h"
-#include "Animation/AnimRootMotionProvider.h"
 #include "Components/SkeletalMeshComponent.h"
 #include "DrawDebugHelpers.h"
 #include "PoseSearch/AnimNode_PoseSearchHistoryCollector.h"
+#include "PoseSearch/PoseSearchDerivedData.h"
 #include "Trace/PoseSearchTraceLogger.h"
 
 #define LOCTEXT_NAMESPACE "AnimNode_MotionMatching"
@@ -59,6 +60,8 @@ void FAnimNode_MotionMatching::UpdateAssetPlayer(const FAnimationUpdateContext&
 {
 	DECLARE_SCOPE_HIERARCHICAL_COUNTER_ANIMNODE(UpdateAssetPlayer);
 
+	using namespace UE::PoseSearch;
+
 	GetEvaluateGraphExposedInputs().Execute(Context);
 
 	bool bNeedsReset =
@@ -75,7 +78,7 @@ void FAnimNode_MotionMatching::UpdateAssetPlayer(const FAnimationUpdateContext&
 	{
 #if WITH_EDITOR
 		// in case we're still indexing MotionMatchingState.CurrentSearchResult.Database we Reset the MotionMatchingState
-		if (MotionMatchingState.CurrentSearchResult.IsValid() && !MotionMatchingState.CurrentSearchResult.Database->GetSearchIndexSafe(true))
+		if (!FAsyncPoseSearchDatabasesManagement::RequestAsyncBuildIndex(MotionMatchingState.CurrentSearchResult.Database.Get(), ERequestAsyncBuildFlag::ContinueRequest))
 		{
 			MotionMatchingState.Reset();
 		}
@@ -137,55 +140,65 @@ void FAnimNode_MotionMatching::PreUpdate(const UAnimInstance* InAnimInstance)
 		check(SkeletalMeshComponent);
 
 		const UE::PoseSearch::FSearchResult& CurResult = MotionMatchingState.CurrentSearchResult;
-		UE::PoseSearch::FDebugDrawParams DrawParams;
-		DrawParams.RootTransform = SkeletalMeshComponent->GetComponentTransform();
-		DrawParams.Database = CurResult.Database.Get();
-		DrawParams.World = SkeletalMeshComponent->GetWorld();
-		DrawParams.DefaultLifeTime = 0.0f;
 
-		if (DrawParams.CanDraw())
+#if WITH_EDITOR
+		// in case we're still indexing MotionMatchingState.CurrentSearchResult.Database we Reset the MotionMatchingState
+		if (!FAsyncPoseSearchDatabasesManagement::RequestAsyncBuildIndex(CurResult.Database.Get(), ERequestAsyncBuildFlag::ContinueRequest))
 		{
-			if (bDebugDrawMatch)
-			{
-				DrawFeatureVector(DrawParams, CurResult.PoseIdx);
-			}
-
-			if (bDebugDrawQuery)
-			{
-				EnumAddFlags(DrawParams.Flags, EDebugDrawFlags::DrawQuery);
-				DrawFeatureVector(DrawParams, CurResult.ComposedQuery.GetValues());
-			}
+		}
+		else
+#endif // WITH_EDITOR
+		{
+			UE::PoseSearch::FDebugDrawParams DrawParams;
+			DrawParams.RootTransform = SkeletalMeshComponent->GetComponentTransform();
+			DrawParams.Database = CurResult.Database.Get();
+			DrawParams.World = SkeletalMeshComponent->GetWorld();
+			DrawParams.DefaultLifeTime = 0.0f;
 
-			if (DrawParams.Database->PoseSearchMode == EPoseSearchMode::PCAKDTree_Compare)
+			if (DrawParams.CanDraw())
 			{
-				FDebugFloatHistory& C = MotionMatchingState.SearchCostHistoryContinuing;
-				FDebugFloatHistory& B = MotionMatchingState.SearchCostHistoryBruteForce;
-				FDebugFloatHistory& K = MotionMatchingState.SearchCostHistoryKDTree;
-
-				C.AddSample(CurResult.ContinuingPoseCost.IsValid() ? CurResult.ContinuingPoseCost.GetTotalCost() : C.MaxValue);
-				B.AddSample(CurResult.BruteForcePoseCost.IsValid() ? CurResult.BruteForcePoseCost.GetTotalCost() : B.MaxValue);
-				K.AddSample(CurResult.PoseCost.IsValid() ? CurResult.PoseCost.GetTotalCost() : K.MaxValue);
-
-				// making SearchCostHistoryKDTree and SearchCostHistoryBruteForce min max consistent
-				const float MinValue = FMath::Min(C.MinValue, FMath::Min(B.MinValue, K.MinValue));
-				const float MaxValue = FMath::Max(C.MaxValue, FMath::Max(B.MaxValue, K.MaxValue));
-				
-				C.MinValue = MinValue;
-				C.MaxValue = MaxValue;
-
-				B.MinValue = MinValue;
-				B.MaxValue = MaxValue;
-
-				K.MinValue = MinValue;
-				K.MaxValue = MaxValue;
-
-				const FVector2D DrawSize(150.f, 100.f);
-				const FTransform OffsetTransform(FRotator(0.f, 0.f, 0.f), FVector(-50.f, -75.f, 100.f));
-				const FTransform DrawTransform = OffsetTransform * DrawParams.RootTransform;
-
-				DrawDebugFloatHistory(*DrawParams.World, K, OffsetTransform * DrawParams.RootTransform, DrawSize, FColor(255, 192, 203, 160)); // pink
-				DrawDebugFloatHistory(*DrawParams.World, B, OffsetTransform * DrawParams.RootTransform, DrawSize, FColor(0, 0, 255, 160)); // blue
-				DrawDebugFloatHistory(*DrawParams.World, C, OffsetTransform * DrawParams.RootTransform, DrawSize, FColor(160, 160, 160, 160)); // gray
+				if (bDebugDrawMatch)
+				{
+					DrawFeatureVector(DrawParams, CurResult.PoseIdx);
+				}
+
+				if (bDebugDrawQuery)
+				{
+					EnumAddFlags(DrawParams.Flags, EDebugDrawFlags::DrawQuery);
+					DrawFeatureVector(DrawParams, CurResult.ComposedQuery.GetValues());
+				}
+
+				if (DrawParams.Database->PoseSearchMode == EPoseSearchMode::PCAKDTree_Compare)
+				{
+					FDebugFloatHistory& C = MotionMatchingState.SearchCostHistoryContinuing;
+					FDebugFloatHistory& B = MotionMatchingState.SearchCostHistoryBruteForce;
+					FDebugFloatHistory& K = MotionMatchingState.SearchCostHistoryKDTree;
+
+					C.AddSample(CurResult.ContinuingPoseCost.IsValid() ? CurResult.ContinuingPoseCost.GetTotalCost() : C.MaxValue);
+					B.AddSample(CurResult.BruteForcePoseCost.IsValid() ? CurResult.BruteForcePoseCost.GetTotalCost() : B.MaxValue);
+					K.AddSample(CurResult.PoseCost.IsValid() ? CurResult.PoseCost.GetTotalCost() : K.MaxValue);
+
+					// making SearchCostHistoryKDTree and SearchCostHistoryBruteForce min max consistent
+					const float MinValue = FMath::Min(C.MinValue, FMath::Min(B.MinValue, K.MinValue));
+					const float MaxValue = FMath::Max(C.MaxValue, FMath::Max(B.MaxValue, K.MaxValue));
+
+					C.MinValue = MinValue;
+					C.MaxValue = MaxValue;
+
+					B.MinValue = MinValue;
+					B.MaxValue = MaxValue;
+
+					K.MinValue = MinValue;
+					K.MaxValue = MaxValue;
+
+					const FVector2D DrawSize(150.f, 100.f);
+					const FTransform OffsetTransform(FRotator(0.f, 0.f, 0.f), FVector(-50.f, -75.f, 100.f));
+					const FTransform DrawTransform = OffsetTransform * DrawParams.RootTransform;
+
+					DrawDebugFloatHistory(*DrawParams.World, K, OffsetTransform * DrawParams.RootTransform, DrawSize, FColor(255, 192, 203, 160)); // pink
+					DrawDebugFloatHistory(*DrawParams.World, B, OffsetTransform * DrawParams.RootTransform, DrawSize, FColor(0, 0, 255, 160)); // blue
+					DrawDebugFloatHistory(*DrawParams.World, C, OffsetTransform * DrawParams.RootTransform, DrawSize, FColor(160, 160, 160, 160)); // gray
+				}
 			}
 		}
 	}
diff --git a/Engine/Plugins/Experimental/Animation/PoseSearch/Source/Runtime/Private/PoseSearch.cpp b/Engine/Plugins/Experimental/Animation/PoseSearch/Source/Runtime/Private/PoseSearch.cpp
--- a/Engine/Plugins/Experimental/Animation/PoseSearch/Source/Runtime/Private/PoseSearch.cpp
+++ b/Engine/Plugins/Experimental/Animation/PoseSearch/Source/Runtime/Private/PoseSearch.cpp
@@ -170,8 +170,9 @@ typedef TArray<size_t, TInlineAllocator<128>> FNonSelectableIdx;
 static void PopulateNonSelectableIdx(FNonSelectableIdx& NonSelectableIdx, FSearchContext& SearchContext, const UPoseSearchDatabase* Database, TConstArrayView<float> QueryValues)
 {
 	check(Database);
-	const FPoseSearchIndex* SearchIndex = Database->GetSearchIndex();
-	check(SearchIndex);
+#if UE_POSE_SEARCH_TRACE_ENABLED
+	const FPoseSearchIndex& SearchIndex = Database->GetSearchIndex();
+#endif
 
 	const FPoseSearchIndexAsset* CurrentIndexAsset = SearchContext.CurrentResult.GetSearchIndexAsset();
 	if (CurrentIndexAsset && SearchContext.IsCurrentResultFromDatabase(Database) && SearchContext.PoseJumpThresholdTime > 0.f)
@@ -202,7 +203,7 @@ static void PopulateNonSelectableIdx(FNonSelectableIdx& NonSelectableIdx, FSearc
 				NonSelectableIdx.Add(PoseIdx);
 
 #if UE_POSE_SEARCH_TRACE_ENABLED
-				const FPoseSearchCost PoseCost = SearchIndex->ComparePoses(PoseIdx, SearchContext.QueryMirrorRequest, EPoseComparisonFlags::None, Database->Schema->MirrorMismatchCostBias, QueryValues);
+				const FPoseSearchCost PoseCost = SearchIndex.ComparePoses(PoseIdx, SearchContext.QueryMirrorRequest, EPoseComparisonFlags::None, Database->Schema->MirrorMismatchCostBias, QueryValues);
 				SearchContext.BestCandidates.Add(PoseCost, PoseIdx, Database, EPoseCandidateFlags::DiscardedBy_PoseJumpThresholdTime);
 #endif
 			}
@@ -235,7 +236,7 @@ static void PopulateNonSelectableIdx(FNonSelectableIdx& NonSelectableIdx, FSearc
 				NonSelectableIdx.Add(PoseIdx);
 
 #if UE_POSE_SEARCH_TRACE_ENABLED
-				const FPoseSearchCost PoseCost = SearchIndex->ComparePoses(PoseIdx, SearchContext.QueryMirrorRequest, EPoseComparisonFlags::None, Database->Schema->MirrorMismatchCostBias, QueryValues);
+				const FPoseSearchCost PoseCost = SearchIndex.ComparePoses(PoseIdx, SearchContext.QueryMirrorRequest, EPoseComparisonFlags::None, Database->Schema->MirrorMismatchCostBias, QueryValues);
 				SearchContext.BestCandidates.Add(PoseCost, PoseIdx, Database, EPoseCandidateFlags::DiscardedBy_PoseJumpThresholdTime);
 #endif
 			}
@@ -260,9 +261,9 @@ static void PopulateNonSelectableIdx(FNonSelectableIdx& NonSelectableIdx, FSearc
 				check(HistoricalPoseIndex.PoseIndex >= 0);
 					
 				// if we're editing the database and removing assets it's possible that the PoseIndicesHistory contains invalid pose indexes
-				if (HistoricalPoseIndex.PoseIndex < SearchIndex->NumPoses)
+				if (HistoricalPoseIndex.PoseIndex < SearchIndex.NumPoses)
 				{
-					const FPoseSearchCost PoseCost = SearchIndex->ComparePoses(HistoricalPoseIndex.PoseIndex, SearchContext.QueryMirrorRequest, EPoseComparisonFlags::None, Database->Schema->MirrorMismatchCostBias, QueryValues);
+					const FPoseSearchCost PoseCost = SearchIndex.ComparePoses(HistoricalPoseIndex.PoseIndex, SearchContext.QueryMirrorRequest, EPoseComparisonFlags::None, Database->Schema->MirrorMismatchCostBias, QueryValues);
 					SearchContext.BestCandidates.Add(PoseCost, HistoricalPoseIndex.PoseIndex, Database, EPoseCandidateFlags::DiscardedBy_PoseReselectHistory);
 				}
 #endif
@@ -298,13 +299,13 @@ struct FPoseFilters
 		}
 	}
 
-	bool AreFiltersValid(const FPoseSearchIndex* SearchIndex, TConstArrayView<float> QueryValues, int32 PoseIdx, const FPoseSearchPoseMetadata& Metadata
+	bool AreFiltersValid(const FPoseSearchIndex& SearchIndex, TConstArrayView<float> QueryValues, int32 PoseIdx, const FPoseSearchPoseMetadata& Metadata
 #if UE_POSE_SEARCH_TRACE_ENABLED
 		, UE::PoseSearch::FSearchContext& SearchContext, const UPoseSearchDatabase* Database
 #endif
 	) const
 	{
-		TConstArrayView<float> PoseValues = SearchIndex->GetPoseValues(PoseIdx);
+		TConstArrayView<float> PoseValues = SearchIndex.GetPoseValues(PoseIdx);
 		for (const IPoseFilter* PoseFilter : AllPoseFilters)
 		{
 			if (!PoseFilter->IsPoseValid(PoseValues, QueryValues, PoseIdx, Metadata))
@@ -316,12 +317,12 @@ struct FPoseFilters
 				}
 				else if (PoseFilter == &BlockTransitionPoseFilter)
 				{
-					const FPoseSearchCost PoseCost = SearchIndex->ComparePoses(PoseIdx, SearchContext.QueryMirrorRequest, EPoseComparisonFlags::None, Database->Schema->MirrorMismatchCostBias, QueryValues);
+					const FPoseSearchCost PoseCost = SearchIndex.ComparePoses(PoseIdx, SearchContext.QueryMirrorRequest, EPoseComparisonFlags::None, Database->Schema->MirrorMismatchCostBias, QueryValues);
 					SearchContext.BestCandidates.Add(PoseCost, PoseIdx, Database, EPoseCandidateFlags::DiscardedBy_BlockTransition);
 				}
 				else
 				{
-					const FPoseSearchCost PoseCost = SearchIndex->ComparePoses(PoseIdx, SearchContext.QueryMirrorRequest, EPoseComparisonFlags::None, Database->Schema->MirrorMismatchCostBias, QueryValues);
+					const FPoseSearchCost PoseCost = SearchIndex.ComparePoses(PoseIdx, SearchContext.QueryMirrorRequest, EPoseComparisonFlags::None, Database->Schema->MirrorMismatchCostBias, QueryValues);
 					SearchContext.BestCandidates.Add(PoseCost, PoseIdx, Database, EPoseCandidateFlags::DiscardedBy_PoseFilter);
 				}
 #endif
@@ -951,62 +952,19 @@ UE::PoseSearch::FSearchResult UPoseSearchSequenceMetaData::Search(UE::PoseSearch
 
 UPoseSearchDatabase::~UPoseSearchDatabase()
 {
-#if WITH_EDITOR
-	UE::PoseSearch::FAsyncPoseSearchDatabasesManagement::Get().WaitOnExistingBuildIndex(*this, false);
-#endif
 }
 
-FPoseSearchIndex* UPoseSearchDatabase::GetSearchIndex()
+void UPoseSearchDatabase::SetSearchIndex(const FPoseSearchIndex& SearchIndex)
 {
-	return &PoseSearchIndex;
+	check(IsInGameThread());
+	SearchIndexPrivate = SearchIndex;
 }
 
-const FPoseSearchIndex* UPoseSearchDatabase::GetSearchIndex() const
+const FPoseSearchIndex& UPoseSearchDatabase::GetSearchIndex() const
 {
-	return &PoseSearchIndex;
-}
-
-const FPoseSearchIndex* UPoseSearchDatabase::GetSearchIndexSafe(bool bVerboseLogging) const
-{
-	if (!Schema)
-	{
-		if (bVerboseLogging)
-		{
-			UE_LOG(LogAnimation, Warning, TEXT("UPoseSearchDatabase %s failed to index. Reason: no Schema!"), *GetName());
-		}
-		return nullptr;
-	}
-
-	if (!Schema->IsValid())
-	{
-		if (bVerboseLogging)
-		{
-			UE_LOG(LogAnimation, Warning, TEXT("UPoseSearchDatabase %s failed to index. Reason: Schema %s is invalid"), *GetName(), *Schema->GetName());
-		}
-		return nullptr;
-	}
-
-	const FPoseSearchIndex* SearchIndex = GetSearchIndex();
-	if (!SearchIndex || SearchIndex->IsEmpty())
-	{
-		if (bVerboseLogging)
-		{
-			UE_LOG(LogAnimation, Warning, TEXT("UPoseSearchDatabase %s failed to index. Reason: is there any unsaved modified asset?"), *GetName());
-		}
-		return nullptr;
-	}
-
-	// @todo: perhaps use FAsyncPoseSearchDatabasesManagement to understand if this UPoseSearchDatabase indexing task is running, instead of using this logic
-	if (!SearchIndex || SearchIndex->IsEmpty() || SearchIndex->WeightsSqrt.Num() != Schema->SchemaCardinality)
-	{
-		if (bVerboseLogging)
-		{
-			UE_LOG(LogAnimation, Warning, TEXT("UPoseSearchDatabase %s SearchIndex mismatch. Indexing..."), *GetName());
-		}
-		return nullptr;
-	}
-	
-	return SearchIndex;
+	// making sure the search index is consistent. if it fails the calling code hasn't been protected by FAsyncPoseSearchDatabasesManagement::RequestAsyncBuildIndex
+	check(Schema && Schema->IsValid() && !SearchIndexPrivate.IsEmpty() && SearchIndexPrivate.WeightsSqrt.Num() == Schema->SchemaCardinality && SearchIndexPrivate.KDTree.Impl);
+	return SearchIndexPrivate;
 }
 
 int32 UPoseSearchDatabase::GetPoseIndexFromTime(float Time, const FPoseSearchIndexAsset& SearchIndexAsset) const
@@ -1149,118 +1107,6 @@ int32 UPoseSearchDatabase::GetNumberOfPrincipalComponents() const
 	return FMath::Min<int32>(NumberOfPrincipalComponents, Schema->SchemaCardinality);
 }
 
-#if WITH_EDITOR
-
-static void AddRawSequenceToWriter(UAnimSequence* Sequence, UE::PoseSearch::FDerivedDataKeyBuilder& KeyBuilder)
-{
-	if (Sequence)
-	{
-		FName SequenceName = Sequence->GetFName();
-		FGuid SequenceGuid = Sequence->GetDataModel()->GenerateGuid();
-		KeyBuilder << SequenceName;
-		KeyBuilder << SequenceGuid;
-		KeyBuilder << Sequence->bLoop;
-	}
-}
-
-static void AddPoseSearchNotifiesToWriter(UAnimSequence* Sequence, UE::PoseSearch::FDerivedDataKeyBuilder& KeyBuilder)
-{
-	if (!Sequence)
-	{
-		return;
-	}
-
-	FAnimNotifyContext NotifyContext;
-	Sequence->GetAnimNotifies(0.0f, Sequence->GetPlayLength(), NotifyContext);
-
-	for (const FAnimNotifyEventReference& EventReference : NotifyContext.ActiveNotifies)
-	{
-		const FAnimNotifyEvent* NotifyEvent = EventReference.GetNotify();
-		if (!NotifyEvent || !NotifyEvent->NotifyStateClass)
-		{
-			continue;
-		}
-
-		if (NotifyEvent->NotifyStateClass->IsA<UAnimNotifyState_PoseSearchBase>())
-		{
-			float StartTime = NotifyEvent->GetTriggerTime();
-			float EndTime = NotifyEvent->GetEndTriggerTime();
-			KeyBuilder << StartTime;
-			KeyBuilder << EndTime;
-			KeyBuilder.Update(NotifyEvent->NotifyStateClass);
-		}
-	}
-}
-
-void FPoseSearchDatabaseSequence::BuildDerivedDataKey(UE::PoseSearch::FDerivedDataKeyBuilder& KeyBuilder) const
-{
-	check(KeyBuilder.IsSaving());
-
-	bool bEnabledCopy = bEnabled;
-	FFloatInterval SamplingRangeCopy = SamplingRange;
-	EPoseSearchMirrorOption MirrorOptionCopy = MirrorOption;
-
-	KeyBuilder << bEnabledCopy;
-	KeyBuilder << SamplingRangeCopy;
-	KeyBuilder << MirrorOptionCopy;
-
-	AddRawSequenceToWriter(Sequence, KeyBuilder);
-	AddRawSequenceToWriter(LeadInSequence, KeyBuilder);
-	AddRawSequenceToWriter(FollowUpSequence, KeyBuilder);
-
-	AddPoseSearchNotifiesToWriter(Sequence, KeyBuilder);
-}
-
-void FPoseSearchDatabaseBlendSpace::BuildDerivedDataKey(UE::PoseSearch::FDerivedDataKeyBuilder& KeyBuilder) const
-{
-	check(KeyBuilder.IsSaving());
-	bool bEnabledCopy = bEnabled;
-	EPoseSearchMirrorOption MirrorOptionCopy = MirrorOption;
-	bool bUseGridForSamplingCopy = bUseGridForSampling;
-	int32 NumberOfHorizontalSamplesCopy = NumberOfHorizontalSamples;
-	int32 NumberOfVerticalSamplesCopy = NumberOfVerticalSamples;
-
-	KeyBuilder << bEnabledCopy;
-	KeyBuilder << MirrorOptionCopy;
-	KeyBuilder << bUseGridForSamplingCopy;
-	KeyBuilder << NumberOfHorizontalSamplesCopy;
-	KeyBuilder << NumberOfVerticalSamplesCopy;
-
-	const TArray<FBlendSample>& BlendSpaceSamples = BlendSpace->GetBlendSamples();
-	for (const FBlendSample& Sample : BlendSpaceSamples)
-	{
-		AddRawSequenceToWriter(Sample.Animation, KeyBuilder);
-		FVector SampleValue = Sample.SampleValue;
-		float RateScale = Sample.RateScale;
- 		KeyBuilder << SampleValue;
- 		KeyBuilder << RateScale;
-	}
-
-	KeyBuilder << BlendSpace->bLoop;
-}
-
-void UPoseSearchDatabase::BuildDerivedDataKey(UE::PoseSearch::FDerivedDataKeyBuilder& KeyBuilder) const
-{
-	KeyBuilder.Update(this);
-
-	if (Schema)
-	{
-		KeyBuilder.Update(Schema.Get());
-	}
-
-	for (const FPoseSearchDatabaseSequence& DbSequence : Sequences)
-	{
-		DbSequence.BuildDerivedDataKey(KeyBuilder);
-	}
-
-	for (const FPoseSearchDatabaseBlendSpace& DbBlendSpace : BlendSpaces)
-	{
-		DbBlendSpace.BuildDerivedDataKey(KeyBuilder);
-	}
-}
-
-#endif // WITH_EDITOR
-
 bool UPoseSearchDatabase::GetSkipSearchIfPossible() const
 {
 	if (PoseSearchMode == EPoseSearchMode::PCAKDTree_Validate || PoseSearchMode == EPoseSearchMode::PCAKDTree_Compare)
@@ -1271,68 +1117,6 @@ bool UPoseSearchDatabase::GetSkipSearchIfPossible() const
 	return bSkipSearchIfPossible;
 }
 
-bool UPoseSearchDatabase::IsValidForIndexing() const
-{
-	bool bValid = Schema && Schema->IsValid() && (!Sequences.IsEmpty() || !BlendSpaces.IsEmpty());
-
-	if (bValid)
-	{
-		for (const FPoseSearchDatabaseSequence& DbSequence : Sequences)
-		{
-			if (!DbSequence.Sequence)
-			{
-				bValid = false;
-				break;
-			}
-
-			const USkeleton* SeqSkeleton = DbSequence.Sequence->GetSkeleton();
-			if (!SeqSkeleton)
-			{
-				bValid = false; 
-				break;
-			}
-		}
-
-		for (const FPoseSearchDatabaseBlendSpace& DbBlendSpace : BlendSpaces)
-		{
-			if (!DbBlendSpace.BlendSpace)
-			{
-				bValid = false;
-				break;
-			}
-
-			const USkeleton* SeqSkeleton = DbBlendSpace.BlendSpace->GetSkeleton();
-			if (!SeqSkeleton)
-			{
-				bValid = false;
-				break;
-			}
-		}
-	}
-
-	return bValid;
-}
-
-bool UPoseSearchDatabase::IsValidForSearch() const
-{
-	bool bIsValid = IsValidForIndexing() && GetSearchIndexSafe(false);
-
-#if WITH_EDITOR
-	if (bIsValid && UE::PoseSearch::FAsyncPoseSearchDatabasesManagement::Get().IsBuildingIndex(*this))
-	{
-		bIsValid = false;
-	}
-#endif // WITH_EDITOR
-
-	return bIsValid;
-}
-
-bool UPoseSearchDatabase::IsValidPoseIndex(int32 PoseIdx) const
-{
-	const FPoseSearchIndex* SearchIndex = GetSearchIndex();
-	return SearchIndex ? SearchIndex->IsValidPoseIndex(PoseIdx) : false;
-}
-
 UAnimationAsset* FPoseSearchDatabaseBlendSpace::GetAnimationAsset() const
 {
 	return BlendSpace.Get();
@@ -1396,7 +1180,8 @@ FVector FPoseSearchDatabaseBlendSpace::BlendParameterForSampleRanges(int32 Horiz
 void UPoseSearchDatabase::PostLoad()
 {
 #if WITH_EDITOR
-	UE::PoseSearch::FAsyncPoseSearchDatabasesManagement::Get().RequestAsyncBuildIndex(*this);
+	using namespace UE::PoseSearch;
+	FAsyncPoseSearchDatabasesManagement::RequestAsyncBuildIndex(this, ERequestAsyncBuildFlag::NewRequest);
 #endif
 
 	Super::PostLoad();
@@ -1412,41 +1197,33 @@ void UPoseSearchDatabase::UnregisterOnDerivedDataRebuild(void* Unregister)
 	OnDerivedDataRebuild.RemoveAll(Unregister);
 }
 
-void UPoseSearchDatabase::NotifyDerivedDataRebuild(EDerivedDataBuildState State) const
+void UPoseSearchDatabase::NotifyDerivedDataRebuild() const
 {
-	OnDerivedDataRebuild.Broadcast(State);
+	OnDerivedDataRebuild.Broadcast();
 }
 
 void UPoseSearchDatabase::BeginCacheForCookedPlatformData(const ITargetPlatform* TargetPlatform)
 {
+	using namespace UE::PoseSearch;
 	Super::BeginCacheForCookedPlatformData(TargetPlatform);
-	UE::PoseSearch::FAsyncPoseSearchDatabasesManagement::Get().RequestAsyncBuildIndex(*this, false, true);
+	FAsyncPoseSearchDatabasesManagement::RequestAsyncBuildIndex(this, ERequestAsyncBuildFlag::NewRequest);
 }
 
 bool UPoseSearchDatabase::IsCachedCookedPlatformDataLoaded(const ITargetPlatform* TargetPlatform)
 {
+	using namespace UE::PoseSearch;
 	check(IsInGameThread());
-
-	if (!UE::PoseSearch::FAsyncPoseSearchDatabasesManagement::Get().IsBuildingIndex(*this))
-	{
-		if (IsValidForSearch())
-		{
-			return true;
-		}
-
-		UE::PoseSearch::FAsyncPoseSearchDatabasesManagement::Get().RequestAsyncBuildIndex(*this, false, true);
-	}
-
-	return false;
+	return FAsyncPoseSearchDatabasesManagement::RequestAsyncBuildIndex(this, ERequestAsyncBuildFlag::ContinueRequest);
 }
 #endif // WITH_EDITOR
 
 void UPoseSearchDatabase::PostSaveRoot(FObjectPostSaveRootContext ObjectSaveContext)
 {
 #if WITH_EDITOR
+	using namespace UE::PoseSearch;
 	if (!IsTemplate() && !ObjectSaveContext.IsProceduralSave())
 	{
-		UE::PoseSearch::FAsyncPoseSearchDatabasesManagement::Get().RequestAsyncBuildIndex(*this, true, true);
+		FAsyncPoseSearchDatabasesManagement::RequestAsyncBuildIndex(this, ERequestAsyncBuildFlag::NewRequest | ERequestAsyncBuildFlag::WaitForCompletion);
 	}
 #endif
 
@@ -1461,7 +1238,7 @@ void UPoseSearchDatabase::Serialize(FArchive& Ar)
 	{
 		if (Ar.IsLoading() || Ar.IsCooking())
 		{
-			Ar << PoseSearchIndex;
+			Ar << SearchIndexPrivate;
 		}
 	}
 }
@@ -1473,18 +1250,12 @@ UE::PoseSearch::FSearchResult UPoseSearchDatabase::Search(UE::PoseSearch::FSearc
 	FSearchResult Result;
 
 #if WITH_EDITOR
-	if (UE::PoseSearch::FAsyncPoseSearchDatabasesManagement::Get().IsBuildingIndex(*this))
+	if (!FAsyncPoseSearchDatabasesManagement::RequestAsyncBuildIndex(this, ERequestAsyncBuildFlag::ContinueRequest))
 	{
 		return Result;
 	}
 #endif
 
-	const FPoseSearchIndex* SearchIndex = GetSearchIndexSafe();
-	if (!SearchIndex)
-	{
-		return Result;
-	}
-
 	if (PoseSearchMode == EPoseSearchMode::BruteForce || PoseSearchMode == EPoseSearchMode::PCAKDTree_Compare)
 	{
 		Result = SearchBruteForce(SearchContext);
@@ -1520,11 +1291,10 @@ UE::PoseSearch::FSearchResult UPoseSearchDatabase::SearchPCAKDTree(UE::PoseSearc
 	FSearchResult Result;
 
 	const int32 NumDimensions = Schema->SchemaCardinality;
-	const FPoseSearchIndex* SearchIndex = GetSearchIndex();
-	check(SearchIndex);
+	const FPoseSearchIndex& SearchIndex = GetSearchIndex();
 
 	const uint32 ClampedNumberOfPrincipalComponents = GetNumberOfPrincipalComponents();
-	const uint32 ClampedKDTreeQueryNumNeighbors = FMath::Clamp<uint32>(KDTreeQueryNumNeighbors, 1, SearchIndex->NumPoses);
+	const uint32 ClampedKDTreeQueryNumNeighbors = FMath::Clamp<uint32>(KDTreeQueryNumNeighbors, 1, SearchIndex.NumPoses);
 
 	//stack allocated temporaries
 	TArrayView<size_t> ResultIndexes((size_t*)FMemory_Alloca((ClampedKDTreeQueryNumNeighbors + 1) * sizeof(size_t)), ClampedKDTreeQueryNumNeighbors + 1);
@@ -1536,23 +1306,23 @@ UE::PoseSearch::FSearchResult UPoseSearchDatabase::SearchPCAKDTree(UE::PoseSearc
 	// KDTree in PCA space search
 	if (PoseSearchMode == EPoseSearchMode::PCAKDTree_Validate)
 	{
-		const RowMajorVectorMapConst MapWeightsSqrt(SearchIndex->WeightsSqrt.GetData(), 1, NumDimensions);
+		const RowMajorVectorMapConst MapWeightsSqrt(SearchIndex.WeightsSqrt.GetData(), 1, NumDimensions);
 
 		// testing the KDTree is returning the proper searches for all the original points transformed in pca space
-		for (int32 PoseIdx = 0; PoseIdx < SearchIndex->NumPoses; ++PoseIdx)
+		for (int32 PoseIdx = 0; PoseIdx < SearchIndex.NumPoses; ++PoseIdx)
 		{
 			FKDTree::KNNResultSet ResultSet(ClampedKDTreeQueryNumNeighbors, ResultIndexes, ResultDistanceSqr);
-			TConstArrayView<float> PoseValues = SearchIndex->GetPoseValues(PoseIdx);
+			TConstArrayView<float> PoseValues = SearchIndex.GetPoseValues(PoseIdx);
 
-			const RowMajorVectorMapConst Mean(SearchIndex->Mean.GetData(), 1, NumDimensions);
-			const ColMajorMatrixMapConst PCAProjectionMatrix(SearchIndex->PCAProjectionMatrix.GetData(), NumDimensions, ClampedNumberOfPrincipalComponents);
+			const RowMajorVectorMapConst Mean(SearchIndex.Mean.GetData(), 1, NumDimensions);
+			const ColMajorMatrixMapConst PCAProjectionMatrix(SearchIndex.PCAProjectionMatrix.GetData(), NumDimensions, ClampedNumberOfPrincipalComponents);
 
 			const RowMajorVectorMapConst QueryValues(PoseValues.GetData(), 1, NumDimensions);
 			WeightedQueryValues = QueryValues.array() * MapWeightsSqrt.array();
 			CenteredQueryValues.noalias() = WeightedQueryValues - Mean;
 			ProjectedQueryValues.noalias() = CenteredQueryValues * PCAProjectionMatrix;
 
-			SearchIndex->KDTree.FindNeighbors(ResultSet, ProjectedQueryValues.data());
+			SearchIndex.KDTree.FindNeighbors(ResultSet, ProjectedQueryValues.data());
 
 			size_t ResultIndex = 0;
 			for (; ResultIndex < ResultSet.Num(); ++ResultIndex)
@@ -1577,7 +1347,7 @@ UE::PoseSearch::FSearchResult UPoseSearchDatabase::SearchPCAKDTree(UE::PoseSearc
 	if (!SearchContext.bForceInterrupt && IsCurrentResultFromThisDatabase && SearchContext.bCanAdvance && !Result.ContinuingPoseCost.IsValid())
 	{
 		Result.PoseIdx = SearchContext.CurrentResult.PoseIdx;
-		Result.PoseCost = SearchIndex->ComparePoses(Result.PoseIdx, SearchContext.QueryMirrorRequest, EPoseComparisonFlags::ContinuingPose, Schema->MirrorMismatchCostBias, QueryValues);
+		Result.PoseCost = SearchIndex.ComparePoses(Result.PoseIdx, SearchContext.QueryMirrorRequest, EPoseComparisonFlags::ContinuingPose, Schema->MirrorMismatchCostBias, QueryValues);
 		Result.ContinuingPoseCost = Result.PoseCost;
 
 		if (GetSkipSearchIfPossible())
@@ -1586,9 +1356,9 @@ UE::PoseSearch::FSearchResult UPoseSearchDatabase::SearchPCAKDTree(UE::PoseSearc
 		}
 	}
 
-	// since any PoseCost calculated here is at least SearchIndex->MinCostAddend,
+	// since any PoseCost calculated here is at least SearchIndex.MinCostAddend,
 	// there's no point in performing the search if CurrentBestTotalCost is already better than that
-	if (SearchContext.GetCurrentBestTotalCost() > SearchIndex->MinCostAddend)
+	if (SearchContext.GetCurrentBestTotalCost() > SearchIndex.MinCostAddend)
 	{
 		FNonSelectableIdx NonSelectableIdx;
 		PopulateNonSelectableIdx(NonSelectableIdx, SearchContext, this, QueryValues);
@@ -1596,30 +1366,30 @@ UE::PoseSearch::FSearchResult UPoseSearchDatabase::SearchPCAKDTree(UE::PoseSearc
 
 		check(QueryValues.Num() == NumDimensions);
 
-		const RowMajorVectorMapConst Mean(SearchIndex->Mean.GetData(), 1, NumDimensions);
-		const ColMajorMatrixMapConst PCAProjectionMatrix(SearchIndex->PCAProjectionMatrix.GetData(), NumDimensions, ClampedNumberOfPrincipalComponents);
+		const RowMajorVectorMapConst Mean(SearchIndex.Mean.GetData(), 1, NumDimensions);
+		const ColMajorMatrixMapConst PCAProjectionMatrix(SearchIndex.PCAProjectionMatrix.GetData(), NumDimensions, ClampedNumberOfPrincipalComponents);
 
 		// transforming query values into PCA space to query the KDTree
 		const RowMajorVectorMapConst QueryValuesMap(QueryValues.GetData(), 1, NumDimensions);
-		const RowMajorVectorMapConst MapWeightsSqrt(SearchIndex->WeightsSqrt.GetData(), 1, NumDimensions);
+		const RowMajorVectorMapConst MapWeightsSqrt(SearchIndex.WeightsSqrt.GetData(), 1, NumDimensions);
 		WeightedQueryValues = QueryValuesMap.array() * MapWeightsSqrt.array();
 		CenteredQueryValues.noalias() = WeightedQueryValues - Mean;
 		ProjectedQueryValues.noalias() = CenteredQueryValues * PCAProjectionMatrix;
 
-		SearchIndex->KDTree.FindNeighbors(ResultSet, ProjectedQueryValues.data());
+		SearchIndex.KDTree.FindNeighbors(ResultSet, ProjectedQueryValues.data());
 
 		// NonSelectableIdx are already filtered out inside the kdtree search
-		const FPoseFilters PoseFilters(Schema, TConstArrayView<size_t>(), SearchIndex->OverallFlags);
+		const FPoseFilters PoseFilters(Schema, TConstArrayView<size_t>(), SearchIndex.OverallFlags);
 		for (size_t ResultIndex = 0; ResultIndex < ResultSet.Num(); ++ResultIndex)
 		{
 			const int32 PoseIdx = ResultIndexes[ResultIndex];
-			if (PoseFilters.AreFiltersValid(SearchIndex, QueryValues, PoseIdx, SearchIndex->PoseMetadata[PoseIdx]
+			if (PoseFilters.AreFiltersValid(SearchIndex, QueryValues, PoseIdx, SearchIndex.PoseMetadata[PoseIdx]
 #if UE_POSE_SEARCH_TRACE_ENABLED
 				, SearchContext, this
 #endif
 			))
 			{
-				const FPoseSearchCost PoseCost = SearchIndex->ComparePoses(PoseIdx, SearchContext.QueryMirrorRequest, EPoseComparisonFlags::None, Schema->MirrorMismatchCostBias, QueryValues);
+				const FPoseSearchCost PoseCost = SearchIndex.ComparePoses(PoseIdx, SearchContext.QueryMirrorRequest, EPoseComparisonFlags::None, Schema->MirrorMismatchCostBias, QueryValues);
 				if (PoseCost < Result.PoseCost)
 				{
 					Result.PoseCost = PoseCost;
@@ -1649,7 +1419,7 @@ UE::PoseSearch::FSearchResult UPoseSearchDatabase::SearchPCAKDTree(UE::PoseSearc
 	// finalizing Result properties
 	if (Result.PoseIdx != INDEX_NONE)
 	{
-		Result.AssetTime = SearchIndex->GetAssetTime(Result.PoseIdx, Schema->GetSamplingInterval());
+		Result.AssetTime = SearchIndex.GetAssetTime(Result.PoseIdx, Schema->GetSamplingInterval());
 		Result.Database = this;
 	}
 
@@ -1665,8 +1435,7 @@ UE::PoseSearch::FSearchResult UPoseSearchDatabase::SearchBruteForce(UE::PoseSear
 	
 	FSearchResult Result;
 
-	const FPoseSearchIndex* SearchIndex = GetSearchIndex();
-	check(SearchIndex);
+	const FPoseSearchIndex& SearchIndex = GetSearchIndex();
 
 	SearchContext.GetOrBuildQuery(this, Result.ComposedQuery);
 	TConstArrayView<float> QueryValues = Result.ComposedQuery.GetValues();
@@ -1678,7 +1447,7 @@ UE::PoseSearch::FSearchResult UPoseSearchDatabase::SearchBruteForce(UE::PoseSear
 		if (SearchContext.bCanAdvance && !Result.ContinuingPoseCost.IsValid())
 		{
 			Result.PoseIdx = SearchContext.CurrentResult.PoseIdx;
-			Result.PoseCost = SearchIndex->ComparePoses(Result.PoseIdx, SearchContext.QueryMirrorRequest, EPoseComparisonFlags::ContinuingPose, Schema->MirrorMismatchCostBias, QueryValues);
+			Result.PoseCost = SearchIndex.ComparePoses(Result.PoseIdx, SearchContext.QueryMirrorRequest, EPoseComparisonFlags::ContinuingPose, Schema->MirrorMismatchCostBias, QueryValues);
 			Result.ContinuingPoseCost = Result.PoseCost;
 
 			if (GetSkipSearchIfPossible())
@@ -1688,24 +1457,24 @@ UE::PoseSearch::FSearchResult UPoseSearchDatabase::SearchBruteForce(UE::PoseSear
 		}
 	}
 
-	// since any PoseCost calculated here is at least SearchIndex->MinCostAddend,
+	// since any PoseCost calculated here is at least SearchIndex.MinCostAddend,
 	// there's no point in performing the search if CurrentBestTotalCost is already better than that
-	if (SearchContext.GetCurrentBestTotalCost() > SearchIndex->MinCostAddend)
+	if (SearchContext.GetCurrentBestTotalCost() > SearchIndex.MinCostAddend)
 	{
 		FNonSelectableIdx NonSelectableIdx;
 		PopulateNonSelectableIdx(NonSelectableIdx, SearchContext, this, QueryValues);
 		check(Algo::IsSorted(NonSelectableIdx));
 
-		const FPoseFilters PoseFilters(Schema, NonSelectableIdx, SearchIndex->OverallFlags);
-		for (int32 PoseIdx = 0; PoseIdx < SearchIndex->NumPoses; ++PoseIdx)
+		const FPoseFilters PoseFilters(Schema, NonSelectableIdx, SearchIndex.OverallFlags);
+		for (int32 PoseIdx = 0; PoseIdx < SearchIndex.NumPoses; ++PoseIdx)
 		{
-			if (PoseFilters.AreFiltersValid(SearchIndex, QueryValues, PoseIdx, SearchIndex->PoseMetadata[PoseIdx]
+			if (PoseFilters.AreFiltersValid(SearchIndex, QueryValues, PoseIdx, SearchIndex.PoseMetadata[PoseIdx]
 #if UE_POSE_SEARCH_TRACE_ENABLED
 				, SearchContext, this
 #endif
 			))
 			{
-				const FPoseSearchCost PoseCost = SearchIndex->ComparePoses(PoseIdx, SearchContext.QueryMirrorRequest, EPoseComparisonFlags::None, Schema->MirrorMismatchCostBias, QueryValues);
+				const FPoseSearchCost PoseCost = SearchIndex.ComparePoses(PoseIdx, SearchContext.QueryMirrorRequest, EPoseComparisonFlags::None, Schema->MirrorMismatchCostBias, QueryValues);
 				if (PoseCost < Result.PoseCost)
 				{
 					Result.PoseCost = PoseCost;
@@ -1738,7 +1507,7 @@ UE::PoseSearch::FSearchResult UPoseSearchDatabase::SearchBruteForce(UE::PoseSear
 	// finalizing Result properties
 	if (Result.PoseIdx != INDEX_NONE)
 	{
-		Result.AssetTime = SearchIndex->GetAssetTime(Result.PoseIdx, Schema->GetSamplingInterval());
+		Result.AssetTime = SearchIndex.GetAssetTime(Result.PoseIdx, Schema->GetSamplingInterval());
 		Result.Database = this;
 	}
 
@@ -1766,32 +1535,33 @@ UE::PoseSearch::FSearchResult UPoseSearchDatabaseSet::Search(UE::PoseSearch::FSe
 #endif
 
 	// evaluating the continuing pose before all the active entries
-	if (bEvaluateContinuingPoseFirst && 
+	const UPoseSearchDatabase* Database = SearchContext.CurrentResult.Database.Get();
+	if (bEvaluateContinuingPoseFirst &&
 		SearchContext.bCanAdvance &&
 		!SearchContext.bForceInterrupt &&
-		SearchContext.CurrentResult.IsValid())
+		SearchContext.CurrentResult.IsValid()
+#if WITH_EDITOR
+		&& FAsyncPoseSearchDatabasesManagement::RequestAsyncBuildIndex(Database, ERequestAsyncBuildFlag::ContinueRequest)
+#endif
+		)
 	{
-		const UPoseSearchDatabase* Database = SearchContext.CurrentResult.Database.Get();
 		check(Database);
-		const FPoseSearchIndex* SearchIndex = Database->GetSearchIndexSafe();
-		if (SearchIndex)
-		{
-			SearchContext.GetOrBuildQuery(Database, Result.ComposedQuery);
+		const FPoseSearchIndex& SearchIndex = Database->GetSearchIndex();
+		SearchContext.GetOrBuildQuery(Database, Result.ComposedQuery);
 
-			TConstArrayView<float> QueryValues = Result.ComposedQuery.GetValues();
+		TConstArrayView<float> QueryValues = Result.ComposedQuery.GetValues();
 
-			Result.PoseIdx = SearchContext.CurrentResult.PoseIdx;
-			Result.PoseCost = SearchIndex->ComparePoses(Result.PoseIdx, SearchContext.QueryMirrorRequest, EPoseComparisonFlags::ContinuingPose, Database->Schema->MirrorMismatchCostBias, QueryValues);
-			Result.ContinuingPoseCost = Result.PoseCost;
-			ContinuingCost = Result.PoseCost;
+		Result.PoseIdx = SearchContext.CurrentResult.PoseIdx;
+		Result.PoseCost = SearchIndex.ComparePoses(Result.PoseIdx, SearchContext.QueryMirrorRequest, EPoseComparisonFlags::ContinuingPose, Database->Schema->MirrorMismatchCostBias, QueryValues);
+		Result.ContinuingPoseCost = Result.PoseCost;
+		ContinuingCost = Result.PoseCost;
 
-			Result.AssetTime = SearchIndex->GetAssetTime(Result.PoseIdx, Database->Schema->GetSamplingInterval());
-			Result.Database = Database;
+		Result.AssetTime = SearchIndex.GetAssetTime(Result.PoseIdx, Database->Schema->GetSamplingInterval());
+		Result.Database = Database;
 
-			if (Database->GetSkipSearchIfPossible())
-			{
-				SearchContext.UpdateCurrentBestCost(Result.PoseCost);
-			}
+		if (Database->GetSkipSearchIfPossible())
+		{
+			SearchContext.UpdateCurrentBestCost(Result.PoseCost);
 		}
 	}
 
@@ -2011,25 +1781,22 @@ bool FSearchContext::IsCurrentResultFromDatabase(const UPoseSearchDatabase* Data
 TConstArrayView<float> FSearchContext::GetCurrentResultPrevPoseVector() const
 {
 	check(CurrentResult.IsValid());
-	const FPoseSearchIndex* SearchIndex = CurrentResult.Database->GetSearchIndex();
-	check(SearchIndex);
-	return SearchIndex->GetPoseValues(CurrentResult.PrevPoseIdx);
+	const FPoseSearchIndex& SearchIndex = CurrentResult.Database->GetSearchIndex();
+	return SearchIndex.GetPoseValues(CurrentResult.PrevPoseIdx);
 }
 
 TConstArrayView<float> FSearchContext::GetCurrentResultPoseVector() const
 {
 	check(CurrentResult.IsValid());
-	const FPoseSearchIndex* SearchIndex = CurrentResult.Database->GetSearchIndex();
-	check(SearchIndex);
-	return SearchIndex->GetPoseValues(CurrentResult.PoseIdx);
+	const FPoseSearchIndex& SearchIndex = CurrentResult.Database->GetSearchIndex();
+	return SearchIndex.GetPoseValues(CurrentResult.PoseIdx);
 }
 
 TConstArrayView<float> FSearchContext::GetCurrentResultNextPoseVector() const
 {
 	check(CurrentResult.IsValid());
-	const FPoseSearchIndex* SearchIndex = CurrentResult.Database->GetSearchIndex();
-	check(SearchIndex);
-	return SearchIndex->GetPoseValues(CurrentResult.NextPoseIdx);
+	const FPoseSearchIndex& SearchIndex = CurrentResult.Database->GetSearchIndex();
+	return SearchIndex.GetPoseValues(CurrentResult.NextPoseIdx);
 }
 
 //////////////////////////////////////////////////////////////////////////
@@ -2489,7 +2256,7 @@ const FPoseSearchIndex* FDebugDrawParams::GetSearchIndex() const
 {
 	if (Database)
 	{
-		return Database->GetSearchIndexSafe(false);
+		return &Database->GetSearchIndex();
 	}
 
 	if (SequenceMetaData)
@@ -2527,7 +2294,7 @@ void FSearchResult::Update(float NewAssetTime)
 	}
 	else
 	{
-		const FPoseSearchIndexAsset& SearchIndexAsset = Database->GetSearchIndex()->GetAssetForPose(PoseIdx);
+		const FPoseSearchIndexAsset& SearchIndexAsset = Database->GetSearchIndex().GetAssetForPose(PoseIdx);
 		if (SearchIndexAsset.Type == ESearchIndexAssetType::Sequence)
 		{
 			if (Database->GetPoseIndicesAndLerpValueFromTime(NewAssetTime, SearchIndexAsset, PrevPoseIdx, PoseIdx, NextPoseIdx, LerpValue))
@@ -2592,7 +2359,7 @@ const FPoseSearchIndexAsset* FSearchResult::GetSearchIndexAsset(bool bMandatory)
 		return nullptr;
 	}
 
-	return &Database->GetSearchIndex()->GetAssetForPose(PoseIdx);
+	return &Database->GetSearchIndex().GetAssetForPose(PoseIdx);
 }
 
 
@@ -4501,7 +4268,8 @@ bool BuildIndex(const UPoseSearchDatabase& Database, FPoseSearchIndex& OutSearch
 #endif
 )
 {
-	if (!Database.IsValidForIndexing())
+	// early out for invalid indexing conditions
+	if (!Database.Schema || !Database.Schema->IsValid() || Database.Schema->SchemaCardinality <= 0)
 	{
 		OutSearchIndex.Reset();
 		return false;
diff --git a/Engine/Plugins/Experimental/Animation/PoseSearch/Source/Runtime/Private/PoseSearchDerivedData.cpp b/Engine/Plugins/Experimental/Animation/PoseSearch/Source/Runtime/Private/PoseSearchDerivedData.cpp
--- a/Engine/Plugins/Experimental/Animation/PoseSearch/Source/Runtime/Private/PoseSearchDerivedData.cpp
+++ b/Engine/Plugins/Experimental/Animation/PoseSearch/Source/Runtime/Private/PoseSearchDerivedData.cpp
@@ -27,18 +27,30 @@ namespace UE::PoseSearch
 		});
 #endif
 
+	//////////////////////////////////////////////////////////////////////////
+	// FPoseSearchDatabaseAsyncCacheTask
 	struct FPoseSearchDatabaseAsyncCacheTask
 	{
-		FPoseSearchDatabaseAsyncCacheTask(UPoseSearchDatabase& InDatabase);
-		~FPoseSearchDatabaseAsyncCacheTask();
-
-		void Cancel();
-		void Wait();
-		bool Poll() const;
+		enum class EState
+		{
+			Prestarted,
+			Cancelled,
+			Ended,
+			Failed
+		};
+
+		// these methods MUST be protected by FPoseSearchDatabaseAsyncCacheTask::Mutex! and to make sure we pass the mutex as input param
+		FPoseSearchDatabaseAsyncCacheTask(UPoseSearchDatabase* InDatabase, FCriticalSection& OuterMutex);
+		void StartNewRequestIfNeeded(FCriticalSection& OuterMutex);
+		bool CancelIfDependsOn(const UObject* Object, FCriticalSection& OuterMutex);
+		void Update(FCriticalSection& OuterMutex);
+		void Wait(FCriticalSection& OuterMutex);
+		void Cancel(FCriticalSection& OuterMutex);
+		bool Poll(FCriticalSection& OuterMutex) const;
+		bool ContainsDatabase(const UPoseSearchDatabase* OtherDatabase, FCriticalSection& OuterMutex) const;
 
-		UPoseSearchDatabase& GetDatabase() { return Database; }
-		const UPoseSearchDatabase& GetDatabase() const { return Database; }
-		static FIoHash CreateKey(UPoseSearchDatabase& Database);
+		~FPoseSearchDatabaseAsyncCacheTask();
+		EState GetState() const { return EState(ThreadSafeState.GetValue()); }
 
 	private:
 		FPoseSearchDatabaseAsyncCacheTask(const FPoseSearchDatabaseAsyncCacheTask& Other) = delete;
@@ -47,67 +59,171 @@ namespace UE::PoseSearch
 		FPoseSearchDatabaseAsyncCacheTask& operator=(FPoseSearchDatabaseAsyncCacheTask&& Other) = delete;
 
 		void OnGetComplete(UE::DerivedData::FCacheGetResponse&& Response);
+		void SetState(EState State) { ThreadSafeState.Set(int32(State)); }
 
-		UPoseSearchDatabase& Database;
+		TWeakObjectPtr<UPoseSearchDatabase> Database;
+		// @todo: this is not relevant when the async task is completed, so to save memory we should move it as pointer perhaps
 		FPoseSearchIndex SearchIndex;
 		UE::DerivedData::FRequestOwner Owner;
+		FIoHash DerivedDataKey = FIoHash::Zero;
+		TSet<TWeakObjectPtr<const UObject>> DatabaseDependencies; // @todo: make this const
+		
+		FThreadSafeCounter ThreadSafeState = int32(EState::Prestarted);
+		bool bBroadcastOnDerivedDataRebuild = false;
 	};
 
 	class FPoseSearchDatabaseAsyncCacheTasks : public TArray<TUniquePtr<FPoseSearchDatabaseAsyncCacheTask>> {};
 
-	FPoseSearchDatabaseAsyncCacheTask::FPoseSearchDatabaseAsyncCacheTask(UPoseSearchDatabase& InDatabase)
+	FPoseSearchDatabaseAsyncCacheTask::FPoseSearchDatabaseAsyncCacheTask(UPoseSearchDatabase* InDatabase, FCriticalSection& OuterMutex)
 		: Database(InDatabase)
-		, SearchIndex()
 		, Owner(UE::DerivedData::EPriority::Normal)
+	{
+		StartNewRequestIfNeeded(OuterMutex);
+	}
+
+	FPoseSearchDatabaseAsyncCacheTask::~FPoseSearchDatabaseAsyncCacheTask()
+	{
+		Owner.Cancel();
+	}
+
+	void FPoseSearchDatabaseAsyncCacheTask::StartNewRequestIfNeeded(FCriticalSection& OuterMutex)
 	{
 		using namespace UE::DerivedData;
-		const FIoHash PendingDerivedDataKey = CreateKey(Database);
 
-		Database.NotifyDerivedDataRebuild(UPoseSearchDatabase::EDerivedDataBuildState::Prestarted);
+		FScopeLock Lock(&OuterMutex);
+
+		// making sure there are no active requests
+		Owner.Cancel();
+
+		// composing the key
+		FDerivedDataKeyBuilder KeyBuilder;
+		FGuid VersionGuid = FDevSystemGuids::GetSystemGuid(FDevSystemGuids::Get().POSESEARCHDB_DERIVEDDATA_VER);
+		UPoseSearchDatabase* Db = Database.Get();
+		check(Db);
+		KeyBuilder << VersionGuid;
+		KeyBuilder << Db;
 
-		UE_LOG(LogPoseSearch, Log, TEXT("%s - %s BeginCache"), *LexToString(PendingDerivedDataKey), *Database.GetName());
+		// Stores a BLAKE3-160 hash, taken from the first 20 bytes of a BLAKE3-256 hash
+		const FIoHash NewDerivedDataKey = KeyBuilder.Finalize();
+		const bool bHasKeyChanged = NewDerivedDataKey != DerivedDataKey;
+		if (bHasKeyChanged)
+		{
+			DerivedDataKey = NewDerivedDataKey;
 
-		TArray<FCacheGetRequest> CacheRequests;
-		const FCacheKey CacheKey{ Bucket, PendingDerivedDataKey };
-		CacheRequests.Add({ { Database.GetPathName() }, CacheKey, ECachePolicy::Default });
-		GetCache().Get(CacheRequests, Owner, [this](FCacheGetResponse&& Response)
+			DatabaseDependencies.Reset();
+			for (const UObject* Dependency : KeyBuilder.GetDependencies())
 			{
-				OnGetComplete(MoveTemp(Response));
-			});
+				DatabaseDependencies.Add(Dependency);
+			}
+
+			SetState(EState::Prestarted);
+
+			UE_LOG(LogPoseSearch, Log, TEXT("%s - %s BeginCache"), *LexToString(DerivedDataKey), *Database->GetName());
+
+			TArray<FCacheGetRequest> CacheRequests;
+			const FCacheKey CacheKey{ Bucket, DerivedDataKey };
+			CacheRequests.Add({ { Database->GetPathName() }, CacheKey, ECachePolicy::Default });
+
+			Owner = FRequestOwner(EPriority::Normal);
+			GetCache().Get(CacheRequests, Owner, [this](FCacheGetResponse&& Response)
+				{
+					OnGetComplete(MoveTemp(Response));
+				});
+		}
 	}
 
-	FPoseSearchDatabaseAsyncCacheTask::~FPoseSearchDatabaseAsyncCacheTask()
+	// it cancels and waits for the task to be done and reset the local SearchIndex. SetState to Cancelled
+	void FPoseSearchDatabaseAsyncCacheTask::Cancel(FCriticalSection& OuterMutex)
 	{
-		check(IsInGameThread());
-		check(Poll());
-		if (!Owner.IsCanceled())
-		{
-			// @todo: implement FPoseSearchIndex move ctor and assignment operator and use a MoveTemp(SearchIndex) here
-			Database.PoseSearchIndex = SearchIndex;
+		FScopeLock Lock(&OuterMutex);
 
-			Database.NotifyDerivedDataRebuild(UPoseSearchDatabase::EDerivedDataBuildState::Ended);
-		}
-		else
+		Owner.Cancel();
+		check(GetState() == EState::Prestarted);
+		SearchIndex.Reset();
+		SetState(EState::Cancelled);
+	}
+
+	bool FPoseSearchDatabaseAsyncCacheTask::CancelIfDependsOn(const UObject* Object, FCriticalSection& OuterMutex)
+	{
+		FScopeLock Lock(&OuterMutex);
+
+		// DatabaseDependencies is updated only in StartNewRequestIfNeeded when there are no active requests, so it's thread safe to access it 
+		if (DatabaseDependencies.Contains(Object))
 		{
-			Database.NotifyDerivedDataRebuild(UPoseSearchDatabase::EDerivedDataBuildState::Cancelled);
+			if (GetState() == EState::Prestarted)
+			{
+				Cancel(OuterMutex);
+			}
+			else
+			{
+				SearchIndex.Reset();
+				SetState(EState::Cancelled);
+			}
+			return true;
 		}
+		return false;
 	}
 
-	void FPoseSearchDatabaseAsyncCacheTask::Cancel()
+	void FPoseSearchDatabaseAsyncCacheTask::Update(FCriticalSection& OuterMutex)
 	{
-		Owner.Cancel();
+		check(IsInGameThread());
+
+		FScopeLock Lock(&OuterMutex);
+
+		check(GetState() != EState::Cancelled); // otherwise FPoseSearchDatabaseAsyncCacheTask should have been already removed
+
+		if (GetState() == EState::Prestarted && Poll(OuterMutex))
+		{
+			// task is done: we need to update the state form Prestarted to Ended/Failed
+			Wait(OuterMutex);
+		}
+
+		if (bBroadcastOnDerivedDataRebuild)
+		{
+			Database->NotifyDerivedDataRebuild();
+			bBroadcastOnDerivedDataRebuild = false;
+		}
 	}
 
-	void FPoseSearchDatabaseAsyncCacheTask::Wait()
+	// it waits for the task to be done and SetSearchIndex on the database. SetState to Ended/Failed
+	void FPoseSearchDatabaseAsyncCacheTask::Wait(FCriticalSection& OuterMutex)
 	{
+		check(GetState() == EState::Prestarted);
+
 		Owner.Wait();
+
+		FScopeLock Lock(&OuterMutex);
+
+		const bool bFailedIndexing = SearchIndex.IsEmpty();
+		if (!bFailedIndexing)
+		{
+			Database->SetSearchIndex(SearchIndex); // @todo: implement FPoseSearchIndex move ctor and assignment operator and use a MoveTemp(SearchIndex) here
+			SetState(EState::Ended);
+			bBroadcastOnDerivedDataRebuild = true;
+		}
+		else
+		{
+			check(!bBroadcastOnDerivedDataRebuild);
+			SetState(EState::Failed);
+		}
+		SearchIndex.Reset();
 	}
 
-	bool FPoseSearchDatabaseAsyncCacheTask::Poll() const
+	// true is the task is done executing
+	bool FPoseSearchDatabaseAsyncCacheTask::Poll(FCriticalSection& OuterMutex) const
 	{
 		return Owner.Poll();
 	}
 
+	bool FPoseSearchDatabaseAsyncCacheTask::ContainsDatabase(const UPoseSearchDatabase* OtherDatabase, FCriticalSection& OuterMutex) const
+	{
+		FScopeLock Lock(&OuterMutex);
+		return Database.Get() == OtherDatabase;
+	}
+
+	// called once the task is done:
+	// if EStatus::Ok (data has been retrieved from DDC) we deserialize the payload into the local SearchIndex
+	// if EStatus::Error we BuildIndex and if that's successful we 'Put' it on DDC
 	void FPoseSearchDatabaseAsyncCacheTask::OnGetComplete(UE::DerivedData::FCacheGetResponse&& Response)
 	{
 		using namespace UE::DerivedData;
@@ -123,7 +239,7 @@ namespace UE::PoseSearch
 			FMemoryReaderView Reader(RawData);
 			Reader << SearchIndex;
 
-			UE_LOG(LogPoseSearch, Log, TEXT("%s - %s BuildIndex From Cache"), *LexToString(Key.Hash), *Database.GetName());
+			UE_LOG(LogPoseSearch, Log, TEXT("%s - %s BuildIndex From Cache"), *LexToString(Key.Hash), *Database->GetName());
 
 			COOK_STAT(Timer.AddHit(RawData.GetSize()));
 		}
@@ -135,9 +251,9 @@ namespace UE::PoseSearch
 					COOK_STAT(auto Timer = UsageStats.TimeSyncWork());
 
 					int32 BytesProcessed = 0;
-					if (BuildIndex(Database, SearchIndex, Owner))
+					if (BuildIndex(*Database, SearchIndex, Owner))
 					{
-						UE_LOG(LogPoseSearch, Log, TEXT("%s - %s BuildIndex Succeeded"), *LexToString(Key.Hash), *Database.GetName());
+						UE_LOG(LogPoseSearch, Log, TEXT("%s - %s BuildIndex Succeeded"), *LexToString(Key.Hash), *Database->GetName());
 
 						TArray<uint8> RawBytes;
 						FMemoryWriter Writer(RawBytes);
@@ -147,11 +263,11 @@ namespace UE::PoseSearch
 
 						FCacheRecordBuilder Builder(Key);
 						Builder.AddValue(Id, RawData);
-						GetCache().Put({ { { Database.GetPathName() }, Builder.Build() } }, Owner, [this, Key](FCachePutResponse&& Response)
+						GetCache().Put({ { { Database->GetPathName() }, Builder.Build() } }, Owner, [this, Key](FCachePutResponse&& Response)
 							{
 								if (Response.Status == EStatus::Error)
 								{
-									UE_LOG(LogPoseSearch, Log, TEXT("%s - %s Failed to store DDC"), *LexToString(Key.Hash), *Database.GetName());
+									UE_LOG(LogPoseSearch, Log, TEXT("%s - %s Failed to store DDC"), *LexToString(Key.Hash), *Database->GetName());
 								}
 							});
 					}
@@ -159,11 +275,11 @@ namespace UE::PoseSearch
 					{
 						if (Owner.IsCanceled())
 						{
-							UE_LOG(LogPoseSearch, Log, TEXT("%s - %s BuildIndex Cancelled"), *LexToString(Key.Hash), *Database.GetName());
+							UE_LOG(LogPoseSearch, Log, TEXT("%s - %s BuildIndex Cancelled"), *LexToString(Key.Hash), *Database->GetName());
 						}
 						else
 						{
-							UE_LOG(LogPoseSearch, Log, TEXT("%s - %s BuildIndex Failed"), *LexToString(Key.Hash), *Database.GetName());
+							UE_LOG(LogPoseSearch, Log, TEXT("%s - %s BuildIndex Failed"), *LexToString(Key.Hash), *Database->GetName());
 						}
 						SearchIndex.Reset();
 					}
@@ -174,28 +290,18 @@ namespace UE::PoseSearch
 		else if(Response.Status == EStatus::Canceled)
 		{
 			SearchIndex.Reset();
-			UE_LOG(LogPoseSearch, Log, TEXT("%s - %s BuildIndex Cancelled"), *LexToString(Key.Hash), *Database.GetName());
+			UE_LOG(LogPoseSearch, Log, TEXT("%s - %s BuildIndex Cancelled"), *LexToString(Key.Hash), *Database->GetName());
 		}
 	}
 
-	FIoHash FPoseSearchDatabaseAsyncCacheTask::CreateKey(UPoseSearchDatabase& Database)
-	{
-		using UE::PoseSearch::FDerivedDataKeyBuilder;
-
-		FDerivedDataKeyBuilder KeyBuilder;
-		FGuid VersionGuid = FDevSystemGuids::GetSystemGuid(FDevSystemGuids::Get().POSESEARCHDB_DERIVEDDATA_VER);
-		KeyBuilder << VersionGuid;
-		Database.BuildDerivedDataKey(KeyBuilder);
-
-		// Stores a BLAKE3-160 hash, taken from the first 20 bytes of a BLAKE3-256 hash
-		FIoHash IoHash(KeyBuilder.Finalize());
-		return IoHash;
-	}
-
 	//////////////////////////////////////////////////////////////////////////
 	// FAsyncPoseSearchDatabasesManagement
+	FCriticalSection FAsyncPoseSearchDatabasesManagement::Mutex;
+
 	FAsyncPoseSearchDatabasesManagement& FAsyncPoseSearchDatabasesManagement::Get()
 	{
+		FScopeLock Lock(&Mutex);
+
 		static FAsyncPoseSearchDatabasesManagement SingletonInstance;
 		return SingletonInstance;
 	}
@@ -203,193 +309,64 @@ namespace UE::PoseSearch
 	FAsyncPoseSearchDatabasesManagement::FAsyncPoseSearchDatabasesManagement()
 		: Tasks(*(new FPoseSearchDatabaseAsyncCacheTasks()))
 	{
-		OnObjectPreSaveHandle = FCoreUObjectDelegates::OnObjectPreSave.AddRaw(this, &FAsyncPoseSearchDatabasesManagement::OnObjectPreSave); 
-		OnPreObjectPropertyChangedHandle = FCoreUObjectDelegates::OnPreObjectPropertyChanged.AddRaw(this, &FAsyncPoseSearchDatabasesManagement::OnPreObjectPropertyChanged);
-		OnObjectPropertyChangedHandle = FCoreUObjectDelegates::OnObjectPropertyChanged.AddRaw(this, &FAsyncPoseSearchDatabasesManagement::OnObjectPropertyChanged);
+		FScopeLock Lock(&Mutex);
 
+		OnObjectModifiedHandle = FCoreUObjectDelegates::OnObjectModified.AddRaw(this, &FAsyncPoseSearchDatabasesManagement::OnObjectModified);
 		FCoreDelegates::OnPreExit.AddRaw(this, &FAsyncPoseSearchDatabasesManagement::Shutdown);
 	}
 
 	FAsyncPoseSearchDatabasesManagement::~FAsyncPoseSearchDatabasesManagement()
 	{
+		FScopeLock Lock(&Mutex);
+
 		FCoreDelegates::OnPreExit.RemoveAll(this);
 		Shutdown();
 
 		delete &Tasks;
 	}
 
-	void GetPoseSearchDatabaseAssetDataList(TArray<FAssetData>& OutPoseSearchDatabaseAssetDataList)
-	{
-		FAssetRegistryModule& AssetRegistryModule = FModuleManager::LoadModuleChecked<FAssetRegistryModule>("AssetRegistry");
-		FARFilter Filter;
-		Filter.bRecursiveClasses = true;
-		Filter.ClassPaths.Add(UPoseSearchDatabase::StaticClass()->GetClassPathName());
-
-		OutPoseSearchDatabaseAssetDataList.Reset();
-		AssetRegistryModule.Get().GetAssets(Filter, OutPoseSearchDatabaseAssetDataList);
-	}
-
-	FPoseSearchDatabaseAsyncCacheTask& FAsyncPoseSearchDatabasesManagement::GetTask(int32 TaskIndex)
+	// we're listening to OnObjectModified to cancel any pending Task indexing databases depending from Object to avoid multi threading issues
+	void FAsyncPoseSearchDatabasesManagement::OnObjectModified(UObject* Object)
 	{
-		return *Tasks[TaskIndex].Get();
-	}
+		FScopeLock Lock(&Mutex);
 
-	const FPoseSearchDatabaseAsyncCacheTask& FAsyncPoseSearchDatabasesManagement::GetTask(int32 TaskIndex) const
-	{
-		return *Tasks[TaskIndex].Get();
-	}
-
-	void FAsyncPoseSearchDatabasesManagement::RemoveTask(int32 TaskIndex)
-	{
-		Tasks.RemoveAtSwap(TaskIndex, 1, false);
-	}
-
-	void FAsyncPoseSearchDatabasesManagement::ForEachPoseSearchDatabase(bool bUseTasksDatabases, TFunctionRef<void(UPoseSearchDatabase&)> InFunction)
-	{
-		if (bUseTasksDatabases)
-		{
-			for (int32 TaskIndex = 0; TaskIndex < Tasks.Num(); ++TaskIndex)
-			{
-				UPoseSearchDatabase& PoseSearchDb = GetTask(TaskIndex).GetDatabase();
-				InFunction(PoseSearchDb);
-			}
-		}
-		else
+		// iterating backwards because of the possible RemoveAtSwap
+		for (int32 TaskIndex = Tasks.Num() - 1; TaskIndex >= 0; --TaskIndex)
 		{
-			TArray<FAssetData> PoseSearchDatabaseAssetDataList;
-			GetPoseSearchDatabaseAssetDataList(PoseSearchDatabaseAssetDataList);
-			for (const auto& PoseSearchDbAssetData : PoseSearchDatabaseAssetDataList)
+			if (Tasks[TaskIndex]->CancelIfDependsOn(Object, Mutex))
 			{
-				if (UPoseSearchDatabase* PoseSearchDb = Cast<UPoseSearchDatabase>(PoseSearchDbAssetData.FastGetAsset(false)))
-				{
-					InFunction(*PoseSearchDb);
-				}
+				Tasks.RemoveAtSwap(TaskIndex, 1, false);
 			}
 		}
 	}
 
-	void FAsyncPoseSearchDatabasesManagement::ExecuteIfObjectIsReferencedByDatabase(UObject* Object, bool bUseTasksDatabases, TFunctionRef<void(UPoseSearchDatabase&)> InFunction)
-	{
-		if (UAnimSequence* Sequence = Cast<UAnimSequence>(Object))
-		{
-			ForEachPoseSearchDatabase(bUseTasksDatabases, [Sequence, InFunction](UPoseSearchDatabase& PoseSearchDb)
-				{
-					const bool bSequenceFound = PoseSearchDb.Sequences.ContainsByPredicate([Sequence](FPoseSearchDatabaseSequence& DbSequence)
-						{
-							return Sequence == DbSequence.Sequence || Sequence == DbSequence.LeadInSequence || Sequence == DbSequence.FollowUpSequence;
-						});
-
-					if (bSequenceFound)
-					{
-						InFunction(PoseSearchDb);
-					}
-				});
-		}
-		else if (UBlendSpace* BlendSpace = Cast<UBlendSpace>(Object))
-		{
-			ForEachPoseSearchDatabase(bUseTasksDatabases, [BlendSpace, InFunction](UPoseSearchDatabase& PoseSearchDb)
-				{
-					const bool bBlendSpaceFound = PoseSearchDb.BlendSpaces.ContainsByPredicate([BlendSpace](FPoseSearchDatabaseBlendSpace& DbBlendSpace)
-						{
-							return BlendSpace == DbBlendSpace.BlendSpace;
-						});
-
-					if (bBlendSpaceFound)
-					{
-						InFunction(PoseSearchDb);
-					}
-				});
-		}
-		else if (UPoseSearchSchema* Schema = Cast<UPoseSearchSchema>(Object))
-		{
-			ForEachPoseSearchDatabase(bUseTasksDatabases, [Schema, InFunction](UPoseSearchDatabase& PoseSearchDb)
-				{
-					if (PoseSearchDb.Schema == Schema)
-					{
-						InFunction(PoseSearchDb);
-					}
-				});
-		}
-		else if (USkeleton* Skeleton = Cast<USkeleton>(Object))
-		{
-			ForEachPoseSearchDatabase(bUseTasksDatabases, [Skeleton, InFunction](UPoseSearchDatabase& PoseSearchDb)
-				{
-					if (PoseSearchDb.Schema && PoseSearchDb.Schema->Skeleton == Skeleton)
-					{
-						InFunction(PoseSearchDb);
-					}
-				});
-		}
-		else if (UPoseSearchDatabase* PoseSearchDb = Cast<UPoseSearchDatabase>(Object))
-		{
-			InFunction(*PoseSearchDb);
-		}
-	}
-
-	// @todo: probably overkilling listening to OnObjectPreSave to RequestAsyncBuildIndex, since we already perform it during OnObjectPropertyChanged
-	void FAsyncPoseSearchDatabasesManagement::OnObjectPreSave(UObject* SavedObject, FObjectPreSaveContext SaveContext)
-	{
-		ExecuteIfObjectIsReferencedByDatabase(SavedObject, false, [this](UPoseSearchDatabase& Database)
-			{
-				RequestAsyncBuildIndex(Database, false, true);
-			});
-	}
-
-	// we're listening to OnPreObjectPropertyChanged to cancel any pending Task indexing databases to avoid multi threading issues
-	void FAsyncPoseSearchDatabasesManagement::OnPreObjectPropertyChanged(UObject* Object, const class FEditPropertyChain& PropChain)
-	{
-		ExecuteIfObjectIsReferencedByDatabase(Object, true, [this](UPoseSearchDatabase& Database)
-			{
-				// cancelling the async indexing request for Database
-				WaitOnExistingBuildIndex(Database, false);
-			});
-	}
-
-	// @todo: investigate if it's possible to move the indexing request then databases index gets accessed
-	void FAsyncPoseSearchDatabasesManagement::OnObjectPropertyChanged(UObject* Object, struct FPropertyChangedEvent& Event)
-	{
-		ExecuteIfObjectIsReferencedByDatabase(Object, false, [this](UPoseSearchDatabase& Database)
-			{
-				// requesting a new async indexing for Database (it should already have been cancelled by OnPreObjectPropertyChanged)
-				RequestAsyncBuildIndex(Database, false, false);
-			});
-	}
-
 	void FAsyncPoseSearchDatabasesManagement::Shutdown()
 	{
-		check(IsInGameThread());
-
-		FCoreUObjectDelegates::OnObjectPreSave.Remove(OnObjectPreSaveHandle);
-		OnObjectPreSaveHandle.Reset();
+		FScopeLock Lock(&Mutex);
 
-		FCoreUObjectDelegates::OnPreObjectPropertyChanged.Remove(OnPreObjectPropertyChangedHandle);
-		OnPreObjectPropertyChangedHandle.Reset();
-
-		FCoreUObjectDelegates::OnObjectPropertyChanged.Remove(OnObjectPropertyChangedHandle);
-		OnObjectPropertyChangedHandle.Reset();
+		FCoreUObjectDelegates::OnObjectModified.Remove(OnObjectModifiedHandle);
+		OnObjectModifiedHandle.Reset();
 	}
 
 	void FAsyncPoseSearchDatabasesManagement::Tick(float DeltaTime)
 	{
+		FScopeLock Lock(&Mutex);
+
 		check(IsInGameThread());
 
-		const int32 NumTasks = Tasks.Num();
-		if (NumTasks > 0)
+		// iterating backwards because of the possible RemoveAtSwap 
+		for (int32 TaskIndex = Tasks.Num() - 1; TaskIndex >= 0; --TaskIndex)
 		{
-			// iterating backwards because of the possible RemoveAtSwap 
-			for (int32 TaskIndex = NumTasks - 1; TaskIndex >= 0; --TaskIndex)
-			{
-				FPoseSearchDatabaseAsyncCacheTask& Task = GetTask(TaskIndex);
-				if (Task.Poll())
-				{
-					RemoveTask(TaskIndex);
-				}
-			}
+			Tasks[TaskIndex]->Update(Mutex);
+			
+			// @todo: check key validity every few ticks, or perhaps delete unused for a long time Tasks
 		}
 	}
 
 	void FAsyncPoseSearchDatabasesManagement::TickCook(float DeltaTime, bool bCookCompete)
 	{
+		FScopeLock Lock(&Mutex);
+
 		Tick(DeltaTime);
 	}
 
@@ -402,70 +379,71 @@ namespace UE::PoseSearch
 	{
 	}
 
-	void FAsyncPoseSearchDatabasesManagement::RequestAsyncBuildIndex(UPoseSearchDatabase& Database, bool bWaitForCompletion, bool bCancelPreviousTask)
+	// returns true if the index has been built and the Database updated correctly  
+	bool FAsyncPoseSearchDatabasesManagement::RequestAsyncBuildIndex(const UPoseSearchDatabase* Database, ERequestAsyncBuildFlag Flag)
 	{
-		check(IsInGameThread());
-
-		WaitOnExistingBuildIndex(Database, !bCancelPreviousTask);
-
-#if DO_CHECK
-		for (int32 TaskIndex = 0; TaskIndex < Tasks.Num(); ++TaskIndex)
+		if (!Database)
 		{
-			const FPoseSearchDatabaseAsyncCacheTask& Task = GetTask(TaskIndex);
-			if (&Task.GetDatabase() == &Database)
-			{
-				// making sure we have no tasks associated to Database
-				checkNoEntry();
-			}
+			return false;
 		}
 
-#endif // DO_CHECK
-		Tasks.Emplace(MakeUnique<FPoseSearchDatabaseAsyncCacheTask>(Database));
+		FScopeLock Lock(&Mutex);
 
-		if (bWaitForCompletion)
-		{
-			WaitOnExistingBuildIndex(Database, true);
-		}
-	}
+		check(Database);
+		check(EnumHasAnyFlags(Flag, ERequestAsyncBuildFlag::NewRequest | ERequestAsyncBuildFlag::ContinueRequest));
 
-	void FAsyncPoseSearchDatabasesManagement::WaitOnExistingBuildIndex(const UPoseSearchDatabase& Database, bool bWantResults)
-	{
-		check(IsInGameThread());
+		FAsyncPoseSearchDatabasesManagement& This = FAsyncPoseSearchDatabasesManagement::Get();
 
-		// iterating backwards because of the possible RemoveAtSwap
-		for (int32 TaskIndex = Tasks.Num() - 1; TaskIndex >= 0; --TaskIndex)
+		FPoseSearchDatabaseAsyncCacheTask* Task = nullptr;
+		for (TUniquePtr<FPoseSearchDatabaseAsyncCacheTask>& TaskPtr : This.Tasks)
 		{
-			FPoseSearchDatabaseAsyncCacheTask& Task = GetTask(TaskIndex);
-			if (&Task.GetDatabase() == &Database)
+			if (TaskPtr->ContainsDatabase(Database, Mutex))
 			{
-				if (bWantResults)
+				Task = TaskPtr.Get();
+
+				if (EnumHasAnyFlags(Flag, ERequestAsyncBuildFlag::NewRequest))
 				{
-					Task.Wait();
+					if (Task->GetState() == FPoseSearchDatabaseAsyncCacheTask::EState::Prestarted)
+					{
+						if (EnumHasAnyFlags(Flag, ERequestAsyncBuildFlag::WaitPreviousRequest))
+						{
+							Task->Wait(Mutex);
+						}
+						else
+						{
+							Task->Cancel(Mutex);
+						}
+					}
+
+					Task->StartNewRequestIfNeeded(Mutex);
 				}
-				else
+				else // if (EnumHasAnyFlags(Flag, ERequestAsyncBuildFlag::ContinueRequest))
 				{
-					Task.Cancel();
+					if (Task->GetState() == FPoseSearchDatabaseAsyncCacheTask::EState::Prestarted)
+					{
+						if (EnumHasAnyFlags(Flag, ERequestAsyncBuildFlag::WaitPreviousRequest))
+						{
+							Task->Wait(Mutex);
+						}
+					}
 				}
-
-				RemoveTask(TaskIndex);
+				break;
 			}
 		}
-	}
+		
+		if (!Task)
+		{
+			// we didn't find the Task, so we Emplace a new one
+			This.Tasks.Emplace(MakeUnique<FPoseSearchDatabaseAsyncCacheTask>(const_cast<UPoseSearchDatabase*>(Database), Mutex));
+			Task = This.Tasks.Last().Get();
+		}
 
-	bool FAsyncPoseSearchDatabasesManagement::IsBuildingIndex(const UPoseSearchDatabase& Database) const
-	{
-		for (int32 TaskIndex = 0; TaskIndex < Tasks.Num(); ++TaskIndex)
+		if (EnumHasAnyFlags(Flag, ERequestAsyncBuildFlag::WaitForCompletion) && Task->GetState() == FPoseSearchDatabaseAsyncCacheTask::EState::Prestarted)
 		{
-			const FPoseSearchDatabaseAsyncCacheTask& Task = GetTask(TaskIndex);
-			if (&Task.GetDatabase() == &Database)
-			{
-				if (!Task.Poll())
-				{
-					return true;
-				}
-			}
+			Task->Wait(Mutex);
 		}
-		return false;
+
+		return Task->GetState() == FPoseSearchDatabaseAsyncCacheTask::EState::Ended;
 	}
 
 } // namespace UE::PoseSearch
diff --git a/Engine/Plugins/Experimental/Animation/PoseSearch/Source/Runtime/Private/PoseSearchLibrary.cpp b/Engine/Plugins/Experimental/Animation/PoseSearch/Source/Runtime/Private/PoseSearchLibrary.cpp
--- a/Engine/Plugins/Experimental/Animation/PoseSearch/Source/Runtime/Private/PoseSearchLibrary.cpp
+++ b/Engine/Plugins/Experimental/Animation/PoseSearch/Source/Runtime/Private/PoseSearchLibrary.cpp
@@ -70,7 +70,7 @@ bool FMotionMatchingState::CanAdvance(float DeltaTime, bool& bOutAdvanceToFollow
 					return Entry.Sequence == DbSequence.FollowUpSequence;
 				});
 
-			int32 FollowUpSearchIndexAssetIdx = CurrentSearchResult.Database->GetSearchIndex()->Assets.IndexOfByPredicate(
+			int32 FollowUpSearchIndexAssetIdx = CurrentSearchResult.Database->GetSearchIndex().Assets.IndexOfByPredicate(
 				[&](const FPoseSearchIndexAsset& Entry)
 				{
 					const bool bIsMatch =
@@ -84,7 +84,7 @@ bool FMotionMatchingState::CanAdvance(float DeltaTime, bool& bOutAdvanceToFollow
 			{
 				bOutAdvanceToFollowUpAsset = true;
 
-				const FPoseSearchIndexAsset& FollowUpSearchIndexAsset = CurrentSearchResult.Database->GetSearchIndex()->Assets[FollowUpSearchIndexAssetIdx];
+				const FPoseSearchIndexAsset& FollowUpSearchIndexAsset = CurrentSearchResult.Database->GetSearchIndex().Assets[FollowUpSearchIndexAssetIdx];
 
 				// Follow up asset time will start slightly before the beginning of the sequence as 
 				// this is essentially what the matching time in the corresponding main sequence is.
@@ -427,7 +427,7 @@ void FMotionMatchingState::UpdateWantedPlayRate(const UE::PoseSearch::FSearchCon
 				if (const UPoseSearchFeatureChannel_Trajectory* TrajectoryChannel = CurrentSearchResult.Database->Schema->FindFirstChannelOfType<UPoseSearchFeatureChannel_Trajectory>())
 				{
 					TConstArrayView<float> QueryData = PoseSearchFeatureVectorBuilder->GetValues();
-					TConstArrayView<float> ResultData = CurrentSearchResult.Database->GetSearchIndex()->GetPoseValues(CurrentSearchResult.PoseIdx);
+					TConstArrayView<float> ResultData = CurrentSearchResult.Database->GetSearchIndex().GetPoseValues(CurrentSearchResult.PoseIdx);
 					float EstimatedSpeedRatio = 1.f;
 					if (TrajectoryChannel->GetEstimatedSpeedRatio(QueryData, ResultData, EstimatedSpeedRatio))
 					{
diff --git a/Engine/Plugins/Experimental/Animation/PoseSearch/Source/Runtime/Public/PoseSearch/PoseSearch.h b/Engine/Plugins/Experimental/Animation/PoseSearch/Source/Runtime/Public/PoseSearch/PoseSearch.h
--- a/Engine/Plugins/Experimental/Animation/PoseSearch/Source/Runtime/Public/PoseSearch/PoseSearch.h
+++ b/Engine/Plugins/Experimental/Animation/PoseSearch/Source/Runtime/Public/PoseSearch/PoseSearch.h
@@ -61,7 +61,6 @@ struct FDebugDrawParams;
 struct FSchemaInitializer;
 struct FQueryBuildingContext;
 struct FSearchContext;
-class FDerivedDataKeyBuilder;
 } // namespace UE::PoseSearch
 
 #if WITH_EDITOR
@@ -496,14 +495,14 @@ private:
 
 	friend class ::UPoseSearchSchema;
 
-	UPROPERTY()
+	UPROPERTY(meta = (ExcludeFromHash))
 	int32 ChannelIdx = INDEX_NONE;
 
 protected:
-	UPROPERTY()
+	UPROPERTY(meta = (ExcludeFromHash))
 	int32 ChannelDataOffset = INDEX_NONE;
 
-	UPROPERTY()
+	UPROPERTY(meta = (ExcludeFromHash))
 	int32 ChannelCardinality = INDEX_NONE;
 };
 
@@ -572,11 +571,11 @@ public:
 	UPROPERTY(EditAnywhere, Category = "Schema")
 	EPoseSearchDataPreprocessor DataPreprocessor = EPoseSearchDataPreprocessor::Normalize;
 
-	UPROPERTY()
+	UPROPERTY(meta = (ExcludeFromHash))
 	int32 SchemaCardinality = 0;
 
 	// @todo: used only for indexing: cache it somewhere else
-	UPROPERTY()
+	UPROPERTY(meta = (ExcludeFromHash))
 	TArray<FBoneReference> BoneReferences;
 
 	UPROPERTY(Transient)
@@ -653,19 +652,19 @@ struct POSESEARCH_API FPoseSearchPoseMetadata
 {
 	GENERATED_BODY()
 
-	UPROPERTY()
+	UPROPERTY(meta = (NeverInHash))
 	EPoseSearchPoseFlags Flags = EPoseSearchPoseFlags::None;
 
 	// @todo: consider float16
-	UPROPERTY()
+	UPROPERTY(meta = (NeverInHash))
 	float CostAddend = 0.0f;
 
 	// @todo: consider float16
-	UPROPERTY()
+	UPROPERTY(meta = (NeverInHash))
 	float ContinuingPoseCostAddend = 0.0f;
 
 	// @todo: consider int16
-	UPROPERTY()
+	UPROPERTY(meta = (NeverInHash))
 	int32 AssetIndex = INDEX_NONE;
 };
 
@@ -815,9 +814,6 @@ struct POSESEARCH_API FPoseSearchDatabaseAnimationAssetBase
 	virtual ~FPoseSearchDatabaseAnimationAssetBase() {}
 	virtual UAnimationAsset* GetAnimationAsset() const { return nullptr; }
 	virtual bool IsLooping() const { return false; }
-#if WITH_EDITOR
-	virtual void BuildDerivedDataKey(UE::PoseSearch::FDerivedDataKeyBuilder& KeyBuilder) const {}
-#endif
 };
 
 /** An entry in a UPoseSearchDatabase. */
@@ -854,9 +850,6 @@ struct POSESEARCH_API FPoseSearchDatabaseSequence : public FPoseSearchDatabaseAn
 
 	virtual UAnimationAsset* GetAnimationAsset() const override { return Sequence; }
 	virtual bool IsLooping() const override { return Sequence->bLoop; }
-#if WITH_EDITOR
-	virtual void BuildDerivedDataKey(UE::PoseSearch::FDerivedDataKeyBuilder& KeyBuilder) const override;
-#endif
 };
 
 /** An blend space entry in a UPoseSearchDatabase. */
@@ -887,9 +880,6 @@ struct POSESEARCH_API FPoseSearchDatabaseBlendSpace : public FPoseSearchDatabase
 
 	virtual UAnimationAsset* GetAnimationAsset() const override;
 	virtual bool IsLooping() const override;
-#if WITH_EDITOR
-	virtual void BuildDerivedDataKey(UE::PoseSearch::FDerivedDataKeyBuilder& KeyBuilder) const override;
-#endif
 
 	void GetBlendSpaceParameterSampleRanges(int32& HorizontalBlendNum, int32& VerticalBlendNum) const;
 	FVector BlendParameterForSampleRanges(int32 HorizontalBlendIndex, int32 VerticalBlendIndex) const;
@@ -1115,18 +1105,17 @@ public:
 	UPROPERTY(EditAnywhere, Category = "Performance")
 	bool bSkipSearchIfPossible = true;
 
+private:
 	UPROPERTY(Transient, meta = (ExcludeFromHash))
-	FPoseSearchIndex PoseSearchIndex;
+	FPoseSearchIndex SearchIndexPrivate; // do not use it directly. use GetSearchIndex / SetSearchIndex interact with it and validate that is ok to do so
 
+public:
 	virtual ~UPoseSearchDatabase();
 
-	const FPoseSearchIndex* GetSearchIndex() const;
-	const FPoseSearchIndex* GetSearchIndexSafe(bool bVerboseLogging = true) const;
+	void SetSearchIndex(const FPoseSearchIndex& SearchIndex);
+	const FPoseSearchIndex& GetSearchIndex() const;
 	
 	bool GetSkipSearchIfPossible() const;
-	bool IsValidForIndexing() const;
-	bool IsValidForSearch() const;
-	bool IsValidPoseIndex(int32 PoseIdx) const;
 
 	int32 GetPoseIndexFromTime(float AssetTime, const FPoseSearchIndexAsset& SearchIndexAsset) const;
 	bool GetPoseIndicesAndLerpValueFromTime(float Time, const FPoseSearchIndexAsset& SearchIndexAsset, int32& PrevPoseIdx, int32& PoseIdx, int32& NextPoseIdx, float& LerpValue) const;
@@ -1150,28 +1139,20 @@ public:
 	// End UPoseSearchSearchableAsset
 
 #if WITH_EDITOR
-	void BuildDerivedDataKey(UE::PoseSearch::FDerivedDataKeyBuilder& KeyBuilder) const;
 	virtual void BeginCacheForCookedPlatformData(const ITargetPlatform* TargetPlatform) override;
 	virtual bool IsCachedCookedPlatformDataLoaded(const ITargetPlatform* TargetPlatform) override;
 
-	enum class EDerivedDataBuildState
-	{
-		Prestarted,
-		Ended,
-		Cancelled
-	};
 private:
-	DECLARE_MULTICAST_DELEGATE_OneParam(FOnDerivedDataRebuildMulticaster, EDerivedDataBuildState State);
+	DECLARE_MULTICAST_DELEGATE(FOnDerivedDataRebuildMulticaster);
 	FOnDerivedDataRebuildMulticaster OnDerivedDataRebuild;
 public:
 	typedef FOnDerivedDataRebuildMulticaster::FDelegate FOnDerivedDataRebuild;
 	void RegisterOnDerivedDataRebuild(const FOnDerivedDataRebuild& Delegate);
 	void UnregisterOnDerivedDataRebuild(void* Unregister);
-	void NotifyDerivedDataRebuild(EDerivedDataBuildState State) const;
+	void NotifyDerivedDataRebuild() const;
 #endif // WITH_EDITOR
 
 private:
-	FPoseSearchIndex* GetSearchIndex();
 	UE::PoseSearch::FSearchResult SearchPCAKDTree(UE::PoseSearch::FSearchContext& SearchContext) const;
 	UE::PoseSearch::FSearchResult SearchBruteForce(UE::PoseSearch::FSearchContext& SearchContext) const;
 };
diff --git a/Engine/Plugins/Experimental/Animation/PoseSearch/Source/Runtime/Public/PoseSearch/PoseSearchDerivedData.h b/Engine/Plugins/Experimental/Animation/PoseSearch/Source/Runtime/Public/PoseSearch/PoseSearchDerivedData.h
--- a/Engine/Plugins/Experimental/Animation/PoseSearch/Source/Runtime/Public/PoseSearch/PoseSearchDerivedData.h
+++ b/Engine/Plugins/Experimental/Animation/PoseSearch/Source/Runtime/Public/PoseSearch/PoseSearchDerivedData.h
@@ -2,9 +2,8 @@
 
 #pragma once
 
-#include "PoseSearch/PoseSearch.h"
-
 #if WITH_EDITOR
+#include "PoseSearch/PoseSearch.h"
 #include "TickableEditorObject.h"
 #include "UObject/ObjectSaveContext.h"
 
@@ -13,28 +12,27 @@ namespace UE::PoseSearch
 	struct FPoseSearchDatabaseAsyncCacheTask;
 	class FPoseSearchDatabaseAsyncCacheTasks;
 
+	enum class ERequestAsyncBuildFlag
+	{
+		NewRequest = 1 << 0,			// generates new key and kick off a task to get updated data (it'll cancel for eventual previous Database request, unless WaitPreviousRequest)
+		ContinueRequest = 1 << 1,		// make sure there's associated data to the Database (doesn't have to be up to date)
+		WaitForCompletion = 1 << 2,		// wait the termination of the NewRequest or ContinueRequest
+		WaitPreviousRequest = 1 << 3	// if NewRequest will WaitPreviousRequest
+	};
+	ENUM_CLASS_FLAGS(ERequestAsyncBuildFlag);
+
 	class POSESEARCH_API FAsyncPoseSearchDatabasesManagement : public FTickableEditorObject, public FTickableCookObject, public FGCObject
 	{
 	public:
-		~FAsyncPoseSearchDatabasesManagement();
-		static FAsyncPoseSearchDatabasesManagement& Get();
-		void RequestAsyncBuildIndex(UPoseSearchDatabase& Database, bool bWaitForCompletion = false, bool bCancelPreviousTask = false);
-		void WaitOnExistingBuildIndex(const UPoseSearchDatabase& Database, bool bWantResults);
-		bool IsBuildingIndex(const UPoseSearchDatabase& Database) const;
+		static bool RequestAsyncBuildIndex(const UPoseSearchDatabase* Database, ERequestAsyncBuildFlag Flag);
 
 	private:
 		FAsyncPoseSearchDatabasesManagement();
+		~FAsyncPoseSearchDatabasesManagement();
 
-		void ForEachPoseSearchDatabase(bool bUseTasksDatabases, TFunctionRef<void(UPoseSearchDatabase&)> InFunction);
-		void ExecuteIfObjectIsReferencedByDatabase(UObject* Object, bool bUseTasksDatabases, TFunctionRef<void(UPoseSearchDatabase&)> InFunction);
-
-		void OnObjectPreSave(UObject* SavedObject, FObjectPreSaveContext SaveContext);
-		void OnPreObjectPropertyChanged(UObject* Object, const class FEditPropertyChain& PropChain);
-		void OnObjectPropertyChanged(UObject* Object, struct FPropertyChangedEvent& Event);
+		static FAsyncPoseSearchDatabasesManagement& Get();
 
-		FPoseSearchDatabaseAsyncCacheTask& GetTask(int32 TaskIndex);
-		const FPoseSearchDatabaseAsyncCacheTask& GetTask(int32 TaskIndex) const;
-		void RemoveTask(int32 TaskIndex);
+		void OnObjectModified(UObject* Object);
 
 		void Shutdown();
 		void StartQueuedTasks(int32 MaxActiveTasks);
@@ -53,12 +51,11 @@ namespace UE::PoseSearch
 		void AddReferencedObjects(FReferenceCollector& Collector) override;
 		virtual FString GetReferencerName() const override { return TEXT("FAsyncPoseSearchDatabaseManagement"); }
 		// End FGCObject
-
+		
 		FPoseSearchDatabaseAsyncCacheTasks& Tasks;
-
-		FDelegateHandle OnObjectPreSaveHandle;
-		FDelegateHandle OnObjectPropertyChangedHandle;
-		FDelegateHandle OnPreObjectPropertyChangedHandle;
+		FDelegateHandle OnObjectModifiedHandle;
+		
+		static FCriticalSection Mutex;
 	};
 } // namespace UE::PoseSearch
 
diff --git a/Engine/Plugins/Experimental/Animation/PoseSearch/Source/Runtime/Public/PoseSearch/PoseSearchDerivedDataKey.h b/Engine/Plugins/Experimental/Animation/PoseSearch/Source/Runtime/Public/PoseSearch/PoseSearchDerivedDataKey.h
--- a/Engine/Plugins/Experimental/Animation/PoseSearch/Source/Runtime/Public/PoseSearch/PoseSearchDerivedDataKey.h
+++ b/Engine/Plugins/Experimental/Animation/PoseSearch/Source/Runtime/Public/PoseSearch/PoseSearchDerivedDataKey.h
@@ -2,23 +2,31 @@
 
 #pragma once
 
-#include "Containers/Queue.h"
+#if WITH_EDITOR
+
+#include "Animation/AnimBoneCompressionSettings.h"
+#include "Animation/AnimCurveCompressionSettings.h"
+#include "Animation/MirrorDataTable.h"
 #include "Containers/Set.h"
 #include "Containers/UnrealString.h"
-#include "Serialization/ArchiveUObject.h"
+#include "Factories/FbxAnimSequenceImportData.h"
 #include "Hash/Blake3.h"
-#include "UObject/UnrealType.h"
+#include "Serialization/ArchiveUObject.h"
 #include "UObject/Object.h"
-
-#if WITH_EDITORONLY_DATA
+#include "UObject/UnrealType.h"
 
 namespace UE::PoseSearch
 {
-
+// log properties and UObjects names
 #ifndef UE_POSE_SEARCH_DERIVED_DATA_LOGGING
 	#define UE_POSE_SEARCH_DERIVED_DATA_LOGGING 0
 #endif
 
+// log properties data
+#ifndef UE_POSE_SEARCH_DERIVED_DATA_LOGGING_VERBOSE
+	#define UE_POSE_SEARCH_DERIVED_DATA_LOGGING_VERBOSE 0
+#endif
+
 //////////////////////////////////////////////////////////////////////////
 // FDerivedDataKeyBuilder
 class POSESEARCH_API FDerivedDataKeyBuilder : public FArchiveUObject
@@ -61,11 +69,16 @@ public:
 			return true;
 		}
 
+		if (!InProperty->HasAllPropertyFlags(CPF_Edit)) // bIsEditAnywhereProperty
+		{
+			return true;
+		}
+
 		if (InProperty->HasAllPropertyFlags(CPF_Transient))
 		{
 			return true;
 		}
-		
+
 		if (InProperty->HasMetaData(ExcludeFromHashName))
 		{
 			return true;
@@ -73,9 +86,10 @@ public:
 		
 		check(!InProperty->HasMetaData(NeverInHashName));
 
-#if UE_POSE_SEARCH_DERIVED_DATA_LOGGING
-		UE_LOG(LogPoseSearch, Log, TEXT("    %s"), *InProperty->GetFullName());
-#endif
+		#if UE_POSE_SEARCH_DERIVED_DATA_LOGGING
+		UE_LOG(LogPoseSearch, Log, TEXT("%s - %s"), *GetIndentation(), *InProperty->GetFullName());
+		#endif
+
 		return false;
 	}
 
@@ -83,10 +97,11 @@ public:
 	{
 		const uint8* HasherData = reinterpret_cast<uint8*>(Data);
 
-#if UE_POSE_SEARCH_DERIVED_DATA_LOGGING
+		#if UE_POSE_SEARCH_DERIVED_DATA_LOGGING_VERBOSE
 		FString RawBytesString = BytesToString(HasherData, Length);
-		UE_LOG(LogPoseSearch, Log, TEXT("    %s"), *RawBytesString);
-#endif
+		UE_LOG(LogPoseSearch, Log, TEXT("%s  > %s"), *GetIndentation(), *RawBytesString);
+		#endif
+
 		Hasher.Update(HasherData, Length);
 	}
 
@@ -112,82 +127,109 @@ public:
 
 	virtual FArchive& operator<<(class UObject*& Object) override
 	{
-		if (!RootObject || !Object || !Object->IsIn(RootObject))
-		{
-			auto UniqueName = GetPathNameSafe(Object);
-			*this << UniqueName;
-		}
-		else
+		if (Object)
 		{
-			ObjectsToSerialize.Enqueue(Object);
+			#if UE_POSE_SEARCH_DERIVED_DATA_LOGGING
+			++Indentation;
+			#endif
+
+			bool bAlreadyProcessed = false;
+			ObjectsAlreadySerialized.Add(Object, &bAlreadyProcessed);
+
+			// If we haven't already serialized this object
+			if (bAlreadyProcessed)
+			{
+				#if UE_POSE_SEARCH_DERIVED_DATA_LOGGING
+				UE_LOG(LogPoseSearch, Log, TEXT("%sAlreadyProcessed '%s' (%s)"), *GetIndentation(), *Object->GetName(), *Object->GetClass()->GetName());
+				#endif
+			}
+			// for specific types we only add their names to the hash
+			else if (AddNameOnly(Object))
+			{
+				#if UE_POSE_SEARCH_DERIVED_DATA_LOGGING
+				UE_LOG(LogPoseSearch, Log, TEXT("%sAddingNameOnly '%s' (%s)"), *GetIndentation(), *Object->GetName(), *Object->GetClass()->GetName());
+				#endif
+
+				FString ObjectName = GetFullNameSafe(Object);
+				*this << ObjectName;
+			}
+			else
+			{
+				// Serialize it
+				ObjectBeingSerialized = Object;
+
+				#if UE_POSE_SEARCH_DERIVED_DATA_LOGGING
+				UE_LOG(LogPoseSearch, Log, TEXT("%sBegin '%s' (%s)"), *GetIndentation(), *Object->GetName(), *Object->GetClass()->GetName());
+				#endif
+
+				const_cast<UObject*>(Object)->Serialize(*this);
+
+				#if UE_POSE_SEARCH_DERIVED_DATA_LOGGING
+				UE_LOG(LogPoseSearch, Log, TEXT("%sEnd '%s' (%s)"), *GetIndentation(), *Object->GetName(), *Object->GetClass()->GetName());
+				#endif
+
+				ObjectBeingSerialized = nullptr;
+			}
+			
+			#if UE_POSE_SEARCH_DERIVED_DATA_LOGGING
+			--Indentation;
+			#endif
 		}
+
 		return *this;
 	}
 
-
 	virtual FString GetArchiveName() const override
 	{
 		return TEXT("FDerivedDataKeyBuilder");
 	}
 	//~ End FArchive Interface
-
-	/**
-	* Serialize the given object and update the hash
-	*/
-	void Update(const UObject* Object, const UObject* Root)
+	
+	HashDigestType Finalize()
 	{
-		RootObject = Root;
-		if (Object)
-		{
-			// Start with the given object
-			ObjectsToSerialize.Enqueue(Object);
-
-			// Continue until we no longer have any objects to serialize
-			while (ObjectsToSerialize.Dequeue(Object))
-			{
-				bool bAlreadyProcessed = false;
-				ObjectsAlreadySerialized.Add(Object, &bAlreadyProcessed);
-				// If we haven't already serialized this object
-				if (!bAlreadyProcessed)
-				{
-					// Serialize it
-					ObjectBeingSerialized = Object;
-					const_cast<UObject*>(Object)->Serialize(*this);
-					ObjectBeingSerialized = nullptr;
-				}
-			}
+		return Hasher.Finalize();
+	}
 
-			// Cleanup
-			RootObject = nullptr;
-		}
+	const TSet<const UObject*>& GetDependencies() const
+	{
+		return ObjectsAlreadySerialized;
 	}
 
-	void Update(const UObject* Object)
+protected:
+	/** to keep the key generation lightweight, we hash only the full names for these types */
+	bool AddNameOnly(class UObject* Object) const
 	{
-		Update(Object, Object);
+		return
+			Cast<UAnimBoneCompressionSettings>(Object) ||
+			Cast<UAnimCurveCompressionSettings>(Object) ||
+			Cast<UFbxAnimSequenceImportData>(Object) ||
+			Cast<UMirrorDataTable>(Object) ||
+			Cast<USkeleton>(Object);
 	}
 
-	HashDigestType Finalize()
+#if UE_POSE_SEARCH_DERIVED_DATA_LOGGING
+	FString GetIndentation() const
 	{
-		return Hasher.Finalize();
+		FString IndentationString;
+		for (int32 i = 0; i < Indentation; ++i)
+		{
+			IndentationString.Append(" ");
+		}
+		return IndentationString;
 	}
 
-protected:
-	HashBuilderType Hasher;
+	int32 Indentation = 0;
+#endif
 
-	/** Queue of object references awaiting serialization */
-	TQueue<const UObject*> ObjectsToSerialize;
+	HashBuilderType Hasher;
 
 	/** Set of objects that have already been serialized */
 	TSet<const UObject*> ObjectsAlreadySerialized;
 
 	/** Object currently being serialized */
 	const UObject* ObjectBeingSerialized = nullptr;
-
-	/** Root of object currently being serialized */
-	const UObject* RootObject = nullptr;
 };
 
 } // namespace UE::PoseSearch
 
-#endif // WITH_EDITORONLY_DATA
\ No newline at end of file
+#endif // WITH_EDITOR
\ No newline at end of file
diff --git a/Engine/Plugins/Experimental/Animation/PoseSearch/Source/Runtime/Public/PoseSearch/PoseSearchFeatureChannels.h b/Engine/Plugins/Experimental/Animation/PoseSearch/Source/Runtime/Public/PoseSearch/PoseSearchFeatureChannels.h
--- a/Engine/Plugins/Experimental/Animation/PoseSearch/Source/Runtime/Public/PoseSearch/PoseSearchFeatureChannels.h
+++ b/Engine/Plugins/Experimental/Animation/PoseSearch/Source/Runtime/Public/PoseSearch/PoseSearchFeatureChannels.h
@@ -41,7 +41,7 @@ public:
 	UPROPERTY(EditAnywhere, Category = "Settings")
 	float SampleTimeOffset = 0.f;
 
-	UPROPERTY()
+	UPROPERTY(meta = (ExcludeFromHash))
 	int8 SchemaBoneIdx = 0;
 
 	UPROPERTY(EditAnywhere, Category = "Settings", meta = (ExcludeFromHash))
