motion matching - packaging fix

diff --git a/Engine/Plugins/Experimental/Animation/PoseSearch/Source/Runtime/Private/PoseSearchDerivedData.cpp b/Engine/Plugins/Experimental/Animation/PoseSearch/Source/Runtime/Private/PoseSearchDerivedData.cpp
--- a/Engine/Plugins/Experimental/Animation/PoseSearch/Source/Runtime/Private/PoseSearchDerivedData.cpp
+++ b/Engine/Plugins/Experimental/Animation/PoseSearch/Source/Runtime/Private/PoseSearchDerivedData.cpp
@@ -77,13 +77,18 @@ namespace UE::PoseSearch
 	FPoseSearchDatabaseAsyncCacheTask::FPoseSearchDatabaseAsyncCacheTask(UPoseSearchDatabase* InDatabase, FCriticalSection& OuterMutex)
 		: Database(InDatabase)
 		, Owner(UE::DerivedData::EPriority::Normal)
+		, DerivedDataKey(FIoHash::Zero)
 	{
 		StartNewRequestIfNeeded(OuterMutex);
 	}
 
 	FPoseSearchDatabaseAsyncCacheTask::~FPoseSearchDatabaseAsyncCacheTask()
 	{
+		Database = nullptr;
+		SearchIndex.Reset();
 		Owner.Cancel();
+		DerivedDataKey = FIoHash::Zero;
+		DatabaseDependencies.Reset();
 	}
 
 	void FPoseSearchDatabaseAsyncCacheTask::StartNewRequestIfNeeded(FCriticalSection& OuterMutex)
@@ -138,8 +143,8 @@ namespace UE::PoseSearch
 		FScopeLock Lock(&OuterMutex);
 
 		Owner.Cancel();
-		check(GetState() == EState::Prestarted);
 		SearchIndex.Reset();
+		DerivedDataKey = FIoHash::Zero;
 		SetState(EState::Cancelled);
 	}
 
@@ -150,15 +155,7 @@ namespace UE::PoseSearch
 		// DatabaseDependencies is updated only in StartNewRequestIfNeeded when there are no active requests, so it's thread safe to access it 
 		if (DatabaseDependencies.Contains(Object))
 		{
-			if (GetState() == EState::Prestarted)
-			{
-				Cancel(OuterMutex);
-			}
-			else
-			{
-				SearchIndex.Reset();
-				SetState(EState::Cancelled);
-			}
+			Cancel(OuterMutex);
 			return true;
 		}
 		return false;
