motion matching - deprecating UPoseSearchSequenceMetaData

diff --git a/Engine/Plugins/Experimental/Animation/PoseSearch/Source/Editor/Private/PoseSearchEditor.cpp b/Engine/Plugins/Experimental/Animation/PoseSearch/Source/Editor/Private/PoseSearchEditor.cpp
--- a/Engine/Plugins/Experimental/Animation/PoseSearch/Source/Editor/Private/PoseSearchEditor.cpp
+++ b/Engine/Plugins/Experimental/Animation/PoseSearch/Source/Editor/Private/PoseSearchEditor.cpp
@@ -23,50 +23,11 @@ DEFINE_LOG_CATEGORY(LogPoseSearchEditor);
 
 #define LOCTEXT_NAMESPACE "PoseSearchEditorModule"
 
-//////////////////////////////////////////////////////////////////////////
-// FEditorCommands
-
 namespace UE::PoseSearch
 {
 
-namespace FEditorCommands
-{
-
-void DrawSearchIndex()
-{
-	FDebugDrawParams DrawParams;
-	DrawParams.DefaultLifeTime = 60.0f;
-	EnumAddFlags(DrawParams.Flags, EDebugDrawFlags::DrawSearchIndex);
-
-	TArray<UObject*> EditedAssets = GEditor->GetEditorSubsystem<UAssetEditorSubsystem>()->GetAllEditedAssets();
-	for (UObject* EditedAsset : EditedAssets)
-	{
-		if (UAnimSequence* Sequence = Cast<UAnimSequence>(EditedAsset))
-		{
-			const UPoseSearchSequenceMetaData* MetaData = Sequence->FindMetaDataByClass<UPoseSearchSequenceMetaData>();
-			if (MetaData)
-			{
-				DrawParams.SequenceMetaData = MetaData;
-				IAssetEditorInstance* EditorInstance = GEditor->GetEditorSubsystem<UAssetEditorSubsystem>()->FindEditorForAsset(Sequence, true /*bFocusIfOpen*/);
-				if (EditorInstance && EditorInstance->GetEditorName() == TEXT("AnimationEditor"))
-				{
-					IAnimationEditor* Editor = static_cast<IAnimationEditor*>(EditorInstance);
-					TSharedRef<IPersonaToolkit> Toolkit = Editor->GetPersonaToolkit();
-					TSharedRef<IPersonaPreviewScene> Scene = Toolkit->GetPreviewScene();
-
-					DrawParams.World = Scene->GetWorld();
-					DrawSearchIndex(DrawParams);
-				}
-			}
-		}
-	}
-}
-
-} // namespace FEditorCommands
-
 //////////////////////////////////////////////////////////////////////////
 // FPoseSearchEditorModule
-
 class FEditorModule : public IPoseSearchEditorModuleInterface
 {
 public:
@@ -111,17 +72,8 @@ void FEditorModule::StartupModule()
 		IModularFeatures::Get().RegisterModularFeature("RewindDebuggerViewCreator", DebuggerViewCreator.Get());
 		IModularFeatures::Get().RegisterModularFeature(TraceServices::ModuleFeatureName, TraceModule.Get());
 		
-		ConsoleCommands.Add(IConsoleManager::Get().RegisterConsoleCommand(
-			TEXT("a.PoseSearch.DrawSearchIndex"),
-			TEXT("Draw the search index for the selected asset"),
-			FConsoleCommandDelegate::CreateStatic(&FEditorCommands::DrawSearchIndex),
-			ECVF_Default
-		));
-
 		// Register Ed Mode used by pose search database
-		FEditorModeRegistry::Get().RegisterMode<FDatabaseEdMode>(
-			FDatabaseEdMode::EdModeId,
-			LOCTEXT("PoseSearchDatabaseEdModeName", "PoseSearchDatabase"));
+		FEditorModeRegistry::Get().RegisterMode<FDatabaseEdMode>(FDatabaseEdMode::EdModeId, LOCTEXT("PoseSearchDatabaseEdModeName", "PoseSearchDatabase"));
 
 		// Register UPoseSearchDatabase Type Actions 
 		IAssetTools& AssetTools = FModuleManager::LoadModuleChecked<FAssetToolsModule>("AssetTools").Get();
@@ -138,7 +90,6 @@ void FEditorModule::StartupModule()
 
 	FPropertyEditorModule& PropertyModule = FModuleManager::GetModuleChecked<FPropertyEditorModule>("PropertyEditor");
 	PropertyModule.NotifyCustomizationModuleChanged();
-
 }
 
 void FEditorModule::ShutdownModule()
@@ -179,7 +130,6 @@ void FEditorModule::RegisterCustomClassLayout(FName ClassName, FOnGetDetailCusto
 	PropertyModule.RegisterCustomClassLayout(ClassName, DetailLayoutDelegate);
 }
 
-
 void FEditorModule::RegisterCustomPropertyTypeLayout(FName PropertyTypeName, FOnGetPropertyTypeCustomizationInstance PropertyTypeLayoutDelegate)
 {
 	check(PropertyTypeName != NAME_None);
@@ -191,7 +141,6 @@ void FEditorModule::RegisterCustomPropertyTypeLayout(FName PropertyTypeName, FOn
 	PropertyModule.RegisterCustomPropertyTypeLayout(PropertyTypeName, PropertyTypeLayoutDelegate);
 }
 
-
 void FEditorModule::RegisterPropertyTypeCustomizations()
 {
 	RegisterCustomPropertyTypeLayout("PoseSearchDatabaseSequence", FOnGetPropertyTypeCustomizationInstance::CreateStatic(&FPoseSearchDatabaseSequenceCustomization::MakeInstance));
@@ -230,7 +179,6 @@ void FEditorModule::UnregisterCustomizations()
 	}
 }
 
-
 } // namespace UE::PoseSearch
 
 IMPLEMENT_MODULE(UE::PoseSearch::FEditorModule, PoseSearchEditor);
diff --git a/Engine/Plugins/Experimental/Animation/PoseSearch/Source/Runtime/Private/PoseSearch.cpp b/Engine/Plugins/Experimental/Animation/PoseSearch/Source/Runtime/Private/PoseSearch.cpp
--- a/Engine/Plugins/Experimental/Animation/PoseSearch/Source/Runtime/Private/PoseSearch.cpp
+++ b/Engine/Plugins/Experimental/Animation/PoseSearch/Source/Runtime/Private/PoseSearch.cpp
@@ -14,7 +14,6 @@
 #include "Templates/Less.h"
 #include "Features/IModularFeatures.h"
 #include "DrawDebugHelpers.h"
-#include "Animation/AnimPoseSearchProvider.h"
 #include "Animation/AnimCurveTypes.h"
 #include "Animation/AnimInstance.h"
 #include "Animation/AnimInstanceProxy.h"
@@ -442,18 +441,17 @@ void UPoseSearchFeatureChannel::ComputeMeanDeviations(const Eigen::MatrixXd& Cen
 }
 
 //////////////////////////////////////////////////////////////////////////
-// UPoseSearchSchema
-
+// FSchemaInitializer
 namespace UE::PoseSearch
 {
-
 int32 FSchemaInitializer::AddBoneReference(const FBoneReference& BoneReference)
 {
 	return BoneReferences.AddUnique(BoneReference);
 }
-
 } // namespace UE::PoseSearch
 
+//////////////////////////////////////////////////////////////////////////
+// UPoseSearchSchema
 void UPoseSearchSchema::Finalize(bool bRemoveEmptyChannels)
 {
 	using namespace UE::PoseSearch;
@@ -564,7 +562,6 @@ void UPoseSearchSchema::ResolveBoneReferences()
 	}
 }
 
-
 bool UPoseSearchSchema::BuildQuery(UE::PoseSearch::FSearchContext& SearchContext, FPoseSearchFeatureVectorBuilder& InOutQuery) const
 {
 	QUICK_SCOPE_CYCLE_COUNTER(STAT_PoseSearch_BuildQuery);
@@ -660,10 +657,6 @@ float FPoseSearchIndex::GetAssetTime(int32 PoseIdx, float SamplingInterval) cons
 	return -1.0f;
 }
 
-
-//////////////////////////////////////////////////////////////////////////
-// FPoseSearchIndex
-
 bool FPoseSearchIndex::IsEmpty() const
 {
 	const bool bEmpty = Assets.Num() == 0 || NumPoses == 0;
@@ -801,7 +794,6 @@ void FPoseSearchIndex::InitSearchIndexAssets(TConstArrayView<FPoseSearchDatabase
 	}
 }
 
-
 FArchive& operator<<(FArchive& Ar, FPoseSearchIndex& Index)
 {
 	int32 NumValues = 0;
@@ -858,95 +850,6 @@ FArchive& operator<<(FArchive& Ar, FPoseSearchIndex& Index)
 	return Ar;
 }
 
-//////////////////////////////////////////////////////////////////////////
-// UPoseSearchSequenceMetaData
-
-void UPoseSearchSequenceMetaData::PreSave(FObjectPreSaveContext ObjectSaveContext)
-{
-	SearchIndex.Reset();
-
-#if WITH_EDITOR
-	if (!IsTemplate())
-	{
-		if (IsValidForIndexing())
-		{
-			UObject* Outer = GetOuter();
-			if (UAnimSequence* Sequence = Cast<UAnimSequence>(Outer))
-			{
-				UE::PoseSearch::BuildIndex(Sequence, this);
-			}
-		}
-	}
-#endif
-
-	Super::PreSave(ObjectSaveContext);
-}
-
-bool UPoseSearchSequenceMetaData::IsValidForIndexing() const
-{
-	return Schema && Schema->IsValid() && UE::PoseSearch::IsSamplingRangeValid(SamplingRange);
-}
-
-bool UPoseSearchSequenceMetaData::IsValidForSearch() const
-{
-	return IsValidForIndexing() && !SearchIndex.IsEmpty();
-}
-
-UE::PoseSearch::FSearchResult UPoseSearchSequenceMetaData::Search(UE::PoseSearch::FSearchContext& SearchContext) const
-{
-	using namespace UE::PoseSearch;
-
-	FSearchResult Result;
-
-	if (!ensure(Schema->IsValid() && !SearchIndex.IsEmpty()))
-	{
-		return Result;
-	}
-
-	Schema->BuildQuery(SearchContext, Result.ComposedQuery);
-	TConstArrayView<float> QueryValues = Result.ComposedQuery.GetValues();
-
-	if (!ensure(QueryValues.Num() == Schema->SchemaCardinality))
-	{
-		return Result;
-	}
-
-	for (const FPoseSearchIndexAsset& Asset : SearchIndex.Assets)
-	{
-		const int32 EndIndex = Asset.FirstPoseIdx + Asset.NumPoses;
-		for (int32 PoseIdx = Asset.FirstPoseIdx; PoseIdx < EndIndex; ++PoseIdx)
-		{
-			const FPoseSearchPoseMetadata& Metadata = SearchIndex.PoseMetadata[PoseIdx];
-
-			if (EnumHasAnyFlags(Metadata.Flags, EPoseSearchPoseFlags::BlockTransition))
-			{
-				continue;
-			}
-
-			const FPoseSearchCost PoseCost = SearchIndex.ComparePoses(PoseIdx, EPoseSearchBooleanRequest::Indifferent, EPoseComparisonFlags::ContinuingPose, Schema->MirrorMismatchCostBias, QueryValues);
-			if (PoseCost < Result.PoseCost)
-			{
-				Result.PoseCost = PoseCost;
-				Result.PoseIdx = PoseIdx;
-			}
-		}
-	}
-
-	if (Result.PoseIdx != INDEX_NONE)
-	{
-		Result.AssetTime = SearchIndex.GetAssetTime(Result.PoseIdx, Schema->GetSamplingInterval());
-	}
-
-#if ENABLE_DRAW_DEBUG
-	DrawFeatureVector(SearchContext.DebugDrawParams, Result.PoseIdx);
-	
-	EnumAddFlags(SearchContext.DebugDrawParams.Flags, EDebugDrawFlags::DrawQuery);
-	DrawFeatureVector(SearchContext.DebugDrawParams, QueryValues);
-#endif // ENABLE_DRAW_DEBUG
-
-	return Result;
-}
-
 //////////////////////////////////////////////////////////////////////////
 // UPoseSearchDatabase
 
@@ -2254,32 +2157,12 @@ FColor FDebugDrawParams::GetColor(int32 ColorPreset) const
 
 const FPoseSearchIndex* FDebugDrawParams::GetSearchIndex() const
 {
-	if (Database)
-	{
-		return &Database->GetSearchIndex();
-	}
-
-	if (SequenceMetaData)
-	{
-		return &SequenceMetaData->SearchIndex;
-	}
-
-	return nullptr;
+	return Database ? &Database->GetSearchIndex() : nullptr;
 }
 
 const UPoseSearchSchema* FDebugDrawParams::GetSchema() const
 {
-	if (Database)
-	{
-		return Database->Schema;
-	}
-
-	if (SequenceMetaData)
-	{
-		return SequenceMetaData->Schema;
-	}
-
-	return nullptr;
+	return Database ? Database->Schema : nullptr;
 }
 
 //////////////////////////////////////////////////////////////////////////
@@ -3986,70 +3869,6 @@ static void PreprocessSearchIndex(FPoseSearchIndex& SearchIndex, const UPoseSear
 	PreprocessSearchIndexKDTree(SearchIndex, Schema->SchemaCardinality, NumberOfPrincipalComponents, PoseSearchMode, KDTreeMaxLeafSize, KDTreeQueryNumNeighbors);
 }
 
-bool BuildIndex(const UAnimSequence* Sequence, UPoseSearchSequenceMetaData* SequenceMetaData)
-{
-	check(Sequence);
-	check(SequenceMetaData);
-
-	if (!SequenceMetaData->IsValidForIndexing())
-	{
-		return false;
-	}
-
-	USkeleton* SeqSkeleton = Sequence->GetSkeleton();
-	if (!SeqSkeleton)
-	{
-		return false;
-	}
-
-	FBoneContainer BoneContainer;
-	BoneContainer.InitializeTo(SequenceMetaData->Schema->BoneIndicesWithParents, FCurveEvaluationOption(false), *SequenceMetaData->Schema->Skeleton);
-	FAssetSamplingContext SamplingContext;
-	SamplingContext.Init(SequenceMetaData->Schema->MirrorDataTable, BoneContainer);
-
-	FSequenceSampler Sampler;
-	FSequenceSampler::FInput SamplerInput;
-	SamplerInput.ExtrapolationParameters = SequenceMetaData->ExtrapolationParameters;
-	SamplerInput.Sequence = Sequence;
-	Sampler.Init(SamplerInput);
-	Sampler.Process();
-
-	FAssetIndexer Indexer;
-	FAssetIndexingContext IndexerContext;
-	IndexerContext.SamplingContext = &SamplingContext;
-	IndexerContext.MainSampler = &Sampler;
-	IndexerContext.Schema = SequenceMetaData->Schema;
-	IndexerContext.RequestedSamplingRange = GetEffectiveSamplingRange(Sequence, SequenceMetaData->SamplingRange);
-
-	Indexer.Init(IndexerContext, BoneContainer);
-	if (!Indexer.Process())
-	{
-		return false;
-	}
-
-	SequenceMetaData->SearchIndex.Assets.Empty();
-	FPoseSearchIndexAsset SearchIndexAsset;
-	SearchIndexAsset.SourceAssetIdx = 0;
-	SearchIndexAsset.FirstPoseIdx = 0;
-	SearchIndexAsset.NumPoses = Indexer.Output.NumIndexedPoses;
-	SearchIndexAsset.SamplingInterval = IndexerContext.RequestedSamplingRange;
-
-	SequenceMetaData->SearchIndex.Values = Indexer.Output.FeatureVectorTable;
-	SequenceMetaData->SearchIndex.NumPoses = Indexer.Output.NumIndexedPoses;
-	SequenceMetaData->SearchIndex.Assets.Add(SearchIndexAsset);
-	SequenceMetaData->SearchIndex.PoseMetadata = Indexer.Output.PoseMetadata;
-
-	SequenceMetaData->SearchIndex.OverallFlags = EPoseSearchPoseFlags::None;
-	for (const FPoseSearchPoseMetadata& PoseMetadata : Indexer.Output.PoseMetadata)
-	{
-		SequenceMetaData->SearchIndex.OverallFlags = PoseMetadata.Flags;
-	}
-
-	// todo: do we need to PreprocessSearchIndex?
-	// PreprocessSearchIndex(&SequenceMetaData->SearchIndex, Database);
-	return true;
-}
-
 struct FDatabaseIndexingContext
 {
 	FPoseSearchIndex* SearchIndex = nullptr;
@@ -4332,60 +4151,13 @@ bool BuildIndex(const UPoseSearchDatabase& Database, FPoseSearchIndex& OutSearch
 
 //////////////////////////////////////////////////////////////////////////
 // FModule
-
-class FModule : public IModuleInterface, public UE::Anim::IPoseSearchProvider
+// @todo: clean up this code
+class FModule : public IModuleInterface
 {
 public: // IModuleInterface
-	virtual void StartupModule() override;
-	virtual void ShutdownModule() override;
-
-public: // IPoseSearchProvider
-	virtual UE::Anim::IPoseSearchProvider::FSearchResult Search(const FAnimationBaseContext& GraphContext, const UAnimSequenceBase* Sequence) override;
+	virtual void StartupModule() override {}
+	virtual void ShutdownModule() override {}
 };
-
-void FModule::StartupModule()
-{
-	IModularFeatures::Get().RegisterModularFeature(UE::Anim::IPoseSearchProvider::ModularFeatureName, this);
-}
-
-void FModule::ShutdownModule()
-{
-	IModularFeatures::Get().UnregisterModularFeature(UE::Anim::IPoseSearchProvider::ModularFeatureName, this);
-}
-
-UE::Anim::IPoseSearchProvider::FSearchResult FModule::Search(const FAnimationBaseContext& GraphContext, const UAnimSequenceBase* Sequence)
-{
-	UE::Anim::IPoseSearchProvider::FSearchResult ProviderResult;
-
-	using namespace UE::PoseSearch;
-
-	const UPoseSearchSequenceMetaData* MetaData = Sequence ? Sequence->FindMetaDataByClass<UPoseSearchSequenceMetaData>() : nullptr;
-	if (!MetaData || !MetaData->IsValidForSearch())
-	{
-		return ProviderResult;
-	}
-
-	IPoseHistoryProvider* PoseHistoryProvider = GraphContext.GetMessage<IPoseHistoryProvider>();
-	if (!PoseHistoryProvider)
-	{
-		return ProviderResult;
-	}
-
-	FPoseHistory& PoseHistory = PoseHistoryProvider->GetPoseHistory();
-
-	FSearchContext SearchContext;
-	SearchContext.OwningComponent = GraphContext.AnimInstanceProxy->GetSkelMeshComponent();
-	SearchContext.BoneContainer = &GraphContext.AnimInstanceProxy->GetRequiredBones();
-	SearchContext.History = &PoseHistoryProvider->GetPoseHistory();
-
-	UE::PoseSearch::FSearchResult Result = MetaData->Search(SearchContext);
-
-	ProviderResult.Dissimilarity = Result.PoseCost.GetTotalCost();
-	ProviderResult.PoseIdx = Result.PoseIdx;
-	ProviderResult.TimeOffsetSeconds = Result.AssetTime;
-	return ProviderResult;
-}
-
 } // namespace UE::PoseSearch
 
 #undef LOCTEXT_NAMESPACE
diff --git a/Engine/Plugins/Experimental/Animation/PoseSearch/Source/Runtime/Public/PoseSearch/PoseSearch.h b/Engine/Plugins/Experimental/Animation/PoseSearch/Source/Runtime/Public/PoseSearch/PoseSearch.h
--- a/Engine/Plugins/Experimental/Animation/PoseSearch/Source/Runtime/Public/PoseSearch/PoseSearch.h
+++ b/Engine/Plugins/Experimental/Animation/PoseSearch/Source/Runtime/Public/PoseSearch/PoseSearch.h
@@ -724,8 +724,7 @@ public:
 };
 
 /**
-* A search index for animation poses. The structure of the search index is determined by its UPoseSearchSchema.
-* May represent a single animation (see UPoseSearchSequenceMetaData) or a collection (see UPoseSearchDatabase).
+* A search index for animation poses for a UPoseSearchDatabase. The structure of the search index is determined by its UPoseSearchSchema.
 */
 USTRUCT()
 struct POSESEARCH_API FPoseSearchIndex
@@ -991,7 +990,6 @@ struct POSESEARCH_API FDebugDrawParams
 {
 	const UWorld* World = nullptr;
 	const UPoseSearchDatabase* Database = nullptr;
-	const UPoseSearchSequenceMetaData* SequenceMetaData = nullptr;
 	EDebugDrawFlags Flags = EDebugDrawFlags::None;
 	uint32 ChannelMask = (uint32)-1;
 
@@ -1157,37 +1155,6 @@ private:
 	UE::PoseSearch::FSearchResult SearchBruteForce(UE::PoseSearch::FSearchContext& SearchContext) const;
 };
 
-//////////////////////////////////////////////////////////////////////////
-// Sequence metadata
-
-/** Animation metadata object for indexing a single animation. */
-UCLASS(BlueprintType, Category = "Animation|Pose Search", Experimental)
-class POSESEARCH_API UPoseSearchSequenceMetaData : public UAnimMetaData
-{
-	GENERATED_BODY()
-public:
-
-	UPROPERTY(EditAnywhere, Category = "Settings")
-	TObjectPtr<const UPoseSearchSchema> Schema;
-
-	UPROPERTY(EditAnywhere, Category = "Settings")
-	FFloatInterval SamplingRange = FFloatInterval(0.f, 0.f);
-
-	UPROPERTY(EditAnywhere, Category = "Settings")
-	FPoseSearchExtrapolationParameters ExtrapolationParameters;
-
-	UPROPERTY(meta = (ExcludeFromHash))
-	FPoseSearchIndex SearchIndex;
-
-	bool IsValidForIndexing() const;
-	bool IsValidForSearch() const;
-
-	UE::PoseSearch::FSearchResult Search(UE::PoseSearch::FSearchContext& SearchContext) const;
-
-public: // UObject
-	virtual void PreSave(FObjectPreSaveContext ObjectSaveContext) override;
-};
-
 //////////////////////////////////////////////////////////////////////////
 // Feature vector reader and builder
 
@@ -1439,16 +1406,6 @@ POSESEARCH_API void DrawFeatureVector(const FDebugDrawParams& DrawParams, int32
 POSESEARCH_API void DrawSearchIndex(const FDebugDrawParams& DrawParams);
 POSESEARCH_API void CompareFeatureVectors(TConstArrayView<float> A, TConstArrayView<float> B, TConstArrayView<float> WeightsSqrt, TArrayView<float> Result);
 
-/**
-* Creates a pose search index for an animation sequence
-* 
-* @param Sequence			The input sequence create a search index for
-* @param SequenceMetaData	The input sequence indexing info and output search index
-* 
-* @return Whether the index was built successfully
-*/
-POSESEARCH_API bool BuildIndex(const UAnimSequence* Sequence, UPoseSearchSequenceMetaData* SequenceMetaData);
-
 /**
 * Creates a pose search index for a collection of animations
 * 
diff --git a/Engine/Source/Editor/AnimGraph/Private/AnimGraphNode_SequencePlayer.cpp b/Engine/Source/Editor/AnimGraph/Private/AnimGraphNode_SequencePlayer.cpp
--- a/Engine/Source/Editor/AnimGraph/Private/AnimGraphNode_SequencePlayer.cpp
+++ b/Engine/Source/Editor/AnimGraph/Private/AnimGraphNode_SequencePlayer.cpp
@@ -16,7 +16,6 @@
 #include "EditorCategoryUtils.h"
 #include "Animation/AnimComposite.h"
 #include "Animation/AnimSequence.h"
-#include "Animation/AnimPoseSearchProvider.h"
 #include "Animation/AnimRootMotionProvider.h"
 #include "UObject/UE5MainStreamObjectVersion.h"
 #include "IAnimBlueprintNodeOverrideAssetsContext.h"
@@ -252,16 +251,6 @@ UScriptStruct* UAnimGraphNode_SequencePlayer::GetTimePropertyStruct() const
 	return FAnimNode_SequencePlayer::StaticStruct();
 }
 
-void UAnimGraphNode_SequencePlayer::CustomizeDetails(IDetailLayoutBuilder& DetailBuilder)
-{
-	Super::CustomizeDetails(DetailBuilder);
-
-	if (!UE::Anim::IPoseSearchProvider::IsAvailable())
-	{
-		DetailBuilder.HideCategory(TEXT("PoseMatching"));
-	}
-}
-
 void UAnimGraphNode_SequencePlayer::CustomizePinData(UEdGraphPin* Pin, FName SourcePropertyName, int32 ArrayIndex) const
 {
 	Super::CustomizePinData(Pin, SourcePropertyName, ArrayIndex);
diff --git a/Engine/Source/Editor/AnimGraph/Public/AnimGraphNode_SequencePlayer.h b/Engine/Source/Editor/AnimGraph/Public/AnimGraphNode_SequencePlayer.h
--- a/Engine/Source/Editor/AnimGraph/Public/AnimGraphNode_SequencePlayer.h
+++ b/Engine/Source/Editor/AnimGraph/Public/AnimGraphNode_SequencePlayer.h
@@ -31,7 +31,6 @@ class UAnimGraphNode_SequencePlayer : public UAnimGraphNode_AssetPlayerBase
 	// End of UEdGraphNode interface
 
 	// UAnimGraphNode_Base interface
-	virtual void CustomizeDetails(IDetailLayoutBuilder& DetailBuilder) override;
 	virtual void CustomizePinData(UEdGraphPin* Pin, FName SourcePropertyName, int32 ArrayIndex) const override;
 	virtual void ValidateAnimNodeDuringCompilation(class USkeleton* ForSkeleton, class FCompilerResultsLog& MessageLog) override;
 	virtual void PreloadRequiredAssets() override;		
diff --git a/Engine/Source/Runtime/Engine/Classes/Animation/AnimNode_SequencePlayer.h b/Engine/Source/Runtime/Engine/Classes/Animation/AnimNode_SequencePlayer.h
--- a/Engine/Source/Runtime/Engine/Classes/Animation/AnimNode_SequencePlayer.h
+++ b/Engine/Source/Runtime/Engine/Classes/Animation/AnimNode_SequencePlayer.h
@@ -38,6 +38,7 @@ public:
 	// End of FAnimNode_Base interface
 
 	float GetTimeFromEnd(float CurrentNodeTime) const;
+	UE_DEPRECATED(5.1, "GetEffectiveStartPosition is no longer supported. Please use GetStartPosition instead")
 	float GetEffectiveStartPosition(const FAnimationBaseContext& Context) const;
 
 	// The animation sequence asset to play
diff --git a/Engine/Source/Runtime/Engine/Private/Animation/AnimNode_SequencePlayer.cpp b/Engine/Source/Runtime/Engine/Private/Animation/AnimNode_SequencePlayer.cpp
--- a/Engine/Source/Runtime/Engine/Private/Animation/AnimNode_SequencePlayer.cpp
+++ b/Engine/Source/Runtime/Engine/Private/Animation/AnimNode_SequencePlayer.cpp
@@ -6,7 +6,6 @@
 #include "AnimEncoding.h"
 #include "Animation/AnimInstanceProxy.h"
 #include "Animation/AnimTrace.h"
-#include "Animation/AnimPoseSearchProvider.h"
 
 #include UE_INLINE_GENERATED_CPP_BY_NAME(AnimNode_SequencePlayer)
 
@@ -61,14 +60,12 @@ void FAnimNode_SequencePlayerBase::Initialize_AnyThread(const FAnimationInitiali
 
 	if (CurrentSequence != nullptr)
 	{
-		const float EffectiveStartPosition = GetEffectiveStartPosition(Context);
 		const float CurrentPlayRate = GetPlayRate();
 		const float CurrentPlayRateBasis = GetPlayRateBasis();
 
-		InternalTimeAccumulator = FMath::Clamp(EffectiveStartPosition, 0.f, CurrentSequence->GetPlayLength());
 		const float AdjustedPlayRate = PlayRateScaleBiasClampState.ApplyTo(GetPlayRateScaleBiasClampConstants(), FMath::IsNearlyZero(CurrentPlayRateBasis) ? 0.f : (CurrentPlayRate / CurrentPlayRateBasis), 0.f);
 		const float EffectivePlayrate = CurrentSequence->RateScale * AdjustedPlayRate;
-		if ((EffectiveStartPosition == 0.f) && (EffectivePlayrate < 0.f))
+		if ((InternalTimeAccumulator == 0.f) && (EffectivePlayrate < 0.f))
 		{
 			InternalTimeAccumulator = CurrentSequence->GetPlayLength();
 		}
@@ -157,21 +154,6 @@ float FAnimNode_SequencePlayerBase::GetTimeFromEnd(float CurrentNodeTime) const
 
 float FAnimNode_SequencePlayerBase::GetEffectiveStartPosition(const FAnimationBaseContext& Context) const
 {
-	// Override the start position if pose matching is enabled
-	UAnimSequenceBase* CurrentSequence = GetSequence();
-	if (CurrentSequence != nullptr && GetStartFromMatchingPose())
-	{
-		UE::Anim::IPoseSearchProvider* PoseSearchProvider = UE::Anim::IPoseSearchProvider::Get();
-		if (PoseSearchProvider)
-		{
-			UE::Anim::IPoseSearchProvider::FSearchResult Result = PoseSearchProvider->Search(Context, CurrentSequence);
-			if (Result.PoseIdx >= 0)
-			{
-				return Result.TimeOffsetSeconds;
-			}
-		}
-	}
-
 	return GetStartPosition();
 }
 
diff --git a/Engine/Source/Runtime/Engine/Private/Animation/AnimPoseSearchProvider.cpp b/Engine/Source/Runtime/Engine/Private/Animation/AnimPoseSearchProvider.cpp
deleted file mode 100644
diff --git a/Engine/Source/Runtime/Engine/Public/Animation/AnimPoseSearchProvider.h b/Engine/Source/Runtime/Engine/Public/Animation/AnimPoseSearchProvider.h
deleted file mode 100644
