motion matching
- creating a FPoseSearchIndexBase with the data indexed from the sequences and blendspaces sampling
- FPoseSearchIndex now contains only the data processed post database normalization (weights, pca data, kdtree...)
- moved all the indexing code into posesearchderived data under the #if WITH_EDITOR directive

diff --git a/Engine/Plugins/Experimental/Animation/PoseSearch/Source/Runtime/Private/PoseSearch.cpp b/Engine/Plugins/Experimental/Animation/PoseSearch/Source/Runtime/Private/PoseSearch.cpp
--- a/Engine/Plugins/Experimental/Animation/PoseSearch/Source/Runtime/Private/PoseSearch.cpp
+++ b/Engine/Plugins/Experimental/Animation/PoseSearch/Source/Runtime/Private/PoseSearch.cpp
@@ -3,7 +3,6 @@
 #include "PoseSearch/PoseSearch.h"
 #include "AssetRegistry/AssetData.h"
 #include "PoseSearch/PoseSearchAnimNotifies.h"
-#include "PoseSearch/PoseSearchDerivedDataKey.h"
 #include "PoseSearch/PoseSearchDerivedData.h"
 #include "PoseSearchEigenHelper.h"
 
@@ -103,16 +102,6 @@ static inline bool IsSamplingRangeValid(FFloatInterval Range)
 	return Range.IsValid() && (Range.Min >= 0.0f);
 }
 
-static inline FFloatInterval GetEffectiveSamplingRange(const UAnimSequenceBase* Sequence, FFloatInterval RequestedSamplingRange)
-{
-	const bool bSampleAll = (RequestedSamplingRange.Min == 0.0f) && (RequestedSamplingRange.Max == 0.0f);
-	const float SequencePlayLength = Sequence->GetPlayLength();
-	FFloatInterval Range;
-	Range.Min = bSampleAll ? 0.0f : RequestedSamplingRange.Min;
-	Range.Max = bSampleAll ? SequencePlayLength : FMath::Min(SequencePlayLength, RequestedSamplingRange.Max);
-	return Range;
-}
-
 /**
 * Algo::LowerBound adapted to TIndexedContainerIterator for use with indexable but not necessarily contiguous containers. Used here with TRingBuffer.
 *
@@ -366,51 +355,6 @@ private:
 	TArray<const IPoseFilter*, TInlineAllocator<64>> AllPoseFilters;
 };
 
-static void FindValidSequenceIntervals(const FPoseSearchDatabaseSequence& DbSequence, const FPoseSearchExcludeFromDatabaseParameters& ExcludeFromDatabaseParameters, TArray<FFloatRange>& ValidRanges)
-{
-	const UAnimSequence* Sequence = DbSequence.Sequence;
-	check(DbSequence.Sequence);
-
-	const float SequenceLength = DbSequence.Sequence->GetPlayLength();
-
-	const FFloatInterval EffectiveSamplingInterval = UE::PoseSearch::GetEffectiveSamplingRange(DbSequence.Sequence, DbSequence.SamplingRange);
-	FFloatRange EffectiveSamplingRange = FFloatRange::Inclusive(EffectiveSamplingInterval.Min, EffectiveSamplingInterval.Max);
-	if (!DbSequence.IsLooping())
-	{
-		const FFloatRange ExcludeFromDatabaseRange(ExcludeFromDatabaseParameters.SequenceStartInterval, SequenceLength - ExcludeFromDatabaseParameters.SequenceEndInterval);
-		EffectiveSamplingRange = FFloatRange::Intersection(EffectiveSamplingRange, ExcludeFromDatabaseRange);
-	}
-
-	// start from a single interval defined by the database sequence sampling range
-	ValidRanges.Empty();
-	ValidRanges.Add(EffectiveSamplingRange);
-
-	FAnimNotifyContext NotifyContext;
-	Sequence->GetAnimNotifies(0.0f, SequenceLength, NotifyContext);
-
-	for (const FAnimNotifyEventReference& EventReference : NotifyContext.ActiveNotifies)
-	{
-		if (const FAnimNotifyEvent* NotifyEvent = EventReference.GetNotify())
-		{
-			if (const UAnimNotifyState_PoseSearchExcludeFromDatabase* ExclusionNotifyState = Cast<const UAnimNotifyState_PoseSearchExcludeFromDatabase>(NotifyEvent->NotifyStateClass))
-			{
-				FFloatRange ExclusionRange = FFloatRange::Inclusive(NotifyEvent->GetTriggerTime(), NotifyEvent->GetEndTriggerTime());
-
-				// Split every valid range based on the exclusion range just found. Because this might increase the 
-				// number of ranges in ValidRanges, the algorithm iterates from end to start.
-				for (int RangeIdx = ValidRanges.Num() - 1; RangeIdx >= 0; --RangeIdx)
-				{
-					FFloatRange EvaluatedRange = ValidRanges[RangeIdx];
-					ValidRanges.RemoveAt(RangeIdx);
-
-					TArray<FFloatRange> Diff = FFloatRange::Difference(EvaluatedRange, ExclusionRange);
-					ValidRanges.Append(Diff);
-				}
-			}
-		}
-	}
-}
-
 } // namespace UE::PoseSearch
 
 
@@ -579,46 +523,14 @@ bool UPoseSearchSchema::BuildQuery(UE::PoseSearch::FSearchContext& SearchContext
 }
 
 //////////////////////////////////////////////////////////////////////////
-// FPoseSearchIndex
-
-FPoseSearchIndex::FPoseSearchIndex(const FPoseSearchIndex& Other)
-	: NumPoses(Other.NumPoses)
-	, Values(Other.Values)
-	, PCAValues(Other.PCAValues)
-#if WITH_EDITORONLY_DATA
-	, PCAExplainedVariance(Other.PCAExplainedVariance)
-	, Deviation(Other.Deviation)
-#endif // WITH_EDITORONLY_DATA
-	, KDTree(Other.KDTree)
-	, PCAProjectionMatrix(Other.PCAProjectionMatrix)
-	, Mean(Other.Mean)
-	, WeightsSqrt(Other.WeightsSqrt)
-	, PoseMetadata(Other.PoseMetadata)
-	, OverallFlags(Other.OverallFlags)
-	, Assets(Other.Assets)
-	, MinCostAddend(Other.MinCostAddend)
-{
-	check(!PCAValues.IsEmpty() || KDTree.DataSource.PointCount == 0);
-	KDTree.DataSource.Data = PCAValues.IsEmpty() ? nullptr : PCAValues.GetData();
-}
-
-FPoseSearchIndex& FPoseSearchIndex::operator=(const FPoseSearchIndex& Other)
-{
-	if (this != &Other)
-	{
-		this->~FPoseSearchIndex();
-		new(this)FPoseSearchIndex(Other);
-	}
-	return *this;
-}
-
-const FPoseSearchIndexAsset& FPoseSearchIndex::GetAssetForPose(int32 PoseIdx) const
+// FPoseSearchBaseIndex
+const FPoseSearchIndexAsset& FPoseSearchIndexBase::GetAssetForPose(int32 PoseIdx) const
 {
 	const int32 AssetIndex = PoseMetadata[PoseIdx].AssetIndex;
 	return Assets[AssetIndex];
 }
 
-const FPoseSearchIndexAsset* FPoseSearchIndex::GetAssetForPoseSafe(int32 PoseIdx) const
+const FPoseSearchIndexAsset* FPoseSearchIndexBase::GetAssetForPoseSafe(int32 PoseIdx) const
 {
 	if (PoseMetadata.IsValidIndex(PoseIdx))
 	{
@@ -631,7 +543,7 @@ const FPoseSearchIndexAsset* FPoseSearchIndex::GetAssetForPoseSafe(int32 PoseIdx
 	return nullptr;
 }
 
-float FPoseSearchIndex::GetAssetTime(int32 PoseIdx, float SamplingInterval) const
+float FPoseSearchIndexBase::GetAssetTime(int32 PoseIdx, float SamplingInterval) const
 {
 	const FPoseSearchIndexAsset& Asset = GetAssetForPose(PoseIdx);
 
@@ -657,16 +569,97 @@ float FPoseSearchIndex::GetAssetTime(int32 PoseIdx, float SamplingInterval) cons
 	return -1.0f;
 }
 
-bool FPoseSearchIndex::IsEmpty() const
+bool FPoseSearchIndexBase::IsEmpty() const
 {
 	const bool bEmpty = Assets.Num() == 0 || NumPoses == 0;
 	return bEmpty;
 }
 
+void FPoseSearchIndexBase::Reset()
+{
+	FPoseSearchIndexBase Default;
+	*this = Default;
+}
+
+FArchive& operator<<(FArchive& Ar, FPoseSearchIndexBase& Index)
+{
+	int32 NumValues = 0;
+	int32 NumAssets = 0;
+
+	if (Ar.IsSaving())
+	{
+		NumValues = Index.Values.Num();
+		NumAssets = Index.Assets.Num();
+	}
+
+	Ar << Index.NumPoses;
+	Ar << NumValues;
+	Ar << NumAssets;
+	Ar << Index.OverallFlags;
+
+	if (Ar.IsLoading())
+	{
+		Index.Values.SetNumUninitialized(NumValues);
+		Index.PoseMetadata.SetNumUninitialized(Index.NumPoses);
+		Index.Assets.SetNumUninitialized(NumAssets);
+	}
+
+	if (Index.Values.Num() > 0)
+	{
+		Ar.Serialize(&Index.Values[0], Index.Values.Num() * Index.Values.GetTypeSize());
+	}
+
+	if (Index.PoseMetadata.Num() > 0)
+	{
+		Ar.Serialize(&Index.PoseMetadata[0], Index.PoseMetadata.Num() * Index.PoseMetadata.GetTypeSize());
+	}
+
+	if (Index.Assets.Num() > 0)
+	{
+		Ar.Serialize(&Index.Assets[0], Index.Assets.Num() * Index.Assets.GetTypeSize());
+	}
+
+	return Ar;
+}
+
+//////////////////////////////////////////////////////////////////////////
+// FPoseSearchIndex
+FPoseSearchIndex::FPoseSearchIndex(const FPoseSearchIndex& Other)
+	: FPoseSearchIndexBase(Other)
+	, PCAValues(Other.PCAValues)
+	, PCAProjectionMatrix(Other.PCAProjectionMatrix)
+	, Mean(Other.Mean)
+	, WeightsSqrt(Other.WeightsSqrt)
+	, KDTree(Other.KDTree)
+#if WITH_EDITORONLY_DATA
+	, PCAExplainedVariance(Other.PCAExplainedVariance)
+	, Deviation(Other.Deviation)
+#endif // WITH_EDITORONLY_DATA
+{
+	check(!PCAValues.IsEmpty() || KDTree.DataSource.PointCount == 0);
+	KDTree.DataSource.Data = PCAValues.IsEmpty() ? nullptr : PCAValues.GetData();
+}
+
+FPoseSearchIndex& FPoseSearchIndex::operator=(const FPoseSearchIndex& Other)
+{
+	if (this != &Other)
+	{
+		this->~FPoseSearchIndex();
+		new(this)FPoseSearchIndex(Other);
+	}
+	return *this;
+}
+
+void FPoseSearchIndex::Reset()
+{
+	FPoseSearchIndex Default;
+	*this = Default;
+}
+
 TConstArrayView<float> FPoseSearchIndex::GetPoseValues(int32 PoseIdx) const
 {
 	const int32 SchemaCardinality = WeightsSqrt.Num();
-	check(PoseIdx >= 0 && PoseIdx < NumPoses && SchemaCardinality > 0);
+	check(PoseIdx >= 0 && PoseIdx < NumPoses&& SchemaCardinality > 0);
 	const int32 ValueOffset = PoseIdx * SchemaCardinality;
 	return MakeArrayView(&Values[ValueOffset], SchemaCardinality);
 }
@@ -682,17 +675,11 @@ TConstArrayView<float> FPoseSearchIndex::GetPoseValuesSafe(int32 PoseIdx) const
 	return TConstArrayView<float>();
 }
 
-void FPoseSearchIndex::Reset()
-{
-	FPoseSearchIndex Default;
-	*this = Default;
-}
-
 FPoseSearchCost FPoseSearchIndex::ComparePoses(int32 PoseIdx, EPoseSearchBooleanRequest QueryMirrorRequest, UE::PoseSearch::EPoseComparisonFlags PoseComparisonFlags, float MirrorMismatchCostBias, TConstArrayView<float> QueryValues) const
 {
 	// base dissimilarity cost representing how the associated PoseIdx differ, in a weighted way, from the query pose (QueryValues)
 	const float DissimilarityCost = UE::PoseSearch::CompareFeatureVectors(GetPoseValues(PoseIdx), QueryValues, WeightsSqrt);
-	
+
 	// cost addend associated to a mismatch in mirror state between query and analyzed PoseIdx
 	float MirrorMismatchAddend = 0.f;
 	if (QueryMirrorRequest != EPoseSearchBooleanRequest::Indifferent)
@@ -708,122 +695,32 @@ FPoseSearchCost FPoseSearchIndex::ComparePoses(int32 PoseIdx, EPoseSearchBoolean
 	}
 
 	const FPoseSearchPoseMetadata& PoseIdxMetadata = PoseMetadata[PoseIdx];
-	
+
 	// cost addend associated to Schema->BaseCostBias or overriden by UAnimNotifyState_PoseSearchModifyCost
 	const float NotifyAddend = PoseIdxMetadata.CostAddend;
 
 	// cost addend associated to Schema->ContinuingPoseCostBias or overriden by UAnimNotifyState_PoseSearchOverrideContinuingPoseCostBias
 	const float ContinuingPoseCostAddend = EnumHasAnyFlags(PoseComparisonFlags, UE::PoseSearch::EPoseComparisonFlags::ContinuingPose) ? PoseIdxMetadata.ContinuingPoseCostAddend : 0.f;
-	
-	return FPoseSearchCost(DissimilarityCost, NotifyAddend, MirrorMismatchAddend, ContinuingPoseCostAddend);
-}
-
-void FPoseSearchIndex::InitSearchIndexAssets(TConstArrayView<FPoseSearchDatabaseSequence> Sequences, TConstArrayView<FPoseSearchDatabaseBlendSpace> BlendSpaces, const FPoseSearchExcludeFromDatabaseParameters& ExcludeFromDatabaseParameters)
-{
-	using namespace UE::PoseSearch;
-
-	Assets.Empty();
-
-	TArray<FFloatRange> ValidRanges;
-
-	for (int32 SequenceIdx = 0; SequenceIdx < Sequences.Num(); ++SequenceIdx)
-	{
-		const FPoseSearchDatabaseSequence& Sequence = Sequences[SequenceIdx];
-
-		if (Sequence.bEnabled)
-		{
-			const bool bAddUnmirrored = Sequence.MirrorOption == EPoseSearchMirrorOption::UnmirroredOnly || Sequence.MirrorOption == EPoseSearchMirrorOption::UnmirroredAndMirrored;
-			const bool bAddMirrored = Sequence.MirrorOption == EPoseSearchMirrorOption::MirroredOnly || Sequence.MirrorOption == EPoseSearchMirrorOption::UnmirroredAndMirrored;
-
-			ValidRanges.Reset();
-			FindValidSequenceIntervals(Sequence, ExcludeFromDatabaseParameters, ValidRanges);
-			for (const FFloatRange& Range : ValidRanges)
-			{
-				if (bAddUnmirrored)
-				{
-					Assets.Add(FPoseSearchIndexAsset(ESearchIndexAssetType::Sequence, SequenceIdx, false, FFloatInterval(Range.GetLowerBoundValue(), Range.GetUpperBoundValue())));
-				}
-
-				if (bAddMirrored)
-				{
-					Assets.Add(FPoseSearchIndexAsset(ESearchIndexAssetType::Sequence, SequenceIdx, true, FFloatInterval(Range.GetLowerBoundValue(), Range.GetUpperBoundValue())));
-				}
-			}
-		}
-	}
-
-	TArray<FBlendSampleData> BlendSamples;
-
-	for (int32 BlendSpaceIdx = 0; BlendSpaceIdx < BlendSpaces.Num(); ++BlendSpaceIdx)
-	{
-		const FPoseSearchDatabaseBlendSpace& BlendSpace = BlendSpaces[BlendSpaceIdx];
-
-		if (BlendSpace.bEnabled)
-		{
-			const bool bAddUnmirrored = BlendSpace.MirrorOption == EPoseSearchMirrorOption::UnmirroredOnly || BlendSpace.MirrorOption == EPoseSearchMirrorOption::UnmirroredAndMirrored;
-			const bool bAddMirrored = BlendSpace.MirrorOption == EPoseSearchMirrorOption::MirroredOnly || BlendSpace.MirrorOption == EPoseSearchMirrorOption::UnmirroredAndMirrored;
-
-			int32 HorizontalBlendNum, VerticalBlendNum;
-			BlendSpace.GetBlendSpaceParameterSampleRanges(HorizontalBlendNum, VerticalBlendNum);
 
-			const bool bWrapInputOnHorizontalAxis = BlendSpace.BlendSpace->GetBlendParameter(0).bWrapInput;
-			const bool bWrapInputOnVerticalAxis = BlendSpace.BlendSpace->GetBlendParameter(1).bWrapInput;
-			for (int32 HorizontalIndex = 0; HorizontalIndex < HorizontalBlendNum; HorizontalIndex++)
-			{
-				for (int32 VerticalIndex = 0; VerticalIndex < VerticalBlendNum; VerticalIndex++)
-				{
-					const FVector BlendParameters = BlendSpace.BlendParameterForSampleRanges(HorizontalIndex, VerticalIndex);
-
-					int32 TriangulationIndex = 0;
-					BlendSpace.BlendSpace->GetSamplesFromBlendInput(BlendParameters, BlendSamples, TriangulationIndex, true);
-
-					float PlayLength = BlendSpace.BlendSpace->GetAnimationLengthFromSampleData(BlendSamples);
-
-					if (bAddUnmirrored)
-					{
-						Assets.Add(FPoseSearchIndexAsset(ESearchIndexAssetType::BlendSpace, BlendSpaceIdx, false, FFloatInterval(0.0f, PlayLength), BlendParameters));
-					}
-
-					if (bAddMirrored)
-					{
-						Assets.Add(FPoseSearchIndexAsset(ESearchIndexAssetType::BlendSpace, BlendSpaceIdx, true, FFloatInterval(0.0f, PlayLength), BlendParameters));
-					}
-				}
-			}
-		}
-	}
+	return FPoseSearchCost(DissimilarityCost, NotifyAddend, MirrorMismatchAddend, ContinuingPoseCostAddend);
 }
 
 FArchive& operator<<(FArchive& Ar, FPoseSearchIndex& Index)
 {
-	int32 NumValues = 0;
+	Ar << static_cast<FPoseSearchIndexBase&>(Index);
+
 	int32 NumPCAValues = 0;
-	int32 NumAssets = 0;
 
 	if (Ar.IsSaving())
 	{
-		NumValues = Index.Values.Num();
 		NumPCAValues = Index.PCAValues.Num();
-		NumAssets = Index.Assets.Num();
 	}
 
-	Ar << Index.NumPoses;
-	Ar << NumValues;
 	Ar << NumPCAValues;
-	Ar << NumAssets;
-	Ar << Index.OverallFlags;
 
 	if (Ar.IsLoading())
 	{
-		Index.Values.SetNumUninitialized(NumValues);
 		Index.PCAValues.SetNumUninitialized(NumPCAValues);
-		Index.PoseMetadata.SetNumUninitialized(Index.NumPoses);
-		Index.Assets.SetNumUninitialized(NumAssets);
-	}
-
-	if (Index.Values.Num() > 0)
-	{
-		Ar.Serialize(&Index.Values[0], Index.Values.Num() * Index.Values.GetTypeSize());
 	}
 
 	if (Index.PCAValues.Num() > 0)
@@ -831,22 +728,17 @@ FArchive& operator<<(FArchive& Ar, FPoseSearchIndex& Index)
 		Ar.Serialize(&Index.PCAValues[0], Index.PCAValues.Num() * Index.PCAValues.GetTypeSize());
 	}
 
-	if (Index.PoseMetadata.Num() > 0)
-	{
-		Ar.Serialize(&Index.PoseMetadata[0], Index.PoseMetadata.Num() * Index.PoseMetadata.GetTypeSize());
-	}
-
-	if (Index.Assets.Num() > 0)
-	{
-		Ar.Serialize(&Index.Assets[0], Index.Assets.Num() * Index.Assets.GetTypeSize());
-	}
-
 	Ar << Index.WeightsSqrt;
 	Ar << Index.Mean;
 	Ar << Index.PCAProjectionMatrix;
 
 	Serialize(Ar, Index.KDTree, Index.PCAValues.GetData());
 
+#if WITH_EDITORONLY_DATA
+	Ar << Index.PCAExplainedVariance;
+	Ar << Index.Deviation;
+#endif // WITH_EDITORONLY_DATA
+
 	return Ar;
 }
 
@@ -3011,561 +2903,45 @@ const UAnimationAsset* FBlendSpaceSampler::GetAsset() const
 }
 
 //////////////////////////////////////////////////////////////////////////
-// FAssetIndexer helpers
-
-struct FSamplingParam
-{
-	float WrappedParam = 0.0f;
-	int32 NumCycles = 0;
-	
-	// If the animation can't loop, WrappedParam contains the clamped value and whatever is left is stored here
-	float Extrapolation = 0.0f;
-};
+// PoseSearch API
 
-static FSamplingParam WrapOrClampSamplingParam(bool bCanWrap, float SamplingParamExtent, float SamplingParam)
+void DrawFeatureVector(const FDebugDrawParams& DrawParams, TConstArrayView<float> PoseVector)
 {
-	// This is a helper function used by both time and distance sampling. A schema may specify time or distance
-	// offsets that are multiple cycles of a clip away from the current pose being sampled.
-	// And that time or distance offset may before the beginning of the clip (SamplingParam < 0.0f)
-	// or after the end of the clip (SamplingParam > SamplingParamExtent). So this function
-	// helps determine how many cycles need to be applied and what the wrapped value should be, clamping
-	// if necessary.
-
-	FSamplingParam Result;
-
-	Result.WrappedParam = SamplingParam;
-
-	const bool bIsSamplingParamExtentKindaSmall = SamplingParamExtent <= UE_KINDA_SMALL_NUMBER;
-	if (!bIsSamplingParamExtentKindaSmall && bCanWrap)
+#if ENABLE_DRAW_DEBUG
+	if (DrawParams.CanDraw())
 	{
-		if (SamplingParam < 0.0f)
-		{
-			while (Result.WrappedParam < 0.0f)
-			{
-				Result.WrappedParam += SamplingParamExtent;
-				++Result.NumCycles;
-			}
-		}
+		const UPoseSearchSchema* Schema = DrawParams.GetSchema();
+		check(Schema);
 
-		else
+		if (PoseVector.Num() == Schema->SchemaCardinality)
 		{
-			while (Result.WrappedParam > SamplingParamExtent)
+			for (int32 ChannelIdx = 0; ChannelIdx != Schema->Channels.Num(); ++ChannelIdx)
 			{
-				Result.WrappedParam -= SamplingParamExtent;
-				++Result.NumCycles;
+				if (DrawParams.ChannelMask & (1 << ChannelIdx))
+				{
+					Schema->Channels[ChannelIdx]->DebugDraw(DrawParams, PoseVector);
+				}
 			}
 		}
 	}
+#endif // ENABLE_DRAW_DEBUG
+}
 
-	const float ParamClamped = FMath::Clamp(Result.WrappedParam, 0.0f, SamplingParamExtent);
-	if (ParamClamped != Result.WrappedParam)
+void DrawFeatureVector(const FDebugDrawParams& DrawParams, int32 PoseIdx)
+{
+#if ENABLE_DRAW_DEBUG
+	// if we're editing the schema while in PIE with Rewind Debugger active, PoseIdx could be out of bound / stale
+	if (DrawParams.CanDraw() && PoseIdx >= 0 && PoseIdx < DrawParams.GetSearchIndex()->NumPoses)
 	{
-		check(bIsSamplingParamExtentKindaSmall || !bCanWrap);
-		Result.Extrapolation = Result.WrappedParam - ParamClamped;
-		Result.WrappedParam = ParamClamped;
+		DrawFeatureVector(DrawParams, DrawParams.GetSearchIndex()->GetPoseValues(PoseIdx));
 	}
-
-	return Result;
+#endif // ENABLE_DRAW_DEBUG
 }
 
-
-//////////////////////////////////////////////////////////////////////////
-// FAssetIndexer
-
-class FAssetIndexer : public IAssetIndexer
+void DrawSearchIndex(const FDebugDrawParams& DrawParams)
 {
-public:
-
-	struct FOutput
-	{
-		int32 FirstIndexedSample = 0;
-		int32 LastIndexedSample = 0;
-		int32 NumIndexedPoses = 0;
-
-		TArray<float> FeatureVectorTable;
-		TArray<FPoseSearchPoseMetadata> PoseMetadata;
-		TBitArray<> AllFeaturesNotAdded;
-	} Output;
-
-	void Reset();
-	void Init(const FAssetIndexingContext& IndexingContext, const FBoneContainer& InBoneContainer);
-	bool Process();
-
-public: // IAssetIndexer
-
-	const FAssetIndexingContext& GetIndexingContext() const override { return IndexingContext; }
-	FSampleInfo GetSampleInfo(float SampleTime) const override;
-	FSampleInfo GetSampleInfoRelative(float SampleTime, const FSampleInfo& Origin) const override;
-	const float GetSampleTimeFromDistance(float Distance) const override;
-	FTransform MirrorTransform(const FTransform& Transform) const override;
-	FTransform GetTransformAndCacheResults(float SampleTime, float OriginTime, int8 SchemaBoneIdx, bool& Clamped) override;
-
-private:
-	FPoseSearchPoseMetadata GetMetadata(int32 SampleIdx) const;
-	
-	struct CachedEntry
-	{
-		float SampleTime;
-		float OriginTime;
-		bool Clamped;
-
-		// @todo: minimize the Entry memory footprint
-		FTransform RootTransform;
-		FCompactPose Pose;
-		FCSPose<FCompactPose> ComponentSpacePose;
-		FBlendedCurve UnusedCurve;
-		UE::Anim::FStackAttributeContainer UnusedAtrribute;
-		FAnimationPoseData AnimPoseData = { Pose, UnusedCurve, UnusedAtrribute };
-	};
-
-	FBoneContainer BoneContainer;
-	FAssetIndexingContext IndexingContext;
-	TArray<CachedEntry> CachedEntries;
-};
-
-void FAssetIndexer::Reset()
-{
-	Output.FirstIndexedSample = 0;
-	Output.LastIndexedSample = 0;
-	Output.NumIndexedPoses = 0;
-
-	Output.FeatureVectorTable.Reset(0);
-	Output.PoseMetadata.Reset(0);
-	Output.AllFeaturesNotAdded.Reset();
-}
-
-void FAssetIndexer::Init(const FAssetIndexingContext& InIndexingContext, const FBoneContainer& InBoneContainer)
-{
-	check(InIndexingContext.Schema);
-	check(InIndexingContext.Schema->IsValid());
-	check(InIndexingContext.MainSampler);
-
-	BoneContainer = InBoneContainer;
-	IndexingContext = InIndexingContext;
-
-	Reset();
-
-	Output.FirstIndexedSample = FMath::FloorToInt(IndexingContext.RequestedSamplingRange.Min * IndexingContext.Schema->SampleRate);
-	Output.LastIndexedSample = FMath::Max(0, FMath::CeilToInt(IndexingContext.RequestedSamplingRange.Max * IndexingContext.Schema->SampleRate));
-	Output.NumIndexedPoses = Output.LastIndexedSample - Output.FirstIndexedSample + 1;
-	
-	Output.FeatureVectorTable.SetNumZeroed(IndexingContext.Schema->SchemaCardinality * Output.NumIndexedPoses);
-	Output.PoseMetadata.SetNum(Output.NumIndexedPoses);
-}
-
-bool FAssetIndexer::Process()
-{
-	check(IndexingContext.Schema);
-	check(IndexingContext.Schema->IsValid());
-	check(IndexingContext.MainSampler);
-
-	FMemMark Mark(FMemStack::Get());
-
-	IndexingContext.BeginSampleIdx = Output.FirstIndexedSample;
-	IndexingContext.EndSampleIdx = Output.LastIndexedSample + 1;
-
-	if (IndexingContext.Schema->SchemaCardinality > 0)
-	{
-		// Index each channel
-		FAssetIndexingOutput AssetIndexingOutput(IndexingContext.Schema->SchemaCardinality, Output.FeatureVectorTable);
-		for (int32 ChannelIdx = 0; ChannelIdx != IndexingContext.Schema->Channels.Num(); ++ChannelIdx) 
-		{
-			IndexingContext.Schema->Channels[ChannelIdx]->IndexAsset(*this, AssetIndexingOutput);
-		}
-	}
-
-	// Generate pose metadata
-	for (int32 SampleIdx = IndexingContext.BeginSampleIdx; SampleIdx != IndexingContext.EndSampleIdx; ++SampleIdx)
-	{
-		const int32 PoseIdx = SampleIdx - Output.FirstIndexedSample;
-		FPoseSearchPoseMetadata& OutputPoseMetadata = Output.PoseMetadata[PoseIdx];
-		OutputPoseMetadata = GetMetadata(SampleIdx);
-	}
-
-	return true;
-}
-
-const float FAssetIndexer::GetSampleTimeFromDistance(float SampleDistance) const
-{
-	auto CanWrapDistanceSamples = [](const IAssetSampler* Sampler) -> bool
-	{
-		constexpr float SMALL_ROOT_DISTANCE = 1.0f;
-		return Sampler->IsLoopable() && Sampler->GetTotalRootDistance() > SMALL_ROOT_DISTANCE;
-	};
-
-	float MainTotalDistance = IndexingContext.MainSampler->GetTotalRootDistance();
-	bool bMainCanWrap = CanWrapDistanceSamples(IndexingContext.MainSampler);
-
-	float SampleTime = MAX_flt;
-
-	if (!bMainCanWrap)
-	{
-		// Use the lead in anim if we would have to clamp to the beginning of the main anim
-		if (IndexingContext.LeadInSampler && (SampleDistance < 0.0f))
-		{
-			const IAssetSampler* ClipSampler = IndexingContext.LeadInSampler;
-
-			bool bLeadInCanWrap = CanWrapDistanceSamples(IndexingContext.LeadInSampler);
-			float LeadRelativeDistance = SampleDistance + ClipSampler->GetTotalRootDistance();
-			FSamplingParam SamplingParam = WrapOrClampSamplingParam(bLeadInCanWrap, ClipSampler->GetTotalRootDistance(), LeadRelativeDistance);
-
-			float ClipTime = ClipSampler->GetTimeFromRootDistance(
-				SamplingParam.WrappedParam + SamplingParam.Extrapolation);
-
-			// Make the lead in clip time relative to the main sequence again and unwrap
-			SampleTime = -((SamplingParam.NumCycles * ClipSampler->GetPlayLength()) + (ClipSampler->GetPlayLength() - ClipTime));
-		}
-
-		// Use the follow up anim if we would have clamp to the end of the main anim
-		else if (IndexingContext.FollowUpSampler && (SampleDistance > MainTotalDistance))
-		{
-			const IAssetSampler* ClipSampler = IndexingContext.FollowUpSampler;
-
-			bool bFollowUpCanWrap = CanWrapDistanceSamples(IndexingContext.FollowUpSampler);
-			float FollowRelativeDistance = SampleDistance - MainTotalDistance;
-			FSamplingParam SamplingParam = WrapOrClampSamplingParam(bFollowUpCanWrap, ClipSampler->GetTotalRootDistance(), FollowRelativeDistance);
-
-			float ClipTime = ClipSampler->GetTimeFromRootDistance(
-				SamplingParam.WrappedParam + SamplingParam.Extrapolation);
-
-			// Make the follow up clip time relative to the main sequence again and unwrap
-			SampleTime = IndexingContext.MainSampler->GetPlayLength() + SamplingParam.NumCycles * ClipSampler->GetPlayLength() + ClipTime;
-		}
-	}
-
-	// Use the main anim if we didn't use the lead-in or follow-up anims.
-	// The main anim sample may have been wrapped or clamped
-	if (SampleTime == MAX_flt)
-	{
-		float MainRelativeDistance = SampleDistance;
-		if (SampleDistance < 0.0f && bMainCanWrap)
-		{
-			// In this case we're sampling a loop backwards, so MainRelativeDistance must adjust so the number of cycles 
-			// is counted correctly.
-			MainRelativeDistance += IndexingContext.MainSampler->GetTotalRootDistance();
-		}
-
-		FSamplingParam SamplingParam = WrapOrClampSamplingParam(bMainCanWrap, MainTotalDistance, MainRelativeDistance);
-		float ClipTime = IndexingContext.MainSampler->GetTimeFromRootDistance(
-			SamplingParam.WrappedParam + SamplingParam.Extrapolation);
-
-		// Unwrap the main clip time
-		if (bMainCanWrap)
-		{
-			if (SampleDistance < 0.0f)
-			{
-				SampleTime = -((SamplingParam.NumCycles * IndexingContext.MainSampler->GetPlayLength()) + (IndexingContext.MainSampler->GetPlayLength() - ClipTime));
-			}
-			else
-			{
-				SampleTime = SamplingParam.NumCycles * IndexingContext.MainSampler->GetPlayLength() + ClipTime;
-			}
-		}
-		else
-		{
-			SampleTime = ClipTime;
-		}
-	}
-
-	return SampleTime;
-}
-
-FAssetIndexer::FSampleInfo FAssetIndexer::GetSampleInfo(float SampleTime) const
-{
-	FSampleInfo Sample;
-
-	FTransform RootMotionLast = FTransform::Identity;
-	FTransform RootMotionInitial = FTransform::Identity;
-
-	float RootDistanceLast = 0.0f;
-	float RootDistanceInitial = 0.0f;
-
-	check(IndexingContext.MainSampler);
-
-	const float MainPlayLength = IndexingContext.MainSampler->GetPlayLength();
-	const bool bMainCanWrap = IndexingContext.MainSampler->IsLoopable();
-
-	FSamplingParam SamplingParam;
-	if (!bMainCanWrap)
-	{
-		// Use the lead in anim if we would have to clamp to the beginning of the main anim
-		if (IndexingContext.LeadInSampler && (SampleTime < 0.0f))
-		{
-			const IAssetSampler* ClipSampler = IndexingContext.LeadInSampler;
-
-			const bool bLeadInCanWrap = IndexingContext.LeadInSampler->IsLoopable();
-			const float LeadRelativeTime = SampleTime + ClipSampler->GetPlayLength();
-			SamplingParam = WrapOrClampSamplingParam(bLeadInCanWrap, ClipSampler->GetPlayLength(), LeadRelativeTime);
-
-			Sample.Clip = IndexingContext.LeadInSampler;
-
-			check(SamplingParam.Extrapolation <= 0.0f);
-			if (SamplingParam.Extrapolation < 0.0f)
-			{
-				RootMotionInitial = IndexingContext.LeadInSampler->GetTotalRootTransform().Inverse();
-				RootDistanceInitial = -IndexingContext.LeadInSampler->GetTotalRootDistance();
-			}
-			else
-			{
-				RootMotionInitial = FTransform::Identity;
-				RootDistanceInitial = 0.0f;
-			}
-
-			RootMotionLast = IndexingContext.LeadInSampler->GetTotalRootTransform();
-			RootDistanceLast = IndexingContext.LeadInSampler->GetTotalRootDistance();
-		}
-
-		// Use the follow up anim if we would have clamp to the end of the main anim
-		else if (IndexingContext.FollowUpSampler && (SampleTime > MainPlayLength))
-		{
-			const IAssetSampler* ClipSampler = IndexingContext.FollowUpSampler;
-
-			const bool bFollowUpCanWrap = IndexingContext.FollowUpSampler->IsLoopable();
-			const float FollowRelativeTime = SampleTime - MainPlayLength;
-			SamplingParam = WrapOrClampSamplingParam(bFollowUpCanWrap, ClipSampler->GetPlayLength(), FollowRelativeTime);
-
-			Sample.Clip = IndexingContext.FollowUpSampler;
-
-			RootMotionInitial = IndexingContext.MainSampler->GetTotalRootTransform();
-			RootDistanceInitial = IndexingContext.MainSampler->GetTotalRootDistance();
-
-			RootMotionLast = IndexingContext.FollowUpSampler->GetTotalRootTransform();
-			RootDistanceLast = IndexingContext.FollowUpSampler->GetTotalRootDistance();
-		}
-	}
-
-	// Use the main anim if we didn't use the lead-in or follow-up anims.
-	// The main anim sample may have been wrapped or clamped
-	if (!Sample.IsValid())
-	{
-		float MainRelativeTime = SampleTime;
-		if (SampleTime < 0.0f && bMainCanWrap)
-		{
-			// In this case we're sampling a loop backwards, so MainRelativeTime must adjust so the number of cycles is
-			// counted correctly.
-			MainRelativeTime += MainPlayLength;
-		}
-
-		SamplingParam = WrapOrClampSamplingParam(bMainCanWrap, MainPlayLength, MainRelativeTime);
-
-		Sample.Clip = IndexingContext.MainSampler;
-
-		RootMotionInitial = FTransform::Identity;
-		RootDistanceInitial = 0.0f;
-
-		RootMotionLast = IndexingContext.MainSampler->GetTotalRootTransform();
-		RootDistanceLast = IndexingContext.MainSampler->GetTotalRootDistance();
-	}
-
-
-	if (FMath::Abs(SamplingParam.Extrapolation) > SMALL_NUMBER)
-	{
-		Sample.bClamped = true;
-		Sample.ClipTime = SamplingParam.WrappedParam + SamplingParam.Extrapolation;
-		const FTransform ClipRootMotion = Sample.Clip->ExtractRootTransform(Sample.ClipTime);
-		const float ClipDistance = Sample.Clip->ExtractRootDistance(Sample.ClipTime);
-
-		Sample.RootTransform = ClipRootMotion * RootMotionInitial;
-		Sample.RootDistance = RootDistanceInitial + ClipDistance;
-	}
-	else
-	{
-		Sample.ClipTime = SamplingParam.WrappedParam;
-
-		// Determine how to accumulate motion for every cycle of the anim. If the sample
-		// had to be clamped, this motion will end up not getting applied below.
-		// Also invert the accumulation direction if the requested sample was wrapped backwards.
-		FTransform RootMotionPerCycle = RootMotionLast;
-		float RootDistancePerCycle = RootDistanceLast;
-		if (SampleTime < 0.0f)
-		{
-			RootMotionPerCycle = RootMotionPerCycle.Inverse();
-			RootDistancePerCycle *= -1.f;
-		}
-
-		// Find the remaining motion deltas after wrapping
-		FTransform RootMotionRemainder = Sample.Clip->ExtractRootTransform(Sample.ClipTime);
-		float RootDistanceRemainder = Sample.Clip->ExtractRootDistance(Sample.ClipTime);
-
-		// Invert motion deltas if we wrapped backwards
-		if (SampleTime < 0.0f)
-		{
-			RootMotionRemainder.SetToRelativeTransform(RootMotionLast);
-			RootDistanceRemainder = -(RootDistanceLast - RootDistanceRemainder);
-		}
-
-		Sample.RootTransform = RootMotionInitial;
-		Sample.RootDistance = RootDistanceInitial;
-
-		// Note if the sample was clamped, no motion will be applied here because NumCycles will be zero
-		int32 CyclesRemaining = SamplingParam.NumCycles;
-		while (CyclesRemaining--)
-		{
-			Sample.RootTransform = RootMotionPerCycle * Sample.RootTransform;
-         	Sample.RootDistance += RootDistancePerCycle;
-		}
-
-		Sample.RootTransform = RootMotionRemainder * Sample.RootTransform;
-		Sample.RootDistance += RootDistanceRemainder;
-	}
-
-	return Sample;
-}
-
-FAssetIndexer::FSampleInfo FAssetIndexer::GetSampleInfoRelative(float SampleTime, const FSampleInfo& Origin) const
-{
-	FSampleInfo Sample = GetSampleInfo(SampleTime);
-	Sample.RootTransform.SetToRelativeTransform(Origin.RootTransform);
-	Sample.RootDistance = Origin.RootDistance - Sample.RootDistance;
-	return Sample;
-}
-
-FTransform FAssetIndexer::MirrorTransform(const FTransform& Transform) const
-{
-	return IndexingContext.bMirrored ? IndexingContext.SamplingContext->MirrorTransform(Transform) : Transform;
-}
-
-FPoseSearchPoseMetadata FAssetIndexer::GetMetadata(int32 SampleIdx) const
-{
-	const float SequenceLength = IndexingContext.MainSampler->GetPlayLength();
-	const float SampleTime = FMath::Min(SampleIdx * IndexingContext.Schema->GetSamplingInterval(), SequenceLength);
-
-	FPoseSearchPoseMetadata Metadata;
-	Metadata.CostAddend = IndexingContext.Schema->BaseCostBias;
-	Metadata.ContinuingPoseCostAddend = IndexingContext.Schema->ContinuingPoseCostBias;
-	
-	TArray<UAnimNotifyState_PoseSearchBase*> NotifyStates;
-	IndexingContext.MainSampler->ExtractPoseSearchNotifyStates(SampleTime, NotifyStates);
-	for (const UAnimNotifyState_PoseSearchBase* PoseSearchNotify : NotifyStates)
-	{
-		if (PoseSearchNotify->GetClass()->IsChildOf<UAnimNotifyState_PoseSearchBlockTransition>())
-		{
-			EnumAddFlags(Metadata.Flags, EPoseSearchPoseFlags::BlockTransition);
-		}
-		else if (PoseSearchNotify->GetClass()->IsChildOf<UAnimNotifyState_PoseSearchModifyCost>())
-		{
-			const UAnimNotifyState_PoseSearchModifyCost* ModifyCostNotify =
-				Cast<const UAnimNotifyState_PoseSearchModifyCost>(PoseSearchNotify);
-			Metadata.CostAddend = ModifyCostNotify->CostAddend;
-		}
-		else if (PoseSearchNotify->GetClass()->IsChildOf<UAnimNotifyState_PoseSearchOverrideContinuingPoseCostBias>())
-		{
-			const UAnimNotifyState_PoseSearchOverrideContinuingPoseCostBias* ContinuingPoseCostBias =
-				Cast<const UAnimNotifyState_PoseSearchOverrideContinuingPoseCostBias>(PoseSearchNotify);
-			Metadata.ContinuingPoseCostAddend = ContinuingPoseCostBias->CostAddend;
-		}
-	}
-	return Metadata;
-}
-
-FTransform FAssetIndexer::GetTransformAndCacheResults(float SampleTime, float OriginTime, int8 SchemaBoneIdx, bool& Clamped)
-{
-	// @todo: use an hashmap if we end up having too many entries
-	CachedEntry* Entry = CachedEntries.FindByPredicate([SampleTime, OriginTime](const FAssetIndexer::CachedEntry& Entry)
-	{
-		return Entry.SampleTime == SampleTime && Entry.OriginTime == OriginTime;
-	});
-
-	const FAssetSamplingContext* SamplingContext = IndexingContext.SamplingContext;
-
-	if (!Entry)
-	{
-		Entry = &CachedEntries[CachedEntries.AddDefaulted()];
-
-		Entry->SampleTime = SampleTime;
-		Entry->OriginTime = OriginTime;
-
-		if (!BoneContainer.IsValid())
-		{
-			UE_LOG(LogPoseSearch,
-				Warning, 
-				TEXT("Invalid BoneContainer encountered in FAssetIndexer::GetTransformAndCacheResults. Asset: %s. Schema: %s. BoneContainerAsset: %s. NumBoneIndices: %d"),
-				*GetNameSafe(IndexingContext.MainSampler->GetAsset()),
-				*GetNameSafe(IndexingContext.Schema),
-				*GetNameSafe(BoneContainer.GetAsset()),
-				BoneContainer.GetCompactPoseNumBones());
-		}
-
-		Entry->Pose.SetBoneContainer(&BoneContainer);
-		Entry->UnusedCurve.InitFrom(BoneContainer);
-
-		IAssetIndexer::FSampleInfo Origin = GetSampleInfo(OriginTime);
-		IAssetIndexer::FSampleInfo Sample = GetSampleInfoRelative(SampleTime, Origin);
-
-		float CurrentTime = Sample.ClipTime;
-		float PreviousTime = CurrentTime - SamplingContext->FiniteDelta;
-
-		FDeltaTimeRecord DeltaTimeRecord;
-		DeltaTimeRecord.Set(PreviousTime, CurrentTime - PreviousTime);
-		FAnimExtractContext ExtractionCtx(static_cast<double>(CurrentTime), true, DeltaTimeRecord, Sample.Clip->IsLoopable());
-
-		Sample.Clip->ExtractPose(ExtractionCtx, Entry->AnimPoseData);
-
-		if (IndexingContext.bMirrored)
-		{
-			FAnimationRuntime::MirrorPose(
-				Entry->AnimPoseData.GetPose(),
-				IndexingContext.Schema->MirrorDataTable->MirrorAxis,
-				SamplingContext->CompactPoseMirrorBones,
-				SamplingContext->ComponentSpaceRefRotations
-			);
-			// Note curves and attributes are not used during the indexing process and therefore don't need to be mirrored
-		}
-
-		Entry->ComponentSpacePose.InitPose(Entry->Pose);
-		Entry->RootTransform = Sample.RootTransform;
-		Entry->Clamped = Sample.bClamped;
-	}
-
-	const FBoneReference& BoneReference = IndexingContext.Schema->BoneReferences[SchemaBoneIdx];
-	FCompactPoseBoneIndex CompactBoneIndex = BoneContainer.MakeCompactPoseIndex(FMeshPoseBoneIndex(BoneReference.BoneIndex));
-
-	const FTransform BoneTransform = Entry->ComponentSpacePose.GetComponentSpaceTransform(CompactBoneIndex) * MirrorTransform(Entry->RootTransform);
-	Clamped = Entry->Clamped;
-
-	return BoneTransform;
-}
-
-//////////////////////////////////////////////////////////////////////////
-// PoseSearch API
-
-void DrawFeatureVector(const FDebugDrawParams& DrawParams, TConstArrayView<float> PoseVector)
-{
-#if ENABLE_DRAW_DEBUG
-	if (DrawParams.CanDraw())
-	{
-		const UPoseSearchSchema* Schema = DrawParams.GetSchema();
-		check(Schema);
-
-		if (PoseVector.Num() == Schema->SchemaCardinality)
-		{
-			for (int32 ChannelIdx = 0; ChannelIdx != Schema->Channels.Num(); ++ChannelIdx)
-			{
-				if (DrawParams.ChannelMask & (1 << ChannelIdx))
-				{
-					Schema->Channels[ChannelIdx]->DebugDraw(DrawParams, PoseVector);
-				}
-			}
-		}
-	}
-#endif // ENABLE_DRAW_DEBUG
-}
-
-void DrawFeatureVector(const FDebugDrawParams& DrawParams, int32 PoseIdx)
-{
-#if ENABLE_DRAW_DEBUG
-	// if we're editing the schema while in PIE with Rewind Debugger active, PoseIdx could be out of bound / stale
-	if (DrawParams.CanDraw() && PoseIdx >= 0 && PoseIdx < DrawParams.GetSearchIndex()->NumPoses)
-	{
-		DrawFeatureVector(DrawParams, DrawParams.GetSearchIndex()->GetPoseValues(PoseIdx));
-	}
-#endif // ENABLE_DRAW_DEBUG
-}
-
-void DrawSearchIndex(const FDebugDrawParams& DrawParams)
-{
-#if ENABLE_DRAW_DEBUG
-	if (DrawParams.CanDraw())
+#if ENABLE_DRAW_DEBUG
+	if (DrawParams.CanDraw())
 	{
 		const FPoseSearchIndex* SearchIndex = DrawParams.GetSearchIndex();
 		const int32 LastPoseIdx = SearchIndex->NumPoses;
@@ -3577,578 +2953,6 @@ void DrawSearchIndex(const FDebugDrawParams& DrawParams)
 #endif // ENABLE_DRAW_DEBUG
 }
 
-static Eigen::VectorXd ComputeChannelsDeviations(const FPoseSearchIndex& SearchIndex, const UPoseSearchSchema* Schema)
-{
-	// This function performs a modified z-score normalization where features are normalized
-	// by mean absolute deviation rather than standard deviation. Both methods are preferable
-	// here to min-max scaling because they preserve outliers.
-	// 
-	// Mean absolute deviation is preferred here over standard deviation because the latter
-	// emphasizes outliers since squaring the distance from the mean increases variance 
-	// exponentially rather than additively and square rooting the sum of squares does not 
-	// remove that bias. [1]
-	//
-	// References:
-	// [1] Gorard, S. (2005), "Revisiting a 90-Year-Old Debate: The Advantages of the Mean Deviation."
-	//     British Journal of Educational Studies, 53: 417-430.
-
-	using namespace Eigen;
-
-	check(Schema->IsValid());
-
-	const int32 NumPoses = SearchIndex.NumPoses;
-	const int32 NumDimensions = Schema->SchemaCardinality;
-
-	// Compute per Channel average distances
-	VectorXd MeanDeviations(NumDimensions);
-	MeanDeviations.setConstant(1.0);
-
-	if (NumPoses > 0)
-	{
-		// Map input buffer
-		auto PoseMatrixSourceMap = RowMajorMatrixMapConst(
-			SearchIndex.Values.GetData(),
-			NumPoses,		// rows
-			NumDimensions	// cols
-		);
-
-		// @todo: evaluate removing the cast to double
-
-		// Copy row major float matrix to column major double matrix
-		MatrixXd PoseMatrix = PoseMatrixSourceMap.transpose().cast<double>();
-		checkSlow(PoseMatrix.rows() == NumDimensions);
-		checkSlow(PoseMatrix.cols() == NumPoses);
-
-		// Mean center
-		VectorXd SampleMean = PoseMatrix.rowwise().mean();
-		PoseMatrix = PoseMatrix.colwise() - SampleMean;
-
-		for (int ChannelIdx = 0; ChannelIdx != Schema->Channels.Num(); ++ChannelIdx)
-		{
-			const UPoseSearchFeatureChannel* Channel = Schema->Channels[ChannelIdx].Get();
-			Channel->ComputeMeanDeviations(PoseMatrix, MeanDeviations);
-		}
-	}
-
-	return MeanDeviations;
-}
-
-static void PreprocessSearchIndexWeights(FPoseSearchIndex& SearchIndex, const UPoseSearchSchema* Schema)
-{
-	const int32 NumDimensions = Schema->SchemaCardinality;
-	SearchIndex.WeightsSqrt.Init(1.f, NumDimensions);
-
-	for (int ChannelIdx = 0; ChannelIdx != Schema->Channels.Num(); ++ChannelIdx)
-	{
-		const UPoseSearchFeatureChannel* Channel = Schema->Channels[ChannelIdx].Get();
-		Channel->FillWeights(SearchIndex.WeightsSqrt);
-	}
-
-	Eigen::VectorXd ChannelsMeanDeviations = ComputeChannelsDeviations(SearchIndex, Schema);
-	TArray<float> Deviation;
-	Deviation.Init(1.f, NumDimensions);
-	for (int32 Dimension = 0; Dimension != NumDimensions; ++Dimension)
-	{
-		const float ChannelsMeanDeviation = ChannelsMeanDeviations[Dimension];
-		Deviation[Dimension] = ChannelsMeanDeviation;
-	}
-
-	EPoseSearchDataPreprocessor DataPreprocessor = Schema->DataPreprocessor;
-
-	if (DataPreprocessor == EPoseSearchDataPreprocessor::Normalize)
-	{
-		// normalizing user weights: the idea behind this step is to be able to compare poses from databases using different schemas
-		RowMajorVectorMap MapWeights(SearchIndex.WeightsSqrt.GetData(), 1, NumDimensions);
-		const float WeightsSum = MapWeights.sum();
-		if (!FMath::IsNearlyZero(WeightsSum))
-		{
-			MapWeights *= (1.0f / WeightsSum);
-		}
-	}
-
-	// extracting the square root
-	for (int32 Dimension = 0; Dimension != NumDimensions; ++Dimension)
-	{
-		SearchIndex.WeightsSqrt[Dimension] = FMath::Sqrt(SearchIndex.WeightsSqrt[Dimension]);
-	}
-
-	if (DataPreprocessor == EPoseSearchDataPreprocessor::Normalize || DataPreprocessor == EPoseSearchDataPreprocessor::NormalizeOnlyByDeviation)
-	{
-		for (int32 Dimension = 0; Dimension != NumDimensions; ++Dimension)
-		{
-			// the idea here is to premultiply the weights by the inverse of the variance (proportional to the square of the deviation) to have a "weighted Mahalanobis" distance
-			SearchIndex.WeightsSqrt[Dimension] /= Deviation[Dimension];
-		}
-	}
-
-#if WITH_EDITORONLY_DATA
-	SearchIndex.Deviation = Deviation;
-#endif // WITH_EDITORONLY_DATA
-}
-
-// it calculates Mean, PCAValues, and PCAProjectionMatrix
-static void PreprocessSearchIndexPCAData(FPoseSearchIndex& SearchIndex, int32 NumDimensions, uint32 NumberOfPrincipalComponents, EPoseSearchMode PoseSearchMode)
-{
-	// binding SearchIndex.Values and SearchIndex.PCAValues Eigen row major matrix maps
-	const int32 NumPoses = SearchIndex.NumPoses;
-
-	SearchIndex.PCAValues.Reset();
-	SearchIndex.Mean.Reset();
-	SearchIndex.PCAProjectionMatrix.Reset();
-
-	SearchIndex.PCAValues.AddZeroed(NumPoses * NumberOfPrincipalComponents);
-	SearchIndex.Mean.AddZeroed(NumDimensions);
-	SearchIndex.PCAProjectionMatrix.AddZeroed(NumDimensions * NumberOfPrincipalComponents);
-
-#if WITH_EDITORONLY_DATA
-	SearchIndex.PCAExplainedVariance = 0.f;
-#endif
-
-	if (NumDimensions > 0)
-	{
-		const RowMajorVectorMapConst MapWeightsSqrt(SearchIndex.WeightsSqrt.GetData(), 1, NumDimensions);
-		const RowMajorMatrixMapConst MapValues(SearchIndex.Values.GetData(), NumPoses, NumDimensions);
-		const RowMajorMatrix WeightedValues = MapValues.array().rowwise() * MapWeightsSqrt.array();
-		RowMajorMatrixMap MapPCAValues(SearchIndex.PCAValues.GetData(), NumPoses, NumberOfPrincipalComponents);
-
-		// calculating the mean
-		RowMajorVectorMap Mean(SearchIndex.Mean.GetData(), 1, NumDimensions);
-		Mean = WeightedValues.colwise().mean();
-
-		// use the mean to center the data points
-		const RowMajorMatrix CenteredValues = WeightedValues.rowwise() - Mean;
-
-		// estimating the covariance matrix (with dimensionality of NumDimensions, NumDimensions)
-		// formula: https://en.wikipedia.org/wiki/Covariance_matrix#Estimation
-		// details: https://en.wikipedia.org/wiki/Estimation_of_covariance_matrices
-		const ColMajorMatrix CovariantMatrix = (CenteredValues.transpose() * CenteredValues) / float(NumPoses - 1);
-		const Eigen::SelfAdjointEigenSolver<ColMajorMatrix> EigenSolver(CovariantMatrix);
-
-		check(EigenSolver.info() == Eigen::Success);
-
-		// validating EigenSolver results
-		const ColMajorMatrix EigenVectors = EigenSolver.eigenvectors().real();
-
-		if (PoseSearchMode == EPoseSearchMode::PCAKDTree_Validate && NumberOfPrincipalComponents == NumDimensions)
-		{
-			const RowMajorVector ReciprocalWeightsSqrt = MapWeightsSqrt.cwiseInverse();
-			const RowMajorMatrix ProjectedValues = CenteredValues * EigenVectors;
-			for (Eigen::Index RowIndex = 0; RowIndex < MapValues.rows(); ++RowIndex)
-			{
-				const RowMajorVector WeightedReconstructedPoint = ProjectedValues.row(RowIndex) * EigenVectors.transpose() + Mean;
-				const RowMajorVector ReconstructedPoint = WeightedReconstructedPoint.array() * ReciprocalWeightsSqrt.array();
-				const float Error = (ReconstructedPoint - MapValues.row(RowIndex)).squaredNorm();
-				check(Error < UE_KINDA_SMALL_NUMBER);
-			}
-		}
-
-		// sorting EigenVectors by EigenValues, so we pick the most significant ones to compose our PCA projection matrix.
-		const RowMajorVector EigenValues = EigenSolver.eigenvalues().real();
-		TArray<size_t> Indexer;
-		Indexer.Reserve(NumDimensions);
-		for (size_t DimensionIndex = 0; DimensionIndex < NumDimensions; ++DimensionIndex)
-		{
-			Indexer.Push(DimensionIndex);
-		}
-		Indexer.Sort([&EigenValues](size_t a, size_t b)
-		{
-			return EigenValues[a] > EigenValues[b];
-		});
-
-		// composing the PCA projection matrix with the PCANumComponents most significant EigenVectors
-		ColMajorMatrixMap PCAProjectionMatrix(SearchIndex.PCAProjectionMatrix.GetData(), NumDimensions, NumberOfPrincipalComponents);
-		float AccumulatedVariance = 0.f;
-		for (size_t PCAComponentIndex = 0; PCAComponentIndex < NumberOfPrincipalComponents; ++PCAComponentIndex)
-		{
-			PCAProjectionMatrix.col(PCAComponentIndex) = EigenVectors.col(Indexer[PCAComponentIndex]);
-			AccumulatedVariance += EigenValues[Indexer[PCAComponentIndex]];
-		}
-
-#if WITH_EDITORONLY_DATA
-		// calculating the total variance knowing that eigen values measure variance along the principal components:
-		const float TotalVariance = EigenValues.sum();
-		// and explained variance as ratio between AccumulatedVariance and TotalVariance: https://ro-che.info/articles/2017-12-11-pca-explained-variance
-		SearchIndex.PCAExplainedVariance = TotalVariance > UE_KINDA_SMALL_NUMBER ? AccumulatedVariance / TotalVariance : 0.f;
-#endif // WITH_EDITORONLY_DATA
-
-		MapPCAValues = CenteredValues * PCAProjectionMatrix;
-
-		if (PoseSearchMode == EPoseSearchMode::PCAKDTree_Validate && NumberOfPrincipalComponents == NumDimensions)
-		{
-			const RowMajorVector ReciprocalWeightsSqrt = MapWeightsSqrt.cwiseInverse();
-			for (Eigen::Index RowIndex = 0; RowIndex < MapValues.rows(); ++RowIndex)
-			{
-				const RowMajorVector WeightedReconstructedValues = MapPCAValues.row(RowIndex) * PCAProjectionMatrix.transpose() + Mean;
-				const RowMajorVector ReconstructedValues = WeightedReconstructedValues.array() * ReciprocalWeightsSqrt.array();
-				const float Error = (ReconstructedValues - MapValues.row(RowIndex)).squaredNorm();
-				check(Error < UE_KINDA_SMALL_NUMBER);
-			}
-		}
-	}
-}
-
-static void PreprocessSearchIndexKDTree(FPoseSearchIndex& SearchIndex, int32 NumDimensions, uint32 NumberOfPrincipalComponents, EPoseSearchMode PoseSearchMode, int32 KDTreeMaxLeafSize, int32 KDTreeQueryNumNeighbors)
-{
-	const int32 NumPoses = SearchIndex.NumPoses;
-	SearchIndex.KDTree.Construct(NumPoses, NumberOfPrincipalComponents, SearchIndex.PCAValues.GetData(), KDTreeMaxLeafSize);
-
-	if (PoseSearchMode == EPoseSearchMode::PCAKDTree_Validate)
-	{
-		// testing the KDTree is returning the proper searches for all the points in pca space
-		int32 NumberOfFailingPoints = 0;
-		for (size_t PointIndex = 0; PointIndex < NumPoses; ++PointIndex)
-		{
-			TArray<size_t> ResultIndexes;
-			TArray<float> ResultDistanceSqr;
-			ResultIndexes.SetNum(KDTreeQueryNumNeighbors + 1);
-			ResultDistanceSqr.SetNum(KDTreeQueryNumNeighbors + 1);
-			FKDTree::KNNResultSet ResultSet(KDTreeQueryNumNeighbors, ResultIndexes, ResultDistanceSqr);
-			SearchIndex.KDTree.FindNeighbors(ResultSet, &SearchIndex.PCAValues[PointIndex * NumberOfPrincipalComponents]);
-
-			size_t ResultIndex = 0;
-			for (; ResultIndex < ResultSet.Num(); ++ResultIndex)
-			{
-				if (PointIndex == ResultIndexes[ResultIndex])
-				{
-					check(ResultDistanceSqr[ResultIndex] < UE_KINDA_SMALL_NUMBER);
-					break;
-				}
-			}
-			if (ResultIndex == ResultSet.Num())
-			{
-				++NumberOfFailingPoints;
-			}
-		}
-
-		check(NumberOfFailingPoints == 0);
-
-		// testing the KDTree is returning the proper searches for all the original points transformed in pca space
-		NumberOfFailingPoints = 0;
-		for (size_t PointIndex = 0; PointIndex < NumPoses; ++PointIndex)
-		{
-			TArray<size_t> ResultIndexes;
-			TArray<float> ResultDistanceSqr;
-			ResultIndexes.SetNum(KDTreeQueryNumNeighbors + 1);
-			ResultDistanceSqr.SetNum(KDTreeQueryNumNeighbors + 1);
-			FKDTree::KNNResultSet ResultSet(KDTreeQueryNumNeighbors, ResultIndexes, ResultDistanceSqr);
-
-			const RowMajorVectorMapConst MapValues(&SearchIndex.Values[PointIndex * NumDimensions], 1, NumDimensions);
-			const RowMajorVectorMapConst MapWeightsSqrt(SearchIndex.WeightsSqrt.GetData(), 1, NumDimensions);
-			const RowMajorVectorMapConst Mean(SearchIndex.Mean.GetData(), 1, NumDimensions);
-			const ColMajorMatrixMapConst PCAProjectionMatrix(SearchIndex.PCAProjectionMatrix.GetData(), NumDimensions, NumberOfPrincipalComponents);
-
-			const RowMajorMatrix WeightedValues = MapValues.array() * MapWeightsSqrt.array();
-			const RowMajorMatrix CenteredValues = WeightedValues - Mean;
-			const RowMajorVector ProjectedValues  = CenteredValues * PCAProjectionMatrix;
-
-			SearchIndex.KDTree.FindNeighbors(ResultSet, ProjectedValues.data());
-
-			size_t ResultIndex = 0;
-			for (; ResultIndex < ResultSet.Num(); ++ResultIndex)
-			{
-				if (PointIndex == ResultIndexes[ResultIndex])
-				{
-					check(ResultDistanceSqr[ResultIndex] < UE_KINDA_SMALL_NUMBER);
-					break;
-				}
-			}
-			if (ResultIndex == ResultSet.Num())
-			{
-				++NumberOfFailingPoints;
-			}
-		}
-
-		check(NumberOfFailingPoints == 0);
-	}
-}
-
-static void PreprocessSearchIndex(FPoseSearchIndex& SearchIndex, const UPoseSearchSchema* Schema, uint32 NumberOfPrincipalComponents, EPoseSearchMode PoseSearchMode, int32 KDTreeMaxLeafSize, int32 KDTreeQueryNumNeighbors)
-{
-	PreprocessSearchIndexWeights(SearchIndex, Schema);
-	PreprocessSearchIndexPCAData(SearchIndex, Schema->SchemaCardinality, NumberOfPrincipalComponents, PoseSearchMode);
-	PreprocessSearchIndexKDTree(SearchIndex, Schema->SchemaCardinality, NumberOfPrincipalComponents, PoseSearchMode, KDTreeMaxLeafSize, KDTreeQueryNumNeighbors);
-}
-
-struct FDatabaseIndexingContext
-{
-	FPoseSearchIndex* SearchIndex = nullptr;
-
-	FAssetSamplingContext SamplingContext;
-	TArray<FSequenceSampler> SequenceSamplers;
-	TArray<FBlendSpaceSampler> BlendSpaceSamplers;
-
-	TArray<FAssetIndexer> Indexers;
-
-	void Prepare(const UPoseSearchSchema* Schema, const FPoseSearchExtrapolationParameters& ExtrapolationParameters, TConstArrayView<FPoseSearchDatabaseSequence> Sequences, TConstArrayView<FPoseSearchDatabaseBlendSpace> BlendSpaces);
-	bool IndexAssets();
-	void JoinIndex();
-	float CalculateMinCostAddend() const;
-};
-
-void FDatabaseIndexingContext::Prepare(const UPoseSearchSchema* Schema, const FPoseSearchExtrapolationParameters& ExtrapolationParameters, TConstArrayView<FPoseSearchDatabaseSequence> Sequences, TConstArrayView<FPoseSearchDatabaseBlendSpace> BlendSpaces)
-{
-	check(Schema);
-
-	FBoneContainer BoneContainer;
-	BoneContainer.InitializeTo(Schema->BoneIndicesWithParents, FCurveEvaluationOption(false), *Schema->Skeleton);
-
-	TMap<const UAnimSequence*, int32> SequenceSamplerMap;
-	TMap<TPair<const UBlendSpace*, FVector>, int32> BlendSpaceSamplerMap;
-
-	SamplingContext.Init(Schema->MirrorDataTable, BoneContainer);
-
-	// Prepare samplers for all sequences
-	auto AddSequenceSampler = [&](const UAnimSequence* Sequence)
-	{
-		if (Sequence && !SequenceSamplerMap.Contains(Sequence))
-		{
-			int32 SequenceSamplerIdx = SequenceSamplers.AddDefaulted();
-			SequenceSamplerMap.Add(Sequence, SequenceSamplerIdx);
-
-			FSequenceSampler::FInput Input;
-			Input.ExtrapolationParameters = ExtrapolationParameters;
-			Input.Sequence = Sequence;
-			SequenceSamplers[SequenceSamplerIdx].Init(Input);
-		}
-	};
-
-	for (const FPoseSearchDatabaseSequence& DbSequence : Sequences)
-	{
-		AddSequenceSampler(DbSequence.Sequence);
-		AddSequenceSampler(DbSequence.LeadInSequence);
-		AddSequenceSampler(DbSequence.FollowUpSequence);
-	}
-		
-	// Prepare samplers for all blend spaces
-	for (const FPoseSearchDatabaseBlendSpace& DbBlendSpace : BlendSpaces)
-	{
-		if (DbBlendSpace.BlendSpace)
-		{
-			int32 HorizontalBlendNum, VerticalBlendNum;
-			DbBlendSpace.GetBlendSpaceParameterSampleRanges(HorizontalBlendNum, VerticalBlendNum);
-
-			for (int32 HorizontalIndex = 0; HorizontalIndex < HorizontalBlendNum; HorizontalIndex++)
-			{
-				for (int32 VerticalIndex = 0; VerticalIndex < VerticalBlendNum; VerticalIndex++)
-				{
-					const FVector BlendParameters = DbBlendSpace.BlendParameterForSampleRanges(HorizontalIndex, VerticalIndex);
-
-					if (!BlendSpaceSamplerMap.Contains({ DbBlendSpace.BlendSpace, BlendParameters }))
-					{
-						int32 BlendSpaceSamplerIdx = BlendSpaceSamplers.AddDefaulted();
-						BlendSpaceSamplerMap.Add({ DbBlendSpace.BlendSpace, BlendParameters }, BlendSpaceSamplerIdx);
-
-						FBlendSpaceSampler::FInput Input;
-						Input.BoneContainer = BoneContainer;
-						Input.ExtrapolationParameters = ExtrapolationParameters;
-						Input.BlendSpace = DbBlendSpace.BlendSpace;
-						Input.BlendParameters = BlendParameters;
-
-						BlendSpaceSamplers[BlendSpaceSamplerIdx].Init(Input);
-					}
-				}
-			}
-		}
-	}
-
-	TArray<IAssetSampler*, TInlineAllocator<512>> AssetSampler;
-	AssetSampler.SetNumUninitialized(SequenceSamplers.Num() + BlendSpaceSamplers.Num());
-	
-	for (int i = 0; i < SequenceSamplers.Num(); ++i)
-	{
-		AssetSampler[i] = &SequenceSamplers[i];
-	}
-	for (int i = 0; i < BlendSpaceSamplers.Num(); ++i)
-	{
-		AssetSampler[i + SequenceSamplers.Num()] = &BlendSpaceSamplers[i];
-	}
-
-	ParallelFor(AssetSampler.Num(), [AssetSampler](int32 SamplerIdx) { AssetSampler[SamplerIdx]->Process(); }, ParallelForFlags);
-
-	// prepare indexers
-	Indexers.Reserve(SearchIndex->Assets.Num());
-
-	auto GetSequenceSampler = [&](const UAnimSequence* Sequence) -> const FSequenceSampler*
-	{
-		return Sequence ? &SequenceSamplers[SequenceSamplerMap[Sequence]] : nullptr;
-	};
-
-	auto GetBlendSpaceSampler = [&](const UBlendSpace* BlendSpace, const FVector BlendParameters) -> const FBlendSpaceSampler*
-	{
-		return BlendSpace ? &BlendSpaceSamplers[BlendSpaceSamplerMap[{BlendSpace, BlendParameters}]] : nullptr;
-	};
-
-	Indexers.Reserve(SearchIndex->Assets.Num());
-
-	for (int32 AssetIdx = 0; AssetIdx != SearchIndex->Assets.Num(); ++AssetIdx)
-	{
-		const FPoseSearchIndexAsset& SearchIndexAsset = SearchIndex->Assets[AssetIdx];
-
-		FAssetIndexingContext IndexerContext;
-		IndexerContext.SamplingContext = &SamplingContext;
-		IndexerContext.Schema = Schema;
-		IndexerContext.RequestedSamplingRange = SearchIndexAsset.SamplingInterval;
-		IndexerContext.bMirrored = SearchIndexAsset.bMirrored;
-
-		if (SearchIndexAsset.Type == ESearchIndexAssetType::Sequence)
-		{
-			const FPoseSearchDatabaseSequence& DbSequence = Sequences[SearchIndexAsset.SourceAssetIdx];
-			const float SequenceLength = DbSequence.Sequence->GetPlayLength();
-			IndexerContext.MainSampler = GetSequenceSampler(DbSequence.Sequence);
-			IndexerContext.LeadInSampler = SearchIndexAsset.SamplingInterval.Min == 0.0f ? GetSequenceSampler(DbSequence.LeadInSequence) : nullptr;
-			IndexerContext.FollowUpSampler = SearchIndexAsset.SamplingInterval.Max == SequenceLength ? GetSequenceSampler(DbSequence.FollowUpSequence) : nullptr;
-		}
-		else if (SearchIndexAsset.Type == ESearchIndexAssetType::BlendSpace)
-		{
-			const FPoseSearchDatabaseBlendSpace& DbBlendSpace = BlendSpaces[SearchIndexAsset.SourceAssetIdx];
-			IndexerContext.MainSampler = GetBlendSpaceSampler(DbBlendSpace.BlendSpace, SearchIndexAsset.BlendParameters);
-		}
-		else
-		{
-			checkNoEntry();
-		}
-
-		FAssetIndexer& Indexer = Indexers.AddDefaulted_GetRef();
-		Indexer.Init(IndexerContext, BoneContainer);
-	}
-}
-
-bool FDatabaseIndexingContext::IndexAssets()
-{
-	// Index asset data
-	ParallelFor(Indexers.Num(), [this](int32 AssetIdx) { Indexers[AssetIdx].Process(); }, ParallelForFlags);
-	return true;
-}
-
-float FDatabaseIndexingContext::CalculateMinCostAddend() const
-{
-	float MinCostAddend = 0.f;
-
-	check(SearchIndex);
-	if (!SearchIndex->PoseMetadata.IsEmpty())
-	{
-		MinCostAddend = MAX_FLT;
-		for (const FPoseSearchPoseMetadata& PoseMetadata : SearchIndex->PoseMetadata)
-		{
-			if (PoseMetadata.CostAddend < MinCostAddend)
-			{
-				MinCostAddend = PoseMetadata.CostAddend;
-			}
-		}
-	}
-	return MinCostAddend;
-}
-
-void FDatabaseIndexingContext::JoinIndex()
-{
-	// Write index info to asset and count up total poses and storage required
-	int32 TotalPoses = 0;
-	int32 TotalFloats = 0;
-
-	check(SearchIndex);
-
-	// Join animation data into a single search index
-	SearchIndex->Values.Reset();
-	SearchIndex->PoseMetadata.Reset();
-	SearchIndex->PCAValues.Reset();
-	SearchIndex->OverallFlags = EPoseSearchPoseFlags::None;
-
-	for (int32 AssetIdx = 0; AssetIdx != SearchIndex->Assets.Num(); ++AssetIdx)
-	{
-		const FAssetIndexer::FOutput& Output = Indexers[AssetIdx].Output;
-
-		FPoseSearchIndexAsset& SearchIndexAsset = SearchIndex->Assets[AssetIdx];
-		SearchIndexAsset.NumPoses = Output.NumIndexedPoses;
-		SearchIndexAsset.FirstPoseIdx = TotalPoses;
-
-		const int32 PoseMetadataStartIdx = SearchIndex->PoseMetadata.Num();
-		const int32 PoseMetadataEndIdx = PoseMetadataStartIdx + Output.PoseMetadata.Num();
-
-		SearchIndex->Values.Append(Output.FeatureVectorTable.GetData(), Output.FeatureVectorTable.Num());
-		SearchIndex->PoseMetadata.Append(Output.PoseMetadata);
-		
-		for (int32 i = PoseMetadataStartIdx; i < PoseMetadataEndIdx; ++i)
-		{
-			SearchIndex->PoseMetadata[i].AssetIndex = AssetIdx;
-			SearchIndex->OverallFlags |= SearchIndex->PoseMetadata[i].Flags;
-		}
-
-		TotalPoses += Output.NumIndexedPoses;
-		TotalFloats += Output.FeatureVectorTable.Num();
-	}
-	
-	SearchIndex->NumPoses = TotalPoses;
-	SearchIndex->MinCostAddend = CalculateMinCostAddend();
-}
-
-bool BuildIndex(const UPoseSearchDatabase& Database, FPoseSearchIndex& OutSearchIndex
-#if WITH_EDITOR
-	, UE::DerivedData::IRequestOwner& Owner
-#endif
-)
-{
-	// early out for invalid indexing conditions
-	if (!Database.Schema || !Database.Schema->IsValid() || Database.Schema->SchemaCardinality <= 0)
-	{
-		OutSearchIndex.Reset();
-		return false;
-	}
-
-#if WITH_EDITOR
-	if (Owner.IsCanceled())
-	{
-		OutSearchIndex.Reset();
-		return false;
-	}
-#endif
-
-	OutSearchIndex.InitSearchIndexAssets(Database.Sequences, Database.BlendSpaces, Database.ExcludeFromDatabaseParameters);
-
-	FDatabaseIndexingContext DbIndexingContext;
-	DbIndexingContext.SearchIndex = &OutSearchIndex;
-	DbIndexingContext.Prepare(Database.Schema, Database.ExtrapolationParameters, Database.Sequences, Database.BlendSpaces);
-	
-#if WITH_EDITOR
-	if (Owner.IsCanceled())
-	{
-		OutSearchIndex.Reset();
-		return false;
-	}
-#endif
-
-	const bool bSuccess = DbIndexingContext.IndexAssets();
-
-#if WITH_EDITOR
-	if (Owner.IsCanceled())
-	{
-		OutSearchIndex.Reset();
-		return false;
-	}
-#endif
-
-	DbIndexingContext.JoinIndex();
-
-#if WITH_EDITOR
-	if (Owner.IsCanceled())
-	{
-		OutSearchIndex.Reset();
-		return false;
-	}
-#endif
-
-	PreprocessSearchIndex(OutSearchIndex, Database.Schema, Database.GetNumberOfPrincipalComponents(), Database.PoseSearchMode, Database.KDTreeMaxLeafSize, Database.KDTreeQueryNumNeighbors);
-
-#if WITH_EDITOR
-	if (Owner.IsCanceled())
-	{
-		OutSearchIndex.Reset();
-		return false;
-	}
-#endif
-
-	return bSuccess;
-}
-
 //////////////////////////////////////////////////////////////////////////
 // FModule
 // @todo: clean up this code
diff --git a/Engine/Plugins/Experimental/Animation/PoseSearch/Source/Runtime/Private/PoseSearchDerivedData.cpp b/Engine/Plugins/Experimental/Animation/PoseSearch/Source/Runtime/Private/PoseSearchDerivedData.cpp
--- a/Engine/Plugins/Experimental/Animation/PoseSearch/Source/Runtime/Private/PoseSearchDerivedData.cpp
+++ b/Engine/Plugins/Experimental/Animation/PoseSearch/Source/Runtime/Private/PoseSearchDerivedData.cpp
@@ -9,439 +9,1648 @@
 #include "DerivedDataCache.h"
 #include "DerivedDataRequestOwner.h"
 #include "Misc/CoreDelegates.h"
+#include "PoseSearch/PoseSearchAnimNotifies.h"
 #include "PoseSearch/PoseSearchDerivedDataKey.h"
+#include "PoseSearchEigenHelper.h"
 #include "ProfilingDebugging/CookStats.h"
 #include "Serialization/BulkDataRegistry.h"
 #include "UObject/NoExportTypes.h"
 
 namespace UE::PoseSearch
 {
-	static const UE::DerivedData::FValueId Id(UE::DerivedData::FValueId::FromName("Data"));
-	static const UE::DerivedData::FCacheBucket Bucket("PoseSearchDatabase");
+static const UE::DerivedData::FValueId Id(UE::DerivedData::FValueId::FromName("Data"));
+static const UE::DerivedData::FCacheBucket Bucket("PoseSearchDatabase");
 
 #if ENABLE_COOK_STATS
-	static FCookStats::FDDCResourceUsageStats UsageStats;
-	static FCookStatsManager::FAutoRegisterCallback RegisterCookStats([](FCookStatsManager::AddStatFuncRef AddStat)
-		{
-			UsageStats.LogStats(AddStat, TEXT("MotionMatching.Usage"), TEXT(""));
-		});
+static FCookStats::FDDCResourceUsageStats UsageStats;
+static FCookStatsManager::FAutoRegisterCallback RegisterCookStats([](FCookStatsManager::AddStatFuncRef AddStat)
+	{
+		UsageStats.LogStats(AddStat, TEXT("MotionMatching.Usage"), TEXT(""));
+	});
 #endif
 
-	//////////////////////////////////////////////////////////////////////////
-	// FPoseSearchDatabaseAsyncCacheTask
-	struct FPoseSearchDatabaseAsyncCacheTask
-	{
-		enum class EState
-		{
-			Prestarted,
-			Cancelled,
-			Ended,
-			Failed
-		};
-
-		// these methods MUST be protected by FPoseSearchDatabaseAsyncCacheTask::Mutex! and to make sure we pass the mutex as input param
-		FPoseSearchDatabaseAsyncCacheTask(UPoseSearchDatabase* InDatabase, FCriticalSection& OuterMutex);
-		void StartNewRequestIfNeeded(FCriticalSection& OuterMutex);
-		bool CancelIfDependsOn(const UObject* Object, FCriticalSection& OuterMutex);
-		void Update(FCriticalSection& OuterMutex);
-		void Wait(FCriticalSection& OuterMutex);
-		void Cancel(FCriticalSection& OuterMutex);
-		bool Poll(FCriticalSection& OuterMutex) const;
-		bool ContainsDatabase(const UPoseSearchDatabase* OtherDatabase, FCriticalSection& OuterMutex) const;
-
-		~FPoseSearchDatabaseAsyncCacheTask();
-		EState GetState() const { return EState(ThreadSafeState.GetValue()); }
-
-	private:
-		FPoseSearchDatabaseAsyncCacheTask(const FPoseSearchDatabaseAsyncCacheTask& Other) = delete;
-		FPoseSearchDatabaseAsyncCacheTask(FPoseSearchDatabaseAsyncCacheTask&& Other) = delete;
-		FPoseSearchDatabaseAsyncCacheTask& operator=(const FPoseSearchDatabaseAsyncCacheTask& Other) = delete;
-		FPoseSearchDatabaseAsyncCacheTask& operator=(FPoseSearchDatabaseAsyncCacheTask&& Other) = delete;
-
-		void OnGetComplete(UE::DerivedData::FCacheGetResponse&& Response);
-		void SetState(EState State) { ThreadSafeState.Set(int32(State)); }
-
-		TWeakObjectPtr<UPoseSearchDatabase> Database;
-		// @todo: this is not relevant when the async task is completed, so to save memory we should move it as pointer perhaps
-		FPoseSearchIndex SearchIndex;
-		UE::DerivedData::FRequestOwner Owner;
-		FIoHash DerivedDataKey = FIoHash::Zero;
-		TSet<TWeakObjectPtr<const UObject>> DatabaseDependencies; // @todo: make this const
-		
-		FThreadSafeCounter ThreadSafeState = int32(EState::Prestarted);
-		bool bBroadcastOnDerivedDataRebuild = false;
-	};
+static Eigen::VectorXd ComputeChannelsDeviations(const FPoseSearchIndex& SearchIndex, const UPoseSearchSchema* Schema)
+{
+	// This function performs a modified z-score normalization where features are normalized
+	// by mean absolute deviation rather than standard deviation. Both methods are preferable
+	// here to min-max scaling because they preserve outliers.
+	// 
+	// Mean absolute deviation is preferred here over standard deviation because the latter
+	// emphasizes outliers since squaring the distance from the mean increases variance 
+	// exponentially rather than additively and square rooting the sum of squares does not 
+	// remove that bias. [1]
+	//
+	// References:
+	// [1] Gorard, S. (2005), "Revisiting a 90-Year-Old Debate: The Advantages of the Mean Deviation."
+	//     British Journal of Educational Studies, 53: 417-430.
+
+	using namespace Eigen;
+
+	check(Schema->IsValid());
+
+	const int32 NumPoses = SearchIndex.NumPoses;
+	const int32 NumDimensions = Schema->SchemaCardinality;
+
+	// Compute per Channel average distances
+	VectorXd MeanDeviations(NumDimensions);
+	MeanDeviations.setConstant(1.0);
+
+	if (NumPoses > 0)
+	{
+		// Map input buffer
+		auto PoseMatrixSourceMap = RowMajorMatrixMapConst(
+			SearchIndex.Values.GetData(),
+			NumPoses,		// rows
+			NumDimensions	// cols
+		);
 
-	class FPoseSearchDatabaseAsyncCacheTasks : public TArray<TUniquePtr<FPoseSearchDatabaseAsyncCacheTask>> {};
+		// @todo: evaluate removing the cast to double
 
-	FPoseSearchDatabaseAsyncCacheTask::FPoseSearchDatabaseAsyncCacheTask(UPoseSearchDatabase* InDatabase, FCriticalSection& OuterMutex)
-		: Database(InDatabase)
-		, Owner(UE::DerivedData::EPriority::Normal)
-		, DerivedDataKey(FIoHash::Zero)
-	{
-		StartNewRequestIfNeeded(OuterMutex);
+		// Copy row major float matrix to column major double matrix
+		MatrixXd PoseMatrix = PoseMatrixSourceMap.transpose().cast<double>();
+		checkSlow(PoseMatrix.rows() == NumDimensions);
+		checkSlow(PoseMatrix.cols() == NumPoses);
+
+		// Mean center
+		VectorXd SampleMean = PoseMatrix.rowwise().mean();
+		PoseMatrix = PoseMatrix.colwise() - SampleMean;
+
+		for (int ChannelIdx = 0; ChannelIdx != Schema->Channels.Num(); ++ChannelIdx)
+		{
+			const UPoseSearchFeatureChannel* Channel = Schema->Channels[ChannelIdx].Get();
+			Channel->ComputeMeanDeviations(PoseMatrix, MeanDeviations);
+		}
 	}
 
-	FPoseSearchDatabaseAsyncCacheTask::~FPoseSearchDatabaseAsyncCacheTask()
+	return MeanDeviations;
+}
+
+static inline FFloatInterval GetEffectiveSamplingRange(const UAnimSequenceBase* Sequence, FFloatInterval RequestedSamplingRange)
+{
+	const bool bSampleAll = (RequestedSamplingRange.Min == 0.0f) && (RequestedSamplingRange.Max == 0.0f);
+	const float SequencePlayLength = Sequence->GetPlayLength();
+	FFloatInterval Range;
+	Range.Min = bSampleAll ? 0.0f : RequestedSamplingRange.Min;
+	Range.Max = bSampleAll ? SequencePlayLength : FMath::Min(SequencePlayLength, RequestedSamplingRange.Max);
+	return Range;
+}
+
+static void FindValidSequenceIntervals(const FPoseSearchDatabaseSequence& DbSequence, const FPoseSearchExcludeFromDatabaseParameters& ExcludeFromDatabaseParameters, TArray<FFloatRange>& ValidRanges)
+{
+	const UAnimSequence* Sequence = DbSequence.Sequence;
+	check(DbSequence.Sequence);
+
+	const float SequenceLength = DbSequence.Sequence->GetPlayLength();
+
+	const FFloatInterval EffectiveSamplingInterval = GetEffectiveSamplingRange(DbSequence.Sequence, DbSequence.SamplingRange);
+	FFloatRange EffectiveSamplingRange = FFloatRange::Inclusive(EffectiveSamplingInterval.Min, EffectiveSamplingInterval.Max);
+	if (!DbSequence.IsLooping())
 	{
-		Database = nullptr;
-		SearchIndex.Reset();
-		Owner.Cancel();
-		DerivedDataKey = FIoHash::Zero;
-		DatabaseDependencies.Reset();
+		const FFloatRange ExcludeFromDatabaseRange(ExcludeFromDatabaseParameters.SequenceStartInterval, SequenceLength - ExcludeFromDatabaseParameters.SequenceEndInterval);
+		EffectiveSamplingRange = FFloatRange::Intersection(EffectiveSamplingRange, ExcludeFromDatabaseRange);
 	}
 
-	void FPoseSearchDatabaseAsyncCacheTask::StartNewRequestIfNeeded(FCriticalSection& OuterMutex)
+	// start from a single interval defined by the database sequence sampling range
+	ValidRanges.Empty();
+	ValidRanges.Add(EffectiveSamplingRange);
+
+	FAnimNotifyContext NotifyContext;
+	Sequence->GetAnimNotifies(0.0f, SequenceLength, NotifyContext);
+
+	for (const FAnimNotifyEventReference& EventReference : NotifyContext.ActiveNotifies)
 	{
-		using namespace UE::DerivedData;
+		if (const FAnimNotifyEvent* NotifyEvent = EventReference.GetNotify())
+		{
+			if (const UAnimNotifyState_PoseSearchExcludeFromDatabase* ExclusionNotifyState = Cast<const UAnimNotifyState_PoseSearchExcludeFromDatabase>(NotifyEvent->NotifyStateClass))
+			{
+				FFloatRange ExclusionRange = FFloatRange::Inclusive(NotifyEvent->GetTriggerTime(), NotifyEvent->GetEndTriggerTime());
+
+				// Split every valid range based on the exclusion range just found. Because this might increase the 
+				// number of ranges in ValidRanges, the algorithm iterates from end to start.
+				for (int RangeIdx = ValidRanges.Num() - 1; RangeIdx >= 0; --RangeIdx)
+				{
+					FFloatRange EvaluatedRange = ValidRanges[RangeIdx];
+					ValidRanges.RemoveAt(RangeIdx);
+
+					TArray<FFloatRange> Diff = FFloatRange::Difference(EvaluatedRange, ExclusionRange);
+					ValidRanges.Append(Diff);
+				}
+			}
+		}
+	}
+}
+
+static void InitSearchIndexAssets(FPoseSearchIndexBase& SearchIndex, TConstArrayView<FPoseSearchDatabaseSequence> Sequences, TConstArrayView<FPoseSearchDatabaseBlendSpace> BlendSpaces, const FPoseSearchExcludeFromDatabaseParameters& ExcludeFromDatabaseParameters)
+{
+	using namespace UE::PoseSearch;
 
-		FScopeLock Lock(&OuterMutex);
+	SearchIndex.Assets.Empty();
 
-		// making sure there are no active requests
-		Owner.Cancel();
+	TArray<FFloatRange> ValidRanges;
 
-		// composing the key
-		FDerivedDataKeyBuilder KeyBuilder;
-		FGuid VersionGuid = FDevSystemGuids::GetSystemGuid(FDevSystemGuids::Get().POSESEARCHDB_DERIVEDDATA_VER);
-		UPoseSearchDatabase* Db = Database.Get();
-		check(Db);
-		KeyBuilder << VersionGuid;
-		KeyBuilder << Db;
+	for (int32 SequenceIdx = 0; SequenceIdx < Sequences.Num(); ++SequenceIdx)
+	{
+		const FPoseSearchDatabaseSequence& Sequence = Sequences[SequenceIdx];
 
-		// Stores a BLAKE3-160 hash, taken from the first 20 bytes of a BLAKE3-256 hash
-		const FIoHash NewDerivedDataKey = KeyBuilder.Finalize();
-		const bool bHasKeyChanged = NewDerivedDataKey != DerivedDataKey;
-		if (bHasKeyChanged)
+		if (Sequence.bEnabled)
 		{
-			DerivedDataKey = NewDerivedDataKey;
+			const bool bAddUnmirrored = Sequence.MirrorOption == EPoseSearchMirrorOption::UnmirroredOnly || Sequence.MirrorOption == EPoseSearchMirrorOption::UnmirroredAndMirrored;
+			const bool bAddMirrored = Sequence.MirrorOption == EPoseSearchMirrorOption::MirroredOnly || Sequence.MirrorOption == EPoseSearchMirrorOption::UnmirroredAndMirrored;
 
-			DatabaseDependencies.Reset();
-			for (const UObject* Dependency : KeyBuilder.GetDependencies())
+			ValidRanges.Reset();
+			FindValidSequenceIntervals(Sequence, ExcludeFromDatabaseParameters, ValidRanges);
+			for (const FFloatRange& Range : ValidRanges)
 			{
-				DatabaseDependencies.Add(Dependency);
+				if (bAddUnmirrored)
+				{
+					SearchIndex.Assets.Add(FPoseSearchIndexAsset(ESearchIndexAssetType::Sequence, SequenceIdx, false, FFloatInterval(Range.GetLowerBoundValue(), Range.GetUpperBoundValue())));
+				}
+
+				if (bAddMirrored)
+				{
+					SearchIndex.Assets.Add(FPoseSearchIndexAsset(ESearchIndexAssetType::Sequence, SequenceIdx, true, FFloatInterval(Range.GetLowerBoundValue(), Range.GetUpperBoundValue())));
+				}
 			}
+		}
+	}
 
-			SetState(EState::Prestarted);
+	TArray<FBlendSampleData> BlendSamples;
 
-			UE_LOG(LogPoseSearch, Log, TEXT("%s - %s BeginCache"), *LexToString(DerivedDataKey), *Database->GetName());
+	for (int32 BlendSpaceIdx = 0; BlendSpaceIdx < BlendSpaces.Num(); ++BlendSpaceIdx)
+	{
+		const FPoseSearchDatabaseBlendSpace& BlendSpace = BlendSpaces[BlendSpaceIdx];
+
+		if (BlendSpace.bEnabled)
+		{
+			const bool bAddUnmirrored = BlendSpace.MirrorOption == EPoseSearchMirrorOption::UnmirroredOnly || BlendSpace.MirrorOption == EPoseSearchMirrorOption::UnmirroredAndMirrored;
+			const bool bAddMirrored = BlendSpace.MirrorOption == EPoseSearchMirrorOption::MirroredOnly || BlendSpace.MirrorOption == EPoseSearchMirrorOption::UnmirroredAndMirrored;
 
-			TArray<FCacheGetRequest> CacheRequests;
-			const FCacheKey CacheKey{ Bucket, DerivedDataKey };
-			CacheRequests.Add({ { Database->GetPathName() }, CacheKey, ECachePolicy::Default });
+			int32 HorizontalBlendNum, VerticalBlendNum;
+			BlendSpace.GetBlendSpaceParameterSampleRanges(HorizontalBlendNum, VerticalBlendNum);
 
-			Owner = FRequestOwner(EPriority::Normal);
-			GetCache().Get(CacheRequests, Owner, [this](FCacheGetResponse&& Response)
+			const bool bWrapInputOnHorizontalAxis = BlendSpace.BlendSpace->GetBlendParameter(0).bWrapInput;
+			const bool bWrapInputOnVerticalAxis = BlendSpace.BlendSpace->GetBlendParameter(1).bWrapInput;
+			for (int32 HorizontalIndex = 0; HorizontalIndex < HorizontalBlendNum; HorizontalIndex++)
+			{
+				for (int32 VerticalIndex = 0; VerticalIndex < VerticalBlendNum; VerticalIndex++)
 				{
-					OnGetComplete(MoveTemp(Response));
-				});
+					const FVector BlendParameters = BlendSpace.BlendParameterForSampleRanges(HorizontalIndex, VerticalIndex);
+
+					int32 TriangulationIndex = 0;
+					BlendSpace.BlendSpace->GetSamplesFromBlendInput(BlendParameters, BlendSamples, TriangulationIndex, true);
+
+					float PlayLength = BlendSpace.BlendSpace->GetAnimationLengthFromSampleData(BlendSamples);
+
+					if (bAddUnmirrored)
+					{
+						SearchIndex.Assets.Add(FPoseSearchIndexAsset(ESearchIndexAssetType::BlendSpace, BlendSpaceIdx, false, FFloatInterval(0.0f, PlayLength), BlendParameters));
+					}
+
+					if (bAddMirrored)
+					{
+						SearchIndex.Assets.Add(FPoseSearchIndexAsset(ESearchIndexAssetType::BlendSpace, BlendSpaceIdx, true, FFloatInterval(0.0f, PlayLength), BlendParameters));
+					}
+				}
+			}
 		}
 	}
+}
+
+static void PreprocessSearchIndexWeights(FPoseSearchIndex& SearchIndex, const UPoseSearchSchema* Schema)
+{
+	const int32 NumDimensions = Schema->SchemaCardinality;
+	SearchIndex.WeightsSqrt.Init(1.f, NumDimensions);
 
-	// it cancels and waits for the task to be done and reset the local SearchIndex. SetState to Cancelled
-	void FPoseSearchDatabaseAsyncCacheTask::Cancel(FCriticalSection& OuterMutex)
+	for (int ChannelIdx = 0; ChannelIdx != Schema->Channels.Num(); ++ChannelIdx)
 	{
-		FScopeLock Lock(&OuterMutex);
+		const UPoseSearchFeatureChannel* Channel = Schema->Channels[ChannelIdx].Get();
+		Channel->FillWeights(SearchIndex.WeightsSqrt);
+	}
 
-		Owner.Cancel();
-		SearchIndex.Reset();
-		DerivedDataKey = FIoHash::Zero;
-		SetState(EState::Cancelled);
+	Eigen::VectorXd ChannelsMeanDeviations = ComputeChannelsDeviations(SearchIndex, Schema);
+	TArray<float> Deviation;
+	Deviation.Init(1.f, NumDimensions);
+	for (int32 Dimension = 0; Dimension != NumDimensions; ++Dimension)
+	{
+		const float ChannelsMeanDeviation = ChannelsMeanDeviations[Dimension];
+		Deviation[Dimension] = ChannelsMeanDeviation;
 	}
 
-	bool FPoseSearchDatabaseAsyncCacheTask::CancelIfDependsOn(const UObject* Object, FCriticalSection& OuterMutex)
+	EPoseSearchDataPreprocessor DataPreprocessor = Schema->DataPreprocessor;
+
+	if (DataPreprocessor == EPoseSearchDataPreprocessor::Normalize)
 	{
-		FScopeLock Lock(&OuterMutex);
+		// normalizing user weights: the idea behind this step is to be able to compare poses from databases using different schemas
+		RowMajorVectorMap MapWeights(SearchIndex.WeightsSqrt.GetData(), 1, NumDimensions);
+		const float WeightsSum = MapWeights.sum();
+		if (!FMath::IsNearlyZero(WeightsSum))
+		{
+			MapWeights *= (1.0f / WeightsSum);
+		}
+	}
 
-		// DatabaseDependencies is updated only in StartNewRequestIfNeeded when there are no active requests, so it's thread safe to access it 
-		if (DatabaseDependencies.Contains(Object))
+	// extracting the square root
+	for (int32 Dimension = 0; Dimension != NumDimensions; ++Dimension)
+	{
+		SearchIndex.WeightsSqrt[Dimension] = FMath::Sqrt(SearchIndex.WeightsSqrt[Dimension]);
+	}
+
+	if (DataPreprocessor == EPoseSearchDataPreprocessor::Normalize || DataPreprocessor == EPoseSearchDataPreprocessor::NormalizeOnlyByDeviation)
+	{
+		for (int32 Dimension = 0; Dimension != NumDimensions; ++Dimension)
 		{
-			Cancel(OuterMutex);
-			return true;
+			// the idea here is to premultiply the weights by the inverse of the variance (proportional to the square of the deviation) to have a "weighted Mahalanobis" distance
+			SearchIndex.WeightsSqrt[Dimension] /= Deviation[Dimension];
 		}
-		return false;
 	}
 
-	void FPoseSearchDatabaseAsyncCacheTask::Update(FCriticalSection& OuterMutex)
+#if WITH_EDITORONLY_DATA
+	SearchIndex.Deviation = Deviation;
+#endif // WITH_EDITORONLY_DATA
+}
+
+// it calculates Mean, PCAValues, and PCAProjectionMatrix
+static void PreprocessSearchIndexPCAData(FPoseSearchIndex& SearchIndex, int32 NumDimensions, uint32 NumberOfPrincipalComponents, EPoseSearchMode PoseSearchMode)
+{
+	// binding SearchIndex.Values and SearchIndex.PCAValues Eigen row major matrix maps
+	const int32 NumPoses = SearchIndex.NumPoses;
+
+	SearchIndex.PCAValues.Reset();
+	SearchIndex.Mean.Reset();
+	SearchIndex.PCAProjectionMatrix.Reset();
+
+	SearchIndex.PCAValues.AddZeroed(NumPoses * NumberOfPrincipalComponents);
+	SearchIndex.Mean.AddZeroed(NumDimensions);
+	SearchIndex.PCAProjectionMatrix.AddZeroed(NumDimensions * NumberOfPrincipalComponents);
+
+#if WITH_EDITORONLY_DATA
+	SearchIndex.PCAExplainedVariance = 0.f;
+#endif
+
+	if (NumDimensions > 0)
 	{
-		check(IsInGameThread());
+		const RowMajorVectorMapConst MapWeightsSqrt(SearchIndex.WeightsSqrt.GetData(), 1, NumDimensions);
+		const RowMajorMatrixMapConst MapValues(SearchIndex.Values.GetData(), NumPoses, NumDimensions);
+		const RowMajorMatrix WeightedValues = MapValues.array().rowwise() * MapWeightsSqrt.array();
+		RowMajorMatrixMap MapPCAValues(SearchIndex.PCAValues.GetData(), NumPoses, NumberOfPrincipalComponents);
+
+		// calculating the mean
+		RowMajorVectorMap Mean(SearchIndex.Mean.GetData(), 1, NumDimensions);
+		Mean = WeightedValues.colwise().mean();
+
+		// use the mean to center the data points
+		const RowMajorMatrix CenteredValues = WeightedValues.rowwise() - Mean;
+
+		// estimating the covariance matrix (with dimensionality of NumDimensions, NumDimensions)
+		// formula: https://en.wikipedia.org/wiki/Covariance_matrix#Estimation
+		// details: https://en.wikipedia.org/wiki/Estimation_of_covariance_matrices
+		const ColMajorMatrix CovariantMatrix = (CenteredValues.transpose() * CenteredValues) / float(NumPoses - 1);
+		const Eigen::SelfAdjointEigenSolver<ColMajorMatrix> EigenSolver(CovariantMatrix);
+
+		check(EigenSolver.info() == Eigen::Success);
+
+		// validating EigenSolver results
+		const ColMajorMatrix EigenVectors = EigenSolver.eigenvectors().real();
 
-		FScopeLock Lock(&OuterMutex);
+		if (PoseSearchMode == EPoseSearchMode::PCAKDTree_Validate && NumberOfPrincipalComponents == NumDimensions)
+		{
+			const RowMajorVector ReciprocalWeightsSqrt = MapWeightsSqrt.cwiseInverse();
+			const RowMajorMatrix ProjectedValues = CenteredValues * EigenVectors;
+			for (Eigen::Index RowIndex = 0; RowIndex < MapValues.rows(); ++RowIndex)
+			{
+				const RowMajorVector WeightedReconstructedPoint = ProjectedValues.row(RowIndex) * EigenVectors.transpose() + Mean;
+				const RowMajorVector ReconstructedPoint = WeightedReconstructedPoint.array() * ReciprocalWeightsSqrt.array();
+				const float Error = (ReconstructedPoint - MapValues.row(RowIndex)).squaredNorm();
+				check(Error < UE_KINDA_SMALL_NUMBER);
+			}
+		}
 
-		check(GetState() != EState::Cancelled); // otherwise FPoseSearchDatabaseAsyncCacheTask should have been already removed
+		// sorting EigenVectors by EigenValues, so we pick the most significant ones to compose our PCA projection matrix.
+		const RowMajorVector EigenValues = EigenSolver.eigenvalues().real();
+		TArray<size_t> Indexer;
+		Indexer.Reserve(NumDimensions);
+		for (size_t DimensionIndex = 0; DimensionIndex < NumDimensions; ++DimensionIndex)
+		{
+			Indexer.Push(DimensionIndex);
+		}
+		Indexer.Sort([&EigenValues](size_t a, size_t b)
+			{
+				return EigenValues[a] > EigenValues[b];
+			});
 
-		if (GetState() == EState::Prestarted && Poll(OuterMutex))
+		// composing the PCA projection matrix with the PCANumComponents most significant EigenVectors
+		ColMajorMatrixMap PCAProjectionMatrix(SearchIndex.PCAProjectionMatrix.GetData(), NumDimensions, NumberOfPrincipalComponents);
+		float AccumulatedVariance = 0.f;
+		for (size_t PCAComponentIndex = 0; PCAComponentIndex < NumberOfPrincipalComponents; ++PCAComponentIndex)
 		{
-			// task is done: we need to update the state form Prestarted to Ended/Failed
-			Wait(OuterMutex);
+			PCAProjectionMatrix.col(PCAComponentIndex) = EigenVectors.col(Indexer[PCAComponentIndex]);
+			AccumulatedVariance += EigenValues[Indexer[PCAComponentIndex]];
 		}
 
-		if (bBroadcastOnDerivedDataRebuild)
+#if WITH_EDITORONLY_DATA
+		// calculating the total variance knowing that eigen values measure variance along the principal components:
+		const float TotalVariance = EigenValues.sum();
+		// and explained variance as ratio between AccumulatedVariance and TotalVariance: https://ro-che.info/articles/2017-12-11-pca-explained-variance
+		SearchIndex.PCAExplainedVariance = TotalVariance > UE_KINDA_SMALL_NUMBER ? AccumulatedVariance / TotalVariance : 0.f;
+#endif // WITH_EDITORONLY_DATA
+
+		MapPCAValues = CenteredValues * PCAProjectionMatrix;
+
+		if (PoseSearchMode == EPoseSearchMode::PCAKDTree_Validate && NumberOfPrincipalComponents == NumDimensions)
 		{
-			Database->NotifyDerivedDataRebuild();
-			bBroadcastOnDerivedDataRebuild = false;
+			const RowMajorVector ReciprocalWeightsSqrt = MapWeightsSqrt.cwiseInverse();
+			for (Eigen::Index RowIndex = 0; RowIndex < MapValues.rows(); ++RowIndex)
+			{
+				const RowMajorVector WeightedReconstructedValues = MapPCAValues.row(RowIndex) * PCAProjectionMatrix.transpose() + Mean;
+				const RowMajorVector ReconstructedValues = WeightedReconstructedValues.array() * ReciprocalWeightsSqrt.array();
+				const float Error = (ReconstructedValues - MapValues.row(RowIndex)).squaredNorm();
+				check(Error < UE_KINDA_SMALL_NUMBER);
+			}
 		}
 	}
+}
+
+static void PreprocessSearchIndexKDTree(FPoseSearchIndex& SearchIndex, int32 NumDimensions, uint32 NumberOfPrincipalComponents, EPoseSearchMode PoseSearchMode, int32 KDTreeMaxLeafSize, int32 KDTreeQueryNumNeighbors)
+{
+	const int32 NumPoses = SearchIndex.NumPoses;
+	SearchIndex.KDTree.Construct(NumPoses, NumberOfPrincipalComponents, SearchIndex.PCAValues.GetData(), KDTreeMaxLeafSize);
 
-	// it waits for the task to be done and SetSearchIndex on the database. SetState to Ended/Failed
-	void FPoseSearchDatabaseAsyncCacheTask::Wait(FCriticalSection& OuterMutex)
+	if (PoseSearchMode == EPoseSearchMode::PCAKDTree_Validate)
 	{
-		check(GetState() == EState::Prestarted);
+		// testing the KDTree is returning the proper searches for all the points in pca space
+		int32 NumberOfFailingPoints = 0;
+		for (size_t PointIndex = 0; PointIndex < NumPoses; ++PointIndex)
+		{
+			TArray<size_t> ResultIndexes;
+			TArray<float> ResultDistanceSqr;
+			ResultIndexes.SetNum(KDTreeQueryNumNeighbors + 1);
+			ResultDistanceSqr.SetNum(KDTreeQueryNumNeighbors + 1);
+			FKDTree::KNNResultSet ResultSet(KDTreeQueryNumNeighbors, ResultIndexes, ResultDistanceSqr);
+			SearchIndex.KDTree.FindNeighbors(ResultSet, &SearchIndex.PCAValues[PointIndex * NumberOfPrincipalComponents]);
+
+			size_t ResultIndex = 0;
+			for (; ResultIndex < ResultSet.Num(); ++ResultIndex)
+			{
+				if (PointIndex == ResultIndexes[ResultIndex])
+				{
+					check(ResultDistanceSqr[ResultIndex] < UE_KINDA_SMALL_NUMBER);
+					break;
+				}
+			}
+			if (ResultIndex == ResultSet.Num())
+			{
+				++NumberOfFailingPoints;
+			}
+		}
+
+		check(NumberOfFailingPoints == 0);
+
+		// testing the KDTree is returning the proper searches for all the original points transformed in pca space
+		NumberOfFailingPoints = 0;
+		for (size_t PointIndex = 0; PointIndex < NumPoses; ++PointIndex)
+		{
+			TArray<size_t> ResultIndexes;
+			TArray<float> ResultDistanceSqr;
+			ResultIndexes.SetNum(KDTreeQueryNumNeighbors + 1);
+			ResultDistanceSqr.SetNum(KDTreeQueryNumNeighbors + 1);
+			FKDTree::KNNResultSet ResultSet(KDTreeQueryNumNeighbors, ResultIndexes, ResultDistanceSqr);
 
-		Owner.Wait();
+			const RowMajorVectorMapConst MapValues(&SearchIndex.Values[PointIndex * NumDimensions], 1, NumDimensions);
+			const RowMajorVectorMapConst MapWeightsSqrt(SearchIndex.WeightsSqrt.GetData(), 1, NumDimensions);
+			const RowMajorVectorMapConst Mean(SearchIndex.Mean.GetData(), 1, NumDimensions);
+			const ColMajorMatrixMapConst PCAProjectionMatrix(SearchIndex.PCAProjectionMatrix.GetData(), NumDimensions, NumberOfPrincipalComponents);
 
-		FScopeLock Lock(&OuterMutex);
+			const RowMajorMatrix WeightedValues = MapValues.array() * MapWeightsSqrt.array();
+			const RowMajorMatrix CenteredValues = WeightedValues - Mean;
+			const RowMajorVector ProjectedValues = CenteredValues * PCAProjectionMatrix;
 
-		const bool bFailedIndexing = SearchIndex.IsEmpty();
-		if (!bFailedIndexing)
+			SearchIndex.KDTree.FindNeighbors(ResultSet, ProjectedValues.data());
+
+			size_t ResultIndex = 0;
+			for (; ResultIndex < ResultSet.Num(); ++ResultIndex)
+			{
+				if (PointIndex == ResultIndexes[ResultIndex])
+				{
+					check(ResultDistanceSqr[ResultIndex] < UE_KINDA_SMALL_NUMBER);
+					break;
+				}
+			}
+			if (ResultIndex == ResultSet.Num())
+			{
+				++NumberOfFailingPoints;
+			}
+		}
+
+		check(NumberOfFailingPoints == 0);
+	}
+}
+
+//////////////////////////////////////////////////////////////////////////
+// FSamplingParam helpers
+struct FSamplingParam
+{
+	float WrappedParam = 0.0f;
+	int32 NumCycles = 0;
+
+	// If the animation can't loop, WrappedParam contains the clamped value and whatever is left is stored here
+	float Extrapolation = 0.0f;
+};
+
+static FSamplingParam WrapOrClampSamplingParam(bool bCanWrap, float SamplingParamExtent, float SamplingParam)
+{
+	// This is a helper function used by both time and distance sampling. A schema may specify time or distance
+	// offsets that are multiple cycles of a clip away from the current pose being sampled.
+	// And that time or distance offset may before the beginning of the clip (SamplingParam < 0.0f)
+	// or after the end of the clip (SamplingParam > SamplingParamExtent). So this function
+	// helps determine how many cycles need to be applied and what the wrapped value should be, clamping
+	// if necessary.
+
+	FSamplingParam Result;
+
+	Result.WrappedParam = SamplingParam;
+
+	const bool bIsSamplingParamExtentKindaSmall = SamplingParamExtent <= UE_KINDA_SMALL_NUMBER;
+	if (!bIsSamplingParamExtentKindaSmall && bCanWrap)
+	{
+		if (SamplingParam < 0.0f)
 		{
-			Database->SetSearchIndex(SearchIndex); // @todo: implement FPoseSearchIndex move ctor and assignment operator and use a MoveTemp(SearchIndex) here
-			SetState(EState::Ended);
-			bBroadcastOnDerivedDataRebuild = true;
+			while (Result.WrappedParam < 0.0f)
+			{
+				Result.WrappedParam += SamplingParamExtent;
+				++Result.NumCycles;
+			}
 		}
+
 		else
 		{
-			check(!bBroadcastOnDerivedDataRebuild);
-			SetState(EState::Failed);
+			while (Result.WrappedParam > SamplingParamExtent)
+			{
+				Result.WrappedParam -= SamplingParamExtent;
+				++Result.NumCycles;
+			}
 		}
-		SearchIndex.Reset();
 	}
 
-	// true is the task is done executing
-	bool FPoseSearchDatabaseAsyncCacheTask::Poll(FCriticalSection& OuterMutex) const
+	const float ParamClamped = FMath::Clamp(Result.WrappedParam, 0.0f, SamplingParamExtent);
+	if (ParamClamped != Result.WrappedParam)
+	{
+		check(bIsSamplingParamExtentKindaSmall || !bCanWrap);
+		Result.Extrapolation = Result.WrappedParam - ParamClamped;
+		Result.WrappedParam = ParamClamped;
+	}
+
+	return Result;
+}
+
+//////////////////////////////////////////////////////////////////////////
+// FAssetIndexer
+class FAssetIndexer : public IAssetIndexer
+{
+public:
+
+	struct FOutput
+	{
+		int32 FirstIndexedSample = 0;
+		int32 LastIndexedSample = 0;
+		int32 NumIndexedPoses = 0;
+
+		TArray<float> FeatureVectorTable;
+		TArray<FPoseSearchPoseMetadata> PoseMetadata;
+		TBitArray<> AllFeaturesNotAdded;
+	} Output;
+
+	void Reset();
+	void Init(const FAssetIndexingContext& IndexingContext, const FBoneContainer& InBoneContainer);
+	bool Process();
+
+public: // IAssetIndexer
+
+	const FAssetIndexingContext& GetIndexingContext() const override { return IndexingContext; }
+	FSampleInfo GetSampleInfo(float SampleTime) const override;
+	FSampleInfo GetSampleInfoRelative(float SampleTime, const FSampleInfo& Origin) const override;
+	const float GetSampleTimeFromDistance(float Distance) const override;
+	FTransform MirrorTransform(const FTransform& Transform) const override;
+	FTransform GetTransformAndCacheResults(float SampleTime, float OriginTime, int8 SchemaBoneIdx, bool& Clamped) override;
+
+private:
+	FPoseSearchPoseMetadata GetMetadata(int32 SampleIdx) const;
+
+	struct CachedEntry
+	{
+		float SampleTime;
+		float OriginTime;
+		bool Clamped;
+
+		// @todo: minimize the Entry memory footprint
+		FTransform RootTransform;
+		FCompactPose Pose;
+		FCSPose<FCompactPose> ComponentSpacePose;
+		FBlendedCurve UnusedCurve;
+		UE::Anim::FStackAttributeContainer UnusedAtrribute;
+		FAnimationPoseData AnimPoseData = { Pose, UnusedCurve, UnusedAtrribute };
+	};
+
+	FBoneContainer BoneContainer;
+	FAssetIndexingContext IndexingContext;
+	TArray<CachedEntry> CachedEntries;
+};
+
+void FAssetIndexer::Reset()
+{
+	Output.FirstIndexedSample = 0;
+	Output.LastIndexedSample = 0;
+	Output.NumIndexedPoses = 0;
+
+	Output.FeatureVectorTable.Reset(0);
+	Output.PoseMetadata.Reset(0);
+	Output.AllFeaturesNotAdded.Reset();
+}
+
+void FAssetIndexer::Init(const FAssetIndexingContext& InIndexingContext, const FBoneContainer& InBoneContainer)
+{
+	check(InIndexingContext.Schema);
+	check(InIndexingContext.Schema->IsValid());
+	check(InIndexingContext.MainSampler);
+
+	BoneContainer = InBoneContainer;
+	IndexingContext = InIndexingContext;
+
+	Reset();
+
+	Output.FirstIndexedSample = FMath::FloorToInt(IndexingContext.RequestedSamplingRange.Min * IndexingContext.Schema->SampleRate);
+	Output.LastIndexedSample = FMath::Max(0, FMath::CeilToInt(IndexingContext.RequestedSamplingRange.Max * IndexingContext.Schema->SampleRate));
+	Output.NumIndexedPoses = Output.LastIndexedSample - Output.FirstIndexedSample + 1;
+
+	Output.FeatureVectorTable.SetNumZeroed(IndexingContext.Schema->SchemaCardinality * Output.NumIndexedPoses);
+	Output.PoseMetadata.SetNum(Output.NumIndexedPoses);
+}
+
+bool FAssetIndexer::Process()
+{
+	check(IndexingContext.Schema);
+	check(IndexingContext.Schema->IsValid());
+	check(IndexingContext.MainSampler);
+
+	FMemMark Mark(FMemStack::Get());
+
+	IndexingContext.BeginSampleIdx = Output.FirstIndexedSample;
+	IndexingContext.EndSampleIdx = Output.LastIndexedSample + 1;
+
+	if (IndexingContext.Schema->SchemaCardinality > 0)
 	{
-		return Owner.Poll();
+		// Index each channel
+		FAssetIndexingOutput AssetIndexingOutput(IndexingContext.Schema->SchemaCardinality, Output.FeatureVectorTable);
+		for (int32 ChannelIdx = 0; ChannelIdx != IndexingContext.Schema->Channels.Num(); ++ChannelIdx)
+		{
+			IndexingContext.Schema->Channels[ChannelIdx]->IndexAsset(*this, AssetIndexingOutput);
+		}
 	}
 
-	bool FPoseSearchDatabaseAsyncCacheTask::ContainsDatabase(const UPoseSearchDatabase* OtherDatabase, FCriticalSection& OuterMutex) const
+	// Generate pose metadata
+	for (int32 SampleIdx = IndexingContext.BeginSampleIdx; SampleIdx != IndexingContext.EndSampleIdx; ++SampleIdx)
 	{
-		FScopeLock Lock(&OuterMutex);
-		return Database.Get() == OtherDatabase;
+		const int32 PoseIdx = SampleIdx - Output.FirstIndexedSample;
+		FPoseSearchPoseMetadata& OutputPoseMetadata = Output.PoseMetadata[PoseIdx];
+		OutputPoseMetadata = GetMetadata(SampleIdx);
 	}
 
-	// called once the task is done:
-	// if EStatus::Ok (data has been retrieved from DDC) we deserialize the payload into the local SearchIndex
-	// if EStatus::Error we BuildIndex and if that's successful we 'Put' it on DDC
-	void FPoseSearchDatabaseAsyncCacheTask::OnGetComplete(UE::DerivedData::FCacheGetResponse&& Response)
+	return true;
+}
+
+const float FAssetIndexer::GetSampleTimeFromDistance(float SampleDistance) const
+{
+	auto CanWrapDistanceSamples = [](const IAssetSampler* Sampler) -> bool
 	{
-		using namespace UE::DerivedData;
+		constexpr float SMALL_ROOT_DISTANCE = 1.0f;
+		return Sampler->IsLoopable() && Sampler->GetTotalRootDistance() > SMALL_ROOT_DISTANCE;
+	};
+
+	float MainTotalDistance = IndexingContext.MainSampler->GetTotalRootDistance();
+	bool bMainCanWrap = CanWrapDistanceSamples(IndexingContext.MainSampler);
 
-		const FCacheKey Key = Response.Record.GetKey();
-		if (Response.Status == EStatus::Ok)
+	float SampleTime = MAX_flt;
+
+	if (!bMainCanWrap)
+	{
+		// Use the lead in anim if we would have to clamp to the beginning of the main anim
+		if (IndexingContext.LeadInSampler && (SampleDistance < 0.0f))
 		{
-			COOK_STAT(auto Timer = UsageStats.TimeAsyncWait());
+			const IAssetSampler* ClipSampler = IndexingContext.LeadInSampler;
 
-			// we found the cached data associated to the PendingDerivedDataKey: we'll deserialized into SearchIndex
-			SearchIndex.Reset();
-			FSharedBuffer RawData = Response.Record.GetValue(Id).GetData().Decompress();
-			FMemoryReaderView Reader(RawData);
-			Reader << SearchIndex;
+			bool bLeadInCanWrap = CanWrapDistanceSamples(IndexingContext.LeadInSampler);
+			float LeadRelativeDistance = SampleDistance + ClipSampler->GetTotalRootDistance();
+			FSamplingParam SamplingParam = WrapOrClampSamplingParam(bLeadInCanWrap, ClipSampler->GetTotalRootDistance(), LeadRelativeDistance);
 
-			UE_LOG(LogPoseSearch, Log, TEXT("%s - %s BuildIndex From Cache"), *LexToString(Key.Hash), *Database->GetName());
+			float ClipTime = ClipSampler->GetTimeFromRootDistance(
+				SamplingParam.WrappedParam + SamplingParam.Extrapolation);
 
-			COOK_STAT(Timer.AddHit(RawData.GetSize()));
+			// Make the lead in clip time relative to the main sequence again and unwrap
+			SampleTime = -((SamplingParam.NumCycles * ClipSampler->GetPlayLength()) + (ClipSampler->GetPlayLength() - ClipTime));
 		}
-		else if (Response.Status == EStatus::Error)
+
+		// Use the follow up anim if we would have clamp to the end of the main anim
+		else if (IndexingContext.FollowUpSampler && (SampleDistance > MainTotalDistance))
 		{
-			// we didn't find the cached data associated to the PendingDerivedDataKey: we'll BuildIndex to update SearchIndex and "Put" the data over the DDC
-			Owner.LaunchTask(TEXT("PoseSearchDatabaseBuild"), [this, Key]
-				{
-					COOK_STAT(auto Timer = UsageStats.TimeSyncWork());
+			const IAssetSampler* ClipSampler = IndexingContext.FollowUpSampler;
 
-					int32 BytesProcessed = 0;
-					if (BuildIndex(*Database, SearchIndex, Owner))
-					{
-						UE_LOG(LogPoseSearch, Log, TEXT("%s - %s BuildIndex Succeeded"), *LexToString(Key.Hash), *Database->GetName());
-
-						TArray<uint8> RawBytes;
-						FMemoryWriter Writer(RawBytes);
-						Writer << SearchIndex;
-						FSharedBuffer RawData = MakeSharedBufferFromArray(MoveTemp(RawBytes));
-						BytesProcessed = RawData.GetSize();
-
-						FCacheRecordBuilder Builder(Key);
-						Builder.AddValue(Id, RawData);
-						GetCache().Put({ { { Database->GetPathName() }, Builder.Build() } }, Owner, [this, Key](FCachePutResponse&& Response)
-							{
-								if (Response.Status == EStatus::Error)
-								{
-									UE_LOG(LogPoseSearch, Log, TEXT("%s - %s Failed to store DDC"), *LexToString(Key.Hash), *Database->GetName());
-								}
-							});
-					}
-					else
-					{
-						if (Owner.IsCanceled())
-						{
-							UE_LOG(LogPoseSearch, Log, TEXT("%s - %s BuildIndex Cancelled"), *LexToString(Key.Hash), *Database->GetName());
-						}
-						else
-						{
-							UE_LOG(LogPoseSearch, Log, TEXT("%s - %s BuildIndex Failed"), *LexToString(Key.Hash), *Database->GetName());
-						}
-						SearchIndex.Reset();
-					}
+			bool bFollowUpCanWrap = CanWrapDistanceSamples(IndexingContext.FollowUpSampler);
+			float FollowRelativeDistance = SampleDistance - MainTotalDistance;
+			FSamplingParam SamplingParam = WrapOrClampSamplingParam(bFollowUpCanWrap, ClipSampler->GetTotalRootDistance(), FollowRelativeDistance);
+
+			float ClipTime = ClipSampler->GetTimeFromRootDistance(
+				SamplingParam.WrappedParam + SamplingParam.Extrapolation);
 
-					COOK_STAT(Timer.AddMiss(BytesProcessed));
-				});
+			// Make the follow up clip time relative to the main sequence again and unwrap
+			SampleTime = IndexingContext.MainSampler->GetPlayLength() + SamplingParam.NumCycles * ClipSampler->GetPlayLength() + ClipTime;
+		}
+	}
+
+	// Use the main anim if we didn't use the lead-in or follow-up anims.
+	// The main anim sample may have been wrapped or clamped
+	if (SampleTime == MAX_flt)
+	{
+		float MainRelativeDistance = SampleDistance;
+		if (SampleDistance < 0.0f && bMainCanWrap)
+		{
+			// In this case we're sampling a loop backwards, so MainRelativeDistance must adjust so the number of cycles 
+			// is counted correctly.
+			MainRelativeDistance += IndexingContext.MainSampler->GetTotalRootDistance();
 		}
-		else if(Response.Status == EStatus::Canceled)
+
+		FSamplingParam SamplingParam = WrapOrClampSamplingParam(bMainCanWrap, MainTotalDistance, MainRelativeDistance);
+		float ClipTime = IndexingContext.MainSampler->GetTimeFromRootDistance(
+			SamplingParam.WrappedParam + SamplingParam.Extrapolation);
+
+		// Unwrap the main clip time
+		if (bMainCanWrap)
 		{
-			SearchIndex.Reset();
-			UE_LOG(LogPoseSearch, Log, TEXT("%s - %s BuildIndex Cancelled"), *LexToString(Key.Hash), *Database->GetName());
+			if (SampleDistance < 0.0f)
+			{
+				SampleTime = -((SamplingParam.NumCycles * IndexingContext.MainSampler->GetPlayLength()) + (IndexingContext.MainSampler->GetPlayLength() - ClipTime));
+			}
+			else
+			{
+				SampleTime = SamplingParam.NumCycles * IndexingContext.MainSampler->GetPlayLength() + ClipTime;
+			}
+		}
+		else
+		{
+			SampleTime = ClipTime;
 		}
 	}
 
-	//////////////////////////////////////////////////////////////////////////
-	// FAsyncPoseSearchDatabasesManagement
-	FCriticalSection FAsyncPoseSearchDatabasesManagement::Mutex;
+	return SampleTime;
+}
+
+FAssetIndexer::FSampleInfo FAssetIndexer::GetSampleInfo(float SampleTime) const
+{
+	FSampleInfo Sample;
+
+	FTransform RootMotionLast = FTransform::Identity;
+	FTransform RootMotionInitial = FTransform::Identity;
+
+	float RootDistanceLast = 0.0f;
+	float RootDistanceInitial = 0.0f;
+
+	check(IndexingContext.MainSampler);
+
+	const float MainPlayLength = IndexingContext.MainSampler->GetPlayLength();
+	const bool bMainCanWrap = IndexingContext.MainSampler->IsLoopable();
 
-	FAsyncPoseSearchDatabasesManagement& FAsyncPoseSearchDatabasesManagement::Get()
+	FSamplingParam SamplingParam;
+	if (!bMainCanWrap)
 	{
-		FScopeLock Lock(&Mutex);
+		// Use the lead in anim if we would have to clamp to the beginning of the main anim
+		if (IndexingContext.LeadInSampler && (SampleTime < 0.0f))
+		{
+			const IAssetSampler* ClipSampler = IndexingContext.LeadInSampler;
+
+			const bool bLeadInCanWrap = IndexingContext.LeadInSampler->IsLoopable();
+			const float LeadRelativeTime = SampleTime + ClipSampler->GetPlayLength();
+			SamplingParam = WrapOrClampSamplingParam(bLeadInCanWrap, ClipSampler->GetPlayLength(), LeadRelativeTime);
+
+			Sample.Clip = IndexingContext.LeadInSampler;
+
+			check(SamplingParam.Extrapolation <= 0.0f);
+			if (SamplingParam.Extrapolation < 0.0f)
+			{
+				RootMotionInitial = IndexingContext.LeadInSampler->GetTotalRootTransform().Inverse();
+				RootDistanceInitial = -IndexingContext.LeadInSampler->GetTotalRootDistance();
+			}
+			else
+			{
+				RootMotionInitial = FTransform::Identity;
+				RootDistanceInitial = 0.0f;
+			}
 
-		static FAsyncPoseSearchDatabasesManagement SingletonInstance;
-		return SingletonInstance;
+			RootMotionLast = IndexingContext.LeadInSampler->GetTotalRootTransform();
+			RootDistanceLast = IndexingContext.LeadInSampler->GetTotalRootDistance();
+		}
+
+		// Use the follow up anim if we would have clamp to the end of the main anim
+		else if (IndexingContext.FollowUpSampler && (SampleTime > MainPlayLength))
+		{
+			const IAssetSampler* ClipSampler = IndexingContext.FollowUpSampler;
+
+			const bool bFollowUpCanWrap = IndexingContext.FollowUpSampler->IsLoopable();
+			const float FollowRelativeTime = SampleTime - MainPlayLength;
+			SamplingParam = WrapOrClampSamplingParam(bFollowUpCanWrap, ClipSampler->GetPlayLength(), FollowRelativeTime);
+
+			Sample.Clip = IndexingContext.FollowUpSampler;
+
+			RootMotionInitial = IndexingContext.MainSampler->GetTotalRootTransform();
+			RootDistanceInitial = IndexingContext.MainSampler->GetTotalRootDistance();
+
+			RootMotionLast = IndexingContext.FollowUpSampler->GetTotalRootTransform();
+			RootDistanceLast = IndexingContext.FollowUpSampler->GetTotalRootDistance();
+		}
 	}
 
-	FAsyncPoseSearchDatabasesManagement::FAsyncPoseSearchDatabasesManagement()
-		: Tasks(*(new FPoseSearchDatabaseAsyncCacheTasks()))
+	// Use the main anim if we didn't use the lead-in or follow-up anims.
+	// The main anim sample may have been wrapped or clamped
+	if (!Sample.IsValid())
 	{
-		FScopeLock Lock(&Mutex);
+		float MainRelativeTime = SampleTime;
+		if (SampleTime < 0.0f && bMainCanWrap)
+		{
+			// In this case we're sampling a loop backwards, so MainRelativeTime must adjust so the number of cycles is
+			// counted correctly.
+			MainRelativeTime += MainPlayLength;
+		}
+
+		SamplingParam = WrapOrClampSamplingParam(bMainCanWrap, MainPlayLength, MainRelativeTime);
+
+		Sample.Clip = IndexingContext.MainSampler;
 
-		OnObjectModifiedHandle = FCoreUObjectDelegates::OnObjectModified.AddRaw(this, &FAsyncPoseSearchDatabasesManagement::OnObjectModified);
-		FCoreDelegates::OnPreExit.AddRaw(this, &FAsyncPoseSearchDatabasesManagement::Shutdown);
+		RootMotionInitial = FTransform::Identity;
+		RootDistanceInitial = 0.0f;
+
+		RootMotionLast = IndexingContext.MainSampler->GetTotalRootTransform();
+		RootDistanceLast = IndexingContext.MainSampler->GetTotalRootDistance();
 	}
 
-	FAsyncPoseSearchDatabasesManagement::~FAsyncPoseSearchDatabasesManagement()
+
+	if (FMath::Abs(SamplingParam.Extrapolation) > SMALL_NUMBER)
 	{
-		FScopeLock Lock(&Mutex);
+		Sample.bClamped = true;
+		Sample.ClipTime = SamplingParam.WrappedParam + SamplingParam.Extrapolation;
+		const FTransform ClipRootMotion = Sample.Clip->ExtractRootTransform(Sample.ClipTime);
+		const float ClipDistance = Sample.Clip->ExtractRootDistance(Sample.ClipTime);
 
-		FCoreDelegates::OnPreExit.RemoveAll(this);
-		Shutdown();
+		Sample.RootTransform = ClipRootMotion * RootMotionInitial;
+		Sample.RootDistance = RootDistanceInitial + ClipDistance;
+	}
+	else
+	{
+		Sample.ClipTime = SamplingParam.WrappedParam;
+
+		// Determine how to accumulate motion for every cycle of the anim. If the sample
+		// had to be clamped, this motion will end up not getting applied below.
+		// Also invert the accumulation direction if the requested sample was wrapped backwards.
+		FTransform RootMotionPerCycle = RootMotionLast;
+		float RootDistancePerCycle = RootDistanceLast;
+		if (SampleTime < 0.0f)
+		{
+			RootMotionPerCycle = RootMotionPerCycle.Inverse();
+			RootDistancePerCycle *= -1.f;
+		}
+
+		// Find the remaining motion deltas after wrapping
+		FTransform RootMotionRemainder = Sample.Clip->ExtractRootTransform(Sample.ClipTime);
+		float RootDistanceRemainder = Sample.Clip->ExtractRootDistance(Sample.ClipTime);
+
+		// Invert motion deltas if we wrapped backwards
+		if (SampleTime < 0.0f)
+		{
+			RootMotionRemainder.SetToRelativeTransform(RootMotionLast);
+			RootDistanceRemainder = -(RootDistanceLast - RootDistanceRemainder);
+		}
+
+		Sample.RootTransform = RootMotionInitial;
+		Sample.RootDistance = RootDistanceInitial;
+
+		// Note if the sample was clamped, no motion will be applied here because NumCycles will be zero
+		int32 CyclesRemaining = SamplingParam.NumCycles;
+		while (CyclesRemaining--)
+		{
+			Sample.RootTransform = RootMotionPerCycle * Sample.RootTransform;
+			Sample.RootDistance += RootDistancePerCycle;
+		}
 
-		delete &Tasks;
+		Sample.RootTransform = RootMotionRemainder * Sample.RootTransform;
+		Sample.RootDistance += RootDistanceRemainder;
 	}
 
-	// we're listening to OnObjectModified to cancel any pending Task indexing databases depending from Object to avoid multi threading issues
-	void FAsyncPoseSearchDatabasesManagement::OnObjectModified(UObject* Object)
+	return Sample;
+}
+
+FAssetIndexer::FSampleInfo FAssetIndexer::GetSampleInfoRelative(float SampleTime, const FSampleInfo& Origin) const
+{
+	FSampleInfo Sample = GetSampleInfo(SampleTime);
+	Sample.RootTransform.SetToRelativeTransform(Origin.RootTransform);
+	Sample.RootDistance = Origin.RootDistance - Sample.RootDistance;
+	return Sample;
+}
+
+FTransform FAssetIndexer::MirrorTransform(const FTransform& Transform) const
+{
+	return IndexingContext.bMirrored ? IndexingContext.SamplingContext->MirrorTransform(Transform) : Transform;
+}
+
+FPoseSearchPoseMetadata FAssetIndexer::GetMetadata(int32 SampleIdx) const
+{
+	const float SequenceLength = IndexingContext.MainSampler->GetPlayLength();
+	const float SampleTime = FMath::Min(SampleIdx * IndexingContext.Schema->GetSamplingInterval(), SequenceLength);
+
+	FPoseSearchPoseMetadata Metadata;
+	Metadata.CostAddend = IndexingContext.Schema->BaseCostBias;
+	Metadata.ContinuingPoseCostAddend = IndexingContext.Schema->ContinuingPoseCostBias;
+
+	TArray<UAnimNotifyState_PoseSearchBase*> NotifyStates;
+	IndexingContext.MainSampler->ExtractPoseSearchNotifyStates(SampleTime, NotifyStates);
+	for (const UAnimNotifyState_PoseSearchBase* PoseSearchNotify : NotifyStates)
 	{
-		FScopeLock Lock(&Mutex);
+		if (PoseSearchNotify->GetClass()->IsChildOf<UAnimNotifyState_PoseSearchBlockTransition>())
+		{
+			EnumAddFlags(Metadata.Flags, EPoseSearchPoseFlags::BlockTransition);
+		}
+		else if (PoseSearchNotify->GetClass()->IsChildOf<UAnimNotifyState_PoseSearchModifyCost>())
+		{
+			const UAnimNotifyState_PoseSearchModifyCost* ModifyCostNotify =
+				Cast<const UAnimNotifyState_PoseSearchModifyCost>(PoseSearchNotify);
+			Metadata.CostAddend = ModifyCostNotify->CostAddend;
+		}
+		else if (PoseSearchNotify->GetClass()->IsChildOf<UAnimNotifyState_PoseSearchOverrideContinuingPoseCostBias>())
+		{
+			const UAnimNotifyState_PoseSearchOverrideContinuingPoseCostBias* ContinuingPoseCostBias =
+				Cast<const UAnimNotifyState_PoseSearchOverrideContinuingPoseCostBias>(PoseSearchNotify);
+			Metadata.ContinuingPoseCostAddend = ContinuingPoseCostBias->CostAddend;
+		}
+	}
+	return Metadata;
+}
 
-		// iterating backwards because of the possible RemoveAtSwap
-		for (int32 TaskIndex = Tasks.Num() - 1; TaskIndex >= 0; --TaskIndex)
+FTransform FAssetIndexer::GetTransformAndCacheResults(float SampleTime, float OriginTime, int8 SchemaBoneIdx, bool& Clamped)
+{
+	// @todo: use an hashmap if we end up having too many entries
+	CachedEntry* Entry = CachedEntries.FindByPredicate([SampleTime, OriginTime](const FAssetIndexer::CachedEntry& Entry)
 		{
-			if (Tasks[TaskIndex]->CancelIfDependsOn(Object, Mutex))
+			return Entry.SampleTime == SampleTime && Entry.OriginTime == OriginTime;
+		});
+
+	const FAssetSamplingContext* SamplingContext = IndexingContext.SamplingContext;
+
+	if (!Entry)
+	{
+		Entry = &CachedEntries[CachedEntries.AddDefaulted()];
+
+		Entry->SampleTime = SampleTime;
+		Entry->OriginTime = OriginTime;
+
+		if (!BoneContainer.IsValid())
+		{
+			UE_LOG(LogPoseSearch,
+				Warning,
+				TEXT("Invalid BoneContainer encountered in FAssetIndexer::GetTransformAndCacheResults. Asset: %s. Schema: %s. BoneContainerAsset: %s. NumBoneIndices: %d"),
+				*GetNameSafe(IndexingContext.MainSampler->GetAsset()),
+				*GetNameSafe(IndexingContext.Schema),
+				*GetNameSafe(BoneContainer.GetAsset()),
+				BoneContainer.GetCompactPoseNumBones());
+		}
+
+		Entry->Pose.SetBoneContainer(&BoneContainer);
+		Entry->UnusedCurve.InitFrom(BoneContainer);
+
+		IAssetIndexer::FSampleInfo Origin = GetSampleInfo(OriginTime);
+		IAssetIndexer::FSampleInfo Sample = GetSampleInfoRelative(SampleTime, Origin);
+
+		float CurrentTime = Sample.ClipTime;
+		float PreviousTime = CurrentTime - SamplingContext->FiniteDelta;
+
+		FDeltaTimeRecord DeltaTimeRecord;
+		DeltaTimeRecord.Set(PreviousTime, CurrentTime - PreviousTime);
+		FAnimExtractContext ExtractionCtx(static_cast<double>(CurrentTime), true, DeltaTimeRecord, Sample.Clip->IsLoopable());
+
+		Sample.Clip->ExtractPose(ExtractionCtx, Entry->AnimPoseData);
+
+		if (IndexingContext.bMirrored)
+		{
+			FAnimationRuntime::MirrorPose(
+				Entry->AnimPoseData.GetPose(),
+				IndexingContext.Schema->MirrorDataTable->MirrorAxis,
+				SamplingContext->CompactPoseMirrorBones,
+				SamplingContext->ComponentSpaceRefRotations
+			);
+			// Note curves and attributes are not used during the indexing process and therefore don't need to be mirrored
+		}
+
+		Entry->ComponentSpacePose.InitPose(Entry->Pose);
+		Entry->RootTransform = Sample.RootTransform;
+		Entry->Clamped = Sample.bClamped;
+	}
+
+	const FBoneReference& BoneReference = IndexingContext.Schema->BoneReferences[SchemaBoneIdx];
+	FCompactPoseBoneIndex CompactBoneIndex = BoneContainer.MakeCompactPoseIndex(FMeshPoseBoneIndex(BoneReference.BoneIndex));
+
+	const FTransform BoneTransform = Entry->ComponentSpacePose.GetComponentSpaceTransform(CompactBoneIndex) * MirrorTransform(Entry->RootTransform);
+	Clamped = Entry->Clamped;
+
+	return BoneTransform;
+}
+
+//////////////////////////////////////////////////////////////////////////
+// FDatabaseIndexingContext
+struct FDatabaseIndexingContext
+{
+	FPoseSearchIndexBase* SearchIndexBase = nullptr;
+
+	FAssetSamplingContext SamplingContext;
+	TArray<FSequenceSampler> SequenceSamplers;
+	TArray<FBlendSpaceSampler> BlendSpaceSamplers;
+
+	TArray<FAssetIndexer> Indexers;
+
+	void Prepare(const UPoseSearchSchema* Schema, const FPoseSearchExtrapolationParameters& ExtrapolationParameters, TConstArrayView<FPoseSearchDatabaseSequence> Sequences, TConstArrayView<FPoseSearchDatabaseBlendSpace> BlendSpaces);
+	bool IndexAssets();
+	void JoinIndex();
+	float CalculateMinCostAddend() const;
+};
+
+void FDatabaseIndexingContext::Prepare(const UPoseSearchSchema* Schema, const FPoseSearchExtrapolationParameters& ExtrapolationParameters, TConstArrayView<FPoseSearchDatabaseSequence> Sequences, TConstArrayView<FPoseSearchDatabaseBlendSpace> BlendSpaces)
+{
+	check(Schema);
+
+	FBoneContainer BoneContainer;
+	BoneContainer.InitializeTo(Schema->BoneIndicesWithParents, FCurveEvaluationOption(false), *Schema->Skeleton);
+
+	TMap<const UAnimSequence*, int32> SequenceSamplerMap;
+	TMap<TPair<const UBlendSpace*, FVector>, int32> BlendSpaceSamplerMap;
+
+	SamplingContext.Init(Schema->MirrorDataTable, BoneContainer);
+
+	// Prepare samplers for all sequences
+	auto AddSequenceSampler = [&](const UAnimSequence* Sequence)
+	{
+		if (Sequence && !SequenceSamplerMap.Contains(Sequence))
+		{
+			int32 SequenceSamplerIdx = SequenceSamplers.AddDefaulted();
+			SequenceSamplerMap.Add(Sequence, SequenceSamplerIdx);
+
+			FSequenceSampler::FInput Input;
+			Input.ExtrapolationParameters = ExtrapolationParameters;
+			Input.Sequence = Sequence;
+			SequenceSamplers[SequenceSamplerIdx].Init(Input);
+		}
+	};
+
+	for (const FPoseSearchDatabaseSequence& DbSequence : Sequences)
+	{
+		AddSequenceSampler(DbSequence.Sequence);
+		AddSequenceSampler(DbSequence.LeadInSequence);
+		AddSequenceSampler(DbSequence.FollowUpSequence);
+	}
+
+	// Prepare samplers for all blend spaces
+	for (const FPoseSearchDatabaseBlendSpace& DbBlendSpace : BlendSpaces)
+	{
+		if (DbBlendSpace.BlendSpace)
+		{
+			int32 HorizontalBlendNum, VerticalBlendNum;
+			DbBlendSpace.GetBlendSpaceParameterSampleRanges(HorizontalBlendNum, VerticalBlendNum);
+
+			for (int32 HorizontalIndex = 0; HorizontalIndex < HorizontalBlendNum; HorizontalIndex++)
 			{
-				Tasks.RemoveAtSwap(TaskIndex, 1, false);
+				for (int32 VerticalIndex = 0; VerticalIndex < VerticalBlendNum; VerticalIndex++)
+				{
+					const FVector BlendParameters = DbBlendSpace.BlendParameterForSampleRanges(HorizontalIndex, VerticalIndex);
+
+					if (!BlendSpaceSamplerMap.Contains({ DbBlendSpace.BlendSpace, BlendParameters }))
+					{
+						int32 BlendSpaceSamplerIdx = BlendSpaceSamplers.AddDefaulted();
+						BlendSpaceSamplerMap.Add({ DbBlendSpace.BlendSpace, BlendParameters }, BlendSpaceSamplerIdx);
+
+						FBlendSpaceSampler::FInput Input;
+						Input.BoneContainer = BoneContainer;
+						Input.ExtrapolationParameters = ExtrapolationParameters;
+						Input.BlendSpace = DbBlendSpace.BlendSpace;
+						Input.BlendParameters = BlendParameters;
+
+						BlendSpaceSamplers[BlendSpaceSamplerIdx].Init(Input);
+					}
+				}
 			}
 		}
 	}
 
-	void FAsyncPoseSearchDatabasesManagement::Shutdown()
+	TArray<IAssetSampler*, TInlineAllocator<512>> AssetSampler;
+	AssetSampler.SetNumUninitialized(SequenceSamplers.Num() + BlendSpaceSamplers.Num());
+
+	for (int i = 0; i < SequenceSamplers.Num(); ++i)
+	{
+		AssetSampler[i] = &SequenceSamplers[i];
+	}
+	for (int i = 0; i < BlendSpaceSamplers.Num(); ++i)
+	{
+		AssetSampler[i + SequenceSamplers.Num()] = &BlendSpaceSamplers[i];
+	}
+
+	ParallelFor(AssetSampler.Num(), [AssetSampler](int32 SamplerIdx) { AssetSampler[SamplerIdx]->Process(); }, ParallelForFlags);
+
+	// prepare indexers
+	Indexers.Reserve(SearchIndexBase->Assets.Num());
+
+	auto GetSequenceSampler = [&](const UAnimSequence* Sequence) -> const FSequenceSampler*
+	{
+		return Sequence ? &SequenceSamplers[SequenceSamplerMap[Sequence]] : nullptr;
+	};
+
+	auto GetBlendSpaceSampler = [&](const UBlendSpace* BlendSpace, const FVector BlendParameters) -> const FBlendSpaceSampler*
+	{
+		return BlendSpace ? &BlendSpaceSamplers[BlendSpaceSamplerMap[{BlendSpace, BlendParameters}]] : nullptr;
+	};
+
+	Indexers.Reserve(SearchIndexBase->Assets.Num());
+
+	for (int32 AssetIdx = 0; AssetIdx != SearchIndexBase->Assets.Num(); ++AssetIdx)
 	{
-		FScopeLock Lock(&Mutex);
+		const FPoseSearchIndexAsset& SearchIndexAsset = SearchIndexBase->Assets[AssetIdx];
+
+		FAssetIndexingContext IndexerContext;
+		IndexerContext.SamplingContext = &SamplingContext;
+		IndexerContext.Schema = Schema;
+		IndexerContext.RequestedSamplingRange = SearchIndexAsset.SamplingInterval;
+		IndexerContext.bMirrored = SearchIndexAsset.bMirrored;
 
-		FCoreUObjectDelegates::OnObjectModified.Remove(OnObjectModifiedHandle);
-		OnObjectModifiedHandle.Reset();
+		if (SearchIndexAsset.Type == ESearchIndexAssetType::Sequence)
+		{
+			const FPoseSearchDatabaseSequence& DbSequence = Sequences[SearchIndexAsset.SourceAssetIdx];
+			const float SequenceLength = DbSequence.Sequence->GetPlayLength();
+			IndexerContext.MainSampler = GetSequenceSampler(DbSequence.Sequence);
+			IndexerContext.LeadInSampler = SearchIndexAsset.SamplingInterval.Min == 0.0f ? GetSequenceSampler(DbSequence.LeadInSequence) : nullptr;
+			IndexerContext.FollowUpSampler = SearchIndexAsset.SamplingInterval.Max == SequenceLength ? GetSequenceSampler(DbSequence.FollowUpSequence) : nullptr;
+		}
+		else if (SearchIndexAsset.Type == ESearchIndexAssetType::BlendSpace)
+		{
+			const FPoseSearchDatabaseBlendSpace& DbBlendSpace = BlendSpaces[SearchIndexAsset.SourceAssetIdx];
+			IndexerContext.MainSampler = GetBlendSpaceSampler(DbBlendSpace.BlendSpace, SearchIndexAsset.BlendParameters);
+		}
+		else
+		{
+			checkNoEntry();
+		}
+
+		FAssetIndexer& Indexer = Indexers.AddDefaulted_GetRef();
+		Indexer.Init(IndexerContext, BoneContainer);
+	}
+}
+
+bool FDatabaseIndexingContext::IndexAssets()
+{
+	// Index asset data
+	ParallelFor(Indexers.Num(), [this](int32 AssetIdx) { Indexers[AssetIdx].Process(); }, ParallelForFlags);
+	return true;
+}
+
+float FDatabaseIndexingContext::CalculateMinCostAddend() const
+{
+	float MinCostAddend = 0.f;
+
+	check(SearchIndexBase);
+	if (!SearchIndexBase->PoseMetadata.IsEmpty())
+	{
+		MinCostAddend = MAX_FLT;
+		for (const FPoseSearchPoseMetadata& PoseMetadata : SearchIndexBase->PoseMetadata)
+		{
+			if (PoseMetadata.CostAddend < MinCostAddend)
+			{
+				MinCostAddend = PoseMetadata.CostAddend;
+			}
+		}
 	}
+	return MinCostAddend;
+}
+
+void FDatabaseIndexingContext::JoinIndex()
+{
+	// Write index info to asset and count up total poses and storage required
+	int32 TotalPoses = 0;
+	int32 TotalFloats = 0;
+
+	check(SearchIndexBase);
 
-	void FAsyncPoseSearchDatabasesManagement::Tick(float DeltaTime)
+	// Join animation data into a single search index
+	SearchIndexBase->Values.Reset();
+	SearchIndexBase->PoseMetadata.Reset();
+	SearchIndexBase->OverallFlags = EPoseSearchPoseFlags::None;
+
+	for (int32 AssetIdx = 0; AssetIdx != SearchIndexBase->Assets.Num(); ++AssetIdx)
 	{
-		FScopeLock Lock(&Mutex);
+		const FAssetIndexer::FOutput& Output = Indexers[AssetIdx].Output;
+
+		FPoseSearchIndexAsset& SearchIndexAsset = SearchIndexBase->Assets[AssetIdx];
+		SearchIndexAsset.NumPoses = Output.NumIndexedPoses;
+		SearchIndexAsset.FirstPoseIdx = TotalPoses;
 
-		check(IsInGameThread());
+		const int32 PoseMetadataStartIdx = SearchIndexBase->PoseMetadata.Num();
+		const int32 PoseMetadataEndIdx = PoseMetadataStartIdx + Output.PoseMetadata.Num();
 
-		// iterating backwards because of the possible RemoveAtSwap 
-		for (int32 TaskIndex = Tasks.Num() - 1; TaskIndex >= 0; --TaskIndex)
+		SearchIndexBase->Values.Append(Output.FeatureVectorTable.GetData(), Output.FeatureVectorTable.Num());
+		SearchIndexBase->PoseMetadata.Append(Output.PoseMetadata);
+
+		for (int32 i = PoseMetadataStartIdx; i < PoseMetadataEndIdx; ++i)
 		{
-			Tasks[TaskIndex]->Update(Mutex);
-			
-			// @todo: check key validity every few ticks, or perhaps delete unused for a long time Tasks
+			SearchIndexBase->PoseMetadata[i].AssetIndex = AssetIdx;
+			SearchIndexBase->OverallFlags |= SearchIndexBase->PoseMetadata[i].Flags;
 		}
+
+		TotalPoses += Output.NumIndexedPoses;
+		TotalFloats += Output.FeatureVectorTable.Num();
 	}
 
-	void FAsyncPoseSearchDatabasesManagement::TickCook(float DeltaTime, bool bCookCompete)
+	SearchIndexBase->NumPoses = TotalPoses;
+	SearchIndexBase->MinCostAddend = CalculateMinCostAddend();
+}
+
+//////////////////////////////////////////////////////////////////////////
+// FPoseSearchDatabaseAsyncCacheTask
+struct FPoseSearchDatabaseAsyncCacheTask
+{
+	enum class EState
 	{
-		FScopeLock Lock(&Mutex);
+		Prestarted,
+		Cancelled,
+		Ended,
+		Failed
+	};
 
-		Tick(DeltaTime);
+	// these methods MUST be protected by FPoseSearchDatabaseAsyncCacheTask::Mutex! and to make sure we pass the mutex as input param
+	FPoseSearchDatabaseAsyncCacheTask(UPoseSearchDatabase* InDatabase, FCriticalSection& OuterMutex);
+	void StartNewRequestIfNeeded(FCriticalSection& OuterMutex);
+	bool CancelIfDependsOn(const UObject* Object, FCriticalSection& OuterMutex);
+	void Update(FCriticalSection& OuterMutex);
+	void Wait(FCriticalSection& OuterMutex);
+	void Cancel(FCriticalSection& OuterMutex);
+	bool Poll(FCriticalSection& OuterMutex) const;
+	bool ContainsDatabase(const UPoseSearchDatabase* OtherDatabase, FCriticalSection& OuterMutex) const;
+
+	~FPoseSearchDatabaseAsyncCacheTask();
+	EState GetState() const { return EState(ThreadSafeState.GetValue()); }
+
+private:
+	FPoseSearchDatabaseAsyncCacheTask(const FPoseSearchDatabaseAsyncCacheTask& Other) = delete;
+	FPoseSearchDatabaseAsyncCacheTask(FPoseSearchDatabaseAsyncCacheTask&& Other) = delete;
+	FPoseSearchDatabaseAsyncCacheTask& operator=(const FPoseSearchDatabaseAsyncCacheTask& Other) = delete;
+	FPoseSearchDatabaseAsyncCacheTask& operator=(FPoseSearchDatabaseAsyncCacheTask&& Other) = delete;
+
+	void OnGetComplete(UE::DerivedData::FCacheGetResponse&& Response);
+	void SetState(EState State) { ThreadSafeState.Set(int32(State)); }
+
+	TWeakObjectPtr<UPoseSearchDatabase> Database;
+	// @todo: this is not relevant when the async task is completed, so to save memory we should move it as pointer perhaps
+	FPoseSearchIndex SearchIndex;
+	UE::DerivedData::FRequestOwner Owner;
+	FIoHash DerivedDataKey = FIoHash::Zero;
+	TSet<TWeakObjectPtr<const UObject>> DatabaseDependencies; // @todo: make this const
+		
+	FThreadSafeCounter ThreadSafeState = int32(EState::Prestarted);
+	bool bBroadcastOnDerivedDataRebuild = false;
+};
+
+class FPoseSearchDatabaseAsyncCacheTasks : public TArray<TUniquePtr<FPoseSearchDatabaseAsyncCacheTask>> {};
+
+FPoseSearchDatabaseAsyncCacheTask::FPoseSearchDatabaseAsyncCacheTask(UPoseSearchDatabase* InDatabase, FCriticalSection& OuterMutex)
+	: Database(InDatabase)
+	, Owner(UE::DerivedData::EPriority::Normal)
+	, DerivedDataKey(FIoHash::Zero)
+{
+	StartNewRequestIfNeeded(OuterMutex);
+}
+
+FPoseSearchDatabaseAsyncCacheTask::~FPoseSearchDatabaseAsyncCacheTask()
+{
+	Database = nullptr;
+	SearchIndex.Reset();
+	Owner.Cancel();
+	DerivedDataKey = FIoHash::Zero;
+	DatabaseDependencies.Reset();
+}
+
+void FPoseSearchDatabaseAsyncCacheTask::StartNewRequestIfNeeded(FCriticalSection& OuterMutex)
+{
+	using namespace UE::DerivedData;
+
+	FScopeLock Lock(&OuterMutex);
+
+	// making sure there are no active requests
+	Owner.Cancel();
+
+	// composing the key
+	const FDerivedDataKeyBuilder KeyBuilder(Database.Get());
+
+	// Stores a BLAKE3-160 hash, taken from the first 20 bytes of a BLAKE3-256 hash
+	const FIoHash NewDerivedDataKey = KeyBuilder.Finalize();
+	const bool bHasKeyChanged = NewDerivedDataKey != DerivedDataKey;
+	if (bHasKeyChanged)
+	{
+		DerivedDataKey = NewDerivedDataKey;
+
+		DatabaseDependencies.Reset();
+		for (const UObject* Dependency : KeyBuilder.GetDependencies())
+		{
+			DatabaseDependencies.Add(Dependency);
+		}
+
+		SetState(EState::Prestarted);
+
+		UE_LOG(LogPoseSearch, Log, TEXT("%s - %s BeginCache"), *LexToString(DerivedDataKey), *Database->GetName());
+
+		TArray<FCacheGetRequest> CacheRequests;
+		const FCacheKey CacheKey{ Bucket, DerivedDataKey };
+		CacheRequests.Add({ { Database->GetPathName() }, CacheKey, ECachePolicy::Default });
+
+		Owner = FRequestOwner(EPriority::Normal);
+		GetCache().Get(CacheRequests, Owner, [this](FCacheGetResponse&& Response)
+			{
+				OnGetComplete(MoveTemp(Response));
+			});
 	}
+}
+
+// it cancels and waits for the task to be done and reset the local SearchIndex. SetState to Cancelled
+void FPoseSearchDatabaseAsyncCacheTask::Cancel(FCriticalSection& OuterMutex)
+{
+	FScopeLock Lock(&OuterMutex);
+
+	Owner.Cancel();
+	SearchIndex.Reset();
+	DerivedDataKey = FIoHash::Zero;
+	SetState(EState::Cancelled);
+}
 
-	TStatId FAsyncPoseSearchDatabasesManagement::GetStatId() const
+bool FPoseSearchDatabaseAsyncCacheTask::CancelIfDependsOn(const UObject* Object, FCriticalSection& OuterMutex)
+{
+	FScopeLock Lock(&OuterMutex);
+
+	// DatabaseDependencies is updated only in StartNewRequestIfNeeded when there are no active requests, so it's thread safe to access it 
+	if (DatabaseDependencies.Contains(Object))
 	{
-		RETURN_QUICK_DECLARE_CYCLE_STAT(FAsyncPoseSearchDatabasesManagement, STATGROUP_Tickables);
+		Cancel(OuterMutex);
+		return true;
 	}
+	return false;
+}
+
+void FPoseSearchDatabaseAsyncCacheTask::Update(FCriticalSection& OuterMutex)
+{
+	check(IsInGameThread());
 
-	void FAsyncPoseSearchDatabasesManagement::AddReferencedObjects(FReferenceCollector& Collector)
+	FScopeLock Lock(&OuterMutex);
+
+	check(GetState() != EState::Cancelled); // otherwise FPoseSearchDatabaseAsyncCacheTask should have been already removed
+
+	if (GetState() == EState::Prestarted && Poll(OuterMutex))
 	{
+		// task is done: we need to update the state form Prestarted to Ended/Failed
+		Wait(OuterMutex);
 	}
 
-	// returns true if the index has been built and the Database updated correctly  
-	bool FAsyncPoseSearchDatabasesManagement::RequestAsyncBuildIndex(const UPoseSearchDatabase* Database, ERequestAsyncBuildFlag Flag)
+	if (bBroadcastOnDerivedDataRebuild)
 	{
-		if (!Database)
-		{
-			return false;
-		}
+		Database->NotifyDerivedDataRebuild();
+		bBroadcastOnDerivedDataRebuild = false;
+	}
+}
 
-		FScopeLock Lock(&Mutex);
+// it waits for the task to be done and SetSearchIndex on the database. SetState to Ended/Failed
+void FPoseSearchDatabaseAsyncCacheTask::Wait(FCriticalSection& OuterMutex)
+{
+	check(GetState() == EState::Prestarted);
 
-		check(Database);
-		check(EnumHasAnyFlags(Flag, ERequestAsyncBuildFlag::NewRequest | ERequestAsyncBuildFlag::ContinueRequest));
+	Owner.Wait();
 
-		FAsyncPoseSearchDatabasesManagement& This = FAsyncPoseSearchDatabasesManagement::Get();
+	FScopeLock Lock(&OuterMutex);
 
-		FPoseSearchDatabaseAsyncCacheTask* Task = nullptr;
-		for (TUniquePtr<FPoseSearchDatabaseAsyncCacheTask>& TaskPtr : This.Tasks)
-		{
-			if (TaskPtr->ContainsDatabase(Database, Mutex))
+	const bool bFailedIndexing = SearchIndex.IsEmpty();
+	if (!bFailedIndexing)
+	{
+		Database->SetSearchIndex(SearchIndex); // @todo: implement FPoseSearchIndex move ctor and assignment operator and use a MoveTemp(SearchIndex) here
+		SetState(EState::Ended);
+		bBroadcastOnDerivedDataRebuild = true;
+	}
+	else
+	{
+		check(!bBroadcastOnDerivedDataRebuild);
+		SetState(EState::Failed);
+	}
+	SearchIndex.Reset();
+}
+
+// true is the task is done executing
+bool FPoseSearchDatabaseAsyncCacheTask::Poll(FCriticalSection& OuterMutex) const
+{
+	return Owner.Poll();
+}
+
+bool FPoseSearchDatabaseAsyncCacheTask::ContainsDatabase(const UPoseSearchDatabase* OtherDatabase, FCriticalSection& OuterMutex) const
+{
+	FScopeLock Lock(&OuterMutex);
+	return Database.Get() == OtherDatabase;
+}
+
+// called once the task is done:
+// if EStatus::Ok (data has been retrieved from DDC) we deserialize the payload into the local SearchIndex
+// if EStatus::Error we BuildIndex and if that's successful we 'Put' it on DDC
+void FPoseSearchDatabaseAsyncCacheTask::OnGetComplete(UE::DerivedData::FCacheGetResponse&& Response)
+{
+	using namespace UE::DerivedData;
+
+	const FCacheKey FullIndexKey = Response.Record.GetKey();
+	if (Response.Status == EStatus::Ok)
+	{
+		COOK_STAT(auto Timer = UsageStats.TimeAsyncWait());
+
+		// we found the cached data associated to the PendingDerivedDataKey: we'll deserialized into SearchIndex
+		SearchIndex.Reset();
+		FSharedBuffer RawData = Response.Record.GetValue(Id).GetData().Decompress();
+		FMemoryReaderView Reader(RawData);
+		Reader << SearchIndex;
+
+		UE_LOG(LogPoseSearch, Log, TEXT("%s - %s BuildIndex From Cache"), *LexToString(FullIndexKey.Hash), *Database->GetName());
+
+		COOK_STAT(Timer.AddHit(RawData.GetSize()));
+	}
+	else if (Response.Status == EStatus::Canceled)
+	{
+		SearchIndex.Reset();
+		UE_LOG(LogPoseSearch, Log, TEXT("%s - %s BuildIndex Cancelled"), *LexToString(FullIndexKey.Hash), *Database->GetName());
+	}
+	else if (Response.Status == EStatus::Error)
+	{
+		// we didn't find the cached data associated to the PendingDerivedDataKey: we'll BuildIndex to update SearchIndex and "Put" the data over the DDC
+		Owner.LaunchTask(TEXT("PoseSearchDatabaseBuild"), [this, FullIndexKey]
 			{
-				Task = TaskPtr.Get();
+				COOK_STAT(auto Timer = UsageStats.TimeSyncWork());
+
+				// early out for invalid indexing conditions
+				if (!Database->Schema || !Database->Schema->IsValid() || Database->Schema->SchemaCardinality <= 0)
+				{
+					UE_LOG(LogPoseSearch, Log, TEXT("%s - %s BuildIndex Failed"), *LexToString(FullIndexKey.Hash), *Database->GetName());
+					SearchIndex.Reset();
+					return;
+				}
+
+				if (Owner.IsCanceled())
+				{
+					UE_LOG(LogPoseSearch, Log, TEXT("%s - %s BuildIndex Cancelled"), *LexToString(FullIndexKey.Hash), *Database->GetName());
+					SearchIndex.Reset();
+					return;
+				}
+
+				// Building all the related FPoseSearchBaseIndex first
+				FPoseSearchIndexBase& SearchIndexBase = SearchIndex;
+				InitSearchIndexAssets(SearchIndexBase, Database->Sequences, Database->BlendSpaces, Database->ExcludeFromDatabaseParameters);
+
+				if (Owner.IsCanceled())
+				{
+					UE_LOG(LogPoseSearch, Log, TEXT("%s - %s BuildIndex Cancelled"), *LexToString(FullIndexKey.Hash), *Database->GetName());
+					SearchIndex.Reset();
+					return;
+				}
 
-				if (EnumHasAnyFlags(Flag, ERequestAsyncBuildFlag::NewRequest))
+				FDatabaseIndexingContext DbIndexingContext;
+				DbIndexingContext.SearchIndexBase = &SearchIndexBase;
+				DbIndexingContext.Prepare(Database->Schema, Database->ExtrapolationParameters, Database->Sequences, Database->BlendSpaces);
+
+				if (Owner.IsCanceled())
 				{
-					if (Task->GetState() == FPoseSearchDatabaseAsyncCacheTask::EState::Prestarted)
+					UE_LOG(LogPoseSearch, Log, TEXT("%s - %s BuildIndex Cancelled"), *LexToString(FullIndexKey.Hash), *Database->GetName());
+					SearchIndex.Reset();
+					return;
+				}
+
+				const bool bSuccess = DbIndexingContext.IndexAssets();
+				if (!bSuccess)
+				{
+					UE_LOG(LogPoseSearch, Log, TEXT("%s - %s BuildIndex Failed"), *LexToString(FullIndexKey.Hash), *Database->GetName());
+					SearchIndex.Reset();
+					return;
+				}
+
+				if (Owner.IsCanceled())
+				{
+					UE_LOG(LogPoseSearch, Log, TEXT("%s - %s BuildIndex Cancelled"), *LexToString(FullIndexKey.Hash), *Database->GetName());
+					SearchIndex.Reset();
+					return;
+				}
+
+				DbIndexingContext.JoinIndex();
+				if (Owner.IsCanceled())
+				{
+					UE_LOG(LogPoseSearch, Log, TEXT("%s - %s BuildIndex Cancelled"), *LexToString(FullIndexKey.Hash), *Database->GetName());
+					SearchIndex.Reset();
+					return;
+				}
+
+				// Building FPoseSearchIndex
+				PreprocessSearchIndexWeights(SearchIndex, Database->Schema);
+				if (Owner.IsCanceled())
+				{
+					UE_LOG(LogPoseSearch, Log, TEXT("%s - %s BuildIndex Cancelled"), *LexToString(FullIndexKey.Hash), *Database->GetName());
+					SearchIndex.Reset();
+					return;
+				}
+
+				PreprocessSearchIndexPCAData(SearchIndex, Database->Schema->SchemaCardinality, Database->GetNumberOfPrincipalComponents(), Database->PoseSearchMode);
+				if (Owner.IsCanceled())
+				{
+					UE_LOG(LogPoseSearch, Log, TEXT("%s - %s BuildIndex Cancelled"), *LexToString(FullIndexKey.Hash), *Database->GetName());
+					SearchIndex.Reset();
+					return;
+				}
+
+				PreprocessSearchIndexKDTree(SearchIndex, Database->Schema->SchemaCardinality, Database->GetNumberOfPrincipalComponents(), Database->PoseSearchMode, Database->KDTreeMaxLeafSize, Database->KDTreeQueryNumNeighbors);
+				if (Owner.IsCanceled())
+				{
+					UE_LOG(LogPoseSearch, Log, TEXT("%s - %s BuildIndex Cancelled"), *LexToString(FullIndexKey.Hash), *Database->GetName());
+					SearchIndex.Reset();
+					return;
+				}
+
+				UE_LOG(LogPoseSearch, Log, TEXT("%s - %s BuildIndex Succeeded"), *LexToString(FullIndexKey.Hash), *Database->GetName());
+
+				// putting SearchIndex to DDC
+				TArray<uint8> RawBytes;
+				FMemoryWriter Writer(RawBytes);
+				Writer << SearchIndex;
+				FSharedBuffer RawData = MakeSharedBufferFromArray(MoveTemp(RawBytes));
+				const int32 BytesProcessed = RawData.GetSize();
+
+				FCacheRecordBuilder Builder(FullIndexKey);
+				Builder.AddValue(Id, RawData);
+				GetCache().Put({ { { Database->GetPathName() }, Builder.Build() } }, Owner, [this, FullIndexKey](FCachePutResponse&& Response)
 					{
-						if (EnumHasAnyFlags(Flag, ERequestAsyncBuildFlag::WaitPreviousRequest))
-						{
-							Task->Wait(Mutex);
-						}
-						else
+						if (Response.Status == EStatus::Error)
 						{
-							Task->Cancel(Mutex);
+							UE_LOG(LogPoseSearch, Log, TEXT("%s - %s Failed to store DDC"), *LexToString(FullIndexKey.Hash), *Database->GetName());
 						}
-					}
+					});
+
+				COOK_STAT(Timer.AddMiss(BytesProcessed));
+			});
+	}
+}
+
+//////////////////////////////////////////////////////////////////////////
+// FAsyncPoseSearchDatabasesManagement
+FCriticalSection FAsyncPoseSearchDatabasesManagement::Mutex;
+
+FAsyncPoseSearchDatabasesManagement& FAsyncPoseSearchDatabasesManagement::Get()
+{
+	FScopeLock Lock(&Mutex);
+
+	static FAsyncPoseSearchDatabasesManagement SingletonInstance;
+	return SingletonInstance;
+}
+
+FAsyncPoseSearchDatabasesManagement::FAsyncPoseSearchDatabasesManagement()
+	: Tasks(*(new FPoseSearchDatabaseAsyncCacheTasks()))
+{
+	FScopeLock Lock(&Mutex);
+
+	OnObjectModifiedHandle = FCoreUObjectDelegates::OnObjectModified.AddRaw(this, &FAsyncPoseSearchDatabasesManagement::OnObjectModified);
+	FCoreDelegates::OnPreExit.AddRaw(this, &FAsyncPoseSearchDatabasesManagement::Shutdown);
+}
+
+FAsyncPoseSearchDatabasesManagement::~FAsyncPoseSearchDatabasesManagement()
+{
+	FScopeLock Lock(&Mutex);
+
+	FCoreDelegates::OnPreExit.RemoveAll(this);
+	Shutdown();
+
+	delete &Tasks;
+}
+
+// we're listening to OnObjectModified to cancel any pending Task indexing databases depending from Object to avoid multi threading issues
+void FAsyncPoseSearchDatabasesManagement::OnObjectModified(UObject* Object)
+{
+	FScopeLock Lock(&Mutex);
+
+	// iterating backwards because of the possible RemoveAtSwap
+	for (int32 TaskIndex = Tasks.Num() - 1; TaskIndex >= 0; --TaskIndex)
+	{
+		if (Tasks[TaskIndex]->CancelIfDependsOn(Object, Mutex))
+		{
+			Tasks.RemoveAtSwap(TaskIndex, 1, false);
+		}
+	}
+}
+
+void FAsyncPoseSearchDatabasesManagement::Shutdown()
+{
+	FScopeLock Lock(&Mutex);
+
+	FCoreUObjectDelegates::OnObjectModified.Remove(OnObjectModifiedHandle);
+	OnObjectModifiedHandle.Reset();
+}
 
-					Task->StartNewRequestIfNeeded(Mutex);
+void FAsyncPoseSearchDatabasesManagement::Tick(float DeltaTime)
+{
+	FScopeLock Lock(&Mutex);
+
+	check(IsInGameThread());
+
+	// iterating backwards because of the possible RemoveAtSwap 
+	for (int32 TaskIndex = Tasks.Num() - 1; TaskIndex >= 0; --TaskIndex)
+	{
+		Tasks[TaskIndex]->Update(Mutex);
+			
+		// @todo: check key validity every few ticks, or perhaps delete unused for a long time Tasks
+	}
+}
+
+void FAsyncPoseSearchDatabasesManagement::TickCook(float DeltaTime, bool bCookCompete)
+{
+	FScopeLock Lock(&Mutex);
+
+	Tick(DeltaTime);
+}
+
+TStatId FAsyncPoseSearchDatabasesManagement::GetStatId() const
+{
+	RETURN_QUICK_DECLARE_CYCLE_STAT(FAsyncPoseSearchDatabasesManagement, STATGROUP_Tickables);
+}
+
+void FAsyncPoseSearchDatabasesManagement::AddReferencedObjects(FReferenceCollector& Collector)
+{
+}
+
+// returns true if the index has been built and the Database updated correctly  
+bool FAsyncPoseSearchDatabasesManagement::RequestAsyncBuildIndex(const UPoseSearchDatabase* Database, ERequestAsyncBuildFlag Flag)
+{
+	if (!Database)
+	{
+		return false;
+	}
+
+	FScopeLock Lock(&Mutex);
+
+	check(Database);
+	check(EnumHasAnyFlags(Flag, ERequestAsyncBuildFlag::NewRequest | ERequestAsyncBuildFlag::ContinueRequest));
+
+	FAsyncPoseSearchDatabasesManagement& This = FAsyncPoseSearchDatabasesManagement::Get();
+
+	FPoseSearchDatabaseAsyncCacheTask* Task = nullptr;
+	for (TUniquePtr<FPoseSearchDatabaseAsyncCacheTask>& TaskPtr : This.Tasks)
+	{
+		if (TaskPtr->ContainsDatabase(Database, Mutex))
+		{
+			Task = TaskPtr.Get();
+
+			if (EnumHasAnyFlags(Flag, ERequestAsyncBuildFlag::NewRequest))
+			{
+				if (Task->GetState() == FPoseSearchDatabaseAsyncCacheTask::EState::Prestarted)
+				{
+					if (EnumHasAnyFlags(Flag, ERequestAsyncBuildFlag::WaitPreviousRequest))
+					{
+						Task->Wait(Mutex);
+					}
+					else
+					{
+						Task->Cancel(Mutex);
+					}
 				}
-				else // if (EnumHasAnyFlags(Flag, ERequestAsyncBuildFlag::ContinueRequest))
+
+				Task->StartNewRequestIfNeeded(Mutex);
+			}
+			else // if (EnumHasAnyFlags(Flag, ERequestAsyncBuildFlag::ContinueRequest))
+			{
+				if (Task->GetState() == FPoseSearchDatabaseAsyncCacheTask::EState::Prestarted)
 				{
-					if (Task->GetState() == FPoseSearchDatabaseAsyncCacheTask::EState::Prestarted)
+					if (EnumHasAnyFlags(Flag, ERequestAsyncBuildFlag::WaitPreviousRequest))
 					{
-						if (EnumHasAnyFlags(Flag, ERequestAsyncBuildFlag::WaitPreviousRequest))
-						{
-							Task->Wait(Mutex);
-						}
+						Task->Wait(Mutex);
 					}
 				}
-				break;
 			}
+			break;
 		}
+	}
 		
-		if (!Task)
-		{
-			// we didn't find the Task, so we Emplace a new one
-			This.Tasks.Emplace(MakeUnique<FPoseSearchDatabaseAsyncCacheTask>(const_cast<UPoseSearchDatabase*>(Database), Mutex));
-			Task = This.Tasks.Last().Get();
-		}
-
-		if (EnumHasAnyFlags(Flag, ERequestAsyncBuildFlag::WaitForCompletion) && Task->GetState() == FPoseSearchDatabaseAsyncCacheTask::EState::Prestarted)
-		{
-			Task->Wait(Mutex);
-		}
+	if (!Task)
+	{
+		// we didn't find the Task, so we Emplace a new one
+		This.Tasks.Emplace(MakeUnique<FPoseSearchDatabaseAsyncCacheTask>(const_cast<UPoseSearchDatabase*>(Database), Mutex));
+		Task = This.Tasks.Last().Get();
+	}
 
-		return Task->GetState() == FPoseSearchDatabaseAsyncCacheTask::EState::Ended;
+	if (EnumHasAnyFlags(Flag, ERequestAsyncBuildFlag::WaitForCompletion) && Task->GetState() == FPoseSearchDatabaseAsyncCacheTask::EState::Prestarted)
+	{
+		Task->Wait(Mutex);
 	}
 
+	return Task->GetState() == FPoseSearchDatabaseAsyncCacheTask::EState::Ended;
+}
+
 } // namespace UE::PoseSearch
 #endif // WITH_EDITOR
diff --git a/Engine/Plugins/Experimental/Animation/PoseSearch/Source/Runtime/Public/PoseSearch/PoseSearch.h b/Engine/Plugins/Experimental/Animation/PoseSearch/Source/Runtime/Public/PoseSearch/PoseSearch.h
--- a/Engine/Plugins/Experimental/Animation/PoseSearch/Source/Runtime/Public/PoseSearch/PoseSearch.h
+++ b/Engine/Plugins/Experimental/Animation/PoseSearch/Source/Runtime/Public/PoseSearch/PoseSearch.h
@@ -724,10 +724,10 @@
 };
 
 /**
-* A search index for animation poses for a UPoseSearchDatabase. The structure of the search index is determined by its UPoseSearchSchema.
+* case class for FPoseSearchIndex. building block used to gather data for data mining and calculate weights, pca, kdtree stuff
 */
 USTRUCT()
-struct POSESEARCH_API FPoseSearchIndex
+struct POSESEARCH_API FPoseSearchIndexBase
 {
 	GENERATED_BODY()
 
@@ -736,19 +736,44 @@
 
 	UPROPERTY(meta = (NeverInHash))
 	TArray<float> Values;
+	
+	UPROPERTY(meta = (NeverInHash))
+	TArray<FPoseSearchPoseMetadata> PoseMetadata;
 
 	UPROPERTY(meta = (NeverInHash))
+	EPoseSearchPoseFlags OverallFlags = EPoseSearchPoseFlags::None;
+
+	UPROPERTY(meta = (NeverInHash))
+	TArray<FPoseSearchIndexAsset> Assets;
+
+	// minimum of the database metadata CostAddend: it represents the minimum cost of any search for the associated database (we'll skip the search in case the search result total cost is already less than MinCostAddend)
+	UPROPERTY(Category = Info, VisibleAnywhere, meta = (NeverInHash))
+	float MinCostAddend = -MAX_FLT;
+
+	bool IsValidPoseIndex(int32 PoseIdx) const { return PoseIdx < NumPoses; }
+	bool IsEmpty() const;
+
+	const FPoseSearchIndexAsset& GetAssetForPose(int32 PoseIdx) const;
+	const FPoseSearchIndexAsset* GetAssetForPoseSafe(int32 PoseIdx) const;
+
+	float GetAssetTime(int32 PoseIdx, float SamplingInterval) const;
+
+	void Reset();
+	
+	friend FArchive& operator<<(FArchive& Ar, FPoseSearchIndexBase& Index);
+};
+
+/**
+* A search index for animation poses. The structure of the search index is determined by its UPoseSearchSchema.
+* May represent a single animation (see UPoseSearchSequenceMetaData) or a collection (see UPoseSearchDatabase).
+*/
+USTRUCT()
+struct POSESEARCH_API FPoseSearchIndex : public FPoseSearchIndexBase
+{
+	GENERATED_BODY()
+
+	UPROPERTY(meta = (NeverInHash))
 	TArray<float> PCAValues;
 
-#if WITH_EDITORONLY_DATA
 	UPROPERTY(Category = Info, VisibleAnywhere, meta = (NeverInHash))
-	float PCAExplainedVariance = 0.f;
-	
-	UPROPERTY(Category = Info, VisibleAnywhere, meta = (NeverInHash))
-	TArray<float> Deviation;
-#endif // WITH_EDITORONLY_DATA
-
-	UE::PoseSearch::FKDTree KDTree;
-
-	UPROPERTY(Category = Info, VisibleAnywhere, meta = (NeverInHash))
 	TArray<float> PCAProjectionMatrix;
@@ -764,39 +789,26 @@
 	UPROPERTY(Category = Info, VisibleAnywhere, meta = (NeverInHash))
 	TArray<float> WeightsSqrt;
 
-	UPROPERTY(meta = (NeverInHash))
-	TArray<FPoseSearchPoseMetadata> PoseMetadata;
+	UE::PoseSearch::FKDTree KDTree;
 
-	UPROPERTY(meta = (NeverInHash))
-	EPoseSearchPoseFlags OverallFlags = EPoseSearchPoseFlags::None;
+#if WITH_EDITORONLY_DATA
+	UPROPERTY(Category = Info, VisibleAnywhere, meta = (NeverInHash))
+	float PCAExplainedVariance = 0.f;
 
-	UPROPERTY(meta = (NeverInHash))
-	TArray<FPoseSearchIndexAsset> Assets;
+	UPROPERTY(Category = Info, VisibleAnywhere, meta = (NeverInHash))
+	TArray<float> Deviation;
+#endif // WITH_EDITORONLY_DATA
 
-	// minimum of the database metadata CostAddend: it represents the minimum cost of any search for the associated database (we'll skip the search in case the search result total cost is already less than MinCostAddend)
-	UPROPERTY(Category = Info, VisibleAnywhere, meta = (NeverInHash))
-	float MinCostAddend = -MAX_FLT;
-
 	FPoseSearchIndex() = default;
+	~FPoseSearchIndex() = default;
 	FPoseSearchIndex(const FPoseSearchIndex& Other); // custom copy constructor to deal with the KDTree DataSrc
 	FPoseSearchIndex(FPoseSearchIndex&& Other) = delete;
 	FPoseSearchIndex& operator=(const FPoseSearchIndex& Other); // custom equal operator to deal with the KDTree DataSrc
 	FPoseSearchIndex& operator=(FPoseSearchIndex&& Other) = delete;
 
-	bool IsValidPoseIndex(int32 PoseIdx) const { return PoseIdx < NumPoses; }
-	bool IsEmpty() const;
-
+	void Reset();
 	TConstArrayView<float> GetPoseValues(int32 PoseIdx) const;
 	TConstArrayView<float> GetPoseValuesSafe(int32 PoseIdx) const;
-
-	const FPoseSearchIndexAsset& GetAssetForPose(int32 PoseIdx) const;
-	const FPoseSearchIndexAsset* GetAssetForPoseSafe(int32 PoseIdx) const;
-
-	float GetAssetTime(int32 PoseIdx, float SamplingInterval) const;
-
-	void Reset();
-	void InitSearchIndexAssets(TConstArrayView<FPoseSearchDatabaseSequence> Sequences, TConstArrayView<FPoseSearchDatabaseBlendSpace> BlendSpaces, const FPoseSearchExcludeFromDatabaseParameters& ExcludeFromDatabaseParameters);
-
 	FPoseSearchCost ComparePoses(int32 PoseIdx, EPoseSearchBooleanRequest QueryMirrorRequest, UE::PoseSearch::EPoseComparisonFlags PoseComparisonFlags, float MirrorMismatchCostBias, TConstArrayView<float> QueryValues) const;
 
 	friend FArchive& operator<<(FArchive& Ar, FPoseSearchIndex& Index);
@@ -1406,19 +1418,6 @@
 POSESEARCH_API void DrawSearchIndex(const FDebugDrawParams& DrawParams);
 POSESEARCH_API void CompareFeatureVectors(TConstArrayView<float> A, TConstArrayView<float> B, TConstArrayView<float> WeightsSqrt, TArrayView<float> Result);
 
-/**
-* Creates a pose search index for a collection of animations
-* 
-* @param Database	The input collection of animations and output search index
-* 
-* @return Whether the index was built successfully
-*/
-POSESEARCH_API bool BuildIndex(const UPoseSearchDatabase& Database, FPoseSearchIndex& OutSearchIndex
-#if WITH_EDITOR
-	, UE::DerivedData::IRequestOwner& Owner
-#endif
-);
-
 } // namespace UE::PoseSearch
 
 UENUM()
diff --git a/Engine/Plugins/Experimental/Animation/PoseSearch/Source/Runtime/Public/PoseSearch/PoseSearchDerivedDataKey.h b/Engine/Plugins/Experimental/Animation/PoseSearch/Source/Runtime/Public/PoseSearch/PoseSearchDerivedDataKey.h
--- a/Engine/Plugins/Experimental/Animation/PoseSearch/Source/Runtime/Public/PoseSearch/PoseSearchDerivedDataKey.h
+++ b/Engine/Plugins/Experimental/Animation/PoseSearch/Source/Runtime/Public/PoseSearch/PoseSearchDerivedDataKey.h
@@ -11,9 +11,10 @@
 #include "Containers/UnrealString.h"
 #include "Factories/FbxAnimSequenceImportData.h"
 #include "Hash/Blake3.h"
-#include "Serialization/ArchiveUObject.h"
+#include "UObject/DevObjectVersion.h"
 #include "UObject/Object.h"
 #include "UObject/UnrealType.h"
+#include "Serialization/ArchiveUObject.h"
 
 namespace UE::PoseSearch
 {
@@ -39,8 +40,10 @@ public:
 	inline static const FName ExcludeFromHashName = FName(TEXT("ExcludeFromHash"));
 	inline static const FName NeverInHashName = FName(TEXT("NeverInHash"));
 
-	FDerivedDataKeyBuilder()
+	FDerivedDataKeyBuilder(UObject* Object)
 	{
+		check(Object);
+
 		ArIgnoreOuterRef = true;
 
 		// Set FDerivedDataKeyBuilder to be a saving archive instead of a reference collector.
@@ -48,6 +51,14 @@ public:
 		// which doesn't give a stable hash.  Serializing these to a saving archive will
 		// use a string reference instead, which is a more meaningful hash value.
 		SetIsSaving(true);
+
+		// used to invalidate the key without having to change POSESEARCHDB_DERIVEDDATA_VER all the times
+		FGuid VersionGuid = FDevSystemGuids::GetSystemGuid(FDevSystemGuids::Get().POSESEARCHDB_DERIVEDDATA_VER);
+		int32 POSESEARCHDB_DERIVEDDATA_VER_SMALL = 3;
+
+		*this << VersionGuid;
+		*this << POSESEARCHDB_DERIVEDDATA_VER_SMALL;
+		*this << Object;
 	}
 
 	virtual void Seek(int64 InPos) override
@@ -185,7 +196,7 @@ public:
 	}
 	//~ End FArchive Interface
 	
-	HashDigestType Finalize()
+	HashDigestType Finalize() const
 	{
 		return Hasher.Finalize();
 	}
