Fixed Motion Matching editor crashing in Statistics Overview panel after reopening it a second time.

diff --git a/Engine/Plugins/Experimental/Animation/PoseSearch/Source/Editor/Private/PoseSearchDatabaseEditor.cpp b/Engine/Plugins/Experimental/Animation/PoseSearch/Source/Editor/Private/PoseSearchDatabaseEditor.cpp
--- a/Engine/Plugins/Experimental/Animation/PoseSearch/Source/Editor/Private/PoseSearchDatabaseEditor.cpp
+++ b/Engine/Plugins/Experimental/Animation/PoseSearch/Source/Editor/Private/PoseSearchDatabaseEditor.cpp
@@ -206,22 +206,14 @@ namespace UE::PoseSearch
 			// Ensure statistics information get updated/populated
 			if (DatabaseAsset)
 			{
-				const auto UpdateStatisticsInformationLambda = [this]()
-				{
-					UPoseSearchDatabaseStatistics* Statistics = NewObject<UPoseSearchDatabaseStatistics>();
-					Statistics->AddToRoot();
-					Statistics->Initialize(GetPoseSearchDatabase());
-					StatisticsOverviewWidget->SetObject(Statistics);
-				};
-
 				// Init statistics
 				if (FAsyncPoseSearchDatabasesManagement::RequestAsyncBuildIndex(GetPoseSearchDatabase(), ERequestAsyncBuildFlag::ContinueRequest))
 				{
-					UpdateStatisticsInformationLambda();
+					RefreshStatisticsWidgetInformation();
 				}
 
 				// Ensure any database changes are reflected
-				DatabaseAsset->RegisterOnDerivedDataRebuild(UPoseSearchDatabase::FOnDerivedDataRebuild::CreateLambda(UpdateStatisticsInformationLambda));
+				DatabaseAsset->RegisterOnDerivedDataRebuild(UPoseSearchDatabase::FOnDerivedDataRebuild::CreateSP(this, &FDatabaseEditor::RefreshStatisticsWidgetInformation));
 			}
 		}
 		
@@ -577,6 +569,14 @@ namespace UE::PoseSearch
 
 		ViewModel->SetSelectedNodes(SelectedItems);
 	}
+
+	void FDatabaseEditor::RefreshStatisticsWidgetInformation()
+	{
+		UPoseSearchDatabaseStatistics* Statistics = NewObject<UPoseSearchDatabaseStatistics>();
+		Statistics->AddToRoot();
+		Statistics->Initialize(GetPoseSearchDatabase());
+		StatisticsOverviewWidget->SetObject(Statistics);
+	}
 }
 
 #undef LOCTEXT_NAMESPACE
diff --git a/Engine/Plugins/Experimental/Animation/PoseSearch/Source/Editor/Public/PoseSearchDatabaseEditor.h b/Engine/Plugins/Experimental/Animation/PoseSearch/Source/Editor/Public/PoseSearchDatabaseEditor.h
--- a/Engine/Plugins/Experimental/Animation/PoseSearch/Source/Editor/Public/PoseSearchDatabaseEditor.h
+++ b/Engine/Plugins/Experimental/Animation/PoseSearch/Source/Editor/Public/PoseSearchDatabaseEditor.h
@@ -67,6 +67,8 @@ namespace UE::PoseSearch
 			const TArrayView<TSharedPtr<FDatabaseAssetTreeNode>>& SelectedItems,
 			ESelectInfo::Type SelectionType);
 
+		void RefreshStatisticsWidgetInformation();
+		
 		TSharedPtr<SDatabasePreview> PreviewWidget;
 
 		TSharedPtr<SDatabaseAssetTree> AssetTreeWidget;
