Fix for "Display Root Motion Speed" viewport option not being able to toggled if multiple database items are selected.

diff --git a/Engine/Plugins/Experimental/Animation/PoseSearch/Source/Editor/Private/PoseSearchDatabaseViewModel.cpp b/Engine/Plugins/Experimental/Animation/PoseSearch/Source/Editor/Private/PoseSearchDatabaseViewModel.cpp
--- a/Engine/Plugins/Experimental/Animation/PoseSearch/Source/Editor/Private/PoseSearchDatabaseViewModel.cpp
+++ b/Engine/Plugins/Experimental/Animation/PoseSearch/Source/Editor/Private/PoseSearchDatabaseViewModel.cpp
@@ -410,11 +410,6 @@ namespace UE::PoseSearch
 		DisplayRootMotionSpeed = !DisplayRootMotionSpeed;
 	}
 
-	bool FDatabaseViewModel::CanDisplayRootMotionSpeed() const
-	{
-		return SelectedNodes.Num() <= 1;
-	}
-
 	void FDatabaseViewModel::AddSequenceToDatabase(UAnimSequence* AnimSequence)
 	{
 		FPoseSearchDatabaseSequence& NewDbSequence = PoseSearchDatabase->Sequences.AddDefaulted_GetRef();
diff --git a/Engine/Plugins/Experimental/Animation/PoseSearch/Source/Editor/Private/SPoseSearchDatabaseViewport.cpp b/Engine/Plugins/Experimental/Animation/PoseSearch/Source/Editor/Private/SPoseSearchDatabaseViewport.cpp
--- a/Engine/Plugins/Experimental/Animation/PoseSearch/Source/Editor/Private/SPoseSearchDatabaseViewport.cpp
+++ b/Engine/Plugins/Experimental/Animation/PoseSearch/Source/Editor/Private/SPoseSearchDatabaseViewport.cpp
@@ -112,7 +112,7 @@ namespace UE::PoseSearch
 		CommandList->MapAction(
 			Commands.ShowDisplayRootMotionSpeed,
 			FExecuteAction::CreateSP(ViewModelRef, &FDatabaseViewModel::OnToggleDisplayRootMotionSpeed),
-			FCanExecuteAction::CreateSP(ViewModelRef, &FDatabaseViewModel::CanDisplayRootMotionSpeed),
+			FCanExecuteAction(),
 			FIsActionChecked::CreateSP(ViewModelRef, &FDatabaseViewModel::IsDisplayRootMotionSpeedChecked));
 	}
 
diff --git a/Engine/Plugins/Experimental/Animation/PoseSearch/Source/Editor/Public/PoseSearchDatabaseViewModel.h b/Engine/Plugins/Experimental/Animation/PoseSearch/Source/Editor/Public/PoseSearchDatabaseViewModel.h
--- a/Engine/Plugins/Experimental/Animation/PoseSearch/Source/Editor/Public/PoseSearchDatabaseViewModel.h
+++ b/Engine/Plugins/Experimental/Animation/PoseSearch/Source/Editor/Public/PoseSearchDatabaseViewModel.h
@@ -98,7 +98,6 @@ namespace UE::PoseSearch
 		bool IsAnimationPreviewMode(EAnimationPreviewMode PreviewMode) const;
 
 		void OnToggleDisplayRootMotionSpeed();
-		bool CanDisplayRootMotionSpeed() const;
 		bool IsDisplayRootMotionSpeedChecked() const { return DisplayRootMotionSpeed; };
 		
 		void AddSequenceToDatabase(UAnimSequence* AnimSequence);
