Undo / Redo support for PoseSearchDatabase editing via the Motion Matching editor.

diff --git a/Engine/Plugins/Experimental/Animation/PoseSearch/Source/Editor/Private/PoseSearchDatabaseEditor.cpp b/Engine/Plugins/Experimental/Animation/PoseSearch/Source/Editor/Private/PoseSearchDatabaseEditor.cpp
--- a/Engine/Plugins/Experimental/Animation/PoseSearch/Source/Editor/Private/PoseSearchDatabaseEditor.cpp
+++ b/Engine/Plugins/Experimental/Animation/PoseSearch/Source/Editor/Private/PoseSearchDatabaseEditor.cpp
@@ -540,6 +540,7 @@ namespace UE::PoseSearch
 					NewSelectionReflection->AddToRoot();
 					NewSelectionReflection->Sequence = GetPoseSearchDatabase()->Sequences[SelectedItem->SourceAssetIdx];
 					NewSelectionReflection->SetSourceLink(SelectedItem, AssetTreeWidget);
+					NewSelectionReflection->SetFlags(RF_Transactional);
 					SelectionReflections[SequenceSelectionIndex].Add(NewSelectionReflection);
 				}
 				else if (SelectedItem->SourceAssetType == ESearchIndexAssetType::BlendSpace)
@@ -548,6 +549,7 @@ namespace UE::PoseSearch
 					NewSelectionReflection->AddToRoot();
 					NewSelectionReflection->BlendSpace = GetPoseSearchDatabase()->BlendSpaces[SelectedItem->SourceAssetIdx];
 					NewSelectionReflection->SetSourceLink(SelectedItem, AssetTreeWidget);
+					NewSelectionReflection->SetFlags(RF_Transactional);
 					SelectionReflections[BlendSpaceSelectionIndex].Add(NewSelectionReflection);
 				}
 			}
diff --git a/Engine/Plugins/Experimental/Animation/PoseSearch/Source/Editor/Private/PoseSearchDatabaseEditorReflection.cpp b/Engine/Plugins/Experimental/Animation/PoseSearch/Source/Editor/Private/PoseSearchDatabaseEditorReflection.cpp
--- a/Engine/Plugins/Experimental/Animation/PoseSearch/Source/Editor/Private/PoseSearchDatabaseEditorReflection.cpp
+++ b/Engine/Plugins/Experimental/Animation/PoseSearch/Source/Editor/Private/PoseSearchDatabaseEditorReflection.cpp
@@ -23,11 +23,16 @@ void UPoseSearchDatabaseSequenceReflection::PostEditChangeProperty(
 
 	check(WeakAssetTreeNode.Pin()->SourceAssetType == ESearchIndexAssetType::Sequence);
 
-	UPoseSearchDatabase* Database = WeakAssetTreeNode.Pin()->EditorViewModel.Pin()->GetPoseSearchDatabase();
-	if (IsValid(Database))
+	if (const TSharedPtr<UE::PoseSearch::FDatabaseViewModel> ViewModel = WeakAssetTreeNode.Pin()->EditorViewModel.Pin())
 	{
-		Database->Sequences[WeakAssetTreeNode.Pin()->SourceAssetIdx] = Sequence;
-		AssetTreeWidget->FinalizeTreeChanges(true);
+		UPoseSearchDatabase* Database = ViewModel->GetPoseSearchDatabase();
+		if (IsValid(Database))
+		{
+			Database->Sequences[WeakAssetTreeNode.Pin()->SourceAssetIdx] = Sequence;
+			Database->MarkPackageDirty();
+		
+			AssetTreeWidget->FinalizeTreeChanges(true);
+		}
 	}
 }
 
@@ -38,11 +43,16 @@ void UPoseSearchDatabaseBlendSpaceReflection::PostEditChangeProperty(
 
 	check(WeakAssetTreeNode.Pin()->SourceAssetType == ESearchIndexAssetType::BlendSpace);
 
-	UPoseSearchDatabase* Database = WeakAssetTreeNode.Pin()->EditorViewModel.Pin()->GetPoseSearchDatabase();
-	if (IsValid(Database))
+	if (const TSharedPtr<UE::PoseSearch::FDatabaseViewModel> ViewModel = WeakAssetTreeNode.Pin()->EditorViewModel.Pin())
 	{
-		Database->BlendSpaces[WeakAssetTreeNode.Pin()->SourceAssetIdx] = BlendSpace;
-		AssetTreeWidget->FinalizeTreeChanges(true);
+		UPoseSearchDatabase* Database = WeakAssetTreeNode.Pin()->EditorViewModel.Pin()->GetPoseSearchDatabase();
+		if (IsValid(Database))
+		{
+			Database->BlendSpaces[WeakAssetTreeNode.Pin()->SourceAssetIdx] = BlendSpace;
+			Database->MarkPackageDirty();
+		
+			AssetTreeWidget->FinalizeTreeChanges(true);
+		}
 	}
 }
 
diff --git a/Engine/Plugins/Experimental/Animation/PoseSearch/Source/Editor/Private/SPoseSearchDatabaseAssetList.cpp b/Engine/Plugins/Experimental/Animation/PoseSearch/Source/Editor/Private/SPoseSearchDatabaseAssetList.cpp
--- a/Engine/Plugins/Experimental/Animation/PoseSearch/Source/Editor/Private/SPoseSearchDatabaseAssetList.cpp
+++ b/Engine/Plugins/Experimental/Animation/PoseSearch/Source/Editor/Private/SPoseSearchDatabaseAssetList.cpp
@@ -8,6 +8,7 @@
 #include "Animation/AnimSequence.h"
 #include "Animation/BlendSpace.h"
 
+#include "Misc/TransactionObjectEvent.h"
 #include "DragAndDrop/AssetDragDropOp.h"
 #include "Misc/FeedbackContext.h"
 #include "AssetSelection.h"
@@ -61,28 +62,34 @@ namespace UE::PoseSearch
 		{
 			if (SourceAssetType == ESearchIndexAssetType::Sequence)
 			{
-				if (const UAnimSequence* Sequence = Database->Sequences[SourceAssetIdx].Sequence)
+				if (Database->Sequences.IsValidIndex(SourceAssetIdx))
 				{
-					return Sequence->HasRootMotion();
+					if (const UAnimSequence* Sequence = Database->Sequences[SourceAssetIdx].Sequence)
+					{
+						return Sequence->HasRootMotion();
+					}
 				}
 			}
 			else if (SourceAssetType == ESearchIndexAssetType::BlendSpace)
 			{
-				if (const UBlendSpace* BlendSpace = Database->BlendSpaces[SourceAssetIdx].BlendSpace)
+				if (Database->BlendSpaces.IsValidIndex(SourceAssetIdx))
 				{
-					bool bIsRootMotionUsedInBlendSpace = false;
-			
-					BlendSpace->ForEachImmutableSample([&bIsRootMotionUsedInBlendSpace](const FBlendSample & Sample)
+					if (const UBlendSpace* BlendSpace = Database->BlendSpaces[SourceAssetIdx].BlendSpace)
 					{
-						const TObjectPtr<UAnimSequence> Sequence = Sample.Animation;
-				
-						if (IsValid(Sequence) && Sequence->HasRootMotion())
+						bool bIsRootMotionUsedInBlendSpace = false;
+			
+						BlendSpace->ForEachImmutableSample([&bIsRootMotionUsedInBlendSpace](const FBlendSample & Sample)
 						{
-							bIsRootMotionUsedInBlendSpace = true;
-						}
-					});
+							const TObjectPtr<UAnimSequence> Sequence = Sample.Animation;
+				
+							if (IsValid(Sequence) && Sequence->HasRootMotion())
+							{
+								bIsRootMotionUsedInBlendSpace = true;
+							}
+						});
 			
-					return bIsRootMotionUsedInBlendSpace;
+						return bIsRootMotionUsedInBlendSpace;
+					}
 				}
 			}
 		}
@@ -96,16 +103,22 @@ namespace UE::PoseSearch
 		{
 			if (SourceAssetType == ESearchIndexAssetType::Sequence)
 			{
-				if (const UAnimSequence* Sequence = Database->Sequences[SourceAssetIdx].Sequence)
+				if (Database->Sequences.IsValidIndex(SourceAssetIdx))
 				{
-					return Sequence->bLoop;
+					if (const UAnimSequence* Sequence = Database->Sequences[SourceAssetIdx].Sequence)
+					{
+						return Sequence->bLoop;
+					}
 				}
 			}
 			else if (SourceAssetType == ESearchIndexAssetType::BlendSpace)
 			{
-				if (const UBlendSpace* BlendSpace = Database->BlendSpaces[SourceAssetIdx].BlendSpace)
+				if (Database->BlendSpaces.IsValidIndex(SourceAssetIdx))
 				{
-					return BlendSpace->bLoop;
+					if (const UBlendSpace* BlendSpace = Database->BlendSpaces[SourceAssetIdx].BlendSpace)
+					{
+						return BlendSpace->bLoop;
+					}
 				}
 			}
 		}
@@ -119,11 +132,17 @@ namespace UE::PoseSearch
 		{
 			if (SourceAssetType == ESearchIndexAssetType::Sequence)
 			{
-				return Database->Sequences[SourceAssetIdx].MirrorOption;
+				if (Database->Sequences.IsValidIndex(SourceAssetIdx))
+				{
+					return Database->Sequences[SourceAssetIdx].MirrorOption;
+				}
 			}
 			else if (SourceAssetType == ESearchIndexAssetType::BlendSpace)
 			{
-				return Database->BlendSpaces[SourceAssetIdx].MirrorOption;
+				if (Database->BlendSpaces.IsValidIndex(SourceAssetIdx))
+				{
+					return Database->BlendSpaces[SourceAssetIdx].MirrorOption;
+				}
 			}
 		}
 
@@ -498,16 +517,23 @@ namespace UE::PoseSearch
 	{
 		TSharedPtr<FDatabaseViewModel> ViewModelPtr = EditorViewModel.Pin();
 		TSharedPtr<FDatabaseAssetTreeNode> TreeNodePtr = WeakAssetTreeNode.Pin();
-
+		const UPoseSearchDatabase * Database = ViewModelPtr->GetPoseSearchDatabase();
+		
 		if (TreeNodePtr->SourceAssetType == ESearchIndexAssetType::Sequence)
 		{
-			const bool bEnabled = ViewModelPtr->IsSelectedSequenceEnabled(TreeNodePtr->SourceAssetIdx);
-			return bEnabled ? ECheckBoxState::Checked : ECheckBoxState::Unchecked;
+			if (Database->Sequences.IsValidIndex(TreeNodePtr->SourceAssetIdx))
+			{
+				const bool bEnabled = ViewModelPtr->IsSelectedSequenceEnabled(TreeNodePtr->SourceAssetIdx);
+				return bEnabled ? ECheckBoxState::Checked : ECheckBoxState::Unchecked;
+			}
 		}
 		else if (TreeNodePtr->SourceAssetType == ESearchIndexAssetType::BlendSpace)
 		{
-			const bool bEnabled = ViewModelPtr->IsSelectedBlendSpaceEnabled(TreeNodePtr->SourceAssetIdx);
-			return bEnabled ? ECheckBoxState::Checked : ECheckBoxState::Unchecked;
+			if (Database->BlendSpaces.IsValidIndex(TreeNodePtr->SourceAssetIdx))
+			{
+				const bool bEnabled = ViewModelPtr->IsSelectedBlendSpaceEnabled(TreeNodePtr->SourceAssetIdx);
+				return bEnabled ? ECheckBoxState::Checked : ECheckBoxState::Unchecked;
+			}
 		}
 
 		return ECheckBoxState::Unchecked;
@@ -515,9 +541,13 @@ namespace UE::PoseSearch
 
 	void SDatabaseAssetListItem::OnAssetIsEnabledChanged(ECheckBoxState NewCheckboxState)
 	{
-		TSharedPtr<FDatabaseViewModel> ViewModelPtr = EditorViewModel.Pin();
-		TSharedPtr<FDatabaseAssetTreeNode> TreeNodePtr = WeakAssetTreeNode.Pin();
+		const FScopedTransaction Transaction(LOCTEXT("EnableChangedForAssetInPoseSearchDatabase", "Update enabled flag for item from Pose Search Database"));
+
+		const TSharedPtr<FDatabaseViewModel> ViewModelPtr = EditorViewModel.Pin();
+		const TSharedPtr<FDatabaseAssetTreeNode> TreeNodePtr = WeakAssetTreeNode.Pin();
 
+		ViewModelPtr->GetPoseSearchDatabase()->Modify();
+		
 		if (TreeNodePtr->SourceAssetType == ESearchIndexAssetType::Sequence)
 		{
 			ViewModelPtr->SetSelectedSequenceEnabled(TreeNodePtr->SourceAssetIdx, NewCheckboxState == ECheckBoxState::Checked ? true : false);
@@ -584,7 +614,7 @@ namespace UE::PoseSearch
 	FText SDatabaseAssetListItem::GetMirrorOptionToolTip() const
 	{
 		const TSharedPtr<FDatabaseAssetTreeNode> Node = WeakAssetTreeNode.Pin();
-		return FText::FromString(LOCTEXT("ToolTipMirrorOption", "Mirror Option: ").ToString() + UEnum::GetDisplayValueAsText(Node->GetMirrorOption()).ToString());
+		return FText::FromString(LOCTEXT("ToolTipMirrorOption", "Mirror Option: ").ToString() + (Node ? UEnum::GetDisplayValueAsText(Node->GetMirrorOption()).ToString() : LOCTEXT("ToolTipMirrorOption_Invalid", "Invalid").ToString()));
 	}
 	
 	FText SDatabaseAssetListItem::GetAssetEnabledToolTip() const
@@ -733,6 +763,48 @@ namespace UE::PoseSearch
 		return FReply::Unhandled();
 	}
 
+	bool SDatabaseAssetTree::MatchesContext(const FTransactionContext& InContext, const TArray<TPair<UObject*, FTransactionObjectEvent>>& TransactionObjectContexts) const
+	{
+		// Ensure that we only react to modifications to the UPosesSearchDatabase.
+		if (const TSharedPtr<FDatabaseViewModel> ViewModel = EditorViewModel.Pin())
+		{
+			if (const UPoseSearchDatabase * Database = ViewModel->GetPoseSearchDatabase())
+			{
+				for (const TPair<UObject*, FTransactionObjectEvent>& TransactionObjectPair : TransactionObjectContexts)
+				{
+					const UObject* Object = TransactionObjectPair.Key;
+					while (Object != nullptr)
+					{
+						if (Object == Database)
+						{
+							return true;
+						}
+
+						Object = Object->GetOuter();
+					}
+				}
+			}
+		}
+		
+		return false;
+	}
+
+	void SDatabaseAssetTree::PostUndo(bool bSuccess)
+	{
+		if (bSuccess)
+		{
+			FinalizeTreeChanges();
+		}
+	}
+
+	void SDatabaseAssetTree::PostRedo(bool bSuccess)
+	{
+		if (bSuccess)
+		{
+			FinalizeTreeChanges();
+		}
+	}
+
 	void SDatabaseAssetTree::RefreshTreeView(bool bIsInitialSetup, bool bRecoverSelection)
 	{
 		const TSharedPtr<FDatabaseViewModel> ViewModel = EditorViewModel.Pin();
@@ -743,6 +815,7 @@ namespace UE::PoseSearch
 
 		const TSharedRef<FDatabaseViewModel> ViewModelRef = ViewModel.ToSharedRef();
 
+		// Empty node data.
 		RootNodes.Reset();
 		AllNodes.Reset();
 
@@ -753,111 +826,127 @@ namespace UE::PoseSearch
 			return;
 		}
 
-		// store selection so we can recover it afterwards (if possible)
+		// Store selection so we can recover it afterwards (if possible)
 		TArray<TSharedPtr<FDatabaseAssetTreeNode>> PreviouslySelectedNodes = TreeView->GetSelectedItems();
 
-		TSharedPtr<FDatabaseAssetTreeNode> DefaultGroupNode = MakeShared<FDatabaseAssetTreeNode>(
-			INDEX_NONE,
-			ESearchIndexAssetType::Invalid,
-			ViewModelRef);
-		AllNodes.Add(DefaultGroupNode);
-		RootNodes.Add(DefaultGroupNode);
-
-		const int32 DefaultGroupIdx = RootNodes.Num() - 1;
-
-		auto CreateAssetNode = [this, ViewModelRef](int32 AssetIdx, ESearchIndexAssetType AssetType, int32 GroupIdx)
-		{
-			TSharedPtr<FDatabaseAssetTreeNode> SequenceGroupNode = MakeShared<FDatabaseAssetTreeNode>(
-				AssetIdx,
-				AssetType,
-				ViewModelRef);
-			TSharedPtr<FDatabaseAssetTreeNode>& ParentGroupNode = RootNodes[GroupIdx];
-			SequenceGroupNode->Parent = ParentGroupNode;
-			ParentGroupNode->Children.Add(SequenceGroupNode);
-			AllNodes.Add(SequenceGroupNode);
-		};
-
-		// Build an index based off of alphabetical order than iterate the index instead
-		TArray<uint32> SequenceIndexArray;
-		SequenceIndexArray.SetNumUninitialized(Database->Sequences.Num());
-		for (int32 SequenceIdx = 0; SequenceIdx < Database->Sequences.Num(); ++SequenceIdx)
-		{
-			SequenceIndexArray[SequenceIdx] = SequenceIdx;
-		}
-
-		SequenceIndexArray.Sort([Database](int32 SequenceIdxA, int32 SequenceIdxB)
+		// Rebuild node hierarchy
 		{
-			const FPoseSearchDatabaseSequence& A = Database->Sequences[SequenceIdxA];
-			const FPoseSearchDatabaseSequence& B = Database->Sequences[SequenceIdxB];
-
-			//If its null add it to the end of the list 
-			if (!B.Sequence)
+			// Setup default group node
 			{
-				return true;
+				const TSharedPtr<FDatabaseAssetTreeNode> DefaultGroupNode = MakeShared<FDatabaseAssetTreeNode>(
+					INDEX_NONE,
+					ESearchIndexAssetType::Invalid,
+					ViewModelRef);
+				AllNodes.Add(DefaultGroupNode);
+				RootNodes.Add(DefaultGroupNode);
 			}
-
-			if (!A.Sequence)
+			
+			const int32 DefaultGroupIdx = RootNodes.Num() - 1;
+			
+			auto CreateAssetNode = [this, ViewModelRef](int32 AssetIdx, ESearchIndexAssetType AssetType, int32 GroupIdx)
+			{
+				// Create sequence node
+				const TSharedPtr<FDatabaseAssetTreeNode> SequenceGroupNode = MakeShared<FDatabaseAssetTreeNode>(
+					AssetIdx,
+					AssetType,
+					ViewModelRef);
+				const TSharedPtr<FDatabaseAssetTreeNode>& ParentGroupNode = RootNodes[GroupIdx];
+
+				// Setup hierarchy
+				SequenceGroupNode->Parent = ParentGroupNode;
+				ParentGroupNode->Children.Add(SequenceGroupNode);
+
+				// Keep track of node
+				AllNodes.Add(SequenceGroupNode);
+			};
+			
+			// Build an index based off of alphabetical order than iterate the index instead
+			TArray<uint32> SequenceIndexArray;
+			SequenceIndexArray.SetNumUninitialized(Database->Sequences.Num());
+			for (int32 SequenceIdx = 0; SequenceIdx < Database->Sequences.Num(); ++SequenceIdx)
 			{
-				return false;
+				SequenceIndexArray[SequenceIdx] = SequenceIdx;
 			}
 
-			const int32 Comparison = A.Sequence->GetName().Compare(B.Sequence->GetName());
-			return Comparison < 0;
-		});
-
-		// create all sequence nodes
-		for (int32 SequenceIdx = 0; SequenceIdx < Database->Sequences.Num(); ++SequenceIdx)
-		{
-			const int32 MappedId = SequenceIndexArray[SequenceIdx];
-			const FPoseSearchDatabaseSequence& DbSequence = Database->Sequences[MappedId];
-			const bool bFiltered = (DbSequence.Sequence == nullptr || GetAssetFilterString().IsEmpty()) ? false : !DbSequence.Sequence->GetName().Contains(GetAssetFilterString());
-
-			if (!bFiltered)
+			SequenceIndexArray.Sort([Database](int32 SequenceIdxA, int32 SequenceIdxB)
 			{
-				CreateAssetNode(MappedId, ESearchIndexAssetType::Sequence, DefaultGroupIdx);
-			}
-		}
+				const FPoseSearchDatabaseSequence& A = Database->Sequences[SequenceIdxA];
+				const FPoseSearchDatabaseSequence& B = Database->Sequences[SequenceIdxB];
 
-		TArray<uint32> BlendspaceIndexArray;
-		BlendspaceIndexArray.SetNumUninitialized(Database->BlendSpaces.Num());
-		for (int32 BlendspaceIdx = 0; BlendspaceIdx < Database->BlendSpaces.Num(); ++BlendspaceIdx)
-		{
-			BlendspaceIndexArray[BlendspaceIdx] = BlendspaceIdx;
-		}
+				//If its null add it to the end of the list 
+				if (!B.Sequence)
+				{
+					return true;
+				}
 
-		BlendspaceIndexArray.Sort([Database](int32 BlendspaceIdxA, int32 BlendspaceIdxB)
-		{
-			const FPoseSearchDatabaseBlendSpace& A = Database->BlendSpaces[BlendspaceIdxA];
-			const FPoseSearchDatabaseBlendSpace& B = Database->BlendSpaces[BlendspaceIdxB];
+				if (!A.Sequence)
+				{
+					return false;
+				}
+
+				const int32 Comparison = A.Sequence->GetName().Compare(B.Sequence->GetName());
+				return Comparison < 0;
+			});
 
-			//If its null add it to the end of the list 
-			if (!B.BlendSpace)
+			// create all sequence nodes
+			for (int32 SequenceIdx = 0; SequenceIdx < Database->Sequences.Num(); ++SequenceIdx)
 			{
-				return true;
+				const int32 MappedId = SequenceIndexArray[SequenceIdx];
+				const FPoseSearchDatabaseSequence& DbSequence = Database->Sequences[MappedId];
+				const bool bFiltered = (DbSequence.Sequence == nullptr || GetAssetFilterString().IsEmpty()) ? false : !DbSequence.Sequence->GetName().Contains(GetAssetFilterString());
+
+				if (!bFiltered)
+				{
+					CreateAssetNode(MappedId, ESearchIndexAssetType::Sequence, DefaultGroupIdx);
+				}
 			}
 
-			if (!A.BlendSpace)
+			// Build an index based off of alphabetical order than iterate the index instead
+			TArray<uint32> BlendspaceIndexArray;
+			BlendspaceIndexArray.SetNumUninitialized(Database->BlendSpaces.Num());
+			for (int32 BlendspaceIdx = 0; BlendspaceIdx < Database->BlendSpaces.Num(); ++BlendspaceIdx)
 			{
-				return false;
+				BlendspaceIndexArray[BlendspaceIdx] = BlendspaceIdx;
 			}
 
-			const int32 Comparison = A.BlendSpace->GetName().Compare(B.BlendSpace->GetName());
-			return Comparison < 0;
-		});
+			BlendspaceIndexArray.Sort([Database](int32 BlendspaceIdxA, int32 BlendspaceIdxB)
+			{
+				const FPoseSearchDatabaseBlendSpace& A = Database->BlendSpaces[BlendspaceIdxA];
+				const FPoseSearchDatabaseBlendSpace& B = Database->BlendSpaces[BlendspaceIdxB];
 
-		// create all blendspace nodes
-		for (int32 BlendSpaceIdx = 0; BlendSpaceIdx < Database->BlendSpaces.Num(); ++BlendSpaceIdx)
-		{
-			const int32 MappedId = BlendspaceIndexArray[BlendSpaceIdx];
-			const FPoseSearchDatabaseBlendSpace& DbBlendSpace = Database->BlendSpaces[MappedId];
-			const bool bFiltered = (DbBlendSpace.BlendSpace == nullptr || GetAssetFilterString().IsEmpty()) ? false : !DbBlendSpace.BlendSpace->GetName().Contains(GetAssetFilterString());
+				//If its null add it to the end of the list 
+				if (!B.BlendSpace)
+				{
+					return true;
+				}
+
+				if (!A.BlendSpace)
+				{
+					return false;
+				}
+
+				const int32 Comparison = A.BlendSpace->GetName().Compare(B.BlendSpace->GetName());
+				return Comparison < 0;
+			});
 
-			if (!bFiltered)
+			// create all blendspace nodes
+			for (int32 BlendSpaceIdx = 0; BlendSpaceIdx < Database->BlendSpaces.Num(); ++BlendSpaceIdx)
 			{
-				CreateAssetNode(MappedId, ESearchIndexAssetType::BlendSpace, DefaultGroupIdx);
+				const int32 MappedId = BlendspaceIndexArray[BlendSpaceIdx];
+				const FPoseSearchDatabaseBlendSpace& DbBlendSpace = Database->BlendSpaces[MappedId];
+				const bool bFiltered = (DbBlendSpace.BlendSpace == nullptr || GetAssetFilterString().IsEmpty()) ? false : !DbBlendSpace.BlendSpace->GetName().Contains(GetAssetFilterString());
+
+				if (!bFiltered)
+				{
+					CreateAssetNode(MappedId, ESearchIndexAssetType::BlendSpace, DefaultGroupIdx);
+				}
 			}
+
+			// Show drag and drop suggestion if tree is empty
+			TreeViewDragAndDropSuggestion->SetVisibility(BlendspaceIndexArray.IsEmpty() && SequenceIndexArray.IsEmpty() ? EVisibility::Visible : EVisibility::Hidden);
 		}
 
+		// Update tree view
 		TreeView->RequestTreeRefresh();
 
 		for (TSharedPtr<FDatabaseAssetTreeNode>& RootNode : RootNodes)
@@ -865,6 +954,7 @@ namespace UE::PoseSearch
 			TreeView->SetItemExpansion(RootNode, true);
 		}
 
+		// Handle selection
 		if (bRecoverSelection)
 		{
 			RecoverSelection(PreviouslySelectedNodes);
@@ -873,9 +963,6 @@ namespace UE::PoseSearch
 		{
 			TreeView->SetItemSelection(PreviouslySelectedNodes, false, ESelectInfo::Direct);
 		}
-		
-		// Show drag and drop suggestion if tree is empty
-		TreeViewDragAndDropSuggestion->SetVisibility(BlendspaceIndexArray.IsEmpty() && SequenceIndexArray.IsEmpty() ? EVisibility::Visible : EVisibility::Hidden);
 	}
 
 	TSharedRef<ITableRow> SDatabaseAssetTree::MakeTableRowWidget(
@@ -928,7 +1015,7 @@ namespace UE::PoseSearch
 		EItemDropZone DropZone,
 		TSharedPtr<FDatabaseAssetTreeNode> TargetItem)
 	{
-		TSharedPtr<FDragDropOperation> Operation = DragDropEvent.GetOperation();
+		const TSharedPtr<FDragDropOperation> Operation = DragDropEvent.GetOperation();
 
 		const bool bValidOperation = Operation.IsValid() && Operation->IsOfType<FAssetDragDropOp>();
 		if (!bValidOperation)
@@ -936,7 +1023,7 @@ namespace UE::PoseSearch
 			return FReply::Unhandled();
 		}
 
-		TSharedPtr<FDatabaseViewModel> ViewModel = EditorViewModel.Pin();
+		const TSharedPtr<FDatabaseViewModel> ViewModel = EditorViewModel.Pin();
 		if (!ViewModel.IsValid())
 		{
 			return FReply::Unhandled();
@@ -952,6 +1039,7 @@ namespace UE::PoseSearch
 
 			{
 				const FScopedTransaction Transaction(LOCTEXT("AddSequencesOrBlendspaces", "Add Sequence(s) and/or Blendspace(s) to Pose Search Database"));
+				ViewModel->GetPoseSearchDatabase()->Modify();
 				
 				for (int32 DroppedAssetIdx = 0; DroppedAssetIdx < NumAssets; ++DroppedAssetIdx)
 				{
@@ -1091,8 +1179,11 @@ namespace UE::PoseSearch
 	void SDatabaseAssetTree::OnAddSequence(bool bFinalizeChanges)
 	{
 		FScopedTransaction Transaction(LOCTEXT("AddSequence", "Add Sequence"));
-
-		EditorViewModel.Pin()->AddSequenceToDatabase(nullptr);
+		const TSharedPtr<FDatabaseViewModel> ViewModel = EditorViewModel.Pin();
+		
+		ViewModel->GetPoseSearchDatabase()->Modify();
+		
+		ViewModel->AddSequenceToDatabase(nullptr);
 
 		if (bFinalizeChanges)
 		{
@@ -1103,8 +1194,11 @@ namespace UE::PoseSearch
 	void SDatabaseAssetTree::OnAddBlendSpace(bool bFinalizeChanges)
 	{
 		FScopedTransaction Transaction(LOCTEXT("AddBlendSpaceTransaction", "Add Blend Space"));
-
-		EditorViewModel.Pin()->AddBlendSpaceToDatabase(nullptr);
+		const TSharedPtr<FDatabaseViewModel> ViewModel = EditorViewModel.Pin();
+		
+		ViewModel->GetPoseSearchDatabase()->Modify();
+		
+		ViewModel->AddBlendSpaceToDatabase(nullptr);
 
 		if (bFinalizeChanges)
 		{
@@ -1115,14 +1209,17 @@ namespace UE::PoseSearch
 	void SDatabaseAssetTree::OnDeleteAsset(TSharedPtr<FDatabaseAssetTreeNode> Node, bool bFinalizeChanges)
 	{
 		FScopedTransaction Transaction(LOCTEXT("DeleteAsset", "Delete Asset"));
+		const TSharedPtr<FDatabaseViewModel> ViewModel = EditorViewModel.Pin();
+		
+		ViewModel->GetPoseSearchDatabase()->Modify();
 
 		if (Node->SourceAssetType == ESearchIndexAssetType::Sequence)
 		{
-			EditorViewModel.Pin()->DeleteSequenceFromDatabase(Node->SourceAssetIdx);
+			ViewModel->DeleteSequenceFromDatabase(Node->SourceAssetIdx);
 		}
 		else if (Node->SourceAssetType == ESearchIndexAssetType::BlendSpace)
 		{
-			EditorViewModel.Pin()->DeleteBlendSpaceFromDatabase(Node->SourceAssetIdx);
+			ViewModel->DeleteBlendSpaceFromDatabase(Node->SourceAssetIdx);
 		}
 		else
 		{
@@ -1151,7 +1248,7 @@ namespace UE::PoseSearch
 
 		for (const TSharedPtr<FDatabaseAssetTreeNode>& Node : AllNodes)
 		{
-			bool bFoundNode = PreviouslySelectedNodes.ContainsByPredicate(
+			const bool bFoundNode = PreviouslySelectedNodes.ContainsByPredicate(
 				[Node](const TSharedPtr<FDatabaseAssetTreeNode>& PrevSelectedNode)
 			{
 				return
@@ -1197,11 +1294,14 @@ namespace UE::PoseSearch
 
 	void SDatabaseAssetTree::OnDeleteNodes()
 	{
-		const FScopedTransaction Transaction(LOCTEXT("DeletePoseSearchDatabaseNodes", "Delete selected items from Pose Search Database"));
-
 		TArray<TSharedPtr<FDatabaseAssetTreeNode>> SelectedNodes = TreeView->GetSelectedItems();
 		if (!SelectedNodes.IsEmpty())
 		{
+			const FScopedTransaction Transaction(LOCTEXT("DeletePoseSearchDatabaseNodes", "Delete selected items from Pose Search Database"));
+			const TSharedPtr<FDatabaseViewModel> ViewModel = EditorViewModel.Pin();
+
+			ViewModel->GetPoseSearchDatabase()->Modify();
+			
 			SelectedNodes.Sort(
 				[](const TSharedPtr<FDatabaseAssetTreeNode>& A, const TSharedPtr<FDatabaseAssetTreeNode>& B)
 			{
@@ -1232,20 +1332,23 @@ namespace UE::PoseSearch
 
 	void SDatabaseAssetTree::OnEnableNodes()
 	{
-		const FScopedTransaction Transaction(LOCTEXT("EnablePoseSearchDatabaseNodes", "Enable selected items from Pose Search Database"));
-
 		TArray<TSharedPtr<FDatabaseAssetTreeNode>> SelectedNodes = TreeView->GetSelectedItems();
 		if (!SelectedNodes.IsEmpty())
 		{
+			const FScopedTransaction Transaction(LOCTEXT("EnablePoseSearchDatabaseNodes", "Enable selected items from Pose Search Database"));
+			const TSharedPtr<FDatabaseViewModel> ViewModel = EditorViewModel.Pin();
+
+			ViewModel->GetPoseSearchDatabase()->Modify();
+			
 			for (TSharedPtr<FDatabaseAssetTreeNode> SelectedNode : SelectedNodes)
 			{
 				if (SelectedNode->SourceAssetType == ESearchIndexAssetType::Sequence)
 				{
-					EditorViewModel.Pin()->SetSelectedSequenceEnabled(SelectedNode->SourceAssetIdx, true);
+					ViewModel->SetSelectedSequenceEnabled(SelectedNode->SourceAssetIdx, true);
 				}
 				else if (SelectedNode->SourceAssetType == ESearchIndexAssetType::BlendSpace)
 				{
-					EditorViewModel.Pin()->SetSelectedBlendSpaceEnabled(SelectedNode->SourceAssetIdx, true);
+					ViewModel->SetSelectedBlendSpaceEnabled(SelectedNode->SourceAssetIdx, true);
 				}
 			}
 		
@@ -1255,20 +1358,23 @@ namespace UE::PoseSearch
 
 	void SDatabaseAssetTree::OnDisableNodes()
 	{
-		const FScopedTransaction Transaction(LOCTEXT("DisablePoseSearchDatabaseNodes", "Disable selected items from Pose Search Database"));
-
 		TArray<TSharedPtr<FDatabaseAssetTreeNode>> SelectedNodes = TreeView->GetSelectedItems();
 		if (!SelectedNodes.IsEmpty())
 		{
+			const FScopedTransaction Transaction(LOCTEXT("DisablePoseSearchDatabaseNodes", "Disable selected items from Pose Search Database"));
+			const TSharedPtr<FDatabaseViewModel> ViewModel = EditorViewModel.Pin();
+
+			ViewModel->GetPoseSearchDatabase()->Modify();
+			
 			for (TSharedPtr<FDatabaseAssetTreeNode> SelectedNode : SelectedNodes)
 			{
 				if (SelectedNode->SourceAssetType == ESearchIndexAssetType::Sequence)
 				{
-					EditorViewModel.Pin()->SetSelectedSequenceEnabled(SelectedNode->SourceAssetIdx, false);
+					ViewModel->SetSelectedSequenceEnabled(SelectedNode->SourceAssetIdx, false);
 				}
 				else if (SelectedNode->SourceAssetType == ESearchIndexAssetType::BlendSpace)
 				{
-					EditorViewModel.Pin()->SetSelectedBlendSpaceEnabled(SelectedNode->SourceAssetIdx, false);
+					ViewModel->SetSelectedBlendSpaceEnabled(SelectedNode->SourceAssetIdx, false);
 				}
 			}
 
diff --git a/Engine/Plugins/Experimental/Animation/PoseSearch/Source/Editor/Public/SPoseSearchDatabaseAssetList.h b/Engine/Plugins/Experimental/Animation/PoseSearch/Source/Editor/Public/SPoseSearchDatabaseAssetList.h
--- a/Engine/Plugins/Experimental/Animation/PoseSearch/Source/Editor/Public/SPoseSearchDatabaseAssetList.h
+++ b/Engine/Plugins/Experimental/Animation/PoseSearch/Source/Editor/Public/SPoseSearchDatabaseAssetList.h
@@ -84,7 +84,7 @@ namespace UE::PoseSearch
 		TWeakPtr<SDatabaseAssetTree> SkeletonView;
 	};
 
-	class SDatabaseAssetTree : public SCompoundWidget, public FEditorUndoClient
+	class SDatabaseAssetTree : public SCompoundWidget, public FSelfRegisteringEditorUndoClient
 	{
 	public:
 		SLATE_BEGIN_ARGS(SDatabaseAssetTree) {}
@@ -99,7 +99,13 @@ namespace UE::PoseSearch
 		virtual FReply OnDrop(const FGeometry& MyGeometry, const FDragDropEvent& DragDropEvent) override;
 		virtual FReply OnKeyDown(const FGeometry& MyGeometry, const FKeyEvent& InKeyEvent) override;
 		// End SWidget interface
-	
+
+		// Begin FEditorUndoClient interface
+		virtual bool MatchesContext(const FTransactionContext& InContext, const TArray<TPair<UObject*, FTransactionObjectEvent>>& TransactionObjectContexts) const override;
+		virtual void PostUndo(bool bSuccess) override;
+		virtual void PostRedo(bool bSuccess) override;
+		// End FEditorUndoClient interface
+		
 		void RefreshTreeView(bool bIsInitialSetup = false, bool bRecoverSelection = false);
 		void FinalizeTreeChanges(bool bRecoverSelection = false);
 
