motion matching - UNormalizationSetAsset: data mining multiple UPoseDatabaseAsset(s) to calculate mean deviations

diff --git a/Engine/Plugins/Experimental/Animation/PoseSearch/Source/Runtime/Private/PoseSearch.cpp b/Engine/Plugins/Experimental/Animation/PoseSearch/Source/Runtime/Private/PoseSearch.cpp
--- a/Engine/Plugins/Experimental/Animation/PoseSearch/Source/Runtime/Private/PoseSearch.cpp
+++ b/Engine/Plugins/Experimental/Animation/PoseSearch/Source/Runtime/Private/PoseSearch.cpp
@@ -4,6 +4,7 @@
 #include "AssetRegistry/AssetData.h"
 #include "PoseSearch/PoseSearchAnimNotifies.h"
 #include "PoseSearch/PoseSearchDerivedData.h"
+#include "PoseSearch/PoseSearchDerivedDataKey.h"
 #include "PoseSearchEigenHelper.h"
 
 #include "Algo/BinarySearch.h"
@@ -369,22 +370,16 @@ void UPoseSearchFeatureChannel::InitializeSchema(UE::PoseSearch::FSchemaInitiali
 }
 
 #if WITH_EDITOR
-void UPoseSearchFeatureChannel::ComputeCostBreakdowns(UE::PoseSearch::ICostBreakDownData& CostBreakDownData, const UPoseSearchSchema* Schema) const
+void UPoseSearchFeatureChannel::PopulateChannelLayoutSet(UE::PoseSearch::FFeatureChannelLayoutSet& FeatureChannelLayoutSet) const
 {
-	CostBreakDownData.AddEntireBreakDownSection(FText::FromString(GetName()), Schema, ChannelDataOffset, ChannelCardinality);
+	FeatureChannelLayoutSet.Add(GetName(), UE::PoseSearch::FKeyBuilder(this).Finalize(), ChannelDataOffset, ChannelCardinality);
 }
-#endif // WITH_EDITOR
 
-// base implementation calculating a single mean deviation value (replicated ChannelCardinality times into MeanDeviations starting at DataOffset index) from all the features data associated to this channel
-void UPoseSearchFeatureChannel::ComputeMeanDeviations(const Eigen::MatrixXd& CenteredPoseMatrix, Eigen::VectorXd& MeanDeviations) const
+void UPoseSearchFeatureChannel::ComputeCostBreakdowns(UE::PoseSearch::ICostBreakDownData& CostBreakDownData, const UPoseSearchSchema* Schema) const
 {
-	using namespace UE::PoseSearch;
-
-	int32 DataOffset = ChannelDataOffset;
-	FFeatureVectorHelper::ComputeMeanDeviations(GetMinimumMeanDeviation(), CenteredPoseMatrix, MeanDeviations, DataOffset, ChannelCardinality);
-
-	check(DataOffset == ChannelDataOffset + ChannelCardinality);
+	CostBreakDownData.AddEntireBreakDownSection(FText::FromString(GetName()), Schema, ChannelDataOffset, ChannelCardinality);
 }
+#endif // WITH_EDITOR
 
 //////////////////////////////////////////////////////////////////////////
 // FSchemaInitializer
@@ -1965,29 +1960,6 @@ float FFeatureVectorHelper::DecodeFloatInternal(TConstArrayView<float> Values, i
 	return Values[DataOffset];
 }
 
-void FFeatureVectorHelper::ComputeMeanDeviations(float MinMeanDeviation, const Eigen::MatrixXd& CenteredPoseMatrix, Eigen::VectorXd& MeanDeviations, int32& DataOffset, int32 Cardinality)
-{
-	const int32 NumPoses = CenteredPoseMatrix.cols();
-
-	// Construct a submatrix for the feature and find the average distance to the feature's centroid.
-	// Since we've already mean centered the data, the average distance to the centroid is simply the average norm.
-	const double FeatureMeanDeviation = CenteredPoseMatrix.block(DataOffset, 0, Cardinality, NumPoses).colwise().norm().mean();
-
-	// Fill the feature's corresponding scaling axes with the average distance
-	// Avoid scaling by zero by leaving near-zero deviations as 1.0
-	MeanDeviations.segment(DataOffset, Cardinality).setConstant(FeatureMeanDeviation > MinMeanDeviation ? FeatureMeanDeviation : 1.f);
-
-	DataOffset += Cardinality;
-}
-
-void FFeatureVectorHelper::SetMeanDeviations(float Deviation, Eigen::VectorXd& MeanDeviations, int32& DataOffset, int32 Cardinality)
-{
-	// Fill the feature's corresponding scaling axes with the supplied value
-	MeanDeviations.segment(DataOffset, Cardinality).setConstant(Deviation);
-
-	DataOffset += Cardinality;
-}
-
 //////////////////////////////////////////////////////////////////////////
 // FDebugDrawParams
 bool FDebugDrawParams::CanDraw() const
@@ -2896,7 +2868,7 @@ void FBlendSpaceSampler::ProcessRootDistance()
 	// Also emit root motion summary info to help with sample wrapping in 
 	// FAssetIndexer::GetSampleTimeFromDistance() and FAssetIndexer::GetSampleInfo()
 	TotalRootTransform = LastRootTransform.GetRelativeTransform(InitialRootTransform);
-	TotalRootDistance = AccumulatedRootDistance.Last();
+	TotalRootDistance = AccumulatedRootDistance.Last(); 
 }
 
 const UAnimationAsset* FBlendSpaceSampler::GetAsset() const
diff --git a/Engine/Plugins/Experimental/Animation/PoseSearch/Source/Runtime/Private/PoseSearchDerivedData.cpp b/Engine/Plugins/Experimental/Animation/PoseSearch/Source/Runtime/Private/PoseSearchDerivedData.cpp
--- a/Engine/Plugins/Experimental/Animation/PoseSearch/Source/Runtime/Private/PoseSearchDerivedData.cpp
+++ b/Engine/Plugins/Experimental/Animation/PoseSearch/Source/Runtime/Private/PoseSearchDerivedData.cpp
@@ -29,7 +29,56 @@ static FCookStatsManager::FAutoRegisterCallback RegisterCookStats([](FCookStatsM
 	});
 #endif
 
-static Eigen::VectorXd ComputeChannelsDeviations(const FPoseSearchIndex& SearchIndex, const UPoseSearchSchema* Schema)
+static float ComputeFeatureMeanDeviation(TConstArrayView<FFeatureChannelLayoutSet::FEntry> Entries, TConstArrayView<FPoseSearchIndexBase> SearchIndexBases, TConstArrayView<const UPoseSearchSchema*> Schemas)
+{
+	check(Schemas.Num() == SearchIndexBases.Num());
+	
+	const int32 EntriesNum = Entries.Num();
+	check(EntriesNum > 0);
+
+	const int32 Cardinality = Entries[0].Cardinality;
+	check(Cardinality > 0);
+
+	int32 TotalNumPoses = 0;
+	for (int32 EntryIdx = 0; EntryIdx < EntriesNum; ++EntryIdx)
+	{
+		TotalNumPoses += SearchIndexBases[Entries[EntryIdx].SchemaIndex].NumPoses;
+	}
+
+	int32 AccumulatedNumPoses = 0;
+	RowMajorMatrix CenteredSubPoseMatrix(TotalNumPoses, Cardinality);
+	for (int32 EntryIdx = 0; EntryIdx < EntriesNum; ++EntryIdx)
+	{
+		const FFeatureChannelLayoutSet::FEntry& Entry = Entries[EntryIdx];
+		check(Cardinality == Entry.Cardinality);
+
+		const int32 DataSetIdx = Entry.SchemaIndex;
+
+		const UPoseSearchSchema* Schema = Schemas[DataSetIdx];
+		const FPoseSearchIndexBase& SearchIndex = SearchIndexBases[DataSetIdx];
+
+		const int32 NumPoses = SearchIndex.NumPoses;
+
+		// Map input buffer with NumPoses as rows and NumDimensions	as cols
+		RowMajorMatrixMapConst PoseMatrixSourceMap(SearchIndex.Values.GetData(), NumPoses, Schema->SchemaCardinality);
+
+		// Given the sub matrix for the features, find the average distance to the feature's centroid.
+		CenteredSubPoseMatrix.block(AccumulatedNumPoses, 0, NumPoses, Cardinality) = PoseMatrixSourceMap.block(0, Entry.DataOffset, NumPoses, Cardinality);
+		AccumulatedNumPoses += NumPoses;
+	}
+
+	RowMajorVector SampleMean = CenteredSubPoseMatrix.colwise().mean();
+	CenteredSubPoseMatrix = CenteredSubPoseMatrix.rowwise() - SampleMean;
+
+	// after mean centering the data, the average distance to the centroid is simply the average norm.
+	const float FeatureMeanDeviation = CenteredSubPoseMatrix.rowwise().norm().mean();
+
+	return FeatureMeanDeviation;
+}
+
+// it collects FFeatureChannelLayoutSet from all the Schemas (for example, figuring out the data offsets of SampledBones at a specific 
+// SampleTimes for a UPoseSearchFeatureChannel_Pose for all the SearchIndexBases), and call ComputeFeatureMeanDeviation
+static TArray<float> ComputeChannelsDeviations(TConstArrayView<FPoseSearchIndexBase> SearchIndexBases, TConstArrayView<const UPoseSearchSchema*> Schemas)
 {
 	// This function performs a modified z-score normalization where features are normalized
 	// by mean absolute deviation rather than standard deviation. Both methods are preferable
@@ -45,46 +94,56 @@ static Eigen::VectorXd ComputeChannelsDeviations(const FPoseSearchIndex& SearchI
 	//     British Journal of Educational Studies, 53: 417-430.
 
 	using namespace Eigen;
+	using namespace UE::PoseSearch;
 
-	check(Schema->IsValid());
-
-	const int32 NumPoses = SearchIndex.NumPoses;
-	const int32 NumDimensions = Schema->SchemaCardinality;
+	int32 ThisSchemaIndex = 0;
+	check(SearchIndexBases.Num() == Schemas.Num() && Schemas.Num() > ThisSchemaIndex);
+	const UPoseSearchSchema* ThisSchema = Schemas[ThisSchemaIndex];
+	check(ThisSchema->IsValid());
+	const int32 NumDimensions = ThisSchema->SchemaCardinality;
 
-	// Compute per Channel average distances
-	VectorXd MeanDeviations(NumDimensions);
-	MeanDeviations.setConstant(1.0);
+	TArray<float> MeanDeviations;
+	MeanDeviations.Init(1.f, NumDimensions);
+	RowMajorVectorMap MeanDeviationsMap(MeanDeviations.GetData(), 1, NumDimensions);
 
-	if (NumPoses > 0)
+	const EPoseSearchDataPreprocessor DataPreprocessor = ThisSchema->DataPreprocessor;
+	if (SearchIndexBases[ThisSchemaIndex].NumPoses > 0 && (DataPreprocessor == EPoseSearchDataPreprocessor::Normalize || DataPreprocessor == EPoseSearchDataPreprocessor::NormalizeOnlyByDeviation))
 	{
-		// Map input buffer
-		auto PoseMatrixSourceMap = RowMajorMatrixMapConst(
-			SearchIndex.Values.GetData(),
-			NumPoses,		// rows
-			NumDimensions	// cols
-		);
-
-		// @todo: evaluate removing the cast to double
-
-		// Copy row major float matrix to column major double matrix
-		MatrixXd PoseMatrix = PoseMatrixSourceMap.transpose().cast<double>();
-		checkSlow(PoseMatrix.rows() == NumDimensions);
-		checkSlow(PoseMatrix.cols() == NumPoses);
+		FFeatureChannelLayoutSet FeatureChannelLayoutSet;
+		for (int32 SchemaIndex = 0; SchemaIndex < Schemas.Num(); ++SchemaIndex)
+		{
+			const UPoseSearchSchema* Schema = Schemas[SchemaIndex];
 
-		// Mean center
-		VectorXd SampleMean = PoseMatrix.rowwise().mean();
-		PoseMatrix = PoseMatrix.colwise() - SampleMean;
+			FeatureChannelLayoutSet.CurrentSchemaIndex = SchemaIndex;
+			FeatureChannelLayoutSet.CurrentSchema = Schema;
+			for (int ChannelIdx = 0; ChannelIdx != Schema->Channels.Num(); ++ChannelIdx)
+			{
+				const UPoseSearchFeatureChannel* Channel = Schema->Channels[ChannelIdx].Get();
+				Channel->PopulateChannelLayoutSet(FeatureChannelLayoutSet);
+			}
+		}
 
-		for (int ChannelIdx = 0; ChannelIdx != Schema->Channels.Num(); ++ChannelIdx)
+		for (auto Pair : FeatureChannelLayoutSet.EntriesMap)
 		{
-			const UPoseSearchFeatureChannel* Channel = Schema->Channels[ChannelIdx].Get();
-			Channel->ComputeMeanDeviations(PoseMatrix, MeanDeviations);
+			const TArray<FFeatureChannelLayoutSet::FEntry>& Entries = Pair.Value;
+			for (const FFeatureChannelLayoutSet::FEntry& Entry : Entries)
+			{
+				if (Entry.SchemaIndex == ThisSchemaIndex)
+				{
+					const float FeatureMeanDeviation = ComputeFeatureMeanDeviation(Entries, SearchIndexBases, Schemas);
+					// the associated data to all the Entries data is going to be used to calculate the deviation of Deviation[Entry.DataOffset] to Deviation[Entry.DataOffset + Entry.Cardinality]
+
+					// Fill the feature's corresponding scaling axes with the average distance
+					// Avoid scaling by zero by leaving near-zero deviations as 1.0
+					static const float MinFeatureMeanDeviation = 0.1f;
+					MeanDeviationsMap.segment(Entry.DataOffset, Entry.Cardinality).setConstant(FeatureMeanDeviation > MinFeatureMeanDeviation ? FeatureMeanDeviation : 1.f);
+				}
+			}
 		}
 	}
-
+	
 	return MeanDeviations;
 }
-
 static inline FFloatInterval GetEffectiveSamplingRange(const UAnimSequenceBase* Sequence, FFloatInterval RequestedSamplingRange)
 {
 	const bool bSampleAll = (RequestedSamplingRange.Min == 0.0f) && (RequestedSamplingRange.Max == 0.0f);
@@ -216,28 +275,17 @@ static void InitSearchIndexAssets(FPoseSearchIndexBase& SearchIndex, TConstArray
 	}
 }
 
-static void PreprocessSearchIndexWeights(FPoseSearchIndex& SearchIndex, const UPoseSearchSchema* Schema)
+static void PreprocessSearchIndexWeights(FPoseSearchIndex& SearchIndex, const UPoseSearchSchema* Schema, TConstArrayView<float> Deviation)
 {
 	const int32 NumDimensions = Schema->SchemaCardinality;
 	SearchIndex.WeightsSqrt.Init(1.f, NumDimensions);
-
 	for (int ChannelIdx = 0; ChannelIdx != Schema->Channels.Num(); ++ChannelIdx)
 	{
 		const UPoseSearchFeatureChannel* Channel = Schema->Channels[ChannelIdx].Get();
 		Channel->FillWeights(SearchIndex.WeightsSqrt);
 	}
 
-	Eigen::VectorXd ChannelsMeanDeviations = ComputeChannelsDeviations(SearchIndex, Schema);
-	TArray<float> Deviation;
-	Deviation.Init(1.f, NumDimensions);
-	for (int32 Dimension = 0; Dimension != NumDimensions; ++Dimension)
-	{
-		const float ChannelsMeanDeviation = ChannelsMeanDeviations[Dimension];
-		Deviation[Dimension] = ChannelsMeanDeviation;
-	}
-
 	EPoseSearchDataPreprocessor DataPreprocessor = Schema->DataPreprocessor;
-
 	if (DataPreprocessor == EPoseSearchDataPreprocessor::Normalize)
 	{
 		// normalizing user weights: the idea behind this step is to be able to compare poses from databases using different schemas
@@ -245,7 +293,7 @@ static void PreprocessSearchIndexWeights(FPoseSearchIndex& SearchIndex, const UP
 		const float WeightsSum = MapWeights.sum();
 		if (!FMath::IsNearlyZero(WeightsSum))
 		{
-			MapWeights *= (1.0f / WeightsSum);
+			MapWeights *= (1.f / WeightsSum);
 		}
 	}
 
@@ -259,14 +307,10 @@ static void PreprocessSearchIndexWeights(FPoseSearchIndex& SearchIndex, const UP
 	{
 		for (int32 Dimension = 0; Dimension != NumDimensions; ++Dimension)
 		{
-			// the idea here is to premultiply the weights by the inverse of the variance (proportional to the square of the deviation) to have a "weighted Mahalanobis" distance
+			// the idea here is to pre-multiply the weights by the inverse of the variance (proportional to the square of the deviation) to have a "weighted Mahalanobis" distance
 			SearchIndex.WeightsSqrt[Dimension] /= Deviation[Dimension];
 		}
 	}
-
-#if WITH_EDITORONLY_DATA
-	SearchIndex.Deviation = Deviation;
-#endif // WITH_EDITORONLY_DATA
 }
 
 // it calculates Mean, PCAValues, and PCAProjectionMatrix
@@ -927,6 +971,8 @@ FTransform FAssetIndexer::GetTransformAndCacheResults(float SampleTime, float Or
 		float CurrentTime = Sample.ClipTime;
 		float PreviousTime = CurrentTime - SamplingContext->FiniteDelta;
 
+		check(Sample.Clip->IsLoopable() || PreviousTime <= Sample.Clip->GetPlayLength());
+
 		FDeltaTimeRecord DeltaTimeRecord;
 		DeltaTimeRecord.Set(PreviousTime, CurrentTime - PreviousTime);
 		FAnimExtractContext ExtractionCtx(static_cast<double>(CurrentTime), true, DeltaTimeRecord, Sample.Clip->IsLoopable());
@@ -1245,10 +1291,8 @@ void FPoseSearchDatabaseAsyncCacheTask::StartNewRequestIfNeeded(FCriticalSection
 	Owner.Cancel();
 
 	// composing the key
-	const FDerivedDataKeyBuilder KeyBuilder(Database.Get());
-
-	// Stores a BLAKE3-160 hash, taken from the first 20 bytes of a BLAKE3-256 hash
-	const FIoHash NewDerivedDataKey = KeyBuilder.Finalize();
+	const FKeyBuilder KeyBuilder(Database.Get(), true);
+	const FIoHash NewDerivedDataKey(KeyBuilder.Finalize());
 	const bool bHasKeyChanged = NewDerivedDataKey != DerivedDataKey;
 	if (bHasKeyChanged)
 	{
@@ -1391,68 +1435,115 @@ void FPoseSearchDatabaseAsyncCacheTask::OnGetComplete(UE::DerivedData::FCacheGet
 			{
 				COOK_STAT(auto Timer = UsageStats.TimeSyncWork());
 
-				// early out for invalid indexing conditions
-				if (!Database->Schema || !Database->Schema->IsValid() || Database->Schema->SchemaCardinality <= 0)
+				// collecting all the databases that need to be built to gather their FPoseSearchIndexBase
+				TArray<TWeakObjectPtr<const UPoseSearchDatabase>> IndexBaseDatabases;
+				IndexBaseDatabases.Add(Database); // the first one is always this Database
+				if (Database->NormalizationSet)
 				{
-					UE_LOG(LogPoseSearch, Log, TEXT("%s - %s BuildIndex Failed"), *LexToString(FullIndexKey.Hash), *Database->GetName());
-					SearchIndex.Reset();
-					return;
+					for (auto OtherDatabase : Database->NormalizationSet->Databases)
+					{
+						if (OtherDatabase)
+						{
+							IndexBaseDatabases.AddUnique(OtherDatabase);
+						}
+					}
 				}
 
-				if (Owner.IsCanceled())
+				// @todo: DDC or parallelize this code
+				TArray<FPoseSearchIndexBase> SearchIndexBases;
+				TArray<const UPoseSearchSchema*> Schemas;
+				SearchIndexBases.AddDefaulted(IndexBaseDatabases.Num());
+				Schemas.AddDefaulted(IndexBaseDatabases.Num());
+				for (int32 IndexBaseIdx = 0; IndexBaseIdx < IndexBaseDatabases.Num(); ++IndexBaseIdx)
 				{
-					UE_LOG(LogPoseSearch, Log, TEXT("%s - %s BuildIndex Cancelled"), *LexToString(FullIndexKey.Hash), *Database->GetName());
-					SearchIndex.Reset();
-					return;
-				}
+					auto IndexBaseDatabase = IndexBaseDatabases[IndexBaseIdx];
+					FPoseSearchIndexBase& SearchIndexBase = SearchIndexBases[IndexBaseIdx];
+					Schemas[IndexBaseIdx] = IndexBaseDatabase->Schema;
 
-				// Building all the related FPoseSearchBaseIndex first
-				FPoseSearchIndexBase& SearchIndexBase = SearchIndex;
-				InitSearchIndexAssets(SearchIndexBase, Database->Sequences, Database->BlendSpaces, Database->ExcludeFromDatabaseParameters);
+					// early out for invalid indexing conditions
+					if (!IndexBaseDatabase->Schema || !IndexBaseDatabase->Schema->IsValid() || IndexBaseDatabase->Schema->SchemaCardinality <= 0)
+					{
+						if (IndexBaseDatabase == Database)
+						{
+							UE_LOG(LogPoseSearch, Log, TEXT("%s - %s BuildIndex Failed"), *LexToString(FullIndexKey.Hash), *Database->GetName());
+						}
+						else
+						{
+							UE_LOG(LogPoseSearch, Log, TEXT("%s - %s BuildIndex Failed because of dependent database fail '%s'"), *LexToString(FullIndexKey.Hash), *Database->GetName(), *IndexBaseDatabase->GetName());
+						}
+						SearchIndex.Reset();
+						return;
+					}
 
-				if (Owner.IsCanceled())
-				{
-					UE_LOG(LogPoseSearch, Log, TEXT("%s - %s BuildIndex Cancelled"), *LexToString(FullIndexKey.Hash), *Database->GetName());
-					SearchIndex.Reset();
-					return;
-				}
+					if (Owner.IsCanceled())
+					{
+						UE_LOG(LogPoseSearch, Log, TEXT("%s - %s BuildIndex Cancelled"), *LexToString(FullIndexKey.Hash), *Database->GetName());
+						SearchIndex.Reset();
+						return;
+					}
 
-				FDatabaseIndexingContext DbIndexingContext;
-				DbIndexingContext.SearchIndexBase = &SearchIndexBase;
-				DbIndexingContext.Prepare(Database->Schema, Database->ExtrapolationParameters, Database->Sequences, Database->BlendSpaces);
+					// Building all the related FPoseSearchBaseIndex first
+					InitSearchIndexAssets(SearchIndexBase, IndexBaseDatabase->Sequences, IndexBaseDatabase->BlendSpaces, IndexBaseDatabase->ExcludeFromDatabaseParameters);
 
-				if (Owner.IsCanceled())
-				{
-					UE_LOG(LogPoseSearch, Log, TEXT("%s - %s BuildIndex Cancelled"), *LexToString(FullIndexKey.Hash), *Database->GetName());
-					SearchIndex.Reset();
-					return;
-				}
+					if (Owner.IsCanceled())
+					{
+						UE_LOG(LogPoseSearch, Log, TEXT("%s - %s BuildIndex Cancelled"), *LexToString(FullIndexKey.Hash), *Database->GetName());
+						SearchIndex.Reset();
+						return;
+					}
 
-				const bool bSuccess = DbIndexingContext.IndexAssets();
-				if (!bSuccess)
-				{
-					UE_LOG(LogPoseSearch, Log, TEXT("%s - %s BuildIndex Failed"), *LexToString(FullIndexKey.Hash), *Database->GetName());
-					SearchIndex.Reset();
-					return;
-				}
+					FDatabaseIndexingContext DbIndexingContext;
+					DbIndexingContext.SearchIndexBase = &SearchIndexBase;
+					DbIndexingContext.Prepare(IndexBaseDatabase->Schema, IndexBaseDatabase->ExtrapolationParameters, IndexBaseDatabase->Sequences, IndexBaseDatabase->BlendSpaces);
 
-				if (Owner.IsCanceled())
-				{
-					UE_LOG(LogPoseSearch, Log, TEXT("%s - %s BuildIndex Cancelled"), *LexToString(FullIndexKey.Hash), *Database->GetName());
-					SearchIndex.Reset();
-					return;
-				}
+					if (Owner.IsCanceled())
+					{
+						UE_LOG(LogPoseSearch, Log, TEXT("%s - %s BuildIndex Cancelled"), *LexToString(FullIndexKey.Hash), *Database->GetName());
+						SearchIndex.Reset();
+						return;
+					}
 
-				DbIndexingContext.JoinIndex();
-				if (Owner.IsCanceled())
-				{
-					UE_LOG(LogPoseSearch, Log, TEXT("%s - %s BuildIndex Cancelled"), *LexToString(FullIndexKey.Hash), *Database->GetName());
-					SearchIndex.Reset();
-					return;
+					const bool bSuccess = DbIndexingContext.IndexAssets();
+					if (!bSuccess)
+					{
+						if (IndexBaseDatabase == Database)
+						{
+							UE_LOG(LogPoseSearch, Log, TEXT("%s - %s BuildIndex Failed"), *LexToString(FullIndexKey.Hash), *Database->GetName());
+						}
+						else
+						{
+							UE_LOG(LogPoseSearch, Log, TEXT("%s - %s BuildIndex Failed because of dependent database fail '%s'"), *LexToString(FullIndexKey.Hash), *Database->GetName(), *IndexBaseDatabase->GetName());
+						}
+						SearchIndex.Reset();
+						return;
+					}
+
+					if (Owner.IsCanceled())
+					{
+						UE_LOG(LogPoseSearch, Log, TEXT("%s - %s BuildIndex Cancelled"), *LexToString(FullIndexKey.Hash), *Database->GetName());
+						SearchIndex.Reset();
+						return;
+					}
+
+					DbIndexingContext.JoinIndex();
+					if (Owner.IsCanceled())
+					{
+						UE_LOG(LogPoseSearch, Log, TEXT("%s - %s BuildIndex Cancelled"), *LexToString(FullIndexKey.Hash), *Database->GetName());
+						SearchIndex.Reset();
+						return;
+					}
 				}
 
+				static_cast<FPoseSearchIndexBase&>(SearchIndex) = SearchIndexBases[0];
+				
+				TArray<float> Deviation = ComputeChannelsDeviations(SearchIndexBases, Schemas);
+
+				#if WITH_EDITORONLY_DATA
+				SearchIndex.Deviation = Deviation;
+				#endif // WITH_EDITORONLY_DATA
+
 				// Building FPoseSearchIndex
-				PreprocessSearchIndexWeights(SearchIndex, Database->Schema);
+				PreprocessSearchIndexWeights(SearchIndex, Database->Schema, Deviation);
 				if (Owner.IsCanceled())
 				{
 					UE_LOG(LogPoseSearch, Log, TEXT("%s - %s BuildIndex Cancelled"), *LexToString(FullIndexKey.Hash), *Database->GetName());
diff --git a/Engine/Plugins/Experimental/Animation/PoseSearch/Source/Runtime/Private/PoseSearchEigenHelper.h b/Engine/Plugins/Experimental/Animation/PoseSearch/Source/Runtime/Private/PoseSearchEigenHelper.h
--- a/Engine/Plugins/Experimental/Animation/PoseSearch/Source/Runtime/Private/PoseSearchEigenHelper.h
+++ b/Engine/Plugins/Experimental/Animation/PoseSearch/Source/Runtime/Private/PoseSearchEigenHelper.h
@@ -64,15 +64,10 @@ namespace UE::PoseSearch
 	using ColMajorMatrixMap = Eigen::Map<ColMajorMatrix, Eigen::ColMajor>;
 	using ColMajorMatrixMapConst = Eigen::Map<const ColMajorMatrix, Eigen::ColMajor>;
 
-	struct ColMajorVectord : public Eigen::VectorXd
-	{
-		ColMajorVectord(int32 Size) : Eigen::VectorXd(Size) {}
-	};
-	struct ColMajorMatrixd : public Eigen::MatrixXd {};
-
 #if UE_POSE_SEARCH_EIGEN_DEBUG
 	template<typename EigenDenseBaseDerivedType>
-	FString EigenMatrixToString(const Eigen::DenseBase<EigenDenseBaseDerivedType>& Matrix){
+	FString EigenMatrixToString(const Eigen::DenseBase<EigenDenseBaseDerivedType>& Matrix)
+	{
 		std::stringstream StringStream;
 		StringStream << Matrix;
 		return StringStream.str().c_str();
diff --git a/Engine/Plugins/Experimental/Animation/PoseSearch/Source/Runtime/Private/PoseSearchFeatureChannels.cpp b/Engine/Plugins/Experimental/Animation/PoseSearch/Source/Runtime/Private/PoseSearchFeatureChannels.cpp
--- a/Engine/Plugins/Experimental/Animation/PoseSearch/Source/Runtime/Private/PoseSearchFeatureChannels.cpp
+++ b/Engine/Plugins/Experimental/Animation/PoseSearch/Source/Runtime/Private/PoseSearchFeatureChannels.cpp
@@ -8,6 +8,7 @@
 #include "Animation/MirrorDataTable.h"
 #include "DrawDebugHelpers.h"
 #include "PoseSearchEigenHelper.h"
+#include "PoseSearch/PoseSearchDerivedDataKey.h"
 #include "UObject/ObjectSaveContext.h"
 
 #include UE_INLINE_GENERATED_CPP_BY_NAME(PoseSearchFeatureChannels)
@@ -728,49 +729,6 @@ void UPoseSearchFeatureChannel_Pose::IndexAsset(UE::PoseSearch::IAssetIndexer& I
 	}
 }
 
-void UPoseSearchFeatureChannel_Pose::ComputeMeanDeviations(const Eigen::MatrixXd& CenteredPoseMatrix, Eigen::VectorXd& MeanDeviations) const
-{
-	using namespace UE::PoseSearch;
-
-	int32 DataOffset = ChannelDataOffset;
-
-	const int32 NumBones = SampledBones.Num();
-	for (int32 ChannelBoneIdx = 0; ChannelBoneIdx != NumBones; ++ChannelBoneIdx)
-	{
-		const FPoseSearchBone& SampledBone = SampledBones[ChannelBoneIdx];
-		if (EnumHasAnyFlags(SampledBone.Flags, EPoseSearchBoneFlags::Position))
-		{
-			for (int32 SubsampleIdx = 0; SubsampleIdx != SampleTimes.Num(); ++SubsampleIdx)
-			{
-				FFeatureVectorHelper::ComputeMeanDeviations(GetMinimumMeanDeviation(), CenteredPoseMatrix, MeanDeviations, DataOffset, FFeatureVectorHelper::EncodeVectorCardinality);
-			}
-		}
-		if (EnumHasAnyFlags(SampledBone.Flags, EPoseSearchBoneFlags::Rotation))
-		{
-			for (int32 SubsampleIdx = 0; SubsampleIdx != SampleTimes.Num(); ++SubsampleIdx)
-			{
-				FFeatureVectorHelper::ComputeMeanDeviations(GetMinimumMeanDeviation(), CenteredPoseMatrix, MeanDeviations, DataOffset, FFeatureVectorHelper::EncodeQuatCardinality);
-			}
-		}
-		if (EnumHasAnyFlags(SampledBone.Flags, EPoseSearchBoneFlags::Velocity))
-		{
-			for (int32 SubsampleIdx = 0; SubsampleIdx != SampleTimes.Num(); ++SubsampleIdx)
-			{
-				FFeatureVectorHelper::ComputeMeanDeviations(GetMinimumMeanDeviation(), CenteredPoseMatrix, MeanDeviations, DataOffset, FFeatureVectorHelper::EncodeVectorCardinality);
-			}
-		}
-		if (EnumHasAnyFlags(SampledBone.Flags, EPoseSearchBoneFlags::Phase))
-		{
-			for (int32 SubsampleIdx = 0; SubsampleIdx != SampleTimes.Num(); ++SubsampleIdx)
-			{
-				FFeatureVectorHelper::ComputeMeanDeviations(GetMinimumMeanDeviation(), CenteredPoseMatrix, MeanDeviations, DataOffset, FFeatureVectorHelper::EncodeVector2DCardinality);
-			}
-		}
-	}
-
-	check(DataOffset == ChannelDataOffset + ChannelCardinality);
-}
-
 void UPoseSearchFeatureChannel_Pose::AddPoseFeatures(UE::PoseSearch::IAssetIndexer& Indexer, int32 SampleIdx, TArrayView<float> FeatureVector, const TArray<TArray<FVector2D>>& Phases) const
 {
 	// This function samples the instantaneous pose at time t as well as the pose's velocity and acceleration at time t.
@@ -1184,6 +1142,59 @@ void UPoseSearchFeatureChannel_Pose::DebugDraw(const UE::PoseSearch::FDebugDrawP
 }
 
 #if WITH_EDITOR
+void UPoseSearchFeatureChannel_Pose::PopulateChannelLayoutSet(UE::PoseSearch::FFeatureChannelLayoutSet& FeatureChannelLayoutSet) const
+{
+	using namespace UE::PoseSearch;
+	int32 DataOffset = ChannelDataOffset;
+
+	auto Add = [&FeatureChannelLayoutSet, &DataOffset](const FPoseSearchBone& SampledBone, EPoseSearchBoneFlags BoneFlag, float SampleTime, const char* Label, int32 Cardinality)
+	{
+		FString SkeletonName = FeatureChannelLayoutSet.CurrentSchema->Skeleton->GetName();
+		FString BoneName = SampledBone.Reference.BoneName.ToString();
+
+		UE::PoseSearch::FKeyBuilder KeyBuilder;
+		KeyBuilder << SkeletonName << BoneName << BoneFlag << SampleTime;
+		FeatureChannelLayoutSet.Add(FString::Format(TEXT("{0} {1} {2}"), { BoneName, Label, SampleTime }), KeyBuilder.Finalize(), DataOffset, Cardinality);
+
+		DataOffset += Cardinality;
+	};
+
+	for (int32 ChannelBoneIdx = 0; ChannelBoneIdx != SampledBones.Num(); ++ChannelBoneIdx)
+	{
+		const FPoseSearchBone& SampledBone = SampledBones[ChannelBoneIdx];
+		if (EnumHasAnyFlags(SampledBone.Flags, EPoseSearchBoneFlags::Position))
+		{
+			for (int32 SubsampleIdx = 0; SubsampleIdx != SampleTimes.Num(); ++SubsampleIdx)
+			{
+				Add(SampledBone, EPoseSearchBoneFlags::Position, SampleTimes[SubsampleIdx], "Pos", FFeatureVectorHelper::EncodeVectorCardinality);
+			}
+		}
+		if (EnumHasAnyFlags(SampledBone.Flags, EPoseSearchBoneFlags::Rotation))
+		{
+			for (int32 SubsampleIdx = 0; SubsampleIdx != SampleTimes.Num(); ++SubsampleIdx)
+			{
+				Add(SampledBone, EPoseSearchBoneFlags::Rotation, SampleTimes[SubsampleIdx], "Rot", FFeatureVectorHelper::EncodeQuatCardinality);
+			}
+		}
+		if (EnumHasAnyFlags(SampledBone.Flags, EPoseSearchBoneFlags::Velocity))
+		{
+			for (int32 SubsampleIdx = 0; SubsampleIdx != SampleTimes.Num(); ++SubsampleIdx)
+			{
+				Add(SampledBone, EPoseSearchBoneFlags::Velocity, SampleTimes[SubsampleIdx], "Vel", FFeatureVectorHelper::EncodeVectorCardinality);
+			}
+		}
+		if (EnumHasAnyFlags(SampledBone.Flags, EPoseSearchBoneFlags::Phase))
+		{
+			for (int32 SubsampleIdx = 0; SubsampleIdx != SampleTimes.Num(); ++SubsampleIdx)
+			{
+				Add(SampledBone, EPoseSearchBoneFlags::Phase, SampleTimes[SubsampleIdx], "Pha", FFeatureVectorHelper::EncodeVector2DCardinality);
+			}
+		}
+	}
+
+	check(DataOffset == ChannelDataOffset + ChannelCardinality);
+}
+
 void UPoseSearchFeatureChannel_Pose::ComputeCostBreakdowns(UE::PoseSearch::ICostBreakDownData& CostBreakDownData, const UPoseSearchSchema* Schema) const
 {
 	using namespace UE::PoseSearch;
@@ -1390,54 +1401,6 @@ void UPoseSearchFeatureChannel_Trajectory::IndexAsset(UE::PoseSearch::IAssetInde
 	}
 }
 
-void UPoseSearchFeatureChannel_Trajectory::ComputeMeanDeviations(const Eigen::MatrixXd& CenteredPoseMatrix, Eigen::VectorXd& MeanDeviations) const
-{
-	using namespace UE::PoseSearch;
-
-	int32 DataOffset = ChannelDataOffset;
-
-	for (const FPoseSearchTrajectorySample& Sample : Samples)
-	{
-		if (EnumHasAnyFlags(Sample.Flags, EPoseSearchTrajectoryFlags::Position))
-		{
-			FFeatureVectorHelper::ComputeMeanDeviations(GetMinimumMeanDeviation(), CenteredPoseMatrix, MeanDeviations, DataOffset, FFeatureVectorHelper::EncodeVectorCardinality);
-		}
-		if (EnumHasAnyFlags(Sample.Flags, EPoseSearchTrajectoryFlags::PositionXY))
-		{
-			FFeatureVectorHelper::ComputeMeanDeviations(GetMinimumMeanDeviation(), CenteredPoseMatrix, MeanDeviations, DataOffset, FFeatureVectorHelper::EncodeVector2DCardinality);
-		}
-
-		if (EnumHasAnyFlags(Sample.Flags, EPoseSearchTrajectoryFlags::Velocity))
-		{
-			FFeatureVectorHelper::ComputeMeanDeviations(GetMinimumMeanDeviation(), CenteredPoseMatrix, MeanDeviations, DataOffset, FFeatureVectorHelper::EncodeVectorCardinality);
-		}
-		if (EnumHasAnyFlags(Sample.Flags, EPoseSearchTrajectoryFlags::VelocityXY))
-		{
-			FFeatureVectorHelper::ComputeMeanDeviations(GetMinimumMeanDeviation(), CenteredPoseMatrix, MeanDeviations, DataOffset, FFeatureVectorHelper::EncodeVector2DCardinality);
-		}
-
-		if (EnumHasAnyFlags(Sample.Flags, EPoseSearchTrajectoryFlags::VelocityDirection))
-		{
-			FFeatureVectorHelper::SetMeanDeviations(1.f, MeanDeviations, DataOffset, FFeatureVectorHelper::EncodeVectorCardinality);
-		}
-		if (EnumHasAnyFlags(Sample.Flags, EPoseSearchTrajectoryFlags::VelocityDirectionXY))
-		{
-			FFeatureVectorHelper::SetMeanDeviations(1.f, MeanDeviations, DataOffset, FFeatureVectorHelper::EncodeVector2DCardinality);
-		}
-
-		if (EnumHasAnyFlags(Sample.Flags, EPoseSearchTrajectoryFlags::FacingDirection))
-		{
-			FFeatureVectorHelper::SetMeanDeviations(1.f, MeanDeviations, DataOffset, FFeatureVectorHelper::EncodeVectorCardinality);
-		}
-		if (EnumHasAnyFlags(Sample.Flags, EPoseSearchTrajectoryFlags::FacingDirectionXY))
-		{
-			FFeatureVectorHelper::SetMeanDeviations(1.f, MeanDeviations, DataOffset, FFeatureVectorHelper::EncodeVector2DCardinality);
-		}
-	}
-
-	check(DataOffset == ChannelDataOffset + ChannelCardinality);
-}
-
 float UPoseSearchFeatureChannel_Trajectory::GetSampleTime(const UE::PoseSearch::IAssetIndexer& Indexer, float Offset, float SampleTime, float RootDistance) const
 {
 	switch (Domain)
@@ -1541,7 +1504,7 @@ void UPoseSearchFeatureChannel_Trajectory::IndexAssetPrivate(const UE::PoseSearc
 		}
 		if (EnumHasAnyFlags(Sample.Flags, EPoseSearchTrajectoryFlags::FacingDirectionXY))
 		{
-			FFeatureVectorHelper::EncodeVector2D(FeatureVector, DataOffset, FVector2D(FacingDirection.X, FacingDirection.Y));
+			FFeatureVectorHelper::EncodeVector2D(FeatureVector, DataOffset, FVector2D(FacingDirection.X, FacingDirection.Y).GetSafeNormal());
 		}
 	}
 	check(DataOffset == ChannelDataOffset + ChannelCardinality);
@@ -1618,7 +1581,7 @@ bool UPoseSearchFeatureChannel_Trajectory::BuildQuery(UE::PoseSearch::FSearchCon
 		}
 		if (EnumHasAnyFlags(Sample.Flags, EPoseSearchTrajectoryFlags::FacingDirectionXY))
 		{
-			FFeatureVectorHelper::EncodeVector2D(InOutQuery.EditValues(), DataOffset, FVector2D(FacingDirection.X, FacingDirection.Y));
+			FFeatureVectorHelper::EncodeVector2D(InOutQuery.EditValues(), DataOffset, FVector2D(FacingDirection.X, FacingDirection.Y).GetSafeNormal());
 		}
 	}
 	check(DataOffset == ChannelDataOffset + ChannelCardinality);
@@ -2051,6 +2014,65 @@ void UPoseSearchFeatureChannel_Trajectory::DebugDraw(const UE::PoseSearch::FDebu
 }
 
 #if WITH_EDITOR
+void UPoseSearchFeatureChannel_Trajectory::PopulateChannelLayoutSet(UE::PoseSearch::FFeatureChannelLayoutSet& FeatureChannelLayoutSet) const
+{
+	using namespace UE::PoseSearch;
+
+	int32 DataOffset = ChannelDataOffset;
+
+	auto Add = [&FeatureChannelLayoutSet, &DataOffset](EPoseSearchTrajectoryFlags SampleFlag, float Offset, const char* Label, int32 Cardinality)
+	{
+		FString SkeletonName = FeatureChannelLayoutSet.CurrentSchema->Skeleton->GetName();
+
+		UE::PoseSearch::FKeyBuilder KeyBuilder;
+		KeyBuilder << SkeletonName << SampleFlag << Offset;
+		FeatureChannelLayoutSet.Add(FString::Format(TEXT("Traj {0} {1}"), { Label, Offset }), KeyBuilder.Finalize(), DataOffset, Cardinality);
+
+		DataOffset += Cardinality;
+	};
+
+	for (const FPoseSearchTrajectorySample& Sample : Samples)
+	{
+		if (EnumHasAnyFlags(Sample.Flags, EPoseSearchTrajectoryFlags::Position))
+		{
+			Add(EPoseSearchTrajectoryFlags::Position, Sample.Offset, "Pos", FFeatureVectorHelper::EncodeVectorCardinality);
+		}
+		if (EnumHasAnyFlags(Sample.Flags, EPoseSearchTrajectoryFlags::PositionXY))
+		{
+			Add(EPoseSearchTrajectoryFlags::PositionXY, Sample.Offset, "PosXY", FFeatureVectorHelper::EncodeVector2DCardinality);
+		}
+
+		if (EnumHasAnyFlags(Sample.Flags, EPoseSearchTrajectoryFlags::Velocity))
+		{
+			Add(EPoseSearchTrajectoryFlags::Velocity, Sample.Offset, "Vel", FFeatureVectorHelper::EncodeVectorCardinality);
+		}
+		if (EnumHasAnyFlags(Sample.Flags, EPoseSearchTrajectoryFlags::VelocityXY))
+		{
+			Add(EPoseSearchTrajectoryFlags::VelocityXY, Sample.Offset, "VelXY", FFeatureVectorHelper::EncodeVector2DCardinality);
+		}
+
+		if (EnumHasAnyFlags(Sample.Flags, EPoseSearchTrajectoryFlags::VelocityDirection))
+		{
+			Add(EPoseSearchTrajectoryFlags::VelocityDirection, Sample.Offset, "VelDir", FFeatureVectorHelper::EncodeVectorCardinality);
+		}
+		if (EnumHasAnyFlags(Sample.Flags, EPoseSearchTrajectoryFlags::VelocityDirectionXY))
+		{
+			Add(EPoseSearchTrajectoryFlags::VelocityDirectionXY, Sample.Offset, "VelDirXY", FFeatureVectorHelper::EncodeVector2DCardinality);
+		}
+
+		if (EnumHasAnyFlags(Sample.Flags, EPoseSearchTrajectoryFlags::FacingDirection))
+		{
+			Add(EPoseSearchTrajectoryFlags::FacingDirection, Sample.Offset, "Fac", FFeatureVectorHelper::EncodeVectorCardinality);
+		}
+		if (EnumHasAnyFlags(Sample.Flags, EPoseSearchTrajectoryFlags::FacingDirectionXY))
+		{
+			Add(EPoseSearchTrajectoryFlags::FacingDirectionXY, Sample.Offset, "FacXY", FFeatureVectorHelper::EncodeVector2DCardinality);
+		}
+	}
+
+	check(DataOffset == ChannelDataOffset + ChannelCardinality);
+}
+
 void UPoseSearchFeatureChannel_Trajectory::ComputeCostBreakdowns(UE::PoseSearch::ICostBreakDownData& CostBreakDownData, const UPoseSearchSchema* Schema) const
 {
 	using namespace UE::PoseSearch;
diff --git a/Engine/Plugins/Experimental/Animation/PoseSearch/Source/Runtime/Public/PoseSearch/PoseSearch.h b/Engine/Plugins/Experimental/Animation/PoseSearch/Source/Runtime/Public/PoseSearch/PoseSearch.h
--- a/Engine/Plugins/Experimental/Animation/PoseSearch/Source/Runtime/Public/PoseSearch/PoseSearch.h
+++ b/Engine/Plugins/Experimental/Animation/PoseSearch/Source/Runtime/Public/PoseSearch/PoseSearch.h
@@ -70,15 +70,6 @@ namespace UE::DerivedData
 } // namespace UE::DerivedData
 #endif
 
-// eigen forward declaration
-namespace Eigen
-{
-	template<typename _Scalar, int _Rows, int _Cols, int _Options, int _MaxRows, int _MaxCols>
-	class Matrix;
-	using MatrixXd = Matrix<double, -1, -1, 0, -1, -1>;
-	using VectorXd = Matrix<double, -1, 1, 0, -1, 1>;
-}
-
 //////////////////////////////////////////////////////////////////////////
 // Constants
 
@@ -380,6 +371,37 @@ private:
 	TArrayView<float> FeatureVectorTable;
 };
 
+#if WITH_EDITOR
+// data structure collecting the internal layout representation of UPoseSearchFeatureChannel,
+// so we can aggregate data from different FPoseSearchIndex and calculate mean deviation with a homogeneous data set (ComputeChannelsDeviations
+struct FFeatureChannelLayoutSet
+{
+	// data structure holding DataOffset and Cardinality to find the data of the DebugName (channel breakdown / layout) in the related SearchIndexBases[SchemaIndex]
+	struct FEntry
+	{
+		FString DebugName; // for easier debugging
+		int32 SchemaIndex = -1; // index of the associated Schemas / SearchIndexBases used as input of the algorithm
+		int32 DataOffset = -1; // data offset from the base of SearchIndexBases[SchemaIndex].Values.GetData() from where the data associated to this Item starts
+		int32 Cardinality = -1; // data cardinality
+	};
+
+	// FIoHash is the hash associated to the channel data breakdown (e.g.: it could be a single SampledBones at a specific SampleTimes for a UPoseSearchFeatureChannel_Pose)
+	TMap<FIoHash, TArray<FEntry>> EntriesMap;
+	int32 CurrentSchemaIndex = -1;
+	TWeakObjectPtr<const UPoseSearchSchema> CurrentSchema;
+
+	void Add(FString DebugName, FIoHash Id, int32 DataOffset, int32 Cardinality)
+	{
+		check(DataOffset >= 0 && Cardinality >= 0 && CurrentSchemaIndex >= 0);
+		TArray<FEntry>& Entries = EntriesMap.FindOrAdd(Id);
+
+		// making sure all the FEntry associated with the same Id have the same Cardinality
+		check(Entries.IsEmpty() || Entries[0].Cardinality == Cardinality);
+		Entries.Add({ DebugName, CurrentSchemaIndex, DataOffset, Cardinality });
+	}
+};
+#endif // WITH_EDITOR
+
 class POSESEARCH_API IAssetIndexer
 {
 public:
@@ -450,7 +472,6 @@ public:
 
 //////////////////////////////////////////////////////////////////////////
 // Feature channels interface
-
 UCLASS(Abstract, BlueprintType, EditInlineNew)
 class POSESEARCH_API UPoseSearchFeatureChannel : public UObject, public IBoneReferenceSkeletonProvider, public IPoseFilter
 {
@@ -472,9 +493,6 @@ public:
 	// Called at database build time to populate pose vectors with this channel's data
 	virtual void IndexAsset(UE::PoseSearch::IAssetIndexer& Indexer, UE::PoseSearch::FAssetIndexingOutput& IndexingOutput) const PURE_VIRTUAL(UPoseSearchFeatureChannel::IndexAsset, );
 
-	// Called at database build time to calculate normalization values
-	virtual void ComputeMeanDeviations(const Eigen::MatrixXd& CenteredPoseMatrix, Eigen::VectorXd& MeanDeviations) const;
-
 	// Called at runtime to add this channel's data to the query pose vector
 	virtual bool BuildQuery(UE::PoseSearch::FSearchContext& SearchContext, FPoseSearchFeatureVectorBuilder& InOutQuery) const PURE_VIRTUAL(UPoseSearchFeatureChannel::BuildQuery, return false;);
 
@@ -482,12 +500,10 @@ public:
 	virtual void DebugDraw(const UE::PoseSearch::FDebugDrawParams& DrawParams, TConstArrayView<float> PoseVector) const PURE_VIRTUAL(UPoseSearchFeatureChannel::DebugDraw, );
 
 #if WITH_EDITOR
+	virtual void PopulateChannelLayoutSet(UE::PoseSearch::FFeatureChannelLayoutSet& FeatureChannelLayoutSet) const;
 	virtual void ComputeCostBreakdowns(UE::PoseSearch::ICostBreakDownData& CostBreakDownData, const UPoseSearchSchema* Schema) const;
 #endif
 
-	// Used during data normalization. If a feature has less than this amount of deviation from the mean across all poses in a database, then it will not be normalized.
-	virtual float GetMinimumMeanDeviation() const { return 0.1f; }
-
 private:
 	// IBoneReferenceSkeletonProvider interface
 	// Note this function is exclusively for FBoneReference details customization
@@ -506,9 +522,6 @@ protected:
 	int32 ChannelCardinality = INDEX_NONE;
 };
 
-//////////////////////////////////////////////////////////////////////////
-// Schema
-
 namespace UE::PoseSearch
 {
 
@@ -1067,6 +1080,17 @@ struct FSearchResult
 
 } // namespace UE::PoseSearch
 
+/** A data asset for indexing a collection of animation sequences. */
+UCLASS(BlueprintType, Category = "Animation|Pose Search", Experimental, meta = (DisplayName = "Normalization Set"))
+class POSESEARCH_API UNormalizationSetAsset : public UDataAsset
+{
+	GENERATED_BODY()
+public:
+
+	UPROPERTY(EditAnywhere, BlueprintReadOnly, Category = "NormalizationSet")
+	TArray<TObjectPtr<const UPoseSearchDatabase>> Databases; // @todo: it should be a UPoseSearchSearchableAsset, and have the UPoseSearchSearchableAsset iterate over all it's contained UPoseSearchSearchableAsset recursively (without duplicates)
+};
+
 /** A data asset for indexing a collection of animation sequences. */
 UCLASS(Abstract, BlueprintType, Experimental)
 class POSESEARCH_API UPoseSearchSearchableAsset : public UDataAsset
@@ -1115,6 +1139,9 @@ public:
 	UPROPERTY(EditAnywhere, Category = "Performance")
 	bool bSkipSearchIfPossible = true;
 
+	UPROPERTY(EditAnywhere, BlueprintReadOnly, Category = "Database")
+	TObjectPtr<const UNormalizationSetAsset> NormalizationSet;
+
 private:
 	UPROPERTY(Transient, meta = (ExcludeFromHash))
 	FPoseSearchIndex SearchIndexPrivate; // do not use it directly. use GetSearchIndex / SetSearchIndex interact with it and validate that is ok to do so
@@ -1216,18 +1243,11 @@ public:
 	static void EncodeFloat(TArrayView<float> Values, int32& DataOffset, TConstArrayView<float> PrevValues, TConstArrayView<float> CurValues, TConstArrayView<float> NextValues, float LerpValue);
 	static float DecodeFloat(TConstArrayView<float> Values, int32& DataOffset);
 
-	// populates MeanDeviations[DataOffset] ... MeanDeviations[DataOffset + Cardinality] with a single value the mean deviation calculated from a centered matrix
-	static void ComputeMeanDeviations(float MinMeanDeviation, const Eigen::MatrixXd& CenteredPoseMatrix, Eigen::VectorXd& MeanDeviations, int32& DataOffset, int32 Cardinality);
-
-	// populates MeanDeviations[DataOffset] ... MeanDeviations[DataOffset + Cardinality] with a single value
-	static void SetMeanDeviations(float Deviation, Eigen::VectorXd& MeanDeviations, int32& DataOffset, int32 Cardinality);
-
 private:
 	static FQuat DecodeQuatInternal(TConstArrayView<float> Values, int32 DataOffset);
 	static FVector DecodeVectorInternal(TConstArrayView<float> Values, int32 DataOffset);
 	static FVector2D DecodeVector2DInternal(TConstArrayView<float> Values, int32 DataOffset);
 	static float DecodeFloatInternal(TConstArrayView<float> Values, int32 DataOffset);
-
 };
 
 /**
diff --git a/Engine/Plugins/Experimental/Animation/PoseSearch/Source/Runtime/Public/PoseSearch/PoseSearchDerivedDataKey.h b/Engine/Plugins/Experimental/Animation/PoseSearch/Source/Runtime/Public/PoseSearch/PoseSearchDerivedDataKey.h
--- a/Engine/Plugins/Experimental/Animation/PoseSearch/Source/Runtime/Public/PoseSearch/PoseSearchDerivedDataKey.h
+++ b/Engine/Plugins/Experimental/Animation/PoseSearch/Source/Runtime/Public/PoseSearch/PoseSearchDerivedDataKey.h
@@ -29,8 +29,8 @@ namespace UE::PoseSearch
 #endif
 
 //////////////////////////////////////////////////////////////////////////
-// FDerivedDataKeyBuilder
-class POSESEARCH_API FDerivedDataKeyBuilder : public FArchiveUObject
+// FKeyBuilder
+class POSESEARCH_API FKeyBuilder : public FArchiveUObject
 {
 public:
 	using Super = FArchiveUObject;
@@ -40,10 +40,8 @@ public:
 	inline static const FName ExcludeFromHashName = FName(TEXT("ExcludeFromHash"));
 	inline static const FName NeverInHashName = FName(TEXT("NeverInHash"));
 
-	FDerivedDataKeyBuilder(UObject* Object)
+	FKeyBuilder()
 	{
-		check(Object);
-
 		ArIgnoreOuterRef = true;
 
 		// Set FDerivedDataKeyBuilder to be a saving archive instead of a reference collector.
@@ -51,14 +49,26 @@ public:
 		// which doesn't give a stable hash.  Serializing these to a saving archive will
 		// use a string reference instead, which is a more meaningful hash value.
 		SetIsSaving(true);
+	}
 
-		// used to invalidate the key without having to change POSESEARCHDB_DERIVEDDATA_VER all the times
-		FGuid VersionGuid = FDevSystemGuids::GetSystemGuid(FDevSystemGuids::Get().POSESEARCHDB_DERIVEDDATA_VER);
-		int32 POSESEARCHDB_DERIVEDDATA_VER_SMALL = 3;
+	FKeyBuilder(const UObject* Object, bool bUseDataVer = false)
+		: FKeyBuilder()
+	{
+		check(Object);
+
+		if (bUseDataVer)
+		{
+			// used to invalidate the key without having to change POSESEARCHDB_DERIVEDDATA_VER all the times
+			int32 POSESEARCHDB_DERIVEDDATA_VER_SMALL = 6;
+			FGuid VersionGuid = FDevSystemGuids::GetSystemGuid(FDevSystemGuids::Get().POSESEARCHDB_DERIVEDDATA_VER);
+
+			*this << VersionGuid;
+			*this << POSESEARCHDB_DERIVEDDATA_VER_SMALL;
+		}
 
-		*this << VersionGuid;
-		*this << POSESEARCHDB_DERIVEDDATA_VER_SMALL;
-		*this << Object;
+		// FKeyBuilder is a saving only archiver, and since it doesn't modify the input Object it's safe to do a const_cast 
+		UObject* NonConstObject = const_cast<UObject*>(Object);
+		*this << NonConstObject;
 	}
 
 	virtual void Seek(int64 InPos) override
@@ -196,9 +206,10 @@ public:
 	}
 	//~ End FArchive Interface
 	
-	HashDigestType Finalize() const
+	FIoHash Finalize() const
 	{
-		return Hasher.Finalize();
+		// Stores a BLAKE3-160 hash, taken from the first 20 bytes of a BLAKE3-256 hash
+		return FIoHash(Hasher.Finalize());
 	}
 
 	const TSet<const UObject*>& GetDependencies() const
diff --git a/Engine/Plugins/Experimental/Animation/PoseSearch/Source/Runtime/Public/PoseSearch/PoseSearchFeatureChannels.h b/Engine/Plugins/Experimental/Animation/PoseSearch/Source/Runtime/Public/PoseSearch/PoseSearchFeatureChannels.h
--- a/Engine/Plugins/Experimental/Animation/PoseSearch/Source/Runtime/Public/PoseSearch/PoseSearchFeatureChannels.h
+++ b/Engine/Plugins/Experimental/Animation/PoseSearch/Source/Runtime/Public/PoseSearch/PoseSearchFeatureChannels.h
@@ -179,11 +179,11 @@ public:
 	virtual void InitializeSchema(UE::PoseSearch::FSchemaInitializer& Initializer) override;
 	virtual void FillWeights(TArray<float>& Weights) const override;
 	virtual void IndexAsset(UE::PoseSearch::IAssetIndexer& Indexer, UE::PoseSearch::FAssetIndexingOutput& IndexingOutput) const override;
-	virtual void ComputeMeanDeviations(const Eigen::MatrixXd& CenteredPoseMatrix, Eigen::VectorXd& MeanDeviations) const override;
 	virtual bool BuildQuery(UE::PoseSearch::FSearchContext& SearchContext, FPoseSearchFeatureVectorBuilder& InOutQuery) const override;
 	virtual void DebugDraw(const UE::PoseSearch::FDebugDrawParams& DrawParams, TConstArrayView<float> PoseVector) const override;
 
 #if WITH_EDITOR
+	virtual void PopulateChannelLayoutSet(UE::PoseSearch::FFeatureChannelLayoutSet& FeatureChannelLayoutSet) const override;
 	virtual void ComputeCostBreakdowns(UE::PoseSearch::ICostBreakDownData& CostBreakDownData, const UPoseSearchSchema* Schema) const override;
 #endif
 
@@ -260,11 +260,11 @@ public:
 	virtual void InitializeSchema(UE::PoseSearch::FSchemaInitializer& Initializer) override;
 	virtual void FillWeights(TArray<float>& Weights) const override;
 	virtual void IndexAsset(UE::PoseSearch::IAssetIndexer& Indexer, UE::PoseSearch::FAssetIndexingOutput& IndexingOutput) const override;
-	virtual void ComputeMeanDeviations(const Eigen::MatrixXd& CenteredPoseMatrix, Eigen::VectorXd& MeanDeviations) const override;
 	virtual bool BuildQuery(UE::PoseSearch::FSearchContext& SearchContext, FPoseSearchFeatureVectorBuilder& InOutQuery) const override;
 	virtual void DebugDraw(const UE::PoseSearch::FDebugDrawParams& DrawParams, TConstArrayView<float> PoseVector) const override;
 
 #if WITH_EDITOR
+	virtual void PopulateChannelLayoutSet(UE::PoseSearch::FFeatureChannelLayoutSet& FeatureChannelLayoutSet) const override;
 	virtual void ComputeCostBreakdowns(UE::PoseSearch::ICostBreakDownData& CostBreakDownData, const UPoseSearchSchema* Schema) const override;
 #endif
 
