motion matching - fix for missing FPoseSearchIndexBase::MinCostAddend serialization

diff --git a/Engine/Plugins/Experimental/Animation/PoseSearch/Source/Runtime/Private/PoseSearch.cpp b/Engine/Plugins/Experimental/Animation/PoseSearch/Source/Runtime/Private/PoseSearch.cpp
--- a/Engine/Plugins/Experimental/Animation/PoseSearch/Source/Runtime/Private/PoseSearch.cpp
+++ b/Engine/Plugins/Experimental/Animation/PoseSearch/Source/Runtime/Private/PoseSearch.cpp
@@ -616,6 +616,8 @@ FArchive& operator<<(FArchive& Ar, FPoseSearchIndexBase& Index)
 		Ar.Serialize(&Index.Assets[0], Index.Assets.Num() * Index.Assets.GetTypeSize());
 	}
 
+	Ar << Index.MinCostAddend;
+
 	return Ar;
 }
 
diff --git a/Engine/Plugins/Experimental/Animation/PoseSearch/Source/Runtime/Public/PoseSearch/PoseSearchDerivedDataKey.h b/Engine/Plugins/Experimental/Animation/PoseSearch/Source/Runtime/Public/PoseSearch/PoseSearchDerivedDataKey.h
--- a/Engine/Plugins/Experimental/Animation/PoseSearch/Source/Runtime/Public/PoseSearch/PoseSearchDerivedDataKey.h
+++ b/Engine/Plugins/Experimental/Animation/PoseSearch/Source/Runtime/Public/PoseSearch/PoseSearchDerivedDataKey.h
@@ -59,7 +59,7 @@ public:
 		if (bUseDataVer)
 		{
 			// used to invalidate the key without having to change POSESEARCHDB_DERIVEDDATA_VER all the times
-			int32 POSESEARCHDB_DERIVEDDATA_VER_SMALL = 6;
+			int32 POSESEARCHDB_DERIVEDDATA_VER_SMALL = 8;
 			FGuid VersionGuid = FDevSystemGuids::GetSystemGuid(FDevSystemGuids::Get().POSESEARCHDB_DERIVEDDATA_VER);
 
 			*this << VersionGuid;
