motion matching - FDatabaseViewModel::UpdatePreviewActors stability improvements

diff --git a/Engine/Plugins/Experimental/Animation/PoseSearch/Source/Editor/Private/PoseSearchDatabasePreviewScene.cpp b/Engine/Plugins/Experimental/Animation/PoseSearch/Source/Editor/Private/PoseSearchDatabasePreviewScene.cpp
--- a/Engine/Plugins/Experimental/Animation/PoseSearch/Source/Editor/Private/PoseSearchDatabasePreviewScene.cpp
+++ b/Engine/Plugins/Experimental/Animation/PoseSearch/Source/Editor/Private/PoseSearchDatabasePreviewScene.cpp
@@ -59,31 +59,34 @@ namespace UE::PoseSearch
 			GetWorld()->Tick(LEVELTICK_All, InDeltaTime);
 		}
 
-		const FDatabaseViewModel* ViewModel = GetEditor()->GetViewModel();
+		FDatabaseViewModel* ViewModel = GetEditor()->GetViewModel();
 		const UPoseSearchDatabase* Database = ViewModel->GetPoseSearchDatabase();
 
 		if (ViewModel->IsPoseFeaturesDrawMode(EFeaturesDrawMode::All | EFeaturesDrawMode::Detailed) && !ViewModel->GetPreviewActors().IsEmpty() &&
 			FAsyncPoseSearchDatabasesManagement::RequestAsyncBuildIndex(Database, ERequestAsyncBuildFlag::ContinueRequest))
 		{
-			for (const FDatabasePreviewActor& PreviewActor : ViewModel->GetPreviewActors())
+			for (FDatabasePreviewActor& PreviewActor : ViewModel->GetPreviewActors())
 			{
 				if (Database->GetSearchIndex().IsValidPoseIndex(PreviewActor.CurrentPoseIndex))
 				{
-					UE::PoseSearch::FDebugDrawParams DrawParams;
-					DrawParams.RootTransform = PreviewActor.Mesh->GetComponentTransform();
-					DrawParams.Database = Database;
-					DrawParams.World = GetWorld();
-					DrawParams.DefaultLifeTime = 0.f;
-					DrawParams.PointSize = 5.f;
-					DrawParams.Mesh = PreviewActor.Mesh;
-
-					if (ViewModel->IsPoseFeaturesDrawMode(EFeaturesDrawMode::Detailed))
+					if (UDebugSkelMeshComponent* Mesh = PreviewActor.GetDebugSkelMeshComponent())
 					{
-						EnumAddFlags(DrawParams.Flags, UE::PoseSearch::EDebugDrawFlags::DrawBoneNames);
-					}
+						UE::PoseSearch::FDebugDrawParams DrawParams;
+						DrawParams.RootTransform = Mesh->GetComponentTransform();
+						DrawParams.Database = Database;
+						DrawParams.World = GetWorld();
+						DrawParams.DefaultLifeTime = 0.f;
+						DrawParams.PointSize = 5.f;
+						DrawParams.Mesh = Mesh;
 
-					EnumAddFlags(DrawParams.Flags, UE::PoseSearch::EDebugDrawFlags::DrawFast);
-					DrawFeatureVector(DrawParams, PreviewActor.CurrentPoseIndex);
+						if (ViewModel->IsPoseFeaturesDrawMode(EFeaturesDrawMode::Detailed))
+						{
+							EnumAddFlags(DrawParams.Flags, UE::PoseSearch::EDebugDrawFlags::DrawBoneNames);
+						}
+
+						EnumAddFlags(DrawParams.Flags, UE::PoseSearch::EDebugDrawFlags::DrawFast);
+						DrawFeatureVector(DrawParams, PreviewActor.CurrentPoseIndex);
+					}
 				}
 			}
 		}
diff --git a/Engine/Plugins/Experimental/Animation/PoseSearch/Source/Editor/Private/PoseSearchDatabaseViewModel.cpp b/Engine/Plugins/Experimental/Animation/PoseSearch/Source/Editor/Private/PoseSearchDatabaseViewModel.cpp
--- a/Engine/Plugins/Experimental/Animation/PoseSearch/Source/Editor/Private/PoseSearchDatabaseViewModel.cpp
+++ b/Engine/Plugins/Experimental/Animation/PoseSearch/Source/Editor/Private/PoseSearchDatabaseViewModel.cpp
@@ -25,8 +25,7 @@ namespace UE::PoseSearch
 	// FDatabasePreviewActor
 	bool FDatabasePreviewActor::IsValid() const
 	{
-		const bool bIsValid = Actor.IsValid() && Mesh.IsValid() && AnimInstance.IsValid();
-		return  bIsValid;
+		return Actor != nullptr;
 	}
 
 	void FDatabasePreviewActor::Process()
@@ -63,6 +62,24 @@ namespace UE::PoseSearch
 		return ScaledTime;
 	}
 
+	UDebugSkelMeshComponent* FDatabasePreviewActor::GetDebugSkelMeshComponent()
+	{
+		if (IsValid())
+		{
+			return Cast<UDebugSkelMeshComponent>(Actor->GetRootComponent());
+		}
+		return nullptr;
+	}
+
+	UAnimPreviewInstance* FDatabasePreviewActor::GetAnimPreviewInstance()
+	{
+		if (UDebugSkelMeshComponent* Mesh = GetDebugSkelMeshComponent())
+		{
+			return Mesh->PreviewInstance.Get();
+		}
+		return nullptr;
+	}
+
 	// FDatabaseViewModel
 	FDatabaseViewModel::FDatabaseViewModel()
 		: PoseSearchDatabase(nullptr)
@@ -150,7 +167,7 @@ namespace UE::PoseSearch
 					FDatabasePreviewActor PreviewActor = SpawnPreviewActor(IndexAssetIndex, BoneContainer);
 					if (PreviewActor.IsValid())
 					{
-						const UAnimationAsset* AnimationAsset = PreviewActor.AnimInstance->GetAnimationAsset();
+						const UAnimationAsset* AnimationAsset = PreviewActor.GetAnimPreviewInstance()->GetAnimationAsset();
 						if (AnimationAsset)
 						{
 							MaxPreviewPlayLength = FMath::Max(MaxPreviewPlayLength, AnimationAsset->GetPlayLength());
@@ -260,26 +277,30 @@ namespace UE::PoseSearch
 		PreviewActor.Actor = GetWorld()->SpawnActor<AActor>(AActor::StaticClass(), FTransform::Identity, Params);
 		PreviewActor.Actor->SetFlags(RF_Transient);
 
-		PreviewActor.Mesh = NewObject<UDebugSkelMeshComponent>(PreviewActor.Actor.Get());
-		PreviewActor.Mesh->RegisterComponentWithWorld(GetWorld());
+		UDebugSkelMeshComponent* Mesh = NewObject<UDebugSkelMeshComponent>(PreviewActor.Actor.Get());
+		Mesh->RegisterComponentWithWorld(GetWorld());
+
+		UAnimPreviewInstance* AnimInstance = NewObject<UAnimPreviewInstance>(Mesh);
 
-		PreviewActor.AnimInstance = NewObject<UAnimPreviewInstance>(PreviewActor.Mesh.Get());
-		PreviewActor.Mesh->PreviewInstance = PreviewActor.AnimInstance.Get();
-		PreviewActor.AnimInstance->InitializeAnimation();
-		PreviewActor.Mesh->SetSkeletalMesh(PoseSearchDatabase->Schema->Skeleton->GetPreviewMesh(true));
-		PreviewActor.Mesh->EnablePreview(true, PreviewAsset);
-		PreviewActor.AnimInstance->SetAnimationAsset(PreviewAsset, false, 0.0f);
-		PreviewActor.AnimInstance->SetBlendSpacePosition(IndexAsset.BlendParameters);
+		Mesh->PreviewInstance = AnimInstance;
+		AnimInstance->InitializeAnimation();
+
+		Mesh->SetSkeletalMesh(PoseSearchDatabase->Schema->Skeleton->GetPreviewMesh(true));
+		Mesh->EnablePreview(true, PreviewAsset);
+		
+		AnimInstance->SetAnimationAsset(PreviewAsset, false, 0.0f);
+		AnimInstance->SetBlendSpacePosition(IndexAsset.BlendParameters);
+		
 		if (IndexAsset.bMirrored && PoseSearchDatabase->Schema)
 		{
-			PreviewActor.AnimInstance->SetMirrorDataTable(PoseSearchDatabase->Schema->MirrorDataTable);
+			AnimInstance->SetMirrorDataTable(PoseSearchDatabase->Schema->MirrorDataTable);
 		}
 
-		PreviewActor.AnimInstance->PlayAnim(false, 0.0f);
+		AnimInstance->PlayAnim(false, 0.0f);
 
 		if (!PreviewActor.Actor->GetRootComponent())
 		{
-			PreviewActor.Actor->SetRootComponent(PreviewActor.Mesh.Get());
+			PreviewActor.Actor->SetRootComponent(Mesh);
 		}
 
 		UE_LOG(LogPoseSearchEditor, Log, TEXT("Spawned preview Actor: %s"), *GetNameSafe(PreviewActor.Actor.Get()));
@@ -322,29 +343,32 @@ namespace UE::PoseSearch
 		{
 			if (PreviewActor.IndexAssetIndex < SearchIndex.Assets.Num())
 			{
-				if (const UAnimationAsset* PreviewAsset = PreviewActor.AnimInstance->GetAnimationAsset())
+				if (const UAnimationAsset* PreviewAsset = PreviewActor.GetAnimPreviewInstance()->GetAnimationAsset())
 				{
 					if (const IAssetSampler* Sampler = PreviewActor.GetSampler())
 					{
 						float CurrentTime = 0.f;
 						FAnimationRuntime::AdvanceTime(false, PlayTime, CurrentTime, PreviewAsset->GetPlayLength());
 
-						const float CurrentScaledTime = PreviewActor.GetScaledTime(CurrentTime);
-						const float OldScaledTime = PreviewActor.AnimInstance->GetCurrentTime();
-						if (!FMath::IsNearlyEqual(CurrentScaledTime, OldScaledTime, UE_KINDA_SMALL_NUMBER))
+						FTransform RootMotion = Sampler->ExtractRootTransform(CurrentTime);
+						if (PreviewActor.GetAnimPreviewInstance()->GetMirrorDataTable())
 						{
-							FTransform RootMotion = Sampler->ExtractRootTransform(CurrentTime);
-							if (PreviewActor.AnimInstance->GetMirrorDataTable())
-							{
-								RootMotion = MirrorRootMotion(RootMotion, PreviewActor.AnimInstance->GetMirrorDataTable());
-							}
-							PreviewActor.Actor->SetActorTransform(RootMotion);
+							RootMotion = MirrorRootMotion(RootMotion, PreviewActor.GetAnimPreviewInstance()->GetMirrorDataTable());
+						}
+						PreviewActor.Actor->SetActorTransform(RootMotion);
 
-							PreviewActor.AnimInstance->SetPosition(CurrentScaledTime);
+						const float CurrentScaledTime = PreviewActor.GetScaledTime(CurrentTime);
+						PreviewActor.GetAnimPreviewInstance()->SetPosition(CurrentScaledTime);
 
-							const FPoseSearchIndexAsset& SearchIndexAsset = SearchIndex.Assets[PreviewActor.IndexAssetIndex];
-							PreviewActor.CurrentPoseIndex = PoseSearchDatabase->GetPoseIndexFromTime(CurrentTime, SearchIndexAsset);
+						float PlayRate = 1.f;
+						if (FMath::IsNearlyZero(DeltaTimeMultiplier) || CurrentScaledTime <= UE_SMALL_NUMBER || CurrentScaledTime >= (PreviewAsset->GetPlayLength() - UE_SMALL_NUMBER))
+						{
+							PlayRate = 0.f;
 						}
+						PreviewActor.GetAnimPreviewInstance()->SetPlayRate(PlayRate);
+
+						const FPoseSearchIndexAsset& SearchIndexAsset = SearchIndex.Assets[PreviewActor.IndexAssetIndex];
+						PreviewActor.CurrentPoseIndex = PoseSearchDatabase->GetPoseIndexFromTime(CurrentTime, SearchIndexAsset);
 					}
 				}
 			}
diff --git a/Engine/Plugins/Experimental/Animation/PoseSearch/Source/Editor/Public/PoseSearchDatabaseViewModel.h b/Engine/Plugins/Experimental/Animation/PoseSearch/Source/Editor/Public/PoseSearchDatabaseViewModel.h
--- a/Engine/Plugins/Experimental/Animation/PoseSearch/Source/Editor/Public/PoseSearchDatabaseViewModel.h
+++ b/Engine/Plugins/Experimental/Animation/PoseSearch/Source/Editor/Public/PoseSearchDatabaseViewModel.h
@@ -42,8 +42,6 @@ namespace UE::PoseSearch
 	{
 	public:
 		TWeakObjectPtr<AActor> Actor = nullptr;
-		TWeakObjectPtr<UDebugSkelMeshComponent> Mesh = nullptr;
-		TWeakObjectPtr<UAnimPreviewInstance> AnimInstance = nullptr;
 		int32 IndexAssetIndex = INDEX_NONE;
 		int32 CurrentPoseIndex = INDEX_NONE;
 
@@ -55,6 +53,8 @@ namespace UE::PoseSearch
 		void Process();
 		const IAssetSampler* GetSampler() const;
 		float GetScaledTime(float Time) const;
+		UDebugSkelMeshComponent* GetDebugSkelMeshComponent();
+		UAnimPreviewInstance* GetAnimPreviewInstance();
 	};
 
 	class FDatabaseViewModel : public TSharedFromThis<FDatabaseViewModel>, public FGCObject
