motion matching - fix for PoseSearchDatabase preview jitteriness

diff --git a/Engine/Plugins/Experimental/Animation/PoseSearch/Source/Editor/Private/PoseSearchDatabaseViewModel.cpp b/Engine/Plugins/Experimental/Animation/PoseSearch/Source/Editor/Private/PoseSearchDatabaseViewModel.cpp
--- a/Engine/Plugins/Experimental/Animation/PoseSearch/Source/Editor/Private/PoseSearchDatabaseViewModel.cpp
+++ b/Engine/Plugins/Experimental/Animation/PoseSearch/Source/Editor/Private/PoseSearchDatabaseViewModel.cpp
@@ -303,6 +303,8 @@ namespace UE::PoseSearch
 			PreviewActor.Actor->SetRootComponent(Mesh);
 		}
 
+		AnimInstance->SetPlayRate(0.f);
+
 		UE_LOG(LogPoseSearchEditor, Log, TEXT("Spawned preview Actor: %s"), *GetNameSafe(PreviewActor.Actor.Get()));
 		return PreviewActor;
 	}
@@ -359,13 +361,7 @@ namespace UE::PoseSearch
 
 						const float CurrentScaledTime = PreviewActor.GetScaledTime(CurrentTime);
 						PreviewActor.GetAnimPreviewInstance()->SetPosition(CurrentScaledTime);
-
-						float PlayRate = 1.f;
-						if (FMath::IsNearlyZero(DeltaTimeMultiplier) || CurrentScaledTime <= UE_SMALL_NUMBER || CurrentScaledTime >= (PreviewAsset->GetPlayLength() - UE_SMALL_NUMBER))
-						{
-							PlayRate = 0.f;
-						}
-						PreviewActor.GetAnimPreviewInstance()->SetPlayRate(PlayRate);
+						PreviewActor.GetAnimPreviewInstance()->SetPlayRate(0.f);
 
 						const FPoseSearchIndexAsset& SearchIndexAsset = SearchIndex.Assets[PreviewActor.IndexAssetIndex];
 						PreviewActor.CurrentPoseIndex = PoseSearchDatabase->GetPoseIndexFromTime(CurrentTime, SearchIndexAsset);
