Motion Matching: Added anim composite support & refactored search database assets into polymorphic structs

* Anim composite support added including:
* Rewind debugger support for anim composites
* Backward compatibility of the pose search database using PostLoad() data conversion. (Also deprecated the former sequence and blend space arrays).
* Increase version number for search database asset's DDC.
* Using FInstancedStruct to introduce polymorphism to the pose search animation asset structs.
* Added an anim composite database asset
* Simplified many conditions by calling polymorphic member functions.
* Utilizing USequenceBase wherever possible (as the anim sequence as well as the anim composite are both inherited from USequenceBase)
* Some more code cleaning by sharing duplicated code and moving it into functions.
* Turned the FSequenceSampler to work on UAnimSequenceBase so that it can be shared across sequences and anim composites

diff --git a/Engine/Plugins/Experimental/Animation/PoseSearch/PoseSearch.uplugin b/Engine/Plugins/Experimental/Animation/PoseSearch/PoseSearch.uplugin
--- a/Engine/Plugins/Experimental/Animation/PoseSearch/PoseSearch.uplugin
+++ b/Engine/Plugins/Experimental/Animation/PoseSearch/PoseSearch.uplugin
@@ -42,6 +42,10 @@
 		{
 			"Name": "MotionTrajectory",
 			"Enabled": true
+		},
+		{
+			"Name": "StructUtils",
+			"Enabled": true
 		}
 	]
 }
\ No newline at end of file
diff --git a/Engine/Plugins/Experimental/Animation/PoseSearch/Source/Editor/PoseSearchEditor.Build.cs b/Engine/Plugins/Experimental/Animation/PoseSearch/Source/Editor/PoseSearchEditor.Build.cs
--- a/Engine/Plugins/Experimental/Animation/PoseSearch/Source/Editor/PoseSearchEditor.Build.cs
+++ b/Engine/Plugins/Experimental/Animation/PoseSearch/Source/Editor/PoseSearchEditor.Build.cs
@@ -20,6 +20,7 @@ public class PoseSearchEditor : ModuleRules
 				"Engine",
 				"PoseSearch",
 				"GameplayTags",
+				"StructUtils",
 				
 				// Trace-related dependencies
 				"TraceLog",
diff --git a/Engine/Plugins/Experimental/Animation/PoseSearch/Source/Editor/Private/PoseSearchCustomization.cpp b/Engine/Plugins/Experimental/Animation/PoseSearch/Source/Editor/Private/PoseSearchCustomization.cpp
--- a/Engine/Plugins/Experimental/Animation/PoseSearch/Source/Editor/Private/PoseSearchCustomization.cpp
+++ b/Engine/Plugins/Experimental/Animation/PoseSearch/Source/Editor/Private/PoseSearchCustomization.cpp
@@ -92,8 +92,7 @@ TSharedRef<IDetailCustomization> FPoseSearchDatabaseDetails::MakeInstance()
 void FPoseSearchDatabaseDetails::CustomizeDetails(IDetailLayoutBuilder& DetailBuilder)
 {
 	TArray<TSharedPtr<IPropertyHandle>> HiddenHandles;
-	HiddenHandles.Add(DetailBuilder.GetProperty(GET_MEMBER_NAME_CHECKED(UPoseSearchDatabase, Sequences)));
-	HiddenHandles.Add(DetailBuilder.GetProperty(GET_MEMBER_NAME_CHECKED(UPoseSearchDatabase, BlendSpaces)));
+	HiddenHandles.Add(DetailBuilder.GetProperty(GET_MEMBER_NAME_CHECKED(UPoseSearchDatabase, AnimationAssets)));
 
 	for (TSharedPtr<IPropertyHandle> PropertyHandle : HiddenHandles)
 	{
diff --git a/Engine/Plugins/Experimental/Animation/PoseSearch/Source/Editor/Private/PoseSearchDatabaseEditor.cpp b/Engine/Plugins/Experimental/Animation/PoseSearch/Source/Editor/Private/PoseSearchDatabaseEditor.cpp
--- a/Engine/Plugins/Experimental/Animation/PoseSearch/Source/Editor/Private/PoseSearchDatabaseEditor.cpp
+++ b/Engine/Plugins/Experimental/Animation/PoseSearch/Source/Editor/Private/PoseSearchDatabaseEditor.cpp
@@ -25,7 +25,6 @@
 namespace UE::PoseSearch
 {
 	const FName PoseSearchDatabaseEditorAppName = FName(TEXT("PoseSearchDatabaseEditorApp"));
-	constexpr int32 SelectionDetailsCount = 2; // sequence, blendspace
 	constexpr double ViewRangeSlack = 0.25;
 
 	// Tab identifiers
@@ -169,27 +168,6 @@ namespace UE::PoseSearch
 			EditingAssetWidget = PropertyModule.CreateDetailView(DatabaseDetailsArgs);
 			EditingAssetWidget->SetObject(DatabaseAsset);
 		}
-		
-		// Selection details widgets
-		{
-			FDetailsViewArgs SelectionDetailsArgs;
-			{
-				SelectionDetailsArgs.bHideSelectionTip = true;
-				SelectionDetailsArgs.NotifyHook = this;
-				SelectionDetailsArgs.bShowScrollBar = false;
-			}
-			
-			SelectionWidgets.Reset(SelectionDetailsCount);
-			for (int32 i = 0; i < SelectionDetailsCount; ++i)
-			{
-				TSharedPtr<IDetailsView>& SelectionWidget = 
-					SelectionWidgets.Add_GetRef(PropertyModule.CreateDetailView(SelectionDetailsArgs));
-				SelectionWidget->SetObject(nullptr);
-				SelectionWidget->OnFinishedChangingProperties().AddSP(
-					this,
-					&FDatabaseEditor::OnFinishedChangingSelectionProperties);
-			}
-		}
 
 		// Statistics details widgets
 		{
@@ -474,18 +452,8 @@ namespace UE::PoseSearch
 	{
 		check(Args.GetTabId() == FDatabaseEditorTabs::SelectionDetailsID);
 
-		TSharedPtr<SVerticalBox> DetailsContainer;
 		SAssignNew(DetailsContainer, SVerticalBox);
 
-		for (TSharedPtr<IDetailsView> SelectionWidget : SelectionWidgets)
-		{
-			DetailsContainer->AddSlot()
-				.AutoHeight()
-				[
-					SelectionWidget.ToSharedRef()
-				];
-		}
-
 		return SNew(SDockTab)
 			.Label(LOCTEXT("SelectionDetails_Title", "Selection Details"))
 			[
@@ -520,56 +488,110 @@ namespace UE::PoseSearch
 		const TArrayView<TSharedPtr<FDatabaseAssetTreeNode>>& SelectedItems,
 		ESelectInfo::Type SelectionType)
 	{
-		enum ESelectionAssetTypeIndex
+		// Reset selected objects on all selection widgets
+		for (auto& SelectionWidgetPair : SelectionWidgets)
 		{
-			SequenceSelectionIndex = 0,
-			BlendSpaceSelectionIndex = 1,
-			SelectionTypeCount = 2
-		};
-
-		TArray<TArray<TWeakObjectPtr<UObject>>> SelectionReflections;
-		SelectionReflections.SetNum(SelectionDetailsCount);
+			SelectionWidgetPair.Value.SelectedReflections.Reset();
+		}
 
 		if (SelectedItems.Num() > 0)
 		{
 			for (TSharedPtr<FDatabaseAssetTreeNode>& SelectedItem : SelectedItems)
 			{
-				if (SelectedItem->SourceAssetType == ESearchIndexAssetType::Sequence)
+				if (!SelectedItem.IsValid() ||
+					!GetPoseSearchDatabase()->AnimationAssets.IsValidIndex(SelectedItem->SourceAssetIdx))
+				{
+					continue;
+				}
+
+				const FInstancedStruct& DatabaseAsset = GetPoseSearchDatabase()->AnimationAssets[SelectedItem->SourceAssetIdx];
+				const UScriptStruct* ScriptStruct = DatabaseAsset.GetScriptStruct();
+				FSelectionWidget& SelectionWidget = FindOrAddSelectionWidget(ScriptStruct);
+
+				if (const FPoseSearchDatabaseSequence* DatabaseSequence = DatabaseAsset.GetPtr<FPoseSearchDatabaseSequence>())
 				{
 					UPoseSearchDatabaseSequenceReflection* NewSelectionReflection = NewObject<UPoseSearchDatabaseSequenceReflection>();
 					NewSelectionReflection->AddToRoot();
-					NewSelectionReflection->Sequence = GetPoseSearchDatabase()->Sequences[SelectedItem->SourceAssetIdx];
+					NewSelectionReflection->Sequence = *DatabaseSequence;
+					NewSelectionReflection->SetSourceLink(SelectedItem, AssetTreeWidget);
+					NewSelectionReflection->SetFlags(RF_Transactional);
+
+					SelectionWidget.SelectedReflections.Add(NewSelectionReflection);
+				}
+				else if (const FPoseSearchDatabaseAnimComposite* DatabaseAnimComposite = DatabaseAsset.GetPtr<FPoseSearchDatabaseAnimComposite>())
+				{
+					UPoseSearchDatabaseAnimCompositeReflection* NewSelectionReflection = NewObject<UPoseSearchDatabaseAnimCompositeReflection>();
+					NewSelectionReflection->AddToRoot();
+					NewSelectionReflection->AnimComposite = *DatabaseAnimComposite;
 					NewSelectionReflection->SetSourceLink(SelectedItem, AssetTreeWidget);
 					NewSelectionReflection->SetFlags(RF_Transactional);
-					SelectionReflections[SequenceSelectionIndex].Add(NewSelectionReflection);
+					
+					SelectionWidget.SelectedReflections.Add(NewSelectionReflection);
 				}
-				else if (SelectedItem->SourceAssetType == ESearchIndexAssetType::BlendSpace)
+				else if (const FPoseSearchDatabaseBlendSpace* DatabaseBlendSpace = DatabaseAsset.GetPtr<FPoseSearchDatabaseBlendSpace>())
 				{
 					UPoseSearchDatabaseBlendSpaceReflection* NewSelectionReflection = NewObject<UPoseSearchDatabaseBlendSpaceReflection>();
 					NewSelectionReflection->AddToRoot();
-					NewSelectionReflection->BlendSpace = GetPoseSearchDatabase()->BlendSpaces[SelectedItem->SourceAssetIdx];
+					NewSelectionReflection->BlendSpace = *DatabaseBlendSpace;
 					NewSelectionReflection->SetSourceLink(SelectedItem, AssetTreeWidget);
 					NewSelectionReflection->SetFlags(RF_Transactional);
-					SelectionReflections[BlendSpaceSelectionIndex].Add(NewSelectionReflection);
+					
+					SelectionWidget.SelectedReflections.Add(NewSelectionReflection);
+				}
+				else
+				{
+					checkNoEntry();
 				}
 			}
 		}
 
-		int32 SelectionWidgetIdx = 0;
-		for (int32 SelectionTypeIdx = 0; SelectionTypeIdx < SelectionTypeCount; ++SelectionTypeIdx)
+		for (auto& SelectionWidgetPair : SelectionWidgets)
 		{
-			if (SelectionReflections[SelectionTypeIdx].Num() > 0)
+			FSelectionWidget& SelectionWidget = SelectionWidgetPair.Value;
+
+			if (SelectionWidget.SelectedReflections.IsEmpty())
 			{
-				SelectionWidgets[SelectionWidgetIdx++]->SetObjects(SelectionReflections[SelectionTypeIdx], true);
+				SelectionWidget.DetailView->SetObject(nullptr, true);
+			}
+			else
+			{
+				SelectionWidget.DetailView->SetObjects(SelectionWidget.SelectedReflections);
 			}
 		}
 
-		while (SelectionWidgetIdx < SelectionDetailsCount)
+		ViewModel->SetSelectedNodes(SelectedItems);
+	}
+
+	FDatabaseEditor::FSelectionWidget& FDatabaseEditor::FindOrAddSelectionWidget(const UScriptStruct* ScriptStructType)
+	{
+		if (SelectionWidgets.Contains(ScriptStructType))
 		{
-			SelectionWidgets[SelectionWidgetIdx++]->SetObject(nullptr, true);
+			return SelectionWidgets[ScriptStructType];
 		}
 
-		ViewModel->SetSelectedNodes(SelectedItems);
+		FSelectionWidget& SelectionWidget = SelectionWidgets.Add(ScriptStructType);
+
+		FDetailsViewArgs SelectionDetailsArgs;
+		{
+			SelectionDetailsArgs.bHideSelectionTip = true;
+			SelectionDetailsArgs.NotifyHook = this;
+			SelectionDetailsArgs.bShowScrollBar = false;
+		}
+
+		FPropertyEditorModule& PropertyModule = FModuleManager::LoadModuleChecked<FPropertyEditorModule>("PropertyEditor");
+		SelectionWidget.DetailView = PropertyModule.CreateDetailView(SelectionDetailsArgs);
+		SelectionWidget.DetailView->SetObject(nullptr);
+		SelectionWidget.DetailView->OnFinishedChangingProperties().AddSP(
+			this,
+			&FDatabaseEditor::OnFinishedChangingSelectionProperties);
+
+		DetailsContainer->AddSlot()
+			.AutoHeight()
+			[
+				SelectionWidget.DetailView.ToSharedRef()
+			];
+
+		return SelectionWidget;
 	}
 
 	void FDatabaseEditor::RefreshStatisticsWidgetInformation()
diff --git a/Engine/Plugins/Experimental/Animation/PoseSearch/Source/Editor/Private/PoseSearchDatabaseEditorReflection.cpp b/Engine/Plugins/Experimental/Animation/PoseSearch/Source/Editor/Private/PoseSearchDatabaseEditorReflection.cpp
--- a/Engine/Plugins/Experimental/Animation/PoseSearch/Source/Editor/Private/PoseSearchDatabaseEditorReflection.cpp
+++ b/Engine/Plugins/Experimental/Animation/PoseSearch/Source/Editor/Private/PoseSearchDatabaseEditorReflection.cpp
@@ -16,11 +16,9 @@ void UPoseSearchDatabaseReflectionBase::SetSourceLink(
 	AssetTreeWidget = InAssetTreeWidget;
 }
 
-void UPoseSearchDatabaseSequenceReflection::PostEditChangeProperty(
-	struct FPropertyChangedEvent& PropertyChangedEvent)
+void UPoseSearchDatabaseSequenceReflection::PostEditChangeProperty(struct FPropertyChangedEvent& PropertyChangedEvent)
 {
 	Super::PostEditChangeProperty(PropertyChangedEvent);
-
 	check(WeakAssetTreeNode.Pin()->SourceAssetType == ESearchIndexAssetType::Sequence);
 
 	if (const TSharedPtr<UE::PoseSearch::FDatabaseViewModel> ViewModel = WeakAssetTreeNode.Pin()->EditorViewModel.Pin())
@@ -28,30 +26,58 @@ void UPoseSearchDatabaseSequenceReflection::PostEditChangeProperty(
 		UPoseSearchDatabase* Database = ViewModel->GetPoseSearchDatabase();
 		if (IsValid(Database))
 		{
-			Database->Sequences[WeakAssetTreeNode.Pin()->SourceAssetIdx] = Sequence;
-			Database->MarkPackageDirty();
+			const FInstancedStruct& DatabaseAsset = Database->GetAnimationAssetStruct(WeakAssetTreeNode.Pin()->SourceAssetIdx);
+			if (FPoseSearchDatabaseSequence* DatabaseSequence = DatabaseAsset.GetMutablePtr<FPoseSearchDatabaseSequence>())
+			{
+				*DatabaseSequence = Sequence;
+				Database->MarkPackageDirty();
 		
-			AssetTreeWidget->FinalizeTreeChanges(true);
+				AssetTreeWidget->FinalizeTreeChanges(true);
+			}
 		}
 	}
 }
 
-void UPoseSearchDatabaseBlendSpaceReflection::PostEditChangeProperty(
-	struct FPropertyChangedEvent& PropertyChangedEvent)
+void UPoseSearchDatabaseBlendSpaceReflection::PostEditChangeProperty(struct FPropertyChangedEvent& PropertyChangedEvent)
 {
 	Super::PostEditChangeProperty(PropertyChangedEvent);
-
 	check(WeakAssetTreeNode.Pin()->SourceAssetType == ESearchIndexAssetType::BlendSpace);
 
 	if (const TSharedPtr<UE::PoseSearch::FDatabaseViewModel> ViewModel = WeakAssetTreeNode.Pin()->EditorViewModel.Pin())
 	{
-		UPoseSearchDatabase* Database = WeakAssetTreeNode.Pin()->EditorViewModel.Pin()->GetPoseSearchDatabase();
+		UPoseSearchDatabase* Database = ViewModel->GetPoseSearchDatabase();
 		if (IsValid(Database))
 		{
-			Database->BlendSpaces[WeakAssetTreeNode.Pin()->SourceAssetIdx] = BlendSpace;
-			Database->MarkPackageDirty();
-		
-			AssetTreeWidget->FinalizeTreeChanges(true);
+			const FInstancedStruct& DatabaseAsset = Database->GetAnimationAssetStruct(WeakAssetTreeNode.Pin()->SourceAssetIdx);
+			if (FPoseSearchDatabaseBlendSpace* DatabaseBlendSpace = DatabaseAsset.GetMutablePtr<FPoseSearchDatabaseBlendSpace>())
+			{
+				*DatabaseBlendSpace = BlendSpace;
+				Database->MarkPackageDirty();
+
+				AssetTreeWidget->FinalizeTreeChanges(true);
+			}
+		}
+	}
+}
+
+void UPoseSearchDatabaseAnimCompositeReflection::PostEditChangeProperty(struct FPropertyChangedEvent& PropertyChangedEvent)
+{
+	Super::PostEditChangeProperty(PropertyChangedEvent);
+	check(WeakAssetTreeNode.Pin()->SourceAssetType == ESearchIndexAssetType::AnimComposite);
+
+	if (const TSharedPtr<UE::PoseSearch::FDatabaseViewModel> ViewModel = WeakAssetTreeNode.Pin()->EditorViewModel.Pin())
+	{
+		UPoseSearchDatabase* Database = ViewModel->GetPoseSearchDatabase();
+		if (IsValid(Database))
+		{
+			const FInstancedStruct& DatabaseAsset = Database->GetAnimationAssetStruct(WeakAssetTreeNode.Pin()->SourceAssetIdx);
+			if (FPoseSearchDatabaseAnimComposite* DatabaseAnimComposite = DatabaseAsset.GetMutablePtr<FPoseSearchDatabaseAnimComposite>())
+			{
+				*DatabaseAnimComposite = AnimComposite;
+				Database->MarkPackageDirty();
+
+				AssetTreeWidget->FinalizeTreeChanges(true);
+			}
 		}
 	}
 }
@@ -67,7 +93,7 @@ void UPoseSearchDatabaseStatistics::Initialize(const UPoseSearchDatabase* PoseSe
 		const FPoseSearchIndex& SearchIndex = PoseSearchDatabase->GetSearchIndex();
 		// General information
 	
-		AnimationSequences = PoseSearchDatabase->Sequences.Num();
+		AnimationSequences = PoseSearchDatabase->AnimationAssets.Num();
 			
 		const int32 SampleRate = FMath::Max(1, PoseSearchDatabase->Schema->SampleRate);
 		TotalAnimationPosesInFrames = SearchIndex.NumPoses;
diff --git a/Engine/Plugins/Experimental/Animation/PoseSearch/Source/Editor/Private/PoseSearchDatabaseViewModel.cpp b/Engine/Plugins/Experimental/Animation/PoseSearch/Source/Editor/Private/PoseSearchDatabaseViewModel.cpp
--- a/Engine/Plugins/Experimental/Animation/PoseSearch/Source/Editor/Private/PoseSearchDatabaseViewModel.cpp
+++ b/Engine/Plugins/Experimental/Animation/PoseSearch/Source/Editor/Private/PoseSearchDatabaseViewModel.cpp
@@ -7,9 +7,10 @@
 #include "PoseSearch/PoseSearchDerivedData.h"
 #include "Modules/ModuleManager.h"
 #include "AnimPreviewInstance.h"
+#include "Animation/AnimComposite.h"
+#include "Animation/AnimSequence.h"
 #include "Animation/DebugSkelMeshComponent.h"
 #include "Animation/MirrorDataTable.h"
-#include "Animation/AnimSequence.h"
 #include "Animation/BlendSpace.h"
 #include "EngineUtils.h"
 #include "PropertyEditorModule.h"
@@ -30,13 +31,28 @@ namespace UE::PoseSearch
 
 	void FDatabasePreviewActor::Process()
 	{
-		if (Type == ESearchIndexAssetType::Sequence)
+		switch (Type)
+		{
+		case ESearchIndexAssetType::Sequence:
 		{
 			SequenceSampler.Process();
+			break;
 		}
-		else if (Type == ESearchIndexAssetType::BlendSpace)
+		case ESearchIndexAssetType::BlendSpace:
 		{
 			BlendSpaceSampler.Process();
+			break;
+		}
+		case ESearchIndexAssetType::AnimComposite:
+		{
+			SequenceSampler.Process();
+			break;
+		}
+		default:
+		{
+			checkNoEntry();
+			break;
+		}
 		}
 	}
 
@@ -48,7 +64,11 @@ namespace UE::PoseSearch
 			return &SequenceSampler;
 		case ESearchIndexAssetType::BlendSpace:
 			return &BlendSpaceSampler;
+		case ESearchIndexAssetType::AnimComposite:
+			return &SequenceSampler;
 		}
+
+		checkNoEntry();
 		return nullptr;
 	}
 
@@ -128,7 +148,7 @@ namespace UE::PoseSearch
 		TSet<int32> AssociatedBlendSpacesAssetIndices;
 		for (const TSharedPtr<FDatabaseAssetTreeNode>& SelectedNode : SelectedNodes)
 		{
-			if (SelectedNode->SourceAssetType == ESearchIndexAssetType::Sequence)
+			if (SelectedNode->SourceAssetType == ESearchIndexAssetType::Sequence || SelectedNode->SourceAssetType == ESearchIndexAssetType::AnimComposite)
 			{
 				AssociatedSequencesAssetIndices.Add(SelectedNode->SourceAssetIdx);
 			}
@@ -152,7 +172,7 @@ namespace UE::PoseSearch
 				if (AnimationPreviewMode == EAnimationPreviewMode::OriginalAndMirrored || !IndexAsset.bMirrored)
 				{
 					bool bIsAssociatedToSelection = false;
-					if (IndexAsset.Type == ESearchIndexAssetType::Sequence)
+					if (IndexAsset.Type == ESearchIndexAssetType::Sequence || IndexAsset.Type == ESearchIndexAssetType::AnimComposite)
 					{
 						bSpawn = AssociatedSequencesAssetIndices.Contains(IndexAsset.SourceAssetIdx);
 					}
@@ -234,39 +254,45 @@ namespace UE::PoseSearch
 		check(PoseSearchDatabase);
 		const FPoseSearchIndex& SearchIndex = PoseSearchDatabase->GetSearchIndex();
 		const FPoseSearchIndexAsset& IndexAsset = SearchIndex.Assets[IndexAssetIndex];
-		UAnimationAsset* PreviewAsset = nullptr;
-		if (IndexAsset.Type == ESearchIndexAssetType::Sequence)
+
+		const FPoseSearchDatabaseAnimationAssetBase* DatabaseAnimationAsset = PoseSearchDatabase->GetAnimationAssetBase(IndexAsset.SourceAssetIdx);
+		UAnimationAsset* PreviewAsset = DatabaseAnimationAsset->GetAnimationAsset();
+		if (!PreviewAsset)
 		{
-			FPoseSearchDatabaseSequence DatabaseSequence = PoseSearchDatabase->Sequences[IndexAsset.SourceAssetIdx];
-			PreviewAsset = DatabaseSequence.Sequence;
-			if (!PreviewAsset)
-			{
-				return PreviewActor;
-			}
+			return PreviewActor;
+		}
 
-			FSequenceSampler::FInput Input;
+		const FInstancedStruct& DatabaseAsset = PoseSearchDatabase->GetAnimationAssetStruct(IndexAsset);
+		if (const FPoseSearchDatabaseSequence* DatabaseSequence = DatabaseAsset.GetPtr<FPoseSearchDatabaseSequence>())
+		{
+			FSequenceBaseSampler::FInput Input;
 			Input.ExtrapolationParameters = PoseSearchDatabase->ExtrapolationParameters;
-			Input.Sequence = DatabaseSequence.Sequence;
+			Input.SequenceBase = DatabaseSequence->Sequence;
 
 			PreviewActor.SequenceSampler.Init(Input);
 		}
-		else if (IndexAsset.Type == ESearchIndexAssetType::BlendSpace)
+		else if (const FPoseSearchDatabaseAnimComposite* DatabaseAnimComposite = DatabaseAsset.GetPtr<FPoseSearchDatabaseAnimComposite>())
 		{
-			FPoseSearchDatabaseBlendSpace DatabaseBlendSpace = PoseSearchDatabase->BlendSpaces[IndexAsset.SourceAssetIdx];
-			PreviewAsset = DatabaseBlendSpace.BlendSpace;
-			if (!PreviewAsset)
-			{
-				return PreviewActor;
-			}
+			FSequenceBaseSampler::FInput Input;
+			Input.ExtrapolationParameters = PoseSearchDatabase->ExtrapolationParameters;
+			Input.SequenceBase = DatabaseAnimComposite->AnimComposite;
 
+			PreviewActor.SequenceSampler.Init(Input);
+		}
+		else if (const FPoseSearchDatabaseBlendSpace* DatabaseBlendSpace = DatabaseAsset.GetPtr<FPoseSearchDatabaseBlendSpace>())
+		{
 			FBlendSpaceSampler::FInput Input;
 			Input.BoneContainer = BoneContainer;
 			Input.ExtrapolationParameters = PoseSearchDatabase->ExtrapolationParameters;
-			Input.BlendSpace = DatabaseBlendSpace.BlendSpace;
+			Input.BlendSpace = DatabaseBlendSpace->BlendSpace;
 			Input.BlendParameters = IndexAsset.BlendParameters;
 
 			PreviewActor.BlendSpaceSampler.Init(Input);
 		}
+		else
+		{
+			checkNoEntry();
+		}
 
 		PreviewActor.Type = IndexAsset.Type;
 		PreviewActor.IndexAssetIndex = IndexAssetIndex;
@@ -343,31 +369,35 @@ namespace UE::PoseSearch
 		TSharedPtr<FDatabasePreviewScene> PreviewScene = PreviewScenePtr.Pin();
 		for (FDatabasePreviewActor& PreviewActor : GetPreviewActors())
 		{
-			if (PreviewActor.IndexAssetIndex < SearchIndex.Assets.Num())
+			if (!PreviewActor.GetAnimPreviewInstance() ||
+				PreviewActor.IndexAssetIndex >= SearchIndex.Assets.Num())
 			{
-				if (const UAnimationAsset* PreviewAsset = PreviewActor.GetAnimPreviewInstance()->GetAnimationAsset())
-				{
-					if (const IAssetSampler* Sampler = PreviewActor.GetSampler())
-					{
-						float CurrentTime = 0.f;
-						FAnimationRuntime::AdvanceTime(false, PlayTime, CurrentTime, PreviewAsset->GetPlayLength());
+				continue;
+			}
 
-						FTransform RootMotion = Sampler->ExtractRootTransform(CurrentTime);
-						if (PreviewActor.GetAnimPreviewInstance()->GetMirrorDataTable())
-						{
-							RootMotion = MirrorRootMotion(RootMotion, PreviewActor.GetAnimPreviewInstance()->GetMirrorDataTable());
-						}
-						PreviewActor.Actor->SetActorTransform(RootMotion);
+			const UAnimationAsset* PreviewAsset = PreviewActor.GetAnimPreviewInstance()->GetAnimationAsset();
+			const IAssetSampler* Sampler = PreviewActor.GetSampler();
+			if (!PreviewAsset || !Sampler)
+			{
+				continue;
+			}
 
-						const float CurrentScaledTime = PreviewActor.GetScaledTime(CurrentTime);
-						PreviewActor.GetAnimPreviewInstance()->SetPosition(CurrentScaledTime);
-						PreviewActor.GetAnimPreviewInstance()->SetPlayRate(0.f);
+			float CurrentTime = 0.f;
+			FAnimationRuntime::AdvanceTime(false, PlayTime, CurrentTime, PreviewAsset->GetPlayLength());
 
-						const FPoseSearchIndexAsset& SearchIndexAsset = SearchIndex.Assets[PreviewActor.IndexAssetIndex];
-						PreviewActor.CurrentPoseIndex = PoseSearchDatabase->GetPoseIndexFromTime(CurrentTime, SearchIndexAsset);
-					}
-				}
+			FTransform RootMotion = Sampler->ExtractRootTransform(CurrentTime);
+			if (PreviewActor.GetAnimPreviewInstance()->GetMirrorDataTable())
+			{
+				RootMotion = MirrorRootMotion(RootMotion, PreviewActor.GetAnimPreviewInstance()->GetMirrorDataTable());
 			}
+			PreviewActor.Actor->SetActorTransform(RootMotion);
+
+			const float CurrentScaledTime = PreviewActor.GetScaledTime(CurrentTime);
+			PreviewActor.GetAnimPreviewInstance()->SetPosition(CurrentScaledTime);
+			PreviewActor.GetAnimPreviewInstance()->SetPlayRate(0.f);
+
+			const FPoseSearchIndexAsset& SearchIndexAsset = SearchIndex.Assets[PreviewActor.IndexAssetIndex];
+			PreviewActor.CurrentPoseIndex = PoseSearchDatabase->GetPoseIndexFromTime(CurrentTime, SearchIndexAsset);
 		}
 	}
 
@@ -432,44 +462,46 @@ namespace UE::PoseSearch
 
 	void FDatabaseViewModel::AddSequenceToDatabase(UAnimSequence* AnimSequence)
 	{
-		FPoseSearchDatabaseSequence& NewDbSequence = PoseSearchDatabase->Sequences.AddDefaulted_GetRef();
-		NewDbSequence.Sequence = AnimSequence;
+		FPoseSearchDatabaseSequence NewAsset;
+		NewAsset.Sequence = AnimSequence;
+		PoseSearchDatabase->AnimationAssets.Add(FInstancedStruct::Make(NewAsset));
 	}
 
 	void FDatabaseViewModel::AddBlendSpaceToDatabase(UBlendSpace* BlendSpace)
 	{
-		FPoseSearchDatabaseBlendSpace& NewDbBlendSpace = PoseSearchDatabase->BlendSpaces.AddDefaulted_GetRef();
-		NewDbBlendSpace.BlendSpace = BlendSpace;
+		FPoseSearchDatabaseBlendSpace NewAsset;
+		NewAsset.BlendSpace = BlendSpace;
+		PoseSearchDatabase->AnimationAssets.Add(FInstancedStruct::Make(NewAsset));
 	}
 
-	void FDatabaseViewModel::DeleteSequenceFromDatabase(int32 SequenceIdx)
+	void FDatabaseViewModel::AddAnimCompositeToDatabase(UAnimComposite* AnimComposite)
 	{
-		PoseSearchDatabase->Sequences.RemoveAt(SequenceIdx);
+		FPoseSearchDatabaseAnimComposite NewAsset;
+		NewAsset.AnimComposite = AnimComposite;
+		PoseSearchDatabase->AnimationAssets.Add(FInstancedStruct::Make(NewAsset));
 	}
 
-	void FDatabaseViewModel::DeleteBlendSpaceFromDatabase(int32 BlendSpaceIdx)
+	void FDatabaseViewModel::DeleteFromDatabase(int32 AnimationAssetIndex)
 	{
-		PoseSearchDatabase->BlendSpaces.RemoveAt(BlendSpaceIdx);
+		PoseSearchDatabase->AnimationAssets.RemoveAt(AnimationAssetIndex);
 	}
 
-	void FDatabaseViewModel::SetSelectedSequenceEnabled(int32 SequenceIndex, bool bEnabled)
-	{
-		PoseSearchDatabase->Sequences[SequenceIndex].bEnabled = bEnabled;
-	}	
-	
-	void FDatabaseViewModel::SetSelectedBlendSpaceEnabled(int32 BlendSpaceIndex, bool bEnabled)
+	void FDatabaseViewModel::SetIsEnabled(int32 AnimationAssetIndex, bool bEnabled)
 	{
-		PoseSearchDatabase->BlendSpaces[BlendSpaceIndex].bEnabled = bEnabled;
+		if (FPoseSearchDatabaseAnimationAssetBase* DatabaseAnimationAsset = PoseSearchDatabase->GetAnimationAssetBase(AnimationAssetIndex))
+		{
+			DatabaseAnimationAsset->SetIsEnabled(bEnabled);
+		}
 	}
 
-	bool FDatabaseViewModel::IsSelectedSequenceEnabled(int32 SequenceIndex) const
+	bool FDatabaseViewModel::IsEnabled(int32 AnimationAssetIndex) const
 	{
-		return PoseSearchDatabase->Sequences[SequenceIndex].bEnabled;
-	}
+		if (const FPoseSearchDatabaseAnimationAssetBase* DatabaseAnimationAsset = PoseSearchDatabase->GetAnimationAssetBase(AnimationAssetIndex))
+		{
+			return DatabaseAnimationAsset->IsEnabled();
+		}
 
-	bool FDatabaseViewModel::IsSelectedBlendSpaceEnabled(int32 BlendSpaceIndex) const
-	{
-		return PoseSearchDatabase->BlendSpaces[BlendSpaceIndex].bEnabled;
+		return false;
 	}
 
 	void FDatabaseViewModel::SetSelectedNodes(const TArrayView<TSharedPtr<FDatabaseAssetTreeNode>>& InSelectedNodes)
diff --git a/Engine/Plugins/Experimental/Animation/PoseSearch/Source/Editor/Private/PoseSearchDebugger.cpp b/Engine/Plugins/Experimental/Animation/PoseSearch/Source/Editor/Private/PoseSearchDebugger.cpp
--- a/Engine/Plugins/Experimental/Animation/PoseSearch/Source/Editor/Private/PoseSearchDebugger.cpp
+++ b/Engine/Plugins/Experimental/Animation/PoseSearch/Source/Editor/Private/PoseSearchDebugger.cpp
@@ -140,7 +140,8 @@ void UPoseSearchMeshComponent::UpdatePose(const FUpdateContext& UpdateContext)
 	UE::Anim::FStackAttributeContainer Attributes;
 	FAnimationPoseData PoseData(CompactPose, Curve, Attributes);
 
-	if (UpdateContext.Type == ESearchIndexAssetType::Sequence)
+	if (UpdateContext.Type == ESearchIndexAssetType::Sequence ||
+		UpdateContext.Type == ESearchIndexAssetType::AnimComposite)
 	{
 		float AdvancedTime = UpdateContext.StartTime;
 
@@ -148,12 +149,12 @@ void UPoseSearchMeshComponent::UpdatePose(const FUpdateContext& UpdateContext)
 			UpdateContext.bLoop,
 			UpdateContext.Time - UpdateContext.StartTime,
 			AdvancedTime,
-			UpdateContext.Sequence->GetPlayLength());
+			UpdateContext.SequenceBase->GetPlayLength());
 
 		FAnimExtractContext ExtractionCtx;
 		ExtractionCtx.CurrentTime = AdvancedTime;
 
-		UpdateContext.Sequence->GetAnimationPose(PoseData, ExtractionCtx);
+		UpdateContext.SequenceBase->GetAnimationPose(PoseData, ExtractionCtx);
 	}
 	else if (UpdateContext.Type == ESearchIndexAssetType::BlendSpace)
 	{
@@ -439,7 +440,8 @@ namespace DebuggerDatabaseColumns
 			FNumberFormattingOptions PercentageFormattingOptions = FNumberFormattingOptions()
 				.SetMaximumFractionalDigits(2);
 
-			if (Row->AssetType == ESearchIndexAssetType::Sequence)
+			if (Row->AssetType == ESearchIndexAssetType::Sequence ||
+				Row->AssetType == ESearchIndexAssetType::AnimComposite)
 			{
 				if (GetDefault<UPersonaOptions>()->bTimelineDisplayPercentage)
 				{
@@ -927,27 +929,35 @@ void SDebuggerDatabaseView::Update(const FTraceMotionMatchingStateMessage& State
 					const bool bIsQueryVectorValid = DbEntry.QueryVector.Num() == PoseValues.Num();
 					CompareFeatureVectors(PoseValues, bIsQueryVectorValid ? DbEntry.QueryVector : PoseValues, SearchIndex.WeightsSqrt, Row->CostVector);
 
-					if (SearchIndexAsset->Type == ESearchIndexAssetType::Sequence)
+					const FInstancedStruct& DatabaseAssetStruct = Database->GetAnimationAssetStruct(*SearchIndexAsset);
+					if (const FPoseSearchDatabaseAnimationAssetBase* DatabaseAsset = DatabaseAssetStruct.GetPtr<FPoseSearchDatabaseAnimationAssetBase>())
 					{
-						const FPoseSearchDatabaseSequence& DbSequence = Database->GetSequenceSourceAsset(*SearchIndexAsset);
-						Row->AssetType = ESearchIndexAssetType::Sequence;
-						Row->AssetName = DbSequence.Sequence->GetName();
-						Row->AssetPath = DbSequence.Sequence->GetPathName();
-						Row->AnimFrame = DbSequence.Sequence->GetFrameAtTime(Time);
-						Row->AnimPercentage = Time / DbSequence.Sequence->GetPlayLength();
-						Row->bLooping = DbSequence.Sequence->bLoop;
+						Row->AssetName = DatabaseAsset->GetName();
+						Row->AssetPath = DatabaseAsset->GetAnimationAsset() ? DatabaseAsset->GetAnimationAsset()->GetPathName() : "";
+						Row->AssetType = DatabaseAsset->GetSearchIndexType();
+						Row->bLooping = DatabaseAsset->IsLooping();
 						Row->BlendParameters = FVector::Zero();
-					}
-					else if (SearchIndexAsset->Type == ESearchIndexAssetType::BlendSpace)
-					{
-						const FPoseSearchDatabaseBlendSpace& DbBlendSpace = Database->GetBlendSpaceSourceAsset(*SearchIndexAsset);
-						Row->AssetType = ESearchIndexAssetType::BlendSpace;
-						Row->AssetName = DbBlendSpace.BlendSpace->GetName();
-						Row->AssetPath = DbBlendSpace.BlendSpace->GetPathName();
-						Row->AnimFrame = 0; // There is no frame index associated with a blendspace
+						Row->AnimFrame = 0;
 						Row->AnimPercentage = 0.0f;
-						Row->bLooping = DbBlendSpace.BlendSpace->bLoop;
-						Row->BlendParameters = SearchIndexAsset->BlendParameters;
+
+						if (const FPoseSearchDatabaseSequence* DatabaseSequence = DatabaseAssetStruct.GetPtr<FPoseSearchDatabaseSequence>())
+						{
+							Row->AnimFrame = DatabaseSequence->Sequence->GetFrameAtTime(Time);
+							Row->AnimPercentage = Time / DatabaseSequence->Sequence->GetPlayLength();
+						}
+						else if (const FPoseSearchDatabaseAnimComposite* DatabaseAnimComposite = DatabaseAssetStruct.GetPtr<FPoseSearchDatabaseAnimComposite>())
+						{
+							Row->AnimFrame = DatabaseAnimComposite->AnimComposite->GetFrameAtTime(Time);
+							Row->AnimPercentage = Time / DatabaseAnimComposite->AnimComposite->GetPlayLength();
+						}
+						else if (const FPoseSearchDatabaseBlendSpace* DatabaseBlendSpace = DatabaseAssetStruct.GetPtr<FPoseSearchDatabaseBlendSpace>())
+						{
+							Row->BlendParameters = SearchIndexAsset->BlendParameters;
+						}
+						else
+						{
+							checkNoEntry();
+						}
 					}
 				}
 			}
@@ -2421,6 +2431,41 @@ void FDebuggerViewModel::OnUpdateNodeSelection(int32 InNodeId)
 	}
 }
 
+void FDebuggerViewModel::UpdatePoseSearchContext(UPoseSearchMeshComponent::FUpdateContext& InOutContext, const FSkeleton& Skeleton) const
+{
+	const FInstancedStruct* DatabaseAssetStruct = Skeleton.GetAnimationAsset();
+	if (DatabaseAssetStruct)
+	{
+		const FPoseSearchDatabaseAnimationAssetBase* DatabaseAsset = DatabaseAssetStruct->GetPtr<FPoseSearchDatabaseAnimationAssetBase>();
+		if (DatabaseAsset)
+		{
+			InOutContext.Type = DatabaseAsset->GetSearchIndexType();
+			InOutContext.StartTime = Skeletons[SelectedPose].Time;
+			InOutContext.Time = Skeletons[SelectedPose].Time;
+			InOutContext.bMirrored = Skeletons[SelectedPose].bMirrored;
+			InOutContext.bLoop = DatabaseAsset->IsLooping();
+		}
+
+		if (const FPoseSearchDatabaseSequence* DatabaseSequence = DatabaseAssetStruct->GetPtr<FPoseSearchDatabaseSequence>())
+		{
+			InOutContext.SequenceBase = DatabaseSequence->Sequence;
+		}
+		else if (const FPoseSearchDatabaseAnimComposite* DatabaseAnimComposite = DatabaseAssetStruct->GetPtr<FPoseSearchDatabaseAnimComposite>())
+		{
+			InOutContext.SequenceBase = DatabaseAnimComposite->AnimComposite;
+		}
+		else if (const FPoseSearchDatabaseBlendSpace* DatabaseBlendSpace = DatabaseAssetStruct->GetPtr<FPoseSearchDatabaseBlendSpace>())
+		{
+			InOutContext.BlendSpace = DatabaseBlendSpace->BlendSpace;
+			InOutContext.BlendParameters = Skeletons[SelectedPose].BlendParameters;
+		}
+		else
+		{
+			checkNoEntry();
+		}
+	}
+}
+
 void FDebuggerViewModel::OnDraw(FSkeletonDrawParams& DrawParams)
 {
 	const UPoseSearchDatabase* CurrentDatabase = GetCurrentDatabase();
@@ -2461,37 +2506,7 @@ void FDebuggerViewModel::OnDraw(FSkeletonDrawParams& DrawParams)
 		UPoseSearchMeshComponent* Component = Skeletons[SelectedPose].Component.Get();
 		if (Component && Component->RequiredBones.IsValid())
 		{
-			if (Skeletons[SelectedPose].Type == ESearchIndexAssetType::Sequence)
-			{
-				const FPoseSearchDatabaseSequence* DatabaseSequence = Skeletons[SelectedPose].GetAnimSequence();
-				if (DatabaseSequence)
-				{
-					UpdateContext.Type = ESearchIndexAssetType::Sequence;
-					UpdateContext.Sequence = DatabaseSequence->Sequence;
-					UpdateContext.StartTime = Skeletons[SelectedPose].Time;
-					UpdateContext.Time = Skeletons[SelectedPose].Time;
-					UpdateContext.bMirrored = Skeletons[SelectedPose].bMirrored;
-					UpdateContext.bLoop = DatabaseSequence->Sequence->bLoop;
-				}
-			}
-			else if (Skeletons[SelectedPose].Type == ESearchIndexAssetType::BlendSpace)
-			{
-				const FPoseSearchDatabaseBlendSpace* DatabaseBlendSpace = Skeletons[SelectedPose].GetBlendSpace();
-				if (DatabaseBlendSpace)
-				{
-					UpdateContext.Type = ESearchIndexAssetType::BlendSpace;
-					UpdateContext.BlendSpace = DatabaseBlendSpace->BlendSpace;
-					UpdateContext.StartTime = Skeletons[SelectedPose].Time;
-					UpdateContext.Time = Skeletons[SelectedPose].Time;
-					UpdateContext.bMirrored = Skeletons[SelectedPose].bMirrored;
-					UpdateContext.bLoop = DatabaseBlendSpace->BlendSpace->bLoop;
-					UpdateContext.BlendParameters = Skeletons[SelectedPose].BlendParameters;
-				}
-			}
-			else
-			{
-				checkNoEntry();
-			}
+			UpdatePoseSearchContext(UpdateContext, Skeletons[SelectedPose]);
 
 			if (UpdateContext.Type != ESearchIndexAssetType::Invalid)
 			{
@@ -2508,37 +2523,7 @@ void FDebuggerViewModel::OnDraw(FSkeletonDrawParams& DrawParams)
 		{
 			SetDrawSkeleton(Component, true);
 
-			if (Skeletons[SelectedPose].Type == ESearchIndexAssetType::Sequence)
-			{
-				const FPoseSearchDatabaseSequence* DatabaseSequence = Skeletons[Asset].GetAnimSequence();
-				if (DatabaseSequence)
-				{
-					UpdateContext.Type = ESearchIndexAssetType::Sequence;
-					UpdateContext.Sequence = DatabaseSequence->Sequence;
-					UpdateContext.StartTime = Skeletons[Asset].Time;
-					UpdateContext.Time = Skeletons[Asset].Time;
-					UpdateContext.bMirrored = Skeletons[Asset].bMirrored;
-					UpdateContext.bLoop = DatabaseSequence->Sequence->bLoop;
-				}
-			}
-			else if (Skeletons[SelectedPose].Type == ESearchIndexAssetType::BlendSpace)
-			{
-				const FPoseSearchDatabaseBlendSpace* DatabaseBlendSpace = Skeletons[Asset].GetBlendSpace();
-				if (DatabaseBlendSpace)
-				{
-					UpdateContext.Type = ESearchIndexAssetType::BlendSpace;
-					UpdateContext.BlendSpace = DatabaseBlendSpace->BlendSpace;
-					UpdateContext.StartTime = Skeletons[Asset].Time;
-					UpdateContext.Time = Skeletons[Asset].Time;
-					UpdateContext.bMirrored = Skeletons[Asset].bMirrored;
-					UpdateContext.bLoop = DatabaseBlendSpace->BlendSpace->bLoop;
-					UpdateContext.BlendParameters = Skeletons[Asset].BlendParameters;
-				}
-			}
-			else
-			{
-				checkNoEntry();
-			}
+			UpdatePoseSearchContext(UpdateContext, Skeletons[SelectedPose]);
 
 			if (UpdateContext.Type != ESearchIndexAssetType::Invalid)
 			{
@@ -2649,77 +2634,72 @@ void FDebuggerViewModel::UpdateAsset()
 		AssetSkeleton.Time = AssetData.StartTime;
 	};
 
-	UAnimationAsset* AnimAsset = nullptr;
-	bool bAssetLooping = false;
-
-	if (AssetSkeleton.Type == ESearchIndexAssetType::Sequence)
+	const FInstancedStruct* AnimationAssetStruct = AssetSkeleton.GetAnimationAsset();
+	if (!AnimationAssetStruct)
 	{
-		const FPoseSearchDatabaseSequence* DatabaseSequence = AssetSkeleton.GetAnimSequence();
-		AnimAsset = DatabaseSequence->Sequence;
-		bAssetLooping = DatabaseSequence->Sequence->bLoop;
-	}
-	else if (AssetSkeleton.Type == ESearchIndexAssetType::BlendSpace)
-	{
-		const FPoseSearchDatabaseBlendSpace* DatabaseBlendSpace = AssetSkeleton.GetBlendSpace();
-		AnimAsset = DatabaseBlendSpace->BlendSpace;
-		bAssetLooping = DatabaseBlendSpace->BlendSpace->bLoop;
+		checkNoEntry();
+		return;
 	}
-	else
+
+	const FPoseSearchDatabaseAnimationAssetBase* DatabaseAsset = AnimationAssetStruct->GetPtr<FPoseSearchDatabaseAnimationAssetBase>();
+	if (!DatabaseAsset)
 	{
 		checkNoEntry();
+		return;
+
 	}
 
-	if (AnimAsset)
+	const UAnimationAsset* AnimAsset = DatabaseAsset->GetAnimationAsset();
+	const bool bAssetLooping = DatabaseAsset->IsLooping();
+
+	const float DT = static_cast<float>(FApp::GetDeltaTime()) * AssetPlayRate;
+	const float PlayLength = AnimAsset->GetPlayLength();
+	const bool bExceededDistanceHorizon = Component->LastRootMotionDelta.GetTranslation().Size() > MAX_DISTANCE_RANGE;
+	const bool bExceededTimeHorizon = (AssetSkeleton.Time - AssetData.StartTime) > MAX_TIME_RANGE;
+	const bool bExceededHorizon = bExceededDistanceHorizon && bExceededTimeHorizon;
+	if (bAssetLooping)
 	{
-		const float DT = static_cast<float>(FApp::GetDeltaTime()) * AssetPlayRate;
-		const float PlayLength = AnimAsset->GetPlayLength();
-		const bool bExceededDistanceHorizon = Component->LastRootMotionDelta.GetTranslation().Size() > MAX_DISTANCE_RANGE;
-		const bool bExceededTimeHorizon = (AssetSkeleton.Time - AssetData.StartTime) > MAX_TIME_RANGE;
-		const bool bExceededHorizon = bExceededDistanceHorizon && bExceededTimeHorizon;
-		if (bAssetLooping)
+		if (bExceededHorizon)
 		{
-			if (bExceededHorizon)
+			// Delay before restarting the asset to give the user some idea of where it would land
+			if (AssetData.AccumulatedTime > AssetData.StopDuration)
 			{
-				// Delay before restarting the asset to give the user some idea of where it would land
-				if (AssetData.AccumulatedTime > AssetData.StopDuration)
-				{
-					RestartAsset();
-				}
-				else
-				{
-					AssetData.AccumulatedTime += DT;
-				}
-				return;
+				RestartAsset();
 			}
-
-			AssetSkeleton.Time += DT;
-			AssetData.AccumulatedTime += DT;
+			else
+			{
+				AssetData.AccumulatedTime += DT;
+			}
+			return;
 		}
-		else
-		{
-			// Used to cap the asset, but avoid modding when updating the pose
-			static constexpr float LengthOffset = 0.001f;
-			const bool bFinishedAsset = AssetSkeleton.Time >= PlayLength - LengthOffset;
 
-			// Asset player reached end of clip or reached distance horizon of trajectory vector
-			if (bFinishedAsset || bExceededHorizon)
+		AssetSkeleton.Time += DT;
+		AssetData.AccumulatedTime += DT;
+	}
+	else
+	{
+		// Used to cap the asset, but avoid modding when updating the pose
+		static constexpr float LengthOffset = 0.001f;
+		const bool bFinishedAsset = AssetSkeleton.Time >= PlayLength - LengthOffset;
+
+		// Asset player reached end of clip or reached distance horizon of trajectory vector
+		if (bFinishedAsset || bExceededHorizon)
+		{
+			// Delay before restarting the asset to give the user some idea of where it would land
+			if (AssetData.AccumulatedTime > AssetData.StopDuration)
 			{
-				// Delay before restarting the asset to give the user some idea of where it would land
-				if (AssetData.AccumulatedTime > AssetData.StopDuration)
-				{
-					RestartAsset();
-				}
-				else
-				{
-					AssetData.AccumulatedTime += DT;
-				}
+				RestartAsset();
 			}
 			else
 			{
-				// If we haven't finished, update the play time capped by the anim asset (not looping)
-				AssetSkeleton.Time += DT;
+				AssetData.AccumulatedTime += DT;
 			}
 		}
+		else
+		{
+			// If we haven't finished, update the play time capped by the anim asset (not looping)
+			AssetSkeleton.Time += DT;
+		}
 	}
 }
 
@@ -2804,20 +2784,11 @@ void FDebuggerViewModel::OnWorldCleanup(UWorld* InWorld, bool bSessionEnded, boo
 	bSkeletonsInitialized = false;
 }
 
-const FPoseSearchDatabaseSequence* FDebuggerViewModel::FSkeleton::GetAnimSequence() const
-{
-	if ((Type == ESearchIndexAssetType::Sequence) && SourceDatabase.IsValid() && SourceDatabase->Sequences.IsValidIndex(AssetIdx))
-	{
-		return &SourceDatabase->Sequences[AssetIdx];
-	}
-	return nullptr;
-}
-
-const FPoseSearchDatabaseBlendSpace* FDebuggerViewModel::FSkeleton::GetBlendSpace() const
+const FInstancedStruct* FDebuggerViewModel::FSkeleton::GetAnimationAsset() const
 {
-	if ((Type == ESearchIndexAssetType::BlendSpace) && SourceDatabase.IsValid() && SourceDatabase->BlendSpaces.IsValidIndex(AssetIdx))
+	if (SourceDatabase.IsValid() && SourceDatabase->AnimationAssets.IsValidIndex(AssetIdx))
 	{
-		return &SourceDatabase->BlendSpaces[AssetIdx];
+		return &SourceDatabase->GetAnimationAssetStruct(AssetIdx);
 	}
 	return nullptr;
 }
diff --git a/Engine/Plugins/Experimental/Animation/PoseSearch/Source/Editor/Private/PoseSearchDebugger.h b/Engine/Plugins/Experimental/Animation/PoseSearch/Source/Editor/Private/PoseSearchDebugger.h
--- a/Engine/Plugins/Experimental/Animation/PoseSearch/Source/Editor/Private/PoseSearchDebugger.h
+++ b/Engine/Plugins/Experimental/Animation/PoseSearch/Source/Editor/Private/PoseSearchDebugger.h
@@ -49,7 +49,7 @@ public:
 	struct FUpdateContext
 	{
 		ESearchIndexAssetType Type = ESearchIndexAssetType::Invalid;
-		const UAnimSequence* Sequence = nullptr;
+		const UAnimSequenceBase* SequenceBase = nullptr;
 		const UBlendSpace* BlendSpace = nullptr;
 		float StartTime = 0.0f;
 		float Time = 0.0f;
@@ -468,12 +468,6 @@ public:
 	/** Updates internal Skeletal Mesh Component depending on input */
 	void OnDraw(FSkeletonDrawParams& DrawParams);
 	
-	/** Get an animation sequence from the sequence ID of the active database */
-	const FPoseSearchDatabaseSequence* GetAnimSequence(int32 SequenceIdx) const;
-	
-	/** Get a blendspace from the blend space ID of the active database */
-	const FPoseSearchDatabaseBlendSpace* GetBlendSpace(int32 BlendSpaceIdx) const;
-
 	/** Sets the selected pose skeleton*/
 	void ShowSelectedSkeleton(const UPoseSearchDatabase* Database, int32 DbPoseIdx, float Time);
 	
@@ -508,11 +502,43 @@ public:
 	const USkinnedMeshComponent* GetMeshComponent() const;
 
 private:
+	/** Debug visualization skeleton actor */
+	struct FSkeleton
+	{
+		/** Actor object for the skeleton */
+		TWeakObjectPtr<AActor> Actor = nullptr;
+
+		/** Derived skeletal mesh for setting the skeleton in the scene */
+		TWeakObjectPtr<UPoseSearchMeshComponent> Component = nullptr;
+
+		/** Type of the asset being played */
+		ESearchIndexAssetType Type = ESearchIndexAssetType::Invalid;
+
+		/** Source database for this skeleton  */
+		TWeakObjectPtr<const UPoseSearchDatabase> SourceDatabase = nullptr;
+
+		/** Source asset for this skeleton */
+		int32 AssetIdx = 0;
+
+		/** Time in the sequence this skeleton is accessing */
+		float Time = 0.0f;
+
+		/** If this asset should be mirrored */
+		bool bMirrored = false;
+
+		/** Blend Parameters if asset is a BlendSpace */
+		FVector BlendParameters = FVector::Zero();
+
+		const FInstancedStruct* GetAnimationAsset() const;
+	};
+
 	/** Update the list of states for this frame */
 	void UpdateFromTimeline();
 
 	/** Populates arrays used for mirroring the animation pose */
 	void FillCompactPoseAndComponentRefRotations();
+
+	void UpdatePoseSearchContext(UPoseSearchMeshComponent::FUpdateContext& InOutContext, const FSkeleton& Skeleton) const;
 	
 	/** List of all Node IDs associated with motion matching states */
 	TArray<int32> NodeIds;
@@ -544,37 +570,7 @@ private:
 	/** Pre-calculated component space rotations of reference pose */
 	TCustomBoneIndexArray<FQuat, FCompactPoseBoneIndex> ComponentSpaceRefRotations;
 
-	/** Debug visualization skeleton actor */
-	struct FSkeleton
-	{
-		/** Actor object for the skeleton */
-		TWeakObjectPtr<AActor> Actor = nullptr;
 
-		/** Derived skeletal mesh for setting the skeleton in the scene */
-		TWeakObjectPtr<UPoseSearchMeshComponent> Component = nullptr;
-		
-		/** Type of the asset being played */
-		ESearchIndexAssetType Type = ESearchIndexAssetType::Invalid;
-
-		/** Source database for this skeleton  */
-		TWeakObjectPtr<const UPoseSearchDatabase> SourceDatabase = nullptr;
-
-		/** Source asset for this skeleton */
-		int32 AssetIdx = 0;
-		
-		/** Time in the sequence this skeleton is accessing */
-		float Time = 0.0f;
-
-		/** If this asset should be mirrored */
-		bool bMirrored = false;
-
-		/** Blend Parameters if asset is a BlendSpace */
-		FVector BlendParameters = FVector::Zero();
-
-		const FPoseSearchDatabaseSequence* GetAnimSequence() const;
-		const FPoseSearchDatabaseBlendSpace* GetBlendSpace() const;
-	};
-	
 	/** Index for each type of skeleton we store for debug visualization */
 	enum ESkeletonIndex
 	{
diff --git a/Engine/Plugins/Experimental/Animation/PoseSearch/Source/Editor/Private/SPoseSearchDatabaseAssetList.cpp b/Engine/Plugins/Experimental/Animation/PoseSearch/Source/Editor/Private/SPoseSearchDatabaseAssetList.cpp
--- a/Engine/Plugins/Experimental/Animation/PoseSearch/Source/Editor/Private/SPoseSearchDatabaseAssetList.cpp
+++ b/Engine/Plugins/Experimental/Animation/PoseSearch/Source/Editor/Private/SPoseSearchDatabaseAssetList.cpp
@@ -6,6 +6,7 @@
 #include "PoseSearch/PoseSearch.h"
 
 #include "Animation/AnimSequence.h"
+#include "Animation/AnimComposite.h"
 #include "Animation/BlendSpace.h"
 
 #include "Misc/TransactionObjectEvent.h"
@@ -60,40 +61,12 @@ namespace UE::PoseSearch
 	{
 		if (const UPoseSearchDatabase* Database = EditorViewModel.Pin()->GetPoseSearchDatabase())
 		{
-			if (SourceAssetType == ESearchIndexAssetType::Sequence)
+			if (const FPoseSearchDatabaseAnimationAssetBase* DatabaseAnimationAsset = Database->GetAnimationAssetBase(SourceAssetIdx))
 			{
-				if (Database->Sequences.IsValidIndex(SourceAssetIdx))
-				{
-					if (const UAnimSequence* Sequence = Database->Sequences[SourceAssetIdx].Sequence)
-					{
-						return Sequence->HasRootMotion();
-					}
-				}
-			}
-			else if (SourceAssetType == ESearchIndexAssetType::BlendSpace)
-			{
-				if (Database->BlendSpaces.IsValidIndex(SourceAssetIdx))
-				{
-					if (const UBlendSpace* BlendSpace = Database->BlendSpaces[SourceAssetIdx].BlendSpace)
-					{
-						bool bIsRootMotionUsedInBlendSpace = false;
-			
-						BlendSpace->ForEachImmutableSample([&bIsRootMotionUsedInBlendSpace](const FBlendSample & Sample)
-						{
-							const TObjectPtr<UAnimSequence> Sequence = Sample.Animation;
-				
-							if (IsValid(Sequence) && Sequence->HasRootMotion())
-							{
-								bIsRootMotionUsedInBlendSpace = true;
-							}
-						});
-			
-						return bIsRootMotionUsedInBlendSpace;
-					}
-				}
+				return DatabaseAnimationAsset->IsRootMotionEnabled();
 			}
 		}
-		
+
 		return false;
 	}
 
@@ -101,25 +74,9 @@ namespace UE::PoseSearch
 	{
 		if (const UPoseSearchDatabase* Database = EditorViewModel.Pin()->GetPoseSearchDatabase())
 		{
-			if (SourceAssetType == ESearchIndexAssetType::Sequence)
+			if (const FPoseSearchDatabaseAnimationAssetBase* DatabaseAnimationAsset = Database->GetAnimationAssetBase(SourceAssetIdx))
 			{
-				if (Database->Sequences.IsValidIndex(SourceAssetIdx))
-				{
-					if (const UAnimSequence* Sequence = Database->Sequences[SourceAssetIdx].Sequence)
-					{
-						return Sequence->bLoop;
-					}
-				}
-			}
-			else if (SourceAssetType == ESearchIndexAssetType::BlendSpace)
-			{
-				if (Database->BlendSpaces.IsValidIndex(SourceAssetIdx))
-				{
-					if (const UBlendSpace* BlendSpace = Database->BlendSpaces[SourceAssetIdx].BlendSpace)
-					{
-						return BlendSpace->bLoop;
-					}
-				}
+				return DatabaseAnimationAsset->IsLooping();
 			}
 		}
 
@@ -130,19 +87,9 @@ namespace UE::PoseSearch
 	{
 		if (const UPoseSearchDatabase* Database = EditorViewModel.Pin()->GetPoseSearchDatabase())
 		{
-			if (SourceAssetType == ESearchIndexAssetType::Sequence)
-			{
-				if (Database->Sequences.IsValidIndex(SourceAssetIdx))
-				{
-					return Database->Sequences[SourceAssetIdx].MirrorOption;
-				}
-			}
-			else if (SourceAssetType == ESearchIndexAssetType::BlendSpace)
+			if (const FPoseSearchDatabaseAnimationAssetBase* DatabaseAnimationAsset = Database->GetAnimationAssetBase(SourceAssetIdx))
 			{
-				if (Database->BlendSpaces.IsValidIndex(SourceAssetIdx))
-				{
-					return Database->BlendSpaces[SourceAssetIdx].MirrorOption;
-				}
+				return DatabaseAnimationAsset->GetMirrorOption();
 			}
 		}
 
@@ -214,45 +161,28 @@ namespace UE::PoseSearch
 		SkeletonView.Pin()->RefreshTreeView(false);
 	}
 
-	FText SDatabaseAssetListItem::GetName() const
+	void SDatabaseAssetListItem::OnAddAnimComposite()
 	{
-		FText Name = LOCTEXT("None", "None");
+		EditorViewModel.Pin()->AddAnimCompositeToDatabase(nullptr);
+		SkeletonView.Pin()->RefreshTreeView(false);
+	}
 
+	FText SDatabaseAssetListItem::GetName() const
+	{
 		TSharedPtr<FDatabaseAssetTreeNode> Node = WeakAssetTreeNode.Pin();
 		TSharedPtr<FDatabaseViewModel> ViewModel = EditorViewModel.Pin();
-		UPoseSearchDatabase* Database = ViewModel->GetPoseSearchDatabase();
 
-		if (Database)
+		if (const UPoseSearchDatabase* Database = ViewModel->GetPoseSearchDatabase())
 		{
-			switch (Node->SourceAssetType)
-			{
-			case ESearchIndexAssetType::Sequence:
+			if (const FPoseSearchDatabaseAnimationAssetBase* DatabaseAnimationAsset = Database->GetAnimationAssetBase(Node->SourceAssetIdx))
 			{
-				const UAnimSequence* Sequence = Database->Sequences[Node->SourceAssetIdx].Sequence;
-				if (Sequence)
-				{
-					Name = FText::FromString(Sequence->GetName());
-				}
-				break;
-			}
-			case ESearchIndexAssetType::BlendSpace:
-			{
-				const UBlendSpace* BlendSpace = Database->BlendSpaces[Node->SourceAssetIdx].BlendSpace;
-				if (BlendSpace)
-				{
-					Name = FText::FromString(BlendSpace->GetName());
-				}
-				break;
-			}
-			default:
-			{
-				Name = FText::FromString(Database->GetName());
-				break;
-			}
+				return FText::FromString(DatabaseAnimationAsset->GetName());
 			}
+
+			return FText::FromString(Database->GetName());
 		}
 
-		return Name;
+		return LOCTEXT("None", "None");
 	}
 
 	TSharedRef<SWidget> SDatabaseAssetListItem::GenerateItemWidget()
@@ -301,18 +231,12 @@ namespace UE::PoseSearch
 		}
 		else
 		{
-			// Item Icon (Sequence or Blendspace)
+			// Item Icon
 			TSharedPtr<SImage> ItemIconWidget;
-			if (Node->SourceAssetType == ESearchIndexAssetType::Sequence)
+			if (const FPoseSearchDatabaseAnimationAssetBase* DatabaseAnimationAsset = Database->GetAnimationAssetBase(Node->SourceAssetIdx))
 			{
-				FSlateIconFinder::FindIconBrushForClass(UBlendSpace::StaticClass());
 				SAssignNew(ItemIconWidget, SImage)
-				.Image(FSlateIconFinder::FindIconBrushForClass(UAnimSequence::StaticClass()));
-			}
-			else
-			{
-				SAssignNew(ItemIconWidget, SImage)
-				.Image(FSlateIconFinder::FindIconBrushForClass(UBlendSpace::StaticClass()));
+					.Image(FSlateIconFinder::FindIconBrushForClass(DatabaseAnimationAsset->GetAnimationAssetStaticClass()));
 			}
 
 			// Setup table row to display 
@@ -465,6 +389,14 @@ namespace UE::PoseSearch
 			NAME_None,
 			EUserInterfaceActionType::Button);
 
+		AddOptions.AddMenuEntry(
+			LOCTEXT("AnimCompositeOption", "Add Anim Composite"),
+			LOCTEXT("AddAnimCompositeToDefaultGroupTooltip", "Add new anim composite to this group"),
+			FSlateIcon(),
+			FUIAction(FExecuteAction::CreateSP(this, &SDatabaseAssetListItem::OnAddAnimComposite)),
+			NAME_None,
+			EUserInterfaceActionType::Button);
+
 		TSharedPtr<SComboButton> AddButton;
 		SAssignNew(AddButton, SComboButton)
 		.ContentPadding(0)
@@ -517,22 +449,13 @@ namespace UE::PoseSearch
 	{
 		TSharedPtr<FDatabaseViewModel> ViewModelPtr = EditorViewModel.Pin();
 		TSharedPtr<FDatabaseAssetTreeNode> TreeNodePtr = WeakAssetTreeNode.Pin();
-		const UPoseSearchDatabase * Database = ViewModelPtr->GetPoseSearchDatabase();
-		
-		if (TreeNodePtr->SourceAssetType == ESearchIndexAssetType::Sequence)
-		{
-			if (Database->Sequences.IsValidIndex(TreeNodePtr->SourceAssetIdx))
-			{
-				const bool bEnabled = ViewModelPtr->IsSelectedSequenceEnabled(TreeNodePtr->SourceAssetIdx);
-				return bEnabled ? ECheckBoxState::Checked : ECheckBoxState::Unchecked;
-			}
-		}
-		else if (TreeNodePtr->SourceAssetType == ESearchIndexAssetType::BlendSpace)
+		const UPoseSearchDatabase* Database = ViewModelPtr->GetPoseSearchDatabase();
+
+		if (Database->AnimationAssets.IsValidIndex(TreeNodePtr->SourceAssetIdx))
 		{
-			if (Database->BlendSpaces.IsValidIndex(TreeNodePtr->SourceAssetIdx))
+			if (ViewModelPtr->IsEnabled(TreeNodePtr->SourceAssetIdx))
 			{
-				const bool bEnabled = ViewModelPtr->IsSelectedBlendSpaceEnabled(TreeNodePtr->SourceAssetIdx);
-				return bEnabled ? ECheckBoxState::Checked : ECheckBoxState::Unchecked;
+				return ECheckBoxState::Checked;
 			}
 		}
 
@@ -548,14 +471,7 @@ namespace UE::PoseSearch
 
 		ViewModelPtr->GetPoseSearchDatabase()->Modify();
 		
-		if (TreeNodePtr->SourceAssetType == ESearchIndexAssetType::Sequence)
-		{
-			ViewModelPtr->SetSelectedSequenceEnabled(TreeNodePtr->SourceAssetIdx, NewCheckboxState == ECheckBoxState::Checked ? true : false);
-		}
-		else if (TreeNodePtr->SourceAssetType == ESearchIndexAssetType::BlendSpace)
-		{
-			ViewModelPtr->SetSelectedBlendSpaceEnabled(TreeNodePtr->SourceAssetIdx, NewCheckboxState == ECheckBoxState::Checked ? true : false);
-		}
+		ViewModelPtr->SetIsEnabled(TreeNodePtr->SourceAssetIdx, NewCheckboxState == ECheckBoxState::Checked ? true : false);
 
 		SkeletonView.Pin()->RefreshTreeView(false, true);
 		ViewModelPtr->BuildSearchIndex();
@@ -700,7 +616,7 @@ namespace UE::PoseSearch
 						.VAlign(VAlign_Center)
 						[
 							SNew(STextBlock)
-							.Text(FText::FromString(TEXT("Drag and drop Animation Sequences or Blendspaces")))
+							.Text(FText::FromString(TEXT("Drag and drop Animation Sequences, Anim Composites or Blendspaces")))
 							.Font(FAppStyle::Get().GetFontStyle("DetailsView.CategoryFontStyle"))
 						]
 					]
@@ -735,6 +651,7 @@ namespace UE::PoseSearch
 						if (UClass* AssetClass = AssetData.GetClass())
 						{
 							if (AssetClass->IsChildOf(UAnimSequence::StaticClass()) ||
+								AssetClass->IsChildOf(UAnimComposite::StaticClass()) ||
 								AssetClass->IsChildOf(UBlendSpace::StaticClass()))
 							{
 								Reply = FReply::Handled();
@@ -861,89 +778,51 @@ namespace UE::PoseSearch
 			};
 			
 			// Build an index based off of alphabetical order than iterate the index instead
-			TArray<uint32> SequenceIndexArray;
-			SequenceIndexArray.SetNumUninitialized(Database->Sequences.Num());
-			for (int32 SequenceIdx = 0; SequenceIdx < Database->Sequences.Num(); ++SequenceIdx)
+			TArray<uint32> IndexArray;
+			IndexArray.SetNumUninitialized(Database->AnimationAssets.Num());
+			for (int32 AnimationAssetIdx = 0; AnimationAssetIdx < Database->AnimationAssets.Num(); ++AnimationAssetIdx)
 			{
-				SequenceIndexArray[SequenceIdx] = SequenceIdx;
+				IndexArray[AnimationAssetIdx] = AnimationAssetIdx;
 			}
 
-			SequenceIndexArray.Sort([Database](int32 SequenceIdxA, int32 SequenceIdxB)
+			IndexArray.Sort([Database](int32 SequenceIdxA, int32 SequenceIdxB)
 			{
-				const FPoseSearchDatabaseSequence& A = Database->Sequences[SequenceIdxA];
-				const FPoseSearchDatabaseSequence& B = Database->Sequences[SequenceIdxB];
+				const FPoseSearchDatabaseAnimationAssetBase* A = Database->GetAnimationAssetBase(SequenceIdxA);
+				const FPoseSearchDatabaseAnimationAssetBase* B = Database->GetAnimationAssetBase(SequenceIdxB);
 
 				//If its null add it to the end of the list 
-				if (!B.Sequence)
+				if (!B->GetAnimationAsset())
 				{
 					return true;
 				}
 
-				if (!A.Sequence)
+				if (!A->GetAnimationAsset())
 				{
 					return false;
 				}
 
-				const int32 Comparison = A.Sequence->GetName().Compare(B.Sequence->GetName());
+				const int32 Comparison = A->GetName().Compare(B->GetName());
 				return Comparison < 0;
 			});
 
-			// create all sequence nodes
-			for (int32 SequenceIdx = 0; SequenceIdx < Database->Sequences.Num(); ++SequenceIdx)
+			// create all nodes
+			for (int32 AnimationAssetIdx = 0; AnimationAssetIdx < Database->AnimationAssets.Num(); ++AnimationAssetIdx)
 			{
-				const int32 MappedId = SequenceIndexArray[SequenceIdx];
-				const FPoseSearchDatabaseSequence& DbSequence = Database->Sequences[MappedId];
-				const bool bFiltered = (DbSequence.Sequence == nullptr || GetAssetFilterString().IsEmpty()) ? false : !DbSequence.Sequence->GetName().Contains(GetAssetFilterString());
+				const int32 MappedId = IndexArray[AnimationAssetIdx];
 
-				if (!bFiltered)
+				if (const FPoseSearchDatabaseAnimationAssetBase* DatabaseAnimationAsset = Database->GetAnimationAssetBase(MappedId))
 				{
-					CreateAssetNode(MappedId, ESearchIndexAssetType::Sequence, DefaultGroupIdx);
-				}
-			}
-
-			// Build an index based off of alphabetical order than iterate the index instead
-			TArray<uint32> BlendspaceIndexArray;
-			BlendspaceIndexArray.SetNumUninitialized(Database->BlendSpaces.Num());
-			for (int32 BlendspaceIdx = 0; BlendspaceIdx < Database->BlendSpaces.Num(); ++BlendspaceIdx)
-			{
-				BlendspaceIndexArray[BlendspaceIdx] = BlendspaceIdx;
-			}
+					const bool bFiltered = (DatabaseAnimationAsset->GetAnimationAsset() == nullptr || GetAssetFilterString().IsEmpty()) ? false : !DatabaseAnimationAsset->GetName().Contains(GetAssetFilterString());
 
-			BlendspaceIndexArray.Sort([Database](int32 BlendspaceIdxA, int32 BlendspaceIdxB)
-			{
-				const FPoseSearchDatabaseBlendSpace& A = Database->BlendSpaces[BlendspaceIdxA];
-				const FPoseSearchDatabaseBlendSpace& B = Database->BlendSpaces[BlendspaceIdxB];
-
-				//If its null add it to the end of the list 
-				if (!B.BlendSpace)
-				{
-					return true;
-				}
-
-				if (!A.BlendSpace)
-				{
-					return false;
-				}
-
-				const int32 Comparison = A.BlendSpace->GetName().Compare(B.BlendSpace->GetName());
-				return Comparison < 0;
-			});
-
-			// create all blendspace nodes
-			for (int32 BlendSpaceIdx = 0; BlendSpaceIdx < Database->BlendSpaces.Num(); ++BlendSpaceIdx)
-			{
-				const int32 MappedId = BlendspaceIndexArray[BlendSpaceIdx];
-				const FPoseSearchDatabaseBlendSpace& DbBlendSpace = Database->BlendSpaces[MappedId];
-				const bool bFiltered = (DbBlendSpace.BlendSpace == nullptr || GetAssetFilterString().IsEmpty()) ? false : !DbBlendSpace.BlendSpace->GetName().Contains(GetAssetFilterString());
-
-				if (!bFiltered)
-				{
-					CreateAssetNode(MappedId, ESearchIndexAssetType::BlendSpace, DefaultGroupIdx);
+					if (!bFiltered)
+					{
+						CreateAssetNode(MappedId, DatabaseAnimationAsset->GetSearchIndexType(), DefaultGroupIdx);
+					}
 				}
 			}
 
 			// Show drag and drop suggestion if tree is empty
-			TreeViewDragAndDropSuggestion->SetVisibility(BlendspaceIndexArray.IsEmpty() && SequenceIndexArray.IsEmpty() ? EVisibility::Visible : EVisibility::Hidden);
+			TreeViewDragAndDropSuggestion->SetVisibility(IndexArray.IsEmpty() ? EVisibility::Visible : EVisibility::Hidden);
 		}
 
 		// Update tree view
@@ -998,6 +877,7 @@ namespace UE::PoseSearch
 				if (UClass* AssetClass = AssetData.GetClass())
 				{
 					if (AssetClass->IsChildOf(UAnimSequence::StaticClass()) ||
+						AssetClass->IsChildOf(UAnimComposite::StaticClass()) ||
 						AssetClass->IsChildOf(UBlendSpace::StaticClass()))
 					{
 						ReturnedDropZone = EItemDropZone::OntoItem;
@@ -1063,6 +943,11 @@ namespace UE::PoseSearch
 						ViewModel->AddSequenceToDatabase(Cast<UAnimSequence>(Asset));
 						++AddedAssets;
 					}
+					if (AssetClass->IsChildOf(UAnimComposite::StaticClass()))
+					{
+						ViewModel->AddAnimCompositeToDatabase(Cast<UAnimComposite>(Asset));
+						++AddedAssets;
+					}
 					else if (AssetClass->IsChildOf(UBlendSpace::StaticClass()))
 					{
 						ViewModel->AddBlendSpaceToDatabase(Cast<UBlendSpace>(Asset));
@@ -1088,21 +973,31 @@ namespace UE::PoseSearch
 		FMenuBuilder AddOptions(true, nullptr);
 
 		AddOptions.BeginSection("AddOptions", LOCTEXT("AssetAddOptions", "Assets"));
-		AddOptions.AddMenuEntry(
-			LOCTEXT("AddSequenceOption", "Sequence"),
-			LOCTEXT("AddSequenceOptionTooltip", "Add new sequence to the default group"),
-			FSlateIcon(),
-			FUIAction(FExecuteAction::CreateSP(this, &SDatabaseAssetTree::OnAddSequence, true)),
-			NAME_None,
-			EUserInterfaceActionType::Button);
+		{
+			AddOptions.AddMenuEntry(
+				LOCTEXT("AddSequenceOption", "Sequence"),
+				LOCTEXT("AddSequenceOptionTooltip", "Add new sequence to the default group"),
+				FSlateIcon(),
+				FUIAction(FExecuteAction::CreateSP(this, &SDatabaseAssetTree::OnAddSequence, true)),
+				NAME_None,
+				EUserInterfaceActionType::Button);
 
-		AddOptions.AddMenuEntry(
-			LOCTEXT("BlendSpaceOption", "Blend Space"),
-			LOCTEXT("AddBlendSpaceToDefaultGroupTooltip", "Add new blend space to the default group"),
-			FSlateIcon(),
-			FUIAction(FExecuteAction::CreateSP(this, &SDatabaseAssetTree::OnAddBlendSpace, true)),
-			NAME_None,
-			EUserInterfaceActionType::Button);
+			AddOptions.AddMenuEntry(
+				LOCTEXT("BlendSpaceOption", "Blend Space"),
+				LOCTEXT("AddBlendSpaceToDefaultGroupTooltip", "Add new blend space to the default group"),
+				FSlateIcon(),
+				FUIAction(FExecuteAction::CreateSP(this, &SDatabaseAssetTree::OnAddBlendSpace, true)),
+				NAME_None,
+				EUserInterfaceActionType::Button);
+
+			AddOptions.AddMenuEntry(
+				LOCTEXT("AnimCompositeOption", "Anim Composite"),
+				LOCTEXT("AddAnimCompositeToDefaultGroupTooltip", "Add new anim composite to the default group"),
+				FSlateIcon(),
+				FUIAction(FExecuteAction::CreateSP(this, &SDatabaseAssetTree::OnAddAnimComposite, true)),
+				NAME_None,
+				EUserInterfaceActionType::Button);
+		}
 		AddOptions.EndSection();
 
 		return AddOptions.MakeWidget();
@@ -1206,25 +1101,25 @@ namespace UE::PoseSearch
 		}
 	}
 
+	void SDatabaseAssetTree::OnAddAnimComposite(bool bFinalizeChanges)
+	{
+		FScopedTransaction Transaction(LOCTEXT("AddAnimCompositeTransaction", "Add Anim Composite"));
+
+		EditorViewModel.Pin()->AddAnimCompositeToDatabase(nullptr);
+
+		if (bFinalizeChanges)
+		{
+			FinalizeTreeChanges();
+		}
+	}
+
 	void SDatabaseAssetTree::OnDeleteAsset(TSharedPtr<FDatabaseAssetTreeNode> Node, bool bFinalizeChanges)
 	{
 		FScopedTransaction Transaction(LOCTEXT("DeleteAsset", "Delete Asset"));
 		const TSharedPtr<FDatabaseViewModel> ViewModel = EditorViewModel.Pin();
 		
 		ViewModel->GetPoseSearchDatabase()->Modify();
-
-		if (Node->SourceAssetType == ESearchIndexAssetType::Sequence)
-		{
-			ViewModel->DeleteSequenceFromDatabase(Node->SourceAssetIdx);
-		}
-		else if (Node->SourceAssetType == ESearchIndexAssetType::BlendSpace)
-		{
-			ViewModel->DeleteBlendSpaceFromDatabase(Node->SourceAssetIdx);
-		}
-		else
-		{
-			checkNoEntry();
-		}
+		ViewModel->DeleteFromDatabase(Node->SourceAssetIdx);
 
 		if (bFinalizeChanges)
 		{
@@ -1330,54 +1225,22 @@ namespace UE::PoseSearch
 		}
 	}
 
-	void SDatabaseAssetTree::OnEnableNodes()
+	void SDatabaseAssetTree::EnableSelectedNodes(bool bIsEnabled)
 	{
 		TArray<TSharedPtr<FDatabaseAssetTreeNode>> SelectedNodes = TreeView->GetSelectedItems();
 		if (!SelectedNodes.IsEmpty())
 		{
-			const FScopedTransaction Transaction(LOCTEXT("EnablePoseSearchDatabaseNodes", "Enable selected items from Pose Search Database"));
+			const FText TransactionName = bIsEnabled ? LOCTEXT("EnablePoseSearchDatabaseNodes", "Enable selected items from Pose Search Database") : LOCTEXT("DisablePoseSearchDatabaseNodes", "Disable selected items from Pose Search Database");
+			const FScopedTransaction Transaction(TransactionName);
 			const TSharedPtr<FDatabaseViewModel> ViewModel = EditorViewModel.Pin();
 
 			ViewModel->GetPoseSearchDatabase()->Modify();
-			
-			for (TSharedPtr<FDatabaseAssetTreeNode> SelectedNode : SelectedNodes)
-			{
-				if (SelectedNode->SourceAssetType == ESearchIndexAssetType::Sequence)
-				{
-					ViewModel->SetSelectedSequenceEnabled(SelectedNode->SourceAssetIdx, true);
-				}
-				else if (SelectedNode->SourceAssetType == ESearchIndexAssetType::BlendSpace)
-				{
-					ViewModel->SetSelectedBlendSpaceEnabled(SelectedNode->SourceAssetIdx, true);
-				}
-			}
-		
-			FinalizeTreeChanges();
-		}
-	}
-
-	void SDatabaseAssetTree::OnDisableNodes()
-	{
-		TArray<TSharedPtr<FDatabaseAssetTreeNode>> SelectedNodes = TreeView->GetSelectedItems();
-		if (!SelectedNodes.IsEmpty())
-		{
-			const FScopedTransaction Transaction(LOCTEXT("DisablePoseSearchDatabaseNodes", "Disable selected items from Pose Search Database"));
-			const TSharedPtr<FDatabaseViewModel> ViewModel = EditorViewModel.Pin();
 
-			ViewModel->GetPoseSearchDatabase()->Modify();
-			
 			for (TSharedPtr<FDatabaseAssetTreeNode> SelectedNode : SelectedNodes)
 			{
-				if (SelectedNode->SourceAssetType == ESearchIndexAssetType::Sequence)
-				{
-					ViewModel->SetSelectedSequenceEnabled(SelectedNode->SourceAssetIdx, false);
-				}
-				else if (SelectedNode->SourceAssetType == ESearchIndexAssetType::BlendSpace)
-				{
-					ViewModel->SetSelectedBlendSpaceEnabled(SelectedNode->SourceAssetIdx, false);
-				}
+				ViewModel->SetIsEnabled(SelectedNode->SourceAssetIdx, bIsEnabled);
 			}
-
+		
 			FinalizeTreeChanges();
 		}
 	}
diff --git a/Engine/Plugins/Experimental/Animation/PoseSearch/Source/Editor/Public/PoseSearchDatabaseEditor.h b/Engine/Plugins/Experimental/Animation/PoseSearch/Source/Editor/Public/PoseSearchDatabaseEditor.h
--- a/Engine/Plugins/Experimental/Animation/PoseSearch/Source/Editor/Public/PoseSearchDatabaseEditor.h
+++ b/Engine/Plugins/Experimental/Animation/PoseSearch/Source/Editor/Public/PoseSearchDatabaseEditor.h
@@ -68,6 +68,13 @@ namespace UE::PoseSearch
 			ESelectInfo::Type SelectionType);
 
 		void RefreshStatisticsWidgetInformation();
+
+		struct FSelectionWidget
+		{
+			TSharedPtr<IDetailsView> DetailView;
+			TArray<TWeakObjectPtr<UObject>> SelectedReflections;
+		};
+		FSelectionWidget& FindOrAddSelectionWidget(const UScriptStruct* ScriptStructType);
 		
 		TSharedPtr<SDatabasePreview> PreviewWidget;
 
@@ -75,7 +82,8 @@ namespace UE::PoseSearch
 
 		TSharedPtr<IDetailsView> EditingAssetWidget;
 
-		TArray<TSharedPtr<IDetailsView>> SelectionWidgets;
+		TSharedPtr<SVerticalBox> DetailsContainer;
+		TMap<const UScriptStruct*, FSelectionWidget> SelectionWidgets;
 
 		TSharedPtr<IDetailsView> StatisticsOverviewWidget;
 		
diff --git a/Engine/Plugins/Experimental/Animation/PoseSearch/Source/Editor/Public/PoseSearchDatabaseEditorReflection.h b/Engine/Plugins/Experimental/Animation/PoseSearch/Source/Editor/Public/PoseSearchDatabaseEditorReflection.h
--- a/Engine/Plugins/Experimental/Animation/PoseSearch/Source/Editor/Public/PoseSearchDatabaseEditorReflection.h
+++ b/Engine/Plugins/Experimental/Animation/PoseSearch/Source/Editor/Public/PoseSearchDatabaseEditorReflection.h
@@ -59,6 +59,20 @@ public:
 #endif
 };
 
+UCLASS()
+class UPoseSearchDatabaseAnimCompositeReflection : public UPoseSearchDatabaseReflectionBase
+{
+	GENERATED_BODY()
+
+public:
+	UPROPERTY(EditAnywhere, Category = "Selected Anim Composite")
+	FPoseSearchDatabaseAnimComposite AnimComposite;
+
+#if WITH_EDITOR
+	virtual void PostEditChangeProperty(struct FPropertyChangedEvent& PropertyChangedEvent) override;
+#endif
+};
+
 USTRUCT()
 struct FPoseSearchDatabaseMemoryStats
 {
diff --git a/Engine/Plugins/Experimental/Animation/PoseSearch/Source/Editor/Public/PoseSearchDatabaseViewModel.h b/Engine/Plugins/Experimental/Animation/PoseSearch/Source/Editor/Public/PoseSearchDatabaseViewModel.h
--- a/Engine/Plugins/Experimental/Animation/PoseSearch/Source/Editor/Public/PoseSearchDatabaseViewModel.h
+++ b/Engine/Plugins/Experimental/Animation/PoseSearch/Source/Editor/Public/PoseSearchDatabaseViewModel.h
@@ -14,6 +14,7 @@ class UPoseSearchDatabase;
 struct FPoseSearchIndexAsset;
 class UAnimPreviewInstance;
 class UDebugSkelMeshComponent;
+class UAnimComposite;
 class UAnimSequence;
 class UBlendSpace;
 class UMirrorDataTable;
@@ -45,7 +46,7 @@ namespace UE::PoseSearch
 		int32 IndexAssetIndex = INDEX_NONE;
 		int32 CurrentPoseIndex = INDEX_NONE;
 
-		FSequenceSampler SequenceSampler;
+		FSequenceBaseSampler SequenceSampler;
 		FBlendSpaceSampler BlendSpaceSampler;
 		ESearchIndexAssetType Type = ESearchIndexAssetType::Invalid;
 
@@ -99,18 +100,14 @@ namespace UE::PoseSearch
 
 		void OnToggleDisplayRootMotionSpeed();
 		bool IsDisplayRootMotionSpeedChecked() const { return DisplayRootMotionSpeed; };
-		
+
 		void AddSequenceToDatabase(UAnimSequence* AnimSequence);
 		void AddBlendSpaceToDatabase(UBlendSpace* BlendSpace);
+		void AddAnimCompositeToDatabase(UAnimComposite* AnimComposite);
+		void DeleteFromDatabase(int32 AnimationAssetIndex);
 
-		void DeleteSequenceFromDatabase(int32 SequenceIdx);
-		void DeleteBlendSpaceFromDatabase(int32 BlendSpaceIdx);
-
-		void SetSelectedSequenceEnabled(int32 SequenceIndex, bool bEnabled);
-		void SetSelectedBlendSpaceEnabled(int32 BlendSpaceIndex, bool bEnabled);
-		
-		bool IsSelectedSequenceEnabled(int32 SequenceIndex) const;
-		bool IsSelectedBlendSpaceEnabled(int32 BlendSpaceIndex) const;
+		void SetIsEnabled(int32 AnimationAssetIndex, bool bEnabled);
+		bool IsEnabled(int32 AnimationAssetIndex) const;
 
 		void SetSelectedNodes(const TArrayView<TSharedPtr<FDatabaseAssetTreeNode>>& InSelectedNodes);
 		void ProcessSelectedActor(AActor* Actor);
diff --git a/Engine/Plugins/Experimental/Animation/PoseSearch/Source/Editor/Public/SPoseSearchDatabaseAssetList.h b/Engine/Plugins/Experimental/Animation/PoseSearch/Source/Editor/Public/SPoseSearchDatabaseAssetList.h
--- a/Engine/Plugins/Experimental/Animation/PoseSearch/Source/Editor/Public/SPoseSearchDatabaseAssetList.h
+++ b/Engine/Plugins/Experimental/Animation/PoseSearch/Source/Editor/Public/SPoseSearchDatabaseAssetList.h
@@ -64,6 +64,7 @@ namespace UE::PoseSearch
 
 		void OnAddSequence();
 		void OnAddBlendSpace();
+		void OnAddAnimComposite();
 
 		EVisibility GetSelectedActorIconVisbility() const;
 
@@ -156,6 +157,7 @@ namespace UE::PoseSearch
 
 		void OnAddSequence(bool bFinalizeChanges = true);
 		void OnAddBlendSpace(bool bFinalizeChanges = true);
+		void OnAddAnimComposite(bool bFinalizeChanges = true);
 
 		void OnDeleteAsset(TSharedPtr<FDatabaseAssetTreeNode> Node, bool bFinalizeChanges = true);
 		void CreateCommandList();
@@ -164,8 +166,9 @@ namespace UE::PoseSearch
 		bool CanDeleteNodes() const;
 		void OnDeleteNodes();
 
-		void OnEnableNodes();
-		void OnDisableNodes();
+		void EnableSelectedNodes(bool bIsEnabled);
+		void OnEnableNodes() { EnableSelectedNodes(true); }
+		void OnDisableNodes() { EnableSelectedNodes(false); }
 
 		friend SDatabaseAssetListItem;
 
diff --git a/Engine/Plugins/Experimental/Animation/PoseSearch/Source/Runtime/PoseSearch.Build.cs b/Engine/Plugins/Experimental/Animation/PoseSearch/Source/Runtime/PoseSearch.Build.cs
--- a/Engine/Plugins/Experimental/Animation/PoseSearch/Source/Runtime/PoseSearch.Build.cs
+++ b/Engine/Plugins/Experimental/Animation/PoseSearch/Source/Runtime/PoseSearch.Build.cs
@@ -23,7 +23,8 @@ public class PoseSearch : ModuleRules
 				"TraceLog",
 				"AnimGraphRuntime",
 				"GameplayTags",
-				"MotionTrajectory"
+				"MotionTrajectory",
+				"StructUtils"
 			}
 		);
 
diff --git a/Engine/Plugins/Experimental/Animation/PoseSearch/Source/Runtime/Private/AnimNode_BlendStack.cpp b/Engine/Plugins/Experimental/Animation/PoseSearch/Source/Runtime/Private/AnimNode_BlendStack.cpp
--- a/Engine/Plugins/Experimental/Animation/PoseSearch/Source/Runtime/Private/AnimNode_BlendStack.cpp
+++ b/Engine/Plugins/Experimental/Animation/PoseSearch/Source/Runtime/Private/AnimNode_BlendStack.cpp
@@ -4,6 +4,7 @@
 
 #include "Algo/MaxElement.h"
 #include "Animation/AnimInstanceProxy.h"
+#include "Animation/AnimComposite.h"
 #include "Animation/AnimSequence.h"
 #include "Animation/BlendSpace.h"
 
@@ -52,7 +53,7 @@ void FPoseSearchAnimPlayer::Initialize(ESearchIndexAssetType InAssetType, UAnima
 	MirrorNode.SetMirrorDataTable(MirrorDataTable);
 	MirrorNode.SetMirror(bMirrored);
 	
-	if (AssetType == ESearchIndexAssetType::Sequence)
+	if (AssetType == ESearchIndexAssetType::Sequence || AssetType == ESearchIndexAssetType::AnimComposite)
 	{
 		UAnimSequenceBase* Sequence = Cast<UAnimSequenceBase>(AnimationAsset);
 		check(Sequence);
@@ -77,12 +78,13 @@ void FPoseSearchAnimPlayer::Initialize(ESearchIndexAssetType InAssetType, UAnima
 	{
 		checkNoEntry();
 	}
+
 	UpdateSourceLinkNode();
 }
 
 void FPoseSearchAnimPlayer::UpdatePlayRate(float PlayRate)
 {
-	if (AssetType == ESearchIndexAssetType::Sequence)
+	if (AssetType == ESearchIndexAssetType::Sequence || AssetType == ESearchIndexAssetType::AnimComposite)
 	{
 		SequencePlayerNode.SetPlayRate(PlayRate);
 	}
@@ -90,6 +92,10 @@ void FPoseSearchAnimPlayer::UpdatePlayRate(float PlayRate)
 	{
 		BlendSpacePlayerNode.SetPlayRate(PlayRate);
 	}
+	else
+	{
+		checkNoEntry();
+	}
 }
 
 void FPoseSearchAnimPlayer::StorePoseContext(const FPoseContext& PoseContext)
@@ -143,7 +149,7 @@ void FPoseSearchAnimPlayer::RestorePoseContext(FPoseContext& PoseContext) const
 // since we're making copies and moving this object in memory, we're using this method to set the MirrorNode SourceLinkNode when necessary
 void FPoseSearchAnimPlayer::UpdateSourceLinkNode()
 {
-	if (AssetType == ESearchIndexAssetType::Sequence)
+	if (AssetType == ESearchIndexAssetType::Sequence || AssetType == ESearchIndexAssetType::AnimComposite)
 	{
 		MirrorNode.SetSourceLinkNode(&SequencePlayerNode);
 	}
@@ -181,7 +187,7 @@ void FPoseSearchAnimPlayer::Update_AnyThread(const FAnimationUpdateContext& Cont
 
 float FPoseSearchAnimPlayer::GetAccumulatedTime() const
 {
-	if (AssetType == ESearchIndexAssetType::Sequence)
+	if (AssetType == ESearchIndexAssetType::Sequence || AssetType == ESearchIndexAssetType::AnimComposite)
 	{
 		return SequencePlayerNode.GetAccumulatedTime();
 	}
@@ -190,13 +196,14 @@ float FPoseSearchAnimPlayer::GetAccumulatedTime() const
 	{
 		return BlendSpacePlayerNode.GetAccumulatedTime();
 	}
-	
+
+	checkNoEntry();
 	return 0.f;
 }
 
 FString FPoseSearchAnimPlayer::GetAnimName() const
 {
-	if (AssetType == ESearchIndexAssetType::Sequence)
+	if (AssetType == ESearchIndexAssetType::Sequence || AssetType == ESearchIndexAssetType::AnimComposite)
 	{
 		check(SequencePlayerNode.GetSequence());
 		return SequencePlayerNode.GetSequence()->GetName();
@@ -208,6 +215,7 @@ FString FPoseSearchAnimPlayer::GetAnimName() const
 		return BlendSpacePlayerNode.GetBlendSpace()->GetName();
 	}
 
+	checkNoEntry();
 	return FString("StoredPose");
 }
 
diff --git a/Engine/Plugins/Experimental/Animation/PoseSearch/Source/Runtime/Private/AnimNode_MotionMatching.cpp b/Engine/Plugins/Experimental/Animation/PoseSearch/Source/Runtime/Private/AnimNode_MotionMatching.cpp
--- a/Engine/Plugins/Experimental/Animation/PoseSearch/Source/Runtime/Private/AnimNode_MotionMatching.cpp
+++ b/Engine/Plugins/Experimental/Animation/PoseSearch/Source/Runtime/Private/AnimNode_MotionMatching.cpp
@@ -111,10 +111,13 @@ void FAnimNode_MotionMatching::UpdateAssetPlayer(const FAnimationUpdateContext&
 		const UPoseSearchDatabase* Database = MotionMatchingState.CurrentSearchResult.Database.Get();
 		if (Database && Database->Schema && EnumHasAnyFlags(MotionMatchingState.Flags, EMotionMatchingFlags::JumpedToPose))
 		{
-			const FPoseSearchDatabaseAnimationAssetBase& AnimationAssetBase = Database->GetAnimationSourceAsset(*SearchIndexAsset);
-			BlendStackNode.BlendTo(SearchIndexAsset->Type, AnimationAssetBase.GetAnimationAsset(), MotionMatchingState.CurrentSearchResult.AssetTime,
-				AnimationAssetBase.IsLooping(), SearchIndexAsset->bMirrored, Database->Schema->MirrorDataTable.Get(),
-				Settings.MaxActiveBlends, Settings.BlendTime, Settings.BlendProfile, Settings.BlendOption, SearchIndexAsset->BlendParameters, MotionMatchingState.WantedPlayRate);
+			if (const FPoseSearchDatabaseAnimationAssetBase* DatabaseAsset = Database->GetAnimationAssetBase(*SearchIndexAsset))
+			{
+				BlendStackNode.BlendTo(SearchIndexAsset->Type, DatabaseAsset->GetAnimationAsset(), MotionMatchingState.CurrentSearchResult.AssetTime,
+					DatabaseAsset->IsLooping(), SearchIndexAsset->bMirrored, Database->Schema->MirrorDataTable.Get(),
+					Settings.MaxActiveBlends, Settings.BlendTime, Settings.BlendProfile, Settings.BlendOption, SearchIndexAsset->BlendParameters, MotionMatchingState.WantedPlayRate);
+			}
+
 		}
 	}
 	BlendStackNode.UpdatePlayRate(MotionMatchingState.WantedPlayRate);
diff --git a/Engine/Plugins/Experimental/Animation/PoseSearch/Source/Runtime/Private/PoseSearch.cpp b/Engine/Plugins/Experimental/Animation/PoseSearch/Source/Runtime/Private/PoseSearch.cpp
--- a/Engine/Plugins/Experimental/Animation/PoseSearch/Source/Runtime/Private/PoseSearch.cpp
+++ b/Engine/Plugins/Experimental/Animation/PoseSearch/Source/Runtime/Private/PoseSearch.cpp
@@ -544,14 +544,14 @@ float FPoseSearchIndexBase::GetAssetTime(int32 PoseIdx, float SamplingInterval)
 {
 	const FPoseSearchIndexAsset& Asset = GetAssetForPose(PoseIdx);
 
-	if (Asset.Type == ESearchIndexAssetType::Sequence)
+	if (Asset.Type == ESearchIndexAssetType::Sequence || Asset.Type == ESearchIndexAssetType::AnimComposite)
 	{
 		const FFloatInterval SamplingRange = Asset.SamplingInterval;
 
 		float AssetTime = FMath::Min(SamplingRange.Min + SamplingInterval * (PoseIdx - Asset.FirstPoseIdx), SamplingRange.Max);
 		return AssetTime;
 	}
-	
+
 	if (Asset.Type == ESearchIndexAssetType::BlendSpace)
 	{
 		const FFloatInterval SamplingRange = Asset.SamplingInterval;
@@ -836,64 +836,40 @@ bool UPoseSearchDatabase::GetPoseIndicesAndLerpValueFromTime(float Time, const F
 	return true;
 }
 
-const FPoseSearchDatabaseAnimationAssetBase& UPoseSearchDatabase::GetAnimationSourceAsset(const FPoseSearchIndexAsset& SearchIndexAsset) const
+const FInstancedStruct& UPoseSearchDatabase::GetAnimationAssetStruct(int32 AnimationAssetIndex) const
 {
-	if (SearchIndexAsset.Type == ESearchIndexAssetType::Sequence)
-	{
-		return Sequences[SearchIndexAsset.SourceAssetIdx];
-	}
-
-	if (SearchIndexAsset.Type == ESearchIndexAssetType::BlendSpace)
-	{
-		return BlendSpaces[SearchIndexAsset.SourceAssetIdx];
-	}
+	check(AnimationAssets.IsValidIndex(AnimationAssetIndex));
+	return AnimationAssets[AnimationAssetIndex];
+}
 
-	checkNoEntry();
-	return Sequences[SearchIndexAsset.SourceAssetIdx];
+const FInstancedStruct& UPoseSearchDatabase::GetAnimationAssetStruct(const FPoseSearchIndexAsset& SearchIndexAsset) const
+{
+	return GetAnimationAssetStruct(SearchIndexAsset.SourceAssetIdx);
 }
 
-const FPoseSearchDatabaseSequence& UPoseSearchDatabase::GetSequenceSourceAsset(const FPoseSearchIndexAsset& SearchIndexAsset) const
+FPoseSearchDatabaseAnimationAssetBase* UPoseSearchDatabase::GetAnimationAssetBase(int32 AnimationAssetIndex) const
 {
-	check(SearchIndexAsset.Type == ESearchIndexAssetType::Sequence);
-	return Sequences[SearchIndexAsset.SourceAssetIdx];
+	if (AnimationAssets.IsValidIndex(AnimationAssetIndex))
+	{
+		return AnimationAssets[AnimationAssetIndex].GetMutablePtr<FPoseSearchDatabaseAnimationAssetBase>();
+	}
+
+	return nullptr;
 }
 
-const FPoseSearchDatabaseBlendSpace& UPoseSearchDatabase::GetBlendSpaceSourceAsset(const FPoseSearchIndexAsset& SearchIndexAsset) const
+FPoseSearchDatabaseAnimationAssetBase* UPoseSearchDatabase::GetAnimationAssetBase(const FPoseSearchIndexAsset& SearchIndexAsset) const
 {
-	check(SearchIndexAsset.Type == ESearchIndexAssetType::BlendSpace);
-	return BlendSpaces[SearchIndexAsset.SourceAssetIdx];
+	return GetAnimationAssetBase(SearchIndexAsset.SourceAssetIdx);
 }
 
 const bool UPoseSearchDatabase::IsSourceAssetLooping(const FPoseSearchIndexAsset& SearchIndexAsset) const
 {
-	if (SearchIndexAsset.Type == ESearchIndexAssetType::Sequence)
-	{
-		return Sequences[SearchIndexAsset.SourceAssetIdx].Sequence->bLoop;
-	}
-
-	if (SearchIndexAsset.Type == ESearchIndexAssetType::BlendSpace)
-	{
-		return BlendSpaces[SearchIndexAsset.SourceAssetIdx].BlendSpace->bLoop;
-	}
-	
-	checkNoEntry();
-	return false;
+	return GetAnimationAssetBase(SearchIndexAsset.SourceAssetIdx)->IsLooping();
 }
 
 const FString UPoseSearchDatabase::GetSourceAssetName(const FPoseSearchIndexAsset& SearchIndexAsset) const
 {
-	if (SearchIndexAsset.Type == ESearchIndexAssetType::Sequence)
-	{
-		return Sequences[SearchIndexAsset.SourceAssetIdx].Sequence->GetName();
-	}
-	
-	if (SearchIndexAsset.Type == ESearchIndexAssetType::BlendSpace)
-	{
-		return BlendSpaces[SearchIndexAsset.SourceAssetIdx].BlendSpace->GetName();
-	}
-	
-	checkNoEntry();
-	return FString();
+	return GetAnimationAssetBase(SearchIndexAsset.SourceAssetIdx)->GetName();
 }
 
 int32 UPoseSearchDatabase::GetNumberOfPrincipalComponents() const
@@ -911,16 +887,49 @@ bool UPoseSearchDatabase::GetSkipSearchIfPossible() const
 	return bSkipSearchIfPossible;
 }
 
+bool FPoseSearchDatabaseSequence::IsRootMotionEnabled() const
+{
+	return Sequence ? Sequence->HasRootMotion() : false;
+}
+
 UAnimationAsset* FPoseSearchDatabaseBlendSpace::GetAnimationAsset() const
 {
 	return BlendSpace.Get();
 }
 
+UClass* FPoseSearchDatabaseBlendSpace::GetAnimationAssetStaticClass() const
+{
+	return UBlendSpace::StaticClass();
+}
+
 bool FPoseSearchDatabaseBlendSpace::IsLooping() const
 {
-	check(BlendSpace);
+	return BlendSpace ? BlendSpace->bLoop : false;
+}
 
-	return BlendSpace->bLoop;
+const FString FPoseSearchDatabaseBlendSpace::GetName() const
+{
+	return BlendSpace ? BlendSpace->GetName() : FString();
+}
+
+bool FPoseSearchDatabaseBlendSpace::IsRootMotionEnabled() const
+{
+	bool bIsRootMotionUsedInBlendSpace = false;
+
+	if (BlendSpace)
+	{
+		BlendSpace->ForEachImmutableSample([&bIsRootMotionUsedInBlendSpace](const FBlendSample& Sample)
+			{
+				const TObjectPtr<UAnimSequence> Sequence = Sample.Animation;
+
+				if (IsValid(Sequence) && Sequence->HasRootMotion())
+				{
+					bIsRootMotionUsedInBlendSpace = true;
+				}
+			});
+	}
+
+	return bIsRootMotionUsedInBlendSpace;
 }
 
 void FPoseSearchDatabaseBlendSpace::GetBlendSpaceParameterSampleRanges(int32& HorizontalBlendNum, int32& VerticalBlendNum) const
@@ -971,8 +980,27 @@ FVector FPoseSearchDatabaseBlendSpace::BlendParameterForSampleRanges(int32 Horiz
 		0.0f);
 }
 
+bool FPoseSearchDatabaseAnimComposite::IsRootMotionEnabled() const
+{
+	return AnimComposite ? AnimComposite->HasRootMotion() : false;
+}
+
 void UPoseSearchDatabase::PostLoad()
 {
+#if WITH_EDITORONLY_DATA
+	for (const FPoseSearchDatabaseSequence& DatabaseSequence : Sequences_DEPRECATED)
+	{
+		AnimationAssets.Add(FInstancedStruct::Make(DatabaseSequence));
+	}
+	Sequences_DEPRECATED.Empty();
+
+	for (const FPoseSearchDatabaseBlendSpace& DatabaseBlendSpace : BlendSpaces_DEPRECATED)
+	{
+		AnimationAssets.Add(FInstancedStruct::Make(DatabaseBlendSpace));
+	}
+	BlendSpaces_DEPRECATED.Empty();
+#endif // WITH_EDITORONLY_DATA
+
 #if WITH_EDITOR
 	using namespace UE::PoseSearch;
 	FAsyncPoseSearchDatabasesManagement::RequestAsyncBuildIndex(this, ERequestAsyncBuildFlag::NewRequest);
@@ -2046,7 +2074,8 @@ void FSearchResult::Update(float NewAssetTime)
 	else
 	{
 		const FPoseSearchIndexAsset& SearchIndexAsset = Database->GetSearchIndex().GetAssetForPose(PoseIdx);
-		if (SearchIndexAsset.Type == ESearchIndexAssetType::Sequence)
+		const FInstancedStruct& DatabaseAsset = Database->GetAnimationAssetStruct(SearchIndexAsset);
+		if (DatabaseAsset.GetPtr<FPoseSearchDatabaseSequence>() || DatabaseAsset.GetPtr<FPoseSearchDatabaseAnimComposite>())
 		{
 			if (Database->GetPoseIndicesAndLerpValueFromTime(NewAssetTime, SearchIndexAsset, PrevPoseIdx, PoseIdx, NextPoseIdx, LerpValue))
 			{
@@ -2057,15 +2086,13 @@ void FSearchResult::Update(float NewAssetTime)
 				Reset();
 			}
 		}
-		else if (SearchIndexAsset.Type == ESearchIndexAssetType::BlendSpace)
+		else if (const FPoseSearchDatabaseBlendSpace* DatabaseBlendSpace = DatabaseAsset.GetPtr<FPoseSearchDatabaseBlendSpace>())
 		{
-			const FPoseSearchDatabaseBlendSpace& DbBlendSpace = Database->GetBlendSpaceSourceAsset(SearchIndexAsset);
-
 			TArray<FBlendSampleData> BlendSamples;
 			int32 TriangulationIndex = 0;
-			DbBlendSpace.BlendSpace->GetSamplesFromBlendInput(SearchIndexAsset.BlendParameters, BlendSamples, TriangulationIndex, true);
+			DatabaseBlendSpace->BlendSpace->GetSamplesFromBlendInput(SearchIndexAsset.BlendParameters, BlendSamples, TriangulationIndex, true);
 
-			const float PlayLength = DbBlendSpace.BlendSpace->GetAnimationLengthFromSampleData(BlendSamples);
+			const float PlayLength = DatabaseBlendSpace->BlendSpace->GetAnimationLengthFromSampleData(BlendSamples);
 
 			// Asset player time for blendspaces is normalized [0, 1] so we need to convert 
 			// to a real time before we advance it
@@ -2298,22 +2325,21 @@ static FTransform ExtrapolateRootMotion(
 	return ExtrapolatedRootMotion;
 }
 
-
 //////////////////////////////////////////////////////////////////////////
-// FSequenceSampler
-void FSequenceSampler::Init(const FInput& InInput)
+// FSequenceBaseSampler
+void FSequenceBaseSampler::Init(const FInput& InInput)
 {
-	check(InInput.Sequence.Get());
+	check(InInput.SequenceBase.Get());
 
 	Input = InInput;
 }
 
-void FSequenceSampler::Process()
+void FSequenceBaseSampler::Process()
 {
 	ProcessRootDistance();
 }
 
-float FSequenceSampler::GetTimeFromRootDistance(float Distance) const
+float FSequenceBaseSampler::GetTimeFromRootDistance(float Distance) const
 {
 	int32 NextSampleIdx = 1;
 	int32 PrevSampleIdx = 0;
@@ -2339,54 +2365,55 @@ float FSequenceSampler::GetTimeFromRootDistance(float Distance) const
 	return ClipTime;
 }
 
-bool FSequenceSampler::IsLoopable() const
+bool FSequenceBaseSampler::IsLoopable() const
 {
-	return Input.Sequence->bLoop;
+	return Input.SequenceBase->bLoop;
 }
 
-void FSequenceSampler::ExtractPose(const FAnimExtractContext& ExtractionCtx, FAnimationPoseData& OutAnimPoseData) const
+void FSequenceBaseSampler::ExtractPose(const FAnimExtractContext& ExtractionCtx, FAnimationPoseData& OutAnimPoseData) const
 {
-	Input.Sequence->GetAnimationPose(OutAnimPoseData, ExtractionCtx);
+	Input.SequenceBase->GetAnimationPose(OutAnimPoseData, ExtractionCtx);
 }
 
-FTransform FSequenceSampler::ExtractRootTransform(float Time) const
+FTransform FSequenceBaseSampler::ExtractRootTransform(float Time) const
 {
+	FTransform RootTransform = FTransform::Identity;
+	const UAnimSequenceBase* SequenceBase = Input.SequenceBase.Get();
+
 	if (IsLoopable())
 	{
-		FTransform LoopableRootTransform = Input.Sequence->ExtractRootMotion(0.0f, Time, true);
+		FTransform LoopableRootTransform = SequenceBase->ExtractRootMotion(0.0f, Time, true);
 		return LoopableRootTransform;
 	}
 
 	const float ExtrapolationSampleTime = Input.ExtrapolationParameters.SampleTime;
 
-	const float PlayLength = Input.Sequence->GetPlayLength();
+	const float PlayLength = SequenceBase->GetPlayLength();
 	const float ClampedTime = FMath::Clamp(Time, 0.0f, PlayLength);
 	const float ExtrapolationTime = Time - ClampedTime;
 
-	FTransform RootTransform = FTransform::Identity;
-
 	// If Time is less than zero, ExtrapolationTime will be negative. In this case, we extrapolate the beginning of the 
 	// animation to estimate where the root would be at Time
 	if (ExtrapolationTime < -SMALL_NUMBER)
 	{
-		FTransform SampleToExtrapolate = Input.Sequence->ExtractRootMotionFromRange(0.0f, ExtrapolationSampleTime);
+		FTransform SampleToExtrapolate = SequenceBase->ExtractRootMotionFromRange(0.0f, ExtrapolationSampleTime);
 
 		const FTransform ExtrapolatedRootMotion = ExtrapolateRootMotion(
 			SampleToExtrapolate,
-			0.0f, ExtrapolationSampleTime, 
+			0.0f, ExtrapolationSampleTime,
 			ExtrapolationTime,
 			Input.ExtrapolationParameters);
 		RootTransform = ExtrapolatedRootMotion;
 	}
 	else
 	{
-		RootTransform = Input.Sequence->ExtractRootMotionFromRange(0.0f, ClampedTime);
+		RootTransform = SequenceBase->ExtractRootMotionFromRange(0.0f, ClampedTime);
 
 		// If Time is greater than PlayLength, ExtrapolationTIme will be a positive number. In this case, we extrapolate
 		// the end of the animation to estimate where the root would be at Time
 		if (ExtrapolationTime > SMALL_NUMBER)
 		{
-			FTransform SampleToExtrapolate = Input.Sequence->ExtractRootMotionFromRange(PlayLength - ExtrapolationSampleTime, PlayLength);
+			FTransform SampleToExtrapolate = SequenceBase->ExtractRootMotionFromRange(PlayLength - ExtrapolationSampleTime, PlayLength);
 
 			const FTransform ExtrapolatedRootMotion = ExtrapolateRootMotion(
 				SampleToExtrapolate,
@@ -2400,24 +2427,24 @@ FTransform FSequenceSampler::ExtractRootTransform(float Time) const
 	return RootTransform;
 }
 
-float FSequenceSampler::ExtractRootDistance(float Time) const
+float FSequenceBaseSampler::ExtractRootDistance(float Time) const
 {
 	return ExtractAccumulatedRootDistance(
 		Input.RootDistanceSamplingRate,
 		AccumulatedRootDistance,
-		Input.Sequence->GetPlayLength(),
+		Input.SequenceBase->GetPlayLength(),
 		Time,
 		Input.ExtrapolationParameters);
 }
 
-void FSequenceSampler::ExtractPoseSearchNotifyStates(
+void FSequenceBaseSampler::ExtractPoseSearchNotifyStates(
 	float Time, 
 	TArray<UAnimNotifyState_PoseSearchBase*>& NotifyStates) const
 {
 	// getting pose search notifies in an interval of size ExtractionInterval, centered on Time
 	constexpr float ExtractionInterval = 1.0f / 120.0f;
 	FAnimNotifyContext NotifyContext;
-	Input.Sequence->GetAnimNotifies(Time - (ExtractionInterval * 0.5f), ExtractionInterval, NotifyContext);
+	Input.SequenceBase->GetAnimNotifies(Time - (ExtractionInterval * 0.5f), ExtractionInterval, NotifyContext);
 
 	// check which notifies actually overlap Time and are of the right base type
 	for (const FAnimNotifyEventReference& EventReference : NotifyContext.ActiveNotifies)
@@ -2443,14 +2470,16 @@ void FSequenceSampler::ExtractPoseSearchNotifyStates(
 	}
 }
 
-void FSequenceSampler::ProcessRootDistance()
+void FSequenceBaseSampler::ProcessRootDistance()
 {
+	const UAnimSequenceBase* SequenceBase = Input.SequenceBase.Get();
+	
 	// Note the distance sampling interval is independent of the schema's sampling interval
 	const float DistanceSamplingInterval = 1.0f / Input.RootDistanceSamplingRate;
 
-	const FTransform InitialRootTransform = Input.Sequence->ExtractRootTrackTransform(0.0f, nullptr);
+	const FTransform InitialRootTransform = SequenceBase->ExtractRootTrackTransform(0.0f, nullptr);
 
-	uint32 NumDistanceSamples = FMath::CeilToInt(Input.Sequence->GetPlayLength() * Input.RootDistanceSamplingRate) + 1;
+	uint32 NumDistanceSamples = FMath::CeilToInt(SequenceBase->GetPlayLength() * Input.RootDistanceSamplingRate) + 1;
 	AccumulatedRootDistance.Reserve(NumDistanceSamples);
 
 	// Build a distance lookup table by sampling root motion at a fixed rate and accumulating
@@ -2463,9 +2492,9 @@ void FSequenceSampler::ProcessRootDistance()
 	float SampleTime = 0.0f;
 	for (int32 SampleIdx = 0; SampleIdx != NumDistanceSamples; ++SampleIdx)
 	{
-		SampleTime = FMath::Min(SampleIdx * DistanceSamplingInterval, Input.Sequence->GetPlayLength());
+		SampleTime = FMath::Min(SampleIdx * DistanceSamplingInterval, SequenceBase->GetPlayLength());
 
-		FTransform RootTransform = Input.Sequence->ExtractRootTrackTransform(SampleTime, nullptr);
+		FTransform RootTransform = SequenceBase->ExtractRootTrackTransform(SampleTime, nullptr);
 		FTransform LocalRootMotion = RootTransform.GetRelativeTransform(LastRootTransform);
 		LastRootTransform = RootTransform;
 
@@ -2474,7 +2503,7 @@ void FSequenceSampler::ProcessRootDistance()
 	}
 
 	// Verify we sampled the final frame of the clip
-	check(SampleTime == Input.Sequence->GetPlayLength());
+	check(SampleTime == SequenceBase->GetPlayLength());
 
 	// Also emit root motion summary info to help with sample wrapping in 
 	// FAssetIndexer::GetSampleTimeFromDistance() and FAssetIndexer::GetSampleInfo()
@@ -2482,9 +2511,9 @@ void FSequenceSampler::ProcessRootDistance()
 	TotalRootDistance = AccumulatedRootDistance.Last();
 }
 
-const UAnimationAsset* FSequenceSampler::GetAsset() const
+const UAnimationAsset* FSequenceBaseSampler::GetAsset() const
 {
-	return Input.Sequence.Get();
+	return Input.SequenceBase.Get();
 }
 
 //////////////////////////////////////////////////////////////////////////
diff --git a/Engine/Plugins/Experimental/Animation/PoseSearch/Source/Runtime/Private/PoseSearchDerivedData.cpp b/Engine/Plugins/Experimental/Animation/PoseSearch/Source/Runtime/Private/PoseSearchDerivedData.cpp
--- a/Engine/Plugins/Experimental/Animation/PoseSearch/Source/Runtime/Private/PoseSearchDerivedData.cpp
+++ b/Engine/Plugins/Experimental/Animation/PoseSearch/Source/Runtime/Private/PoseSearchDerivedData.cpp
@@ -154,16 +154,16 @@ static inline FFloatInterval GetEffectiveSamplingRange(const UAnimSequenceBase*
 	return Range;
 }
 
-static void FindValidSequenceIntervals(const FPoseSearchDatabaseSequence& DbSequence, const FPoseSearchExcludeFromDatabaseParameters& ExcludeFromDatabaseParameters, TArray<FFloatRange>& ValidRanges)
+static void FindValidSequenceIntervals(const UAnimSequenceBase* SequenceBase, FFloatInterval SamplingRange, bool bIsLooping,
+	const FPoseSearchExcludeFromDatabaseParameters& ExcludeFromDatabaseParameters, TArray<FFloatRange>& ValidRanges)
 {
-	const UAnimSequence* Sequence = DbSequence.Sequence;
-	check(DbSequence.Sequence);
+	check(SequenceBase);
 
-	const float SequenceLength = DbSequence.Sequence->GetPlayLength();
+	const float SequenceLength = SequenceBase->GetPlayLength();
 
-	const FFloatInterval EffectiveSamplingInterval = GetEffectiveSamplingRange(DbSequence.Sequence, DbSequence.SamplingRange);
+	const FFloatInterval EffectiveSamplingInterval = GetEffectiveSamplingRange(SequenceBase, SamplingRange);
 	FFloatRange EffectiveSamplingRange = FFloatRange::Inclusive(EffectiveSamplingInterval.Min, EffectiveSamplingInterval.Max);
-	if (!DbSequence.IsLooping())
+	if (!bIsLooping)
 	{
 		const FFloatRange ExcludeFromDatabaseRange(ExcludeFromDatabaseParameters.SequenceStartInterval, SequenceLength - ExcludeFromDatabaseParameters.SequenceEndInterval);
 		EffectiveSamplingRange = FFloatRange::Intersection(EffectiveSamplingRange, ExcludeFromDatabaseRange);
@@ -174,7 +174,7 @@ static void FindValidSequenceIntervals(const FPoseSearchDatabaseSequence& DbSequ
 	ValidRanges.Add(EffectiveSamplingRange);
 
 	FAnimNotifyContext NotifyContext;
-	Sequence->GetAnimNotifies(0.0f, SequenceLength, NotifyContext);
+	SequenceBase->GetAnimNotifies(0.0f, SequenceLength, NotifyContext);
 
 	for (const FAnimNotifyEventReference& EventReference : NotifyContext.ActiveNotifies)
 	{
@@ -199,78 +199,95 @@ static void FindValidSequenceIntervals(const FPoseSearchDatabaseSequence& DbSequ
 	}
 }
 
-static void InitSearchIndexAssets(FPoseSearchIndexBase& SearchIndex, TConstArrayView<FPoseSearchDatabaseSequence> Sequences, TConstArrayView<FPoseSearchDatabaseBlendSpace> BlendSpaces, const FPoseSearchExcludeFromDatabaseParameters& ExcludeFromDatabaseParameters)
+static void InitSearchIndexAssets(FPoseSearchIndexBase& SearchIndex, UPoseSearchDatabase* Database)
 {
 	using namespace UE::PoseSearch;
 
 	SearchIndex.Assets.Empty();
-
 	TArray<FFloatRange> ValidRanges;
+	TArray<FBlendSampleData> BlendSamples;
 
-	for (int32 SequenceIdx = 0; SequenceIdx < Sequences.Num(); ++SequenceIdx)
+	for (int32 AnimationAssetIndex = 0; AnimationAssetIndex < Database->AnimationAssets.Num(); ++AnimationAssetIndex)
 	{
-		const FPoseSearchDatabaseSequence& Sequence = Sequences[SequenceIdx];
-
-		if (Sequence.bEnabled)
+		const FInstancedStruct& DatabaseAssetStruct = Database->GetAnimationAssetStruct(AnimationAssetIndex);
+		if (const FPoseSearchDatabaseAnimationAssetBase* DatabaseAsset = DatabaseAssetStruct.GetPtr<FPoseSearchDatabaseAnimationAssetBase>())
 		{
-			const bool bAddUnmirrored = Sequence.MirrorOption == EPoseSearchMirrorOption::UnmirroredOnly || Sequence.MirrorOption == EPoseSearchMirrorOption::UnmirroredAndMirrored;
-			const bool bAddMirrored = Sequence.MirrorOption == EPoseSearchMirrorOption::MirroredOnly || Sequence.MirrorOption == EPoseSearchMirrorOption::UnmirroredAndMirrored;
+			if (!DatabaseAsset->IsEnabled() || !DatabaseAsset->GetAnimationAsset())
+			{
+				continue;
+			}
 
-			ValidRanges.Reset();
-			FindValidSequenceIntervals(Sequence, ExcludeFromDatabaseParameters, ValidRanges);
-			for (const FFloatRange& Range : ValidRanges)
+			const bool bAddUnmirrored = DatabaseAsset->GetMirrorOption() == EPoseSearchMirrorOption::UnmirroredOnly || DatabaseAsset->GetMirrorOption() == EPoseSearchMirrorOption::UnmirroredAndMirrored;
+			const bool bAddMirrored = DatabaseAsset->GetMirrorOption() == EPoseSearchMirrorOption::MirroredOnly || DatabaseAsset->GetMirrorOption() == EPoseSearchMirrorOption::UnmirroredAndMirrored;
+
+			if (const FPoseSearchDatabaseSequence* DatabaseSequence = DatabaseAssetStruct.GetPtr<FPoseSearchDatabaseSequence>())
 			{
-				if (bAddUnmirrored)
+				ValidRanges.Reset();
+				FindValidSequenceIntervals(DatabaseSequence->Sequence, DatabaseSequence->SamplingRange, DatabaseSequence->IsLooping(), Database->ExcludeFromDatabaseParameters, ValidRanges);
+				for (const FFloatRange& Range : ValidRanges)
 				{
-					SearchIndex.Assets.Add(FPoseSearchIndexAsset(ESearchIndexAssetType::Sequence, SequenceIdx, false, FFloatInterval(Range.GetLowerBoundValue(), Range.GetUpperBoundValue())));
-				}
+					if (bAddUnmirrored)
+					{
+						SearchIndex.Assets.Add(FPoseSearchIndexAsset(ESearchIndexAssetType::Sequence, AnimationAssetIndex, false, FFloatInterval(Range.GetLowerBoundValue(), Range.GetUpperBoundValue())));
+					}
 
-				if (bAddMirrored)
-				{
-					SearchIndex.Assets.Add(FPoseSearchIndexAsset(ESearchIndexAssetType::Sequence, SequenceIdx, true, FFloatInterval(Range.GetLowerBoundValue(), Range.GetUpperBoundValue())));
+					if (bAddMirrored)
+					{
+						SearchIndex.Assets.Add(FPoseSearchIndexAsset(ESearchIndexAssetType::Sequence, AnimationAssetIndex, true, FFloatInterval(Range.GetLowerBoundValue(), Range.GetUpperBoundValue())));
+					}
 				}
 			}
-		}
-	}
-
-	TArray<FBlendSampleData> BlendSamples;
-
-	for (int32 BlendSpaceIdx = 0; BlendSpaceIdx < BlendSpaces.Num(); ++BlendSpaceIdx)
-	{
-		const FPoseSearchDatabaseBlendSpace& BlendSpace = BlendSpaces[BlendSpaceIdx];
-
-		if (BlendSpace.bEnabled)
-		{
-			const bool bAddUnmirrored = BlendSpace.MirrorOption == EPoseSearchMirrorOption::UnmirroredOnly || BlendSpace.MirrorOption == EPoseSearchMirrorOption::UnmirroredAndMirrored;
-			const bool bAddMirrored = BlendSpace.MirrorOption == EPoseSearchMirrorOption::MirroredOnly || BlendSpace.MirrorOption == EPoseSearchMirrorOption::UnmirroredAndMirrored;
-
-			int32 HorizontalBlendNum, VerticalBlendNum;
-			BlendSpace.GetBlendSpaceParameterSampleRanges(HorizontalBlendNum, VerticalBlendNum);
-
-			const bool bWrapInputOnHorizontalAxis = BlendSpace.BlendSpace->GetBlendParameter(0).bWrapInput;
-			const bool bWrapInputOnVerticalAxis = BlendSpace.BlendSpace->GetBlendParameter(1).bWrapInput;
-			for (int32 HorizontalIndex = 0; HorizontalIndex < HorizontalBlendNum; HorizontalIndex++)
+			else if (const FPoseSearchDatabaseAnimComposite* DatabaseAnimComposite = DatabaseAssetStruct.GetPtr<FPoseSearchDatabaseAnimComposite>())
 			{
-				for (int32 VerticalIndex = 0; VerticalIndex < VerticalBlendNum; VerticalIndex++)
+				ValidRanges.Reset();
+				FindValidSequenceIntervals(DatabaseAnimComposite->AnimComposite, DatabaseAnimComposite->SamplingRange, DatabaseAnimComposite->IsLooping(), Database->ExcludeFromDatabaseParameters, ValidRanges);
+				for (const FFloatRange& Range : ValidRanges)
 				{
-					const FVector BlendParameters = BlendSpace.BlendParameterForSampleRanges(HorizontalIndex, VerticalIndex);
-
-					int32 TriangulationIndex = 0;
-					BlendSpace.BlendSpace->GetSamplesFromBlendInput(BlendParameters, BlendSamples, TriangulationIndex, true);
-
-					float PlayLength = BlendSpace.BlendSpace->GetAnimationLengthFromSampleData(BlendSamples);
-
 					if (bAddUnmirrored)
 					{
-						SearchIndex.Assets.Add(FPoseSearchIndexAsset(ESearchIndexAssetType::BlendSpace, BlendSpaceIdx, false, FFloatInterval(0.0f, PlayLength), BlendParameters));
+						SearchIndex.Assets.Add(FPoseSearchIndexAsset(ESearchIndexAssetType::AnimComposite, AnimationAssetIndex, false, FFloatInterval(Range.GetLowerBoundValue(), Range.GetUpperBoundValue())));
 					}
 
 					if (bAddMirrored)
 					{
-						SearchIndex.Assets.Add(FPoseSearchIndexAsset(ESearchIndexAssetType::BlendSpace, BlendSpaceIdx, true, FFloatInterval(0.0f, PlayLength), BlendParameters));
+						SearchIndex.Assets.Add(FPoseSearchIndexAsset(ESearchIndexAssetType::AnimComposite, AnimationAssetIndex, true, FFloatInterval(Range.GetLowerBoundValue(), Range.GetUpperBoundValue())));
 					}
 				}
 			}
+			else if (const FPoseSearchDatabaseBlendSpace* DatabaseBlendSpace = DatabaseAssetStruct.GetPtr<FPoseSearchDatabaseBlendSpace>())
+			{
+				int32 HorizontalBlendNum, VerticalBlendNum;
+				DatabaseBlendSpace->GetBlendSpaceParameterSampleRanges(HorizontalBlendNum, VerticalBlendNum);
+
+				const bool bWrapInputOnHorizontalAxis = DatabaseBlendSpace->BlendSpace->GetBlendParameter(0).bWrapInput;
+				const bool bWrapInputOnVerticalAxis = DatabaseBlendSpace->BlendSpace->GetBlendParameter(1).bWrapInput;
+				for (int32 HorizontalIndex = 0; HorizontalIndex < HorizontalBlendNum; HorizontalIndex++)
+				{
+					for (int32 VerticalIndex = 0; VerticalIndex < VerticalBlendNum; VerticalIndex++)
+					{
+						const FVector BlendParameters = DatabaseBlendSpace->BlendParameterForSampleRanges(HorizontalIndex, VerticalIndex);
+
+						int32 TriangulationIndex = 0;
+						DatabaseBlendSpace->BlendSpace->GetSamplesFromBlendInput(BlendParameters, BlendSamples, TriangulationIndex, true);
+
+						float PlayLength = DatabaseBlendSpace->BlendSpace->GetAnimationLengthFromSampleData(BlendSamples);
+
+						if (bAddUnmirrored)
+						{
+							SearchIndex.Assets.Add(FPoseSearchIndexAsset(ESearchIndexAssetType::BlendSpace, AnimationAssetIndex, false, FFloatInterval(0.0f, PlayLength), BlendParameters));
+						}
+
+						if (bAddMirrored)
+						{
+							SearchIndex.Assets.Add(FPoseSearchIndexAsset(ESearchIndexAssetType::BlendSpace, AnimationAssetIndex, true, FFloatInterval(0.0f, PlayLength), BlendParameters));
+						}
+					}
+				}
+			}
+			else
+			{
+				checkNoEntry();
+			}
 		}
 	}
 }
@@ -1011,77 +1028,83 @@ struct FDatabaseIndexingContext
 	FPoseSearchIndexBase* SearchIndexBase = nullptr;
 
 	FAssetSamplingContext SamplingContext;
-	TArray<FSequenceSampler> SequenceSamplers;
+	TArray<FSequenceBaseSampler> SequenceSamplers; // Composite and sequence samplers
 	TArray<FBlendSpaceSampler> BlendSpaceSamplers;
 
 	TArray<FAssetIndexer> Indexers;
 
-	void Prepare(const UPoseSearchSchema* Schema, const FPoseSearchExtrapolationParameters& ExtrapolationParameters, TConstArrayView<FPoseSearchDatabaseSequence> Sequences, TConstArrayView<FPoseSearchDatabaseBlendSpace> BlendSpaces);
+	void Prepare(const UPoseSearchDatabase* Database);
 	bool IndexAssets();
 	void JoinIndex();
 	float CalculateMinCostAddend() const;
 };
 
-void FDatabaseIndexingContext::Prepare(const UPoseSearchSchema* Schema, const FPoseSearchExtrapolationParameters& ExtrapolationParameters, TConstArrayView<FPoseSearchDatabaseSequence> Sequences, TConstArrayView<FPoseSearchDatabaseBlendSpace> BlendSpaces)
+void FDatabaseIndexingContext::Prepare(const UPoseSearchDatabase* Database)
 {
+	const UPoseSearchSchema* Schema = Database->Schema;
 	check(Schema);
 
 	FBoneContainer BoneContainer;
 	BoneContainer.InitializeTo(Schema->BoneIndicesWithParents, FCurveEvaluationOption(false), *Schema->Skeleton);
 
-	TMap<const UAnimSequence*, int32> SequenceSamplerMap;
+	TMap<const UAnimSequenceBase*, int32> SequenceSamplerMap;
 	TMap<TPair<const UBlendSpace*, FVector>, int32> BlendSpaceSamplerMap;
 
 	SamplingContext.Init(Schema->MirrorDataTable, BoneContainer);
 
-	// Prepare samplers for all sequences
-	auto AddSequenceSampler = [&](const UAnimSequence* Sequence)
+	// Prepare samplers for all animation assets.
+	for (const FInstancedStruct& DatabaseAssetStruct : Database->AnimationAssets)
 	{
-		if (Sequence && !SequenceSamplerMap.Contains(Sequence))
+		auto AddSequenceBaseSampler = [&](const UAnimSequenceBase* Sequence)
 		{
-			int32 SequenceSamplerIdx = SequenceSamplers.AddDefaulted();
-			SequenceSamplerMap.Add(Sequence, SequenceSamplerIdx);
-
-			FSequenceSampler::FInput Input;
-			Input.ExtrapolationParameters = ExtrapolationParameters;
-			Input.Sequence = Sequence;
-			SequenceSamplers[SequenceSamplerIdx].Init(Input);
-		}
-	};
+			if (Sequence && !SequenceSamplerMap.Contains(Sequence))
+			{
+				int32 SequenceSamplerIdx = SequenceSamplers.AddDefaulted();
+				SequenceSamplerMap.Add(Sequence, SequenceSamplerIdx);
 
-	for (const FPoseSearchDatabaseSequence& DbSequence : Sequences)
-	{
-		AddSequenceSampler(DbSequence.Sequence);
-		AddSequenceSampler(DbSequence.LeadInSequence);
-		AddSequenceSampler(DbSequence.FollowUpSequence);
-	}
+				FSequenceBaseSampler::FInput Input;
+				Input.ExtrapolationParameters = Database->ExtrapolationParameters;
+				Input.SequenceBase = Sequence;
+				SequenceSamplers[SequenceSamplerIdx].Init(Input);
+			}
+		};
 
-	// Prepare samplers for all blend spaces
-	for (const FPoseSearchDatabaseBlendSpace& DbBlendSpace : BlendSpaces)
-	{
-		if (DbBlendSpace.BlendSpace)
+		if (const FPoseSearchDatabaseSequence* DatabaseSequence = DatabaseAssetStruct.GetPtr<FPoseSearchDatabaseSequence>())
 		{
-			int32 HorizontalBlendNum, VerticalBlendNum;
-			DbBlendSpace.GetBlendSpaceParameterSampleRanges(HorizontalBlendNum, VerticalBlendNum);
-
-			for (int32 HorizontalIndex = 0; HorizontalIndex < HorizontalBlendNum; HorizontalIndex++)
+			AddSequenceBaseSampler(DatabaseSequence->Sequence);
+			AddSequenceBaseSampler(DatabaseSequence->LeadInSequence);
+			AddSequenceBaseSampler(DatabaseSequence->FollowUpSequence);
+		}
+		else if (const FPoseSearchDatabaseAnimComposite* DatabaseAnimComposite = DatabaseAssetStruct.GetPtr<FPoseSearchDatabaseAnimComposite>())
+		{
+			AddSequenceBaseSampler(DatabaseAnimComposite->AnimComposite);
+		}
+		else if (const FPoseSearchDatabaseBlendSpace* DatabaseBlendSpace = DatabaseAssetStruct.GetPtr<FPoseSearchDatabaseBlendSpace>())
+		{
+			if (DatabaseBlendSpace->BlendSpace)
 			{
-				for (int32 VerticalIndex = 0; VerticalIndex < VerticalBlendNum; VerticalIndex++)
-				{
-					const FVector BlendParameters = DbBlendSpace.BlendParameterForSampleRanges(HorizontalIndex, VerticalIndex);
+				int32 HorizontalBlendNum, VerticalBlendNum;
+				DatabaseBlendSpace->GetBlendSpaceParameterSampleRanges(HorizontalBlendNum, VerticalBlendNum);
 
-					if (!BlendSpaceSamplerMap.Contains({ DbBlendSpace.BlendSpace, BlendParameters }))
+				for (int32 HorizontalIndex = 0; HorizontalIndex < HorizontalBlendNum; HorizontalIndex++)
+				{
+					for (int32 VerticalIndex = 0; VerticalIndex < VerticalBlendNum; VerticalIndex++)
 					{
-						int32 BlendSpaceSamplerIdx = BlendSpaceSamplers.AddDefaulted();
-						BlendSpaceSamplerMap.Add({ DbBlendSpace.BlendSpace, BlendParameters }, BlendSpaceSamplerIdx);
+						const FVector BlendParameters = DatabaseBlendSpace->BlendParameterForSampleRanges(HorizontalIndex, VerticalIndex);
+
+						if (!BlendSpaceSamplerMap.Contains({ DatabaseBlendSpace->BlendSpace, BlendParameters }))
+						{
+							int32 BlendSpaceSamplerIdx = BlendSpaceSamplers.AddDefaulted();
+							BlendSpaceSamplerMap.Add({ DatabaseBlendSpace->BlendSpace, BlendParameters }, BlendSpaceSamplerIdx);
 
-						FBlendSpaceSampler::FInput Input;
-						Input.BoneContainer = BoneContainer;
-						Input.ExtrapolationParameters = ExtrapolationParameters;
-						Input.BlendSpace = DbBlendSpace.BlendSpace;
-						Input.BlendParameters = BlendParameters;
+							FBlendSpaceSampler::FInput Input;
+							Input.BoneContainer = BoneContainer;
+							Input.ExtrapolationParameters = Database->ExtrapolationParameters;
+							Input.BlendSpace = DatabaseBlendSpace->BlendSpace;
+							Input.BlendParameters = BlendParameters;
 
-						BlendSpaceSamplers[BlendSpaceSamplerIdx].Init(Input);
+							BlendSpaceSamplers[BlendSpaceSamplerIdx].Init(Input);
+						}
 					}
 				}
 			}
@@ -1105,7 +1128,7 @@ void FDatabaseIndexingContext::Prepare(const UPoseSearchSchema* Schema, const FP
 	// prepare indexers
 	Indexers.Reserve(SearchIndexBase->Assets.Num());
 
-	auto GetSequenceSampler = [&](const UAnimSequence* Sequence) -> const FSequenceSampler*
+	auto GetSequenceBaseSampler = [&](const UAnimSequenceBase* Sequence) -> const FSequenceBaseSampler*
 	{
 		return Sequence ? &SequenceSamplers[SequenceSamplerMap[Sequence]] : nullptr;
 	};
@@ -1127,22 +1150,21 @@ void FDatabaseIndexingContext::Prepare(const UPoseSearchSchema* Schema, const FP
 		IndexerContext.RequestedSamplingRange = SearchIndexAsset.SamplingInterval;
 		IndexerContext.bMirrored = SearchIndexAsset.bMirrored;
 
-		if (SearchIndexAsset.Type == ESearchIndexAssetType::Sequence)
+		const FInstancedStruct& DatabaseAsset = Database->GetAnimationAssetStruct(SearchIndexAsset.SourceAssetIdx);
+		if (const FPoseSearchDatabaseSequence* DatabaseSequence = DatabaseAsset.GetPtr<FPoseSearchDatabaseSequence>())
 		{
-			const FPoseSearchDatabaseSequence& DbSequence = Sequences[SearchIndexAsset.SourceAssetIdx];
-			const float SequenceLength = DbSequence.Sequence->GetPlayLength();
-			IndexerContext.MainSampler = GetSequenceSampler(DbSequence.Sequence);
-			IndexerContext.LeadInSampler = SearchIndexAsset.SamplingInterval.Min == 0.0f ? GetSequenceSampler(DbSequence.LeadInSequence) : nullptr;
-			IndexerContext.FollowUpSampler = SearchIndexAsset.SamplingInterval.Max == SequenceLength ? GetSequenceSampler(DbSequence.FollowUpSequence) : nullptr;
+			const float SequenceLength = DatabaseSequence->Sequence->GetPlayLength();
+			IndexerContext.MainSampler = GetSequenceBaseSampler(DatabaseSequence->Sequence);
+			IndexerContext.LeadInSampler = SearchIndexAsset.SamplingInterval.Min == 0.0f ? GetSequenceBaseSampler(DatabaseSequence->LeadInSequence) : nullptr;
+			IndexerContext.FollowUpSampler = SearchIndexAsset.SamplingInterval.Max == SequenceLength ? GetSequenceBaseSampler(DatabaseSequence->FollowUpSequence) : nullptr;
 		}
-		else if (SearchIndexAsset.Type == ESearchIndexAssetType::BlendSpace)
+		else if (const FPoseSearchDatabaseAnimComposite* DatabaseAnimComposite = DatabaseAsset.GetPtr<FPoseSearchDatabaseAnimComposite>())
 		{
-			const FPoseSearchDatabaseBlendSpace& DbBlendSpace = BlendSpaces[SearchIndexAsset.SourceAssetIdx];
-			IndexerContext.MainSampler = GetBlendSpaceSampler(DbBlendSpace.BlendSpace, SearchIndexAsset.BlendParameters);
+			IndexerContext.MainSampler = GetSequenceBaseSampler(DatabaseAnimComposite->AnimComposite);
 		}
-		else
+		else if (const FPoseSearchDatabaseBlendSpace* DatabaseBlendSpace = DatabaseAsset.GetPtr<FPoseSearchDatabaseBlendSpace>())
 		{
-			checkNoEntry();
+			IndexerContext.MainSampler = GetBlendSpaceSampler(DatabaseBlendSpace->BlendSpace, SearchIndexAsset.BlendParameters);
 		}
 
 		FAssetIndexer& Indexer = Indexers.AddDefaulted_GetRef();
@@ -1409,6 +1431,8 @@ void FPoseSearchDatabaseAsyncCacheTask::OnGetComplete(UE::DerivedData::FCacheGet
 	using namespace UE::DerivedData;
 
 	const FCacheKey FullIndexKey = Response.Record.GetKey();
+
+	// The database is part of the derived data cache and up to date, skip re-building it.
 	if (Response.Status == EStatus::Ok)
 	{
 		COOK_STAT(auto Timer = UsageStats.TimeAsyncWait());
@@ -1483,7 +1507,7 @@ void FPoseSearchDatabaseAsyncCacheTask::OnGetComplete(UE::DerivedData::FCacheGet
 					}
 
 					// Building all the related FPoseSearchBaseIndex first
-					InitSearchIndexAssets(SearchIndexBase, IndexBaseDatabase->Sequences, IndexBaseDatabase->BlendSpaces, IndexBaseDatabase->ExcludeFromDatabaseParameters);
+					InitSearchIndexAssets(SearchIndexBase, Database.Get());
 
 					if (Owner.IsCanceled())
 					{
@@ -1494,7 +1518,7 @@ void FPoseSearchDatabaseAsyncCacheTask::OnGetComplete(UE::DerivedData::FCacheGet
 
 					FDatabaseIndexingContext DbIndexingContext;
 					DbIndexingContext.SearchIndexBase = &SearchIndexBase;
-					DbIndexingContext.Prepare(IndexBaseDatabase->Schema, IndexBaseDatabase->ExtrapolationParameters, IndexBaseDatabase->Sequences, IndexBaseDatabase->BlendSpaces);
+					DbIndexingContext.Prepare(IndexBaseDatabase.Get());
 
 					if (Owner.IsCanceled())
 					{
diff --git a/Engine/Plugins/Experimental/Animation/PoseSearch/Source/Runtime/Private/PoseSearchLibrary.cpp b/Engine/Plugins/Experimental/Animation/PoseSearch/Source/Runtime/Private/PoseSearchLibrary.cpp
--- a/Engine/Plugins/Experimental/Animation/PoseSearch/Source/Runtime/Private/PoseSearchLibrary.cpp
+++ b/Engine/Plugins/Experimental/Animation/PoseSearch/Source/Runtime/Private/PoseSearchLibrary.cpp
@@ -44,17 +44,16 @@ bool FMotionMatchingState::CanAdvance(float DeltaTime, bool& bOutAdvanceToFollow
 		return false;
 	}
 
-	
 	const FPoseSearchIndexAsset* SearchIndexAsset = CurrentSearchResult.GetSearchIndexAsset(true);
 
-	if (SearchIndexAsset->Type == ESearchIndexAssetType::Sequence)
+	const FInstancedStruct& DatabaseAsset = CurrentSearchResult.Database->GetAnimationAssetStruct(*SearchIndexAsset);
+	if (const FPoseSearchDatabaseSequence* DatabaseSequence = DatabaseAsset.GetPtr<FPoseSearchDatabaseSequence>())
 	{
-		const FPoseSearchDatabaseSequence& DbSequence = CurrentSearchResult.Database->GetSequenceSourceAsset(*SearchIndexAsset);
-		const float AssetLength = DbSequence.Sequence->GetPlayLength();
+		const float AssetLength = DatabaseSequence->Sequence->GetPlayLength();
 
 		float SteppedTime = CurrentSearchResult.AssetTime;
 		ETypeAdvanceAnim AdvanceType = FAnimationRuntime::AdvanceTime(
-			DbSequence.Sequence->bLoop,
+			DatabaseSequence->Sequence->bLoop,
 			DeltaTime,
 			SteppedTime,
 			AssetLength);
@@ -66,10 +65,15 @@ bool FMotionMatchingState::CanAdvance(float DeltaTime, bool& bOutAdvanceToFollow
 		else
 		{
 			// check if there's a follow-up that can be used
-			int32 FollowUpDbSequenceIdx = CurrentSearchResult.Database->Sequences.IndexOfByPredicate(
-				[&](const FPoseSearchDatabaseSequence& Entry)
+			int32 FollowUpDbSequenceIdx = CurrentSearchResult.Database->AnimationAssets.IndexOfByPredicate(
+				[&](const FInstancedStruct& DatabaseAsset)
 				{
-					return Entry.Sequence == DbSequence.FollowUpSequence;
+					if (const FPoseSearchDatabaseSequence* DatabaseSequence = DatabaseAsset.GetPtr<FPoseSearchDatabaseSequence>())
+					{
+						return (DatabaseSequence->Sequence == DatabaseSequence->FollowUpSequence);
+					}
+
+					return false;
 				});
 
 			int32 FollowUpSearchIndexAssetIdx = CurrentSearchResult.Database->GetSearchIndex().Assets.IndexOfByPredicate(
@@ -102,22 +106,36 @@ bool FMotionMatchingState::CanAdvance(float DeltaTime, bool& bOutAdvanceToFollow
 			}
 		}
 	}
-	else if (SearchIndexAsset->Type == ESearchIndexAssetType::BlendSpace)
+	else if (const FPoseSearchDatabaseAnimComposite* DatabaseAnimComposite = DatabaseAsset.GetPtr<FPoseSearchDatabaseAnimComposite>())
 	{
-		const FPoseSearchDatabaseBlendSpace& DbBlendSpace = CurrentSearchResult.Database->GetBlendSpaceSourceAsset(*SearchIndexAsset);
+		const float AssetLength = DatabaseAnimComposite->GetAnimationAsset()->GetPlayLength();
 
+		float SteppedTime = CurrentSearchResult.AssetTime;
+		ETypeAdvanceAnim AdvanceType = FAnimationRuntime::AdvanceTime(
+			DatabaseAnimComposite->IsLooping(),
+			DeltaTime,
+			SteppedTime,
+			AssetLength);
+
+		if (AdvanceType != ETAA_Finished)
+		{
+			return SearchIndexAsset->SamplingInterval.Contains(SteppedTime);
+		}
+	}
+	else if (const FPoseSearchDatabaseBlendSpace* DatabaseBlendSpace = DatabaseAsset.GetPtr<FPoseSearchDatabaseBlendSpace>())
+	{
 		TArray<FBlendSampleData> BlendSamples;
 		int32 TriangulationIndex = 0;
-		DbBlendSpace.BlendSpace->GetSamplesFromBlendInput(SearchIndexAsset->BlendParameters, BlendSamples, TriangulationIndex, true);
+		DatabaseBlendSpace->BlendSpace->GetSamplesFromBlendInput(SearchIndexAsset->BlendParameters, BlendSamples, TriangulationIndex, true);
 
-		float PlayLength = DbBlendSpace.BlendSpace->GetAnimationLengthFromSampleData(BlendSamples);
+		float PlayLength = DatabaseBlendSpace->BlendSpace->GetAnimationLengthFromSampleData(BlendSamples);
 
 		// Asset player time for blendspaces is normalized [0, 1] so we need to convert 
 		// to a real time before we advance it
 		float RealTime = CurrentSearchResult.AssetTime * PlayLength;
 		float SteppedTime = RealTime;
 		ETypeAdvanceAnim AdvanceType = FAnimationRuntime::AdvanceTime(
-			DbBlendSpace.BlendSpace->bLoop,
+			DatabaseBlendSpace->BlendSpace->bLoop,
 			DeltaTime,
 			SteppedTime,
 			PlayLength);
diff --git a/Engine/Plugins/Experimental/Animation/PoseSearch/Source/Runtime/Public/PoseSearch/PoseSearch.h b/Engine/Plugins/Experimental/Animation/PoseSearch/Source/Runtime/Public/PoseSearch/PoseSearch.h
--- a/Engine/Plugins/Experimental/Animation/PoseSearch/Source/Runtime/Public/PoseSearch/PoseSearch.h
+++ b/Engine/Plugins/Experimental/Animation/PoseSearch/Source/Runtime/Public/PoseSearch/PoseSearch.h
@@ -7,12 +7,14 @@
 #include "CoreMinimal.h"
 #include "Containers/RingBuffer.h"
 #include "Engine/DataAsset.h"
+#include "InstancedStruct.h"
 #include "IO/IoHash.h"
 #include "Modules/ModuleManager.h"
 #include "Misc/EnumClassFlags.h"
 #include "UObject/Object.h"
 #include "UObject/ObjectMacros.h"
 #include "Templates/TypeHash.h"
+#include "Animation/AnimComposite.h"
 #include "Animation/AnimMetaData.h"
 #include "Animation/AnimNodeMessages.h"
 #include "Animation/AnimationAsset.h"
@@ -123,6 +125,7 @@ enum class ESearchIndexAssetType : int32
 	Invalid,
 	Sequence,
 	BlendSpace,
+	AnimComposite,
 };
 
 namespace UE::PoseSearch
@@ -249,12 +252,13 @@ public:
 	virtual void Process() = 0;
 };
 
-struct POSESEARCH_API FSequenceSampler : public IAssetSampler
+// Sampler working with UAnimSequenceBase so it can be used for UAnimSequence as well as UAnimComposite.
+struct POSESEARCH_API FSequenceBaseSampler : public IAssetSampler
 {
 public:
 	struct FInput
 	{
-		TWeakObjectPtr<const UAnimSequence> Sequence;
+		TWeakObjectPtr<const UAnimSequenceBase> SequenceBase;
 		int32 RootDistanceSamplingRate = 60;
 		FPoseSearchExtrapolationParameters ExtrapolationParameters;
 	} Input;
@@ -262,7 +266,7 @@ public:
 	void Init(const FInput& Input);
 	virtual void Process() override;
 
-	virtual float GetPlayLength() const override { return Input.Sequence->GetPlayLength(); }
+	virtual float GetPlayLength() const override { return Input.SequenceBase->GetPlayLength(); }
 	virtual bool IsLoopable() const override;
 
 	virtual float GetTimeFromRootDistance(float Distance) const override;
@@ -834,10 +838,17 @@ USTRUCT()
 struct POSESEARCH_API FPoseSearchDatabaseAnimationAssetBase
 {
 	GENERATED_BODY()
+	virtual ~FPoseSearchDatabaseAnimationAssetBase() = default;
 
-	virtual ~FPoseSearchDatabaseAnimationAssetBase() {}
 	virtual UAnimationAsset* GetAnimationAsset() const { return nullptr; }
+	virtual UClass* GetAnimationAssetStaticClass() const { return nullptr; }
 	virtual bool IsLooping() const { return false; }
+	virtual const FString GetName() const { return {}; }
+	virtual bool IsEnabled() const { return false; }
+	virtual void SetIsEnabled(bool bInIsEnabled) {}
+	virtual bool IsRootMotionEnabled() const { return false; }
+	virtual EPoseSearchMirrorOption GetMirrorOption() const { return EPoseSearchMirrorOption::Invalid; }
+	virtual ESearchIndexAssetType GetSearchIndexType() const { return ESearchIndexAssetType::Invalid; }
 };
 
 /** An entry in a UPoseSearchDatabase. */
@@ -845,6 +856,7 @@ USTRUCT(BlueprintType, Category = "Animation|Pose Search")
 struct POSESEARCH_API FPoseSearchDatabaseSequence : public FPoseSearchDatabaseAnimationAssetBase
 {
 	GENERATED_BODY()
+	virtual ~FPoseSearchDatabaseSequence() = default;
 
 	UPROPERTY(EditAnywhere, Category="Sequence")
 	TObjectPtr<UAnimSequence> Sequence;
@@ -872,8 +884,15 @@ struct POSESEARCH_API FPoseSearchDatabaseSequence : public FPoseSearchDatabaseAn
 	UPROPERTY(EditAnywhere, Category="Sequence")
 	TObjectPtr<UAnimSequence> FollowUpSequence;
 
-	virtual UAnimationAsset* GetAnimationAsset() const override { return Sequence; }
-	virtual bool IsLooping() const override { return Sequence->bLoop; }
+	UAnimationAsset* GetAnimationAsset() const override { return Sequence; }
+	UClass* GetAnimationAssetStaticClass() const override { return UAnimSequence::StaticClass(); }
+	bool IsLooping() const override { return Sequence ? Sequence->bLoop : false; }
+	const FString GetName() const override { return Sequence ? Sequence->GetName() : FString(); }
+	bool IsEnabled() const override { return bEnabled; }
+	void SetIsEnabled(bool bInIsEnabled) override { bEnabled = bInIsEnabled; }
+	bool IsRootMotionEnabled() const override;
+	EPoseSearchMirrorOption GetMirrorOption() const override { return MirrorOption; }
+	ESearchIndexAssetType GetSearchIndexType() const override { return ESearchIndexAssetType::Sequence; }
 };
 
 /** An blend space entry in a UPoseSearchDatabase. */
@@ -881,6 +900,7 @@ USTRUCT(BlueprintType, Category = "Animation|Pose Search")
 struct POSESEARCH_API FPoseSearchDatabaseBlendSpace : public FPoseSearchDatabaseAnimationAssetBase
 {
 	GENERATED_BODY()
+	virtual ~FPoseSearchDatabaseBlendSpace() = default;
 
 	UPROPERTY(EditAnywhere, Category = "BlendSpace")
 	TObjectPtr<UBlendSpace> BlendSpace;
@@ -902,13 +922,50 @@ struct POSESEARCH_API FPoseSearchDatabaseBlendSpace : public FPoseSearchDatabase
 	UPROPERTY(EditAnywhere, Category = "BlendSpace", meta = (EditCondition = "!bUseGridForSampling", EditConditionHides, ClampMin = "1", UIMin = "1", UIMax = "25"))
 	int32 NumberOfVerticalSamples = 2;
 
-	virtual UAnimationAsset* GetAnimationAsset() const override;
-	virtual bool IsLooping() const override;
+	UAnimationAsset* GetAnimationAsset() const override;
+	UClass* GetAnimationAssetStaticClass() const override;
+	bool IsLooping() const override;
+	const FString GetName() const override;
+	bool IsEnabled() const override { return bEnabled; }
+	void SetIsEnabled(bool bInIsEnabled) override { bEnabled = bInIsEnabled; }
+	bool IsRootMotionEnabled() const override;
+	EPoseSearchMirrorOption GetMirrorOption() const override { return MirrorOption; }
+	ESearchIndexAssetType GetSearchIndexType() const override { return ESearchIndexAssetType::BlendSpace; }
 
 	void GetBlendSpaceParameterSampleRanges(int32& HorizontalBlendNum, int32& VerticalBlendNum) const;
 	FVector BlendParameterForSampleRanges(int32 HorizontalBlendIndex, int32 VerticalBlendIndex) const;
 };
 
+/** An entry in a UPoseSearchDatabase. */
+USTRUCT(BlueprintType, Category = "Animation|Pose Search")
+struct POSESEARCH_API FPoseSearchDatabaseAnimComposite : public FPoseSearchDatabaseAnimationAssetBase
+{
+	GENERATED_BODY()
+	virtual ~FPoseSearchDatabaseAnimComposite() = default;
+
+	UPROPERTY(EditAnywhere, Category = "AnimComposite")
+	TObjectPtr<UAnimComposite> AnimComposite;
+
+	UPROPERTY(EditAnywhere, Category = "AnimComposite")
+	bool bEnabled = true;
+
+	UPROPERTY(EditAnywhere, Category = "AnimComposite")
+	FFloatInterval SamplingRange = FFloatInterval(0.0f, 0.0f);
+
+	UPROPERTY(EditAnywhere, Category = "AnimComposite")
+	EPoseSearchMirrorOption MirrorOption = EPoseSearchMirrorOption::UnmirroredOnly;
+
+	UAnimationAsset* GetAnimationAsset() const override { return AnimComposite; }
+	UClass* GetAnimationAssetStaticClass() const override { return UAnimComposite::StaticClass(); }
+	bool IsLooping() const override { return AnimComposite ? AnimComposite->bLoop : false; }
+	const FString GetName() const override { return AnimComposite ? AnimComposite->GetName() : FString(); }
+	bool IsEnabled() const override { return bEnabled; }
+	void SetIsEnabled(bool bInIsEnabled) override { bEnabled = bInIsEnabled; }
+	bool IsRootMotionEnabled() const override;
+	EPoseSearchMirrorOption GetMirrorOption() const override { return MirrorOption; }
+	ESearchIndexAssetType GetSearchIndexType() const override { return ESearchIndexAssetType::AnimComposite; }
+};
+
 USTRUCT()
 struct POSESEARCH_API FPoseSearchCost
 {
@@ -1117,11 +1174,17 @@ public:
 	UPROPERTY(EditAnywhere, Category = "Database")
 	FPoseSearchExcludeFromDatabaseParameters ExcludeFromDatabaseParameters;
 
-	UPROPERTY(EditAnywhere, Category="Database")
-	TArray<FPoseSearchDatabaseSequence> Sequences;
+#if WITH_EDITORONLY_DATA
+	// Sequences and Blendspaces are deprecated and its data will be part of the AnimationAssets.
+	// All sequences and blend spaces will be added to the AnimationAssets in PostLoad().
+	UPROPERTY()
+	TArray<FPoseSearchDatabaseSequence> Sequences_DEPRECATED;
+	UPROPERTY()
+	TArray<FPoseSearchDatabaseBlendSpace> BlendSpaces_DEPRECATED;
+#endif // WITH_EDITORONLY_DATA
 
-	UPROPERTY(EditAnywhere, Category = "Database")
-	TArray<FPoseSearchDatabaseBlendSpace> BlendSpaces;
+	UPROPERTY(EditAnywhere, Category="Database")
+	TArray<FInstancedStruct> AnimationAssets;
 
 	UPROPERTY(EditAnywhere, Category = "Performance")
 	EPoseSearchMode PoseSearchMode = EPoseSearchMode::PCAKDTree;
@@ -1157,9 +1220,10 @@ public:
 	int32 GetPoseIndexFromTime(float AssetTime, const FPoseSearchIndexAsset& SearchIndexAsset) const;
 	bool GetPoseIndicesAndLerpValueFromTime(float Time, const FPoseSearchIndexAsset& SearchIndexAsset, int32& PrevPoseIdx, int32& PoseIdx, int32& NextPoseIdx, float& LerpValue) const;
 
-	const FPoseSearchDatabaseAnimationAssetBase& GetAnimationSourceAsset(const FPoseSearchIndexAsset& SearchIndexAsset) const;
-	const FPoseSearchDatabaseSequence& GetSequenceSourceAsset(const FPoseSearchIndexAsset& SearchIndexAsset) const;
-	const FPoseSearchDatabaseBlendSpace& GetBlendSpaceSourceAsset(const FPoseSearchIndexAsset& SearchIndexAsset) const;
+	const FInstancedStruct& GetAnimationAssetStruct(int32 AnimationAssetIndex) const;
+	const FInstancedStruct& GetAnimationAssetStruct(const FPoseSearchIndexAsset& SearchIndexAsset) const;
+	FPoseSearchDatabaseAnimationAssetBase* GetAnimationAssetBase(int32 AnimationAssetIndex) const;
+	FPoseSearchDatabaseAnimationAssetBase* GetAnimationAssetBase(const FPoseSearchIndexAsset& SearchIndexAsset) const;
 	const bool IsSourceAssetLooping(const FPoseSearchIndexAsset& SearchIndexAsset) const;
 	const FString GetSourceAssetName(const FPoseSearchIndexAsset& SearchIndexAsset) const;
 	int32 GetNumberOfPrincipalComponents() const;
diff --git a/Engine/Plugins/Experimental/Animation/PoseSearch/Source/Runtime/Public/PoseSearch/PoseSearchDerivedDataKey.h b/Engine/Plugins/Experimental/Animation/PoseSearch/Source/Runtime/Public/PoseSearch/PoseSearchDerivedDataKey.h
--- a/Engine/Plugins/Experimental/Animation/PoseSearch/Source/Runtime/Public/PoseSearch/PoseSearchDerivedDataKey.h
+++ b/Engine/Plugins/Experimental/Animation/PoseSearch/Source/Runtime/Public/PoseSearch/PoseSearchDerivedDataKey.h
@@ -59,7 +59,7 @@ public:
 		if (bUseDataVer)
 		{
 			// used to invalidate the key without having to change POSESEARCHDB_DERIVEDDATA_VER all the times
-			int32 POSESEARCHDB_DERIVEDDATA_VER_SMALL = 8;
+			int32 POSESEARCHDB_DERIVEDDATA_VER_SMALL = 9;
 			FGuid VersionGuid = FDevSystemGuids::GetSystemGuid(FDevSystemGuids::Get().POSESEARCHDB_DERIVEDDATA_VER);
 
 			*this << VersionGuid;
diff --git a/Engine/Source/Runtime/Engine/Classes/Animation/AnimComposite.h b/Engine/Source/Runtime/Engine/Classes/Animation/AnimComposite.h
--- a/Engine/Source/Runtime/Engine/Classes/Animation/AnimComposite.h
+++ b/Engine/Source/Runtime/Engine/Classes/Animation/AnimComposite.h
@@ -43,6 +43,9 @@ public:
 	virtual bool IsValidAdditive() const override { return GetAdditiveAnimType() != AAT_None; }
 	virtual void EnableRootMotionSettingFromMontage(bool bInEnableRootMotion, const ERootMotionRootLock::Type InRootMotionRootLock) override;
 	virtual bool HasRootMotion() const override;
+	FTransform ExtractRootMotion(float StartTime, float DeltaTime, bool bAllowLooping) const override final;
+	FTransform ExtractRootMotionFromRange(float StartTrackPosition, float EndTrackPosition) const override final;
+	FTransform ExtractRootTrackTransform(float Time, const FBoneContainer* RequiredBones) const override final;
 	virtual void GetAnimNotifiesFromDeltaPositions(const float& PreviousPosition, const float & CurrentPosition, FAnimNotifyContext& NotifyContext) const override;
 	virtual bool IsNotifyAvailable() const override;
 	//~ End UAnimSequenceBase Interface
diff --git a/Engine/Source/Runtime/Engine/Classes/Animation/AnimSequence.h b/Engine/Source/Runtime/Engine/Classes/Animation/AnimSequence.h
--- a/Engine/Source/Runtime/Engine/Classes/Animation/AnimSequence.h
+++ b/Engine/Source/Runtime/Engine/Classes/Animation/AnimSequence.h
@@ -420,13 +420,13 @@ public:
 	//~ End UAnimSequenceBase Interface
 
 	// Extract Root Motion transform from the animation
-	FTransform ExtractRootMotion(float StartTime, float DeltaTime, bool bAllowLooping) const;
+	FTransform ExtractRootMotion(float StartTime, float DeltaTime, bool bAllowLooping) const override final;
 
 	// Extract Root Motion transform from a contiguous position range (no looping)
-	FTransform ExtractRootMotionFromRange(float StartTrackPosition, float EndTrackPosition) const;
+	FTransform ExtractRootMotionFromRange(float StartTrackPosition, float EndTrackPosition) const override final;
 
 	// Extract the transform from the root track for the given animation position
-	FTransform ExtractRootTrackTransform(float Pos, const FBoneContainer * RequiredBones) const;
+	FTransform ExtractRootTrackTransform(float Time, const FBoneContainer* RequiredBones) const override final;
 
 	// Begin Transform related functions 
 	virtual void GetAnimationPose(FAnimationPoseData& OutAnimationPoseData, const FAnimExtractContext& ExtractionContext) const override;
diff --git a/Engine/Source/Runtime/Engine/Classes/Animation/AnimSequenceBase.h b/Engine/Source/Runtime/Engine/Classes/Animation/AnimSequenceBase.h
--- a/Engine/Source/Runtime/Engine/Classes/Animation/AnimSequenceBase.h
+++ b/Engine/Source/Runtime/Engine/Classes/Animation/AnimSequenceBase.h
@@ -208,6 +208,15 @@ public:
 
 	virtual bool HasRootMotion() const { return false; }
 
+	// Extract Root Motion transform from the animation
+	virtual FTransform ExtractRootMotion(float StartTime, float DeltaTime, bool bAllowLooping) const { return {}; }
+
+	// Extract Root Motion transform from a contiguous position range (no looping)
+	virtual FTransform ExtractRootMotionFromRange(float StartTrackPosition, float EndTrackPosition) const { return {}; }
+
+	// Extract the transform from the root track for the given animation position
+	virtual FTransform ExtractRootTrackTransform(float Time, const FBoneContainer* RequiredBones) const { return {}; }
+
 	virtual void Serialize(FArchive& Ar) override;
 
 	UE_DEPRECATED(5.0, "Use other AdvanceMarkerPhaseAsLeader signature")
diff --git a/Engine/Source/Runtime/Engine/Private/Animation/AnimComposite.cpp b/Engine/Source/Runtime/Engine/Private/Animation/AnimComposite.cpp
--- a/Engine/Source/Runtime/Engine/Private/Animation/AnimComposite.cpp
+++ b/Engine/Source/Runtime/Engine/Private/Animation/AnimComposite.cpp
@@ -109,6 +109,43 @@ bool UAnimComposite::HasRootMotion() const
 	return AnimationTrack.HasRootMotion();
 }
 
+FTransform UAnimComposite::ExtractRootMotion(float StartTime, float DeltaTime, bool bAllowLooping) const
+{
+	return ExtractRootMotionFromRange(StartTime, DeltaTime);
+}
+
+FTransform UAnimComposite::ExtractRootMotionFromRange(float StartTrackPosition, float EndTrackPosition) const
+{
+	FRootMotionMovementParams RootMotion;
+	ExtractRootMotionFromTrack(AnimationTrack, StartTrackPosition, EndTrackPosition, RootMotion);
+	return RootMotion.GetRootMotionTransform();
+}
+
+FTransform UAnimComposite::ExtractRootTrackTransform(float Time, const FBoneContainer* RequiredBones) const
+{
+	if (const FAnimSegment* AnimSegment = AnimationTrack.GetSegmentAtTime(Time))
+	{
+		float SegmentTime = 0.0f;
+		if (const UAnimSequenceBase* SequenceBase = AnimSegment->GetAnimationData(Time, SegmentTime))
+		{
+			return SequenceBase->ExtractRootTrackTransform(SegmentTime, RequiredBones);
+		}
+	}
+
+	// Return the last valid value in case we're requesting for a time after the anim composite end time.
+	if (!AnimationTrack.AnimSegments.IsEmpty())
+	{
+		const int32 LastSegmentIndex = AnimationTrack.AnimSegments.Num() - 1;
+		if (Time > AnimationTrack.AnimSegments[LastSegmentIndex].AnimEndTime)
+		{
+			const UAnimSequenceBase* SequenceBase = AnimationTrack.AnimSegments[LastSegmentIndex].GetAnimReference().Get();
+			return SequenceBase->ExtractRootTrackTransform(SequenceBase->GetPlayLength(), RequiredBones);
+		}
+	}
+
+	return {};
+}
+
 #if WITH_EDITOR
 class UAnimSequence* UAnimComposite::GetAdditiveBasePose() const
 {
diff --git a/Engine/Source/Runtime/Engine/Private/Animation/AnimSequence.cpp b/Engine/Source/Runtime/Engine/Private/Animation/AnimSequence.cpp
--- a/Engine/Source/Runtime/Engine/Private/Animation/AnimSequence.cpp
+++ b/Engine/Source/Runtime/Engine/Private/Animation/AnimSequence.cpp
@@ -1229,7 +1229,7 @@ void UAnimSequence::HandleAssetPlayerTickedInternal(FAnimAssetTickContext &Conte
 	}
 }
 
-FTransform UAnimSequence::ExtractRootTrackTransform(float Pos, const FBoneContainer * RequiredBones) const
+FTransform UAnimSequence::ExtractRootTrackTransform(float Time, const FBoneContainer * RequiredBones) const
 {
 	const int32 RootBoneIndex = 0;
 	const bool bContainsRootBoneTrack = [this, RootBoneIndex]() 
@@ -1263,7 +1263,7 @@ FTransform UAnimSequence::ExtractRootTrackTransform(float Pos, const FBoneContai
 	{
 		// if we do have root data, then return root data
 		FTransform RootTransform;
-		GetBoneTransform(RootTransform, FSkeletonPoseBoneIndex(RootBoneIndex), static_cast<double>(Pos), bUseRawDataOnly);
+		GetBoneTransform(RootTransform, FSkeletonPoseBoneIndex(RootBoneIndex), static_cast<double>(Time), bUseRawDataOnly);
 		return RootTransform;
 	}
 
