motion matching - removing FPoseSearchDatabaseSequence LeadInSequence and FollowUpSequence (since now MM DB supports AnimComposites)

diff --git a/Engine/Plugins/Experimental/Animation/PoseSearch/Source/Editor/Private/PoseSearchDebugger.cpp b/Engine/Plugins/Experimental/Animation/PoseSearch/Source/Editor/Private/PoseSearchDebugger.cpp
--- a/Engine/Plugins/Experimental/Animation/PoseSearch/Source/Editor/Private/PoseSearchDebugger.cpp
+++ b/Engine/Plugins/Experimental/Animation/PoseSearch/Source/Editor/Private/PoseSearchDebugger.cpp
@@ -1579,7 +1579,6 @@ void SDebuggerDetailsView::UpdateReflection(const FTraceMotionMatchingStateMessa
 
 		Reflection->CurrentDatabaseName = CurrentDatabase->GetName();
 		Reflection->ElapsedPoseJumpTime = State.ElapsedPoseJumpTime;
-		Reflection->bFollowUpAnimation = EnumHasAnyFlags(State.Flags, FTraceMotionMatchingState::EFlags::FollowupAnimation);
 
 		Reflection->AssetPlayerAssetName = "None";
 		if (const FPoseSearchIndexAsset* IndexAsset = CurrentSearchIndex.GetAssetForPoseSafe(CurrentDbPoseIdx))
diff --git a/Engine/Plugins/Experimental/Animation/PoseSearch/Source/Editor/Private/PoseSearchDebugger.h b/Engine/Plugins/Experimental/Animation/PoseSearch/Source/Editor/Private/PoseSearchDebugger.h
--- a/Engine/Plugins/Experimental/Animation/PoseSearch/Source/Editor/Private/PoseSearchDebugger.h
+++ b/Engine/Plugins/Experimental/Animation/PoseSearch/Source/Editor/Private/PoseSearchDebugger.h
@@ -104,10 +104,6 @@ public:
 	UPROPERTY(VisibleAnywhere, Category="Motion Matching State")
 	float ElapsedPoseJumpTime = 0.0f;
 
-	/** Whether it is playing the loop following the expended animation runway */
-	UPROPERTY(VisibleAnywhere, Category="Motion Matching State")
-	bool bFollowUpAnimation = false;
-	
 	UPROPERTY(VisibleAnywhere, Category = "Motion Matching State")
 	FString AssetPlayerAssetName = "";
 
diff --git a/Engine/Plugins/Experimental/Animation/PoseSearch/Source/Runtime/Private/PoseSearchDerivedData.cpp b/Engine/Plugins/Experimental/Animation/PoseSearch/Source/Runtime/Private/PoseSearchDerivedData.cpp
--- a/Engine/Plugins/Experimental/Animation/PoseSearch/Source/Runtime/Private/PoseSearchDerivedData.cpp
+++ b/Engine/Plugins/Experimental/Animation/PoseSearch/Source/Runtime/Private/PoseSearchDerivedData.cpp
@@ -631,7 +631,7 @@ void FAssetIndexer::Init(const FAssetIndexingContext& InIndexingContext, const F
 {
 	check(InIndexingContext.Schema);
 	check(InIndexingContext.Schema->IsValid());
-	check(InIndexingContext.MainSampler);
+	check(InIndexingContext.AssetSampler);
 
 	BoneContainer = InBoneContainer;
 	IndexingContext = InIndexingContext;
@@ -650,7 +650,7 @@ bool FAssetIndexer::Process()
 {
 	check(IndexingContext.Schema);
 	check(IndexingContext.Schema->IsValid());
-	check(IndexingContext.MainSampler);
+	check(IndexingContext.AssetSampler);
 
 	FMemMark Mark(FMemStack::Get());
 
@@ -680,84 +680,37 @@ bool FAssetIndexer::Process()
 
 const float FAssetIndexer::GetSampleTimeFromDistance(float SampleDistance) const
 {
-	auto CanWrapDistanceSamples = [](const IAssetSampler* Sampler) -> bool
+	constexpr float SMALL_ROOT_DISTANCE = 1.0f;
+	const bool bCanWrap = IndexingContext.AssetSampler->IsLoopable() && IndexingContext.AssetSampler->GetTotalRootDistance() > SMALL_ROOT_DISTANCE;
+
+	float RelativeDistance = SampleDistance;
+	if (SampleDistance < 0.0f && bCanWrap)
 	{
-		constexpr float SMALL_ROOT_DISTANCE = 1.0f;
-		return Sampler->IsLoopable() && Sampler->GetTotalRootDistance() > SMALL_ROOT_DISTANCE;
-	};
+		// In this case we're sampling a loop backwards, so MainRelativeDistance must adjust so the number of cycles 
+		// is counted correctly.
+		RelativeDistance += IndexingContext.AssetSampler->GetTotalRootDistance();
+	}
 
-	float MainTotalDistance = IndexingContext.MainSampler->GetTotalRootDistance();
-	bool bMainCanWrap = CanWrapDistanceSamples(IndexingContext.MainSampler);
+	const float MainTotalDistance = IndexingContext.AssetSampler->GetTotalRootDistance();
+	const FSamplingParam SamplingParam = WrapOrClampSamplingParam(bCanWrap, MainTotalDistance, RelativeDistance);
+	const float ClipTime = IndexingContext.AssetSampler->GetTimeFromRootDistance(SamplingParam.WrappedParam + SamplingParam.Extrapolation);
 
+	// Unwrap the clip time
 	float SampleTime = MAX_flt;
-
-	if (!bMainCanWrap)
+	if (bCanWrap)
 	{
-		// Use the lead in anim if we would have to clamp to the beginning of the main anim
-		if (IndexingContext.LeadInSampler && (SampleDistance < 0.0f))
+		if (SampleDistance < 0.0f)
 		{
-			const IAssetSampler* ClipSampler = IndexingContext.LeadInSampler;
-
-			bool bLeadInCanWrap = CanWrapDistanceSamples(IndexingContext.LeadInSampler);
-			float LeadRelativeDistance = SampleDistance + ClipSampler->GetTotalRootDistance();
-			FSamplingParam SamplingParam = WrapOrClampSamplingParam(bLeadInCanWrap, ClipSampler->GetTotalRootDistance(), LeadRelativeDistance);
-
-			float ClipTime = ClipSampler->GetTimeFromRootDistance(
-				SamplingParam.WrappedParam + SamplingParam.Extrapolation);
-
-			// Make the lead in clip time relative to the main sequence again and unwrap
-			SampleTime = -((SamplingParam.NumCycles * ClipSampler->GetPlayLength()) + (ClipSampler->GetPlayLength() - ClipTime));
+			SampleTime = -((SamplingParam.NumCycles * IndexingContext.AssetSampler->GetPlayLength()) + (IndexingContext.AssetSampler->GetPlayLength() - ClipTime));
 		}
-
-		// Use the follow up anim if we would have clamp to the end of the main anim
-		else if (IndexingContext.FollowUpSampler && (SampleDistance > MainTotalDistance))
+		else
 		{
-			const IAssetSampler* ClipSampler = IndexingContext.FollowUpSampler;
-
-			bool bFollowUpCanWrap = CanWrapDistanceSamples(IndexingContext.FollowUpSampler);
-			float FollowRelativeDistance = SampleDistance - MainTotalDistance;
-			FSamplingParam SamplingParam = WrapOrClampSamplingParam(bFollowUpCanWrap, ClipSampler->GetTotalRootDistance(), FollowRelativeDistance);
-
-			float ClipTime = ClipSampler->GetTimeFromRootDistance(
-				SamplingParam.WrappedParam + SamplingParam.Extrapolation);
-
-			// Make the follow up clip time relative to the main sequence again and unwrap
-			SampleTime = IndexingContext.MainSampler->GetPlayLength() + SamplingParam.NumCycles * ClipSampler->GetPlayLength() + ClipTime;
+			SampleTime = SamplingParam.NumCycles * IndexingContext.AssetSampler->GetPlayLength() + ClipTime;
 		}
 	}
-
-	// Use the main anim if we didn't use the lead-in or follow-up anims.
-	// The main anim sample may have been wrapped or clamped
-	if (SampleTime == MAX_flt)
+	else
 	{
-		float MainRelativeDistance = SampleDistance;
-		if (SampleDistance < 0.0f && bMainCanWrap)
-		{
-			// In this case we're sampling a loop backwards, so MainRelativeDistance must adjust so the number of cycles 
-			// is counted correctly.
-			MainRelativeDistance += IndexingContext.MainSampler->GetTotalRootDistance();
-		}
-
-		FSamplingParam SamplingParam = WrapOrClampSamplingParam(bMainCanWrap, MainTotalDistance, MainRelativeDistance);
-		float ClipTime = IndexingContext.MainSampler->GetTimeFromRootDistance(
-			SamplingParam.WrappedParam + SamplingParam.Extrapolation);
-
-		// Unwrap the main clip time
-		if (bMainCanWrap)
-		{
-			if (SampleDistance < 0.0f)
-			{
-				SampleTime = -((SamplingParam.NumCycles * IndexingContext.MainSampler->GetPlayLength()) + (IndexingContext.MainSampler->GetPlayLength() - ClipTime));
-			}
-			else
-			{
-				SampleTime = SamplingParam.NumCycles * IndexingContext.MainSampler->GetPlayLength() + ClipTime;
-			}
-		}
-		else
-		{
-			SampleTime = ClipTime;
-		}
+		SampleTime = ClipTime;
 	}
 
 	return SampleTime;
@@ -767,89 +720,22 @@ FAssetIndexer::FSampleInfo FAssetIndexer::GetSampleInfo(float SampleTime) const
 {
 	FSampleInfo Sample;
 
-	FTransform RootMotionLast = FTransform::Identity;
-	FTransform RootMotionInitial = FTransform::Identity;
-
-	float RootDistanceLast = 0.0f;
-	float RootDistanceInitial = 0.0f;
+	check(IndexingContext.AssetSampler);
 
-	check(IndexingContext.MainSampler);
+	const float PlayLength = IndexingContext.AssetSampler->GetPlayLength();
+	const bool bCanWrap = IndexingContext.AssetSampler->IsLoopable();
 
-	const float MainPlayLength = IndexingContext.MainSampler->GetPlayLength();
-	const bool bMainCanWrap = IndexingContext.MainSampler->IsLoopable();
-
-	FSamplingParam SamplingParam;
-	if (!bMainCanWrap)
+	float MainRelativeTime = SampleTime;
+	if (SampleTime < 0.0f && bCanWrap)
 	{
-		// Use the lead in anim if we would have to clamp to the beginning of the main anim
-		if (IndexingContext.LeadInSampler && (SampleTime < 0.0f))
-		{
-			const IAssetSampler* ClipSampler = IndexingContext.LeadInSampler;
-
-			const bool bLeadInCanWrap = IndexingContext.LeadInSampler->IsLoopable();
-			const float LeadRelativeTime = SampleTime + ClipSampler->GetPlayLength();
-			SamplingParam = WrapOrClampSamplingParam(bLeadInCanWrap, ClipSampler->GetPlayLength(), LeadRelativeTime);
-
-			Sample.Clip = IndexingContext.LeadInSampler;
-
-			check(SamplingParam.Extrapolation <= 0.0f);
-			if (SamplingParam.Extrapolation < 0.0f)
-			{
-				RootMotionInitial = IndexingContext.LeadInSampler->GetTotalRootTransform().Inverse();
-				RootDistanceInitial = -IndexingContext.LeadInSampler->GetTotalRootDistance();
-			}
-			else
-			{
-				RootMotionInitial = FTransform::Identity;
-				RootDistanceInitial = 0.0f;
-			}
-
-			RootMotionLast = IndexingContext.LeadInSampler->GetTotalRootTransform();
-			RootDistanceLast = IndexingContext.LeadInSampler->GetTotalRootDistance();
-		}
-
-		// Use the follow up anim if we would have clamp to the end of the main anim
-		else if (IndexingContext.FollowUpSampler && (SampleTime > MainPlayLength))
-		{
-			const IAssetSampler* ClipSampler = IndexingContext.FollowUpSampler;
-
-			const bool bFollowUpCanWrap = IndexingContext.FollowUpSampler->IsLoopable();
-			const float FollowRelativeTime = SampleTime - MainPlayLength;
-			SamplingParam = WrapOrClampSamplingParam(bFollowUpCanWrap, ClipSampler->GetPlayLength(), FollowRelativeTime);
-
-			Sample.Clip = IndexingContext.FollowUpSampler;
-
-			RootMotionInitial = IndexingContext.MainSampler->GetTotalRootTransform();
-			RootDistanceInitial = IndexingContext.MainSampler->GetTotalRootDistance();
-
-			RootMotionLast = IndexingContext.FollowUpSampler->GetTotalRootTransform();
-			RootDistanceLast = IndexingContext.FollowUpSampler->GetTotalRootDistance();
-		}
+		// In this case we're sampling a loop backwards, so MainRelativeTime must adjust so the number of cycles is
+		// counted correctly.
+		MainRelativeTime += PlayLength;
 	}
 
-	// Use the main anim if we didn't use the lead-in or follow-up anims.
-	// The main anim sample may have been wrapped or clamped
-	if (!Sample.IsValid())
-	{
-		float MainRelativeTime = SampleTime;
-		if (SampleTime < 0.0f && bMainCanWrap)
-		{
-			// In this case we're sampling a loop backwards, so MainRelativeTime must adjust so the number of cycles is
-			// counted correctly.
-			MainRelativeTime += MainPlayLength;
-		}
-
-		SamplingParam = WrapOrClampSamplingParam(bMainCanWrap, MainPlayLength, MainRelativeTime);
-
-		Sample.Clip = IndexingContext.MainSampler;
-
-		RootMotionInitial = FTransform::Identity;
-		RootDistanceInitial = 0.0f;
-
-		RootMotionLast = IndexingContext.MainSampler->GetTotalRootTransform();
-		RootDistanceLast = IndexingContext.MainSampler->GetTotalRootDistance();
-	}
+	const FSamplingParam SamplingParam = WrapOrClampSamplingParam(bCanWrap, PlayLength, MainRelativeTime);
 
+	Sample.Clip = IndexingContext.AssetSampler;
 
 	if (FMath::Abs(SamplingParam.Extrapolation) > SMALL_NUMBER)
 	{
@@ -858,13 +744,16 @@ FAssetIndexer::FSampleInfo FAssetIndexer::GetSampleInfo(float SampleTime) const
 		const FTransform ClipRootMotion = Sample.Clip->ExtractRootTransform(Sample.ClipTime);
 		const float ClipDistance = Sample.Clip->ExtractRootDistance(Sample.ClipTime);
 
-		Sample.RootTransform = ClipRootMotion * RootMotionInitial;
-		Sample.RootDistance = RootDistanceInitial + ClipDistance;
+		Sample.RootTransform = ClipRootMotion;
+		Sample.RootDistance = ClipDistance;
 	}
 	else
 	{
 		Sample.ClipTime = SamplingParam.WrappedParam;
 
+		const FTransform RootMotionLast = IndexingContext.AssetSampler->GetTotalRootTransform();
+		float RootDistanceLast = IndexingContext.AssetSampler->GetTotalRootDistance();
+
 		// Determine how to accumulate motion for every cycle of the anim. If the sample
 		// had to be clamped, this motion will end up not getting applied below.
 		// Also invert the accumulation direction if the requested sample was wrapped backwards.
@@ -887,8 +776,8 @@ FAssetIndexer::FSampleInfo FAssetIndexer::GetSampleInfo(float SampleTime) const
 			RootDistanceRemainder = -(RootDistanceLast - RootDistanceRemainder);
 		}
 
-		Sample.RootTransform = RootMotionInitial;
-		Sample.RootDistance = RootDistanceInitial;
+		Sample.RootTransform = FTransform::Identity;
+		Sample.RootDistance = 0.f;
 
 		// Note if the sample was clamped, no motion will be applied here because NumCycles will be zero
 		int32 CyclesRemaining = SamplingParam.NumCycles;
@@ -920,7 +809,7 @@ FTransform FAssetIndexer::MirrorTransform(const FTransform& Transform) const
 
 FPoseSearchPoseMetadata FAssetIndexer::GetMetadata(int32 SampleIdx) const
 {
-	const float SequenceLength = IndexingContext.MainSampler->GetPlayLength();
+	const float SequenceLength = IndexingContext.AssetSampler->GetPlayLength();
 	const float SampleTime = FMath::Min(SampleIdx * IndexingContext.Schema->GetSamplingInterval(), SequenceLength);
 
 	FPoseSearchPoseMetadata Metadata;
@@ -928,7 +817,7 @@ FPoseSearchPoseMetadata FAssetIndexer::GetMetadata(int32 SampleIdx) const
 	Metadata.ContinuingPoseCostAddend = IndexingContext.Schema->ContinuingPoseCostBias;
 
 	TArray<UAnimNotifyState_PoseSearchBase*> NotifyStates;
-	IndexingContext.MainSampler->ExtractPoseSearchNotifyStates(SampleTime, NotifyStates);
+	IndexingContext.AssetSampler->ExtractPoseSearchNotifyStates(SampleTime, NotifyStates);
 	for (const UAnimNotifyState_PoseSearchBase* PoseSearchNotify : NotifyStates)
 	{
 		if (PoseSearchNotify->GetClass()->IsChildOf<UAnimNotifyState_PoseSearchBlockTransition>())
@@ -973,7 +862,7 @@ FTransform FAssetIndexer::GetTransformAndCacheResults(float SampleTime, float Or
 			UE_LOG(LogPoseSearch,
 				Warning,
 				TEXT("Invalid BoneContainer encountered in FAssetIndexer::GetTransformAndCacheResults. Asset: %s. Schema: %s. BoneContainerAsset: %s. NumBoneIndices: %d"),
-				*GetNameSafe(IndexingContext.MainSampler->GetAsset()),
+				*GetNameSafe(IndexingContext.AssetSampler->GetAsset()),
 				*GetNameSafe(IndexingContext.Schema),
 				*GetNameSafe(BoneContainer.GetAsset()),
 				BoneContainer.GetCompactPoseNumBones());
@@ -1072,8 +961,6 @@ void FDatabaseIndexingContext::Prepare(const UPoseSearchDatabase* Database)
 		if (const FPoseSearchDatabaseSequence* DatabaseSequence = DatabaseAssetStruct.GetPtr<FPoseSearchDatabaseSequence>())
 		{
 			AddSequenceBaseSampler(DatabaseSequence->Sequence);
-			AddSequenceBaseSampler(DatabaseSequence->LeadInSequence);
-			AddSequenceBaseSampler(DatabaseSequence->FollowUpSequence);
 		}
 		else if (const FPoseSearchDatabaseAnimComposite* DatabaseAnimComposite = DatabaseAssetStruct.GetPtr<FPoseSearchDatabaseAnimComposite>())
 		{
@@ -1154,17 +1041,15 @@ void FDatabaseIndexingContext::Prepare(const UPoseSearchDatabase* Database)
 		if (const FPoseSearchDatabaseSequence* DatabaseSequence = DatabaseAsset.GetPtr<FPoseSearchDatabaseSequence>())
 		{
 			const float SequenceLength = DatabaseSequence->Sequence->GetPlayLength();
-			IndexerContext.MainSampler = GetSequenceBaseSampler(DatabaseSequence->Sequence);
-			IndexerContext.LeadInSampler = SearchIndexAsset.SamplingInterval.Min == 0.0f ? GetSequenceBaseSampler(DatabaseSequence->LeadInSequence) : nullptr;
-			IndexerContext.FollowUpSampler = SearchIndexAsset.SamplingInterval.Max == SequenceLength ? GetSequenceBaseSampler(DatabaseSequence->FollowUpSequence) : nullptr;
+			IndexerContext.AssetSampler = GetSequenceBaseSampler(DatabaseSequence->Sequence);
 		}
 		else if (const FPoseSearchDatabaseAnimComposite* DatabaseAnimComposite = DatabaseAsset.GetPtr<FPoseSearchDatabaseAnimComposite>())
 		{
-			IndexerContext.MainSampler = GetSequenceBaseSampler(DatabaseAnimComposite->AnimComposite);
+			IndexerContext.AssetSampler = GetSequenceBaseSampler(DatabaseAnimComposite->AnimComposite);
 		}
 		else if (const FPoseSearchDatabaseBlendSpace* DatabaseBlendSpace = DatabaseAsset.GetPtr<FPoseSearchDatabaseBlendSpace>())
 		{
-			IndexerContext.MainSampler = GetBlendSpaceSampler(DatabaseBlendSpace->BlendSpace, SearchIndexAsset.BlendParameters);
+			IndexerContext.AssetSampler = GetBlendSpaceSampler(DatabaseBlendSpace->BlendSpace, SearchIndexAsset.BlendParameters);
 		}
 
 		FAssetIndexer& Indexer = Indexers.AddDefaulted_GetRef();
diff --git a/Engine/Plugins/Experimental/Animation/PoseSearch/Source/Runtime/Private/PoseSearchFeatureChannels.cpp b/Engine/Plugins/Experimental/Animation/PoseSearch/Source/Runtime/Private/PoseSearchFeatureChannels.cpp
--- a/Engine/Plugins/Experimental/Animation/PoseSearch/Source/Runtime/Private/PoseSearchFeatureChannels.cpp
+++ b/Engine/Plugins/Experimental/Animation/PoseSearch/Source/Runtime/Private/PoseSearchFeatureChannels.cpp
@@ -65,7 +65,7 @@ static void CollectBonePositions(TArray<FVector>& BonePositions, IAssetIndexer&
 {
 	const FAssetIndexingContext& IndexingContext = Indexer.GetIndexingContext();
 	const float FiniteDelta = IndexingContext.Schema->GetSamplingInterval();
-	const float SampleTimeStart = FMath::Min(IndexingContext.BeginSampleIdx * FiniteDelta, IndexingContext.MainSampler->GetPlayLength());
+	const float SampleTimeStart = FMath::Min(IndexingContext.BeginSampleIdx * FiniteDelta, IndexingContext.AssetSampler->GetPlayLength());
 	const int32 NumSamples = IndexingContext.EndSampleIdx - IndexingContext.BeginSampleIdx;
 
 	// collecting all the bone transforms
@@ -346,7 +346,7 @@ void UPoseSearchFeatureChannel_Position::IndexAsset(UE::PoseSearch::IAssetIndexe
 	{
 		const int32 VectorIdx = SampleIdx - IndexingContext.BeginSampleIdx;
 
-		const float OriginSampleTime = FMath::Min(SampleIdx * IndexingContext.Schema->GetSamplingInterval(), IndexingContext.MainSampler->GetPlayLength());
+		const float OriginSampleTime = FMath::Min(SampleIdx * IndexingContext.Schema->GetSamplingInterval(), IndexingContext.AssetSampler->GetPlayLength());
 		const float SubsampleTime = OriginSampleTime + SampleTimeOffset;
 
 		bool ClampedPresent;
@@ -470,7 +470,7 @@ void UPoseSearchFeatureChannel_Heading::IndexAsset(UE::PoseSearch::IAssetIndexer
 	{
 		const int32 VectorIdx = SampleIdx - IndexingContext.BeginSampleIdx;
 
-		const float OriginSampleTime = FMath::Min(SampleIdx * IndexingContext.Schema->GetSamplingInterval(), IndexingContext.MainSampler->GetPlayLength());
+		const float OriginSampleTime = FMath::Min(SampleIdx * IndexingContext.Schema->GetSamplingInterval(), IndexingContext.AssetSampler->GetPlayLength());
 		const float SubsampleTime = OriginSampleTime + SampleTimeOffset;
 
 		bool ClampedPresent;
@@ -749,7 +749,7 @@ void UPoseSearchFeatureChannel_Pose::AddPoseFeatures(UE::PoseSearch::IAssetIndex
 	const FAssetIndexingContext& IndexingContext = Indexer.GetIndexingContext();
 	const FAssetSamplingContext* SamplingContext = IndexingContext.SamplingContext;
 
-	const float SampleTime = FMath::Min(SampleIdx * IndexingContext.Schema->GetSamplingInterval(), IndexingContext.MainSampler->GetPlayLength());
+	const float SampleTime = FMath::Min(SampleIdx * IndexingContext.Schema->GetSamplingInterval(), IndexingContext.AssetSampler->GetPlayLength());
 
 	int32 DataOffset = ChannelDataOffset;
 	for (int32 ChannelBoneIdx = 0; ChannelBoneIdx != SampledBones.Num(); ++ChannelBoneIdx)
@@ -1431,7 +1431,7 @@ void UPoseSearchFeatureChannel_Trajectory::IndexAssetPrivate(const UE::PoseSearc
 	using FSampleInfo = IAssetIndexer::FSampleInfo;
 
 	const FAssetIndexingContext& IndexingContext = Indexer.GetIndexingContext();
-	float SampleTime = FMath::Min(SampleIdx * IndexingContext.Schema->GetSamplingInterval(), IndexingContext.MainSampler->GetPlayLength());
+	float SampleTime = FMath::Min(SampleIdx * IndexingContext.Schema->GetSamplingInterval(), IndexingContext.AssetSampler->GetPlayLength());
 	FSampleInfo Origin = Indexer.GetSampleInfo(SampleTime);
 
 	int32 DataOffset = ChannelDataOffset;
diff --git a/Engine/Plugins/Experimental/Animation/PoseSearch/Source/Runtime/Private/PoseSearchLibrary.cpp b/Engine/Plugins/Experimental/Animation/PoseSearch/Source/Runtime/Private/PoseSearchLibrary.cpp
--- a/Engine/Plugins/Experimental/Animation/PoseSearch/Source/Runtime/Private/PoseSearchLibrary.cpp
+++ b/Engine/Plugins/Experimental/Animation/PoseSearch/Source/Runtime/Private/PoseSearchLibrary.cpp
@@ -34,11 +34,8 @@ void FMotionMatchingState::AdjustAssetTime(float AssetTime)
 	CurrentSearchResult.Update(AssetTime);
 }
 
-bool FMotionMatchingState::CanAdvance(float DeltaTime, bool& bOutAdvanceToFollowUpAsset, UE::PoseSearch::FSearchResult& OutFollowUpAsset) const
+bool FMotionMatchingState::CanAdvance(float DeltaTime) const
 {
-	bOutAdvanceToFollowUpAsset = false;
-	OutFollowUpAsset = UE::PoseSearch::FSearchResult();
-
 	if (!CurrentSearchResult.IsValid())
 	{
 		return false;
@@ -62,49 +59,6 @@ bool FMotionMatchingState::CanAdvance(float DeltaTime, bool& bOutAdvanceToFollow
 		{
 			return SearchIndexAsset->SamplingInterval.Contains(SteppedTime);
 		}
-		else
-		{
-			// check if there's a follow-up that can be used
-			int32 FollowUpDbSequenceIdx = CurrentSearchResult.Database->AnimationAssets.IndexOfByPredicate(
-				[&](const FInstancedStruct& DatabaseAsset)
-				{
-					if (const FPoseSearchDatabaseSequence* DatabaseSequence = DatabaseAsset.GetPtr<FPoseSearchDatabaseSequence>())
-					{
-						return (DatabaseSequence->Sequence == DatabaseSequence->FollowUpSequence);
-					}
-
-					return false;
-				});
-
-			int32 FollowUpSearchIndexAssetIdx = CurrentSearchResult.Database->GetSearchIndex().Assets.IndexOfByPredicate(
-				[&](const FPoseSearchIndexAsset& Entry)
-				{
-					const bool bIsMatch =
-						Entry.SourceAssetIdx == FollowUpDbSequenceIdx &&
-						Entry.bMirrored == SearchIndexAsset->bMirrored &&
-						Entry.SamplingInterval.Contains(0.0f);
-					return bIsMatch;
-				});
-
-			if (FollowUpSearchIndexAssetIdx != INDEX_NONE)
-			{
-				bOutAdvanceToFollowUpAsset = true;
-
-				const FPoseSearchIndexAsset& FollowUpSearchIndexAsset = CurrentSearchResult.Database->GetSearchIndex().Assets[FollowUpSearchIndexAssetIdx];
-
-				// Follow up asset time will start slightly before the beginning of the sequence as 
-				// this is essentially what the matching time in the corresponding main sequence is.
-				// Here we are assuming that the tick will advance the asset player timer into the 
-				// valid region
-				const float FollowUpAssetTime = CurrentSearchResult.AssetTime - AssetLength;
-
-				// There is no correspoding pose index when we switch due to what is mentioned above
-				// so for now we just take whatever pose index is associated with the first frame.
-				OutFollowUpAsset.PoseIdx = CurrentSearchResult.Database->GetPoseIndexFromTime(FollowUpSearchIndexAsset.SamplingInterval.Min, FollowUpSearchIndexAsset);
-				OutFollowUpAsset.AssetTime = FollowUpAssetTime;
-				return true;
-			}
-		}
 	}
 	else if (const FPoseSearchDatabaseAnimComposite* DatabaseAnimComposite = DatabaseAsset.GetPtr<FPoseSearchDatabaseAnimComposite>())
 	{
@@ -301,11 +255,6 @@ static void TraceMotionMatchingState(
 		TraceState.AnimAngularVelocity = FMath::RadiansToDegrees(AnimDelta.GetRotation().GetAngle()) / DeltaTime;
 	}
 
-	if (EnumHasAnyFlags(MotionMatchingState.Flags, EMotionMatchingFlags::JumpedToFollowUp))
-	{
-		TraceState.Flags |= FTraceMotionMatchingState::EFlags::FollowupAnimation;
-	}
-
 	TraceState.SearchableAssetId = FTraceMotionMatchingState::GetIdFromObject(Searchable);
 	TraceState.ElapsedPoseJumpTime = MotionMatchingState.ElapsedPoseJumpTime;
 	TraceState.AssetPlayerTime = MotionMatchingState.CurrentSearchResult.AssetTime;
@@ -347,10 +296,8 @@ void UpdateMotionMatchingState(
 	const FSearchResult LastResult = InOutMotionMatchingState.CurrentSearchResult;
 #endif
 
-	// Check if we can advance. Includes the case where we can advance but only by switching to a follow up asset.
-	bool bAdvanceToFollowUpAsset = false;
-	FSearchResult FollowUpAssetResult;
-	const bool bCanAdvance = InOutMotionMatchingState.CanAdvance(DeltaTime, bAdvanceToFollowUpAsset, FollowUpAssetResult);
+	// Check if we can advance.
+	const bool bCanAdvance = InOutMotionMatchingState.CanAdvance(DeltaTime);
 
 	// If we can't advance or enough time has elapsed since the last pose jump then search
 	FSearchContext SearchContext;
@@ -404,17 +351,6 @@ void UpdateMotionMatchingState(
 		}
 	}
 
-	// If we didn't search or it didn't find a pose to jump to, and we can 
-	// advance but only with the follow up asset, jump to that. Otherwise we 
-	// are advancing as normal, and nothing needs to be done.
-	if (!(InOutMotionMatchingState.Flags & EMotionMatchingFlags::JumpedToPose)
-		&& bCanAdvance
-		&& bAdvanceToFollowUpAsset)
-	{
-		InOutMotionMatchingState.JumpToPose(Context, Settings, FollowUpAssetResult);
-		InOutMotionMatchingState.Flags |= EMotionMatchingFlags::JumpedToFollowUp;
-	}
-
 	// Tick elapsed pose jump timer
 	if (!(InOutMotionMatchingState.Flags & EMotionMatchingFlags::JumpedToPose))
 	{
diff --git a/Engine/Plugins/Experimental/Animation/PoseSearch/Source/Runtime/Private/Trace/PoseSearchTraceLogger.cpp b/Engine/Plugins/Experimental/Animation/PoseSearch/Source/Runtime/Private/Trace/PoseSearchTraceLogger.cpp
--- a/Engine/Plugins/Experimental/Animation/PoseSearch/Source/Runtime/Private/Trace/PoseSearchTraceLogger.cpp
+++ b/Engine/Plugins/Experimental/Animation/PoseSearch/Source/Runtime/Private/Trace/PoseSearchTraceLogger.cpp
@@ -71,7 +71,6 @@ FArchive& operator<<(FArchive& Ar, FTraceMotionMatchingState& State)
 {
 	Ar << State.SearchableAssetId;
 	Ar << State.ElapsedPoseJumpTime;
-	Ar << State.Flags;
 	Ar << State.AssetPlayerTime;
 	Ar << State.DeltaTime;
 	Ar << State.SimLinearVelocity;
diff --git a/Engine/Plugins/Experimental/Animation/PoseSearch/Source/Runtime/Private/Trace/PoseSearchTraceLogger.h b/Engine/Plugins/Experimental/Animation/PoseSearch/Source/Runtime/Private/Trace/PoseSearchTraceLogger.h
--- a/Engine/Plugins/Experimental/Animation/PoseSearch/Source/Runtime/Private/Trace/PoseSearchTraceLogger.h
+++ b/Engine/Plugins/Experimental/Animation/PoseSearch/Source/Runtime/Private/Trace/PoseSearchTraceLogger.h
@@ -64,24 +64,12 @@ POSESEARCH_API FArchive& operator<<(FArchive& Ar, FTraceMotionMatchingStateDatab
  */
 struct POSESEARCH_API FTraceMotionMatchingState
 {
-	/** Bitfield for various state booleans */
-	enum class EFlags : uint32
-	{
-		None = 0u,
-
-		/** Whether the last animation was a forced follow-up animation due to expended animation runway */
-		FollowupAnimation = 1u << 0
-	};
-
 	/** ObjectId of active searchable asset */
 	uint64 SearchableAssetId = 0;
 	
 	/** Amount of time since the last pose switch */
 	float ElapsedPoseJumpTime = 0.0f;
 
-	/** Storage container for state booleans */
-	EFlags Flags = EFlags::None;
-
 	float AssetPlayerTime = 0.0f;
 	float DeltaTime = 0.0f;
 	float SimLinearVelocity = 0.0f;
@@ -132,7 +120,6 @@ struct POSESEARCH_API FTraceMotionMatchingState
 	
 	static const FName Name;
 };
-ENUM_CLASS_FLAGS(FTraceMotionMatchingState::EFlags)
 
 POSESEARCH_API FArchive& operator<<(FArchive& Ar, FTraceMotionMatchingState& State);
 
diff --git a/Engine/Plugins/Experimental/Animation/PoseSearch/Source/Runtime/Public/PoseSearch/PoseSearch.h b/Engine/Plugins/Experimental/Animation/PoseSearch/Source/Runtime/Public/PoseSearch/PoseSearch.h
--- a/Engine/Plugins/Experimental/Animation/PoseSearch/Source/Runtime/Public/PoseSearch/PoseSearch.h
+++ b/Engine/Plugins/Experimental/Animation/PoseSearch/Source/Runtime/Public/PoseSearch/PoseSearch.h
@@ -343,9 +343,7 @@ struct FAssetIndexingContext
 {
 	const FAssetSamplingContext* SamplingContext = nullptr;
 	const UPoseSearchSchema* Schema = nullptr;
-	const IAssetSampler* MainSampler = nullptr;
-	const IAssetSampler* LeadInSampler = nullptr;
-	const IAssetSampler* FollowUpSampler = nullptr;
+	const IAssetSampler* AssetSampler = nullptr;
 	bool bMirrored = false;
 	FFloatInterval RequestedSamplingRange = FFloatInterval(0.0f, 0.0f);
 	
@@ -870,20 +868,6 @@ struct POSESEARCH_API FPoseSearchDatabaseSequence : public FPoseSearchDatabaseAn
 	UPROPERTY(EditAnywhere, Category = "Sequence")
 	EPoseSearchMirrorOption MirrorOption = EPoseSearchMirrorOption::UnmirroredOnly;
 
-	// Used for sampling past pose information at the beginning of the main sequence.
-	// This setting is intended for transitions between cycles. It is optional and only used
-	// for one shot anims with past sampling. When past sampling is used without a lead in sequence,
-	// the sampling range of the main sequence will be clamped if necessary.
-	UPROPERTY(EditAnywhere, Category="Sequence")
-	TObjectPtr<UAnimSequence> LeadInSequence;
-
-	// Used for sampling future pose information at the end of the main sequence.
-	// This setting is intended for transitions between cycles. It is optional and only used
-	// for one shot anims with future sampling. When future sampling is used without a follow up sequence,
-	// the sampling range of the main sequence will be clamped if necessary.
-	UPROPERTY(EditAnywhere, Category="Sequence")
-	TObjectPtr<UAnimSequence> FollowUpSequence;
-
 	UAnimationAsset* GetAnimationAsset() const override { return Sequence; }
 	UClass* GetAnimationAssetStaticClass() const override { return UAnimSequence::StaticClass(); }
 	bool IsLooping() const override { return Sequence ? Sequence->bLoop : false; }
diff --git a/Engine/Plugins/Experimental/Animation/PoseSearch/Source/Runtime/Public/PoseSearch/PoseSearchLibrary.h b/Engine/Plugins/Experimental/Animation/PoseSearch/Source/Runtime/Public/PoseSearch/PoseSearchLibrary.h
--- a/Engine/Plugins/Experimental/Animation/PoseSearch/Source/Runtime/Public/PoseSearch/PoseSearchLibrary.h
+++ b/Engine/Plugins/Experimental/Animation/PoseSearch/Source/Runtime/Public/PoseSearch/PoseSearchLibrary.h
@@ -18,7 +18,6 @@ enum class EMotionMatchingFlags : uint8
 {
 	None = 0 UMETA(Hidden),
 	JumpedToPose = 1 << 0,		// Signals that motion matching has made a significant deviation in the selected sequence/pose index
-	JumpedToFollowUp = 1 << 1,	// Motion matching chose the follow up animation of the prior sequence
 };
 ENUM_CLASS_FLAGS(EMotionMatchingFlags);
 
@@ -72,10 +71,8 @@ struct POSESEARCH_API FMotionMatchingState
 	// Reset the state to a default state using the current Database
 	void Reset();
 
-	// Checks if the currently playing asset can advance and stay in bounds under the provided DeltaTime. If the 
-	// asset cannot advance but a follow up asset is provided by the database then this function will return true and 
-	// also set bAdvanceToFollowUpAsset to true, while providing the pose to jump to in the follow up asset in OutFollowUpAsset
-	bool CanAdvance(float DeltaTime, bool& bOutAdvanceToFollowUpAsset, UE::PoseSearch::FSearchResult& OutFollowUpAsset) const;
+	// Checks if the currently playing asset can advance and stay in bounds under the provided DeltaTime.
+	bool CanAdvance(float DeltaTime) const;
 
 	// Attempts to set the internal state to match the provided asset time including updating the internal DbPoseIdx. 
 	// If the provided asset time is out of bounds for the currently playing asset then this function will reset the 
