motion matching - Hide Invalid Poses to Pose Search Debugger

diff --git a/Engine/Plugins/Experimental/Animation/PoseSearch/Source/Editor/Private/PoseSearchDebugger.cpp b/Engine/Plugins/Experimental/Animation/PoseSearch/Source/Editor/Private/PoseSearchDebugger.cpp
--- a/Engine/Plugins/Experimental/Animation/PoseSearch/Source/Editor/Private/PoseSearchDebugger.cpp
+++ b/Engine/Plugins/Experimental/Animation/PoseSearch/Source/Editor/Private/PoseSearchDebugger.cpp
@@ -1152,6 +1152,19 @@ void SDebuggerDatabaseView::OnFilterTextChanged(const FText& SearchText)
 	PopulateViewRows();
 }
 
+void SDebuggerDatabaseView::OnHideInvalidPosesCheckboxChanged(ECheckBoxState State)
+{
+	if (State == ECheckBoxState::Checked) 
+	{
+		bHideInvalidPoses = true;
+	}
+	else 
+	{
+		bHideInvalidPoses = false;
+	}
+	PopulateViewRows();
+}
+
 void SDebuggerDatabaseView::OnDatabaseRowSelectionChanged(
 	TSharedPtr<FDebuggerDatabaseRowData> Row, 
 	ESelectInfo::Type SelectInfo)
@@ -1210,18 +1223,21 @@ void SDebuggerDatabaseView::PopulateViewRows()
 
 		if (bTryAddToFilteredDatabaseViewRows)
 		{
-			bool bPassesNameFilter = true;
-			if (bHasNameFilter)
+			if (!bHideInvalidPoses || EnumHasAnyFlags(UnfilteredRow->PoseCandidateFlags, EPoseCandidateFlags::AnyValidMask))
 			{
-				bPassesNameFilter = Algo::AllOf(Tokens, [&](FString Token)
+				bool bPassesNameFilter = true;
+				if (bHasNameFilter)
 				{
-					return UnfilteredRow->AssetName.Contains(Token);
-				});
-			}
+					bPassesNameFilter = Algo::AllOf(Tokens, [&](FString Token)
+					{
+						return UnfilteredRow->AssetName.Contains(Token);
+					});
+				}
 
-			if (bPassesNameFilter)
-			{
-				FilteredDatabaseView.Rows.Add(UnfilteredRow);
+				if (bPassesNameFilter)
+				{
+					FilteredDatabaseView.Rows.Add(UnfilteredRow);
+				}
 			}
 		}
 	}
@@ -1498,6 +1514,20 @@ void SDebuggerDatabaseView::Construct(const FArguments& InArgs)
 					SAssignNew(FilterBox, SSearchBox)
 					.OnTextChanged(this, &SDebuggerDatabaseView::OnFilterTextChanged)
 				]
+				+ SHorizontalBox::Slot()
+				.AutoWidth()
+				.Padding(10, 5, 10, 5)
+				[
+					SNew(SCheckBox)
+					.OnCheckStateChanged_Lambda([this](ECheckBoxState State)
+					{
+						   SDebuggerDatabaseView::OnHideInvalidPosesCheckboxChanged(State);
+					})
+					[
+						SNew(STextBlock)
+						.Text(LOCTEXT("PoseSearchDebuggerHideInvalidPosesFlag", "Hide Invalid Poses"))
+					]
+				]
 			]
 		
 			+ SVerticalBox::Slot()
diff --git a/Engine/Plugins/Experimental/Animation/PoseSearch/Source/Editor/Private/PoseSearchDebugger.h b/Engine/Plugins/Experimental/Animation/PoseSearch/Source/Editor/Private/PoseSearchDebugger.h
--- a/Engine/Plugins/Experimental/Animation/PoseSearch/Source/Editor/Private/PoseSearchDebugger.h
+++ b/Engine/Plugins/Experimental/Animation/PoseSearch/Source/Editor/Private/PoseSearchDebugger.h
@@ -222,6 +222,9 @@ private:
 	/** Called when the text in the filter box is modified to update the filtering */
 	void OnFilterTextChanged(const FText& SearchText);
 
+	/** Called when the HideInvalidPoses Checkbox is Toggled */
+	void OnHideInvalidPosesCheckboxChanged(ECheckBoxState State);
+
 	/** Row selection to update model view */
 	void OnDatabaseRowSelectionChanged(TSharedPtr<FDebuggerDatabaseRowData> Row, ESelectInfo::Type SelectInfo);
 
@@ -288,6 +291,9 @@ private:
 
 	/** Text used to filter DatabaseView */
 	FText FilterText;
+
+	/* Bool used to hide invalid poses from display */
+	bool bHideInvalidPoses = false;
 };
 
 /**
