motion matching - removing UPoseSearchFeatureChannel_Pose::SampleTimes

diff --git a/Engine/Plugins/Experimental/Animation/PoseSearch/Source/Runtime/Private/PoseSearchFeatureChannels.cpp b/Engine/Plugins/Experimental/Animation/PoseSearch/Source/Runtime/Private/PoseSearchFeatureChannels.cpp
--- a/Engine/Plugins/Experimental/Animation/PoseSearch/Source/Runtime/Private/PoseSearchFeatureChannels.cpp
+++ b/Engine/Plugins/Experimental/Animation/PoseSearch/Source/Runtime/Private/PoseSearchFeatureChannels.cpp
@@ -553,12 +553,6 @@ void UPoseSearchFeatureChannel_Heading::DebugDraw(const UE::PoseSearch::FDebugDr
 //////////////////////////////////////////////////////////////////////////
 // UPoseSearchFeatureChannel_Pose
 
-void UPoseSearchFeatureChannel_Pose::PreSave(FObjectPreSaveContext ObjectSaveContext)
-{
-	SampleTimes.Sort(TLess<>());
-	Super::PreSave(ObjectSaveContext);
-}
-
 void UPoseSearchFeatureChannel_Pose::InitializeSchema(UE::PoseSearch::FSchemaInitializer& Initializer)
 {
 	using namespace UE::PoseSearch;
@@ -571,31 +565,19 @@ void UPoseSearchFeatureChannel_Pose::InitializeSchema(UE::PoseSearch::FSchemaIni
 		const FPoseSearchBone& SampledBone = SampledBones[ChannelBoneIdx];
 		if (EnumHasAnyFlags(SampledBone.Flags, EPoseSearchBoneFlags::Position))
 		{
-			for (int32 SubsampleIdx = 0; SubsampleIdx != SampleTimes.Num(); ++SubsampleIdx)
-			{
-				DataOffset += FFeatureVectorHelper::EncodeVectorCardinality;
-			}
+			DataOffset += FFeatureVectorHelper::EncodeVectorCardinality;
 		}
 		if (EnumHasAnyFlags(SampledBone.Flags, EPoseSearchBoneFlags::Rotation))
 		{
-			for (int32 SubsampleIdx = 0; SubsampleIdx != SampleTimes.Num(); ++SubsampleIdx)
-			{
-				DataOffset += FFeatureVectorHelper::EncodeQuatCardinality;
-			}
+			DataOffset += FFeatureVectorHelper::EncodeQuatCardinality;
 		}
 		if (EnumHasAnyFlags(SampledBone.Flags, EPoseSearchBoneFlags::Velocity))
 		{
-			for (int32 SubsampleIdx = 0; SubsampleIdx != SampleTimes.Num(); ++SubsampleIdx)
-			{
-				DataOffset += FFeatureVectorHelper::EncodeVectorCardinality;
-			}
+			DataOffset += FFeatureVectorHelper::EncodeVectorCardinality;
 		}
 		if (EnumHasAnyFlags(SampledBone.Flags, EPoseSearchBoneFlags::Phase))
 		{
-			for (int32 SubsampleIdx = 0; SubsampleIdx != SampleTimes.Num(); ++SubsampleIdx)
-			{
-				DataOffset += FFeatureVectorHelper::EncodeVector2DCardinality;
-			}
+			DataOffset += FFeatureVectorHelper::EncodeVector2DCardinality;
 		}
 	}
 
@@ -622,47 +604,35 @@ void UPoseSearchFeatureChannel_Pose::FillWeights(TArray<float>& Weights) const
 		const FPoseSearchBone& SampledBone = SampledBones[ChannelBoneIdx];
 		if (EnumHasAnyFlags(SampledBone.Flags, EPoseSearchBoneFlags::Position))
 		{
-			for (int32 SubsampleIdx = 0; SubsampleIdx != SampleTimes.Num(); ++SubsampleIdx)
+			for (int32 i = 0; i != FFeatureVectorHelper::EncodeVectorCardinality; ++i)
 			{
-				for (int32 i = 0; i != FFeatureVectorHelper::EncodeVectorCardinality; ++i)
-				{
-					Weights[DataOffset + i] = Weight * SampledBone.Weight;
-				}
-				DataOffset += FFeatureVectorHelper::EncodeVectorCardinality;
+				Weights[DataOffset + i] = Weight * SampledBone.Weight;
 			}
+			DataOffset += FFeatureVectorHelper::EncodeVectorCardinality;
 		}
 		if (EnumHasAnyFlags(SampledBone.Flags, EPoseSearchBoneFlags::Rotation))
 		{
-			for (int32 SubsampleIdx = 0; SubsampleIdx != SampleTimes.Num(); ++SubsampleIdx)
+			for (int32 i = 0; i != FFeatureVectorHelper::EncodeQuatCardinality; ++i)
 			{
-				for (int32 i = 0; i != FFeatureVectorHelper::EncodeQuatCardinality; ++i)
-				{
-					Weights[DataOffset + i] = Weight * SampledBone.Weight;
-				}
-				DataOffset += FFeatureVectorHelper::EncodeQuatCardinality;
+				Weights[DataOffset + i] = Weight * SampledBone.Weight;
 			}
+			DataOffset += FFeatureVectorHelper::EncodeQuatCardinality;
 		}
 		if (EnumHasAnyFlags(SampledBone.Flags, EPoseSearchBoneFlags::Velocity))
 		{
-			for (int32 SubsampleIdx = 0; SubsampleIdx != SampleTimes.Num(); ++SubsampleIdx)
+			for (int32 i = 0; i != FFeatureVectorHelper::EncodeVectorCardinality; ++i)
 			{
-				for (int32 i = 0; i != FFeatureVectorHelper::EncodeVectorCardinality; ++i)
-				{
-					Weights[DataOffset + i] = Weight * SampledBone.Weight;
-				}
-				DataOffset += FFeatureVectorHelper::EncodeVectorCardinality;
+				Weights[DataOffset + i] = Weight * SampledBone.Weight;
 			}
+			DataOffset += FFeatureVectorHelper::EncodeVectorCardinality;
 		}
 		if (EnumHasAnyFlags(SampledBone.Flags, EPoseSearchBoneFlags::Phase))
 		{
-			for (int32 SubsampleIdx = 0; SubsampleIdx != SampleTimes.Num(); ++SubsampleIdx)
+			for (int32 i = 0; i != FFeatureVectorHelper::EncodeVector2DCardinality; ++i)
 			{
-				for (int32 i = 0; i != FFeatureVectorHelper::EncodeVector2DCardinality; ++i)
-				{
-					Weights[DataOffset + i] = Weight * SampledBone.Weight;
-				}
-				DataOffset += FFeatureVectorHelper::EncodeVector2DCardinality;
+				Weights[DataOffset + i] = Weight * SampledBone.Weight;
 			}
+			DataOffset += FFeatureVectorHelper::EncodeVector2DCardinality;
 		}
 	}
 
@@ -741,7 +711,7 @@ void UPoseSearchFeatureChannel_Pose::AddPoseFeatures(UE::PoseSearch::IAssetIndex
 	using namespace UE::PoseSearch;
 	using FSampleInfo = IAssetIndexer::FSampleInfo;
 
-	if (SampledBones.IsEmpty() || SampleTimes.IsEmpty())
+	if (SampledBones.IsEmpty())
 	{
 		return;
 	}
@@ -762,69 +732,50 @@ void UPoseSearchFeatureChannel_Pose::AddPoseFeatures(UE::PoseSearch::IAssetIndex
 
 		if (EnumHasAnyFlags(SampledBone.Flags, EPoseSearchBoneFlags::Position))
 		{
-			for (int32 SubsampleIdx = 0; SubsampleIdx != SampleTimes.Num(); ++SubsampleIdx)
-			{
-				const float SubsampleTime = SampleTime + SampleTimes[SubsampleIdx];
-
-				bool ClampedPresent;
-				const FTransform BoneTransformsPresent = Indexer.GetTransformAndCacheResults(SubsampleTime, UseCharacterSpaceVelocities ? SubsampleTime : SampleTime, SchemaBoneIdx[ChannelBoneIdx], ClampedPresent);
-				FFeatureVectorHelper::EncodeVector(FeatureVector, DataOffset, BoneTransformsPresent.GetTranslation());
-			}
+			bool ClampedPresent;
+			const FTransform BoneTransformsPresent = Indexer.GetTransformAndCacheResults(SampleTime, SampleTime, SchemaBoneIdx[ChannelBoneIdx], ClampedPresent);
+			FFeatureVectorHelper::EncodeVector(FeatureVector, DataOffset, BoneTransformsPresent.GetTranslation());
 		}
 
 		if (EnumHasAnyFlags(SampledBone.Flags, EPoseSearchBoneFlags::Rotation))
 		{
-			for (int32 SubsampleIdx = 0; SubsampleIdx != SampleTimes.Num(); ++SubsampleIdx)
-			{
-				const float SubsampleTime = SampleTime + SampleTimes[SubsampleIdx];
-
-				bool ClampedPresent;
-				const FTransform BoneTransformsPresent = Indexer.GetTransformAndCacheResults(SubsampleTime, UseCharacterSpaceVelocities ? SubsampleTime : SampleTime, SchemaBoneIdx[ChannelBoneIdx], ClampedPresent);
-				FFeatureVectorHelper::EncodeQuat(FeatureVector, DataOffset, BoneTransformsPresent.GetRotation());
-			}
+			bool ClampedPresent;
+			const FTransform BoneTransformsPresent = Indexer.GetTransformAndCacheResults(SampleTime, SampleTime, SchemaBoneIdx[ChannelBoneIdx], ClampedPresent);
+			FFeatureVectorHelper::EncodeQuat(FeatureVector, DataOffset, BoneTransformsPresent.GetRotation());
 		}
 
 		if (EnumHasAnyFlags(SampledBone.Flags, EPoseSearchBoneFlags::Velocity))
 		{
-			for (int32 SubsampleIdx = 0; SubsampleIdx != SampleTimes.Num(); ++SubsampleIdx)
-			{
-				const float SubsampleTime = SampleTime + SampleTimes[SubsampleIdx];
+			bool ClampedPast, ClampedPresent, ClampedFuture;
+			const FTransform BoneTransformsPast = Indexer.GetTransformAndCacheResults(SampleTime - SamplingContext->FiniteDelta, UseCharacterSpaceVelocities ? SampleTime - SamplingContext->FiniteDelta : SampleTime, SchemaBoneIdx[ChannelBoneIdx], ClampedPast);
+			const FTransform BoneTransformsPresent = Indexer.GetTransformAndCacheResults(SampleTime, UseCharacterSpaceVelocities ? SampleTime : SampleTime, SchemaBoneIdx[ChannelBoneIdx], ClampedPresent);
+			const FTransform BoneTransformsFuture = Indexer.GetTransformAndCacheResults(SampleTime + SamplingContext->FiniteDelta, UseCharacterSpaceVelocities ? SampleTime + SamplingContext->FiniteDelta : SampleTime, SchemaBoneIdx[ChannelBoneIdx], ClampedFuture);
 
-				bool ClampedPast, ClampedPresent, ClampedFuture;
-				const FTransform BoneTransformsPast = Indexer.GetTransformAndCacheResults(SubsampleTime - SamplingContext->FiniteDelta, UseCharacterSpaceVelocities ? SubsampleTime - SamplingContext->FiniteDelta : SampleTime, SchemaBoneIdx[ChannelBoneIdx], ClampedPast);
-				const FTransform BoneTransformsPresent = Indexer.GetTransformAndCacheResults(SubsampleTime, UseCharacterSpaceVelocities ? SubsampleTime : SampleTime, SchemaBoneIdx[ChannelBoneIdx], ClampedPresent);
-				const FTransform BoneTransformsFuture = Indexer.GetTransformAndCacheResults(SubsampleTime + SamplingContext->FiniteDelta, UseCharacterSpaceVelocities ? SubsampleTime + SamplingContext->FiniteDelta : SampleTime, SchemaBoneIdx[ChannelBoneIdx], ClampedFuture);
-
-				// We can get a better finite difference if we ignore samples that have
-				// been clamped at either side of the clip. However, if the central sample 
-				// itself is clamped, or there are no samples that are clamped, we can just 
-				// use the central difference as normal.
-				FVector LinearVelocity;
-				if (ClampedPast && !ClampedPresent && !ClampedFuture)
-				{
-					LinearVelocity = (BoneTransformsFuture.GetTranslation() - BoneTransformsPresent.GetTranslation()) / SamplingContext->FiniteDelta;
-				}
-				else if (ClampedFuture && !ClampedPresent && !ClampedPast)
-				{
-					LinearVelocity = (BoneTransformsPresent.GetTranslation() - BoneTransformsPast.GetTranslation()) / SamplingContext->FiniteDelta;
-				}
-				else
-				{
-					LinearVelocity = (BoneTransformsFuture.GetTranslation() - BoneTransformsPast.GetTranslation()) / (SamplingContext->FiniteDelta * 2.0f);
-				}
-
-				FFeatureVectorHelper::EncodeVector(FeatureVector, DataOffset, LinearVelocity);
+			// We can get a better finite difference if we ignore samples that have
+			// been clamped at either side of the clip. However, if the central sample 
+			// itself is clamped, or there are no samples that are clamped, we can just 
+			// use the central difference as normal.
+			FVector LinearVelocity;
+			if (ClampedPast && !ClampedPresent && !ClampedFuture)
+			{
+				LinearVelocity = (BoneTransformsFuture.GetTranslation() - BoneTransformsPresent.GetTranslation()) / SamplingContext->FiniteDelta;
+			}
+			else if (ClampedFuture && !ClampedPresent && !ClampedPast)
+			{
+				LinearVelocity = (BoneTransformsPresent.GetTranslation() - BoneTransformsPast.GetTranslation()) / SamplingContext->FiniteDelta;
+			}
+			else
+			{
+				LinearVelocity = (BoneTransformsFuture.GetTranslation() - BoneTransformsPast.GetTranslation()) / (SamplingContext->FiniteDelta * 2.0f);
 			}
+
+			FFeatureVectorHelper::EncodeVector(FeatureVector, DataOffset, LinearVelocity);
 		}
 
 		if (EnumHasAnyFlags(SampledBone.Flags, EPoseSearchBoneFlags::Phase))
 		{
-			for (int32 SubsampleIdx = 0; SubsampleIdx != SampleTimes.Num(); ++SubsampleIdx)
-			{
-				// @todo: support for SubsampleIdx
-				const int32 VectorIdx = SampleIdx - IndexingContext.BeginSampleIdx;
-				FFeatureVectorHelper::EncodeVector2D(FeatureVector, DataOffset, Phases[ChannelBoneIdx][VectorIdx]);
-			}
+			const int32 VectorIdx = SampleIdx - IndexingContext.BeginSampleIdx;
+			FFeatureVectorHelper::EncodeVector2D(FeatureVector, DataOffset, Phases[ChannelBoneIdx][VectorIdx]);
 		}
 	}
 
@@ -848,34 +799,22 @@ bool UPoseSearchFeatureChannel_Pose::BuildQuery(UE::PoseSearch::FSearchContext&
 				const FPoseSearchBone& SampledBone = SampledBones[SampledBoneIdx];
 				if (EnumHasAnyFlags(SampledBone.Flags, EPoseSearchBoneFlags::Position))
 				{
-					for (int32 SubsampleIdx = 0; SubsampleIdx != SampleTimes.Num(); ++SubsampleIdx)
-					{
-						FFeatureVectorHelper::EncodeVector(InOutQuery.EditValues(), DataOffset, SearchContext.GetCurrentResultPrevPoseVector(), SearchContext.GetCurrentResultPoseVector(), SearchContext.GetCurrentResultNextPoseVector(), LerpValue);
-					}
+					FFeatureVectorHelper::EncodeVector(InOutQuery.EditValues(), DataOffset, SearchContext.GetCurrentResultPrevPoseVector(), SearchContext.GetCurrentResultPoseVector(), SearchContext.GetCurrentResultNextPoseVector(), LerpValue);
 				}
 
 				if (EnumHasAnyFlags(SampledBone.Flags, EPoseSearchBoneFlags::Rotation))
 				{
-					for (int32 SubsampleIdx = 0; SubsampleIdx != SampleTimes.Num(); ++SubsampleIdx)
-					{
-						FFeatureVectorHelper::EncodeQuat(InOutQuery.EditValues(), DataOffset, SearchContext.GetCurrentResultPrevPoseVector(), SearchContext.GetCurrentResultPoseVector(), SearchContext.GetCurrentResultNextPoseVector(), LerpValue);
-					}
+					FFeatureVectorHelper::EncodeQuat(InOutQuery.EditValues(), DataOffset, SearchContext.GetCurrentResultPrevPoseVector(), SearchContext.GetCurrentResultPoseVector(), SearchContext.GetCurrentResultNextPoseVector(), LerpValue);
 				}
 
 				if (EnumHasAnyFlags(SampledBone.Flags, EPoseSearchBoneFlags::Velocity))
 				{
-					for (int32 SubsampleIdx = 0; SubsampleIdx != SampleTimes.Num(); ++SubsampleIdx)
-					{
-						FFeatureVectorHelper::EncodeVector(InOutQuery.EditValues(), DataOffset, SearchContext.GetCurrentResultPrevPoseVector(), SearchContext.GetCurrentResultPoseVector(), SearchContext.GetCurrentResultNextPoseVector(), LerpValue);
-					}
+					FFeatureVectorHelper::EncodeVector(InOutQuery.EditValues(), DataOffset, SearchContext.GetCurrentResultPrevPoseVector(), SearchContext.GetCurrentResultPoseVector(), SearchContext.GetCurrentResultNextPoseVector(), LerpValue);
 				}
 
 				if (EnumHasAnyFlags(SampledBone.Flags, EPoseSearchBoneFlags::Phase))
 				{
-					for (int32 SubsampleIdx = 0; SubsampleIdx != SampleTimes.Num(); ++SubsampleIdx)
-					{
-						FFeatureVectorHelper::EncodeVector2D(InOutQuery.EditValues(), DataOffset, SearchContext.GetCurrentResultPrevPoseVector(), SearchContext.GetCurrentResultPoseVector(), SearchContext.GetCurrentResultNextPoseVector(), LerpValue);
-					}
+					FFeatureVectorHelper::EncodeVector2D(InOutQuery.EditValues(), DataOffset, SearchContext.GetCurrentResultPrevPoseVector(), SearchContext.GetCurrentResultPoseVector(), SearchContext.GetCurrentResultNextPoseVector(), LerpValue);
 				}
 			}
 		}
@@ -890,42 +829,33 @@ bool UPoseSearchFeatureChannel_Pose::BuildQuery(UE::PoseSearch::FSearchContext&
 		bool Valid = false;
 	};
 	TArray<CachedTransforms, TInlineAllocator<32>> CachedTransforms;
-	CachedTransforms.AddUninitialized(SampleTimes.Num() * SampledBones.Num());
+	CachedTransforms.AddUninitialized(SampledBones.Num());
 
 	bool bAnyError = false;
-	for (int32 SubsampleIdx = 0; SubsampleIdx != SampleTimes.Num(); ++SubsampleIdx)
+	float SampleTime = 0.f;
+
+	for (int32 SampledBoneIdx = 0; SampledBoneIdx != SampledBones.Num(); ++SampledBoneIdx)
 	{
-		// Stop when we've reached future samples
-		float SampleTime = SampleTimes[SubsampleIdx];
-		if (SampleTime > 0.0f)
-		{
-			break;
-		}
+		CachedTransforms[SampledBoneIdx].Current = SearchContext.TryGetTransformAndCacheResults(SampleTime, InOutQuery.GetSchema(), SchemaBoneIdx[SampledBoneIdx], bAnyError);
+		const FPoseSearchBone& SampledBone = SampledBones[SampledBoneIdx];
 
-		for (int32 SampledBoneIdx = 0; SampledBoneIdx != SampledBones.Num(); ++SampledBoneIdx)
+		if (EnumHasAnyFlags(SampledBone.Flags, EPoseSearchBoneFlags::Velocity))
 		{
-			const int32 CachedTransformsIndex = SubsampleIdx * SampledBones.Num() + SampledBoneIdx;
-			CachedTransforms[CachedTransformsIndex].Current = SearchContext.TryGetTransformAndCacheResults(SampleTime, InOutQuery.GetSchema(), SchemaBoneIdx[SampledBoneIdx], bAnyError);
-			const FPoseSearchBone& SampledBone = SampledBones[SampledBoneIdx];
+			check(SearchContext.History);
+			const float HistorySameplInterval = SearchContext.History->GetSampleTimeInterval();
 
-			if (EnumHasAnyFlags(SampledBone.Flags, EPoseSearchBoneFlags::Velocity))
-			{
-				check(SearchContext.History);
-				const float HistorySameplInterval = SearchContext.History->GetSampleTimeInterval();
-
-				CachedTransforms[CachedTransformsIndex].Previous = SearchContext.TryGetTransformAndCacheResults(SampleTime - HistorySameplInterval, InOutQuery.GetSchema(), SchemaBoneIdx[SampledBoneIdx], bAnyError);
+			CachedTransforms[SampledBoneIdx].Previous = SearchContext.TryGetTransformAndCacheResults(SampleTime - HistorySameplInterval, InOutQuery.GetSchema(), SchemaBoneIdx[SampledBoneIdx], bAnyError);
 
-				if (!UE::PoseSearch::UseCharacterSpaceVelocities)
-				{
-					const FTransform RootTransform = SearchContext.TryGetTransformAndCacheResults(SampleTime, InOutQuery.GetSchema(), FSearchContext::SchemaRootBoneIdx, bAnyError);
-					const FTransform RootTransformPrev = SearchContext.TryGetTransformAndCacheResults(SampleTime - HistorySameplInterval, InOutQuery.GetSchema(), FSearchContext::SchemaRootBoneIdx, bAnyError);
+			if (!UE::PoseSearch::UseCharacterSpaceVelocities)
+			{
+				const FTransform RootTransform = SearchContext.TryGetTransformAndCacheResults(SampleTime, InOutQuery.GetSchema(), FSearchContext::SchemaRootBoneIdx, bAnyError);
+				const FTransform RootTransformPrev = SearchContext.TryGetTransformAndCacheResults(SampleTime - HistorySameplInterval, InOutQuery.GetSchema(), FSearchContext::SchemaRootBoneIdx, bAnyError);
 
-					// animation space velocity
-					CachedTransforms[CachedTransformsIndex].Previous = CachedTransforms[CachedTransformsIndex].Previous * (RootTransformPrev * RootTransform.Inverse());
-				}
+				// animation space velocity
+				CachedTransforms[SampledBoneIdx].Previous = CachedTransforms[SampledBoneIdx].Previous * (RootTransformPrev * RootTransform.Inverse());
 			}
-			CachedTransforms[CachedTransformsIndex].Valid = true;
 		}
+		CachedTransforms[SampledBoneIdx].Valid = true;
 	}
 
 	if (bAnyError)
@@ -939,66 +869,50 @@ bool UPoseSearchFeatureChannel_Pose::BuildQuery(UE::PoseSearch::FSearchContext&
 		const FPoseSearchBone& SampledBone = SampledBones[SampledBoneIdx];
 		if (EnumHasAnyFlags(SampledBone.Flags, EPoseSearchBoneFlags::Position))
 		{
-			for (int32 SubsampleIdx = 0; SubsampleIdx != SampleTimes.Num(); ++SubsampleIdx)
+			if (CachedTransforms[SampledBoneIdx].Valid)
 			{
-				const int32 CachedTransformsIndex = SubsampleIdx * SampledBones.Num() + SampledBoneIdx;
-				if (CachedTransforms[CachedTransformsIndex].Valid)
-				{
-					FFeatureVectorHelper::EncodeVector(InOutQuery.EditValues(), DataOffset, CachedTransforms[CachedTransformsIndex].Current.GetTranslation());
-				}
-				else
-				{
-					// preserve the InOutQuery.EditValues() and increase the DataOffset
-					DataOffset += FFeatureVectorHelper::EncodeVectorCardinality;
-				}
+				FFeatureVectorHelper::EncodeVector(InOutQuery.EditValues(), DataOffset, CachedTransforms[SampledBoneIdx].Current.GetTranslation());
+			}
+			else
+			{
+				// preserve the InOutQuery.EditValues() and increase the DataOffset
+				DataOffset += FFeatureVectorHelper::EncodeVectorCardinality;
 			}
 		}
 
 		if (EnumHasAnyFlags(SampledBone.Flags, EPoseSearchBoneFlags::Rotation))
 		{
-			for (int32 SubsampleIdx = 0; SubsampleIdx != SampleTimes.Num(); ++SubsampleIdx)
+			if (CachedTransforms[SampledBoneIdx].Valid)
 			{
-				const int32 CachedTransformsIndex = SubsampleIdx * SampledBones.Num() + SampledBoneIdx;
-				if (CachedTransforms[CachedTransformsIndex].Valid)
-				{
-					FFeatureVectorHelper::EncodeQuat(InOutQuery.EditValues(), DataOffset, CachedTransforms[CachedTransformsIndex].Current.GetRotation());
-				}
-				else
-				{
-					// preserve the InOutQuery.EditValues() and increase the DataOffset
-					DataOffset += FFeatureVectorHelper::EncodeQuatCardinality;
-				}
+				FFeatureVectorHelper::EncodeQuat(InOutQuery.EditValues(), DataOffset, CachedTransforms[SampledBoneIdx].Current.GetRotation());
+			}
+			else
+			{
+				// preserve the InOutQuery.EditValues() and increase the DataOffset
+				DataOffset += FFeatureVectorHelper::EncodeQuatCardinality;
 			}
 		}
 
 		if (EnumHasAnyFlags(SampledBone.Flags, EPoseSearchBoneFlags::Velocity))
 		{
-			for (int32 SubsampleIdx = 0; SubsampleIdx != SampleTimes.Num(); ++SubsampleIdx)
+			if (CachedTransforms[SampledBoneIdx].Valid)
 			{
-				const int32 CachedTransformsIndex = SubsampleIdx * SampledBones.Num() + SampledBoneIdx;
-				if (CachedTransforms[CachedTransformsIndex].Valid)
-				{
-					const FVector LinearVelocity = (CachedTransforms[CachedTransformsIndex].Current.GetTranslation() - CachedTransforms[CachedTransformsIndex].Previous.GetTranslation()) / SearchContext.History->GetSampleTimeInterval();
-					FFeatureVectorHelper::EncodeVector(InOutQuery.EditValues(), DataOffset, LinearVelocity);
-				}
-				else
-				{
-					// preserve the InOutQuery.EditValues() and increase the DataOffset
-					DataOffset += FFeatureVectorHelper::EncodeVectorCardinality;
-				}
+				const FVector LinearVelocity = (CachedTransforms[SampledBoneIdx].Current.GetTranslation() - CachedTransforms[SampledBoneIdx].Previous.GetTranslation()) / SearchContext.History->GetSampleTimeInterval();
+				FFeatureVectorHelper::EncodeVector(InOutQuery.EditValues(), DataOffset, LinearVelocity);
+			}
+			else
+			{
+				// preserve the InOutQuery.EditValues() and increase the DataOffset
+				DataOffset += FFeatureVectorHelper::EncodeVectorCardinality;
 			}
 		}
 
 		if (EnumHasAnyFlags(SampledBone.Flags, EPoseSearchBoneFlags::Phase))
 		{
-			for (int32 SubsampleIdx = 0; SubsampleIdx != SampleTimes.Num(); ++SubsampleIdx)
-			{
-				const int32 CachedTransformsIndex = SubsampleIdx * SampledBones.Num() + SampledBoneIdx;
-				DataOffset += FFeatureVectorHelper::EncodeVector2DCardinality;
+			DataOffset += FFeatureVectorHelper::EncodeVector2DCardinality;
 
-				// @todo: Support phase in BuildQuery
-				// FFeatureVectorHelper::EncodeVector2D(InOutQuery.EditValues(), DataOffset, ???);
-			}
+			// @todo: Support phase in BuildQuery
+			// FFeatureVectorHelper::EncodeVector2D(InOutQuery.EditValues(), DataOffset, ???);
 		}
 	}
 
@@ -1019,10 +933,9 @@ void UPoseSearchFeatureChannel_Pose::DebugDraw(const UE::PoseSearch::FDebugDrawP
 	const uint8 DepthPriority = ESceneDepthPriorityGroup::SDPG_Foreground + 2;
 	const bool bPersistent = EnumHasAnyFlags(DrawParams.Flags, EDebugDrawFlags::Persistent);
 
-	const int32 NumSubsamples = SampleTimes.Num();
 	const int32 NumBones = SampledBones.Num();
 
-	if ((NumSubsamples * NumBones) == 0)
+	if (NumBones == 0)
 	{
 		return;
 	}
@@ -1032,108 +945,91 @@ void UPoseSearchFeatureChannel_Pose::DebugDraw(const UE::PoseSearch::FDebugDrawP
 	{
 		const FPoseSearchBone& SampledBone = SampledBones[ChannelBoneIdx];
 
-		TArray<FVector, TInlineAllocator<32>> BonePos;
-		BonePos.AddUninitialized(NumSubsamples);
+		FVector BonePos;
 		if (EnumHasAnyFlags(SampledBone.Flags, EPoseSearchBoneFlags::Position))
 		{
-			for (int32 SubsampleIdx = 0; SubsampleIdx != NumSubsamples; ++SubsampleIdx)
-			{
-				BonePos[SubsampleIdx] = FFeatureVectorHelper::DecodeVector(PoseVector, DataOffset);
+			BonePos = FFeatureVectorHelper::DecodeVector(PoseVector, DataOffset);
 
-				const FColor Color = DrawParams.GetColor(SampledBone.ColorPresetIndex);
+			const FColor Color = DrawParams.GetColor(SampledBone.ColorPresetIndex);
 
-				BonePos[SubsampleIdx] = DrawParams.RootTransform.TransformPosition(BonePos[SubsampleIdx]);
-				if (EnumHasAnyFlags(DrawParams.Flags, EDebugDrawFlags::DrawFast | EDebugDrawFlags::DrawSearchIndex))
-				{
-					DrawDebugPoint(DrawParams.World, BonePos[SubsampleIdx], DrawParams.PointSize, Color, bPersistent, LifeTime, DepthPriority);
-				}
-				else
-				{
-					DrawDebugSphere(DrawParams.World, BonePos[SubsampleIdx], DrawDebugSphereSize, DrawDebugSphereSegments, Color, bPersistent, LifeTime, DepthPriority);
-				}
+			BonePos = DrawParams.RootTransform.TransformPosition(BonePos);
+			if (EnumHasAnyFlags(DrawParams.Flags, EDebugDrawFlags::DrawFast | EDebugDrawFlags::DrawSearchIndex))
+			{
+				DrawDebugPoint(DrawParams.World, BonePos, DrawParams.PointSize, Color, bPersistent, LifeTime, DepthPriority);
+			}
+			else
+			{
+				DrawDebugSphere(DrawParams.World, BonePos, DrawDebugSphereSize, DrawDebugSphereSegments, Color, bPersistent, LifeTime, DepthPriority);
+			}
 
-				if (EnumHasAnyFlags(DrawParams.Flags, EDebugDrawFlags::DrawBoneNames))
-				{
-					DrawDebugString(
-						DrawParams.World, BonePos[SubsampleIdx] + FVector(0.0, 0.0, 10.0),
-						Schema->BoneReferences[SchemaBoneIdx[ChannelBoneIdx]].BoneName.ToString(),
-						nullptr, Color, LifeTime, false, 1.0f);
-				}
+			if (EnumHasAnyFlags(DrawParams.Flags, EDebugDrawFlags::DrawBoneNames))
+			{
+				DrawDebugString(
+					DrawParams.World, BonePos + FVector(0.0, 0.0, 10.0),
+					Schema->BoneReferences[SchemaBoneIdx[ChannelBoneIdx]].BoneName.ToString(),
+					nullptr, Color, LifeTime, false, 1.0f);
 			}
 		}
 		else
 		{
-			for (int32 SubsampleIdx = 0; SubsampleIdx != NumSubsamples; ++SubsampleIdx)
-			{
-				BonePos[SubsampleIdx] = DrawParams.Mesh != nullptr ? DrawParams.Mesh->GetSocketTransform(SampledBones[ChannelBoneIdx].Reference.BoneName).GetLocation() : DrawParams.RootTransform.GetTranslation();
-			}
+			BonePos = DrawParams.Mesh != nullptr ? DrawParams.Mesh->GetSocketTransform(SampledBones[ChannelBoneIdx].Reference.BoneName).GetLocation() : DrawParams.RootTransform.GetTranslation();
 		}
 
 		if (EnumHasAnyFlags(SampledBone.Flags, EPoseSearchBoneFlags::Rotation))
 		{
-			for (int32 SubsampleIdx = 0; SubsampleIdx != NumSubsamples; ++SubsampleIdx)
-			{
-				const FQuat BoneRot = FFeatureVectorHelper::DecodeQuat(PoseVector, DataOffset);
-
-				// @todo: debug draw rotation
-			}
+			const FQuat BoneRot = FFeatureVectorHelper::DecodeQuat(PoseVector, DataOffset);
+			// @todo: debug draw rotation
 		}
 
 		if (EnumHasAnyFlags(SampledBone.Flags, EPoseSearchBoneFlags::Velocity))
 		{
-			for (int32 SubsampleIdx = 0; SubsampleIdx != NumSubsamples; ++SubsampleIdx)
-			{
-				FVector BoneVel = FFeatureVectorHelper::DecodeVector(PoseVector, DataOffset);
+			FVector BoneVel = FFeatureVectorHelper::DecodeVector(PoseVector, DataOffset);
 
-				const FColor Color = DrawParams.GetColor(SampledBone.ColorPresetIndex);
+			const FColor Color = DrawParams.GetColor(SampledBone.ColorPresetIndex);
 
-				BoneVel *= DrawDebugVelocityScale;
-				BoneVel = DrawParams.RootTransform.TransformVector(BoneVel);
-				FVector BoneVelDirection = BoneVel.GetSafeNormal();
+			BoneVel *= DrawDebugVelocityScale;
+			BoneVel = DrawParams.RootTransform.TransformVector(BoneVel);
+			FVector BoneVelDirection = BoneVel.GetSafeNormal();
 
-				if (EnumHasAnyFlags(DrawParams.Flags, EDebugDrawFlags::DrawSearchIndex))
-				{
-					DrawDebugLine(DrawParams.World, BonePos[SubsampleIdx], BonePos[SubsampleIdx] + BoneVel, Color, bPersistent, LifeTime, DepthPriority);
-				}
-				else
-				{
-					const float AdjustedThickness = EnumHasAnyFlags(DrawParams.Flags, EDebugDrawFlags::DrawFast) ? 0.0f : DrawDebugLineThickness;
-
-					DrawDebugLine(
-						DrawParams.World,
-						BonePos[SubsampleIdx] + BoneVelDirection * DrawDebugSphereSize,
-						BonePos[SubsampleIdx] + BoneVel,
-						Color,
-						bPersistent,
-						LifeTime,
-						DepthPriority,
-						AdjustedThickness);
-				}
+			if (EnumHasAnyFlags(DrawParams.Flags, EDebugDrawFlags::DrawSearchIndex))
+			{
+				DrawDebugLine(DrawParams.World, BonePos, BonePos + BoneVel, Color, bPersistent, LifeTime, DepthPriority);
+			}
+			else
+			{
+				const float AdjustedThickness = EnumHasAnyFlags(DrawParams.Flags, EDebugDrawFlags::DrawFast) ? 0.0f : DrawDebugLineThickness;
+
+				DrawDebugLine(
+					DrawParams.World,
+					BonePos + BoneVelDirection * DrawDebugSphereSize,
+					BonePos + BoneVel,
+					Color,
+					bPersistent,
+					LifeTime,
+					DepthPriority,
+					AdjustedThickness);
 			}
 		}
 
 		if (EnumHasAnyFlags(SampledBone.Flags, EPoseSearchBoneFlags::Phase))
 		{
-			for (int32 SubsampleIdx = 0; SubsampleIdx != NumSubsamples; ++SubsampleIdx)
-			{
-				const FVector2D Phase = FFeatureVectorHelper::DecodeVector2D(PoseVector, DataOffset);
+			const FVector2D Phase = FFeatureVectorHelper::DecodeVector2D(PoseVector, DataOffset);
 
-				const FColor Color = DrawParams.GetColor(SampledBone.ColorPresetIndex);
+			const FColor Color = DrawParams.GetColor(SampledBone.ColorPresetIndex);
 
-				static float ScaleFactor = 1.f;
+			static float ScaleFactor = 1.f;
 
-				const FVector TransformXAxisVector = DrawParams.RootTransform.TransformVector(FVector::XAxisVector);
-				const FVector TransformYAxisVector = DrawParams.RootTransform.TransformVector(FVector::YAxisVector);
-				const FVector TransformZAxisVector = DrawParams.RootTransform.TransformVector(FVector::ZAxisVector);
+			const FVector TransformXAxisVector = DrawParams.RootTransform.TransformVector(FVector::XAxisVector);
+			const FVector TransformYAxisVector = DrawParams.RootTransform.TransformVector(FVector::YAxisVector);
+			const FVector TransformZAxisVector = DrawParams.RootTransform.TransformVector(FVector::ZAxisVector);
 
-				const FVector PhaseVector = (TransformZAxisVector * Phase.X + TransformYAxisVector * Phase.Y) * ScaleFactor;
-				DrawDebugLine(DrawParams.World, BonePos[SubsampleIdx], BonePos[SubsampleIdx] + PhaseVector, Color, bPersistent, LifeTime, DepthPriority, 0.f);
+			const FVector PhaseVector = (TransformZAxisVector * Phase.X + TransformYAxisVector * Phase.Y) * ScaleFactor;
+			DrawDebugLine(DrawParams.World, BonePos, BonePos + PhaseVector, Color, bPersistent, LifeTime, DepthPriority, 0.f);
 
-				static int32 Segments = 64;
-				FMatrix CircleTransform;
-				CircleTransform.SetAxes(&TransformXAxisVector, &TransformYAxisVector, &TransformZAxisVector, &BonePos[SubsampleIdx]);
-				DrawDebugCircle(DrawParams.World, CircleTransform, PhaseVector.Length(), Segments, Color, bPersistent, LifeTime, DepthPriority, 0.f, false);
-			}
+			static int32 Segments = 64;
+			FMatrix CircleTransform;
+			CircleTransform.SetAxes(&TransformXAxisVector, &TransformYAxisVector, &TransformZAxisVector, &BonePos);
+			DrawDebugCircle(DrawParams.World, CircleTransform, PhaseVector.Length(), Segments, Color, bPersistent, LifeTime, DepthPriority, 0.f, false);
 		}
 	}
 
@@ -1147,14 +1043,14 @@ void UPoseSearchFeatureChannel_Pose::PopulateChannelLayoutSet(UE::PoseSearch::FF
 	using namespace UE::PoseSearch;
 	int32 DataOffset = ChannelDataOffset;
 
-	auto Add = [&FeatureChannelLayoutSet, &DataOffset](const FPoseSearchBone& SampledBone, EPoseSearchBoneFlags BoneFlag, float SampleTime, const char* Label, int32 Cardinality)
+	auto Add = [&FeatureChannelLayoutSet, &DataOffset](const FPoseSearchBone& SampledBone, EPoseSearchBoneFlags BoneFlag, const char* Label, int32 Cardinality)
 	{
 		FString SkeletonName = FeatureChannelLayoutSet.CurrentSchema->Skeleton->GetName();
 		FString BoneName = SampledBone.Reference.BoneName.ToString();
 
 		UE::PoseSearch::FKeyBuilder KeyBuilder;
-		KeyBuilder << SkeletonName << BoneName << BoneFlag << SampleTime;
-		FeatureChannelLayoutSet.Add(FString::Format(TEXT("{0} {1} {2}"), { BoneName, Label, SampleTime }), KeyBuilder.Finalize(), DataOffset, Cardinality);
+		KeyBuilder << SkeletonName << BoneName << BoneFlag;
+		FeatureChannelLayoutSet.Add(FString::Format(TEXT("{0} {1}"), { BoneName, Label }), KeyBuilder.Finalize(), DataOffset, Cardinality);
 
 		DataOffset += Cardinality;
 	};
@@ -1164,31 +1060,19 @@ void UPoseSearchFeatureChannel_Pose::PopulateChannelLayoutSet(UE::PoseSearch::FF
 		const FPoseSearchBone& SampledBone = SampledBones[ChannelBoneIdx];
 		if (EnumHasAnyFlags(SampledBone.Flags, EPoseSearchBoneFlags::Position))
 		{
-			for (int32 SubsampleIdx = 0; SubsampleIdx != SampleTimes.Num(); ++SubsampleIdx)
-			{
-				Add(SampledBone, EPoseSearchBoneFlags::Position, SampleTimes[SubsampleIdx], "Pos", FFeatureVectorHelper::EncodeVectorCardinality);
-			}
+			Add(SampledBone, EPoseSearchBoneFlags::Position, "Pos", FFeatureVectorHelper::EncodeVectorCardinality);
 		}
 		if (EnumHasAnyFlags(SampledBone.Flags, EPoseSearchBoneFlags::Rotation))
 		{
-			for (int32 SubsampleIdx = 0; SubsampleIdx != SampleTimes.Num(); ++SubsampleIdx)
-			{
-				Add(SampledBone, EPoseSearchBoneFlags::Rotation, SampleTimes[SubsampleIdx], "Rot", FFeatureVectorHelper::EncodeQuatCardinality);
-			}
+			Add(SampledBone, EPoseSearchBoneFlags::Rotation, "Rot", FFeatureVectorHelper::EncodeQuatCardinality);
 		}
 		if (EnumHasAnyFlags(SampledBone.Flags, EPoseSearchBoneFlags::Velocity))
 		{
-			for (int32 SubsampleIdx = 0; SubsampleIdx != SampleTimes.Num(); ++SubsampleIdx)
-			{
-				Add(SampledBone, EPoseSearchBoneFlags::Velocity, SampleTimes[SubsampleIdx], "Vel", FFeatureVectorHelper::EncodeVectorCardinality);
-			}
+			Add(SampledBone, EPoseSearchBoneFlags::Velocity, "Vel", FFeatureVectorHelper::EncodeVectorCardinality);
 		}
 		if (EnumHasAnyFlags(SampledBone.Flags, EPoseSearchBoneFlags::Phase))
 		{
-			for (int32 SubsampleIdx = 0; SubsampleIdx != SampleTimes.Num(); ++SubsampleIdx)
-			{
-				Add(SampledBone, EPoseSearchBoneFlags::Phase, SampleTimes[SubsampleIdx], "Pha", FFeatureVectorHelper::EncodeVector2DCardinality);
-			}
+			Add(SampledBone, EPoseSearchBoneFlags::Phase, "Pha", FFeatureVectorHelper::EncodeVector2DCardinality);
 		}
 	}
 
@@ -1210,35 +1094,23 @@ void UPoseSearchFeatureChannel_Pose::ComputeCostBreakdowns(UE::PoseSearch::ICost
 			const FPoseSearchBone& SampledBone = SampledBones[ChannelBoneIdx];
 			if (EnumHasAnyFlags(SampledBone.Flags, EPoseSearchBoneFlags::Position))
 			{
-				for (int32 SubsampleIdx = 0; SubsampleIdx != SampleTimes.Num(); ++SubsampleIdx)
-				{
-					CostBreakDownData.AddEntireBreakDownSection(FText::Format(LOCTEXT("ColumnLabelPoseChannelPosition", "{0} Pos {1}"), FText::FromName(SampledBone.Reference.BoneName), SampleTimes[SubsampleIdx]), Schema, DataOffset, FFeatureVectorHelper::EncodeVectorCardinality);
-					DataOffset += FFeatureVectorHelper::EncodeVectorCardinality;
-				}
+				CostBreakDownData.AddEntireBreakDownSection(FText::Format(LOCTEXT("ColumnLabelPoseChannelPosition", "{0} Pos"), FText::FromName(SampledBone.Reference.BoneName)), Schema, DataOffset, FFeatureVectorHelper::EncodeVectorCardinality);
+				DataOffset += FFeatureVectorHelper::EncodeVectorCardinality;
 			}
 			if (EnumHasAnyFlags(SampledBone.Flags, EPoseSearchBoneFlags::Rotation))
 			{
-				for (int32 SubsampleIdx = 0; SubsampleIdx != SampleTimes.Num(); ++SubsampleIdx)
-				{
-					CostBreakDownData.AddEntireBreakDownSection(FText::Format(LOCTEXT("ColumnLabelPoseChannelRotation", "{0} Rot {1}"), FText::FromName(SampledBone.Reference.BoneName), SampleTimes[SubsampleIdx]), Schema, DataOffset, FFeatureVectorHelper::EncodeQuatCardinality);
-					DataOffset += FFeatureVectorHelper::EncodeQuatCardinality;
-				}
+				CostBreakDownData.AddEntireBreakDownSection(FText::Format(LOCTEXT("ColumnLabelPoseChannelRotation", "{0} Rot"), FText::FromName(SampledBone.Reference.BoneName)), Schema, DataOffset, FFeatureVectorHelper::EncodeQuatCardinality);
+				DataOffset += FFeatureVectorHelper::EncodeQuatCardinality;
 			}
 			if (EnumHasAnyFlags(SampledBone.Flags, EPoseSearchBoneFlags::Velocity))
 			{
-				for (int32 SubsampleIdx = 0; SubsampleIdx != SampleTimes.Num(); ++SubsampleIdx)
-				{
-					CostBreakDownData.AddEntireBreakDownSection(FText::Format(LOCTEXT("ColumnLabelPoseChannelVelocity", "{0} Vel {1}"), FText::FromName(SampledBone.Reference.BoneName), SampleTimes[SubsampleIdx]), Schema, DataOffset, FFeatureVectorHelper::EncodeVectorCardinality);
-					DataOffset += FFeatureVectorHelper::EncodeVectorCardinality;
-				}
+				CostBreakDownData.AddEntireBreakDownSection(FText::Format(LOCTEXT("ColumnLabelPoseChannelVelocity", "{0} Vel"), FText::FromName(SampledBone.Reference.BoneName)), Schema, DataOffset, FFeatureVectorHelper::EncodeVectorCardinality);
+				DataOffset += FFeatureVectorHelper::EncodeVectorCardinality;
 			}
 			if (EnumHasAnyFlags(SampledBone.Flags, EPoseSearchBoneFlags::Phase))
 			{
-				for (int32 SubsampleIdx = 0; SubsampleIdx != SampleTimes.Num(); ++SubsampleIdx)
-				{
-					CostBreakDownData.AddEntireBreakDownSection(FText::Format(LOCTEXT("ColumnLabelPoseChannelPhase", "{0} Pha {1}"), FText::FromName(SampledBone.Reference.BoneName), SampleTimes[SubsampleIdx]), Schema, DataOffset, FFeatureVectorHelper::EncodeVector2DCardinality);
-					DataOffset += FFeatureVectorHelper::EncodeVector2DCardinality;
-				}
+				CostBreakDownData.AddEntireBreakDownSection(FText::Format(LOCTEXT("ColumnLabelPoseChannelPhase", "{0} Pha"), FText::FromName(SampledBone.Reference.BoneName)), Schema, DataOffset, FFeatureVectorHelper::EncodeVector2DCardinality);
+				DataOffset += FFeatureVectorHelper::EncodeVector2DCardinality;
 			}
 		}
 
diff --git a/Engine/Plugins/Experimental/Animation/PoseSearch/Source/Runtime/Public/PoseSearch/PoseSearchFeatureChannels.h b/Engine/Plugins/Experimental/Animation/PoseSearch/Source/Runtime/Public/PoseSearch/PoseSearchFeatureChannels.h
--- a/Engine/Plugins/Experimental/Animation/PoseSearch/Source/Runtime/Public/PoseSearch/PoseSearchFeatureChannels.h
+++ b/Engine/Plugins/Experimental/Animation/PoseSearch/Source/Runtime/Public/PoseSearch/PoseSearchFeatureChannels.h
@@ -163,18 +163,12 @@ public:
 	UPROPERTY(EditAnywhere, Category = "Settings")
 	TArray<FPoseSearchBone> SampledBones;
 
-	UPROPERTY(EditAnywhere, Category = "Settings")
-	TArray<float> SampleTimes;
-
 	UPROPERTY()
 	TArray<int8> SchemaBoneIdx;
 
 	UPROPERTY(EditAnywhere, Category = "Settings")
 	EInputQueryPose InputQueryPose = EInputQueryPose::UseContinuingPose;
 
-	// UObject interface
-	virtual void PreSave(FObjectPreSaveContext ObjectSaveContext) override;
-
 	// UPoseSearchFeatureChannel interface
 	virtual void InitializeSchema(UE::PoseSearch::FSchemaInitializer& Initializer) override;
 	virtual void FillWeights(TArray<float>& Weights) const override;
