Remove unncessary plugin dependencies

diff --git a/Engine/Plugins/Experimental/Animation/PoseSearch/PoseSearch.uplugin b/Engine/Plugins/Experimental/Animation/PoseSearch/PoseSearch.uplugin
--- a/Engine/Plugins/Experimental/Animation/PoseSearch/PoseSearch.uplugin
+++ b/Engine/Plugins/Experimental/Animation/PoseSearch/PoseSearch.uplugin
@@ -35,14 +35,6 @@
 				"Editor"
 			]
 		},
-		{
-			"Name": "AnimationWarping",
-			"Enabled": true
-		},
-		{
-			"Name": "MotionTrajectory",
-			"Enabled": true
-		},
 		{
 			"Name": "StructUtils",
 			"Enabled": true
diff --git a/Engine/Plugins/Experimental/Animation/PoseSearch/Source/Runtime/PoseSearch.Build.cs b/Engine/Plugins/Experimental/Animation/PoseSearch/Source/Runtime/PoseSearch.Build.cs
--- a/Engine/Plugins/Experimental/Animation/PoseSearch/Source/Runtime/PoseSearch.Build.cs
+++ b/Engine/Plugins/Experimental/Animation/PoseSearch/Source/Runtime/PoseSearch.Build.cs
@@ -23,7 +23,6 @@ public class PoseSearch : ModuleRules
 				"TraceLog",
 				"AnimGraphRuntime",
 				"GameplayTags",
-				"MotionTrajectory",
 				"StructUtils"
 			}
 		);
diff --git a/Engine/Plugins/Experimental/Animation/PoseSearch/Source/Runtime/Private/PoseSearchLibrary.cpp b/Engine/Plugins/Experimental/Animation/PoseSearch/Source/Runtime/Private/PoseSearchLibrary.cpp
--- a/Engine/Plugins/Experimental/Animation/PoseSearch/Source/Runtime/Private/PoseSearchLibrary.cpp
+++ b/Engine/Plugins/Experimental/Animation/PoseSearch/Source/Runtime/Private/PoseSearchLibrary.cpp
@@ -10,7 +10,6 @@
 #include "Animation/BlendSpace.h"
 #include "PoseSearch/AnimNode_MotionMatching.h"
 #include "Trace/PoseSearchTraceLogger.h"
-#include "MotionTrajectoryLibrary.h"
 #include "PoseSearch/PoseSearchFeatureChannels.h"
 
 #include UE_INLINE_GENERATED_CPP_BY_NAME(PoseSearchLibrary)
