motion matching - fix for PoseSearchFeatureChannels.cpp(751): warning V583: The '?:' operator, regardless of its conditional expression, always returns one and the same value: SampleTime.

diff --git a/Engine/Plugins/Experimental/Animation/PoseSearch/Source/Runtime/Private/PoseSearchFeatureChannels.cpp b/Engine/Plugins/Experimental/Animation/PoseSearch/Source/Runtime/Private/PoseSearchFeatureChannels.cpp
--- a/Engine/Plugins/Experimental/Animation/PoseSearch/Source/Runtime/Private/PoseSearchFeatureChannels.cpp
+++ b/Engine/Plugins/Experimental/Animation/PoseSearch/Source/Runtime/Private/PoseSearchFeatureChannels.cpp
@@ -748,7 +748,7 @@ void UPoseSearchFeatureChannel_Pose::AddPoseFeatures(UE::PoseSearch::IAssetIndex
 		{
 			bool ClampedPast, ClampedPresent, ClampedFuture;
 			const FTransform BoneTransformsPast = Indexer.GetTransformAndCacheResults(SampleTime - SamplingContext->FiniteDelta, UseCharacterSpaceVelocities ? SampleTime - SamplingContext->FiniteDelta : SampleTime, SchemaBoneIdx[ChannelBoneIdx], ClampedPast);
-			const FTransform BoneTransformsPresent = Indexer.GetTransformAndCacheResults(SampleTime, UseCharacterSpaceVelocities ? SampleTime : SampleTime, SchemaBoneIdx[ChannelBoneIdx], ClampedPresent);
+			const FTransform BoneTransformsPresent = Indexer.GetTransformAndCacheResults(SampleTime, SampleTime, SchemaBoneIdx[ChannelBoneIdx], ClampedPresent);
 			const FTransform BoneTransformsFuture = Indexer.GetTransformAndCacheResults(SampleTime + SamplingContext->FiniteDelta, UseCharacterSpaceVelocities ? SampleTime + SamplingContext->FiniteDelta : SampleTime, SchemaBoneIdx[ChannelBoneIdx], ClampedFuture);
 
 			// We can get a better finite difference if we ignore samples that have
