motion matching - pose search plugin editor refactoring

diff --git a/Engine/Plugins/Experimental/Animation/PoseSearch/Source/Editor/Private/PoseSearchDatabaseAssetListItem.cpp b/Engine/Plugins/Experimental/Animation/PoseSearch/Source/Editor/Private/PoseSearchDatabaseAssetListItem.cpp
new file mode 100644
--- /dev/null
+++ b/Engine/Plugins/Experimental/Animation/PoseSearch/Source/Editor/Private/PoseSearchDatabaseAssetListItem.cpp
@@ -0,0 +1,485 @@
+// Copyright Epic Games, Inc. All Rights Reserved.
+
+#include "PoseSearchDatabaseAssetListItem.h"
+#include "PoseSearchDatabaseViewModel.h"
+
+#include "PoseSearch/PoseSearch.h"
+
+#include "Animation/AnimSequence.h"
+#include "Animation/AnimComposite.h"
+#include "Animation/BlendSpace.h"
+
+#include "Misc/TransactionObjectEvent.h"
+#include "DragAndDrop/AssetDragDropOp.h"
+#include "Misc/FeedbackContext.h"
+#include "AssetSelection.h"
+#include "ClassIconFinder.h"
+#include "DetailColumnSizeData.h"
+#include "PoseSearchDatabaseAssetTree.h"
+
+#include "Widgets/Text/SRichTextBlock.h"
+#include "Widgets/Input/SCheckBox.h"
+#include "Widgets/Input/SSearchBox.h"
+#include "Framework/Commands/GenericCommands.h"
+#include "Framework/MultiBox/MultiBoxBuilder.h"
+#include "SPositiveActionButton.h"
+#include "Styling/AppStyle.h"
+
+#include "ScopedTransaction.h"
+#include "Styling/StyleColors.h"
+
+#define LOCTEXT_NAMESPACE "SDatabaseAssetListItem"
+
+namespace UE::PoseSearch
+{
+	static constexpr FLinearColor DisabledColor = FLinearColor(1.0f, 1.0f, 1.0f, 0.25f);
+	
+	void SDatabaseAssetListItem::Construct(
+		const FArguments& InArgs,
+		const TSharedRef<FDatabaseViewModel>& InEditorViewModel,
+		const TSharedRef<STableViewBase>& OwnerTable,
+		TSharedRef<FDatabaseAssetTreeNode> InAssetTreeNode,
+		TSharedRef<FUICommandList> InCommandList,
+		TSharedPtr<SDatabaseAssetTree> InHierarchy)
+	{
+		WeakAssetTreeNode = InAssetTreeNode;
+		EditorViewModel = InEditorViewModel;
+		SkeletonView = InHierarchy;
+
+		if (InAssetTreeNode->SourceAssetType == ESearchIndexAssetType::Invalid)
+		{
+			ConstructGroupItem(OwnerTable);
+		}
+		else
+		{
+			ConstructAssetItem(OwnerTable);
+		}
+	}
+
+	void SDatabaseAssetListItem::ConstructGroupItem(const TSharedRef<STableViewBase>& OwnerTable)
+	{
+		STableRow<TSharedPtr<FDatabaseAssetTreeNode>>::ChildSlot
+		.Padding(0.0f, 2.0f, 0.0f, 0.0f)
+		[
+			GenerateItemWidget()
+		];
+
+		STableRow<TSharedPtr<FDatabaseAssetTreeNode>>::ConstructInternal(
+			STableRow<TSharedPtr<FDatabaseAssetTreeNode>>::FArguments()
+			.Style(FAppStyle::Get(), "DetailsView.TreeView.TableRow")
+			.OnCanAcceptDrop(SkeletonView.Pin().Get(), &SDatabaseAssetTree::OnCanAcceptDrop)
+			.OnAcceptDrop(SkeletonView.Pin().Get(), &SDatabaseAssetTree::OnAcceptDrop)
+			.ShowSelection(true),
+			OwnerTable);
+	}
+
+	void SDatabaseAssetListItem::ConstructAssetItem(const TSharedRef<STableViewBase>& OwnerTable)
+	{
+		STableRow<TSharedPtr<FDatabaseAssetTreeNode>>::Construct(
+			STableRow<TSharedPtr<FDatabaseAssetTreeNode>>::FArguments()
+			.Style(&FAppStyle::Get().GetWidgetStyle<FTableRowStyle>("TableView.Row"))
+			.OnCanAcceptDrop(SkeletonView.Pin().Get(), &SDatabaseAssetTree::OnCanAcceptDrop)
+			.OnAcceptDrop(SkeletonView.Pin().Get(), &SDatabaseAssetTree::OnAcceptDrop)
+			.ShowWires(false)
+			.Content()
+			[
+				GenerateItemWidget()
+			], OwnerTable);
+	}
+
+	void SDatabaseAssetListItem::OnAddSequence()
+	{
+		EditorViewModel.Pin()->AddSequenceToDatabase(nullptr);
+		SkeletonView.Pin()->RefreshTreeView(false);
+	}
+
+	void SDatabaseAssetListItem::OnAddBlendSpace()
+	{
+		EditorViewModel.Pin()->AddBlendSpaceToDatabase(nullptr);
+		SkeletonView.Pin()->RefreshTreeView(false);
+	}
+
+	void SDatabaseAssetListItem::OnAddAnimComposite()
+	{
+		EditorViewModel.Pin()->AddAnimCompositeToDatabase(nullptr);
+		SkeletonView.Pin()->RefreshTreeView(false);
+	}
+
+	FText SDatabaseAssetListItem::GetName() const
+	{
+		TSharedPtr<FDatabaseAssetTreeNode> Node = WeakAssetTreeNode.Pin();
+		TSharedPtr<FDatabaseViewModel> ViewModel = EditorViewModel.Pin();
+
+		if (const UPoseSearchDatabase* Database = ViewModel->GetPoseSearchDatabase())
+		{
+			if (const FPoseSearchDatabaseAnimationAssetBase* DatabaseAnimationAsset = Database->GetAnimationAssetBase(Node->SourceAssetIdx))
+			{
+				return FText::FromString(DatabaseAnimationAsset->GetName());
+			}
+
+			return FText::FromString(Database->GetName());
+		}
+
+		return LOCTEXT("None", "None");
+	}
+
+	TSharedRef<SWidget> SDatabaseAssetListItem::GenerateItemWidget()
+	{
+		TSharedPtr<FDatabaseAssetTreeNode> Node = WeakAssetTreeNode.Pin();
+		TSharedPtr<FDatabaseViewModel> ViewModel = EditorViewModel.Pin();
+		UPoseSearchDatabase* Database = ViewModel->GetPoseSearchDatabase();
+		
+		TSharedPtr<SWidget> ItemWidget;
+		const FDetailColumnSizeData& ColumnSizeData = SkeletonView.Pin()->GetColumnSizeData();
+		
+		if (Node->SourceAssetType == ESearchIndexAssetType::Invalid)
+		{
+			// it's a group
+			SAssignNew(ItemWidget, SBorder)
+			.BorderImage(this, &SDatabaseAssetListItem::GetGroupBackgroundImage)
+			.Padding(FMargin(3.0f, 5.0f))
+			[
+				SNew(SHorizontalBox)
+				+ SHorizontalBox::Slot()
+				.VAlign(VAlign_Center)
+				.Padding(5.0f)
+				.AutoWidth()
+				[
+					SNew(SExpanderArrow, STableRow<TSharedPtr<FDatabaseAssetTreeNode>>::SharedThis(this))
+				]
+				+ SHorizontalBox::Slot()
+				.FillWidth(1.0f)
+				.VAlign(VAlign_Center)
+				[
+					SNew(SRichTextBlock)
+					.Text(this, &SDatabaseAssetListItem::GetName)
+					.TransformPolicy(ETextTransformPolicy::ToUpper)
+					.DecoratorStyleSet(&FAppStyle::Get())
+					.TextStyle(FAppStyle::Get(), "DetailsView.CategoryTextStyle")
+				]
+				+SHorizontalBox::Slot()
+				.AutoWidth()
+				.VAlign(VAlign_Center)
+				.HAlign(HAlign_Right)
+				.Padding(2, 0, 0, 0)
+				[
+					GenerateAddButtonWidget()
+				]
+			];
+		}
+		else
+		{
+			// Item Icon
+			TSharedPtr<SImage> ItemIconWidget;
+			if (const FPoseSearchDatabaseAnimationAssetBase* DatabaseAnimationAsset = Database->GetAnimationAssetBase(Node->SourceAssetIdx))
+			{
+				SAssignNew(ItemIconWidget, SImage)
+					.Image(FSlateIconFinder::FindIconBrushForClass(DatabaseAnimationAsset->GetAnimationAssetStaticClass()));
+			}
+
+			// Setup table row to display 
+			SAssignNew(ItemWidget, SHorizontalBox)
+			+SHorizontalBox::Slot()
+			.FillWidth(1.0f)
+			[
+				SNew(SSplitter)
+				.Style(FAppStyle::Get(), "FoliageEditMode.Splitter")
+				.PhysicalSplitterHandleSize(1.0f)
+				.HitDetectionSplitterHandleSize(5.0f)
+				.HighlightedHandleIndex(ColumnSizeData.GetHoveredSplitterIndex())
+				.MinimumSlotHeight(0.5f)
+				
+				// Asset Name with type icon
+				+SSplitter::Slot()
+				.Value(ColumnSizeData.GetNameColumnWidth())
+				.MinSize(0.3f)
+				.OnSlotResized(ColumnSizeData.GetOnNameColumnResized())
+				[
+					SNew(SHorizontalBox)
+					.Clipping(EWidgetClipping::ClipToBounds)
+					+ SHorizontalBox::Slot()
+					.MaxWidth(18)
+					.AutoWidth()
+					.Padding(0.0f, 0.0f, 5.0f, 0.0f)
+					.HAlign(HAlign_Left)
+					.VAlign(VAlign_Center)
+					[
+						ItemIconWidget.ToSharedRef()
+					]
+					+ SHorizontalBox::Slot()
+					.FillWidth(1.0f)
+					.VAlign(VAlign_Center)
+					[
+						SNew(STextBlock)
+						.Text(this, &SDatabaseAssetListItem::GetName)
+						.ColorAndOpacity(this, &SDatabaseAssetListItem::GetNameTextColorAndOpacity)
+					]
+				]
+				
+				// Display information via icons
+				+SSplitter::Slot()
+				.Value(ColumnSizeData.GetValueColumnWidth())
+				.MinSize(0.3f)
+				.OnSlotResized(ColumnSizeData.GetOnValueColumnResized())
+				[
+					// Asset Info.
+
+					// Looping
+					SNew(SHorizontalBox)
+					.Clipping(EWidgetClipping::ClipToBounds)
+					+ SHorizontalBox::Slot()
+					.AutoWidth()
+					.Padding(4.0f, 1.0f)
+					.HAlign(HAlign_Center)
+					.VAlign(VAlign_Center)
+					[
+						SNew(SImage)
+						.Image(FAppStyle::Get().GetBrush("Graph.Node.Loop"))
+						.ColorAndOpacity(this, &SDatabaseAssetListItem::GetLoopingColorAndOpacity)
+						.ToolTipText(this, &SDatabaseAssetListItem::GetLoopingToolTip)
+					]
+
+					// Root Motion
+					+ SHorizontalBox::Slot()
+					.Padding(1.0f, 1.0f)
+					.AutoWidth()
+					.HAlign(HAlign_Center)
+					.VAlign(VAlign_Center)
+					[
+						SNew(SImage)
+						.Image(FAppStyle::Get().GetBrush("AnimGraph.Attribute.RootMotionDelta.Icon"))
+						.ColorAndOpacity(this, &SDatabaseAssetListItem::GetRootMotionColorAndOpacity)
+						.ToolTipText(this, &SDatabaseAssetListItem::GetRootMotionOptionToolTip)
+					]
+					
+					// Mirror Type
+					+ SHorizontalBox::Slot()
+					.Padding(1.0f, 1.0f)
+					.AutoWidth()
+					.HAlign(HAlign_Center)
+					.VAlign(VAlign_Center)
+					[
+						SNew(SImage)
+						.Image(this, &SDatabaseAssetListItem::GetMirrorOptionSlateBrush)
+						.ToolTipText(this, &SDatabaseAssetListItem::GetMirrorOptionToolTip)
+					]
+				]
+			]
+			
+			+SHorizontalBox::Slot()
+			.AutoWidth()
+			.HAlign(HAlign_Right)
+			.VAlign(VAlign_Center)
+			[
+				SNew(SHorizontalBox)
+				+ SHorizontalBox::Slot()
+				.MaxWidth(18)
+				.Padding(4.0f, 0.0f, 0.0f, 0.0f)
+				.AutoWidth()
+				.HAlign(HAlign_Right)
+				.VAlign(VAlign_Center)
+				[
+					SNew(SImage)
+					.Image(FAppStyle::Get().GetBrush("Icons.EyeDropper"))
+					.Visibility_Raw(this, &SDatabaseAssetListItem::GetSelectedActorIconVisbility)
+				]
+				+ SHorizontalBox::Slot()
+				.MaxWidth(16)
+				.Padding(4.0f, 0.0f)
+				.AutoWidth()
+				.HAlign(HAlign_Right)
+				.VAlign(VAlign_Center)
+				[
+					SNew(SCheckBox)
+					.IsChecked(this, &SDatabaseAssetListItem::GetAssetEnabledChecked)
+					.OnCheckStateChanged(const_cast<SDatabaseAssetListItem*>(this), &SDatabaseAssetListItem::OnAssetIsEnabledChanged)
+					.ToolTipText(this, &SDatabaseAssetListItem::GetAssetEnabledToolTip)
+					.CheckedImage(FAppStyle::Get().GetBrush("Icons.Visible"))
+					.CheckedHoveredImage(FAppStyle::Get().GetBrush("Icons.Visible"))
+					.CheckedPressedImage(FAppStyle::Get().GetBrush("Icons.Visible"))
+					.UncheckedImage(FAppStyle::Get().GetBrush("Icons.Hidden"))
+					.UncheckedHoveredImage(FAppStyle::Get().GetBrush("Icons.Hidden"))
+					.UncheckedPressedImage(FAppStyle::Get().GetBrush("Icons.Hidden"))
+				]
+			];
+		}
+
+		return ItemWidget.ToSharedRef();
+	}
+
+	TSharedRef<SWidget> SDatabaseAssetListItem::GenerateAddButtonWidget()
+	{
+		FMenuBuilder AddOptions(true, nullptr);
+
+		AddOptions.AddMenuEntry(
+			LOCTEXT("AddSequence", "Add Sequence"),
+			LOCTEXT("AddSequenceTooltip", "Add new sequence to this group"),
+			FSlateIcon(),
+			FUIAction(FExecuteAction::CreateSP(this, &SDatabaseAssetListItem::OnAddSequence)),
+			NAME_None,
+			EUserInterfaceActionType::Button);
+
+		AddOptions.AddMenuEntry(
+			LOCTEXT("AddBlendSpaceOption", "Add Blend Space"),
+			LOCTEXT("AddBlendSpaceOptionTooltip", "Add new blend space to this group"),
+			FSlateIcon(),
+			FUIAction(FExecuteAction::CreateSP(this, &SDatabaseAssetListItem::OnAddBlendSpace)),
+			NAME_None,
+			EUserInterfaceActionType::Button);
+
+		AddOptions.AddMenuEntry(
+			LOCTEXT("AnimCompositeOption", "Add Anim Composite"),
+			LOCTEXT("AddAnimCompositeToDefaultGroupTooltip", "Add new anim composite to this group"),
+			FSlateIcon(),
+			FUIAction(FExecuteAction::CreateSP(this, &SDatabaseAssetListItem::OnAddAnimComposite)),
+			NAME_None,
+			EUserInterfaceActionType::Button);
+
+		TSharedPtr<SComboButton> AddButton;
+		SAssignNew(AddButton, SComboButton)
+		.ContentPadding(0)
+		.ComboButtonStyle(&FAppStyle::Get().GetWidgetStyle<FComboButtonStyle>("SimpleComboButton"))
+		.HasDownArrow(false)
+		.ButtonContent()
+		[
+			SNew(SImage)
+			.ColorAndOpacity(FSlateColor::UseForeground())
+			.Image(FAppStyle::Get().GetBrush("Icons.PlusCircle"))
+		]
+		.MenuContent()
+		[
+			AddOptions.MakeWidget()
+		];
+
+		return AddButton.ToSharedRef();
+	}
+
+
+	const FSlateBrush* SDatabaseAssetListItem::GetGroupBackgroundImage() const
+	{
+		if (STableRow<TSharedPtr<FDatabaseAssetTreeNode>>::IsHovered())
+		{
+			return FAppStyle::Get().GetBrush("Brushes.Secondary");
+		}
+		else
+		{
+			return FAppStyle::Get().GetBrush("Brushes.Header");
+		}
+	}
+
+	EVisibility SDatabaseAssetListItem::GetSelectedActorIconVisbility() const
+	{
+		TSharedPtr<FDatabaseViewModel> ViewModelPtr = EditorViewModel.Pin();
+		TSharedPtr<FDatabaseAssetTreeNode> TreeNodePtr = WeakAssetTreeNode.Pin();
+		if (const FPoseSearchIndexAsset* SelectedIndexAsset = ViewModelPtr->GetSelectedActorIndexAsset())
+		{
+			if (TreeNodePtr->SourceAssetType == ESearchIndexAssetType::Sequence &&
+				TreeNodePtr->SourceAssetIdx == SelectedIndexAsset->SourceAssetIdx)
+			{
+				return EVisibility::Visible;
+			}
+		}
+
+		return EVisibility::Hidden;
+	}
+
+	ECheckBoxState SDatabaseAssetListItem::GetAssetEnabledChecked() const
+	{
+		TSharedPtr<FDatabaseViewModel> ViewModelPtr = EditorViewModel.Pin();
+		TSharedPtr<FDatabaseAssetTreeNode> TreeNodePtr = WeakAssetTreeNode.Pin();
+		const UPoseSearchDatabase* Database = ViewModelPtr->GetPoseSearchDatabase();
+
+		if (Database->AnimationAssets.IsValidIndex(TreeNodePtr->SourceAssetIdx))
+		{
+			if (ViewModelPtr->IsEnabled(TreeNodePtr->SourceAssetIdx))
+			{
+				return ECheckBoxState::Checked;
+			}
+		}
+
+		return ECheckBoxState::Unchecked;
+	}
+
+	void SDatabaseAssetListItem::OnAssetIsEnabledChanged(ECheckBoxState NewCheckboxState)
+	{
+		const FScopedTransaction Transaction(LOCTEXT("EnableChangedForAssetInPoseSearchDatabase", "Update enabled flag for item from Pose Search Database"));
+
+		const TSharedPtr<FDatabaseViewModel> ViewModelPtr = EditorViewModel.Pin();
+		const TSharedPtr<FDatabaseAssetTreeNode> TreeNodePtr = WeakAssetTreeNode.Pin();
+
+		ViewModelPtr->GetPoseSearchDatabase()->Modify();
+		
+		ViewModelPtr->SetIsEnabled(TreeNodePtr->SourceAssetIdx, NewCheckboxState == ECheckBoxState::Checked ? true : false);
+
+		SkeletonView.Pin()->RefreshTreeView(false, true);
+		ViewModelPtr->BuildSearchIndex();
+	}
+
+	FSlateColor SDatabaseAssetListItem::GetNameTextColorAndOpacity() const
+	{
+		return GetAssetEnabledChecked() == ECheckBoxState::Checked ? FLinearColor::White : DisabledColor;
+	}
+
+	FSlateColor SDatabaseAssetListItem::GetLoopingColorAndOpacity() const
+	{
+		const TSharedPtr<FDatabaseAssetTreeNode> Node = WeakAssetTreeNode.Pin();
+		return Node->IsLooping() ? FLinearColor::White : DisabledColor;
+	}
+
+	FText SDatabaseAssetListItem::GetLoopingToolTip() const
+	{
+		const TSharedPtr<FDatabaseAssetTreeNode> Node = WeakAssetTreeNode.Pin();
+		return Node->IsLooping() ? LOCTEXT("NodeLoopEnabledToolTip", "Looping") : LOCTEXT("NodeLoopDisabledToolTip", "Not looping");
+	}
+
+	FSlateColor SDatabaseAssetListItem::GetRootMotionColorAndOpacity() const
+	{
+		const TSharedPtr<FDatabaseAssetTreeNode> Node = WeakAssetTreeNode.Pin();
+		return Node->IsRootMotionEnabled() ? FLinearColor::White : DisabledColor;
+	}
+
+	FText SDatabaseAssetListItem::GetRootMotionOptionToolTip() const
+	{
+		const TSharedPtr<FDatabaseAssetTreeNode> Node = WeakAssetTreeNode.Pin();
+		return Node->IsRootMotionEnabled() ? LOCTEXT("NodeRootMotionEnabledToolTip", "Root motion enabled") : LOCTEXT("NodeRootMotionDisabledToolTip", "No root motion enabled");
+
+	}
+	const FSlateBrush* SDatabaseAssetListItem::GetMirrorOptionSlateBrush() const
+	{
+		const TSharedPtr<FDatabaseAssetTreeNode> Node = WeakAssetTreeNode.Pin();
+
+		// TODO: Update icons when appropriate assets become available.
+		switch (Node->GetMirrorOption())
+		{
+			case EPoseSearchMirrorOption::UnmirroredOnly: 
+				return FAppStyle::Get().GetBrush("Icons.Minus");
+			
+			case EPoseSearchMirrorOption::MirroredOnly: 
+				return FAppStyle::Get().GetBrush("Icons.Plus");
+			
+			case EPoseSearchMirrorOption::UnmirroredAndMirrored:
+				return FAppStyle::Get().GetBrush("Icons.X");
+			
+			default:
+				return nullptr;
+		}
+	}
+
+	FText SDatabaseAssetListItem::GetMirrorOptionToolTip() const
+	{
+		const TSharedPtr<FDatabaseAssetTreeNode> Node = WeakAssetTreeNode.Pin();
+		return FText::FromString(LOCTEXT("ToolTipMirrorOption", "Mirror Option: ").ToString() + (Node ? UEnum::GetDisplayValueAsText(Node->GetMirrorOption()).ToString() : LOCTEXT("ToolTipMirrorOption_Invalid", "Invalid").ToString()));
+	}
+	
+	FText SDatabaseAssetListItem::GetAssetEnabledToolTip() const
+	{
+		if (GetAssetEnabledChecked() == ECheckBoxState::Checked)
+		{
+			return LOCTEXT("DisableAssetTooltip", "Disable this asset in the Pose Search Database.");
+		}
+		
+		return LOCTEXT("EnableAssetTooltip", "Enable this asset in the Pose Search Database.");
+	}
+}
+
+#undef LOCTEXT_NAMESPACE
diff --git a/Engine/Plugins/Experimental/Animation/PoseSearch/Source/Editor/Private/PoseSearchDatabaseAssetTree.cpp b/Engine/Plugins/Experimental/Animation/PoseSearch/Source/Editor/Private/PoseSearchDatabaseAssetTree.cpp
new file mode 100644
--- /dev/null
+++ b/Engine/Plugins/Experimental/Animation/PoseSearch/Source/Editor/Private/PoseSearchDatabaseAssetTree.cpp
@@ -0,0 +1,743 @@
+// Copyright Epic Games, Inc. All Rights Reserved.
+
+#include "PoseSearchDatabaseAssetTree.h"
+#include "PoseSearchDatabaseViewModel.h"
+
+#include "PoseSearch/PoseSearch.h"
+
+#include "Animation/AnimSequence.h"
+#include "Animation/AnimComposite.h"
+#include "Animation/BlendSpace.h"
+
+#include "Misc/TransactionObjectEvent.h"
+#include "DragAndDrop/AssetDragDropOp.h"
+#include "Misc/FeedbackContext.h"
+#include "AssetSelection.h"
+#include "ClassIconFinder.h"
+#include "DetailColumnSizeData.h"
+
+#include "Widgets/Text/SRichTextBlock.h"
+#include "Widgets/Input/SCheckBox.h"
+#include "Widgets/Input/SSearchBox.h"
+#include "Framework/Commands/GenericCommands.h"
+#include "Framework/MultiBox/MultiBoxBuilder.h"
+#include "SPositiveActionButton.h"
+#include "Styling/AppStyle.h"
+
+#include "ScopedTransaction.h"
+#include "Styling/StyleColors.h"
+
+#define LOCTEXT_NAMESPACE "SDatabaseAssetTree"
+
+namespace UE::PoseSearch
+{
+	SDatabaseAssetTree::~SDatabaseAssetTree()
+	{
+	}
+
+	void SDatabaseAssetTree::Construct(
+		const FArguments& InArgs, 
+		TSharedRef<FDatabaseViewModel> InEditorViewModel)
+	{
+		EditorViewModel = InEditorViewModel;
+		
+		ColumnSizeData.SetValueColumnWidth(0.6f);
+
+		CreateCommandList();
+
+		ChildSlot
+		[
+			SNew(SVerticalBox)
+			+SVerticalBox::Slot()
+			.AutoHeight()
+			[
+				SNew(SHorizontalBox)
+				+SHorizontalBox::Slot()
+				.AutoWidth()
+				.VAlign(VAlign_Center)
+				.Padding(0, 0, 4, 0)
+				[
+					SNew(SPositiveActionButton)
+					.Icon(FAppStyle::Get().GetBrush("Icons.Plus"))
+					.Text(LOCTEXT("AddNew", "Add"))
+					.ToolTipText(LOCTEXT("AddNewToolTip", "Add a new Sequence, Blend Space or Group"))
+					.OnGetMenuContent(this, &SDatabaseAssetTree::CreateAddNewMenuWidget)
+				]
+				+ SHorizontalBox::Slot()
+				.AutoWidth()
+				.VAlign(VAlign_Center)
+				.HAlign(HAlign_Right)
+				.Padding(2, 0, 0, 0)
+				[
+					GenerateFilterBoxWidget()
+				]
+			]
+			+SVerticalBox::Slot()
+			.Padding(0.0f, 0.0f)
+			[
+				SNew(SBorder)
+				.Padding(2.0f)
+				.BorderImage(FAppStyle::GetBrush("SCSEditor.TreePanel"))
+				[
+					SNew(SOverlay)
+					+SOverlay::Slot()
+					[
+						SAssignNew(TreeView, STreeView<TSharedPtr<FDatabaseAssetTreeNode>>)
+						.TreeItemsSource(&RootNodes)
+						.SelectionMode(ESelectionMode::Multi)
+						.OnGenerateRow(this, &SDatabaseAssetTree::MakeTableRowWidget)
+						.OnGetChildren(this, &SDatabaseAssetTree::HandleGetChildrenForTree)
+						.OnContextMenuOpening(this, &SDatabaseAssetTree::CreateContextMenu)
+						.HighlightParentNodesForSelection(false)
+						.OnSelectionChanged_Lambda([this](TSharedPtr<FDatabaseAssetTreeNode> Item, ESelectInfo::Type Type)
+							{
+								TArray<TSharedPtr<FDatabaseAssetTreeNode>> SelectedItems = TreeView->GetSelectedItems();
+								OnSelectionChanged.Broadcast(SelectedItems, Type);
+							})
+						.ItemHeight(24)
+					]
+					+SOverlay::Slot()
+					[
+						SAssignNew(TreeViewDragAndDropSuggestion, SVerticalBox)
+						+ SVerticalBox::Slot()
+						.HAlign(HAlign_Center)
+						.VAlign(VAlign_Center)
+						[
+							SNew(STextBlock)
+							.Text(FText::FromString(TEXT("Drag and drop Animation Sequences, Anim Composites or Blendspaces")))
+							.Font(FAppStyle::Get().GetFontStyle("DetailsView.CategoryFontStyle"))
+						]
+					]
+				]
+			]
+		];
+
+		RefreshTreeView(true);
+	}
+
+	FReply SDatabaseAssetTree::OnDragOver(const FGeometry& MyGeometry, const FDragDropEvent& DragDropEvent)
+	{
+		FReply Reply = FReply::Unhandled();
+
+		TSharedPtr<FDragDropOperation> Operation = DragDropEvent.GetOperation();
+
+		const bool bValidOperation =
+			Operation.IsValid() &&
+			(Operation->IsOfType<FExternalDragOperation>() || Operation->IsOfType<FAssetDragDropOp>());
+		if (bValidOperation)
+		{
+			Reply = AssetUtil::CanHandleAssetDrag(DragDropEvent);
+
+			if (!Reply.IsEventHandled())
+			{
+				if (Operation->IsOfType<FAssetDragDropOp>())
+				{
+					const TSharedPtr<FAssetDragDropOp> AssetDragDropOp = StaticCastSharedPtr<FAssetDragDropOp>(Operation);
+
+					for (const FAssetData& AssetData : AssetDragDropOp->GetAssets())
+					{
+						if (UClass* AssetClass = AssetData.GetClass())
+						{
+							if (AssetClass->IsChildOf(UAnimSequence::StaticClass()) ||
+								AssetClass->IsChildOf(UAnimComposite::StaticClass()) ||
+								AssetClass->IsChildOf(UBlendSpace::StaticClass()))
+							{
+								Reply = FReply::Handled();
+								break;
+							}
+						}
+					}
+				}
+			}
+		}
+
+		return Reply;
+	}
+
+	FReply SDatabaseAssetTree::OnDrop(const FGeometry& MyGeometry, const FDragDropEvent& DragDropEvent)
+	{
+		return OnAcceptDrop(DragDropEvent, EItemDropZone::OntoItem, nullptr);
+	}
+
+	FReply SDatabaseAssetTree::OnKeyDown(const FGeometry& MyGeometry, const FKeyEvent& InKeyEvent)
+	{
+		if (CommandList->ProcessCommandBindings(InKeyEvent))
+		{
+			return FReply::Handled();
+		}
+		return FReply::Unhandled();
+	}
+
+	bool SDatabaseAssetTree::MatchesContext(const FTransactionContext& InContext, const TArray<TPair<UObject*, FTransactionObjectEvent>>& TransactionObjectContexts) const
+	{
+		// Ensure that we only react to modifications to the UPosesSearchDatabase.
+		if (const TSharedPtr<FDatabaseViewModel> ViewModel = EditorViewModel.Pin())
+		{
+			if (const UPoseSearchDatabase * Database = ViewModel->GetPoseSearchDatabase())
+			{
+				for (const TPair<UObject*, FTransactionObjectEvent>& TransactionObjectPair : TransactionObjectContexts)
+				{
+					const UObject* Object = TransactionObjectPair.Key;
+					while (Object != nullptr)
+					{
+						if (Object == Database)
+						{
+							return true;
+						}
+
+						Object = Object->GetOuter();
+					}
+				}
+			}
+		}
+		
+		return false;
+	}
+
+	void SDatabaseAssetTree::PostUndo(bool bSuccess)
+	{
+		if (bSuccess)
+		{
+			FinalizeTreeChanges();
+		}
+	}
+
+	void SDatabaseAssetTree::PostRedo(bool bSuccess)
+	{
+		if (bSuccess)
+		{
+			FinalizeTreeChanges();
+		}
+	}
+
+	void SDatabaseAssetTree::RefreshTreeView(bool bIsInitialSetup, bool bRecoverSelection)
+	{
+		const TSharedPtr<FDatabaseViewModel> ViewModel = EditorViewModel.Pin();
+		if (!ViewModel.IsValid())
+		{
+			return;
+		}
+
+		const TSharedRef<FDatabaseViewModel> ViewModelRef = ViewModel.ToSharedRef();
+
+		// Empty node data.
+		RootNodes.Reset();
+		AllNodes.Reset();
+
+		const UPoseSearchDatabase* Database = ViewModel->GetPoseSearchDatabase();
+		if (!IsValid(Database))
+		{
+			TreeView->RequestTreeRefresh();
+			return;
+		}
+
+		// Store selection so we can recover it afterwards (if possible)
+		TArray<TSharedPtr<FDatabaseAssetTreeNode>> PreviouslySelectedNodes = TreeView->GetSelectedItems();
+
+		// Rebuild node hierarchy
+		{
+			// Setup default group node
+			{
+				const TSharedPtr<FDatabaseAssetTreeNode> DefaultGroupNode = MakeShared<FDatabaseAssetTreeNode>(
+					INDEX_NONE,
+					ESearchIndexAssetType::Invalid,
+					ViewModelRef);
+				AllNodes.Add(DefaultGroupNode);
+				RootNodes.Add(DefaultGroupNode);
+			}
+			
+			const int32 DefaultGroupIdx = RootNodes.Num() - 1;
+			
+			auto CreateAssetNode = [this, ViewModelRef](int32 AssetIdx, ESearchIndexAssetType AssetType, int32 GroupIdx)
+			{
+				// Create sequence node
+				const TSharedPtr<FDatabaseAssetTreeNode> SequenceGroupNode = MakeShared<FDatabaseAssetTreeNode>(
+					AssetIdx,
+					AssetType,
+					ViewModelRef);
+				const TSharedPtr<FDatabaseAssetTreeNode>& ParentGroupNode = RootNodes[GroupIdx];
+
+				// Setup hierarchy
+				SequenceGroupNode->Parent = ParentGroupNode;
+				ParentGroupNode->Children.Add(SequenceGroupNode);
+
+				// Keep track of node
+				AllNodes.Add(SequenceGroupNode);
+			};
+			
+			// Build an index based off of alphabetical order than iterate the index instead
+			TArray<uint32> IndexArray;
+			IndexArray.SetNumUninitialized(Database->AnimationAssets.Num());
+			for (int32 AnimationAssetIdx = 0; AnimationAssetIdx < Database->AnimationAssets.Num(); ++AnimationAssetIdx)
+			{
+				IndexArray[AnimationAssetIdx] = AnimationAssetIdx;
+			}
+
+			IndexArray.Sort([Database](int32 SequenceIdxA, int32 SequenceIdxB)
+			{
+				const FPoseSearchDatabaseAnimationAssetBase* A = Database->GetAnimationAssetBase(SequenceIdxA);
+				const FPoseSearchDatabaseAnimationAssetBase* B = Database->GetAnimationAssetBase(SequenceIdxB);
+
+				//If its null add it to the end of the list 
+				if (!B->GetAnimationAsset())
+				{
+					return true;
+				}
+
+				if (!A->GetAnimationAsset())
+				{
+					return false;
+				}
+
+				const int32 Comparison = A->GetName().Compare(B->GetName());
+				return Comparison < 0;
+			});
+
+			// create all nodes
+			for (int32 AnimationAssetIdx = 0; AnimationAssetIdx < Database->AnimationAssets.Num(); ++AnimationAssetIdx)
+			{
+				const int32 MappedId = IndexArray[AnimationAssetIdx];
+
+				if (const FPoseSearchDatabaseAnimationAssetBase* DatabaseAnimationAsset = Database->GetAnimationAssetBase(MappedId))
+				{
+					const bool bFiltered = (DatabaseAnimationAsset->GetAnimationAsset() == nullptr || GetAssetFilterString().IsEmpty()) ? false : !DatabaseAnimationAsset->GetName().Contains(GetAssetFilterString());
+
+					if (!bFiltered)
+					{
+						CreateAssetNode(MappedId, DatabaseAnimationAsset->GetSearchIndexType(), DefaultGroupIdx);
+					}
+				}
+			}
+
+			// Show drag and drop suggestion if tree is empty
+			TreeViewDragAndDropSuggestion->SetVisibility(IndexArray.IsEmpty() ? EVisibility::Visible : EVisibility::Hidden);
+		}
+
+		// Update tree view
+		TreeView->RequestTreeRefresh();
+
+		for (TSharedPtr<FDatabaseAssetTreeNode>& RootNode : RootNodes)
+		{
+			TreeView->SetItemExpansion(RootNode, true);
+		}
+
+		// Handle selection
+		if (bRecoverSelection)
+		{
+			RecoverSelection(PreviouslySelectedNodes);
+		}
+		else
+		{
+			TreeView->SetItemSelection(PreviouslySelectedNodes, false, ESelectInfo::Direct);
+		}
+	}
+
+	TSharedRef<ITableRow> SDatabaseAssetTree::MakeTableRowWidget(
+		TSharedPtr<FDatabaseAssetTreeNode> InItem,
+		const TSharedRef<STableViewBase>& OwnerTable)
+	{
+		return InItem->MakeTreeRowWidget(OwnerTable, InItem.ToSharedRef(), CommandList.ToSharedRef(), SharedThis(this));
+	}
+
+	void SDatabaseAssetTree::HandleGetChildrenForTree(
+		TSharedPtr<FDatabaseAssetTreeNode> InNode,
+		TArray<TSharedPtr<FDatabaseAssetTreeNode>>& OutChildren)
+	{
+		OutChildren = InNode.Get()->Children;
+	}
+
+	TOptional<EItemDropZone> SDatabaseAssetTree::OnCanAcceptDrop(
+		const FDragDropEvent& DragDropEvent,
+		EItemDropZone DropZone,
+		TSharedPtr<FDatabaseAssetTreeNode> TargetItem)
+	{
+		TOptional<EItemDropZone> ReturnedDropZone;
+
+		TSharedPtr<FDragDropOperation> Operation = DragDropEvent.GetOperation();
+
+		const bool bValidOperation = Operation.IsValid() && Operation->IsOfType<FAssetDragDropOp>();
+		if (bValidOperation)
+		{
+			const TSharedPtr<FAssetDragDropOp> AssetDragDropOp = StaticCastSharedPtr<FAssetDragDropOp>(Operation);
+
+			for (const FAssetData& AssetData : AssetDragDropOp->GetAssets())
+			{
+				if (UClass* AssetClass = AssetData.GetClass())
+				{
+					if (AssetClass->IsChildOf(UAnimSequence::StaticClass()) ||
+						AssetClass->IsChildOf(UAnimComposite::StaticClass()) ||
+						AssetClass->IsChildOf(UBlendSpace::StaticClass()))
+					{
+						ReturnedDropZone = EItemDropZone::OntoItem;
+						break;
+					}
+				}
+			}
+		}
+
+		return ReturnedDropZone;
+	}
+
+	FReply SDatabaseAssetTree::OnAcceptDrop(
+		const FDragDropEvent& DragDropEvent,
+		EItemDropZone DropZone,
+		TSharedPtr<FDatabaseAssetTreeNode> TargetItem)
+	{
+		const TSharedPtr<FDragDropOperation> Operation = DragDropEvent.GetOperation();
+
+		const bool bValidOperation = Operation.IsValid() && Operation->IsOfType<FAssetDragDropOp>();
+		if (!bValidOperation)
+		{
+			return FReply::Unhandled();
+		}
+
+		const TSharedPtr<FDatabaseViewModel> ViewModel = EditorViewModel.Pin();
+		if (!ViewModel.IsValid())
+		{
+			return FReply::Unhandled();
+		}
+
+		TArray<FAssetData> DroppedAssetData = AssetUtil::ExtractAssetDataFromDrag(Operation);
+		const int32 NumAssets = DroppedAssetData.Num();
+
+		int32 AddedAssets = 0;
+		if (NumAssets > 0)
+		{
+			GWarn->BeginSlowTask(LOCTEXT("LoadingAssets", "Loading Asset(s)"), true);
+
+			{
+				const FScopedTransaction Transaction(LOCTEXT("AddSequencesOrBlendspaces", "Add Sequence(s) and/or Blendspace(s) to Pose Search Database"));
+				ViewModel->GetPoseSearchDatabase()->Modify();
+				
+				for (int32 DroppedAssetIdx = 0; DroppedAssetIdx < NumAssets; ++DroppedAssetIdx)
+				{
+					const FAssetData& AssetData = DroppedAssetData[DroppedAssetIdx];
+
+					if (!AssetData.IsAssetLoaded())
+					{
+						GWarn->StatusUpdate(
+							DroppedAssetIdx,
+							NumAssets,
+							FText::Format(
+								LOCTEXT("LoadingAsset", "Loading Asset {0}"),
+								FText::FromName(AssetData.AssetName)));
+					}
+
+					UClass* AssetClass = AssetData.GetClass();
+					UObject* Asset = AssetData.GetAsset();
+					
+					if (AssetClass->IsChildOf(UAnimSequence::StaticClass()))
+					{
+						ViewModel->AddSequenceToDatabase(Cast<UAnimSequence>(Asset));
+						++AddedAssets;
+					}
+					if (AssetClass->IsChildOf(UAnimComposite::StaticClass()))
+					{
+						ViewModel->AddAnimCompositeToDatabase(Cast<UAnimComposite>(Asset));
+						++AddedAssets;
+					}
+					else if (AssetClass->IsChildOf(UBlendSpace::StaticClass()))
+					{
+						ViewModel->AddBlendSpaceToDatabase(Cast<UBlendSpace>(Asset));
+						++AddedAssets;
+					}
+				}
+			}
+			
+			GWarn->EndSlowTask();
+		}
+
+		if (AddedAssets == 0)
+		{
+			return FReply::Unhandled();
+		}
+
+		FinalizeTreeChanges(false);
+		return FReply::Handled();
+	}
+
+	TSharedRef<SWidget> SDatabaseAssetTree::CreateAddNewMenuWidget()
+	{
+		FMenuBuilder AddOptions(true, nullptr);
+
+		AddOptions.BeginSection("AddOptions", LOCTEXT("AssetAddOptions", "Assets"));
+		{
+			AddOptions.AddMenuEntry(
+				LOCTEXT("AddSequenceOption", "Sequence"),
+				LOCTEXT("AddSequenceOptionTooltip", "Add new sequence to the default group"),
+				FSlateIcon(),
+				FUIAction(FExecuteAction::CreateSP(this, &SDatabaseAssetTree::OnAddSequence, true)),
+				NAME_None,
+				EUserInterfaceActionType::Button);
+
+			AddOptions.AddMenuEntry(
+				LOCTEXT("BlendSpaceOption", "Blend Space"),
+				LOCTEXT("AddBlendSpaceToDefaultGroupTooltip", "Add new blend space to the default group"),
+				FSlateIcon(),
+				FUIAction(FExecuteAction::CreateSP(this, &SDatabaseAssetTree::OnAddBlendSpace, true)),
+				NAME_None,
+				EUserInterfaceActionType::Button);
+
+			AddOptions.AddMenuEntry(
+				LOCTEXT("AnimCompositeOption", "Anim Composite"),
+				LOCTEXT("AddAnimCompositeToDefaultGroupTooltip", "Add new anim composite to the default group"),
+				FSlateIcon(),
+				FUIAction(FExecuteAction::CreateSP(this, &SDatabaseAssetTree::OnAddAnimComposite, true)),
+				NAME_None,
+				EUserInterfaceActionType::Button);
+		}
+		AddOptions.EndSection();
+
+		return AddOptions.MakeWidget();
+	}
+
+	TSharedPtr<SWidget> SDatabaseAssetTree::CreateContextMenu()
+	{
+		const bool CloseAfterSelection = true;
+		FMenuBuilder MenuBuilder(CloseAfterSelection, CommandList);
+
+		const TArray<TSharedPtr<FDatabaseAssetTreeNode>> SelectedNodes = TreeView->GetSelectedItems();
+		if (!SelectedNodes.IsEmpty())
+		{
+			MenuBuilder.AddMenuEntry(
+				LOCTEXT("DeleteUngroup", "Delete / Remove"),
+				LOCTEXT(
+					"DeleteUngroupTooltip", 
+					"Deletes groups and ungrouped assets; removes grouped assets from group."),
+				FSlateIcon(),
+				FUIAction(FExecuteAction::CreateSP(this, &SDatabaseAssetTree::OnDeleteNodes)),
+				NAME_None,
+				EUserInterfaceActionType::Button);
+
+			MenuBuilder.AddMenuEntry(
+				LOCTEXT("Enable", "Enable"),
+				LOCTEXT(
+					"EnableTooltip",
+					"Sets Assets Enabled."),
+				FSlateIcon(),
+				FUIAction(FExecuteAction::CreateSP(this, &SDatabaseAssetTree::OnEnableNodes)),
+				NAME_None,
+				EUserInterfaceActionType::Button);
+
+			MenuBuilder.AddMenuEntry(
+				LOCTEXT("Disable", "Disable"),
+				LOCTEXT(
+					"DisableToolTip",
+					"Sets Assets Disabled."),
+				FSlateIcon(),
+				FUIAction(FExecuteAction::CreateSP(this, &SDatabaseAssetTree::OnDisableNodes)),
+				NAME_None,
+				EUserInterfaceActionType::Button);
+		}
+
+		return MenuBuilder.MakeWidget();
+	}
+
+	TSharedRef<SWidget> SDatabaseAssetTree::GenerateFilterBoxWidget()
+	{
+		TSharedPtr<SSearchBox> SearchBox;
+		SAssignNew(SearchBox, SSearchBox)
+			.MinDesiredWidth(300.0f)
+			.InitialText(this, &SDatabaseAssetTree::GetFilterText)
+			.ToolTipText(FText::FromString(TEXT("Enter Asset Filter...")))
+			.OnTextChanged(this, &SDatabaseAssetTree::OnAssetFilterTextCommitted, ETextCommit::Default)
+			.OnTextCommitted(this, &SDatabaseAssetTree::OnAssetFilterTextCommitted);
+
+		return SearchBox.ToSharedRef();
+	}
+
+
+	FText SDatabaseAssetTree::GetFilterText() const
+	{
+		return FText::FromString(GetAssetFilterString());
+	}
+
+	void SDatabaseAssetTree::OnAssetFilterTextCommitted(const FText& InText, ETextCommit::Type CommitInfo)
+	{
+		SetAssetFilterString(InText.ToString());
+		RefreshTreeView(false);
+	}
+
+
+	void SDatabaseAssetTree::OnAddSequence(bool bFinalizeChanges)
+	{
+		FScopedTransaction Transaction(LOCTEXT("AddSequence", "Add Sequence"));
+		const TSharedPtr<FDatabaseViewModel> ViewModel = EditorViewModel.Pin();
+		
+		ViewModel->GetPoseSearchDatabase()->Modify();
+		
+		ViewModel->AddSequenceToDatabase(nullptr);
+
+		if (bFinalizeChanges)
+		{
+			FinalizeTreeChanges();
+		}
+	}
+
+	void SDatabaseAssetTree::OnAddBlendSpace(bool bFinalizeChanges)
+	{
+		FScopedTransaction Transaction(LOCTEXT("AddBlendSpaceTransaction", "Add Blend Space"));
+		const TSharedPtr<FDatabaseViewModel> ViewModel = EditorViewModel.Pin();
+		
+		ViewModel->GetPoseSearchDatabase()->Modify();
+		
+		ViewModel->AddBlendSpaceToDatabase(nullptr);
+
+		if (bFinalizeChanges)
+		{
+			FinalizeTreeChanges();
+		}
+	}
+
+	void SDatabaseAssetTree::OnAddAnimComposite(bool bFinalizeChanges)
+	{
+		FScopedTransaction Transaction(LOCTEXT("AddAnimCompositeTransaction", "Add Anim Composite"));
+
+		EditorViewModel.Pin()->AddAnimCompositeToDatabase(nullptr);
+
+		if (bFinalizeChanges)
+		{
+			FinalizeTreeChanges();
+		}
+	}
+
+	void SDatabaseAssetTree::OnDeleteAsset(TSharedPtr<FDatabaseAssetTreeNode> Node, bool bFinalizeChanges)
+	{
+		FScopedTransaction Transaction(LOCTEXT("DeleteAsset", "Delete Asset"));
+		const TSharedPtr<FDatabaseViewModel> ViewModel = EditorViewModel.Pin();
+		
+		ViewModel->GetPoseSearchDatabase()->Modify();
+		ViewModel->DeleteFromDatabase(Node->SourceAssetIdx);
+
+		if (bFinalizeChanges)
+		{
+			FinalizeTreeChanges();
+		}
+	}
+
+	void SDatabaseAssetTree::RegisterOnSelectionChanged(const FOnSelectionChanged& Delegate)
+	{
+		OnSelectionChanged.Add(Delegate);
+	}
+
+	void SDatabaseAssetTree::UnregisterOnSelectionChanged(void* Unregister)
+	{
+		OnSelectionChanged.RemoveAll(Unregister);
+	}
+
+	void SDatabaseAssetTree::RecoverSelection(const TArray<TSharedPtr<FDatabaseAssetTreeNode>>& PreviouslySelectedNodes)
+	{
+		TArray<TSharedPtr<FDatabaseAssetTreeNode>> NewSelectedNodes;
+
+		for (const TSharedPtr<FDatabaseAssetTreeNode>& Node : AllNodes)
+		{
+			const bool bFoundNode = PreviouslySelectedNodes.ContainsByPredicate(
+				[Node](const TSharedPtr<FDatabaseAssetTreeNode>& PrevSelectedNode)
+			{
+				return
+					PrevSelectedNode->SourceAssetType == Node->SourceAssetType &&
+					PrevSelectedNode->SourceAssetIdx == Node->SourceAssetIdx;
+			});
+
+			if (bFoundNode)
+			{
+				NewSelectedNodes.Add(Node);
+			}
+		}
+
+		// @todo: investigate if we should call a TreeView->ClearSelection() before TreeView->SetItemSelection
+		TreeView->SetItemSelection(NewSelectedNodes, true, ESelectInfo::Direct);
+	}
+
+	void SDatabaseAssetTree::CreateCommandList()
+	{
+		CommandList = MakeShared<FUICommandList>();
+
+		CommandList->MapAction(
+			FGenericCommands::Get().Delete,
+			FUIAction(
+				FExecuteAction::CreateSP(this, &SDatabaseAssetTree::OnDeleteNodes),
+				FCanExecuteAction::CreateSP(this, &SDatabaseAssetTree::CanDeleteNodes)));
+	}
+
+	bool SDatabaseAssetTree::CanDeleteNodes() const
+	{
+		TArray<TSharedPtr<FDatabaseAssetTreeNode>> SelectedNodes = TreeView->GetSelectedItems();
+		for (TSharedPtr<FDatabaseAssetTreeNode> SelectedNode : SelectedNodes)
+		{
+			if (SelectedNode->SourceAssetType != ESearchIndexAssetType::Invalid ||
+				SelectedNode->SourceAssetIdx != INDEX_NONE)
+			{
+				return true;
+			}
+		}
+
+		return false;
+	}
+
+	void SDatabaseAssetTree::OnDeleteNodes()
+	{
+		TArray<TSharedPtr<FDatabaseAssetTreeNode>> SelectedNodes = TreeView->GetSelectedItems();
+		if (!SelectedNodes.IsEmpty())
+		{
+			const FScopedTransaction Transaction(LOCTEXT("DeletePoseSearchDatabaseNodes", "Delete selected items from Pose Search Database"));
+			const TSharedPtr<FDatabaseViewModel> ViewModel = EditorViewModel.Pin();
+
+			ViewModel->GetPoseSearchDatabase()->Modify();
+			
+			SelectedNodes.Sort(
+				[](const TSharedPtr<FDatabaseAssetTreeNode>& A, const TSharedPtr<FDatabaseAssetTreeNode>& B)
+			{
+				if (A->SourceAssetType != ESearchIndexAssetType::Invalid &&
+					B->SourceAssetType == ESearchIndexAssetType::Invalid)
+				{
+					return true;
+				}
+				if (B->SourceAssetType != ESearchIndexAssetType::Invalid &&
+					A->SourceAssetType == ESearchIndexAssetType::Invalid)
+				{
+					return false;
+				}
+				return B->SourceAssetIdx < A->SourceAssetIdx;
+			});
+
+			for (TSharedPtr<FDatabaseAssetTreeNode> SelectedNode : SelectedNodes)
+			{
+				if (SelectedNode->SourceAssetType != ESearchIndexAssetType::Invalid)
+				{
+					OnDeleteAsset(SelectedNode, false);
+				}
+			}
+			
+			FinalizeTreeChanges();
+		}
+	}
+
+	void SDatabaseAssetTree::EnableSelectedNodes(bool bIsEnabled)
+	{
+		TArray<TSharedPtr<FDatabaseAssetTreeNode>> SelectedNodes = TreeView->GetSelectedItems();
+		if (!SelectedNodes.IsEmpty())
+		{
+			const FText TransactionName = bIsEnabled ? LOCTEXT("EnablePoseSearchDatabaseNodes", "Enable selected items from Pose Search Database") : LOCTEXT("DisablePoseSearchDatabaseNodes", "Disable selected items from Pose Search Database");
+			const FScopedTransaction Transaction(TransactionName);
+			const TSharedPtr<FDatabaseViewModel> ViewModel = EditorViewModel.Pin();
+
+			ViewModel->GetPoseSearchDatabase()->Modify();
+
+			for (TSharedPtr<FDatabaseAssetTreeNode> SelectedNode : SelectedNodes)
+			{
+				ViewModel->SetIsEnabled(SelectedNode->SourceAssetIdx, bIsEnabled);
+			}
+		
+			FinalizeTreeChanges();
+		}
+	}
+
+	void SDatabaseAssetTree::FinalizeTreeChanges(bool bRecoverSelection)
+	{
+		RefreshTreeView(false, bRecoverSelection);
+		EditorViewModel.Pin()->BuildSearchIndex();
+	}
+}
+
+#undef LOCTEXT_NAMESPACE
diff --git a/Engine/Plugins/Experimental/Animation/PoseSearch/Source/Editor/Private/PoseSearchDatabaseAssetTreeNode.cpp b/Engine/Plugins/Experimental/Animation/PoseSearch/Source/Editor/Private/PoseSearchDatabaseAssetTreeNode.cpp
new file mode 100644
--- /dev/null
+++ b/Engine/Plugins/Experimental/Animation/PoseSearch/Source/Editor/Private/PoseSearchDatabaseAssetTreeNode.cpp
@@ -0,0 +1,97 @@
+// Copyright Epic Games, Inc. All Rights Reserved.
+
+#include "PoseSearchDatabaseAssetTreeNode.h"
+#include "PoseSearchDatabaseViewModel.h"
+
+#include "PoseSearch/PoseSearch.h"
+
+#include "Animation/AnimSequence.h"
+#include "Animation/AnimComposite.h"
+#include "Animation/BlendSpace.h"
+
+#include "Misc/TransactionObjectEvent.h"
+#include "DragAndDrop/AssetDragDropOp.h"
+#include "Misc/FeedbackContext.h"
+#include "AssetSelection.h"
+#include "ClassIconFinder.h"
+#include "DetailColumnSizeData.h"
+#include "PoseSearchDatabaseAssetListItem.h"
+
+#include "Widgets/Text/SRichTextBlock.h"
+#include "Widgets/Input/SCheckBox.h"
+#include "Widgets/Input/SSearchBox.h"
+#include "Widgets/Views/SListView.h"
+
+#include "Framework/Commands/GenericCommands.h"
+#include "Framework/MultiBox/MultiBoxBuilder.h"
+#include "SPositiveActionButton.h"
+#include "Styling/AppStyle.h"
+
+#include "ScopedTransaction.h"
+#include "Styling/StyleColors.h"
+
+namespace UE::PoseSearch
+{
+	FDatabaseAssetTreeNode::FDatabaseAssetTreeNode(
+		int32 InSourceAssetIdx,
+		ESearchIndexAssetType InSourceAssetType,
+		const TSharedRef<FDatabaseViewModel>& InEditorViewModel) 
+		: SourceAssetIdx(InSourceAssetIdx)
+		, SourceAssetType(InSourceAssetType)
+		, EditorViewModel(InEditorViewModel)
+	{ }
+
+	TSharedRef<ITableRow> FDatabaseAssetTreeNode::MakeTreeRowWidget(
+		const TSharedRef<STableViewBase>& InOwnerTable,
+		TSharedRef<FDatabaseAssetTreeNode> InDatabaseAssetNode,
+		TSharedRef<FUICommandList> InCommandList,
+		TSharedPtr<SDatabaseAssetTree> InHierarchy)
+	{
+		return SNew(
+			SDatabaseAssetListItem, 
+			EditorViewModel.Pin().ToSharedRef(), 
+			InOwnerTable, 
+			InDatabaseAssetNode, 
+			InCommandList, 
+			InHierarchy);
+	}
+
+	bool FDatabaseAssetTreeNode::IsRootMotionEnabled() const
+	{
+		if (const UPoseSearchDatabase* Database = EditorViewModel.Pin()->GetPoseSearchDatabase())
+		{
+			if (const FPoseSearchDatabaseAnimationAssetBase* DatabaseAnimationAsset = Database->GetAnimationAssetBase(SourceAssetIdx))
+			{
+				return DatabaseAnimationAsset->IsRootMotionEnabled();
+			}
+		}
+
+		return false;
+	}
+
+	bool FDatabaseAssetTreeNode::IsLooping() const
+	{
+		if (const UPoseSearchDatabase* Database = EditorViewModel.Pin()->GetPoseSearchDatabase())
+		{
+			if (const FPoseSearchDatabaseAnimationAssetBase* DatabaseAnimationAsset = Database->GetAnimationAssetBase(SourceAssetIdx))
+			{
+				return DatabaseAnimationAsset->IsLooping();
+			}
+		}
+
+		return false;
+	}
+
+	EPoseSearchMirrorOption FDatabaseAssetTreeNode::GetMirrorOption() const
+	{
+		if (const UPoseSearchDatabase* Database = EditorViewModel.Pin()->GetPoseSearchDatabase())
+		{
+			if (const FPoseSearchDatabaseAnimationAssetBase* DatabaseAnimationAsset = Database->GetAnimationAssetBase(SourceAssetIdx))
+			{
+				return DatabaseAnimationAsset->GetMirrorOption();
+			}
+		}
+
+		return EPoseSearchMirrorOption::Invalid;
+	}
+}
diff --git a/Engine/Plugins/Experimental/Animation/PoseSearch/Source/Editor/Private/PoseSearchDatabaseEditor.cpp b/Engine/Plugins/Experimental/Animation/PoseSearch/Source/Editor/Private/PoseSearchDatabaseEditor.cpp
--- a/Engine/Plugins/Experimental/Animation/PoseSearch/Source/Editor/Private/PoseSearchDatabaseEditor.cpp
+++ b/Engine/Plugins/Experimental/Animation/PoseSearch/Source/Editor/Private/PoseSearchDatabaseEditor.cpp
@@ -1,8 +1,8 @@
 // Copyright Epic Games, Inc. All Rights Reserved.
 
 #include "PoseSearchDatabaseEditor.h"
+#include "PoseSearchDatabaseAssetTree.h"
 #include "SPoseSearchDatabaseViewport.h"
-#include "SPoseSearchDatabaseAssetList.h"
 #include "PoseSearchDatabasePreviewScene.h"
 #include "PoseSearchDatabaseEditorCommands.h"
 #include "PoseSearchDatabaseViewModel.h"
diff --git a/Engine/Plugins/Experimental/Animation/PoseSearch/Source/Editor/Private/PoseSearchDatabaseEditorReflection.cpp b/Engine/Plugins/Experimental/Animation/PoseSearch/Source/Editor/Private/PoseSearchDatabaseEditorReflection.cpp
--- a/Engine/Plugins/Experimental/Animation/PoseSearch/Source/Editor/Private/PoseSearchDatabaseEditorReflection.cpp
+++ b/Engine/Plugins/Experimental/Animation/PoseSearch/Source/Editor/Private/PoseSearchDatabaseEditorReflection.cpp
@@ -1,7 +1,8 @@
 // Copyright Epic Games, Inc. All Rights Reserved.
 
 #include "PoseSearchDatabaseEditorReflection.h"
-#include "SPoseSearchDatabaseAssetList.h"
+#include "PoseSearchDatabaseAssetTree.h"
+#include "PoseSearchDatabaseAssetTreeNode.h"
 #include "PoseSearchDatabaseViewModel.h"
 
 #define LOCTEXT_NAMESPACE "UPoseSearchDatabaseReflection"
diff --git a/Engine/Plugins/Experimental/Animation/PoseSearch/Source/Editor/Private/PoseSearchDatabasePreviewScene.cpp b/Engine/Plugins/Experimental/Animation/PoseSearch/Source/Editor/Private/PoseSearchDatabasePreviewScene.cpp
--- a/Engine/Plugins/Experimental/Animation/PoseSearch/Source/Editor/Private/PoseSearchDatabasePreviewScene.cpp
+++ b/Engine/Plugins/Experimental/Animation/PoseSearch/Source/Editor/Private/PoseSearchDatabasePreviewScene.cpp
@@ -5,7 +5,6 @@
 #include "PoseSearchDatabaseEditor.h"
 #include "PoseSearchDatabaseViewModel.h"
 
-#include "PoseSearch/PoseSearch.h"
 #include "PoseSearch/PoseSearchDerivedData.h"
 
 #include "Animation/DebugSkelMeshComponent.h"
diff --git a/Engine/Plugins/Experimental/Animation/PoseSearch/Source/Editor/Private/PoseSearchDatabaseViewModel.cpp b/Engine/Plugins/Experimental/Animation/PoseSearch/Source/Editor/Private/PoseSearchDatabaseViewModel.cpp
--- a/Engine/Plugins/Experimental/Animation/PoseSearch/Source/Editor/Private/PoseSearchDatabaseViewModel.cpp
+++ b/Engine/Plugins/Experimental/Animation/PoseSearch/Source/Editor/Private/PoseSearchDatabaseViewModel.cpp
@@ -2,7 +2,7 @@
 
 #include "PoseSearchDatabaseViewModel.h"
 #include "PoseSearchDatabasePreviewScene.h"
-#include "SPoseSearchDatabaseAssetList.h"
+#include "PoseSearchDatabaseAssetTreeNode.h"
 #include "PoseSearchEditor.h"
 #include "PoseSearch/PoseSearchDerivedData.h"
 #include "Modules/ModuleManager.h"
diff --git a/Engine/Plugins/Experimental/Animation/PoseSearch/Source/Editor/Private/PoseSearchDebugger.cpp b/Engine/Plugins/Experimental/Animation/PoseSearch/Source/Editor/Private/PoseSearchDebugger.cpp
--- a/Engine/Plugins/Experimental/Animation/PoseSearch/Source/Editor/Private/PoseSearchDebugger.cpp
+++ b/Engine/Plugins/Experimental/Animation/PoseSearch/Source/Editor/Private/PoseSearchDebugger.cpp
@@ -1,2827 +1,19 @@
 ﻿// Copyright Epic Games, Inc. All Rights Reserved.
 
 #include "PoseSearchDebugger.h"
-#include "Algo/AllOf.h"
-#include "Animation/AnimRootMotionProvider.h"
-#include "Animation/AnimSequence.h"
-#include "Animation/AnimationPoseData.h"
-#include "Animation/BlendSpace.h"
-#include "Animation/MirrorDataTable.h"
-#include "Editor.h"
-#include "Engine/SkeletalMesh.h"
-#include "Engine/SkinnedAsset.h"
 #include "IAnimationProvider.h"
 #include "IGameplayProvider.h"
 #include "IRewindDebugger.h"
-#include "ObjectTrace.h"
-#include "PoseSearch/PoseSearch.h"
-#include "PoseSearch/PoseSearchDerivedData.h"
-#include "Preferences/PersonaOptions.h"
-#include "PropertyEditorModule.h"
-#include "Styling/AppStyle.h"
+#include "PoseSearchDebuggerView.h"
+#include "PoseSearchDebuggerViewModel.h"
 #include "Styling/SlateIconFinder.h"
-#include "Subsystems/AssetEditorSubsystem.h"
 #include "Trace/PoseSearchTraceProvider.h"
-#include "TraceServices/Model/AnalysisSession.h"
-#include "TraceServices/Model/Frames.h"
-#include "Widgets/Images/SImage.h"
-#include "Widgets/Input/SButton.h"
-#include "Widgets/Input/SHyperlink.h"
-#include "Widgets/Input/SNumericEntryBox.h"
-#include "Widgets/Input/SSearchBox.h"
-#include "Widgets/Layout/SScrollBox.h"
-#include "Widgets/Layout/SWidgetSwitcher.h"
-#include "Widgets/SBoxPanel.h"
-#include "Widgets/Views/SListView.h"
 
 #define LOCTEXT_NAMESPACE "PoseSearchDebugger"
 
 namespace UE::PoseSearch
 {
 
-inline static FLinearColor LinearColorBlend(FLinearColor LinearColorA, FLinearColor LinearColorB, float BlendParam)
-{
-	return LinearColorA + (LinearColorB - LinearColorA) * BlendParam;
-}
-
-inline static void LinearColorBlend(FLinearColor LinearColorA, FLinearColor LinearColorB, TConstArrayView<float> BlendParam, TArray<FLinearColor>& BlendedColors)
-{
-	const int32 Num = BlendParam.Num();
-	BlendedColors.SetNumUninitialized(Num);
-	for (int i = 0; i < Num; ++i)
-	{
-		BlendedColors[i] = LinearColorBlend(LinearColorA, LinearColorB, BlendParam[i]);
-	}
-}
-
-static void ArrayMinMax(TConstArrayView<float> View, TArrayView<float> Min, TArrayView<float> Max, float InvalidValue)
-{
-	const int32 Num = View.Num();
-	check(Num == Min.Num() && Num == Max.Num());
-	for (int i = 0; i < Num; ++i)
-	{
-		const float Value = View[i];
-		if (Value != InvalidValue)
-		{
-			Min[i] = FMath::Min(Min[i], Value);
-			Max[i] = FMath::Max(Max[i], Value);
-		}
-	}
-}
-
-static void ArraySafeNormalize(TConstArrayView<float> View, TConstArrayView<float> Min, TConstArrayView<float> Max, TArrayView<float> NormalizedView)
-{
-	const int32 Num = View.Num();
-	check(Num == Min.Num() && Num == Max.Num() && Num == NormalizedView.Num());
-	for (int i = 0; i < Num; ++i)
-	{
-		const float Delta = Max[i] - Min[i];
-		if (FMath::IsNearlyZero(Delta, UE_KINDA_SMALL_NUMBER))
-		{
-			NormalizedView[i] = 0.f;
-		}
-		else
-		{
-			NormalizedView[i] = (View[i] - Min[i]) / Delta;
-		}
-	}
-}
-
-} // namespace UE::PoseSearch
-
-void UPoseSearchMeshComponent::Initialize(const FTransform& InComponentToWorld)
-{
-	SetComponentToWorld(InComponentToWorld);
-	const FReferenceSkeleton& SkeletalMeshRefSkeleton = GetSkinnedAsset()->GetRefSkeleton();
-
-	// set up bone visibility states as this gets skipped since we allocate the component array before registration
-	for (int32 BaseIndex = 0; BaseIndex < 2; ++BaseIndex)
-	{
-		BoneVisibilityStates[BaseIndex].SetNum(SkeletalMeshRefSkeleton.GetNum());
-		for (int32 BoneIndex = 0; BoneIndex < SkeletalMeshRefSkeleton.GetNum(); BoneIndex++)
-		{
-			BoneVisibilityStates[BaseIndex][BoneIndex] = BVS_ExplicitlyHidden;
-		}
-	}
-
-	StartingTransform = InComponentToWorld;
-	Refresh();
-}
-
-void UPoseSearchMeshComponent::Refresh()
-{
-	// Flip buffers once to copy the directly-written component space transforms
-	bNeedToFlipSpaceBaseBuffers = true;
-	bHasValidBoneTransform = false;
-	FlipEditableSpaceBases();
-	bHasValidBoneTransform = true;
-
-	InvalidateCachedBounds();
-	UpdateBounds();
-	MarkRenderTransformDirty();
-	MarkRenderDynamicDataDirty();
-	MarkRenderStateDirty();
-}
-
-void UPoseSearchMeshComponent::ResetToStart()
-{
-	SetComponentToWorld(StartingTransform);
-	Refresh();
-}
-
-void UPoseSearchMeshComponent::UpdatePose(const FUpdateContext& UpdateContext)
-{
-	FMemMark Mark(FMemStack::Get());
-
-	FCompactPose CompactPose;
-	CompactPose.SetBoneContainer(&RequiredBones);
-	FBlendedCurve Curve;
-	Curve.InitFrom(RequiredBones);
-	UE::Anim::FStackAttributeContainer Attributes;
-	FAnimationPoseData PoseData(CompactPose, Curve, Attributes);
-
-	if (UpdateContext.Type == ESearchIndexAssetType::Sequence ||
-		UpdateContext.Type == ESearchIndexAssetType::AnimComposite)
-	{
-		float AdvancedTime = UpdateContext.StartTime;
-
-		FAnimationRuntime::AdvanceTime(
-			UpdateContext.bLoop,
-			UpdateContext.Time - UpdateContext.StartTime,
-			AdvancedTime,
-			UpdateContext.SequenceBase->GetPlayLength());
-
-		FAnimExtractContext ExtractionCtx;
-		ExtractionCtx.CurrentTime = AdvancedTime;
-
-		UpdateContext.SequenceBase->GetAnimationPose(PoseData, ExtractionCtx);
-	}
-	else if (UpdateContext.Type == ESearchIndexAssetType::BlendSpace)
-	{
-		TArray<FBlendSampleData> BlendSamples;
-		int32 TriangulationIndex = 0;
-		UpdateContext.BlendSpace->GetSamplesFromBlendInput(UpdateContext.BlendParameters, BlendSamples, TriangulationIndex, true);
-		
-		float PlayLength = UpdateContext.BlendSpace->GetAnimationLengthFromSampleData(BlendSamples);
-		
-		float PreviousTime = UpdateContext.StartTime * PlayLength;
-		float CurrentTime = UpdateContext.Time * PlayLength;
-
-		float AdvancedTime = PreviousTime;
-		FAnimationRuntime::AdvanceTime(
-			UpdateContext.bLoop,
-			CurrentTime - PreviousTime,
-			CurrentTime,
-			PlayLength);
-		
-		FDeltaTimeRecord DeltaTimeRecord;
-		DeltaTimeRecord.Set(PreviousTime, AdvancedTime - PreviousTime);
-		FAnimExtractContext ExtractionCtx(static_cast<double>(AdvancedTime), true, DeltaTimeRecord, UpdateContext.bLoop);
-
-		for (int32 BlendSampleIdex = 0; BlendSampleIdex < BlendSamples.Num(); BlendSampleIdex++)
-		{
-			float Scale = BlendSamples[BlendSampleIdex].Animation->GetPlayLength() / PlayLength;
-
-			FDeltaTimeRecord BlendSampleDeltaTimeRecord;
-			BlendSampleDeltaTimeRecord.Set(DeltaTimeRecord.GetPrevious() * Scale, DeltaTimeRecord.Delta * Scale);
-
-			BlendSamples[BlendSampleIdex].DeltaTimeRecord = BlendSampleDeltaTimeRecord;
-			BlendSamples[BlendSampleIdex].PreviousTime = PreviousTime * Scale;
-			BlendSamples[BlendSampleIdex].Time = AdvancedTime * Scale;
-		}
-
-		UpdateContext.BlendSpace->GetAnimationPose(BlendSamples, ExtractionCtx, PoseData);
-	}
-	else
-	{
-		checkNoEntry();
-	}
-
-	LastRootMotionDelta = FTransform::Identity;
-
-	const UE::Anim::IAnimRootMotionProvider* RootMotionProvider = UE::Anim::IAnimRootMotionProvider::Get();
-
-	if (ensureMsgf(RootMotionProvider, TEXT("Could not get Root Motion Provider.")))
-	{
-		if (ensureMsgf(RootMotionProvider->HasRootMotion(Attributes), TEXT("Blend Space had no Root Motion Attribute.")))
-		{
-			RootMotionProvider->ExtractRootMotion(Attributes, LastRootMotionDelta);
-		}
-	}
-
-	if (UpdateContext.bMirrored)
-	{
-		FAnimationRuntime::MirrorPose(
-			CompactPose, 
-			UpdateContext.MirrorDataTable->MirrorAxis, 
-			*UpdateContext.CompactPoseMirrorBones, 
-			*UpdateContext.ComponentSpaceRefRotations);
-	}
-
-	FCSPose<FCompactPose> ComponentSpacePose;
-	ComponentSpacePose.InitPose(CompactPose);
-
-	for (const FBoneIndexType BoneIndex : RequiredBones.GetBoneIndicesArray())
-	{
-		const FTransform BoneTransform = 
-			ComponentSpacePose.GetComponentSpaceTransform(FCompactPoseBoneIndex(BoneIndex));
-
-		FSkeletonPoseBoneIndex SkeletonBoneIndex =
-			RequiredBones.GetSkeletonPoseIndexFromCompactPoseIndex(FCompactPoseBoneIndex(BoneIndex));
-		FName BoneName = 
-			RequiredBones.GetSkeletonAsset()->GetReferenceSkeleton().GetBoneName(SkeletonBoneIndex.GetInt());
-		SetBoneTransformByName(BoneName, BoneTransform, EBoneSpaces::ComponentSpace);
-	}
-
-	if (UpdateContext.bMirrored)
-	{
-		const EAxis::Type MirrorAxis = UpdateContext.MirrorDataTable->MirrorAxis;
-		FVector T = LastRootMotionDelta.GetTranslation();
-		T = FAnimationRuntime::MirrorVector(T, MirrorAxis);
-		const FQuat ReferenceRotation = (*UpdateContext.ComponentSpaceRefRotations)[FCompactPoseBoneIndex(0)];
-		FQuat Q = LastRootMotionDelta.GetRotation();
-		Q = FAnimationRuntime::MirrorQuat(Q, MirrorAxis);
-		Q *= FAnimationRuntime::MirrorQuat(ReferenceRotation, MirrorAxis).Inverse() * ReferenceRotation;
-		LastRootMotionDelta = FTransform(Q, T, LastRootMotionDelta.GetScale3D());
-	}
-
-	const FTransform ComponentTransform = LastRootMotionDelta * StartingTransform;
-
-	SetComponentToWorld(ComponentTransform);
-	FillComponentSpaceTransforms();
-	Refresh();
-}
-
-
-namespace UE::PoseSearch
-{
-
-class FDebuggerDatabaseRowData : public TSharedFromThis<FDebuggerDatabaseRowData>
-{
-public:
-	FDebuggerDatabaseRowData() = default;
-	
-	ESearchIndexAssetType AssetType = ESearchIndexAssetType::Invalid;
-	int32 PoseIdx = 0;
-	TWeakObjectPtr<const UPoseSearchDatabase> SourceDatabase = nullptr;
-	EPoseCandidateFlags PoseCandidateFlags = EPoseCandidateFlags::None;
-	FString DatabaseName = "";
-	FString DatabasePath = "";
-	FString AssetName = "";
-	FString AssetPath = "";
-	int32 DbAssetIdx = 0;
-	int32 AnimFrame = 0;
-	float AnimPercentage = 0.0f;
-	float AssetTime = 0.0f;
-	bool bMirrored = false;
-	bool bLooping = false;
-	FVector BlendParameters = FVector::Zero();
-	FPoseSearchCost PoseCost;
-	FLinearColor CostColor = FLinearColor::White;
-	TArray<float> CostBreakdowns;
-	TArray<FLinearColor> CostBreakdownsColors;
-	TArray<float> CostVector;
-};
-
-namespace DebuggerDatabaseColumns
-{
-	/** Column struct to represent each column in the debugger database */
-	struct IColumn : TSharedFromThis<IColumn>
-	{
-		explicit IColumn(int32 InSortIndex, bool InEnabled = true)
-			: SortIndex(InSortIndex)
-			, bEnabled(InEnabled)
-		{
-			ColumnId = FName(FString::Printf(TEXT("Column %d"), SortIndex));
-		}
-
-		virtual ~IColumn() = default;
-
-		FName ColumnId;
-
-		/** Sorted left to right based on this index */
-		int32 SortIndex = 0;
-		/** Current width, starts at 1 to be evenly spaced between all columns */
-		float Width = 1.0f;
-		/** Disabled selectively with view options */
-		bool bEnabled = false;
-
-		virtual FText GetLabel() const = 0;
-
-		using FRowDataRef = TSharedRef<FDebuggerDatabaseRowData>;
-		using FSortPredicate = TFunction<bool(const FRowDataRef&, const FRowDataRef&)>;
-
-		virtual FSortPredicate GetSortPredicate() const = 0;
-
-		virtual TSharedRef<SWidget> GenerateWidget(const FRowDataRef& RowData) const = 0;
-	};
-
-	/** Column struct to represent each column in the debugger database */
-	struct ITextColumn : IColumn
-	{
-		explicit ITextColumn(int32 InSortIndex, bool InEnabled = true)
-			: IColumn(InSortIndex, InEnabled)
-		{
-		}
-
-		virtual ~ITextColumn() = default;
-		
-		virtual TSharedRef<SWidget> GenerateWidget(const FRowDataRef& RowData) const override
-		{
-			static FSlateFontInfo RowFont = FAppStyle::Get().GetFontStyle("DetailsView.CategoryTextStyle");
-			
-			return SNew(STextBlock)
-            	.Font(RowFont)
-				.Text_Lambda([this, RowData]() -> FText { return GetRowText(RowData); })
-            	.Justification(ETextJustify::Center)
-				.ColorAndOpacity_Lambda([this, RowData] { return GetColorAndOpacity(RowData); });
-		}
-
-	protected:
-		virtual FText GetRowText(const FRowDataRef& Row) const = 0;
-		virtual FSlateColor GetColorAndOpacity(const FRowDataRef& Row) const
-		{
-			return FSlateColor(FLinearColor::White);
-		}
-	};
-
-	struct FPoseIdx : ITextColumn
-	{
-		using ITextColumn::ITextColumn;
-
-		virtual FText GetLabel() const override { return LOCTEXT("ColumnLabelPoseIndex", "Index"); }
-
-		virtual FSortPredicate GetSortPredicate() const override
-		{
-			return [](const FRowDataRef& Row0, const FRowDataRef& Row1) -> bool { return Row0->PoseIdx < Row1->PoseIdx; };
-		}
-
-		virtual FText GetRowText(const FRowDataRef& Row) const override
-		{
-			return FText::AsNumber(Row->PoseIdx, &FNumberFormattingOptions::DefaultNoGrouping());
-		}
-	};
-
-	struct FDatabaseName : IColumn
-	{
-		using IColumn::IColumn;
-
-		virtual FText GetLabel() const override { return LOCTEXT("ColumnLabelDatabaseName", "Database"); }
-
-		virtual FSortPredicate GetSortPredicate() const override
-		{
-			return [](const FRowDataRef& Row0, const FRowDataRef& Row1) -> bool { return Row0->DatabaseName < Row1->DatabaseName; };
-		}
-
-		virtual TSharedRef<SWidget> GenerateWidget(const FRowDataRef& RowData) const override
-		{
-			return SNew(SHyperlink)
-				.Text_Lambda([RowData]() -> FText { return FText::FromString(RowData->DatabaseName); })
-				.TextStyle(&FCoreStyle::Get().GetWidgetStyle<FTextBlockStyle>("SmallText"))
-				.ToolTipText_Lambda([RowData]() -> FText
-					{
-						return FText::Format(
-							LOCTEXT("DatabaseHyperlinkTooltipFormat", "Open database '{0}'"),
-							FText::FromString(RowData->DatabasePath));
-					})
-				.OnNavigate_Lambda([RowData]()
-					{
-						GEditor->GetEditorSubsystem<UAssetEditorSubsystem>()->OpenEditorForAsset(RowData->DatabasePath);
-					});
-		}
-	};
-
-	struct FAssetName : IColumn
-	{
-		using IColumn::IColumn;
-
-		virtual FText GetLabel() const override { return LOCTEXT("ColumnLabelAssetName", "Asset"); }
-		
-		virtual FSortPredicate GetSortPredicate() const override
-		{
-			return [](const FRowDataRef& Row0, const FRowDataRef& Row1) -> bool { return Row0->AssetName < Row1->AssetName; };
-		}
-
-		virtual TSharedRef<SWidget> GenerateWidget(const FRowDataRef& RowData) const override
-		{
-			return SNew(SHyperlink)
-				.Text_Lambda([RowData]() -> FText { return FText::FromString(RowData->AssetName); })
-				.TextStyle(&FCoreStyle::Get().GetWidgetStyle<FTextBlockStyle>("SmallText"))
-				.ToolTipText_Lambda([RowData]() -> FText 
-					{ 
-						return FText::Format(
-							LOCTEXT("AssetHyperlinkTooltipFormat", "Open asset '{0}'"), 
-							FText::FromString(RowData->AssetPath)); 
-					})
-				.OnNavigate_Lambda([RowData]()
-					{
-						GEditor->GetEditorSubsystem<UAssetEditorSubsystem>()->OpenEditorForAsset(RowData->AssetPath);
-					});
-		}
-	};
-
-	struct FFrame : ITextColumn
-	{
-		using ITextColumn::ITextColumn;
-
-		virtual FText GetLabel() const override { return LOCTEXT("ColumnLabelFrame", "Frame"); }
-		
-		virtual FSortPredicate GetSortPredicate() const override
-		{
-			return [](const FRowDataRef& Row0, const FRowDataRef& Row1) -> bool { return Row0->AssetTime < Row1->AssetTime; };
-		}
-
-		virtual FText GetRowText(const FRowDataRef& Row) const override
-		{
-			FNumberFormattingOptions TimeFormattingOptions = FNumberFormattingOptions()
-				.SetUseGrouping(false)
-				.SetMaximumFractionalDigits(2);
-
-			FNumberFormattingOptions PercentageFormattingOptions = FNumberFormattingOptions()
-				.SetMaximumFractionalDigits(2);
-
-			if (Row->AssetType == ESearchIndexAssetType::Sequence ||
-				Row->AssetType == ESearchIndexAssetType::AnimComposite)
-			{
-				if (GetDefault<UPersonaOptions>()->bTimelineDisplayPercentage)
-				{
-					return FText::Format(
-						FText::FromString("{0} ({1}) ({2})"),
-						FText::AsNumber(Row->AnimFrame, &FNumberFormattingOptions::DefaultNoGrouping()),
-						FText::AsNumber(Row->AssetTime, &TimeFormattingOptions),
-						FText::AsPercent(Row->AnimPercentage, &PercentageFormattingOptions));
-				}
-				else
-				{
-					return FText::Format(
-						FText::FromString("{0} ({1})"),
-						FText::AsNumber(Row->AnimFrame, &FNumberFormattingOptions::DefaultNoGrouping()),
-						FText::AsNumber(Row->AssetTime, &TimeFormattingOptions));
-				}
-			}
-			else if (Row->AssetType == ESearchIndexAssetType::BlendSpace)
-			{
-				// There is no frame index associated with a blendspace
-				return FText::Format(
-					FText::FromString("({0})"),
-					FText::AsNumber(Row->AssetTime, &TimeFormattingOptions));
-			}
-			else
-			{
-				return FText::FromString(TEXT("-"));
-			}
-		}
-	};
-
-	struct FCost : ITextColumn
-	{
-		using ITextColumn::ITextColumn;
-
-		virtual FText GetLabel() const override { return LOCTEXT("ColumnLabelCost", "Cost"); }
-		
-		virtual FSortPredicate GetSortPredicate() const override
-		{
-			return [](const FRowDataRef& Row0, const FRowDataRef& Row1) -> bool { return Row0->PoseCost < Row1->PoseCost; };
-		}
-		
-		virtual FText GetRowText(const FRowDataRef& Row) const override
-		{
-			if (Row->PoseCost.IsValid())
-			{
-				return FText::AsNumber(Row->PoseCost.GetTotalCost());
-			}
-			else
-			{
-				return FText::FromString(TEXT("-"));
-			}
-        }
-
-		virtual FSlateColor GetColorAndOpacity(const FRowDataRef& Row) const override
-		{
-			return Row->CostColor;
-		}
-	};
-
-	struct FChannelBreakdownCostColumn : ITextColumn
-	{
-		FChannelBreakdownCostColumn(int32 SortIndex, int32 InBreakdownCostIndex, const FText& InLabel)
-			: ITextColumn(SortIndex)
-			, Label(InLabel)
-			, BreakdownCostIndex(InBreakdownCostIndex)
-		{
-		}
-
-		virtual FText GetLabel() const override { return Label; }
-
-		virtual FSortPredicate GetSortPredicate() const override
-		{
-			return [this](const FRowDataRef& Row0, const FRowDataRef& Row1) -> bool
-			{
-				return Row0->CostBreakdowns[BreakdownCostIndex] < Row1->CostBreakdowns[BreakdownCostIndex];
-			};
-		}
-
-		virtual FText GetRowText(const FRowDataRef& Row) const override
-		{
-			const float LabelCost = Row->CostBreakdowns[BreakdownCostIndex];
-			if (LabelCost != UE_MAX_FLT)
-			{
-				return FText::AsNumber(LabelCost);
-			}
-			return FText::FromString(TEXT("-"));
-		}
-
-		virtual FSlateColor GetColorAndOpacity(const FRowDataRef& Row) const override
-		{
-			return Row->CostBreakdownsColors[BreakdownCostIndex];
-		}
-
-		FText Label;
-		int32 BreakdownCostIndex = INDEX_NONE;
-	};
-
-#if WITH_EDITORONLY_DATA
-	struct FCostModifier : ITextColumn
-	{
-		using ITextColumn::ITextColumn;
-
-		virtual FText GetLabel() const override { return LOCTEXT("ColumnLabelCostModifier", "Bias"); }
-
-		virtual FSortPredicate GetSortPredicate() const override
-		{
-			return [](const FRowDataRef& Row0, const FRowDataRef& Row1) -> bool { return Row0->PoseCost.GetCostAddend() < Row1->PoseCost.GetCostAddend(); };
-		}
-
-		virtual FText GetRowText(const FRowDataRef& Row) const override
-		{
-			return FText::AsNumber(Row->PoseCost.GetCostAddend());
-		}
-	};
-#endif // WITH_EDITORONLY_DATA
-
-	struct FMirrored : ITextColumn
-	{
-		using ITextColumn::ITextColumn;
-
-		virtual FText GetLabel() const override { return LOCTEXT("ColumnLabelMirrored", "Mirror"); }
-
-		virtual FSortPredicate GetSortPredicate() const override
-		{
-			return [](const FRowDataRef& Row0, const FRowDataRef& Row1) -> bool { return Row0->bMirrored < Row1->bMirrored; };
-		}
-
-		virtual FText GetRowText(const FRowDataRef& Row) const override
-		{
-			return FText::Format(LOCTEXT("Mirrored", "{0}"), { Row->bMirrored } );
-		}
-	};
-
-	struct FLooping : ITextColumn
-	{
-		using ITextColumn::ITextColumn;
-
-		virtual FText GetLabel() const override { return LOCTEXT("ColumnLabelLooping", "Loop"); }
-
-		virtual FSortPredicate GetSortPredicate() const override
-		{
-			return [](const FRowDataRef& Row0, const FRowDataRef& Row1) -> bool { return Row0->bLooping < Row1->bLooping; };
-		}
-
-		virtual FText GetRowText(const FRowDataRef& Row) const override
-		{
-			return FText::Format(LOCTEXT("Looping", "{0}"), { Row->bLooping });
-		}
-	};
-
-	struct FBlendParameters : ITextColumn
-	{
-		using ITextColumn::ITextColumn;
-
-		virtual FText GetLabel() const override { return LOCTEXT("ColumnLabelBlendParams", "Blend Params"); }
-
-		virtual FSortPredicate GetSortPredicate() const override
-		{
-			return [](const FRowDataRef& Row0, const FRowDataRef& Row1) -> bool
-			{
-				return (
-					Row0->BlendParameters[0] < Row1->BlendParameters[0] ||
-					Row0->BlendParameters[1] < Row1->BlendParameters[1]);
-			};
-		}
-
-		virtual FText GetRowText(const FRowDataRef& Row) const override
-		{
-			if (Row->AssetType == ESearchIndexAssetType::BlendSpace)
-			{
-				return FText::Format(
-					LOCTEXT("Blend Parameters", "({0}, {1})"),
-					FText::AsNumber(Row->BlendParameters[0]),
-					FText::AsNumber(Row->BlendParameters[1]));
-			}
-			else
-			{
-				return FText::FromString(TEXT("-"));
-			}
-		}
-	};
-
-	struct FPoseCandidateFlags : ITextColumn
-	{
-		using ITextColumn::ITextColumn;
-
-		virtual FText GetLabel() const override { return LOCTEXT("ColumnLabelPoseCandidateFlags", "Flags"); }
-
-		virtual FSortPredicate GetSortPredicate() const override
-		{
-			return [](const FRowDataRef& Row0, const FRowDataRef& Row1) -> bool { return Row0->PoseCandidateFlags < Row1->PoseCandidateFlags; };
-		}
-
-		virtual FText GetRowText(const FRowDataRef& Row) const override
-		{
-			if (EnumHasAnyFlags(Row->PoseCandidateFlags, EPoseCandidateFlags::AnyDiscardedMask))
-			{
-				auto AddDelimiter = [](FTextBuilder& TextBuilder, bool& bNeedDelimiter)
-				{
-					if (bNeedDelimiter)
-					{
-						TextBuilder.AppendLine(LOCTEXT("DiscardedBy_Delimiter", " | "));
-					}
-					bNeedDelimiter = true;
-				};
-
-				bool bNeedDelimiter = false;
-
-				FTextBuilder TextBuilder;
-				if (EnumHasAnyFlags(Row->PoseCandidateFlags, EPoseCandidateFlags::DiscardedBy_PoseJumpThresholdTime))
-				{
-					AddDelimiter(TextBuilder, bNeedDelimiter);
-					TextBuilder.AppendLine(LOCTEXT("DiscardedBy_PoseJumpThresholdTime", "PoseJumpThresholdTime"));
-				}
-				
-				if (EnumHasAnyFlags(Row->PoseCandidateFlags, EPoseCandidateFlags::DiscardedBy_PoseReselectHistory))
-				{
-					AddDelimiter(TextBuilder, bNeedDelimiter);
-					TextBuilder.AppendLine(LOCTEXT("DiscardedBy_PoseReselectHistory", "PoseReselectHistory"));
-				}
-
-				if (EnumHasAnyFlags(Row->PoseCandidateFlags, EPoseCandidateFlags::DiscardedBy_BlockTransition))
-				{
-					AddDelimiter(TextBuilder, bNeedDelimiter);
-					TextBuilder.AppendLine(LOCTEXT("DiscardedBy_BlockTransition", "BlockTransition"));
-				}
-
-				if (EnumHasAnyFlags(Row->PoseCandidateFlags, EPoseCandidateFlags::DiscardedBy_PoseFilter))
-				{
-					AddDelimiter(TextBuilder, bNeedDelimiter);
-					TextBuilder.AppendLine(LOCTEXT("DiscardedBy_PoseFilter", "PoseFilter"));
-				}
-
-				return TextBuilder.ToText();
-			}
-
-			return FText::GetEmpty();
-		}
-	};
-}
-
-
-/**
- * Widget representing a single row of the database view
- */
-class SDebuggerDatabaseRow : public SMultiColumnTableRow<TSharedRef<FDebuggerDatabaseRowData>>
-{
-	SLATE_BEGIN_ARGS(SDebuggerDatabaseRow) {}
-		SLATE_ATTRIBUTE(const SDebuggerDatabaseView::FColumnMap*, ColumnMap)
-	SLATE_END_ARGS()
-
-	void Construct(
-		const FArguments& InArgs,
-		const TSharedRef<STableViewBase>& InOwnerTable,
-		TSharedRef<FDebuggerDatabaseRowData> InRow,
-		const FTableRowStyle& InRowStyle,
-		const FSlateBrush* InRowBrush,
-		FMargin InPaddingMargin
-	)
-	{
-		ColumnMap = InArgs._ColumnMap;
-		check(ColumnMap.IsBound());
-		
-		Row = InRow;
-		
-		RowBrush = InRowBrush;
-		check(RowBrush);
-
-		SMultiColumnTableRow<TSharedRef<FDebuggerDatabaseRowData>>::Construct(
-			FSuperRowType::FArguments()
-			.Padding(InPaddingMargin)
-			.Style(&InRowStyle),
-			InOwnerTable
-		);
-	}
-
-	virtual TSharedRef<SWidget> GenerateWidgetForColumn(const FName& InColumnName) override
-	{
-		// Get column
-		const TSharedRef<DebuggerDatabaseColumns::IColumn>& Column = (*ColumnMap.Get())[InColumnName];
-		
-		static FSlateFontInfo NormalFont = FAppStyle::Get().GetFontStyle("DetailsView.CategoryTextStyle");
-		const TSharedRef<SWidget> Widget = Column->GenerateWidget(Row.ToSharedRef());
-		
-		return
-			SNew(SBorder)
-			.HAlign(HAlign_Center)
-			.VAlign(VAlign_Center)
-			.BorderImage(RowBrush)
-			[
-				SNew(SHorizontalBox)
-				+ SHorizontalBox::Slot()
-				.VAlign(VAlign_Center)
-				.HAlign(HAlign_Center)
-				.Padding(0.0f, 3.0f)
-				[
-					Widget
-				]
-			];
-	}
-
-	/** Row data associated with this widget */
-	TSharedPtr<FDebuggerDatabaseRowData> Row;
-
-	/** Used for cell styles (active vs database row) */
-	const FSlateBrush* RowBrush = nullptr;
-
-	/** Used to grab the column struct given a column name */
-	TAttribute<const SDebuggerDatabaseView::FColumnMap*> ColumnMap;
-};
-
-class SDebuggerMessageBox : public SCompoundWidget
-{
-	SLATE_BEGIN_ARGS(SDebuggerMessageBox) {}
-	SLATE_END_ARGS()
-
-	void Construct(const FArguments& InArgs, const FString& Message)
-	{
-		ChildSlot
-		[
-			SNew(SVerticalBox)
-			+ SVerticalBox::Slot()
-			.HAlign(HAlign_Center)
-			.VAlign(VAlign_Center)
-			[
-				SNew(STextBlock)
-				.Text(FText::FromString(Message))
-				.Font(FAppStyle::Get().GetFontStyle("DetailsView.CategoryFontStyle"))
-			]
-		];
-	}
-};
-
-class SCostBreakDownData : public ICostBreakDownData
-{
-public:
-	SCostBreakDownData(TArray<TSharedRef<FDebuggerDatabaseRowData>>& UnfilteredDatabaseRows_, bool bIsVerbose_)
-	: UnfilteredDatabaseRows(UnfilteredDatabaseRows_)
-	, bIsVerbose(bIsVerbose_)
-	{
-	}
-
-	// begin ICostBreakDownData
-	virtual ~SCostBreakDownData()
-	{
-		check(CurrentLabeIndex == -1);
-	}
-
-	virtual int32 Num() const override
-	{
-		return UnfilteredDatabaseRows.Num();
-	}
-
-	virtual bool IsCostVectorFromSchema(int32 Index, const UPoseSearchSchema* Schema) const override
-	{
-		return UnfilteredDatabaseRows[Index]->SourceDatabase->Schema == Schema;
-	}
-
-	virtual TConstArrayView<float> GetCostVector(int32 Index, const UPoseSearchSchema* Schema) const override
-	{
-		check(UnfilteredDatabaseRows[Index]->SourceDatabase->Schema == Schema);
-		return UnfilteredDatabaseRows[Index]->CostVector;
-	}
-
-	virtual void BeginBreakDownSection(const FText& Label) override
-	{
-		check(CurrentLabeIndex == -1);
-		CurrentLabeIndex = GetOrAddLabelIndex(Label);
-	}
-
-	virtual void SetCostBreakDown(float CostBreakDown, int32 Index, const UPoseSearchSchema* Schema) override
-	{
-		FDebuggerDatabaseRowData& UnfilteredDatabaseRow = UnfilteredDatabaseRows[Index].Get();
-		check(UnfilteredDatabaseRow.SourceDatabase->Schema == Schema);
-		TArray<float>& CostBreakdowns = UnfilteredDatabaseRow.CostBreakdowns;
-		CostBreakdowns.SetNumZeroed(CurrentLabeIndex + 1, false);
-		CostBreakdowns[CurrentLabeIndex] = CostBreakDown;
-	}
-
-	virtual void EndBreakDownSection(const FText& Label) override
-	{
-		check(CurrentLabeIndex == GetOrAddLabelIndex(Label));
-		CurrentLabeIndex = INDEX_NONE;
-	}
-
-	virtual bool IsVerbose() const override
-	{
-		return bIsVerbose;
-	}
-	// end ICostBreakDownData
-
-	int32 GetOrAddLabelIndex(const FText& Label)
-	{
-		for (int32 i = 0; i < Labels.Num(); ++i)
-		{
-			if (Labels[i].EqualTo(Label))
-			{
-				return i;
-			}
-		}
-		Labels.Add(Label);
-		return Labels.Num() - 1;
-	}
-
-	bool AreLabelsEqualTo(TArray<FText> OtherLabels)
-	{
-		if (Labels.Num() != OtherLabels.Num())
-		{
-			return false;
-		}
-
-		for (int32 i = 0; i < Labels.Num(); ++i)
-		{
-			if (!Labels[i].EqualTo(OtherLabels[i]))
-			{
-				return false;
-			}
-		}
-
-		return true;
-	}
-
-	const TArray<FText>& GetLabels() const
-	{
-		return Labels;
-	}
-
-private:
-	TArray<TSharedRef<FDebuggerDatabaseRowData>>& UnfilteredDatabaseRows;
-
-	TArray<FText> Labels;
-
-	int32 CurrentLabeIndex = INDEX_NONE;
-
-	bool bIsVerbose = true;
-};
-
-void SDebuggerDatabaseView::Update(const FTraceMotionMatchingStateMessage& State)
-{
-	// row cost color palette
-	static const FLinearColor DiscardedRowColor(0.314f, 0.314f, 0.314f); // darker gray
-	static const FLinearColor BestScoreRowColor = FLinearColor::Green;
-	static const FLinearColor WorstScoreRowColor = FLinearColor::Red;
-
-	using namespace DebuggerDatabaseColumns;
-
-	bool bIsVerbose = false;
-	auto DebuggerView = ParentDebuggerViewPtr.Pin();
-	if (DebuggerView.IsValid())
-	{
-		bIsVerbose = DebuggerView->GetViewModel()->IsVerbose();
-	}
-
-	UnfilteredDatabaseRows.Reset();
-	for (const FTraceMotionMatchingStateDatabaseEntry& DbEntry : State.DatabaseEntries)
-	{
-		const UPoseSearchDatabase* Database = FTraceMotionMatchingState::GetObjectFromId<UPoseSearchDatabase>(DbEntry.DatabaseId);
-		if (FAsyncPoseSearchDatabasesManagement::RequestAsyncBuildIndex(Database, ERequestAsyncBuildFlag::ContinueRequest))
-		{
-			const FPoseSearchIndex& SearchIndex = Database->GetSearchIndex();
-			for (const FTraceMotionMatchingStatePoseEntry& PoseEntry : DbEntry.PoseEntries)
-			{
-				if (const FPoseSearchIndexAsset* SearchIndexAsset = SearchIndex.GetAssetForPoseSafe(PoseEntry.DbPoseIdx))
-				{
-					TSharedRef<FDebuggerDatabaseRowData>& Row = UnfilteredDatabaseRows.Add_GetRef(MakeShared<FDebuggerDatabaseRowData>());
-
-					const float Time = SearchIndex.GetAssetTime(PoseEntry.DbPoseIdx, Database->Schema->GetSamplingInterval());
-
-					Row->PoseIdx = PoseEntry.DbPoseIdx;
-					Row->SourceDatabase = Database;
-					Row->DatabaseName = Database->GetName();
-					Row->DatabasePath = Database->GetPathName();
-					Row->PoseCandidateFlags = PoseEntry.PoseCandidateFlags;
-					Row->DbAssetIdx = SearchIndexAsset->SourceAssetIdx;
-					Row->AssetTime = Time;
-					Row->bMirrored = SearchIndexAsset->bMirrored;
-					Row->PoseCost = PoseEntry.Cost;
-
-					Row->CostVector.SetNum(Database->Schema->SchemaCardinality);
-					TConstArrayView<float> PoseValues = SearchIndex.GetPoseValuesSafe(PoseEntry.DbPoseIdx);
-
-					// in case we modify the schema while PIE is paused and displaying the Pose Search Editor, we could end up with a stale State with a DbEntry.QueryVector saved with the previous schema
-					// so the cardinality of DbEntry.QueryVector and PoseValues don't match. In that case we just use PoseValues as query to have all costs set to zero
-					const bool bIsQueryVectorValid = DbEntry.QueryVector.Num() == PoseValues.Num();
-					CompareFeatureVectors(PoseValues, bIsQueryVectorValid ? DbEntry.QueryVector : PoseValues, SearchIndex.WeightsSqrt, Row->CostVector);
-
-					const FInstancedStruct& DatabaseAssetStruct = Database->GetAnimationAssetStruct(*SearchIndexAsset);
-					if (const FPoseSearchDatabaseAnimationAssetBase* DatabaseAsset = DatabaseAssetStruct.GetPtr<FPoseSearchDatabaseAnimationAssetBase>())
-					{
-						Row->AssetName = DatabaseAsset->GetName();
-						Row->AssetPath = DatabaseAsset->GetAnimationAsset() ? DatabaseAsset->GetAnimationAsset()->GetPathName() : "";
-						Row->AssetType = DatabaseAsset->GetSearchIndexType();
-						Row->bLooping = DatabaseAsset->IsLooping();
-						Row->BlendParameters = FVector::Zero();
-						Row->AnimFrame = 0;
-						Row->AnimPercentage = 0.0f;
-
-						if (const FPoseSearchDatabaseSequence* DatabaseSequence = DatabaseAssetStruct.GetPtr<FPoseSearchDatabaseSequence>())
-						{
-							Row->AnimFrame = DatabaseSequence->Sequence->GetFrameAtTime(Time);
-							Row->AnimPercentage = Time / DatabaseSequence->Sequence->GetPlayLength();
-						}
-						else if (const FPoseSearchDatabaseAnimComposite* DatabaseAnimComposite = DatabaseAssetStruct.GetPtr<FPoseSearchDatabaseAnimComposite>())
-						{
-							Row->AnimFrame = DatabaseAnimComposite->AnimComposite->GetFrameAtTime(Time);
-							Row->AnimPercentage = Time / DatabaseAnimComposite->AnimComposite->GetPlayLength();
-						}
-						else if (const FPoseSearchDatabaseBlendSpace* DatabaseBlendSpace = DatabaseAssetStruct.GetPtr<FPoseSearchDatabaseBlendSpace>())
-						{
-							Row->BlendParameters = SearchIndexAsset->BlendParameters;
-						}
-						else
-						{
-							checkNoEntry();
-						}
-					}
-				}
-			}
-		}
-	}
-
-	SCostBreakDownData CostBreakDownData(UnfilteredDatabaseRows, bIsVerbose);
-	if (!UnfilteredDatabaseRows.IsEmpty())
-	{
-		for (const FTraceMotionMatchingStateDatabaseEntry& DbEntry : State.DatabaseEntries)
-		{
-			const UPoseSearchDatabase* Database = FTraceMotionMatchingState::GetObjectFromId<UPoseSearchDatabase>(DbEntry.DatabaseId);
-			if (FAsyncPoseSearchDatabasesManagement::RequestAsyncBuildIndex(Database, ERequestAsyncBuildFlag::ContinueRequest))
-			{
-				Database->Schema->ComputeCostBreakdowns(CostBreakDownData);
-			}
-		}
-
-		// calculating breakdowns min max and colors
-		TArray<float> MinCostBreakdowns;
-		TArray<float> MaxCostBreakdowns;
-
-		const int32 CostBreakdownsCardinality = UnfilteredDatabaseRows[0]->CostBreakdowns.Num();
-		MinCostBreakdowns.Init(UE_MAX_FLT, CostBreakdownsCardinality);
-		MaxCostBreakdowns.Init(-UE_MAX_FLT, CostBreakdownsCardinality);
-
-		for (auto& UnfilteredRow : UnfilteredDatabaseRows)
-		{
-			ArrayMinMax(UnfilteredRow->CostBreakdowns, MinCostBreakdowns, MaxCostBreakdowns, UE_MAX_FLT);
-		}
-		
-		TArray<float> CostBreakdownsColorBlend;
-		CostBreakdownsColorBlend.Init(0, CostBreakdownsCardinality);
-		for (auto& UnfilteredRow : UnfilteredDatabaseRows)
-		{
-			if (EnumHasAnyFlags(UnfilteredRow->PoseCandidateFlags, EPoseCandidateFlags::AnyValidMask))
-			{
-				ArraySafeNormalize(UnfilteredRow->CostBreakdowns, MinCostBreakdowns, MaxCostBreakdowns, CostBreakdownsColorBlend);
-				LinearColorBlend(BestScoreRowColor, WorstScoreRowColor, CostBreakdownsColorBlend, UnfilteredRow->CostBreakdownsColors);
-			}
-			else
-			{
-				UnfilteredRow->CostBreakdownsColors.Init(DiscardedRowColor, CostBreakdownsCardinality);
-			}
-		}
-
-		float MinCost = UE_MAX_FLT;
-		float MaxCost = -UE_MAX_FLT;
-		for (auto& UnfilteredRow : UnfilteredDatabaseRows)
-		{
-			const float Cost = UnfilteredRow->PoseCost.GetTotalCost();
-			MinCost = FMath::Min(MinCost, Cost);
-			MaxCost = FMath::Max(MaxCost, Cost);
-		}
-
-		const float DeltaCost = MaxCost - MinCost;
-		for (auto& UnfilteredRow : UnfilteredDatabaseRows)
-		{
-			if (EnumHasAnyFlags(UnfilteredRow->PoseCandidateFlags, EPoseCandidateFlags::AnyValidMask))
-			{
-				const float CostColorBlend = DeltaCost > UE_KINDA_SMALL_NUMBER ? (UnfilteredRow->PoseCost.GetTotalCost() - MinCost) / DeltaCost : 0.f;
-				UnfilteredRow->CostColor = UE::PoseSearch::LinearColorBlend(BestScoreRowColor, WorstScoreRowColor, CostColorBlend);
-			}
-			else
-			{
-				UnfilteredRow->CostColor = DiscardedRowColor;
-			}
-		}
-	}
-
-	if (!CostBreakDownData.AreLabelsEqualTo(OldLabels))
-	{
-		OldLabels = CostBreakDownData.GetLabels();
-
-		// recreating and binding the columns
-		Columns.Reset();
-
-		// Construct all column types
-		int32 ColumnIdx = 0;
-		AddColumn(MakeShared<FDatabaseName>(ColumnIdx++));
-		AddColumn(MakeShared<FAssetName>(ColumnIdx++));
-
-		auto CostColumn = MakeShared<FCost>(ColumnIdx++);
-		AddColumn(CostColumn);
-
-		int32 LabelIdx = 0;
-		for (const FText& Label : CostBreakDownData.GetLabels())
-		{
-			AddColumn(MakeShared<FChannelBreakdownCostColumn>(ColumnIdx++, LabelIdx++, Label));
-		}
-
-#if WITH_EDITORONLY_DATA
-		AddColumn(MakeShared<FCostModifier>(ColumnIdx++));
-#endif // WITH_EDITORONLY_DATA
-		AddColumn(MakeShared<FFrame>(ColumnIdx++));
-		AddColumn(MakeShared<FMirrored>(ColumnIdx++));
-		AddColumn(MakeShared<FLooping>(ColumnIdx++));
-		AddColumn(MakeShared<FPoseIdx>(ColumnIdx++));
-		AddColumn(MakeShared<FBlendParameters>(ColumnIdx++));
-		AddColumn(MakeShared<FPoseCandidateFlags>(ColumnIdx++));
-
-		SortColumn = CostColumn->ColumnId;
-
-		// Active and Continuing Pose view scroll bars only for indenting the columns to align w/ database
-		ActiveView.ScrollBar->SetVisibility(EVisibility::Hidden);
-		ContinuingPoseView.ScrollBar->SetVisibility(EVisibility::Hidden);
-
-		// Refresh Columns
-		ActiveView.HeaderRow->ClearColumns();
-		ContinuingPoseView.HeaderRow->ClearColumns();
-		FilteredDatabaseView.HeaderRow->ClearColumns();
-
-		// Sort columns by index
-		Columns.ValueSort([](const TSharedRef<IColumn> Column0, const TSharedRef<IColumn> Column1)
-		{
-			return Column0->SortIndex < Column1->SortIndex;
-		});
-
-		// Add columns from map to header row
-		for (TPair<FName, TSharedRef<IColumn>>& ColumnPair : Columns)
-		{
-			IColumn& Column = ColumnPair.Value.Get();
-			if (ColumnPair.Value->bEnabled)
-			{
-				SHeaderRow::FColumn::FArguments ColumnArgs = SHeaderRow::FColumn::FArguments()
-					.ColumnId(Column.ColumnId)
-					.DefaultLabel(Column.GetLabel())
-					.SortMode(this, &SDebuggerDatabaseView::GetColumnSortMode, Column.ColumnId)
-					.OnSort(this, &SDebuggerDatabaseView::OnColumnSortModeChanged)
-					.FillWidth(this, &SDebuggerDatabaseView::GetColumnWidth, Column.ColumnId)
-					.VAlignCell(VAlign_Center)
-					.VAlignHeader(VAlign_Center)
-					.HAlignHeader(HAlign_Center)
-					.HAlignCell(HAlign_Fill);
-
-				FilteredDatabaseView.HeaderRow->AddColumn(ColumnArgs);
-
-				// Every time the active column is changed, update the database column
-				ActiveView.HeaderRow->AddColumn(ColumnArgs.OnWidthChanged(this, &SDebuggerDatabaseView::OnColumnWidthChanged, Column.ColumnId));
-
-				ContinuingPoseView.HeaderRow->AddColumn(ColumnArgs.OnWidthChanged(this, &SDebuggerDatabaseView::OnColumnWidthChanged, Column.ColumnId));
-			}
-		}
-	}
-
-	SortDatabaseRows();
-	PopulateViewRows();
-}
-
-void SDebuggerDatabaseView::AddColumn(TSharedRef<DebuggerDatabaseColumns::IColumn>&& Column)
-{
-	Columns.Add(Column->ColumnId, Column);
-}
-
-EColumnSortMode::Type SDebuggerDatabaseView::GetColumnSortMode(const FName ColumnId) const
-{
-	if (ColumnId == SortColumn)
-	{
-		return SortMode;
-	}
-	
-	return EColumnSortMode::None;
-}
-
-float SDebuggerDatabaseView::GetColumnWidth(const FName ColumnId) const
-{
-	check(Columns.Find(ColumnId));
-
-	return Columns[ColumnId]->Width;
-}
-
-
-void SDebuggerDatabaseView::OnColumnSortModeChanged(const EColumnSortPriority::Type SortPriority, const FName & ColumnId, const EColumnSortMode::Type InSortMode)
-{
-	check(Columns.Find(ColumnId));
-	SortColumn = ColumnId;
-	SortMode = InSortMode;
-	SortDatabaseRows();
-	PopulateViewRows();
-}
-
-void SDebuggerDatabaseView::OnColumnWidthChanged(const float NewWidth, FName ColumnId) const
-{
-	check(Columns.Find(ColumnId));
-	
-	Columns[ColumnId]->Width = NewWidth;
-}
-
-void SDebuggerDatabaseView::OnFilterTextChanged(const FText& SearchText)
-{
-	FilterText = SearchText;
-	PopulateViewRows();
-}
-
-void SDebuggerDatabaseView::OnHideInvalidPosesCheckboxChanged(ECheckBoxState State)
-{
-	if (State == ECheckBoxState::Checked) 
-	{
-		bHideInvalidPoses = true;
-	}
-	else 
-	{
-		bHideInvalidPoses = false;
-	}
-	PopulateViewRows();
-}
-
-void SDebuggerDatabaseView::OnDatabaseRowSelectionChanged(
-	TSharedPtr<FDebuggerDatabaseRowData> Row, 
-	ESelectInfo::Type SelectInfo)
-{
-	if (Row.IsValid())
-	{
-		OnPoseSelectionChanged.ExecuteIfBound(Row->SourceDatabase.Get(), Row->PoseIdx, Row->AssetTime);
-	}
-}
-
-void SDebuggerDatabaseView::SortDatabaseRows()
-{
-	if (!UnfilteredDatabaseRows.IsEmpty())
-	{
-		if (SortMode == EColumnSortMode::Ascending)
-		{
-			UnfilteredDatabaseRows.Sort(Columns[SortColumn]->GetSortPredicate());
-		}
-		else if (SortMode == EColumnSortMode::Descending)
-		{
-			auto DescendingPredicate = [this](const auto& Lhs, const auto& Rhs) -> bool
-			{
-				return !Columns[SortColumn]->GetSortPredicate()(Lhs, Rhs);
-			};
-
-			UnfilteredDatabaseRows.Sort(DescendingPredicate);
-		}
-	}
-}
-
-void SDebuggerDatabaseView::PopulateViewRows()
-{
-	ActiveView.Rows.Reset();
-	ContinuingPoseView.Rows.Reset();
-	FilteredDatabaseView.Rows.Empty();
-
-	FString FilterString = FilterText.ToString();
-	TArray<FString> Tokens;
-	FilterString.ParseIntoArrayWS(Tokens);
-	const bool bHasNameFilter = !Tokens.IsEmpty();
-
-	for (const auto& UnfilteredRow : UnfilteredDatabaseRows)
-	{
-		bool bTryAddToFilteredDatabaseViewRows = true;
-		if (EnumHasAnyFlags(UnfilteredRow->PoseCandidateFlags, EPoseCandidateFlags::Valid_ContinuingPose))
-		{
-			ContinuingPoseView.Rows.Add(UnfilteredRow);
-			bTryAddToFilteredDatabaseViewRows = false;
-		}
-
-		if (EnumHasAnyFlags(UnfilteredRow->PoseCandidateFlags, EPoseCandidateFlags::Valid_CurrentPose))
-		{
-			ActiveView.Rows.Add(UnfilteredRow);
-			bTryAddToFilteredDatabaseViewRows = false;
-		}
-
-		if (bTryAddToFilteredDatabaseViewRows)
-		{
-			if (!bHideInvalidPoses || EnumHasAnyFlags(UnfilteredRow->PoseCandidateFlags, EPoseCandidateFlags::AnyValidMask))
-			{
-				bool bPassesNameFilter = true;
-				if (bHasNameFilter)
-				{
-					bPassesNameFilter = Algo::AllOf(Tokens, [&](FString Token)
-					{
-						return UnfilteredRow->AssetName.Contains(Token);
-					});
-				}
-
-				if (bPassesNameFilter)
-				{
-					FilteredDatabaseView.Rows.Add(UnfilteredRow);
-				}
-			}
-		}
-	}
-
-	ActiveView.ListView->RequestListRefresh();
-	ContinuingPoseView.ListView->RequestListRefresh();
-	FilteredDatabaseView.ListView->RequestListRefresh();
-}
-
-TSharedRef<ITableRow> SDebuggerDatabaseView::HandleGenerateDatabaseRow(TSharedRef<FDebuggerDatabaseRowData> Item, const TSharedRef<STableViewBase>& OwnerTable) const
-{
-	return SNew(SDebuggerDatabaseRow, OwnerTable, Item, FilteredDatabaseView.RowStyle, &FilteredDatabaseView.RowBrush, FMargin(0.0f, 2.0f, 6.0f, 2.0f))
-		.ColumnMap(this, &SDebuggerDatabaseView::GetColumnMap);	
-}
-
-TSharedRef<ITableRow> SDebuggerDatabaseView::HandleGenerateActiveRow(TSharedRef<FDebuggerDatabaseRowData> Item, const TSharedRef<STableViewBase>& OwnerTable) const
-{
-	return SNew(SDebuggerDatabaseRow, OwnerTable, Item, ActiveView.RowStyle, &ActiveView.RowBrush, FMargin(0.0f, 2.0f, 6.0f, 4.0f))
-		.ColumnMap(this, &SDebuggerDatabaseView::GetColumnMap);
-}
-
-TSharedRef<ITableRow> SDebuggerDatabaseView::HandleGenerateContinuingPoseRow(TSharedRef<FDebuggerDatabaseRowData> Item, const TSharedRef<STableViewBase>& OwnerTable) const
-{
-	return SNew(SDebuggerDatabaseRow, OwnerTable, Item, ContinuingPoseView.RowStyle, &ContinuingPoseView.RowBrush, FMargin(0.0f, 2.0f, 6.0f, 4.0f))
-		.ColumnMap(this, &SDebuggerDatabaseView::GetColumnMap);
-}
-
-void SDebuggerDatabaseView::Construct(const FArguments& InArgs)
-{
-	ParentDebuggerViewPtr = InArgs._Parent;
-	OnPoseSelectionChanged = InArgs._OnPoseSelectionChanged;
-	check(OnPoseSelectionChanged.IsBound());
-
-	// Active Row
-	ActiveView.HeaderRow = SNew(SHeaderRow);
-
-	// Used for spacing
-	ActiveView.ScrollBar = SNew(SScrollBar)
-		.Orientation(Orient_Vertical)
-		.HideWhenNotInUse(false)
-		.AlwaysShowScrollbar(true)
-		.AlwaysShowScrollbarTrack(true);
-
-	ActiveView.ListView = SNew(SListView<TSharedRef<FDebuggerDatabaseRowData>>)
-		.ListItemsSource(&ActiveView.Rows)
-		.HeaderRow(ActiveView.HeaderRow.ToSharedRef())
-		.OnGenerateRow(this, &SDebuggerDatabaseView::HandleGenerateActiveRow)
-		.ExternalScrollbar(ActiveView.ScrollBar)
-		.SelectionMode(ESelectionMode::SingleToggle)
-		.ConsumeMouseWheel(EConsumeMouseWheel::Never);
-
-	ActiveView.RowStyle = FAppStyle::GetWidgetStyle<FTableRowStyle>("TableView.Row");
-	ActiveView.RowBrush = *FAppStyle::GetBrush("DetailsView.CategoryTop");
-
-	// ContinuingPose Row
-	ContinuingPoseView.HeaderRow = SNew(SHeaderRow).Visibility(EVisibility::Collapsed);
-
-	// Used for spacing
-	ContinuingPoseView.ScrollBar = SNew(SScrollBar)
-		.Orientation(Orient_Vertical)
-		.HideWhenNotInUse(false)
-		.AlwaysShowScrollbar(true)
-		.AlwaysShowScrollbarTrack(true);
-
-	ContinuingPoseView.ListView = SNew(SListView<TSharedRef<FDebuggerDatabaseRowData>>)
-		.ListItemsSource(&ContinuingPoseView.Rows)
-		.HeaderRow(ContinuingPoseView.HeaderRow.ToSharedRef())
-		.OnGenerateRow(this, &SDebuggerDatabaseView::HandleGenerateContinuingPoseRow)
-		.ExternalScrollbar(ContinuingPoseView.ScrollBar)
-		.SelectionMode(ESelectionMode::SingleToggle)
-		.ConsumeMouseWheel(EConsumeMouseWheel::Never);
-
-	ContinuingPoseView.RowStyle = FAppStyle::GetWidgetStyle<FTableRowStyle>("TableView.Row");
-	ContinuingPoseView.RowBrush = *FAppStyle::GetBrush("DetailsView.CategoryTop");
-
-	// Filtered Database
-	FilteredDatabaseView.ScrollBar =
-		SNew(SScrollBar)
-		.Orientation(Orient_Vertical)
-		.HideWhenNotInUse(false)
-		.AlwaysShowScrollbar(true)
-		.AlwaysShowScrollbarTrack(true);
-	FilteredDatabaseView.HeaderRow = SNew(SHeaderRow).Visibility(EVisibility::Collapsed);
-
-	FilteredDatabaseView.ListView = SNew(SListView<TSharedRef<FDebuggerDatabaseRowData>>)
-		.ListItemsSource(&FilteredDatabaseView.Rows)
-		.HeaderRow(FilteredDatabaseView.HeaderRow.ToSharedRef())
-		.OnGenerateRow(this, &SDebuggerDatabaseView::HandleGenerateDatabaseRow)
-		.ExternalScrollbar(FilteredDatabaseView.ScrollBar)
-		.SelectionMode(ESelectionMode::Multi)
-		.ConsumeMouseWheel(EConsumeMouseWheel::WhenScrollingPossible)
-		.OnSelectionChanged(this, &SDebuggerDatabaseView::OnDatabaseRowSelectionChanged);
-
-	FilteredDatabaseView.RowStyle = FAppStyle::GetWidgetStyle<FTableRowStyle>("TableView.Row");
-	// Set selected color to white to retain visibility when multi-selecting
-	FilteredDatabaseView.RowStyle.SetSelectedTextColor(FLinearColor(FVector3f(0.8f)));
-	FilteredDatabaseView.RowBrush = *FAppStyle::GetBrush("ToolPanel.GroupBorder");
-
-	ChildSlot
-	[
-		SNew(SVerticalBox)
-		+ SVerticalBox::Slot()
-		// Side and top margins, ignore bottom handled by the color border below
-		.Padding(0.0f, 5.0f, 0.0f, 0.0f)
-		.AutoHeight()
-		[
-			// Active Row text tab
-			SNew(SVerticalBox)
-			+ SVerticalBox::Slot()
-			.AutoHeight()
-			.Padding(0.0f)
-			[
-				SNew(SHorizontalBox)
-				+ SHorizontalBox::Slot()
-				.HAlign(HAlign_Center)
-				.VAlign(VAlign_Fill)
-				.Padding(0.0f)
-				.AutoWidth()
-				[
-					SNew(SBorder)
-					.BorderImage(FAppStyle::GetBrush("DetailsView.CategoryTop"))
-					.Padding(FMargin(30.0f, 3.0f, 30.0f, 0.0f))
-					.HAlign(HAlign_Center)
-					.VAlign(VAlign_Fill)
-					[
-						SNew(STextBlock)
-						.Text(FText::FromString("Active Pose"))	
-					]
-				]
-			]
-
-			// Active row list view with scroll bar
-			+ SVerticalBox::Slot()
-			
-			.AutoHeight()
-			[
-				SNew(SHorizontalBox)
-				
-				+ SHorizontalBox::Slot()
-				.HAlign(HAlign_Fill)
-				.VAlign(VAlign_Fill)
-				.Padding(0.0f)
-				[
-					SNew(SBorder)
-					
-					.BorderImage(FAppStyle::GetBrush("NoBorder"))
-					.Padding(0.0f)
-					[
-						ActiveView.ListView.ToSharedRef()
-					]
-				]
-
-				+ SHorizontalBox::Slot()
-				.AutoWidth()
-				[
-					ActiveView.ScrollBar.ToSharedRef()
-				]
-			]	
-		]
-
-		+ SVerticalBox::Slot()
-		// Side and top margins, ignore bottom handled by the color border below
-		.Padding(0.0f, 5.0f, 0.0f, 0.0f)
-		.AutoHeight()
-		[
-			// ContinuingPose Row text tab
-			SNew(SVerticalBox)
-			+ SVerticalBox::Slot()
-			.AutoHeight()
-			.Padding(0.0f)
-			[
-				SNew(SHorizontalBox)
-				+ SHorizontalBox::Slot()
-				.HAlign(HAlign_Center)
-				.VAlign(VAlign_Fill)
-				.Padding(0.0f)
-				.AutoWidth()
-				[
-					SNew(SBorder)
-					.BorderImage(FAppStyle::GetBrush("DetailsView.CategoryTop"))
-					.Padding(FMargin(30.0f, 3.0f, 30.0f, 0.0f))
-					.HAlign(HAlign_Center)
-					.VAlign(VAlign_Fill)
-					[
-						SNew(STextBlock)
-						.Text(FText::FromString("Continuing Pose"))	
-					]
-				]
-			]
-
-			// ContinuingPose row list view with scroll bar
-			+ SVerticalBox::Slot()
-			
-			.AutoHeight()
-			[
-				SNew(SHorizontalBox)
-				
-				+ SHorizontalBox::Slot()
-				.HAlign(HAlign_Fill)
-				.VAlign(VAlign_Fill)
-				.Padding(0.0f)
-				[
-					SNew(SBorder)
-					
-					.BorderImage(FAppStyle::GetBrush("NoBorder"))
-					.Padding(0.0f)
-					[
-						ContinuingPoseView.ListView.ToSharedRef()
-					]
-				]
-
-				+ SHorizontalBox::Slot()
-				.AutoWidth()
-				[
-					ContinuingPoseView.ScrollBar.ToSharedRef()
-				]
-			]	
-		]
-		
-		+ SVerticalBox::Slot()
-		.Padding(0.0f, 0.0f, 0.0f, 5.0f)
-		[
-			// Database view text tab
-			SNew(SVerticalBox)
-			+ SVerticalBox::Slot()
-			[
-				SNew(SHorizontalBox)
-				+ SHorizontalBox::Slot()
-				.HAlign(HAlign_Center)
-				.VAlign(VAlign_Fill)
-				.Padding(0.0f)
-				[
-					SNew(SBorder)
-					.BorderImage(FAppStyle::GetBrush("DetailsView.CategoryTop"))
-					.Padding(FMargin(30.0f, 3.0f, 30.0f, 0.0f))
-					.HAlign(HAlign_Center)
-					.VAlign(VAlign_Fill)
-					[
-						SNew(STextBlock)
-						.Text(FText::FromString("Pose Candidates"))
-					]
-				]
-				.AutoWidth()
-				
-				+ SHorizontalBox::Slot()
-				.HAlign(HAlign_Fill)
-				[
-					SNew(SBorder)
-					.BorderImage(&FilteredDatabaseView.RowStyle.EvenRowBackgroundBrush)
-				]
-			]
-			.AutoHeight()
-
-			// Gray line below the tab 
-			+ SVerticalBox::Slot()
-			.AutoHeight()
-			.Padding(0.0f)
-			[
-				SNew(SBorder)
-				.BorderImage(FAppStyle::GetBrush("DetailsView.CategoryTop"))
-				.Padding(FMargin(0.0f, 3.0f, 0.0f, 3.0f))
-				.HAlign(HAlign_Fill)
-				.VAlign(VAlign_Fill)
-			]
-
-			+ SVerticalBox::Slot()
-			.Padding(0.0f, 0.0f, 0.0f, 5.0f)
-			.AutoHeight()
-			[
-				SNew(SHorizontalBox)
-				+ SHorizontalBox::Slot()
-				.Padding(10, 5, 10, 5)
-				[
-					SAssignNew(FilterBox, SSearchBox)
-					.OnTextChanged(this, &SDebuggerDatabaseView::OnFilterTextChanged)
-				]
-				+ SHorizontalBox::Slot()
-				.AutoWidth()
-				.Padding(10, 5, 10, 5)
-				[
-					SNew(SCheckBox)
-					.OnCheckStateChanged_Lambda([this](ECheckBoxState State)
-					{
-						   SDebuggerDatabaseView::OnHideInvalidPosesCheckboxChanged(State);
-					})
-					[
-						SNew(STextBlock)
-						.Text(LOCTEXT("PoseSearchDebuggerHideInvalidPosesFlag", "Hide Invalid Poses"))
-					]
-				]
-			]
-		
-			+ SVerticalBox::Slot()
-			[
-				SNew(SHorizontalBox)
-				+ SHorizontalBox::Slot()
-				.Padding(0.0f)
-				[
-					SNew(SBorder)
-					.BorderImage(FAppStyle::GetBrush("NoBorder"))
-					.Padding(0.0f)
-					[
-						FilteredDatabaseView.ListView.ToSharedRef()
-					]
-				]
-				
-				+ SHorizontalBox::Slot()
-				.AutoWidth()
-				[
-					FilteredDatabaseView.ScrollBar.ToSharedRef()
-				]
-			]
-		]
-	];
-
-	SortMode = EColumnSortMode::Ascending;
-	OldLabels.Reset();
-	Columns.Reset();
-}
-
-void SDebuggerDetailsView::Construct(const FArguments& InArgs)
-{
-	ParentDebuggerViewPtr = InArgs._Parent;
-
-	// Add property editor (detail view) UObject to world root so that it persists when PIE is stopped
-	Reflection = NewObject<UPoseSearchDebuggerReflection>();
-	Reflection->AddToRoot();
-	check(IsValid(Reflection));
-
-	// @TODO: Convert this to a custom builder instead of of a standard details view
-	// Load property module and create details view with our reflection UObject
-	FPropertyEditorModule& PropPlugin = FModuleManager::LoadModuleChecked<FPropertyEditorModule>("PropertyEditor");
-	FDetailsViewArgs DetailsViewArgs;
-    DetailsViewArgs.NameAreaSettings = FDetailsViewArgs::HideNameArea;
-	DetailsViewArgs.DefaultsOnlyVisibility = EEditDefaultsOnlyNodeVisibility::Hide;
-
-	Details = PropPlugin.CreateDetailView(DetailsViewArgs);
-	Details->SetObject(Reflection);
-	
-	ChildSlot
-	[
-		Details.ToSharedRef()
-	];
-}
-
-void SDebuggerDetailsView::Update(const FTraceMotionMatchingStateMessage& State) const
-{
-	UpdateReflection(State);
-}
-
-SDebuggerDetailsView::~SDebuggerDetailsView()
-{
-	// Our previously instantiated object attached to root may be cleaned up at this point
-	if (UObjectInitialized())
-	{
-		Reflection->RemoveFromRoot();
-	}
-}
-
-void SDebuggerDetailsView::UpdateReflection(const FTraceMotionMatchingStateMessage& State) const
-{
-	check(Reflection);
-
-	const UPoseSearchDatabase* CurrentDatabase = State.GetCurrentDatabase();
-	if (FAsyncPoseSearchDatabasesManagement::RequestAsyncBuildIndex(CurrentDatabase, ERequestAsyncBuildFlag::ContinueRequest))
-	{
-		const FPoseSearchIndex& CurrentSearchIndex = CurrentDatabase->GetSearchIndex();
-		int32 CurrentDbPoseIdx = State.GetCurrentDatabasePoseIndex();
-
-		Reflection->CurrentDatabaseName = CurrentDatabase->GetName();
-		Reflection->ElapsedPoseJumpTime = State.ElapsedPoseJumpTime;
-
-		Reflection->AssetPlayerAssetName = "None";
-		if (const FPoseSearchIndexAsset* IndexAsset = CurrentSearchIndex.GetAssetForPoseSafe(CurrentDbPoseIdx))
-		{
-			Reflection->AssetPlayerAssetName = CurrentDatabase->GetSourceAssetName(*IndexAsset);
-		}
-
-		Reflection->AssetPlayerTime = State.AssetPlayerTime;
-		Reflection->LastDeltaTime = State.DeltaTime;
-		Reflection->SimLinearVelocity = State.SimLinearVelocity;
-		Reflection->SimAngularVelocity = State.SimAngularVelocity;
-		Reflection->AnimLinearVelocity = State.AnimLinearVelocity;
-		Reflection->AnimAngularVelocity = State.AnimAngularVelocity;
-
-		// Query pose
-		Reflection->QueryPoseVector.Reset();
-		for (const FTraceMotionMatchingStateDatabaseEntry& DbEntry : State.DatabaseEntries)
-		{
-			const UPoseSearchDatabase* Database = FTraceMotionMatchingState::GetObjectFromId<UPoseSearchDatabase>(DbEntry.DatabaseId);
-			if (Database && Database == CurrentDatabase)
-			{
-				Reflection->QueryPoseVector = DbEntry.QueryVector;
-				break;
-			}
-		}
-
-		// Active pose
-		Reflection->ActivePoseVector = CurrentSearchIndex.GetPoseValuesSafe(CurrentDbPoseIdx);
-	}
-
-	auto DebuggerView = ParentDebuggerViewPtr.Pin();
-	if (DebuggerView.IsValid())
-	{
-		TArray<TSharedRef<FDebuggerDatabaseRowData>> SelectedRows = DebuggerView->GetSelectedDatabaseRows();
-		if (!SelectedRows.IsEmpty())
-		{
-			const TSharedRef<FDebuggerDatabaseRowData>& Selected = SelectedRows[0];
-			if (FAsyncPoseSearchDatabasesManagement::RequestAsyncBuildIndex(Selected->SourceDatabase.Get(), ERequestAsyncBuildFlag::ContinueRequest))
-			{
-				const FPoseSearchIndex& SelectedSearchIndex = Selected->SourceDatabase->GetSearchIndex();
-				Reflection->SelectedPoseVector = SelectedSearchIndex.GetPoseValuesSafe(Selected->PoseIdx);
-			}
-			Reflection->CostVector = Selected->CostVector;
-		}
-	}
-}
-
-void SDebuggerView::Construct(const FArguments& InArgs, uint64 InAnimInstanceId)
-{
-	ViewModel = InArgs._ViewModel;
-	OnViewClosed = InArgs._OnViewClosed;
-	
-	// Validate the existence of the passed getters
-	check(ViewModel.IsBound())
-	check(OnViewClosed.IsBound());
-	
-	AnimInstanceId = InAnimInstanceId;
-	SelectedNodeId = INDEX_NONE;
-
-	ChildSlot
-	[
-		SAssignNew(DebuggerView, SVerticalBox)
-		+ SVerticalBox::Slot()
-		.FillHeight(1.0f)
-		.VAlign(VAlign_Center)
-		.HAlign(HAlign_Fill)
-		[
-			SAssignNew(Switcher, SWidgetSwitcher)
-			.WidgetIndex(this, &SDebuggerView::SelectView)
-
-			// [0] Selection view before node selection is made
-			+ SWidgetSwitcher::Slot()
-			.Padding(40.0f)
-			.HAlign(HAlign_Fill)
-            .VAlign(VAlign_Center)
-			[
-				SAssignNew(SelectionView, SVerticalBox)
-			]
-
-			// [1] Node selected; node debugger view
-			+ SWidgetSwitcher::Slot()
-			[
-				GenerateNodeDebuggerView()
-			]
-
-			// [2] Occluding message box when stopped (no recording)
-			+ SWidgetSwitcher::Slot()
-			[
-				SNew(SDebuggerMessageBox, "Record gameplay to begin debugging")
-			]
-
-			// [3] Occluding message box when recording
-			+ SWidgetSwitcher::Slot()
-			[
-				SNew(SDebuggerMessageBox, "Recording...")
-			]
-			
-			// [4] Occluding message box when there is no data for the selected MM node
-			+ SWidgetSwitcher::Slot()
-			[
-				GenerateNoDataMessageView()
-			]
-		]
-	];
-}
-
-void SDebuggerView::SetTimeMarker(double InTimeMarker)
-{
-	if (FDebugger::IsPIESimulating())
-	{
-		return;
-	}
-
-	TimeMarker = InTimeMarker;
-}
-
-void SDebuggerView::Tick(const FGeometry& AllottedGeometry, const double InCurrentTime, const float InDeltaTime)
-{
-	if (FDebugger::IsPIESimulating())
-	{
-		return;
-	}
-
-	const UWorld* DebuggerWorld = FDebugger::GetWorld();
-    check(DebuggerWorld);
-	
-	// @TODO: Handle editor world when those features are enabled for the Rewind Debugger
-	// Currently prevents debug draw remnants from stopped world
-	if (DebuggerWorld->WorldType != EWorldType::PIE)
-	{
-		return;
-	}
-	
-	const bool bSameTime = FMath::Abs(TimeMarker - PreviousTimeMarker) < DOUBLE_SMALL_NUMBER;
-	PreviousTimeMarker = TimeMarker;
-
-	TSharedPtr<FDebuggerViewModel> Model = ViewModel.Get();
-
-	bool bNeedUpdate = Model->HasSearchableAssetChanged();
-
-	// We haven't reached the update point yet
-	if (CurrentConsecutiveFrames < ConsecutiveFramesUpdateThreshold)
-	{
-		// If we're on the same time marker, it is consecutive
-		if (bSameTime)
-		{
-			++CurrentConsecutiveFrames;
-		}
-	}
-	else
-	{
-		// New frame after having updated, reset consecutive frames count and start counting again
-		if (!bSameTime)
-		{
-			CurrentConsecutiveFrames = 0;
-			bUpdated = false;
-		}
-		// Haven't updated since passing through frame gate, update once
-		else if (!bUpdated)
-		{
-			bNeedUpdate = true;
-		}
-	}
-
-	if (bNeedUpdate)
-	{
-		Model->OnUpdate();
-		if (UpdateNodeSelection())
-		{
-			Model->OnUpdateNodeSelection(SelectedNodeId);
-			UpdateViews();
-		}
-		bUpdated = true;
-	}
-
-	Model->UpdateAsset();
-	
-	// Draw visualization every tick
-	DrawVisualization();
-}
-
-bool SDebuggerView::UpdateNodeSelection()
-{
-	TSharedPtr<FDebuggerViewModel> Model = ViewModel.Get();
-
-	// Update selection view if no node selected
-	bool bNodeSelected = SelectedNodeId != INDEX_NONE;
-	if (!bNodeSelected)
-	{
-		const TArray<int32>& NodeIds = *Model->GetNodeIds();
-		// Only one node active, bypass selection view
-		if (NodeIds.Num() == 1)
-		{
-			SelectedNodeId = *NodeIds.begin();
-			bNodeSelected = true;
-		}
-		// Create selection view with buttons for each node, displaying the database name
-		else
-		{
-			SelectionView->ClearChildren();
-			for (int32 NodeId : NodeIds)
-			{
-				Model->OnUpdateNodeSelection(NodeId);
-				SelectionView->AddSlot()
-				.HAlign(HAlign_Fill)
-				.VAlign(VAlign_Center)
-				.Padding(10.0f)
-				[
-					SNew(SButton)
-					.Text(FText::FromString(Model->GetSearchableAsset()->GetName()))
-					.HAlign(HAlign_Center)
-					.VAlign(VAlign_Center)
-					.ContentPadding(10.0f)
-					.OnClicked(this, &SDebuggerView::OnUpdateNodeSelection, NodeId)
-				];
-			}
-		}
-	}
-
-	return bNodeSelected;
-}
-
-void SDebuggerView::UpdateViews() const
-{
-	const FTraceMotionMatchingStateMessage* State = ViewModel.Get()->GetMotionMatchingState();
-	if (State)
-	{
-		DatabaseView->Update(*State);
-		DetailsView->Update(*State);
-	}
-}
-
-void SDebuggerView::DrawVisualization() const
-{
-	const UWorld* DebuggerWorld = FDebugger::GetWorld();
-	check(DebuggerWorld);
-
-	const FTraceMotionMatchingStateMessage* State = ViewModel.Get()->GetMotionMatchingState();
-	const FTransform* Transform = ViewModel.Get()->GetRootTransform();
-	if (State && Transform)
-	{
-		DrawFeatures(*DebuggerWorld, *State, *Transform, ViewModel.Get()->GetMeshComponent());
-	}
-}
-
-TArray<TSharedRef<FDebuggerDatabaseRowData>> SDebuggerView::GetSelectedDatabaseRows() const
-{
-	return DatabaseView->GetDatabaseRows()->GetSelectedItems();
-}
-
-void SDebuggerView::DrawFeatures(
-	const UWorld& DebuggerWorld,
-	const FTraceMotionMatchingStateMessage& State,
-	const FTransform& Transform,
-	const USkinnedMeshComponent* Mesh
-) const
-{
-	auto SetDrawFlags = [](FDebugDrawParams& InDrawParams, const FPoseSearchDebuggerFeatureDrawOptions& Options)
-	{
-		InDrawParams.Flags = EDebugDrawFlags::None;
-		if (!Options.bDisable)
-		{
-			if (Options.bDrawBoneNames)
-			{
-				EnumAddFlags(InDrawParams.Flags, EDebugDrawFlags::DrawBoneNames);
-			}
-			
-			if (Options.bDrawSampleLabels)
-			{
-				EnumAddFlags(InDrawParams.Flags, EDebugDrawFlags::DrawSampleLabels);
-			}
-		}
-	};
-
-	const TObjectPtr<UPoseSearchDebuggerReflection> Reflection = DetailsView->GetReflection();
-
-	// Draw query vector
-	{
-		const UPoseSearchDatabase* CurrentDatabase = ViewModel.Get()->GetCurrentDatabase();
-		if (CurrentDatabase)
-		{
-			for (const FTraceMotionMatchingStateDatabaseEntry& DbEntry : State.DatabaseEntries)
-			{
-				const UPoseSearchDatabase* Database = FTraceMotionMatchingState::GetObjectFromId<UPoseSearchDatabase>(DbEntry.DatabaseId);
-				if (Database && Database == CurrentDatabase && DbEntry.QueryVector.Num() == Database->Schema->SchemaCardinality && 
-					FAsyncPoseSearchDatabasesManagement::RequestAsyncBuildIndex(CurrentDatabase, ERequestAsyncBuildFlag::ContinueRequest))
-				{
-					// Set shared state
-					FDebugDrawParams DrawParams;
-					DrawParams.World = &DebuggerWorld;
-					DrawParams.RootTransform = Transform;
-					DrawParams.DefaultLifeTime = 0.0f; // Single frame render
-					DrawParams.Mesh = Mesh;
-					DrawParams.Database = CurrentDatabase;
-					SetDrawFlags(DrawParams, Reflection ? Reflection->QueryDrawOptions : FPoseSearchDebuggerFeatureDrawOptions());
-					EnumAddFlags(DrawParams.Flags, EDebugDrawFlags::DrawQuery);
-					DrawFeatureVector(DrawParams, DbEntry.QueryVector);
-					EnumRemoveFlags(DrawParams.Flags, EDebugDrawFlags::DrawQuery);
-					break;
-				}
-			}
-		}
-	}
-
-	// Draw selected poses
-	{
-		const TSharedPtr<SListView<TSharedRef<FDebuggerDatabaseRowData>>>& DatabaseRows = DatabaseView->GetDatabaseRows();
-		TArray<TSharedRef<FDebuggerDatabaseRowData>> SelectedRows = DatabaseRows->GetSelectedItems();
-	
-		// Red for non-active database view
-		FDebugDrawParams DrawParams;
-		DrawParams.World = &DebuggerWorld;
-		DrawParams.RootTransform = Transform;
-		DrawParams.DefaultLifeTime = 0.0f; // Single frame render
-		DrawParams.Mesh = Mesh;
-		SetDrawFlags(DrawParams, Reflection ? Reflection->SelectedPoseDrawOptions : FPoseSearchDebuggerFeatureDrawOptions());
-
-		// Draw any selected database vectors
-		for (const TSharedRef<FDebuggerDatabaseRowData>& Row : SelectedRows)
-		{
-			DrawParams.Database = Row->SourceDatabase.Get();
-			DrawFeatureVector(DrawParams, Row->PoseIdx);
-		}
-	}
-
-
-	// Draw active pose
-	{
-		TArray<TSharedRef<FDebuggerDatabaseRowData>> ActiveRows = DatabaseView->GetActiveRow()->GetSelectedItems();
-
-		// Active row should only have 0 or 1
-		check(ActiveRows.Num() < 2);
-
-		if (!ActiveRows.IsEmpty())
-		{		
-			// Use the motion-matching state's pose idx, as the active row may be update-throttled at this point
-			FDebugDrawParams DrawParams;
-			DrawParams.World = &DebuggerWorld;
-			DrawParams.RootTransform = Transform;
-			DrawParams.DefaultLifeTime = 0.0f; // Single frame render
-			DrawParams.Mesh = Mesh;
-			DrawParams.Database = ActiveRows[0]->SourceDatabase.Get();
-			DrawFeatureVector(DrawParams, ActiveRows[0]->PoseIdx);
-		}
-	}
-
-
-	// Draw continuing pose
-	{
-		TArray<TSharedRef<FDebuggerDatabaseRowData>> ContinuingRows = DatabaseView->GetContinuingPoseRow()->GetSelectedItems();
-
-		// ContinuingPose row should only have 0 or 1
-		check(ContinuingRows.Num() < 2);
-
-		if (!ContinuingRows.IsEmpty())
-		{
-			FDebugDrawParams DrawParams;
-			DrawParams.World = &DebuggerWorld;
-			DrawParams.RootTransform = Transform;
-			DrawParams.DefaultLifeTime = 0.0f; // Single frame render
-			DrawParams.Mesh = Mesh;
-			DrawParams.Database = ContinuingRows[0]->SourceDatabase.Get();
-			DrawFeatureVector(DrawParams, ContinuingRows[0]->PoseIdx);
-		}
-	}
-
-
-	// Draw skeleton
-	{
-		FSkeletonDrawParams SkeletonDrawParams;
-		if (Reflection && Reflection->bDrawSelectedSkeleton)
-		{
-			SkeletonDrawParams.Flags |= ESkeletonDrawFlags::SelectedPose;
-		}
-		if (Reflection && Reflection->bDrawActiveSkeleton)
-		{
-			SkeletonDrawParams.Flags |= ESkeletonDrawFlags::ActivePose;
-		}
-
-		SkeletonDrawParams.Flags |= ESkeletonDrawFlags::Asset;
-
-		ViewModel.Get()->OnDraw(SkeletonDrawParams);
-	}
-}
-
-int32 SDebuggerView::SelectView() const
-{
-	// Currently recording
-	if (FDebugger::IsPIESimulating() && FDebugger::IsRecording())
-	{
-		return RecordingMsg;
-	}
-
-	// Data has not been recorded yet
-	if (FDebugger::GetRecordingDuration() < DOUBLE_SMALL_NUMBER)
-	{
-		return StoppedMsg;
-	}
-
-	const TSharedPtr<FDebuggerViewModel> Model = ViewModel.Get();
-
-	const bool bNoActiveNodes = Model->GetNodesNum() == 0;
-	const bool bNodeSelectedWithoutData = SelectedNodeId != INDEX_NONE && Model->GetMotionMatchingState() == nullptr;
-
-	// No active nodes, or node selected has no data
-	if (bNoActiveNodes || bNodeSelectedWithoutData)
-    {
-    	return NoDataMsg;
-    }
-
-	// Node not selected yet, showcase selection view
-	if (SelectedNodeId == INDEX_NONE)
-	{
-		return Selection;
-	}
-
-	// Standard debugger view
-	return Debugger;
-}
-
-void SDebuggerView::OnPoseSelectionChanged(const UPoseSearchDatabase* Database, int32 DbPoseIdx, float Time)
-{
-	const TSharedPtr<FDebuggerViewModel> Model = ViewModel.Get();
-	const FTraceMotionMatchingStateMessage* State = Model->GetMotionMatchingState();
-
-	if (State)
-	{
-		DetailsView->Update(*State);
-	}
-	
-	if (DbPoseIdx == INDEX_NONE)
-	{
-		Model->ClearSelectedSkeleton();
-	}
-	else
-	{
-		Model->ShowSelectedSkeleton(Database, DbPoseIdx, Time);
-
-		// Stop asset player when switching selections
-		Model->StopSelection();
-	}
-}
-
-FReply SDebuggerView::OnUpdateNodeSelection(int32 InSelectedNodeId)
-{
-	// -1 will backtrack to selection view
-	SelectedNodeId = InSelectedNodeId;
-	bUpdated = false;
-	return FReply::Handled();
-}
-
-FReply SDebuggerView::TogglePlaySelectedAssets() const
-{
-	const TSharedPtr<SListView<TSharedRef<FDebuggerDatabaseRowData>>>& DatabaseRows = DatabaseView->GetDatabaseRows();
-	TArray<TSharedRef<FDebuggerDatabaseRowData>> Selected = DatabaseRows->GetSelectedItems();
-	const bool bPlaying = ViewModel.Get()->IsPlayingSelections();
-	if (!bPlaying)
-	{
-		if (!Selected.IsEmpty())
-		{
-			// @TODO: Make functional with multiple poses being selected
-			ViewModel.Get()->PlaySelection(Selected[0]->PoseIdx, Selected[0]->AssetTime);
-		}
-	}
-	else
-	{
-		ViewModel.Get()->StopSelection();
-	}
-
-	return FReply::Handled();
-}
-
-TSharedRef<SWidget> SDebuggerView::GenerateNoDataMessageView()
-{
-	TSharedRef<SWidget> ReturnButtonView = GenerateReturnButtonView();
-	ReturnButtonView->SetVisibility(TAttribute<EVisibility>::CreateLambda([this]() -> EVisibility
-	{
-		// Hide the return button for the no data message if we have no nodes at all
-		return ViewModel.Get()->GetNodesNum() > 0 ? EVisibility::Visible : EVisibility::Hidden;
-	}));
-	
-	return 
-		SNew(SVerticalBox)
-
-		+ SVerticalBox::Slot()
-		[
-			SNew(SDebuggerMessageBox, "No recorded data available for the selected frame")
-		]
-
-		+ SVerticalBox::Slot()
-		.AutoHeight()
-		.Padding(20.0f)
-		.HAlign(HAlign_Center)
-		.VAlign(VAlign_Center)
-		[
-			ReturnButtonView
-		];
-}
-
-TSharedRef<SHorizontalBox> SDebuggerView::GenerateReturnButtonView()
-{
-	return
-		SNew(SHorizontalBox)
-
-		+ SHorizontalBox::Slot()
-		.Padding(10, 5, 0, 0)
-		.AutoWidth()
-		[
-			SNew(SButton)
-			.VAlign(VAlign_Center)
-			.Visibility_Lambda([this] { return ViewModel.Get()->GetNodesNum() > 1 ? EVisibility::Visible : EVisibility::Hidden; })
-			.ButtonStyle(FAppStyle::Get(), "SimpleButton")
-			.ContentPadding( FMargin(1, 0) )
-			.OnClicked(this, &SDebuggerView::OnUpdateNodeSelection, static_cast<int32>(INDEX_NONE))
-			// Contents of button, icon then text
-			[
-				SNew(SHorizontalBox)
-
-				+ SHorizontalBox::Slot()
-				.VAlign(VAlign_Center)
-				[
-					SNew(SImage)
-					.Image(FAppStyle::Get().GetBrush("Icons.CircleArrowLeft"))
-					.ColorAndOpacity(FSlateColor::UseForeground())
-				]
-
-				+ SHorizontalBox::Slot()
-				.AutoWidth()
-				.Padding(FMargin(8.0f, 0.0f, 0.0f, 0.0f))
-				.VAlign(VAlign_Center)
-				[
-					SNew(STextBlock)
-					.Text(FText::FromString("Return to Database Selection"))
-					.Justification(ETextJustify::Center)
-				]
-			]
-		];
-}
-
-TSharedRef<SWidget> SDebuggerView::GenerateNodeDebuggerView()
-{
-	TSharedRef<SHorizontalBox> ReturnButtonView = GenerateReturnButtonView();
-	ReturnButtonView->AddSlot()
-	.VAlign(VAlign_Center)
-	.HAlign(HAlign_Fill)
-	.Padding(32, 5, 0, 0)
-	.AutoWidth()
-	[
-		SNew(SButton)
-		.VAlign(VAlign_Center)
-		.HAlign(HAlign_Fill)
-		.ButtonStyle(FAppStyle::Get(), "Button")
-		.ContentPadding( FMargin(5, 0) )
-		.OnClicked(this, &SDebuggerView::TogglePlaySelectedAssets)
-		[
-			SNew(SHorizontalBox)
-			// Icon
-			+ SHorizontalBox::Slot()
-			.AutoWidth()
-			[
-				SNew(SImage)
-				.Image_Lambda([this]
-				{
-					const bool bPlayingSelections = ViewModel.Get()->IsPlayingSelections();
-					return FSlateIcon("FAppStyle", bPlayingSelections ? "PlayWorld.StopPlaySession.Small" : "PlayWorld.PlayInViewport.Small").GetSmallIcon();
-				})
-			]
-			// Text
-			+ SHorizontalBox::Slot()
-			.Padding(FMargin(8, 0, 0, 0))
-			.AutoWidth()
-			[
-				SNew(STextBlock)
-				.Text_Lambda([this] { return ViewModel.Get()->IsPlayingSelections() ? FText::FromString("Stop Selected Asset") : FText::FromString("Play Selected Asset"); })
-				.Justification(ETextJustify::Center)
-			]
-		]
-	];
-	
-	ReturnButtonView->AddSlot()
-	.VAlign(VAlign_Center)
-	.HAlign(HAlign_Left)
-	.Padding(64, 5, 0, 0)
-	.AutoWidth()
-	[
-		SNew(SHorizontalBox)
-		
-		+ SHorizontalBox::Slot()
-		.AutoWidth()
-		.Padding(0, 5, 0, 0)
-		[
-			SNew(STextBlock)
-			.Text(FText::FromString("Asset Play Rate: "))
-		]
-		+ SHorizontalBox::Slot()
-		.AutoWidth()
-		.Padding(8, 0, 0, 0)
-		[
-			SNew(SNumericEntryBox<float>)
-			.MinValue(0.0f)
-			.MaxValue(5.0f)
-			.MinSliderValue(0.0f)
-			.MaxSliderValue(5.0f)
-			.Delta(0.01f)
-			.AllowSpin(true)
-			// Lambda to accomodate the TOptional this requires (for now)
-			.Value_Lambda([this] { return ViewModel.Get()->GetPlayRate(); })
-			.OnValueChanged(ViewModel.Get().ToSharedRef(), &FDebuggerViewModel::ChangePlayRate)	
-		]
-	];
-
-	ReturnButtonView->AddSlot()
-	.VAlign(VAlign_Center)
-	.HAlign(HAlign_Left)
-	.Padding(64, 5, 0, 0)
-	.AutoWidth()
-	[
-		SNew(SHorizontalBox)
-
-		+ SHorizontalBox::Slot()
-		.AutoWidth()
-		.Padding(0, 5, 0, 0)
-		[
-			SNew(SCheckBox)
-			.IsChecked_Lambda([this]
-			{
-				return ViewModel.Get()->IsVerbose() ? ECheckBoxState::Checked : ECheckBoxState::Unchecked; 
-			})
-			.OnCheckStateChanged_Lambda([this](ECheckBoxState State)
-			{
-				ViewModel.Get()->SetVerbose(State == ECheckBoxState::Checked);
-				UpdateViews();
-			})
-			[
-				SNew(STextBlock)
-				.Text(LOCTEXT("PoseSearchDebuggerShowVerbose", "Channels Breakdown"))
-			]
-		]
-	];
-
-	return 
-		SNew(SSplitter)
-		.Orientation(Orient_Horizontal)
-		.ResizeMode(ESplitterResizeMode::Fill)
-	
-		// Database view
-		+ SSplitter::Slot()
-		.Value(0.65f)
-		[
-			SNew(SVerticalBox)
-			
-			+ SVerticalBox::Slot()
-			.AutoHeight()
-			[
-				ReturnButtonView
-			]
-			
-			+ SVerticalBox::Slot()
-			[
-				SAssignNew(DatabaseView, SDebuggerDatabaseView)
-				.Parent(SharedThis(this))
-				.OnPoseSelectionChanged(this, &SDebuggerView::OnPoseSelectionChanged)
-			]
-		]
-
-		// Details panel view
-		+ SSplitter::Slot()
-		.Value(0.35f)
-		[
-			SAssignNew(DetailsView, SDebuggerDetailsView)
-			.Parent(SharedThis(this))
-		];
-}
-
-FName SDebuggerView::GetName() const
-{
-	static const FName DebuggerName("PoseSearchDebugger");
-	return DebuggerName;
-}
-
-uint64 SDebuggerView::GetObjectId() const
-{
-	return AnimInstanceId;
-}
-
-SDebuggerView::~SDebuggerView()
-{
-	OnViewClosed.Execute(AnimInstanceId);
-}
-
-FDebuggerViewModel::FDebuggerViewModel(uint64 InAnimInstanceId)
-	: AnimInstanceId(InAnimInstanceId)
-{
-	Skeletons.AddDefaulted(ESkeletonIndex::Num);
-}
-
-FDebuggerViewModel::~FDebuggerViewModel()
-{
-	for (FSkeleton& Skeleton : Skeletons)
-	{
-		if (Skeleton.Actor.IsValid())
-		{
-			Skeleton.Actor->Destroy();
-		}
-	}
-
-	Skeletons.Empty();
-}
-
-const FTraceMotionMatchingStateMessage* FDebuggerViewModel::GetMotionMatchingState() const
-{
-	return ActiveMotionMatchingState;
-}
-
-const UPoseSearchDatabase* FDebuggerViewModel::GetCurrentDatabase() const
-{
-	return ActiveMotionMatchingState ? ActiveMotionMatchingState->GetCurrentDatabase() : nullptr;
-}
-
-const UPoseSearchSearchableAsset* FDebuggerViewModel::GetSearchableAsset() const
-{
-	if (ActiveMotionMatchingState)
-	{
-		return FTraceMotionMatchingState::GetObjectFromId<UPoseSearchSearchableAsset>(ActiveMotionMatchingState->SearchableAssetId);
-	}
-
-	return nullptr;
-}
-
-void FDebuggerViewModel::ShowSelectedSkeleton(const UPoseSearchDatabase* Database, int32 DbPoseIdx, float Time)
-{
-	UPoseSearchMeshComponent* Component = Skeletons[SelectedPose].Component.Get();
-	if (!Component)
-	{
-		return;
-	}
-
-	if (!FAsyncPoseSearchDatabasesManagement::RequestAsyncBuildIndex(Database, ERequestAsyncBuildFlag::ContinueRequest))
-	{
-		return;
-	}
-
-	const FPoseSearchIndex& SearchIndex = Database->GetSearchIndex();
-	const FPoseSearchIndexAsset& IndexAsset = SearchIndex.GetAssetForPose(DbPoseIdx);
-	
-	Component->ResetToStart(); 
-	bSelecting = true;
-	
-	Skeletons[SelectedPose].Type = IndexAsset.Type;
-	Skeletons[SelectedPose].Time = Time;
-	Skeletons[SelectedPose].bMirrored = IndexAsset.bMirrored;
-	Skeletons[SelectedPose].SourceDatabase = Database;
-	Skeletons[SelectedPose].AssetIdx = IndexAsset.SourceAssetIdx;
-	Skeletons[SelectedPose].BlendParameters = IndexAsset.BlendParameters;
-}
-
-void FDebuggerViewModel::ClearSelectedSkeleton()
-{
-	bSelecting = false;
-}
-
-const TArray<int32>* FDebuggerViewModel::GetNodeIds() const
-{
-	return &NodeIds;
-}
-
-int32 FDebuggerViewModel::GetNodesNum() const
-{
-	return MotionMatchingStates.Num();
-}
-
-const FTransform* FDebuggerViewModel::GetRootTransform() const
-{
-	return RootTransform;
-}
-
-bool FDebuggerViewModel::HasSearchableAssetChanged() const
-{
-	uint64 NewSearchableAssetId = ActiveMotionMatchingState ? ActiveMotionMatchingState->SearchableAssetId : 0;
-	return NewSearchableAssetId != SearchableAssetId;
-}
-
-void FDebuggerViewModel::OnUpdate()
-{
-	if (!bSkeletonsInitialized)
-	{
-		UWorld* World = RewindDebugger.Get()->GetWorldToVisualize();
-		for (FSkeleton& Skeleton : Skeletons)
-		{
-			FActorSpawnParameters ActorSpawnParameters;
-			ActorSpawnParameters.bHideFromSceneOutliner = false;
-			ActorSpawnParameters.ObjectFlags |= RF_Transient;
-			Skeleton.Actor = World->SpawnActor<AActor>(ActorSpawnParameters);
-			Skeleton.Actor->SetActorLabel(TEXT("PoseSearch"));
-			Skeleton.Component = NewObject<UPoseSearchMeshComponent>(Skeleton.Actor.Get());
-			Skeleton.Actor->AddInstanceComponent(Skeleton.Component.Get());
-			Skeleton.Component->RegisterComponentWithWorld(World);
-		}
-		FWorldDelegates::OnWorldCleanup.AddRaw(this, &FDebuggerViewModel::OnWorldCleanup);
-		bSkeletonsInitialized = true;
-	}
-
-	UpdateFromTimeline();
-}
-
-void FDebuggerViewModel::OnUpdateNodeSelection(int32 InNodeId)
-{
-	if (InNodeId == INDEX_NONE)
-	{
-		return;
-	}
-
-	ActiveMotionMatchingState = nullptr;
-
-	// Find node in all motion matching states this frame
-	const int32 NodesNum = NodeIds.Num();
-	for (int i = 0; i < NodesNum; ++i)
-	{
-		if (NodeIds[i] == InNodeId)
-		{
-			ActiveMotionMatchingState = MotionMatchingStates[i];
-			break;
-		}
-	}
-
-	if (ActiveMotionMatchingState)
-	{
-		const UPoseSearchDatabase* CurrentDatabase = ActiveMotionMatchingState->GetCurrentDatabase();
-		if (FAsyncPoseSearchDatabasesManagement::RequestAsyncBuildIndex(CurrentDatabase, ERequestAsyncBuildFlag::ContinueRequest))
-		{
-			const FPoseSearchIndex& CurrentSearchIndex = CurrentDatabase->GetSearchIndex();
-			int32 CurrentPoseIdx = ActiveMotionMatchingState->GetCurrentDatabasePoseIndex();
-			if (const FPoseSearchIndexAsset* IndexAsset = CurrentSearchIndex.GetAssetForPoseSafe(CurrentPoseIdx))
-			{
-				Skeletons[Asset].Type = IndexAsset->Type;
-				Skeletons[Asset].bMirrored = IndexAsset->bMirrored;
-				Skeletons[Asset].SourceDatabase = CurrentDatabase;
-				Skeletons[Asset].AssetIdx = IndexAsset->SourceAssetIdx;
-				Skeletons[Asset].BlendParameters = IndexAsset->BlendParameters;
-			}
-		}
-	}
-
-	uint64 NewSearchableAssetId = ActiveMotionMatchingState ? ActiveMotionMatchingState->SearchableAssetId : 0;
-	if (NewSearchableAssetId != SearchableAssetId)
-	{
-		ClearSelectedSkeleton();
-		SearchableAssetId = NewSearchableAssetId;
-	}
-}
-
-void FDebuggerViewModel::UpdatePoseSearchContext(UPoseSearchMeshComponent::FUpdateContext& InOutContext, const FSkeleton& Skeleton) const
-{
-	const FInstancedStruct* DatabaseAssetStruct = Skeleton.GetAnimationAsset();
-	if (DatabaseAssetStruct)
-	{
-		const FPoseSearchDatabaseAnimationAssetBase* DatabaseAsset = DatabaseAssetStruct->GetPtr<FPoseSearchDatabaseAnimationAssetBase>();
-		if (DatabaseAsset)
-		{
-			InOutContext.Type = DatabaseAsset->GetSearchIndexType();
-			InOutContext.StartTime = Skeletons[SelectedPose].Time;
-			InOutContext.Time = Skeletons[SelectedPose].Time;
-			InOutContext.bMirrored = Skeletons[SelectedPose].bMirrored;
-			InOutContext.bLoop = DatabaseAsset->IsLooping();
-		}
-
-		if (const FPoseSearchDatabaseSequence* DatabaseSequence = DatabaseAssetStruct->GetPtr<FPoseSearchDatabaseSequence>())
-		{
-			InOutContext.SequenceBase = DatabaseSequence->Sequence;
-		}
-		else if (const FPoseSearchDatabaseAnimComposite* DatabaseAnimComposite = DatabaseAssetStruct->GetPtr<FPoseSearchDatabaseAnimComposite>())
-		{
-			InOutContext.SequenceBase = DatabaseAnimComposite->AnimComposite;
-		}
-		else if (const FPoseSearchDatabaseBlendSpace* DatabaseBlendSpace = DatabaseAssetStruct->GetPtr<FPoseSearchDatabaseBlendSpace>())
-		{
-			InOutContext.BlendSpace = DatabaseBlendSpace->BlendSpace;
-			InOutContext.BlendParameters = Skeletons[SelectedPose].BlendParameters;
-		}
-		else
-		{
-			checkNoEntry();
-		}
-	}
-}
-
-void FDebuggerViewModel::OnDraw(FSkeletonDrawParams& DrawParams)
-{
-	const UPoseSearchDatabase* CurrentDatabase = GetCurrentDatabase();
-	if (!CurrentDatabase)
-	{
-		return;
-	}
-
-	// Returns if it is to be drawn this frame
-	auto SetDrawSkeleton = [this](UPoseSearchMeshComponent* InComponent, bool bDraw)
-	{
-		if (InComponent && InComponent->RequiredBones.IsValid())
-		{
-			const bool bIsDrawingSkeleton = InComponent->ShouldDrawDebugSkeleton();
-			if (bIsDrawingSkeleton != bDraw)
-			{
-				InComponent->SetDrawDebugSkeleton(bDraw);
-			}
-			InComponent->MarkRenderStateDirty();
-		}
-	};
-	const bool bDrawActivePose = EnumHasAnyFlags(DrawParams.Flags, ESkeletonDrawFlags::ActivePose);
-	SetDrawSkeleton(Skeletons[ActivePose].Component.Get(), bDrawActivePose);
-	// If flag is set and we are currently in a valid drawing state
-	const bool bDrawSelectedPose = EnumHasAnyFlags(DrawParams.Flags, ESkeletonDrawFlags::SelectedPose) && bSelecting;
-	SetDrawSkeleton(Skeletons[SelectedPose].Component.Get(), bDrawSelectedPose);
-
-	FillCompactPoseAndComponentRefRotations();
-
-	UPoseSearchMeshComponent::FUpdateContext UpdateContext;
-
-	UpdateContext.MirrorDataTable = CurrentDatabase->Schema->MirrorDataTable;
-	UpdateContext.CompactPoseMirrorBones = &CompactPoseMirrorBones;
-	UpdateContext.ComponentSpaceRefRotations = &ComponentSpaceRefRotations;
-
-	if (bDrawSelectedPose)
-	{
-		UPoseSearchMeshComponent* Component = Skeletons[SelectedPose].Component.Get();
-		if (Component && Component->RequiredBones.IsValid())
-		{
-			UpdatePoseSearchContext(UpdateContext, Skeletons[SelectedPose]);
-
-			if (UpdateContext.Type != ESearchIndexAssetType::Invalid)
-			{
-				Component->UpdatePose(UpdateContext);
-			}
-		}
-	}
-
-	const bool bDrawAsset = EnumHasAnyFlags(DrawParams.Flags, ESkeletonDrawFlags::Asset);
-	if (bDrawAsset && AssetData.bActive)
-	{
-		UPoseSearchMeshComponent* Component = Skeletons[Asset].Component.Get();
-		if (Component && Component->RequiredBones.IsValid())
-		{
-			SetDrawSkeleton(Component, true);
-
-			UpdatePoseSearchContext(UpdateContext, Skeletons[SelectedPose]);
-
-			if (UpdateContext.Type != ESearchIndexAssetType::Invalid)
-			{
-				Component->UpdatePose(UpdateContext);
-			}
-		}
-	}
-}
-
-void FDebuggerViewModel::UpdateFromTimeline()
-{
-	NodeIds.Empty();
-	MotionMatchingStates.Empty();
-	SkeletalMeshComponentId = 0;
-
-	// Get provider and validate
-	const TraceServices::IAnalysisSession* Session = RewindDebugger.Get()->GetAnalysisSession();
-	TraceServices::FAnalysisSessionReadScope SessionReadScope(*Session);
-
-	const FTraceProvider* PoseSearchProvider = Session->ReadProvider<FTraceProvider>(FTraceProvider::ProviderName);
-	const IAnimationProvider* AnimationProvider = Session->ReadProvider<IAnimationProvider>("AnimationProvider");
-	const IGameplayProvider* GameplayProvider = Session->ReadProvider<IGameplayProvider>("GameplayProvider");
-	if (!(PoseSearchProvider && AnimationProvider && GameplayProvider))
-	{
-		return;
-	}
-	const double TraceTime = RewindDebugger.Get()->CurrentTraceTime();
-	TraceServices::FFrame Frame;
-	ReadFrameProvider(*Session).GetFrameFromTime(TraceFrameType_Game, TraceTime, Frame);
-	PoseSearchProvider->EnumerateMotionMatchingStateTimelines(AnimInstanceId, [&](const FTraceProvider::FMotionMatchingStateTimeline& InTimeline)
-	{
-		const FTraceMotionMatchingStateMessage* Message = nullptr;
-
-		InTimeline.EnumerateEvents(Frame.StartTime, Frame.EndTime, [&Message](double InStartTime, double InEndTime, const FTraceMotionMatchingStateMessage& InMessage)
-		{
-			Message = &InMessage;
-			return TraceServices::EEventEnumerate::Stop;
-		});
-		if (Message)
-		{
-			NodeIds.Add(Message->NodeId);
-			MotionMatchingStates.Add(Message);
-			SkeletalMeshComponentId = Message->SkeletalMeshComponentId;
-		}
-	});
-	/** No active motion matching state as no messages were read */
-	if (SkeletalMeshComponentId == 0)
-	{
-		return;
-	}
-	AnimationProvider->ReadSkeletalMeshPoseTimeline(SkeletalMeshComponentId, [&](const IAnimationProvider::SkeletalMeshPoseTimeline& TimelineData, bool bHasCurves)
-	{
-		TimelineData.EnumerateEvents(Frame.StartTime, Frame.EndTime, [&](double InStartTime, double InEndTime, uint32 InDepth, const FSkeletalMeshPoseMessage& PoseMessage) -> TraceServices::EEventEnumerate
-		{
-			// Update root transform
-			RootTransform = &PoseMessage.ComponentToWorld;
-			const FSkeletalMeshInfo* SkeletalMeshInfo = AnimationProvider->FindSkeletalMeshInfo(PoseMessage.MeshId);
-			const FObjectInfo* SkeletalMeshObjectInfo = GameplayProvider->FindObjectInfo(PoseMessage.MeshId);
-			if (!SkeletalMeshInfo || !SkeletalMeshObjectInfo)
-			{
-
-				return TraceServices::EEventEnumerate::Stop;
-			}
-			UPoseSearchMeshComponent* ActiveComponent = Skeletons[ActivePose].Component.Get();
-			UPoseSearchMeshComponent* SelectedComponent = Skeletons[SelectedPose].Component.Get();
-			UPoseSearchMeshComponent* AssetComponent = Skeletons[Asset].Component.Get();
-			USkeletalMesh* SkeletalMesh = TSoftObjectPtr<USkeletalMesh>(FSoftObjectPath(SkeletalMeshObjectInfo->PathName)).LoadSynchronous();
-			if (SkeletalMesh)
-			{
-				ActiveComponent->SetSkinnedAssetAndUpdate(SkeletalMesh, true);
-				SelectedComponent->SetSkinnedAssetAndUpdate(SkeletalMesh, true);
-				AssetComponent->SetSkinnedAssetAndUpdate(SkeletalMesh, true);
-			}
-			FTransform ComponentWorldTransform;
-			// Active skeleton is simply the traced bone transforms
-			AnimationProvider->GetSkeletalMeshComponentSpacePose(PoseMessage, *SkeletalMeshInfo, ComponentWorldTransform, ActiveComponent->GetEditableComponentSpaceTransforms());
-			ActiveComponent->Initialize(ComponentWorldTransform);
-			ActiveComponent->SetDebugDrawColor(FLinearColor::Green);
-			SelectedComponent->SetDebugDrawColor(FLinearColor::Blue);
-			SelectedComponent->Initialize(ComponentWorldTransform);
-			AssetComponent->SetDebugDrawColor(FLinearColor::Red);
-			AssetComponent->Initialize(ComponentWorldTransform);
-
-			return TraceServices::EEventEnumerate::Stop;
-		});
-	});
-}
-
-void FDebuggerViewModel::UpdateAsset()
-{
-	// @todo: expose those parameters
-	static float MAX_DISTANCE_RANGE = 200.f;
-	static float MAX_TIME_RANGE = 2.f;
-
-	const UPoseSearchDatabase* Database = GetCurrentDatabase();
-	if (!Database || !IsPlayingSelections())
-	{
-		return;
-	}
-
-	FSkeleton& AssetSkeleton = Skeletons[Asset];
-	UPoseSearchMeshComponent* Component = AssetSkeleton.Component.Get();
-
-	auto RestartAsset = [&]()
-	{
-		Component->ResetToStart();
-		AssetData.AccumulatedTime = 0.0;
-		AssetSkeleton.Time = AssetData.StartTime;
-	};
-
-	const FInstancedStruct* AnimationAssetStruct = AssetSkeleton.GetAnimationAsset();
-	if (!AnimationAssetStruct)
-	{
-		checkNoEntry();
-		return;
-	}
-
-	const FPoseSearchDatabaseAnimationAssetBase* DatabaseAsset = AnimationAssetStruct->GetPtr<FPoseSearchDatabaseAnimationAssetBase>();
-	if (!DatabaseAsset)
-	{
-		checkNoEntry();
-		return;
-
-	}
-
-	const UAnimationAsset* AnimAsset = DatabaseAsset->GetAnimationAsset();
-	const bool bAssetLooping = DatabaseAsset->IsLooping();
-
-	const float DT = static_cast<float>(FApp::GetDeltaTime()) * AssetPlayRate;
-	const float PlayLength = AnimAsset->GetPlayLength();
-	const bool bExceededDistanceHorizon = Component->LastRootMotionDelta.GetTranslation().Size() > MAX_DISTANCE_RANGE;
-	const bool bExceededTimeHorizon = (AssetSkeleton.Time - AssetData.StartTime) > MAX_TIME_RANGE;
-	const bool bExceededHorizon = bExceededDistanceHorizon && bExceededTimeHorizon;
-	if (bAssetLooping)
-	{
-		if (bExceededHorizon)
-		{
-			// Delay before restarting the asset to give the user some idea of where it would land
-			if (AssetData.AccumulatedTime > AssetData.StopDuration)
-			{
-				RestartAsset();
-			}
-			else
-			{
-				AssetData.AccumulatedTime += DT;
-			}
-			return;
-		}
-
-		AssetSkeleton.Time += DT;
-		AssetData.AccumulatedTime += DT;
-	}
-	else
-	{
-		// Used to cap the asset, but avoid modding when updating the pose
-		static constexpr float LengthOffset = 0.001f;
-		const bool bFinishedAsset = AssetSkeleton.Time >= PlayLength - LengthOffset;
-
-		// Asset player reached end of clip or reached distance horizon of trajectory vector
-		if (bFinishedAsset || bExceededHorizon)
-		{
-			// Delay before restarting the asset to give the user some idea of where it would land
-			if (AssetData.AccumulatedTime > AssetData.StopDuration)
-			{
-				RestartAsset();
-			}
-			else
-			{
-				AssetData.AccumulatedTime += DT;
-			}
-		}
-		else
-		{
-			// If we haven't finished, update the play time capped by the anim asset (not looping)
-			AssetSkeleton.Time += DT;
-		}
-	}
-}
-
-const USkinnedMeshComponent* FDebuggerViewModel::GetMeshComponent() const
-{
-	if (Skeletons.Num() > FDebuggerViewModel::Asset)
-	{
-		return Skeletons[FDebuggerViewModel::Asset].Component.Get();
-	}
-	return nullptr;
-}
-
-void FDebuggerViewModel::FillCompactPoseAndComponentRefRotations()
-{
-	bool bResetMirrorBonesAndRotations = true;
-	const UPoseSearchDatabase* Database = GetCurrentDatabase();
-	if (Database)
-	{
-		UMirrorDataTable* MirrorDataTable = Database->Schema->MirrorDataTable;
-		if (MirrorDataTable != nullptr && Skeletons[ActivePose].Component->RequiredBones.IsValid())
-		{
-			if (CompactPoseMirrorBones.Num() == 0 || ComponentSpaceRefRotations.Num() == 0)
-			{
-				MirrorDataTable->FillCompactPoseAndComponentRefRotations(
-					Skeletons[ActivePose].Component->RequiredBones,
-					CompactPoseMirrorBones,
-					ComponentSpaceRefRotations);
-			}
-			bResetMirrorBonesAndRotations = false;
-		}
-	}
-
-	if (bResetMirrorBonesAndRotations)
-	{
-		CompactPoseMirrorBones.Reset();
-		ComponentSpaceRefRotations.Reset();
-	}
-}
-
-void FDebuggerViewModel::PlaySelection(int32 PoseIdx, float Time)
-{
-	UPoseSearchMeshComponent* Component = Skeletons[Asset].Component.Get();
-	if (!Component)
-	{
-		return;
-	}
-
-	const UPoseSearchDatabase* Database = GetCurrentDatabase();
-	if (!FAsyncPoseSearchDatabasesManagement::RequestAsyncBuildIndex(Database, ERequestAsyncBuildFlag::ContinueRequest))
-	{
-		return;
-	}
-	
-	const FPoseSearchIndexAsset& IndexAsset = Database->GetSearchIndex().GetAssetForPose(PoseIdx);
-	Component->ResetToStart();
-	
-	Skeletons[Asset].Type = IndexAsset.Type;
-	Skeletons[Asset].AssetIdx = IndexAsset.SourceAssetIdx;
-	Skeletons[Asset].Time = Time;
-	Skeletons[Asset].bMirrored = IndexAsset.bMirrored;
-	Skeletons[Asset].BlendParameters = IndexAsset.BlendParameters;
-
-	AssetData.StartTime = Time;
-	AssetData.AccumulatedTime = 0.0f;
-	AssetData.bActive = true;
-}
-
-void FDebuggerViewModel::StopSelection()
-{
-	UPoseSearchMeshComponent* Component = Skeletons[Asset].Component.Get();
-	if (!Component)
-	{
-		return;
-	}
-
-	AssetData = {};
-	// @TODO: Make more functionality rely on checking if it should draw the asset
-	Component->SetDrawDebugSkeleton(false);
-}
-void FDebuggerViewModel::OnWorldCleanup(UWorld* InWorld, bool bSessionEnded, bool bCleanupResources)
-{
-	bSkeletonsInitialized = false;
-}
-
-const FInstancedStruct* FDebuggerViewModel::FSkeleton::GetAnimationAsset() const
-{
-	if (SourceDatabase.IsValid() && SourceDatabase->AnimationAssets.IsValidIndex(AssetIdx))
-	{
-		return &SourceDatabase->GetAnimationAssetStruct(AssetIdx);
-	}
-	return nullptr;
-}
-
 FDebugger* FDebugger::Debugger;
 void FDebugger::Initialize()
 {
diff --git a/Engine/Plugins/Experimental/Animation/PoseSearch/Source/Editor/Private/PoseSearchDebugger.h b/Engine/Plugins/Experimental/Animation/PoseSearch/Source/Editor/Private/PoseSearchDebugger.h
--- a/Engine/Plugins/Experimental/Animation/PoseSearch/Source/Editor/Private/PoseSearchDebugger.h
+++ b/Engine/Plugins/Experimental/Animation/PoseSearch/Source/Editor/Private/PoseSearchDebugger.h
@@ -2,622 +2,22 @@
 
 #pragma once
 
-#include "Components/PoseableMeshComponent.h"
 #include "CoreMinimal.h"
-#include "CustomBoneIndexArray.h"
+//#include "CustomBoneIndexArray.h"
 #include "IRewindDebuggerExtension.h"
-#include "IRewindDebuggerView.h"
+//#include "IRewindDebuggerView.h"
 #include "IRewindDebuggerViewCreator.h"
-#include "PoseSearch/PoseSearch.h"
-#include "Widgets/Views/SHeaderRow.h"
-#include "PoseSearchDebugger.generated.h"
-
-struct FPoseSearchDatabaseSequence;
-struct FPoseSearchDatabaseBlendSpace;
-class UAnimSequence;
-class UBlendSpace;
+//#include "PoseSearch/PoseSearch.h"
+//#include "PoseSearchMeshComponent.h"
 
 namespace TraceServices { class IAnalysisSession; }
-class UPoseSearchDatabase;
-class IUnrealInsightsModule;
-class ITableRow;
-class IDetailsView;
-class STableViewBase;
-class SVerticalBox;
-class SHorizontalBox;
-class SScrollBox;
-class SWidgetSwitcher;
-class SSearchBox;
-template <typename ItemType> class SListView;
 
 namespace UE::PoseSearch
 {
-	namespace DebuggerDatabaseColumns { struct IColumn; }
-	struct FDebugDrawParams;
-	struct FTraceMotionMatchingStateMessage;
-	class FDebuggerDatabaseRowData;
-} // namespace UE::PoseSearch
-
-
-UCLASS()
-class UPoseSearchMeshComponent : public UPoseableMeshComponent
-{
-	GENERATED_BODY()
-public:
-
-	struct FUpdateContext
-	{
-		ESearchIndexAssetType Type = ESearchIndexAssetType::Invalid;
-		const UAnimSequenceBase* SequenceBase = nullptr;
-		const UBlendSpace* BlendSpace = nullptr;
-		float StartTime = 0.0f;
-		float Time = 0.0f;
-		bool bLoop = false;
-		bool bMirrored = false;
-		FVector BlendParameters = FVector::Zero();
-		const UMirrorDataTable* MirrorDataTable = nullptr;
-		TCustomBoneIndexArray<FCompactPoseBoneIndex, FCompactPoseBoneIndex>* CompactPoseMirrorBones = nullptr;
-		TCustomBoneIndexArray<FQuat, FCompactPoseBoneIndex>* ComponentSpaceRefRotations = nullptr;
-	};
-
-	void Refresh();
-	void ResetToStart();
-	void UpdatePose(const FUpdateContext& UpdateContext);
-	void Initialize(const FTransform& InComponentToWorld);
-	FTransform StartingTransform = FTransform::Identity;
-	FTransform LastRootMotionDelta = FTransform::Identity;
-};
-
-/**
- * Used by the reflection UObject to encompass draw options for the query and database selections
- */
-USTRUCT()
-struct FPoseSearchDebuggerFeatureDrawOptions
-{
-	GENERATED_BODY()
-
-	UPROPERTY(EditAnywhere, Category="Draw Options")
-    bool bDisable = false;
-
-	UPROPERTY(EditAnywhere, Category = "Draw Options", Meta = (EditCondition = "!bDisable"))
-	bool bDrawBoneNames = false;
-
-	UPROPERTY(EditAnywhere, Category = "Draw Options", Meta = (EditCondition = "!bDisable"))
-	bool bDrawSampleLabels = false;
-};
-
-/**
- * Reflection UObject being observed in the details view panel of the debugger
- */
-UCLASS()
-class POSESEARCHEDITOR_API UPoseSearchDebuggerReflection : public UObject
-{
-	GENERATED_BODY()
-
-public:
-
-	UPROPERTY(VisibleAnywhere, Category="Motion Matching State", Meta=(DisplayName="Current Database"))
-	FString CurrentDatabaseName = "";
-
-	/** Time since last PoseSearch jump */
-	UPROPERTY(VisibleAnywhere, Category="Motion Matching State")
-	float ElapsedPoseJumpTime = 0.0f;
-
-	UPROPERTY(VisibleAnywhere, Category = "Motion Matching State")
-	FString AssetPlayerAssetName = "";
 
-	UPROPERTY(VisibleAnywhere, Category = "Motion Matching State")
-	float AssetPlayerTime = 0.0f;
-
-	UPROPERTY(VisibleAnywhere, Category = "Motion Matching State")
-	float LastDeltaTime = 0.0f;
-
-	UPROPERTY(VisibleAnywhere, Category = "Motion Matching State")
-	float SimLinearVelocity;
-
-	UPROPERTY(VisibleAnywhere, Category = "Motion Matching State")
-	float SimAngularVelocity;
-
-	UPROPERTY(VisibleAnywhere, Category = "Motion Matching State")
-	float AnimLinearVelocity;
-
-	UPROPERTY(VisibleAnywhere, Category = "Motion Matching State")
-	float AnimAngularVelocity;
-
-	UPROPERTY(EditAnywhere, Category="Draw Options", Meta=(DisplayName="Query"))
-	FPoseSearchDebuggerFeatureDrawOptions QueryDrawOptions;
-
-	UPROPERTY(EditAnywhere, Category="Draw Options", Meta=(DisplayName="Selected Pose"))
-	FPoseSearchDebuggerFeatureDrawOptions SelectedPoseDrawOptions;
-
-	UPROPERTY(EditAnywhere, Category = "Draw Options")
-	bool bDrawActiveSkeleton = false;
-	
-	UPROPERTY(EditAnywhere, Category = "Draw Options")
-	bool bDrawSelectedSkeleton = false;
-
-    UPROPERTY(VisibleAnywhere, Category="Pose Vectors")
-	TArray<float> QueryPoseVector;
-    	
-    UPROPERTY(VisibleAnywhere, Category="Pose Vectors")
-	TArray<float> ActivePoseVector;
-
-	UPROPERTY(VisibleAnywhere, Category="Pose Vectors")
-	TArray<float> SelectedPoseVector;
-
-	UPROPERTY(VisibleAnywhere, Category="Pose Vectors")
-	TArray<float> CostVector;
-};
-
-
-namespace UE::PoseSearch
-{
-
-/** Draw flags for the view's debug draw */
-enum class ESkeletonDrawFlags : uint32
-{
-	None			= 0,
-	ActivePose		= 1 << 0,
-	SelectedPose	= 1 << 1,
-	Asset			= 1 << 2,
-};
-ENUM_CLASS_FLAGS(ESkeletonDrawFlags)
-
-struct FSkeletonDrawParams
-{
-	ESkeletonDrawFlags Flags = ESkeletonDrawFlags::None;
-};
-
-/** Sets model selection data on row selection */
-DECLARE_DELEGATE_ThreeParams(FOnPoseSelectionChanged, const UPoseSearchDatabase*, int32, float)
+class SDebuggerView;
 class FDebuggerViewModel;
 
-/**
- * Database panel view widget of the PoseSearch debugger
- */
-
-class SDebuggerDatabaseView : public SCompoundWidget
-{
-	SLATE_BEGIN_ARGS(SDebuggerDatabaseView) {}
-		SLATE_ARGUMENT(TWeakPtr<class SDebuggerView>, Parent)
-		SLATE_EVENT(FOnPoseSelectionChanged, OnPoseSelectionChanged)
-	SLATE_END_ARGS()
-	
-	void Construct(const FArguments& InArgs);
-	void Update(const FTraceMotionMatchingStateMessage& State);
-
-	const TSharedPtr<SListView<TSharedRef<FDebuggerDatabaseRowData>>>& GetActiveRow() const { return ActiveView.ListView; }
-	const TSharedPtr<SListView<TSharedRef<FDebuggerDatabaseRowData>>>& GetContinuingPoseRow() const { return ContinuingPoseView.ListView; }
-	const TSharedPtr<SListView<TSharedRef<FDebuggerDatabaseRowData>>>& GetDatabaseRows() const { return FilteredDatabaseView.ListView; }
-	const TSharedRef<FDebuggerDatabaseRowData>& GetPoseIdxDatabaseRow(int32 PoseIdx) const;
-
-	/** Used by database rows to acquire column-specific information */
-	using FColumnMap = TMap<FName, TSharedRef<DebuggerDatabaseColumns::IColumn>>;
-
-private:
-	/** Adds a column to the existing list */
-	void AddColumn(TSharedRef<DebuggerDatabaseColumns::IColumn>&& Column);
-
-	/** Retrieves current column map, used as an attribute by rows */
-	const FColumnMap* GetColumnMap() const { return &Columns; }
-
-	/** Sorts the database by the current sort predicate, updating the view order */
-	void SortDatabaseRows();
-
-	void PopulateViewRows();
-
-	/** Acquires sort predicate for the given column */
-	EColumnSortMode::Type GetColumnSortMode(const FName ColumnId) const;
-
-	/** Gets active column width, used to align active and database view */
-	float GetColumnWidth(const FName ColumnId) const;
-	
-	/** Updates the active sort predicate, setting the sorting order of all other columns to none
-	 * (to be dependent on active column */
-	void OnColumnSortModeChanged(const EColumnSortPriority::Type SortPriority, const FName & ColumnId, const EColumnSortMode::Type InSortMode);
-
-	/** Aligns the active and database views */
-	void OnColumnWidthChanged(const float NewWidth, FName ColumnId) const;
-
-	/** Called when the text in the filter box is modified to update the filtering */
-	void OnFilterTextChanged(const FText& SearchText);
-
-	/** Called when the HideInvalidPoses Checkbox is Toggled */
-	void OnHideInvalidPosesCheckboxChanged(ECheckBoxState State);
-
-	/** Row selection to update model view */
-	void OnDatabaseRowSelectionChanged(TSharedPtr<FDebuggerDatabaseRowData> Row, ESelectInfo::Type SelectInfo);
-
-	/** Generates a database row widget for the given data */
-	TSharedRef<ITableRow> HandleGenerateDatabaseRow(TSharedRef<FDebuggerDatabaseRowData> Item, const TSharedRef<STableViewBase>& OwnerTable) const;
-	
-	/** Generates the active row widget for the given data */
-	TSharedRef<ITableRow> HandleGenerateActiveRow(TSharedRef<FDebuggerDatabaseRowData> Item, const TSharedRef<STableViewBase>& OwnerTable) const;
-
-	/** Generates the continuing pose row widget for the given data */
-	TSharedRef<ITableRow> HandleGenerateContinuingPoseRow(TSharedRef<FDebuggerDatabaseRowData> Item, const TSharedRef<STableViewBase>& OwnerTable) const;
-
-	TWeakPtr<SDebuggerView> ParentDebuggerViewPtr;
-
-	FOnPoseSelectionChanged OnPoseSelectionChanged;
-
-	/** Current column to sort by */
-	FName SortColumn = "";
-
-	/** Current sorting mode */
-	EColumnSortMode::Type SortMode = EColumnSortMode::Ascending;
-	
-	/** Column data container, used to emplace defined column structures of various types */
-    FColumnMap Columns;
-
-	TArray<FText> OldLabels;
-
-	struct FTable
-	{
-		/** Header row*/
-		TSharedPtr<SHeaderRow> HeaderRow;
-
-		/** Widget for displaying the list of row objects */
-		TSharedPtr<SListView<TSharedRef<FDebuggerDatabaseRowData>>> ListView;
-
-		// @TODO: Explore options for active row other than displaying array of 1 element
-		/** List of row objects */
-		TArray<TSharedRef<FDebuggerDatabaseRowData>> Rows;
-
-		/** Background style for the list view */
-		FTableRowStyle RowStyle;
-
-		/** Row color */
-		FSlateBrush RowBrush;
-
-		/** Scroll bar for the data table */
-		TSharedPtr<SScrollBar> ScrollBar;
-	};
-
-	/** Active row at the top of the view */
-	FTable ActiveView;
-
-	/** Continuing pose row below Active row */
-	FTable ContinuingPoseView;
-
-	/** All database poses */
-	TArray<TSharedRef<FDebuggerDatabaseRowData>> UnfilteredDatabaseRows;
-
-	/** Database listing for filtered poses */
-	FTable FilteredDatabaseView;
-
-	/** Search box widget */
-	TSharedPtr<SSearchBox> FilterBox;
-
-	/** Text used to filter DatabaseView */
-	FText FilterText;
-
-	/* Bool used to hide invalid poses from display */
-	bool bHideInvalidPoses = false;
-};
-
-/**
- * Details panel view widget of the PoseSearch debugger
- */
-class SDebuggerDetailsView : public SCompoundWidget
-{
-	SLATE_BEGIN_ARGS(SDebuggerDetailsView) {}
-		SLATE_ARGUMENT(TWeakPtr<class SDebuggerView>, Parent)
-	SLATE_END_ARGS()
-
-	SDebuggerDetailsView() = default;
-	virtual ~SDebuggerDetailsView() override;
-	
-	void Construct(const FArguments& InArgs);
-	void Update(const FTraceMotionMatchingStateMessage& State) const;
-
-	/** Get a const version of our reflection object */
-	const TObjectPtr<UPoseSearchDebuggerReflection> GetReflection() const { return Reflection; }
-	
-private:
-	/** Update our details view object with new state information */
-	void UpdateReflection(const FTraceMotionMatchingStateMessage& State) const;
-	
-	TWeakPtr<SDebuggerView> ParentDebuggerViewPtr;
-
-	/** Details widget constructed for the MM node */
-	TSharedPtr<IDetailsView> Details;
-
-	/** Last updated reflection data relative to MM state */
-	TObjectPtr<UPoseSearchDebuggerReflection> Reflection = nullptr;
-};
-/** Callback to relay closing of the view to destroy the debugger instance */
-DECLARE_DELEGATE_OneParam(FOnViewClosed, uint64 AnimInstanceId);
-
-/**
- * Entire view of the PoseSearch debugger, containing all sub-widgets
- */
-class SDebuggerView : public IRewindDebuggerView
-{
-public:
-	SLATE_BEGIN_ARGS(SDebuggerView){}
-		SLATE_ATTRIBUTE(TSharedPtr<FDebuggerViewModel>, ViewModel)
-		SLATE_EVENT(FOnViewClosed, OnViewClosed)
-	SLATE_END_ARGS()
-
-	SDebuggerView() = default;
-    virtual ~SDebuggerView() override;
-
-	void Construct(const FArguments& InArgs, uint64 InAnimInstanceId);
-	virtual void SetTimeMarker(double InTimeMarker) override;
-	virtual void Tick(const FGeometry& AllottedGeometry, const double InCurrentTime, const float InDeltaTime) override;
-	virtual FName GetName() const override;
-	virtual uint64 GetObjectId() const override;
-
-	TSharedPtr<FDebuggerViewModel> GetViewModel() const { return ViewModel.Get(); }
-	TArray<TSharedRef<FDebuggerDatabaseRowData>> GetSelectedDatabaseRows() const;
-
-private:
-	/** Called each frame to draw features of the query vector & database selections */
-	void DrawFeatures(const UWorld& DebuggerWorld, const FTraceMotionMatchingStateMessage& State, const FTransform& Transform, const USkinnedMeshComponent* Mesh) const;
-	
-	/** Check if a node selection was made, true if a node is selected */
-	bool UpdateNodeSelection();
-
-	/** Update the database and details views */
-	void UpdateViews() const;
-
-	void DrawVisualization() const;
-	
-	/** Returns an int32 appropriate to the index of our widget selector */
-	int32 SelectView() const;
-
-	/** Callback when a button in the selection view is clicked */
-	FReply OnUpdateNodeSelection(int32 InSelectedNodeId);
-
-	void OnPoseSelectionChanged(const UPoseSearchDatabase* Database, int32 PoseIdx, float Time);
-
-	/** Button interaction to toggle play / stop of the asset */
-	FReply TogglePlaySelectedAssets() const;
-
-	/** Generates the message view relaying that there is no data */
-	TSharedRef<SWidget> GenerateNoDataMessageView();
-
-	/** Generates the return button to go back to the selection mode */
-	TSharedRef<SHorizontalBox> GenerateReturnButtonView();
-
-	/** Generates the entire node debugger widget, including database and details view */
-	TSharedRef<SWidget> GenerateNodeDebuggerView();
-
-	/** Pointer to the debugger instance / model for this view */
-	TAttribute<TSharedPtr<FDebuggerViewModel>> ViewModel;
-	
-	/** Destroy the debugger instanced when closed */
-	FOnViewClosed OnViewClosed;
-
-	/** Active node being debugged */
-	int32 SelectedNodeId = -1;
-
-	/** Database view of the motion matching node */
-	TSharedPtr<SDebuggerDatabaseView> DatabaseView;
-	
-	/** Details panel for introspecting the motion matching node */
-	TSharedPtr<SDebuggerDetailsView> DetailsView;
-	
-	/** Node debugger view hosts the above two views */
-	TSharedPtr<SSplitter> NodeDebuggerView;
-
-	/** Selection view before node is selected */
-	TSharedPtr<SVerticalBox> SelectionView;
-	
-	/** Gray box occluding the debugger view when simulating */
-	TSharedPtr<SVerticalBox> SimulatingView;
-
-	/** Used to switch between views in the switcher, int32 maps to index in the SWidgetSwitcher */
-	enum ESwitcherViewType : int32
-	{
-		Selection = 0,
-		Debugger = 1,
-		StoppedMsg = 2,
-		RecordingMsg = 3,
-		NoDataMsg = 4
-	} SwitcherViewType = StoppedMsg;
-	
-	/** Contains all the above, switches between them depending on context */
-	TSharedPtr<SWidgetSwitcher> Switcher;
-
-	/** Contains the switcher, the entire debugger view */
-	TSharedPtr<SVerticalBox> DebuggerView;
-
-	/** AnimInstance this view was created for */
-	uint64 AnimInstanceId = 0;
-
-	/** Current position of the time marker */
-	double TimeMarker = -1.0;
-
-	/** Previous position of the time marker */
-	double PreviousTimeMarker = -1.0;
-
-	/** Tracks if the current time has been updated yet (delayed) */
-	bool bUpdated = false;
-
-	/** Tracks number of consecutive frames, once it reaches threshold it will update the view */
-	int32 CurrentConsecutiveFrames = 0;
-
-	/** Once the frame count has reached this value, an update will trigger for the view */
-	static constexpr int32 ConsecutiveFramesUpdateThreshold = 10;
-};
-
-class FDebuggerViewModel : public TSharedFromThis<FDebuggerViewModel>
-{
-public:
-	explicit FDebuggerViewModel(uint64 InAnimInstanceId);
-	virtual ~FDebuggerViewModel();
-
-	// Used for view callbacks
-    const FTraceMotionMatchingStateMessage* GetMotionMatchingState() const;
-	const UPoseSearchDatabase* GetCurrentDatabase() const;
-	const UPoseSearchSearchableAsset* GetSearchableAsset() const;
-	const TArray<int32>* GetNodeIds() const;
-	int32 GetNodesNum() const;
-	const FTransform* GetRootTransform() const;
-
-	/** Checks if Update must be called */
-	bool HasSearchableAssetChanged() const;
-
-	/** Update motion matching states for frame */
-	void OnUpdate();
-	
-	/** Updates active motion matching state based on node selection */
-	void OnUpdateNodeSelection(int32 InNodeId);
-
-	/** Updates internal Skeletal Mesh Component depending on input */
-	void OnDraw(FSkeletonDrawParams& DrawParams);
-	
-	/** Sets the selected pose skeleton*/
-	void ShowSelectedSkeleton(const UPoseSearchDatabase* Database, int32 DbPoseIdx, float Time);
-	
-	/** Clears the selected pose skeleton */
-	void ClearSelectedSkeleton();
-	
-	/** Plays the selected row upon button press */
-	void PlaySelection(int32 PoseIdx, float Time);
-
-	/** Stops the playing selection upon button press */
-	void StopSelection();
-	
-	/** If there is an asset selection playing */
-	bool IsPlayingSelections() const { return AssetData.bActive; }
-	
-	/** Play Rate of the asset selection player */
-	void ChangePlayRate(float PlayRate) { AssetPlayRate = PlayRate; }
-	
-	/** Current play rate of the asset selection player */
-	float GetPlayRate() const { return AssetPlayRate; }
-	
-	void SetVerbose(bool bVerbose) { bIsVerbose = bVerbose; }
-
-	bool IsVerbose() const { return bIsVerbose; }
-
-	/** Callback to reset debug skeletons for the active world */
-	void OnWorldCleanup(UWorld* InWorld, bool bSessionEnded, bool bCleanupResources);
-
-	/** Updates the current playing asset */
-	void UpdateAsset();
-
-	const USkinnedMeshComponent* GetMeshComponent() const;
-
-private:
-	/** Debug visualization skeleton actor */
-	struct FSkeleton
-	{
-		/** Actor object for the skeleton */
-		TWeakObjectPtr<AActor> Actor = nullptr;
-
-		/** Derived skeletal mesh for setting the skeleton in the scene */
-		TWeakObjectPtr<UPoseSearchMeshComponent> Component = nullptr;
-
-		/** Type of the asset being played */
-		ESearchIndexAssetType Type = ESearchIndexAssetType::Invalid;
-
-		/** Source database for this skeleton  */
-		TWeakObjectPtr<const UPoseSearchDatabase> SourceDatabase = nullptr;
-
-		/** Source asset for this skeleton */
-		int32 AssetIdx = 0;
-
-		/** Time in the sequence this skeleton is accessing */
-		float Time = 0.0f;
-
-		/** If this asset should be mirrored */
-		bool bMirrored = false;
-
-		/** Blend Parameters if asset is a BlendSpace */
-		FVector BlendParameters = FVector::Zero();
-
-		const FInstancedStruct* GetAnimationAsset() const;
-	};
-
-	/** Update the list of states for this frame */
-	void UpdateFromTimeline();
-
-	/** Populates arrays used for mirroring the animation pose */
-	void FillCompactPoseAndComponentRefRotations();
-
-	void UpdatePoseSearchContext(UPoseSearchMeshComponent::FUpdateContext& InOutContext, const FSkeleton& Skeleton) const;
-	
-	/** List of all Node IDs associated with motion matching states */
-	TArray<int32> NodeIds;
-	
-	/** List of all updated motion matching states per node */
-	TArray<const FTraceMotionMatchingStateMessage*> MotionMatchingStates;
-	
-	/** Currently active motion matching state based on node selection in the view */
-	const FTraceMotionMatchingStateMessage* ActiveMotionMatchingState = nullptr;
-
-	/** Active motion matching state's searchable asset */
-	uint64 SearchableAssetId = 0;
-
-	/** Current Skeletal Mesh Component Id for the AnimInstance */
-	uint64 SkeletalMeshComponentId = 0;
-
-	/** Currently active root transform on the skeletal mesh */
-	const FTransform* RootTransform = nullptr;
-
-	/** Pointer to the active rewind debugger in the scene */
-	TAttribute<const IRewindDebugger*> RewindDebugger;
-
-	/** Anim Instance associated with this debugger instance */
-	uint64 AnimInstanceId = 0;
-
-	/** Compact pose format of Mirror Bone Map */
-	TCustomBoneIndexArray<FCompactPoseBoneIndex, FCompactPoseBoneIndex> CompactPoseMirrorBones;
-
-	/** Pre-calculated component space rotations of reference pose */
-	TCustomBoneIndexArray<FQuat, FCompactPoseBoneIndex> ComponentSpaceRefRotations;
-
-
-	/** Index for each type of skeleton we store for debug visualization */
-	enum ESkeletonIndex
-	{
-		ActivePose = 0,
-		SelectedPose,
-		Asset,
-
-		Num
-	};
-
-	/** Skeleton container for each type */
-	TArray<FSkeleton, TFixedAllocator<ESkeletonIndex::Num>> Skeletons;
-
-	/** Whether the skeletons have been initialized for this world */
-	bool bSkeletonsInitialized = false;
-	
-	/** If we currently have a selection active in the view */
-	bool bSelecting = false;
-	
-	/** Data of the active playing asset */
-	struct FAsset
-	{
-		/** How long to stop upon reaching target */
-		static constexpr float StopDuration = 2.0f;
-
-		/** Time since the start of play */
-		float AccumulatedTime = 0.0f;
-
-		/** Start time of the asset */
-		float StartTime = 0.0f;
-
-		/** If we are currently playing a asset */
-		bool bActive = false;
-
-	} AssetData;
-	
-	/** Current play rate of the asset selection player */
-	float AssetPlayRate = 1.0f;
-
-	bool bIsVerbose = false;
-
-	/** Limits some public API */
-	friend class FDebugger;
-};
-
-
 /**
  * PoseSearch debugger, containing the data to be acquired and relayed to the view
  */
diff --git a/Engine/Plugins/Experimental/Animation/PoseSearch/Source/Editor/Private/PoseSearchDebuggerDatabaseColumns.h b/Engine/Plugins/Experimental/Animation/PoseSearch/Source/Editor/Private/PoseSearchDebuggerDatabaseColumns.h
new file mode 100644
--- /dev/null
+++ b/Engine/Plugins/Experimental/Animation/PoseSearch/Source/Editor/Private/PoseSearchDebuggerDatabaseColumns.h
@@ -0,0 +1,418 @@
+// Copyright Epic Games, Inc. All Rights Reserved.
+
+#pragma once
+
+#include "Editor.h"
+#include "PoseSearchDebuggerDatabaseRowData.h"
+#include "Preferences/PersonaOptions.h"
+#include "Styling/AppStyle.h"
+#include "Subsystems/AssetEditorSubsystem.h"
+#include "Widgets/Input/SHyperlink.h"
+#include "Widgets/Text/STextBlock.h"
+
+#define LOCTEXT_NAMESPACE "PoseSearchDebugger"
+
+namespace UE::PoseSearch::DebuggerDatabaseColumns
+{
+	
+/** Column struct to represent each column in the debugger database */
+struct IColumn : TSharedFromThis<IColumn>
+{
+	explicit IColumn(int32 InSortIndex, bool InEnabled = true)
+		: SortIndex(InSortIndex)
+		, bEnabled(InEnabled)
+	{
+		ColumnId = FName(FString::Printf(TEXT("Column %d"), SortIndex));
+	}
+
+	virtual ~IColumn() = default;
+
+	FName ColumnId;
+
+	/** Sorted left to right based on this index */
+	int32 SortIndex = 0;
+	/** Current width, starts at 1 to be evenly spaced between all columns */
+	float Width = 1.0f;
+	/** Disabled selectively with view options */
+	bool bEnabled = false;
+
+	virtual FText GetLabel() const = 0;
+
+	using FRowDataRef = TSharedRef<FDebuggerDatabaseRowData>;
+	using FSortPredicate = TFunction<bool(const FRowDataRef&, const FRowDataRef&)>;
+
+	virtual FSortPredicate GetSortPredicate() const = 0;
+
+	virtual TSharedRef<SWidget> GenerateWidget(const FRowDataRef& RowData) const = 0;
+};
+
+/** Column struct to represent each column in the debugger database */
+struct ITextColumn : IColumn
+{
+	explicit ITextColumn(int32 InSortIndex, bool InEnabled = true)
+		: IColumn(InSortIndex, InEnabled)
+	{
+	}
+
+	virtual ~ITextColumn() = default;
+		
+	virtual TSharedRef<SWidget> GenerateWidget(const FRowDataRef& RowData) const override
+	{
+		static FSlateFontInfo RowFont = FAppStyle::Get().GetFontStyle("DetailsView.CategoryTextStyle");
+			
+		return SNew(STextBlock)
+            .Font(RowFont)
+			.Text_Lambda([this, RowData]() -> FText { return GetRowText(RowData); })
+            .Justification(ETextJustify::Center)
+			.ColorAndOpacity_Lambda([this, RowData] { return GetColorAndOpacity(RowData); });
+	}
+
+protected:
+	virtual FText GetRowText(const FRowDataRef& Row) const = 0;
+	virtual FSlateColor GetColorAndOpacity(const FRowDataRef& Row) const
+	{
+		return FSlateColor(FLinearColor::White);
+	}
+};
+
+struct FPoseIdx : ITextColumn
+{
+	using ITextColumn::ITextColumn;
+
+	virtual FText GetLabel() const override { return LOCTEXT("ColumnLabelPoseIndex", "Index"); }
+
+	virtual FSortPredicate GetSortPredicate() const override
+	{
+		return [](const FRowDataRef& Row0, const FRowDataRef& Row1) -> bool { return Row0->PoseIdx < Row1->PoseIdx; };
+	}
+
+	virtual FText GetRowText(const FRowDataRef& Row) const override
+	{
+		return FText::AsNumber(Row->PoseIdx, &FNumberFormattingOptions::DefaultNoGrouping());
+	}
+};
+
+struct FDatabaseName : IColumn
+{
+	using IColumn::IColumn;
+
+	virtual FText GetLabel() const override { return LOCTEXT("ColumnLabelDatabaseName", "Database"); }
+
+	virtual FSortPredicate GetSortPredicate() const override
+	{
+		return [](const FRowDataRef& Row0, const FRowDataRef& Row1) -> bool { return Row0->DatabaseName < Row1->DatabaseName; };
+	}
+
+	virtual TSharedRef<SWidget> GenerateWidget(const FRowDataRef& RowData) const override
+	{
+		return SNew(SHyperlink)
+			.Text_Lambda([RowData]() -> FText { return FText::FromString(RowData->DatabaseName); })
+			.TextStyle(&FCoreStyle::Get().GetWidgetStyle<FTextBlockStyle>("SmallText"))
+			.ToolTipText_Lambda([RowData]() -> FText
+				{
+					return FText::Format(
+						LOCTEXT("DatabaseHyperlinkTooltipFormat", "Open database '{0}'"),
+						FText::FromString(RowData->DatabasePath));
+				})
+			.OnNavigate_Lambda([RowData]()
+				{
+					GEditor->GetEditorSubsystem<UAssetEditorSubsystem>()->OpenEditorForAsset(RowData->DatabasePath);
+				});
+	}
+};
+
+struct FAssetName : IColumn
+{
+	using IColumn::IColumn;
+
+	virtual FText GetLabel() const override { return LOCTEXT("ColumnLabelAssetName", "Asset"); }
+		
+	virtual FSortPredicate GetSortPredicate() const override
+	{
+		return [](const FRowDataRef& Row0, const FRowDataRef& Row1) -> bool { return Row0->AssetName < Row1->AssetName; };
+	}
+
+	virtual TSharedRef<SWidget> GenerateWidget(const FRowDataRef& RowData) const override
+	{
+		return SNew(SHyperlink)
+			.Text_Lambda([RowData]() -> FText { return FText::FromString(RowData->AssetName); })
+			.TextStyle(&FCoreStyle::Get().GetWidgetStyle<FTextBlockStyle>("SmallText"))
+			.ToolTipText_Lambda([RowData]() -> FText 
+				{ 
+					return FText::Format(
+						LOCTEXT("AssetHyperlinkTooltipFormat", "Open asset '{0}'"), 
+						FText::FromString(RowData->AssetPath)); 
+				})
+			.OnNavigate_Lambda([RowData]()
+				{
+					GEditor->GetEditorSubsystem<UAssetEditorSubsystem>()->OpenEditorForAsset(RowData->AssetPath);
+				});
+	}
+};
+
+struct FFrame : ITextColumn
+{
+	using ITextColumn::ITextColumn;
+
+	virtual FText GetLabel() const override { return LOCTEXT("ColumnLabelFrame", "Frame"); }
+		
+	virtual FSortPredicate GetSortPredicate() const override
+	{
+		return [](const FRowDataRef& Row0, const FRowDataRef& Row1) -> bool { return Row0->AssetTime < Row1->AssetTime; };
+	}
+
+	virtual FText GetRowText(const FRowDataRef& Row) const override
+	{
+		FNumberFormattingOptions TimeFormattingOptions = FNumberFormattingOptions()
+			.SetUseGrouping(false)
+			.SetMaximumFractionalDigits(2);
+
+		FNumberFormattingOptions PercentageFormattingOptions = FNumberFormattingOptions()
+			.SetMaximumFractionalDigits(2);
+
+		if (Row->AssetType == ESearchIndexAssetType::Sequence ||
+			Row->AssetType == ESearchIndexAssetType::AnimComposite)
+		{
+			if (GetDefault<UPersonaOptions>()->bTimelineDisplayPercentage)
+			{
+				return FText::Format(
+					FText::FromString("{0} ({1}) ({2})"),
+					FText::AsNumber(Row->AnimFrame, &FNumberFormattingOptions::DefaultNoGrouping()),
+					FText::AsNumber(Row->AssetTime, &TimeFormattingOptions),
+					FText::AsPercent(Row->AnimPercentage, &PercentageFormattingOptions));
+			}
+			else
+			{
+				return FText::Format(
+					FText::FromString("{0} ({1})"),
+					FText::AsNumber(Row->AnimFrame, &FNumberFormattingOptions::DefaultNoGrouping()),
+					FText::AsNumber(Row->AssetTime, &TimeFormattingOptions));
+			}
+		}
+		else if (Row->AssetType == ESearchIndexAssetType::BlendSpace)
+		{
+			// There is no frame index associated with a blendspace
+			return FText::Format(
+				FText::FromString("({0})"),
+				FText::AsNumber(Row->AssetTime, &TimeFormattingOptions));
+		}
+		else
+		{
+			return FText::FromString(TEXT("-"));
+		}
+	}
+};
+
+struct FCost : ITextColumn
+{
+	using ITextColumn::ITextColumn;
+
+	virtual FText GetLabel() const override { return LOCTEXT("ColumnLabelCost", "Cost"); }
+		
+	virtual FSortPredicate GetSortPredicate() const override
+	{
+		return [](const FRowDataRef& Row0, const FRowDataRef& Row1) -> bool { return Row0->PoseCost < Row1->PoseCost; };
+	}
+		
+	virtual FText GetRowText(const FRowDataRef& Row) const override
+	{
+		if (Row->PoseCost.IsValid())
+		{
+			return FText::AsNumber(Row->PoseCost.GetTotalCost());
+		}
+		else
+		{
+			return FText::FromString(TEXT("-"));
+		}
+    }
+
+	virtual FSlateColor GetColorAndOpacity(const FRowDataRef& Row) const override
+	{
+		return Row->CostColor;
+	}
+};
+
+struct FChannelBreakdownCostColumn : ITextColumn
+{
+	FChannelBreakdownCostColumn(int32 SortIndex, int32 InBreakdownCostIndex, const FText& InLabel)
+		: ITextColumn(SortIndex)
+		, Label(InLabel)
+		, BreakdownCostIndex(InBreakdownCostIndex)
+	{
+	}
+
+	virtual FText GetLabel() const override { return Label; }
+
+	virtual FSortPredicate GetSortPredicate() const override
+	{
+		return [this](const FRowDataRef& Row0, const FRowDataRef& Row1) -> bool
+		{
+			return Row0->CostBreakdowns[BreakdownCostIndex] < Row1->CostBreakdowns[BreakdownCostIndex];
+		};
+	}
+
+	virtual FText GetRowText(const FRowDataRef& Row) const override
+	{
+		const float LabelCost = Row->CostBreakdowns[BreakdownCostIndex];
+		if (LabelCost != UE_MAX_FLT)
+		{
+			return FText::AsNumber(LabelCost);
+		}
+		return FText::FromString(TEXT("-"));
+	}
+
+	virtual FSlateColor GetColorAndOpacity(const FRowDataRef& Row) const override
+	{
+		return Row->CostBreakdownsColors[BreakdownCostIndex];
+	}
+
+	FText Label;
+	int32 BreakdownCostIndex = INDEX_NONE;
+};
+
+#if WITH_EDITORONLY_DATA
+struct FCostModifier : ITextColumn
+{
+	using ITextColumn::ITextColumn;
+
+	virtual FText GetLabel() const override { return LOCTEXT("ColumnLabelCostModifier", "Bias"); }
+
+	virtual FSortPredicate GetSortPredicate() const override
+	{
+		return [](const FRowDataRef& Row0, const FRowDataRef& Row1) -> bool { return Row0->PoseCost.GetCostAddend() < Row1->PoseCost.GetCostAddend(); };
+	}
+
+	virtual FText GetRowText(const FRowDataRef& Row) const override
+	{
+		return FText::AsNumber(Row->PoseCost.GetCostAddend());
+	}
+};
+#endif // WITH_EDITORONLY_DATA
+
+struct FMirrored : ITextColumn
+{
+	using ITextColumn::ITextColumn;
+
+	virtual FText GetLabel() const override { return LOCTEXT("ColumnLabelMirrored", "Mirror"); }
+
+	virtual FSortPredicate GetSortPredicate() const override
+	{
+		return [](const FRowDataRef& Row0, const FRowDataRef& Row1) -> bool { return Row0->bMirrored < Row1->bMirrored; };
+	}
+
+	virtual FText GetRowText(const FRowDataRef& Row) const override
+	{
+		return FText::Format(LOCTEXT("Mirrored", "{0}"), { Row->bMirrored } );
+	}
+};
+
+struct FLooping : ITextColumn
+{
+	using ITextColumn::ITextColumn;
+
+	virtual FText GetLabel() const override { return LOCTEXT("ColumnLabelLooping", "Loop"); }
+
+	virtual FSortPredicate GetSortPredicate() const override
+	{
+		return [](const FRowDataRef& Row0, const FRowDataRef& Row1) -> bool { return Row0->bLooping < Row1->bLooping; };
+	}
+
+	virtual FText GetRowText(const FRowDataRef& Row) const override
+	{
+		return FText::Format(LOCTEXT("Looping", "{0}"), { Row->bLooping });
+	}
+};
+
+struct FBlendParameters : ITextColumn
+{
+	using ITextColumn::ITextColumn;
+
+	virtual FText GetLabel() const override { return LOCTEXT("ColumnLabelBlendParams", "Blend Params"); }
+
+	virtual FSortPredicate GetSortPredicate() const override
+	{
+		return [](const FRowDataRef& Row0, const FRowDataRef& Row1) -> bool
+		{
+			return (
+				Row0->BlendParameters[0] < Row1->BlendParameters[0] ||
+				Row0->BlendParameters[1] < Row1->BlendParameters[1]);
+		};
+	}
+
+	virtual FText GetRowText(const FRowDataRef& Row) const override
+	{
+		if (Row->AssetType == ESearchIndexAssetType::BlendSpace)
+		{
+			return FText::Format(
+				LOCTEXT("Blend Parameters", "({0}, {1})"),
+				FText::AsNumber(Row->BlendParameters[0]),
+				FText::AsNumber(Row->BlendParameters[1]));
+		}
+		else
+		{
+			return FText::FromString(TEXT("-"));
+		}
+	}
+};
+
+struct FPoseCandidateFlags : ITextColumn
+{
+	using ITextColumn::ITextColumn;
+
+	virtual FText GetLabel() const override { return LOCTEXT("ColumnLabelPoseCandidateFlags", "Flags"); }
+
+	virtual FSortPredicate GetSortPredicate() const override
+	{
+		return [](const FRowDataRef& Row0, const FRowDataRef& Row1) -> bool { return Row0->PoseCandidateFlags < Row1->PoseCandidateFlags; };
+	}
+
+	virtual FText GetRowText(const FRowDataRef& Row) const override
+	{
+		if (EnumHasAnyFlags(Row->PoseCandidateFlags, EPoseCandidateFlags::AnyDiscardedMask))
+		{
+			auto AddDelimiter = [](FTextBuilder& TextBuilder, bool& bNeedDelimiter)
+			{
+				if (bNeedDelimiter)
+				{
+					TextBuilder.AppendLine(LOCTEXT("DiscardedBy_Delimiter", " | "));
+				}
+				bNeedDelimiter = true;
+			};
+
+			bool bNeedDelimiter = false;
+
+			FTextBuilder TextBuilder;
+			if (EnumHasAnyFlags(Row->PoseCandidateFlags, EPoseCandidateFlags::DiscardedBy_PoseJumpThresholdTime))
+			{
+				AddDelimiter(TextBuilder, bNeedDelimiter);
+				TextBuilder.AppendLine(LOCTEXT("DiscardedBy_PoseJumpThresholdTime", "PoseJumpThresholdTime"));
+			}
+				
+			if (EnumHasAnyFlags(Row->PoseCandidateFlags, EPoseCandidateFlags::DiscardedBy_PoseReselectHistory))
+			{
+				AddDelimiter(TextBuilder, bNeedDelimiter);
+				TextBuilder.AppendLine(LOCTEXT("DiscardedBy_PoseReselectHistory", "PoseReselectHistory"));
+			}
+
+			if (EnumHasAnyFlags(Row->PoseCandidateFlags, EPoseCandidateFlags::DiscardedBy_BlockTransition))
+			{
+				AddDelimiter(TextBuilder, bNeedDelimiter);
+				TextBuilder.AppendLine(LOCTEXT("DiscardedBy_BlockTransition", "BlockTransition"));
+			}
+
+			if (EnumHasAnyFlags(Row->PoseCandidateFlags, EPoseCandidateFlags::DiscardedBy_PoseFilter))
+			{
+				AddDelimiter(TextBuilder, bNeedDelimiter);
+				TextBuilder.AppendLine(LOCTEXT("DiscardedBy_PoseFilter", "PoseFilter"));
+			}
+
+			return TextBuilder.ToText();
+		}
+
+		return FText::GetEmpty();
+	}
+};
+
+} // namespace UE::PoseSearch::DebuggerDatabaseColumns
+
+#undef LOCTEXT_NAMESPACE
diff --git a/Engine/Plugins/Experimental/Animation/PoseSearch/Source/Editor/Private/PoseSearchDebuggerDatabaseRow.h b/Engine/Plugins/Experimental/Animation/PoseSearch/Source/Editor/Private/PoseSearchDebuggerDatabaseRow.h
new file mode 100644
--- /dev/null
+++ b/Engine/Plugins/Experimental/Animation/PoseSearch/Source/Editor/Private/PoseSearchDebuggerDatabaseRow.h
@@ -0,0 +1,82 @@
+// Copyright Epic Games, Inc. All Rights Reserved.
+
+#pragma once
+
+#include "PoseSearchDebuggerDatabaseColumns.h"
+#include "PoseSearchDebuggerDatabaseView.h"
+#include "Widgets/Views/STableRow.h"
+
+namespace UE::PoseSearch
+{
+
+/**
+ * Widget representing a single row of the database view
+ */
+class SDebuggerDatabaseRow : public SMultiColumnTableRow<TSharedRef<FDebuggerDatabaseRowData>>
+{
+	SLATE_BEGIN_ARGS(SDebuggerDatabaseRow) {}
+		SLATE_ATTRIBUTE(const SDebuggerDatabaseView::FColumnMap*, ColumnMap)
+	SLATE_END_ARGS()
+
+	void Construct(
+		const FArguments& InArgs,
+		const TSharedRef<STableViewBase>& InOwnerTable,
+		TSharedRef<FDebuggerDatabaseRowData> InRow,
+		const FTableRowStyle& InRowStyle,
+		const FSlateBrush* InRowBrush,
+		FMargin InPaddingMargin
+	)
+	{
+		ColumnMap = InArgs._ColumnMap;
+		check(ColumnMap.IsBound());
+		
+		Row = InRow;
+		
+		RowBrush = InRowBrush;
+		check(RowBrush);
+
+		SMultiColumnTableRow<TSharedRef<FDebuggerDatabaseRowData>>::Construct(
+			FSuperRowType::FArguments()
+			.Padding(InPaddingMargin)
+			.Style(&InRowStyle),
+			InOwnerTable
+		);
+	}
+
+	virtual TSharedRef<SWidget> GenerateWidgetForColumn(const FName& InColumnName) override
+	{
+		// Get column
+		const TSharedRef<DebuggerDatabaseColumns::IColumn>& Column = (*ColumnMap.Get())[InColumnName];
+		
+		static FSlateFontInfo NormalFont = FAppStyle::Get().GetFontStyle("DetailsView.CategoryTextStyle");
+		const TSharedRef<SWidget> Widget = Column->GenerateWidget(Row.ToSharedRef());
+		
+		return
+			SNew(SBorder)
+			.HAlign(HAlign_Center)
+			.VAlign(VAlign_Center)
+			.BorderImage(RowBrush)
+			[
+				SNew(SHorizontalBox)
+				+ SHorizontalBox::Slot()
+				.VAlign(VAlign_Center)
+				.HAlign(HAlign_Center)
+				.Padding(0.0f, 3.0f)
+				[
+					Widget
+				]
+			];
+	}
+
+	/** Row data associated with this widget */
+	TSharedPtr<FDebuggerDatabaseRowData> Row;
+
+	/** Used for cell styles (active vs database row) */
+	const FSlateBrush* RowBrush = nullptr;
+
+	/** Used to grab the column struct given a column name */
+	TAttribute<const SDebuggerDatabaseView::FColumnMap*> ColumnMap;
+};
+
+} // namespace UE::PoseSearch
+
diff --git a/Engine/Plugins/Experimental/Animation/PoseSearch/Source/Editor/Private/PoseSearchDebuggerDatabaseRowData.h b/Engine/Plugins/Experimental/Animation/PoseSearch/Source/Editor/Private/PoseSearchDebuggerDatabaseRowData.h
new file mode 100644
--- /dev/null
+++ b/Engine/Plugins/Experimental/Animation/PoseSearch/Source/Editor/Private/PoseSearchDebuggerDatabaseRowData.h
@@ -0,0 +1,37 @@
+// Copyright Epic Games, Inc. All Rights Reserved.
+
+#pragma once
+
+#include "PoseSearch/PoseSearch.h"
+
+namespace UE::PoseSearch
+{
+
+class FDebuggerDatabaseRowData : public TSharedFromThis<FDebuggerDatabaseRowData>
+{
+public:
+	FDebuggerDatabaseRowData() = default;
+	
+	ESearchIndexAssetType AssetType = ESearchIndexAssetType::Invalid;
+	int32 PoseIdx = 0;
+	TWeakObjectPtr<const UPoseSearchDatabase> SourceDatabase = nullptr;
+	EPoseCandidateFlags PoseCandidateFlags = EPoseCandidateFlags::None;
+	FString DatabaseName = "";
+	FString DatabasePath = "";
+	FString AssetName = "";
+	FString AssetPath = "";
+	int32 DbAssetIdx = 0;
+	int32 AnimFrame = 0;
+	float AnimPercentage = 0.0f;
+	float AssetTime = 0.0f;
+	bool bMirrored = false;
+	bool bLooping = false;
+	FVector BlendParameters = FVector::Zero();
+	FPoseSearchCost PoseCost;
+	FLinearColor CostColor = FLinearColor::White;
+	TArray<float> CostBreakdowns;
+	TArray<FLinearColor> CostBreakdownsColors;
+	TArray<float> CostVector;
+};
+
+} // namespace UE::PoseSearch
diff --git a/Engine/Plugins/Experimental/Animation/PoseSearch/Source/Editor/Private/PoseSearchDebuggerDatabaseView.cpp b/Engine/Plugins/Experimental/Animation/PoseSearch/Source/Editor/Private/PoseSearchDebuggerDatabaseView.cpp
new file mode 100644
--- /dev/null
+++ b/Engine/Plugins/Experimental/Animation/PoseSearch/Source/Editor/Private/PoseSearchDebuggerDatabaseView.cpp
@@ -0,0 +1,848 @@
+// Copyright Epic Games, Inc. All Rights Reserved.
+
+#include "PoseSearchDebuggerDatabaseView.h"
+#include "Algo/AllOf.h"
+#include "PoseSearchDebuggerDatabaseRow.h"
+#include "PoseSearchDebuggerView.h"
+#include "PoseSearchDebuggerViewModel.h"
+#include "PoseSearch/PoseSearchDerivedData.h"
+#include "Trace/PoseSearchTraceProvider.h"
+#include "Widgets/Input/SCheckBox.h"
+#include "Widgets/Layout/SScrollBar.h"
+#include "Widgets/Views/SListView.h"
+
+#define LOCTEXT_NAMESPACE "PoseSearchDebugger"
+
+namespace UE::PoseSearch
+{
+
+void SDebuggerDatabaseView::Update(const FTraceMotionMatchingStateMessage& State)
+{
+	class SCostBreakDownData : public ICostBreakDownData
+	{
+	public:
+		SCostBreakDownData(TArray<TSharedRef<FDebuggerDatabaseRowData>>& UnfilteredDatabaseRows_, bool bIsVerbose_)
+			: UnfilteredDatabaseRows(UnfilteredDatabaseRows_)
+			, bIsVerbose(bIsVerbose_)
+		{
+		}
+
+		// begin ICostBreakDownData
+		virtual ~SCostBreakDownData()
+		{
+			check(CurrentLabeIndex == -1);
+		}
+
+		virtual int32 Num() const override
+		{
+			return UnfilteredDatabaseRows.Num();
+		}
+
+		virtual bool IsCostVectorFromSchema(int32 Index, const UPoseSearchSchema* Schema) const override
+		{
+			return UnfilteredDatabaseRows[Index]->SourceDatabase->Schema == Schema;
+		}
+
+		virtual TConstArrayView<float> GetCostVector(int32 Index, const UPoseSearchSchema* Schema) const override
+		{
+			check(UnfilteredDatabaseRows[Index]->SourceDatabase->Schema == Schema);
+			return UnfilteredDatabaseRows[Index]->CostVector;
+		}
+
+		virtual void BeginBreakDownSection(const FText& Label) override
+		{
+			check(CurrentLabeIndex == -1);
+			CurrentLabeIndex = GetOrAddLabelIndex(Label);
+		}
+
+		virtual void SetCostBreakDown(float CostBreakDown, int32 Index, const UPoseSearchSchema* Schema) override
+		{
+			FDebuggerDatabaseRowData& UnfilteredDatabaseRow = UnfilteredDatabaseRows[Index].Get();
+			check(UnfilteredDatabaseRow.SourceDatabase->Schema == Schema);
+			TArray<float>& CostBreakdowns = UnfilteredDatabaseRow.CostBreakdowns;
+			CostBreakdowns.SetNumZeroed(CurrentLabeIndex + 1, false);
+			CostBreakdowns[CurrentLabeIndex] = CostBreakDown;
+		}
+
+		virtual void EndBreakDownSection(const FText& Label) override
+		{
+			check(CurrentLabeIndex == GetOrAddLabelIndex(Label));
+			CurrentLabeIndex = INDEX_NONE;
+		}
+
+		virtual bool IsVerbose() const override
+		{
+			return bIsVerbose;
+		}
+		// end ICostBreakDownData
+
+		int32 GetOrAddLabelIndex(const FText& Label)
+		{
+			for (int32 i = 0; i < Labels.Num(); ++i)
+			{
+				if (Labels[i].EqualTo(Label))
+				{
+					return i;
+				}
+			}
+			Labels.Add(Label);
+			return Labels.Num() - 1;
+		}
+
+		bool AreLabelsEqualTo(TArray<FText> OtherLabels)
+		{
+			if (Labels.Num() != OtherLabels.Num())
+			{
+				return false;
+			}
+
+			for (int32 i = 0; i < Labels.Num(); ++i)
+			{
+				if (!Labels[i].EqualTo(OtherLabels[i]))
+				{
+					return false;
+				}
+			}
+
+			return true;
+		}
+
+		const TArray<FText>& GetLabels() const
+		{
+			return Labels;
+		}
+
+	private:
+		TArray<TSharedRef<FDebuggerDatabaseRowData>>& UnfilteredDatabaseRows;
+
+		TArray<FText> Labels;
+
+		int32 CurrentLabeIndex = INDEX_NONE;
+
+		bool bIsVerbose = true;
+	};
+
+	// row cost color palette
+	static const FLinearColor DiscardedRowColor(0.314f, 0.314f, 0.314f); // darker gray
+	static const FLinearColor BestScoreRowColor = FLinearColor::Green;
+	static const FLinearColor WorstScoreRowColor = FLinearColor::Red;
+
+	using namespace DebuggerDatabaseColumns;
+
+	bool bIsVerbose = false;
+	auto DebuggerView = ParentDebuggerViewPtr.Pin();
+	if (DebuggerView.IsValid())
+	{
+		bIsVerbose = DebuggerView->GetViewModel()->IsVerbose();
+	}
+
+	UnfilteredDatabaseRows.Reset();
+	for (const FTraceMotionMatchingStateDatabaseEntry& DbEntry : State.DatabaseEntries)
+	{
+		const UPoseSearchDatabase* Database = FTraceMotionMatchingState::GetObjectFromId<UPoseSearchDatabase>(DbEntry.DatabaseId);
+		if (FAsyncPoseSearchDatabasesManagement::RequestAsyncBuildIndex(Database, ERequestAsyncBuildFlag::ContinueRequest))
+		{
+			const FPoseSearchIndex& SearchIndex = Database->GetSearchIndex();
+			for (const FTraceMotionMatchingStatePoseEntry& PoseEntry : DbEntry.PoseEntries)
+			{
+				if (const FPoseSearchIndexAsset* SearchIndexAsset = SearchIndex.GetAssetForPoseSafe(PoseEntry.DbPoseIdx))
+				{
+					TSharedRef<FDebuggerDatabaseRowData>& Row = UnfilteredDatabaseRows.Add_GetRef(MakeShared<FDebuggerDatabaseRowData>());
+
+					const float Time = SearchIndex.GetAssetTime(PoseEntry.DbPoseIdx, Database->Schema->GetSamplingInterval());
+
+					Row->PoseIdx = PoseEntry.DbPoseIdx;
+					Row->SourceDatabase = Database;
+					Row->DatabaseName = Database->GetName();
+					Row->DatabasePath = Database->GetPathName();
+					Row->PoseCandidateFlags = PoseEntry.PoseCandidateFlags;
+					Row->DbAssetIdx = SearchIndexAsset->SourceAssetIdx;
+					Row->AssetTime = Time;
+					Row->bMirrored = SearchIndexAsset->bMirrored;
+					Row->PoseCost = PoseEntry.Cost;
+
+					Row->CostVector.SetNum(Database->Schema->SchemaCardinality);
+					TConstArrayView<float> PoseValues = SearchIndex.GetPoseValuesSafe(PoseEntry.DbPoseIdx);
+
+					// in case we modify the schema while PIE is paused and displaying the Pose Search Editor, we could end up with a stale State with a DbEntry.QueryVector saved with the previous schema
+					// so the cardinality of DbEntry.QueryVector and PoseValues don't match. In that case we just use PoseValues as query to have all costs set to zero
+					const bool bIsQueryVectorValid = DbEntry.QueryVector.Num() == PoseValues.Num();
+					CompareFeatureVectors(PoseValues, bIsQueryVectorValid ? DbEntry.QueryVector : PoseValues, SearchIndex.WeightsSqrt, Row->CostVector);
+
+					const FInstancedStruct& DatabaseAssetStruct = Database->GetAnimationAssetStruct(*SearchIndexAsset);
+					if (const FPoseSearchDatabaseAnimationAssetBase* DatabaseAsset = DatabaseAssetStruct.GetPtr<FPoseSearchDatabaseAnimationAssetBase>())
+					{
+						Row->AssetName = DatabaseAsset->GetName();
+						Row->AssetPath = DatabaseAsset->GetAnimationAsset() ? DatabaseAsset->GetAnimationAsset()->GetPathName() : "";
+						Row->AssetType = DatabaseAsset->GetSearchIndexType();
+						Row->bLooping = DatabaseAsset->IsLooping();
+						Row->BlendParameters = FVector::Zero();
+						Row->AnimFrame = 0;
+						Row->AnimPercentage = 0.0f;
+
+						if (const FPoseSearchDatabaseSequence* DatabaseSequence = DatabaseAssetStruct.GetPtr<FPoseSearchDatabaseSequence>())
+						{
+							Row->AnimFrame = DatabaseSequence->Sequence->GetFrameAtTime(Time);
+							Row->AnimPercentage = Time / DatabaseSequence->Sequence->GetPlayLength();
+						}
+						else if (const FPoseSearchDatabaseAnimComposite* DatabaseAnimComposite = DatabaseAssetStruct.GetPtr<FPoseSearchDatabaseAnimComposite>())
+						{
+							Row->AnimFrame = DatabaseAnimComposite->AnimComposite->GetFrameAtTime(Time);
+							Row->AnimPercentage = Time / DatabaseAnimComposite->AnimComposite->GetPlayLength();
+						}
+						else if (const FPoseSearchDatabaseBlendSpace* DatabaseBlendSpace = DatabaseAssetStruct.GetPtr<FPoseSearchDatabaseBlendSpace>())
+						{
+							Row->BlendParameters = SearchIndexAsset->BlendParameters;
+						}
+						else
+						{
+							checkNoEntry();
+						}
+					}
+				}
+			}
+		}
+	}
+
+	SCostBreakDownData CostBreakDownData(UnfilteredDatabaseRows, bIsVerbose);
+	if (!UnfilteredDatabaseRows.IsEmpty())
+	{
+		for (const FTraceMotionMatchingStateDatabaseEntry& DbEntry : State.DatabaseEntries)
+		{
+			const UPoseSearchDatabase* Database = FTraceMotionMatchingState::GetObjectFromId<UPoseSearchDatabase>(DbEntry.DatabaseId);
+			if (FAsyncPoseSearchDatabasesManagement::RequestAsyncBuildIndex(Database, ERequestAsyncBuildFlag::ContinueRequest))
+			{
+				Database->Schema->ComputeCostBreakdowns(CostBreakDownData);
+			}
+		}
+
+		// calculating breakdowns min max and colors
+		TArray<float> MinCostBreakdowns;
+		TArray<float> MaxCostBreakdowns;
+
+		const int32 CostBreakdownsCardinality = UnfilteredDatabaseRows[0]->CostBreakdowns.Num();
+		MinCostBreakdowns.Init(UE_MAX_FLT, CostBreakdownsCardinality);
+		MaxCostBreakdowns.Init(-UE_MAX_FLT, CostBreakdownsCardinality);
+
+		auto ArrayMinMax = [](TConstArrayView<float> View, TArrayView<float> Min, TArrayView<float> Max, float InvalidValue)
+		{
+			const int32 Num = View.Num();
+			check(Num == Min.Num() && Num == Max.Num());
+			for (int i = 0; i < Num; ++i)
+			{
+				const float Value = View[i];
+				if (Value != InvalidValue)
+				{
+					Min[i] = FMath::Min(Min[i], Value);
+					Max[i] = FMath::Max(Max[i], Value);
+				}
+			}
+		};
+
+		for (auto& UnfilteredRow : UnfilteredDatabaseRows)
+		{
+			ArrayMinMax(UnfilteredRow->CostBreakdowns, MinCostBreakdowns, MaxCostBreakdowns, UE_MAX_FLT);
+		}
+		
+		auto ArraySafeNormalize = [](TConstArrayView<float> View, TConstArrayView<float> Min, TConstArrayView<float> Max, TArrayView<float> NormalizedView)
+		{
+			const int32 Num = View.Num();
+			check(Num == Min.Num() && Num == Max.Num() && Num == NormalizedView.Num());
+			for (int i = 0; i < Num; ++i)
+			{
+				const float Delta = Max[i] - Min[i];
+				if (FMath::IsNearlyZero(Delta, UE_KINDA_SMALL_NUMBER))
+				{
+					NormalizedView[i] = 0.f;
+				}
+				else
+				{
+					NormalizedView[i] = (View[i] - Min[i]) / Delta;
+				}
+			}
+		};
+
+		auto LinearColorBlend = [](FLinearColor LinearColorA, FLinearColor LinearColorB, float BlendParam) -> FLinearColor
+		{
+			return LinearColorA + (LinearColorB - LinearColorA) * BlendParam;
+		};
+
+		auto LinearColorArrayBlend = [](FLinearColor LinearColorA, FLinearColor LinearColorB, TConstArrayView<float> BlendParam, TArray<FLinearColor>& BlendedColors) -> void
+		{
+			const int32 Num = BlendParam.Num();
+			BlendedColors.SetNumUninitialized(Num);
+			for (int i = 0; i < Num; ++i)
+			{
+				BlendedColors[i] = LinearColorA + (LinearColorB - LinearColorA) * BlendParam[i];
+			}
+		};
+
+		TArray<float> CostBreakdownsColorBlend;
+		CostBreakdownsColorBlend.Init(0, CostBreakdownsCardinality);
+		for (auto& UnfilteredRow : UnfilteredDatabaseRows)
+		{
+			if (EnumHasAnyFlags(UnfilteredRow->PoseCandidateFlags, EPoseCandidateFlags::AnyValidMask))
+			{
+				ArraySafeNormalize(UnfilteredRow->CostBreakdowns, MinCostBreakdowns, MaxCostBreakdowns, CostBreakdownsColorBlend);
+				LinearColorArrayBlend(BestScoreRowColor, WorstScoreRowColor, CostBreakdownsColorBlend, UnfilteredRow->CostBreakdownsColors);
+			}
+			else
+			{
+				UnfilteredRow->CostBreakdownsColors.Init(DiscardedRowColor, CostBreakdownsCardinality);
+			}
+		}
+
+		float MinCost = UE_MAX_FLT;
+		float MaxCost = -UE_MAX_FLT;
+		for (auto& UnfilteredRow : UnfilteredDatabaseRows)
+		{
+			const float Cost = UnfilteredRow->PoseCost.GetTotalCost();
+			MinCost = FMath::Min(MinCost, Cost);
+			MaxCost = FMath::Max(MaxCost, Cost);
+		}
+
+		const float DeltaCost = MaxCost - MinCost;
+		for (auto& UnfilteredRow : UnfilteredDatabaseRows)
+		{
+			if (EnumHasAnyFlags(UnfilteredRow->PoseCandidateFlags, EPoseCandidateFlags::AnyValidMask))
+			{
+				const float CostColorBlend = DeltaCost > UE_KINDA_SMALL_NUMBER ? (UnfilteredRow->PoseCost.GetTotalCost() - MinCost) / DeltaCost : 0.f;
+				UnfilteredRow->CostColor = LinearColorBlend(BestScoreRowColor, WorstScoreRowColor, CostColorBlend);
+			}
+			else
+			{
+				UnfilteredRow->CostColor = DiscardedRowColor;
+			}
+		}
+	}
+
+	if (!CostBreakDownData.AreLabelsEqualTo(OldLabels))
+	{
+		OldLabels = CostBreakDownData.GetLabels();
+
+		// recreating and binding the columns
+		Columns.Reset();
+
+		// Construct all column types
+		int32 ColumnIdx = 0;
+		AddColumn(MakeShared<FDatabaseName>(ColumnIdx++));
+		AddColumn(MakeShared<FAssetName>(ColumnIdx++));
+
+		auto CostColumn = MakeShared<FCost>(ColumnIdx++);
+		AddColumn(CostColumn);
+
+		int32 LabelIdx = 0;
+		for (const FText& Label : CostBreakDownData.GetLabels())
+		{
+			AddColumn(MakeShared<FChannelBreakdownCostColumn>(ColumnIdx++, LabelIdx++, Label));
+		}
+
+#if WITH_EDITORONLY_DATA
+		AddColumn(MakeShared<FCostModifier>(ColumnIdx++));
+#endif // WITH_EDITORONLY_DATA
+		AddColumn(MakeShared<FFrame>(ColumnIdx++));
+		AddColumn(MakeShared<FMirrored>(ColumnIdx++));
+		AddColumn(MakeShared<FLooping>(ColumnIdx++));
+		AddColumn(MakeShared<FPoseIdx>(ColumnIdx++));
+		AddColumn(MakeShared<FBlendParameters>(ColumnIdx++));
+		AddColumn(MakeShared<FPoseCandidateFlags>(ColumnIdx++));
+
+		SortColumn = CostColumn->ColumnId;
+
+		// Active and Continuing Pose view scroll bars only for indenting the columns to align w/ database
+		ActiveView.ScrollBar->SetVisibility(EVisibility::Hidden);
+		ContinuingPoseView.ScrollBar->SetVisibility(EVisibility::Hidden);
+
+		// Refresh Columns
+		ActiveView.HeaderRow->ClearColumns();
+		ContinuingPoseView.HeaderRow->ClearColumns();
+		FilteredDatabaseView.HeaderRow->ClearColumns();
+
+		// Sort columns by index
+		Columns.ValueSort([](const TSharedRef<IColumn> Column0, const TSharedRef<IColumn> Column1)
+		{
+			return Column0->SortIndex < Column1->SortIndex;
+		});
+
+		// Add columns from map to header row
+		for (TPair<FName, TSharedRef<IColumn>>& ColumnPair : Columns)
+		{
+			IColumn& Column = ColumnPair.Value.Get();
+			if (ColumnPair.Value->bEnabled)
+			{
+				SHeaderRow::FColumn::FArguments ColumnArgs = SHeaderRow::FColumn::FArguments()
+					.ColumnId(Column.ColumnId)
+					.DefaultLabel(Column.GetLabel())
+					.SortMode(this, &SDebuggerDatabaseView::GetColumnSortMode, Column.ColumnId)
+					.OnSort(this, &SDebuggerDatabaseView::OnColumnSortModeChanged)
+					.FillWidth(this, &SDebuggerDatabaseView::GetColumnWidth, Column.ColumnId)
+					.VAlignCell(VAlign_Center)
+					.VAlignHeader(VAlign_Center)
+					.HAlignHeader(HAlign_Center)
+					.HAlignCell(HAlign_Fill);
+
+				FilteredDatabaseView.HeaderRow->AddColumn(ColumnArgs);
+
+				// Every time the active column is changed, update the database column
+				ActiveView.HeaderRow->AddColumn(ColumnArgs.OnWidthChanged(this, &SDebuggerDatabaseView::OnColumnWidthChanged, Column.ColumnId));
+
+				ContinuingPoseView.HeaderRow->AddColumn(ColumnArgs.OnWidthChanged(this, &SDebuggerDatabaseView::OnColumnWidthChanged, Column.ColumnId));
+			}
+		}
+	}
+
+	SortDatabaseRows();
+	PopulateViewRows();
+}
+
+void SDebuggerDatabaseView::AddColumn(TSharedRef<DebuggerDatabaseColumns::IColumn>&& Column)
+{
+	Columns.Add(Column->ColumnId, Column);
+}
+
+EColumnSortMode::Type SDebuggerDatabaseView::GetColumnSortMode(const FName ColumnId) const
+{
+	if (ColumnId == SortColumn)
+	{
+		return SortMode;
+	}
+	
+	return EColumnSortMode::None;
+}
+
+float SDebuggerDatabaseView::GetColumnWidth(const FName ColumnId) const
+{
+	check(Columns.Find(ColumnId));
+
+	return Columns[ColumnId]->Width;
+}
+
+void SDebuggerDatabaseView::OnColumnSortModeChanged(const EColumnSortPriority::Type SortPriority, const FName & ColumnId, const EColumnSortMode::Type InSortMode)
+{
+	check(Columns.Find(ColumnId));
+	SortColumn = ColumnId;
+	SortMode = InSortMode;
+	SortDatabaseRows();
+	PopulateViewRows();
+}
+
+void SDebuggerDatabaseView::OnColumnWidthChanged(const float NewWidth, FName ColumnId) const
+{
+	check(Columns.Find(ColumnId));
+	
+	Columns[ColumnId]->Width = NewWidth;
+}
+
+void SDebuggerDatabaseView::OnFilterTextChanged(const FText& SearchText)
+{
+	FilterText = SearchText;
+	PopulateViewRows();
+}
+
+void SDebuggerDatabaseView::OnHideInvalidPosesCheckboxChanged(ECheckBoxState State)
+{
+	if (State == ECheckBoxState::Checked) 
+	{
+		bHideInvalidPoses = true;
+	}
+	else 
+	{
+		bHideInvalidPoses = false;
+	}
+	PopulateViewRows();
+}
+
+void SDebuggerDatabaseView::OnDatabaseRowSelectionChanged(TSharedPtr<FDebuggerDatabaseRowData> Row, ESelectInfo::Type SelectInfo)
+{
+	if (Row.IsValid())
+	{
+		OnPoseSelectionChanged.ExecuteIfBound(Row->SourceDatabase.Get(), Row->PoseIdx, Row->AssetTime);
+	}
+}
+
+void SDebuggerDatabaseView::SortDatabaseRows()
+{
+	if (!UnfilteredDatabaseRows.IsEmpty())
+	{
+		if (SortMode == EColumnSortMode::Ascending)
+		{
+			UnfilteredDatabaseRows.Sort(Columns[SortColumn]->GetSortPredicate());
+		}
+		else if (SortMode == EColumnSortMode::Descending)
+		{
+			auto DescendingPredicate = [this](const auto& Lhs, const auto& Rhs) -> bool
+			{
+				return !Columns[SortColumn]->GetSortPredicate()(Lhs, Rhs);
+			};
+
+			UnfilteredDatabaseRows.Sort(DescendingPredicate);
+		}
+	}
+}
+
+void SDebuggerDatabaseView::PopulateViewRows()
+{
+	ActiveView.Rows.Reset();
+	ContinuingPoseView.Rows.Reset();
+	FilteredDatabaseView.Rows.Empty();
+
+	FString FilterString = FilterText.ToString();
+	TArray<FString> Tokens;
+	FilterString.ParseIntoArrayWS(Tokens);
+	const bool bHasNameFilter = !Tokens.IsEmpty();
+
+	for (const auto& UnfilteredRow : UnfilteredDatabaseRows)
+	{
+		bool bTryAddToFilteredDatabaseViewRows = true;
+		if (EnumHasAnyFlags(UnfilteredRow->PoseCandidateFlags, EPoseCandidateFlags::Valid_ContinuingPose))
+		{
+			ContinuingPoseView.Rows.Add(UnfilteredRow);
+			bTryAddToFilteredDatabaseViewRows = false;
+		}
+
+		if (EnumHasAnyFlags(UnfilteredRow->PoseCandidateFlags, EPoseCandidateFlags::Valid_CurrentPose))
+		{
+			ActiveView.Rows.Add(UnfilteredRow);
+			bTryAddToFilteredDatabaseViewRows = false;
+		}
+
+		if (bTryAddToFilteredDatabaseViewRows)
+		{
+			if (!bHideInvalidPoses || EnumHasAnyFlags(UnfilteredRow->PoseCandidateFlags, EPoseCandidateFlags::AnyValidMask))
+			{
+				bool bPassesNameFilter = true;
+				if (bHasNameFilter)
+				{
+					bPassesNameFilter = Algo::AllOf(Tokens, [&](FString Token)
+					{
+						return UnfilteredRow->AssetName.Contains(Token);
+					});
+				}
+
+				if (bPassesNameFilter)
+				{
+					FilteredDatabaseView.Rows.Add(UnfilteredRow);
+				}
+			}
+		}
+	}
+
+	ActiveView.ListView->RequestListRefresh();
+	ContinuingPoseView.ListView->RequestListRefresh();
+	FilteredDatabaseView.ListView->RequestListRefresh();
+}
+
+TSharedRef<ITableRow> SDebuggerDatabaseView::HandleGenerateDatabaseRow(TSharedRef<FDebuggerDatabaseRowData> Item, const TSharedRef<STableViewBase>& OwnerTable) const
+{
+	return SNew(SDebuggerDatabaseRow, OwnerTable, Item, FilteredDatabaseView.RowStyle, &FilteredDatabaseView.RowBrush, FMargin(0.0f, 2.0f, 6.0f, 2.0f))
+		.ColumnMap(this, &SDebuggerDatabaseView::GetColumnMap);	
+}
+
+TSharedRef<ITableRow> SDebuggerDatabaseView::HandleGenerateActiveRow(TSharedRef<FDebuggerDatabaseRowData> Item, const TSharedRef<STableViewBase>& OwnerTable) const
+{
+	return SNew(SDebuggerDatabaseRow, OwnerTable, Item, ActiveView.RowStyle, &ActiveView.RowBrush, FMargin(0.0f, 2.0f, 6.0f, 4.0f))
+		.ColumnMap(this, &SDebuggerDatabaseView::GetColumnMap);
+}
+
+TSharedRef<ITableRow> SDebuggerDatabaseView::HandleGenerateContinuingPoseRow(TSharedRef<FDebuggerDatabaseRowData> Item, const TSharedRef<STableViewBase>& OwnerTable) const
+{
+	return SNew(SDebuggerDatabaseRow, OwnerTable, Item, ContinuingPoseView.RowStyle, &ContinuingPoseView.RowBrush, FMargin(0.0f, 2.0f, 6.0f, 4.0f))
+		.ColumnMap(this, &SDebuggerDatabaseView::GetColumnMap);
+}
+
+void SDebuggerDatabaseView::Construct(const FArguments& InArgs)
+{
+	ParentDebuggerViewPtr = InArgs._Parent;
+	OnPoseSelectionChanged = InArgs._OnPoseSelectionChanged;
+	check(OnPoseSelectionChanged.IsBound());
+
+	// Active Row
+	ActiveView.HeaderRow = SNew(SHeaderRow);
+
+	// Used for spacing
+	ActiveView.ScrollBar = SNew(SScrollBar)
+		.Orientation(Orient_Vertical)
+		.HideWhenNotInUse(false)
+		.AlwaysShowScrollbar(true)
+		.AlwaysShowScrollbarTrack(true);
+
+	ActiveView.ListView = SNew(SListView<TSharedRef<FDebuggerDatabaseRowData>>)
+		.ListItemsSource(&ActiveView.Rows)
+		.HeaderRow(ActiveView.HeaderRow.ToSharedRef())
+		.OnGenerateRow(this, &SDebuggerDatabaseView::HandleGenerateActiveRow)
+		.ExternalScrollbar(ActiveView.ScrollBar)
+		.SelectionMode(ESelectionMode::SingleToggle)
+		.ConsumeMouseWheel(EConsumeMouseWheel::Never);
+
+	ActiveView.RowStyle = FAppStyle::GetWidgetStyle<FTableRowStyle>("TableView.Row");
+	ActiveView.RowBrush = *FAppStyle::GetBrush("DetailsView.CategoryTop");
+
+	// ContinuingPose Row
+	ContinuingPoseView.HeaderRow = SNew(SHeaderRow).Visibility(EVisibility::Collapsed);
+
+	// Used for spacing
+	ContinuingPoseView.ScrollBar = SNew(SScrollBar)
+		.Orientation(Orient_Vertical)
+		.HideWhenNotInUse(false)
+		.AlwaysShowScrollbar(true)
+		.AlwaysShowScrollbarTrack(true);
+
+	ContinuingPoseView.ListView = SNew(SListView<TSharedRef<FDebuggerDatabaseRowData>>)
+		.ListItemsSource(&ContinuingPoseView.Rows)
+		.HeaderRow(ContinuingPoseView.HeaderRow.ToSharedRef())
+		.OnGenerateRow(this, &SDebuggerDatabaseView::HandleGenerateContinuingPoseRow)
+		.ExternalScrollbar(ContinuingPoseView.ScrollBar)
+		.SelectionMode(ESelectionMode::SingleToggle)
+		.ConsumeMouseWheel(EConsumeMouseWheel::Never);
+
+	ContinuingPoseView.RowStyle = FAppStyle::GetWidgetStyle<FTableRowStyle>("TableView.Row");
+	ContinuingPoseView.RowBrush = *FAppStyle::GetBrush("DetailsView.CategoryTop");
+
+	// Filtered Database
+	FilteredDatabaseView.ScrollBar =
+		SNew(SScrollBar)
+		.Orientation(Orient_Vertical)
+		.HideWhenNotInUse(false)
+		.AlwaysShowScrollbar(true)
+		.AlwaysShowScrollbarTrack(true);
+	FilteredDatabaseView.HeaderRow = SNew(SHeaderRow).Visibility(EVisibility::Collapsed);
+
+	FilteredDatabaseView.ListView = SNew(SListView<TSharedRef<FDebuggerDatabaseRowData>>)
+		.ListItemsSource(&FilteredDatabaseView.Rows)
+		.HeaderRow(FilteredDatabaseView.HeaderRow.ToSharedRef())
+		.OnGenerateRow(this, &SDebuggerDatabaseView::HandleGenerateDatabaseRow)
+		.ExternalScrollbar(FilteredDatabaseView.ScrollBar)
+		.SelectionMode(ESelectionMode::Multi)
+		.ConsumeMouseWheel(EConsumeMouseWheel::WhenScrollingPossible)
+		.OnSelectionChanged(this, &SDebuggerDatabaseView::OnDatabaseRowSelectionChanged);
+
+	FilteredDatabaseView.RowStyle = FAppStyle::GetWidgetStyle<FTableRowStyle>("TableView.Row");
+	// Set selected color to white to retain visibility when multi-selecting
+	FilteredDatabaseView.RowStyle.SetSelectedTextColor(FLinearColor(FVector3f(0.8f)));
+	FilteredDatabaseView.RowBrush = *FAppStyle::GetBrush("ToolPanel.GroupBorder");
+
+	ChildSlot
+	[
+		SNew(SVerticalBox)
+		+ SVerticalBox::Slot()
+		// Side and top margins, ignore bottom handled by the color border below
+		.Padding(0.0f, 5.0f, 0.0f, 0.0f)
+		.AutoHeight()
+		[
+			// Active Row text tab
+			SNew(SVerticalBox)
+			+ SVerticalBox::Slot()
+			.AutoHeight()
+			.Padding(0.0f)
+			[
+				SNew(SHorizontalBox)
+				+ SHorizontalBox::Slot()
+				.HAlign(HAlign_Center)
+				.VAlign(VAlign_Fill)
+				.Padding(0.0f)
+				.AutoWidth()
+				[
+					SNew(SBorder)
+					.BorderImage(FAppStyle::GetBrush("DetailsView.CategoryTop"))
+					.Padding(FMargin(30.0f, 3.0f, 30.0f, 0.0f))
+					.HAlign(HAlign_Center)
+					.VAlign(VAlign_Fill)
+					[
+						SNew(STextBlock)
+						.Text(FText::FromString("Active Pose"))	
+					]
+				]
+			]
+
+			// Active row list view with scroll bar
+			+ SVerticalBox::Slot()
+			
+			.AutoHeight()
+			[
+				SNew(SHorizontalBox)
+				
+				+ SHorizontalBox::Slot()
+				.HAlign(HAlign_Fill)
+				.VAlign(VAlign_Fill)
+				.Padding(0.0f)
+				[
+					SNew(SBorder)
+					
+					.BorderImage(FAppStyle::GetBrush("NoBorder"))
+					.Padding(0.0f)
+					[
+						ActiveView.ListView.ToSharedRef()
+					]
+				]
+
+				+ SHorizontalBox::Slot()
+				.AutoWidth()
+				[
+					ActiveView.ScrollBar.ToSharedRef()
+				]
+			]	
+		]
+
+		+ SVerticalBox::Slot()
+		// Side and top margins, ignore bottom handled by the color border below
+		.Padding(0.0f, 5.0f, 0.0f, 0.0f)
+		.AutoHeight()
+		[
+			// ContinuingPose Row text tab
+			SNew(SVerticalBox)
+			+ SVerticalBox::Slot()
+			.AutoHeight()
+			.Padding(0.0f)
+			[
+				SNew(SHorizontalBox)
+				+ SHorizontalBox::Slot()
+				.HAlign(HAlign_Center)
+				.VAlign(VAlign_Fill)
+				.Padding(0.0f)
+				.AutoWidth()
+				[
+					SNew(SBorder)
+					.BorderImage(FAppStyle::GetBrush("DetailsView.CategoryTop"))
+					.Padding(FMargin(30.0f, 3.0f, 30.0f, 0.0f))
+					.HAlign(HAlign_Center)
+					.VAlign(VAlign_Fill)
+					[
+						SNew(STextBlock)
+						.Text(FText::FromString("Continuing Pose"))	
+					]
+				]
+			]
+
+			// ContinuingPose row list view with scroll bar
+			+ SVerticalBox::Slot()
+			
+			.AutoHeight()
+			[
+				SNew(SHorizontalBox)
+				
+				+ SHorizontalBox::Slot()
+				.HAlign(HAlign_Fill)
+				.VAlign(VAlign_Fill)
+				.Padding(0.0f)
+				[
+					SNew(SBorder)
+					
+					.BorderImage(FAppStyle::GetBrush("NoBorder"))
+					.Padding(0.0f)
+					[
+						ContinuingPoseView.ListView.ToSharedRef()
+					]
+				]
+
+				+ SHorizontalBox::Slot()
+				.AutoWidth()
+				[
+					ContinuingPoseView.ScrollBar.ToSharedRef()
+				]
+			]	
+		]
+		
+		+ SVerticalBox::Slot()
+		.Padding(0.0f, 0.0f, 0.0f, 5.0f)
+		[
+			// Database view text tab
+			SNew(SVerticalBox)
+			+ SVerticalBox::Slot()
+			[
+				SNew(SHorizontalBox)
+				+ SHorizontalBox::Slot()
+				.HAlign(HAlign_Center)
+				.VAlign(VAlign_Fill)
+				.Padding(0.0f)
+				[
+					SNew(SBorder)
+					.BorderImage(FAppStyle::GetBrush("DetailsView.CategoryTop"))
+					.Padding(FMargin(30.0f, 3.0f, 30.0f, 0.0f))
+					.HAlign(HAlign_Center)
+					.VAlign(VAlign_Fill)
+					[
+						SNew(STextBlock)
+						.Text(FText::FromString("Pose Candidates"))
+					]
+				]
+				.AutoWidth()
+				
+				+ SHorizontalBox::Slot()
+				.HAlign(HAlign_Fill)
+				[
+					SNew(SBorder)
+					.BorderImage(&FilteredDatabaseView.RowStyle.EvenRowBackgroundBrush)
+				]
+			]
+			.AutoHeight()
+
+			// Gray line below the tab 
+			+ SVerticalBox::Slot()
+			.AutoHeight()
+			.Padding(0.0f)
+			[
+				SNew(SBorder)
+				.BorderImage(FAppStyle::GetBrush("DetailsView.CategoryTop"))
+				.Padding(FMargin(0.0f, 3.0f, 0.0f, 3.0f))
+				.HAlign(HAlign_Fill)
+				.VAlign(VAlign_Fill)
+			]
+
+			+ SVerticalBox::Slot()
+			.Padding(0.0f, 0.0f, 0.0f, 5.0f)
+			.AutoHeight()
+			[
+				SNew(SHorizontalBox)
+				+ SHorizontalBox::Slot()
+				.Padding(10, 5, 10, 5)
+				[
+					SAssignNew(FilterBox, SSearchBox)
+					.OnTextChanged(this, &SDebuggerDatabaseView::OnFilterTextChanged)
+				]
+				+ SHorizontalBox::Slot()
+				.AutoWidth()
+				.Padding(10, 5, 10, 5)
+				[
+					SNew(SCheckBox)
+					.OnCheckStateChanged_Lambda([this](ECheckBoxState State)
+					{
+						   SDebuggerDatabaseView::OnHideInvalidPosesCheckboxChanged(State);
+					})
+					[
+						SNew(STextBlock)
+						.Text(LOCTEXT("PoseSearchDebuggerHideInvalidPosesFlag", "Hide Invalid Poses"))
+					]
+				]
+			]
+		
+			+ SVerticalBox::Slot()
+			[
+				SNew(SHorizontalBox)
+				+ SHorizontalBox::Slot()
+				.Padding(0.0f)
+				[
+					SNew(SBorder)
+					.BorderImage(FAppStyle::GetBrush("NoBorder"))
+					.Padding(0.0f)
+					[
+						FilteredDatabaseView.ListView.ToSharedRef()
+					]
+				]
+				
+				+ SHorizontalBox::Slot()
+				.AutoWidth()
+				[
+					FilteredDatabaseView.ScrollBar.ToSharedRef()
+				]
+			]
+		]
+	];
+
+	SortMode = EColumnSortMode::Ascending;
+	OldLabels.Reset();
+	Columns.Reset();
+}
+
+} // namespace UE::PoseSearch
+
+#undef LOCTEXT_NAMESPACE
diff --git a/Engine/Plugins/Experimental/Animation/PoseSearch/Source/Editor/Private/PoseSearchDebuggerDatabaseView.h b/Engine/Plugins/Experimental/Animation/PoseSearch/Source/Editor/Private/PoseSearchDebuggerDatabaseView.h
new file mode 100644
--- /dev/null
+++ b/Engine/Plugins/Experimental/Animation/PoseSearch/Source/Editor/Private/PoseSearchDebuggerDatabaseView.h
@@ -0,0 +1,143 @@
+// Copyright Epic Games, Inc. All Rights Reserved.
+
+#pragma once
+
+#include "PoseSearchDebuggerDatabaseRowData.h"
+#include "Widgets/Input/SSearchBox.h"
+#include "Widgets/Views/SHeaderRow.h"
+
+template <typename ItemType> class SListView;
+
+namespace UE::PoseSearch
+{
+
+namespace DebuggerDatabaseColumns { struct IColumn; }
+struct FTraceMotionMatchingStateMessage;
+
+/** Sets model selection data on row selection */
+DECLARE_DELEGATE_ThreeParams(FOnPoseSelectionChanged, const UPoseSearchDatabase*, int32, float)
+
+/**
+ * Database panel view widget of the PoseSearch debugger
+ */
+class SDebuggerDatabaseView : public SCompoundWidget
+{
+	SLATE_BEGIN_ARGS(SDebuggerDatabaseView) {}
+		SLATE_ARGUMENT(TWeakPtr<class SDebuggerView>, Parent)
+		SLATE_EVENT(FOnPoseSelectionChanged, OnPoseSelectionChanged)
+	SLATE_END_ARGS()
+	
+	void Construct(const FArguments& InArgs);
+	void Update(const FTraceMotionMatchingStateMessage& State);
+
+	const TSharedPtr<SListView<TSharedRef<FDebuggerDatabaseRowData>>>& GetActiveRow() const { return ActiveView.ListView; }
+	const TSharedPtr<SListView<TSharedRef<FDebuggerDatabaseRowData>>>& GetContinuingPoseRow() const { return ContinuingPoseView.ListView; }
+	const TSharedPtr<SListView<TSharedRef<FDebuggerDatabaseRowData>>>& GetDatabaseRows() const { return FilteredDatabaseView.ListView; }
+	const TSharedRef<FDebuggerDatabaseRowData>& GetPoseIdxDatabaseRow(int32 PoseIdx) const;
+
+	/** Used by database rows to acquire column-specific information */
+	using FColumnMap = TMap<FName, TSharedRef<DebuggerDatabaseColumns::IColumn>>;
+
+private:
+	/** Adds a column to the existing list */
+	void AddColumn(TSharedRef<DebuggerDatabaseColumns::IColumn>&& Column);
+
+	/** Retrieves current column map, used as an attribute by rows */
+	const FColumnMap* GetColumnMap() const { return &Columns; }
+
+	/** Sorts the database by the current sort predicate, updating the view order */
+	void SortDatabaseRows();
+
+	void PopulateViewRows();
+
+	/** Acquires sort predicate for the given column */
+	EColumnSortMode::Type GetColumnSortMode(const FName ColumnId) const;
+
+	/** Gets active column width, used to align active and database view */
+	float GetColumnWidth(const FName ColumnId) const;
+	
+	/** Updates the active sort predicate, setting the sorting order of all other columns to none
+	 * (to be dependent on active column */
+	void OnColumnSortModeChanged(const EColumnSortPriority::Type SortPriority, const FName & ColumnId, const EColumnSortMode::Type InSortMode);
+
+	/** Aligns the active and database views */
+	void OnColumnWidthChanged(const float NewWidth, FName ColumnId) const;
+
+	/** Called when the text in the filter box is modified to update the filtering */
+	void OnFilterTextChanged(const FText& SearchText);
+
+	/** Called when the HideInvalidPoses Checkbox is Toggled */
+	void OnHideInvalidPosesCheckboxChanged(ECheckBoxState State);
+
+	/** Row selection to update model view */
+	void OnDatabaseRowSelectionChanged(TSharedPtr<FDebuggerDatabaseRowData> Row, ESelectInfo::Type SelectInfo);
+
+	/** Generates a database row widget for the given data */
+	TSharedRef<ITableRow> HandleGenerateDatabaseRow(TSharedRef<FDebuggerDatabaseRowData> Item, const TSharedRef<STableViewBase>& OwnerTable) const;
+	
+	/** Generates the active row widget for the given data */
+	TSharedRef<ITableRow> HandleGenerateActiveRow(TSharedRef<FDebuggerDatabaseRowData> Item, const TSharedRef<STableViewBase>& OwnerTable) const;
+
+	/** Generates the continuing pose row widget for the given data */
+	TSharedRef<ITableRow> HandleGenerateContinuingPoseRow(TSharedRef<FDebuggerDatabaseRowData> Item, const TSharedRef<STableViewBase>& OwnerTable) const;
+
+	TWeakPtr<SDebuggerView> ParentDebuggerViewPtr;
+
+	FOnPoseSelectionChanged OnPoseSelectionChanged;
+
+	/** Current column to sort by */
+	FName SortColumn = "";
+
+	/** Current sorting mode */
+	EColumnSortMode::Type SortMode = EColumnSortMode::Ascending;
+	
+	/** Column data container, used to emplace defined column structures of various types */
+    FColumnMap Columns;
+
+	TArray<FText> OldLabels;
+
+	struct FTable
+	{
+		/** Header row*/
+		TSharedPtr<SHeaderRow> HeaderRow;
+
+		/** Widget for displaying the list of row objects */
+		TSharedPtr<SListView<TSharedRef<FDebuggerDatabaseRowData>>> ListView;
+
+		// @TODO: Explore options for active row other than displaying array of 1 element
+		/** List of row objects */
+		TArray<TSharedRef<FDebuggerDatabaseRowData>> Rows;
+
+		/** Background style for the list view */
+		FTableRowStyle RowStyle;
+
+		/** Row color */
+		FSlateBrush RowBrush;
+
+		/** Scroll bar for the data table */
+		TSharedPtr<SScrollBar> ScrollBar;
+	};
+
+	/** Active row at the top of the view */
+	FTable ActiveView;
+
+	/** Continuing pose row below Active row */
+	FTable ContinuingPoseView;
+
+	/** All database poses */
+	TArray<TSharedRef<FDebuggerDatabaseRowData>> UnfilteredDatabaseRows;
+
+	/** Database listing for filtered poses */
+	FTable FilteredDatabaseView;
+
+	/** Search box widget */
+	TSharedPtr<SSearchBox> FilterBox;
+
+	/** Text used to filter DatabaseView */
+	FText FilterText;
+
+	/* Bool used to hide invalid poses from display */
+	bool bHideInvalidPoses = false;
+};
+
+} // namespace UE::PoseSearch
diff --git a/Engine/Plugins/Experimental/Animation/PoseSearch/Source/Editor/Private/PoseSearchDebuggerReflection.h b/Engine/Plugins/Experimental/Animation/PoseSearch/Source/Editor/Private/PoseSearchDebuggerReflection.h
new file mode 100644
--- /dev/null
+++ b/Engine/Plugins/Experimental/Animation/PoseSearch/Source/Editor/Private/PoseSearchDebuggerReflection.h
@@ -0,0 +1,87 @@
+// Copyright Epic Games, Inc. All Rights Reserved.
+
+#pragma once
+
+#include "CoreMinimal.h"
+#include "PoseSearchDebuggerReflection.generated.h"
+
+/**
+ * Used by the reflection UObject to encompass draw options for the query and database selections
+ */
+USTRUCT()
+struct FPoseSearchDebuggerFeatureDrawOptions
+{
+	GENERATED_BODY()
+
+	UPROPERTY(EditAnywhere, Category="Draw Options")
+    bool bDisable = false;
+
+	UPROPERTY(EditAnywhere, Category = "Draw Options", Meta = (EditCondition = "!bDisable"))
+	bool bDrawBoneNames = false;
+
+	UPROPERTY(EditAnywhere, Category = "Draw Options", Meta = (EditCondition = "!bDisable"))
+	bool bDrawSampleLabels = false;
+};
+
+/**
+ * Reflection UObject being observed in the details view panel of the debugger
+ */
+UCLASS()
+class POSESEARCHEDITOR_API UPoseSearchDebuggerReflection : public UObject
+{
+	GENERATED_BODY()
+
+public:
+
+	UPROPERTY(VisibleAnywhere, Category="Motion Matching State", Meta=(DisplayName="Current Database"))
+	FString CurrentDatabaseName = "";
+
+	/** Time since last PoseSearch jump */
+	UPROPERTY(VisibleAnywhere, Category="Motion Matching State")
+	float ElapsedPoseJumpTime = 0.0f;
+
+	UPROPERTY(VisibleAnywhere, Category = "Motion Matching State")
+	FString AssetPlayerAssetName = "";
+
+	UPROPERTY(VisibleAnywhere, Category = "Motion Matching State")
+	float AssetPlayerTime = 0.0f;
+
+	UPROPERTY(VisibleAnywhere, Category = "Motion Matching State")
+	float LastDeltaTime = 0.0f;
+
+	UPROPERTY(VisibleAnywhere, Category = "Motion Matching State")
+	float SimLinearVelocity;
+
+	UPROPERTY(VisibleAnywhere, Category = "Motion Matching State")
+	float SimAngularVelocity;
+
+	UPROPERTY(VisibleAnywhere, Category = "Motion Matching State")
+	float AnimLinearVelocity;
+
+	UPROPERTY(VisibleAnywhere, Category = "Motion Matching State")
+	float AnimAngularVelocity;
+
+	UPROPERTY(EditAnywhere, Category="Draw Options", Meta=(DisplayName="Query"))
+	FPoseSearchDebuggerFeatureDrawOptions QueryDrawOptions;
+
+	UPROPERTY(EditAnywhere, Category="Draw Options", Meta=(DisplayName="Selected Pose"))
+	FPoseSearchDebuggerFeatureDrawOptions SelectedPoseDrawOptions;
+
+	UPROPERTY(EditAnywhere, Category = "Draw Options")
+	bool bDrawActiveSkeleton = false;
+	
+	UPROPERTY(EditAnywhere, Category = "Draw Options")
+	bool bDrawSelectedSkeleton = false;
+
+    UPROPERTY(VisibleAnywhere, Category="Pose Vectors")
+	TArray<float> QueryPoseVector;
+    	
+    UPROPERTY(VisibleAnywhere, Category="Pose Vectors")
+	TArray<float> ActivePoseVector;
+
+	UPROPERTY(VisibleAnywhere, Category="Pose Vectors")
+	TArray<float> SelectedPoseVector;
+
+	UPROPERTY(VisibleAnywhere, Category="Pose Vectors")
+	TArray<float> CostVector;
+};
diff --git a/Engine/Plugins/Experimental/Animation/PoseSearch/Source/Editor/Private/PoseSearchDebuggerView.cpp b/Engine/Plugins/Experimental/Animation/PoseSearch/Source/Editor/Private/PoseSearchDebuggerView.cpp
new file mode 100644
--- /dev/null
+++ b/Engine/Plugins/Experimental/Animation/PoseSearch/Source/Editor/Private/PoseSearchDebuggerView.cpp
@@ -0,0 +1,790 @@
+// Copyright Epic Games, Inc. All Rights Reserved.
+
+#include "PoseSearchDebuggerView.h"
+#include "Modules/ModuleManager.h"
+#include "PoseSearch/PoseSearchDerivedData.h"
+#include "PoseSearchDebugger.h"
+#include "PoseSearchDebuggerDatabaseRowData.h"
+#include "PoseSearchDebuggerDatabaseView.h"
+#include "PoseSearchDebuggerReflection.h"
+#include "PoseSearchDebuggerViewModel.h"
+#include "PropertyEditorModule.h"
+#include "Trace/PoseSearchTraceProvider.h"
+#include "Widgets/Input/SButton.h"
+#include "Widgets/Input/SNumericEntryBox.h"
+#include "Widgets/Layout/SWidgetSwitcher.h"
+#include "Widgets/SBoxPanel.h"
+#include "Widgets/Text/STextBlock.h"
+#include "Widgets/Views/SListView.h"
+
+#define LOCTEXT_NAMESPACE "PoseSearchDebugger"
+
+namespace UE::PoseSearch
+{
+
+class SDebuggerMessageBox : public SCompoundWidget
+{
+	SLATE_BEGIN_ARGS(SDebuggerMessageBox) {}
+	SLATE_END_ARGS()
+
+	void Construct(const FArguments& InArgs, const FString& Message)
+	{
+		ChildSlot
+		[
+			SNew(SVerticalBox)
+			+ SVerticalBox::Slot()
+			.HAlign(HAlign_Center)
+			.VAlign(VAlign_Center)
+			[
+				SNew(STextBlock)
+				.Text(FText::FromString(Message))
+				.Font(FAppStyle::Get().GetFontStyle("DetailsView.CategoryFontStyle"))
+			]
+		];
+	}
+};
+
+void SDebuggerDetailsView::Construct(const FArguments& InArgs)
+{
+	ParentDebuggerViewPtr = InArgs._Parent;
+
+	// Add property editor (detail view) UObject to world root so that it persists when PIE is stopped
+	Reflection = NewObject<UPoseSearchDebuggerReflection>();
+	Reflection->AddToRoot();
+	check(IsValid(Reflection));
+
+	// @TODO: Convert this to a custom builder instead of of a standard details view
+	// Load property module and create details view with our reflection UObject
+	FPropertyEditorModule& PropPlugin = FModuleManager::LoadModuleChecked<FPropertyEditorModule>("PropertyEditor");
+	FDetailsViewArgs DetailsViewArgs;
+    DetailsViewArgs.NameAreaSettings = FDetailsViewArgs::HideNameArea;
+	DetailsViewArgs.DefaultsOnlyVisibility = EEditDefaultsOnlyNodeVisibility::Hide;
+
+	Details = PropPlugin.CreateDetailView(DetailsViewArgs);
+	Details->SetObject(Reflection);
+	
+	ChildSlot
+	[
+		Details.ToSharedRef()
+	];
+}
+
+void SDebuggerDetailsView::Update(const FTraceMotionMatchingStateMessage& State) const
+{
+	UpdateReflection(State);
+}
+
+SDebuggerDetailsView::~SDebuggerDetailsView()
+{
+	// Our previously instantiated object attached to root may be cleaned up at this point
+	if (UObjectInitialized())
+	{
+		Reflection->RemoveFromRoot();
+	}
+}
+
+void SDebuggerDetailsView::UpdateReflection(const FTraceMotionMatchingStateMessage& State) const
+{
+	check(Reflection);
+
+	const UPoseSearchDatabase* CurrentDatabase = State.GetCurrentDatabase();
+	if (FAsyncPoseSearchDatabasesManagement::RequestAsyncBuildIndex(CurrentDatabase, ERequestAsyncBuildFlag::ContinueRequest))
+	{
+		const FPoseSearchIndex& CurrentSearchIndex = CurrentDatabase->GetSearchIndex();
+		int32 CurrentDbPoseIdx = State.GetCurrentDatabasePoseIndex();
+
+		Reflection->CurrentDatabaseName = CurrentDatabase->GetName();
+		Reflection->ElapsedPoseJumpTime = State.ElapsedPoseJumpTime;
+
+		Reflection->AssetPlayerAssetName = "None";
+		if (const FPoseSearchIndexAsset* IndexAsset = CurrentSearchIndex.GetAssetForPoseSafe(CurrentDbPoseIdx))
+		{
+			Reflection->AssetPlayerAssetName = CurrentDatabase->GetSourceAssetName(*IndexAsset);
+		}
+
+		Reflection->AssetPlayerTime = State.AssetPlayerTime;
+		Reflection->LastDeltaTime = State.DeltaTime;
+		Reflection->SimLinearVelocity = State.SimLinearVelocity;
+		Reflection->SimAngularVelocity = State.SimAngularVelocity;
+		Reflection->AnimLinearVelocity = State.AnimLinearVelocity;
+		Reflection->AnimAngularVelocity = State.AnimAngularVelocity;
+
+		// Query pose
+		Reflection->QueryPoseVector.Reset();
+		for (const FTraceMotionMatchingStateDatabaseEntry& DbEntry : State.DatabaseEntries)
+		{
+			const UPoseSearchDatabase* Database = FTraceMotionMatchingState::GetObjectFromId<UPoseSearchDatabase>(DbEntry.DatabaseId);
+			if (Database && Database == CurrentDatabase)
+			{
+				Reflection->QueryPoseVector = DbEntry.QueryVector;
+				break;
+			}
+		}
+
+		// Active pose
+		Reflection->ActivePoseVector = CurrentSearchIndex.GetPoseValuesSafe(CurrentDbPoseIdx);
+	}
+
+	auto DebuggerView = ParentDebuggerViewPtr.Pin();
+	if (DebuggerView.IsValid())
+	{
+		TArray<TSharedRef<FDebuggerDatabaseRowData>> SelectedRows = DebuggerView->GetSelectedDatabaseRows();
+		if (!SelectedRows.IsEmpty())
+		{
+			const TSharedRef<FDebuggerDatabaseRowData>& Selected = SelectedRows[0];
+			if (FAsyncPoseSearchDatabasesManagement::RequestAsyncBuildIndex(Selected->SourceDatabase.Get(), ERequestAsyncBuildFlag::ContinueRequest))
+			{
+				const FPoseSearchIndex& SelectedSearchIndex = Selected->SourceDatabase->GetSearchIndex();
+				Reflection->SelectedPoseVector = SelectedSearchIndex.GetPoseValuesSafe(Selected->PoseIdx);
+			}
+			Reflection->CostVector = Selected->CostVector;
+		}
+	}
+}
+
+void SDebuggerView::Construct(const FArguments& InArgs, uint64 InAnimInstanceId)
+{
+	ViewModel = InArgs._ViewModel;
+	OnViewClosed = InArgs._OnViewClosed;
+	
+	// Validate the existence of the passed getters
+	check(ViewModel.IsBound())
+	check(OnViewClosed.IsBound());
+	
+	AnimInstanceId = InAnimInstanceId;
+	SelectedNodeId = INDEX_NONE;
+
+	ChildSlot
+	[
+		SAssignNew(DebuggerView, SVerticalBox)
+		+ SVerticalBox::Slot()
+		.FillHeight(1.0f)
+		.VAlign(VAlign_Center)
+		.HAlign(HAlign_Fill)
+		[
+			SAssignNew(Switcher, SWidgetSwitcher)
+			.WidgetIndex(this, &SDebuggerView::SelectView)
+
+			// [0] Selection view before node selection is made
+			+ SWidgetSwitcher::Slot()
+			.Padding(40.0f)
+			.HAlign(HAlign_Fill)
+            .VAlign(VAlign_Center)
+			[
+				SAssignNew(SelectionView, SVerticalBox)
+			]
+
+			// [1] Node selected; node debugger view
+			+ SWidgetSwitcher::Slot()
+			[
+				GenerateNodeDebuggerView()
+			]
+
+			// [2] Occluding message box when stopped (no recording)
+			+ SWidgetSwitcher::Slot()
+			[
+				SNew(SDebuggerMessageBox, "Record gameplay to begin debugging")
+			]
+
+			// [3] Occluding message box when recording
+			+ SWidgetSwitcher::Slot()
+			[
+				SNew(SDebuggerMessageBox, "Recording...")
+			]
+			
+			// [4] Occluding message box when there is no data for the selected MM node
+			+ SWidgetSwitcher::Slot()
+			[
+				GenerateNoDataMessageView()
+			]
+		]
+	];
+}
+
+void SDebuggerView::SetTimeMarker(double InTimeMarker)
+{
+	if (FDebugger::IsPIESimulating())
+	{
+		return;
+	}
+
+	TimeMarker = InTimeMarker;
+}
+
+void SDebuggerView::Tick(const FGeometry& AllottedGeometry, const double InCurrentTime, const float InDeltaTime)
+{
+	if (FDebugger::IsPIESimulating())
+	{
+		return;
+	}
+
+	const UWorld* DebuggerWorld = FDebugger::GetWorld();
+    check(DebuggerWorld);
+	
+	// @TODO: Handle editor world when those features are enabled for the Rewind Debugger
+	// Currently prevents debug draw remnants from stopped world
+	if (DebuggerWorld->WorldType != EWorldType::PIE)
+	{
+		return;
+	}
+	
+	const bool bSameTime = FMath::Abs(TimeMarker - PreviousTimeMarker) < DOUBLE_SMALL_NUMBER;
+	PreviousTimeMarker = TimeMarker;
+
+	TSharedPtr<FDebuggerViewModel> Model = ViewModel.Get();
+
+	bool bNeedUpdate = Model->HasSearchableAssetChanged();
+
+	// We haven't reached the update point yet
+	if (CurrentConsecutiveFrames < ConsecutiveFramesUpdateThreshold)
+	{
+		// If we're on the same time marker, it is consecutive
+		if (bSameTime)
+		{
+			++CurrentConsecutiveFrames;
+		}
+	}
+	else
+	{
+		// New frame after having updated, reset consecutive frames count and start counting again
+		if (!bSameTime)
+		{
+			CurrentConsecutiveFrames = 0;
+			bUpdated = false;
+		}
+		// Haven't updated since passing through frame gate, update once
+		else if (!bUpdated)
+		{
+			bNeedUpdate = true;
+		}
+	}
+
+	if (bNeedUpdate)
+	{
+		Model->OnUpdate();
+		if (UpdateNodeSelection())
+		{
+			Model->OnUpdateNodeSelection(SelectedNodeId);
+			UpdateViews();
+		}
+		bUpdated = true;
+	}
+
+	Model->UpdateAsset();
+	
+	// Draw visualization every tick
+	DrawVisualization();
+}
+
+bool SDebuggerView::UpdateNodeSelection()
+{
+	TSharedPtr<FDebuggerViewModel> Model = ViewModel.Get();
+
+	// Update selection view if no node selected
+	bool bNodeSelected = SelectedNodeId != INDEX_NONE;
+	if (!bNodeSelected)
+	{
+		const TArray<int32>& NodeIds = *Model->GetNodeIds();
+		// Only one node active, bypass selection view
+		if (NodeIds.Num() == 1)
+		{
+			SelectedNodeId = *NodeIds.begin();
+			bNodeSelected = true;
+		}
+		// Create selection view with buttons for each node, displaying the database name
+		else
+		{
+			SelectionView->ClearChildren();
+			for (int32 NodeId : NodeIds)
+			{
+				Model->OnUpdateNodeSelection(NodeId);
+				SelectionView->AddSlot()
+				.HAlign(HAlign_Fill)
+				.VAlign(VAlign_Center)
+				.Padding(10.0f)
+				[
+					SNew(SButton)
+					.Text(FText::FromString(Model->GetSearchableAsset()->GetName()))
+					.HAlign(HAlign_Center)
+					.VAlign(VAlign_Center)
+					.ContentPadding(10.0f)
+					.OnClicked(this, &SDebuggerView::OnUpdateNodeSelection, NodeId)
+				];
+			}
+		}
+	}
+
+	return bNodeSelected;
+}
+
+void SDebuggerView::UpdateViews() const
+{
+	const FTraceMotionMatchingStateMessage* State = ViewModel.Get()->GetMotionMatchingState();
+	if (State)
+	{
+		DatabaseView->Update(*State);
+		DetailsView->Update(*State);
+	}
+}
+
+void SDebuggerView::DrawVisualization() const
+{
+	const UWorld* DebuggerWorld = FDebugger::GetWorld();
+	check(DebuggerWorld);
+
+	const FTraceMotionMatchingStateMessage* State = ViewModel.Get()->GetMotionMatchingState();
+	const FTransform* Transform = ViewModel.Get()->GetRootTransform();
+	if (State && Transform)
+	{
+		DrawFeatures(*DebuggerWorld, *State, *Transform, ViewModel.Get()->GetMeshComponent());
+	}
+}
+
+TArray<TSharedRef<FDebuggerDatabaseRowData>> SDebuggerView::GetSelectedDatabaseRows() const
+{
+	return DatabaseView->GetDatabaseRows()->GetSelectedItems();
+}
+
+void SDebuggerView::DrawFeatures(
+	const UWorld& DebuggerWorld,
+	const FTraceMotionMatchingStateMessage& State,
+	const FTransform& Transform,
+	const USkinnedMeshComponent* Mesh
+) const
+{
+	auto SetDrawFlags = [](FDebugDrawParams& InDrawParams, const FPoseSearchDebuggerFeatureDrawOptions& Options)
+	{
+		InDrawParams.Flags = EDebugDrawFlags::None;
+		if (!Options.bDisable)
+		{
+			if (Options.bDrawBoneNames)
+			{
+				EnumAddFlags(InDrawParams.Flags, EDebugDrawFlags::DrawBoneNames);
+			}
+			
+			if (Options.bDrawSampleLabels)
+			{
+				EnumAddFlags(InDrawParams.Flags, EDebugDrawFlags::DrawSampleLabels);
+			}
+		}
+	};
+
+	const TObjectPtr<UPoseSearchDebuggerReflection> Reflection = DetailsView->GetReflection();
+
+	// Draw query vector
+	{
+		const UPoseSearchDatabase* CurrentDatabase = ViewModel.Get()->GetCurrentDatabase();
+		if (CurrentDatabase)
+		{
+			for (const FTraceMotionMatchingStateDatabaseEntry& DbEntry : State.DatabaseEntries)
+			{
+				const UPoseSearchDatabase* Database = FTraceMotionMatchingState::GetObjectFromId<UPoseSearchDatabase>(DbEntry.DatabaseId);
+				if (Database && Database == CurrentDatabase && DbEntry.QueryVector.Num() == Database->Schema->SchemaCardinality && 
+					FAsyncPoseSearchDatabasesManagement::RequestAsyncBuildIndex(CurrentDatabase, ERequestAsyncBuildFlag::ContinueRequest))
+				{
+					// Set shared state
+					FDebugDrawParams DrawParams;
+					DrawParams.World = &DebuggerWorld;
+					DrawParams.RootTransform = Transform;
+					DrawParams.DefaultLifeTime = 0.0f; // Single frame render
+					DrawParams.Mesh = Mesh;
+					DrawParams.Database = CurrentDatabase;
+					SetDrawFlags(DrawParams, Reflection ? Reflection->QueryDrawOptions : FPoseSearchDebuggerFeatureDrawOptions());
+					EnumAddFlags(DrawParams.Flags, EDebugDrawFlags::DrawQuery);
+					DrawFeatureVector(DrawParams, DbEntry.QueryVector);
+					EnumRemoveFlags(DrawParams.Flags, EDebugDrawFlags::DrawQuery);
+					break;
+				}
+			}
+		}
+	}
+
+	// Draw selected poses
+	{
+		const TSharedPtr<SListView<TSharedRef<FDebuggerDatabaseRowData>>>& DatabaseRows = DatabaseView->GetDatabaseRows();
+		TArray<TSharedRef<FDebuggerDatabaseRowData>> SelectedRows = DatabaseRows->GetSelectedItems();
+	
+		// Red for non-active database view
+		FDebugDrawParams DrawParams;
+		DrawParams.World = &DebuggerWorld;
+		DrawParams.RootTransform = Transform;
+		DrawParams.DefaultLifeTime = 0.0f; // Single frame render
+		DrawParams.Mesh = Mesh;
+		SetDrawFlags(DrawParams, Reflection ? Reflection->SelectedPoseDrawOptions : FPoseSearchDebuggerFeatureDrawOptions());
+
+		// Draw any selected database vectors
+		for (const TSharedRef<FDebuggerDatabaseRowData>& Row : SelectedRows)
+		{
+			DrawParams.Database = Row->SourceDatabase.Get();
+			DrawFeatureVector(DrawParams, Row->PoseIdx);
+		}
+	}
+
+
+	// Draw active pose
+	{
+		TArray<TSharedRef<FDebuggerDatabaseRowData>> ActiveRows = DatabaseView->GetActiveRow()->GetSelectedItems();
+
+		// Active row should only have 0 or 1
+		check(ActiveRows.Num() < 2);
+
+		if (!ActiveRows.IsEmpty())
+		{		
+			// Use the motion-matching state's pose idx, as the active row may be update-throttled at this point
+			FDebugDrawParams DrawParams;
+			DrawParams.World = &DebuggerWorld;
+			DrawParams.RootTransform = Transform;
+			DrawParams.DefaultLifeTime = 0.0f; // Single frame render
+			DrawParams.Mesh = Mesh;
+			DrawParams.Database = ActiveRows[0]->SourceDatabase.Get();
+			DrawFeatureVector(DrawParams, ActiveRows[0]->PoseIdx);
+		}
+	}
+
+
+	// Draw continuing pose
+	{
+		TArray<TSharedRef<FDebuggerDatabaseRowData>> ContinuingRows = DatabaseView->GetContinuingPoseRow()->GetSelectedItems();
+
+		// ContinuingPose row should only have 0 or 1
+		check(ContinuingRows.Num() < 2);
+
+		if (!ContinuingRows.IsEmpty())
+		{
+			FDebugDrawParams DrawParams;
+			DrawParams.World = &DebuggerWorld;
+			DrawParams.RootTransform = Transform;
+			DrawParams.DefaultLifeTime = 0.0f; // Single frame render
+			DrawParams.Mesh = Mesh;
+			DrawParams.Database = ContinuingRows[0]->SourceDatabase.Get();
+			DrawFeatureVector(DrawParams, ContinuingRows[0]->PoseIdx);
+		}
+	}
+
+
+	// Draw skeleton
+	{
+		FSkeletonDrawParams SkeletonDrawParams;
+		if (Reflection && Reflection->bDrawSelectedSkeleton)
+		{
+			SkeletonDrawParams.Flags |= ESkeletonDrawFlags::SelectedPose;
+		}
+		if (Reflection && Reflection->bDrawActiveSkeleton)
+		{
+			SkeletonDrawParams.Flags |= ESkeletonDrawFlags::ActivePose;
+		}
+
+		SkeletonDrawParams.Flags |= ESkeletonDrawFlags::Asset;
+
+		ViewModel.Get()->OnDraw(SkeletonDrawParams);
+	}
+}
+
+int32 SDebuggerView::SelectView() const
+{
+	// Currently recording
+	if (FDebugger::IsPIESimulating() && FDebugger::IsRecording())
+	{
+		return RecordingMsg;
+	}
+
+	// Data has not been recorded yet
+	if (FDebugger::GetRecordingDuration() < DOUBLE_SMALL_NUMBER)
+	{
+		return StoppedMsg;
+	}
+
+	const TSharedPtr<FDebuggerViewModel> Model = ViewModel.Get();
+
+	const bool bNoActiveNodes = Model->GetNodesNum() == 0;
+	const bool bNodeSelectedWithoutData = SelectedNodeId != INDEX_NONE && Model->GetMotionMatchingState() == nullptr;
+
+	// No active nodes, or node selected has no data
+	if (bNoActiveNodes || bNodeSelectedWithoutData)
+    {
+    	return NoDataMsg;
+    }
+
+	// Node not selected yet, showcase selection view
+	if (SelectedNodeId == INDEX_NONE)
+	{
+		return Selection;
+	}
+
+	// Standard debugger view
+	return Debugger;
+}
+
+void SDebuggerView::OnPoseSelectionChanged(const UPoseSearchDatabase* Database, int32 DbPoseIdx, float Time)
+{
+	const TSharedPtr<FDebuggerViewModel> Model = ViewModel.Get();
+	const FTraceMotionMatchingStateMessage* State = Model->GetMotionMatchingState();
+
+	if (State)
+	{
+		DetailsView->Update(*State);
+	}
+	
+	if (DbPoseIdx == INDEX_NONE)
+	{
+		Model->ClearSelectedSkeleton();
+	}
+	else
+	{
+		Model->ShowSelectedSkeleton(Database, DbPoseIdx, Time);
+
+		// Stop asset player when switching selections
+		Model->StopSelection();
+	}
+}
+
+FReply SDebuggerView::OnUpdateNodeSelection(int32 InSelectedNodeId)
+{
+	// -1 will backtrack to selection view
+	SelectedNodeId = InSelectedNodeId;
+	bUpdated = false;
+	return FReply::Handled();
+}
+
+FReply SDebuggerView::TogglePlaySelectedAssets() const
+{
+	const TSharedPtr<SListView<TSharedRef<FDebuggerDatabaseRowData>>>& DatabaseRows = DatabaseView->GetDatabaseRows();
+	TArray<TSharedRef<FDebuggerDatabaseRowData>> Selected = DatabaseRows->GetSelectedItems();
+	const bool bPlaying = ViewModel.Get()->IsPlayingSelections();
+	if (!bPlaying)
+	{
+		if (!Selected.IsEmpty())
+		{
+			// @TODO: Make functional with multiple poses being selected
+			ViewModel.Get()->PlaySelection(Selected[0]->PoseIdx, Selected[0]->AssetTime);
+		}
+	}
+	else
+	{
+		ViewModel.Get()->StopSelection();
+	}
+
+	return FReply::Handled();
+}
+
+TSharedRef<SWidget> SDebuggerView::GenerateNoDataMessageView()
+{
+	TSharedRef<SWidget> ReturnButtonView = GenerateReturnButtonView();
+	ReturnButtonView->SetVisibility(TAttribute<EVisibility>::CreateLambda([this]() -> EVisibility
+	{
+		// Hide the return button for the no data message if we have no nodes at all
+		return ViewModel.Get()->GetNodesNum() > 0 ? EVisibility::Visible : EVisibility::Hidden;
+	}));
+	
+	return 
+		SNew(SVerticalBox)
+
+		+ SVerticalBox::Slot()
+		[
+			SNew(SDebuggerMessageBox, "No recorded data available for the selected frame")
+		]
+
+		+ SVerticalBox::Slot()
+		.AutoHeight()
+		.Padding(20.0f)
+		.HAlign(HAlign_Center)
+		.VAlign(VAlign_Center)
+		[
+			ReturnButtonView
+		];
+}
+
+TSharedRef<SHorizontalBox> SDebuggerView::GenerateReturnButtonView()
+{
+	return
+		SNew(SHorizontalBox)
+
+		+ SHorizontalBox::Slot()
+		.Padding(10, 5, 0, 0)
+		.AutoWidth()
+		[
+			SNew(SButton)
+			.VAlign(VAlign_Center)
+			.Visibility_Lambda([this] { return ViewModel.Get()->GetNodesNum() > 1 ? EVisibility::Visible : EVisibility::Hidden; })
+			.ButtonStyle(FAppStyle::Get(), "SimpleButton")
+			.ContentPadding( FMargin(1, 0) )
+			.OnClicked(this, &SDebuggerView::OnUpdateNodeSelection, static_cast<int32>(INDEX_NONE))
+			// Contents of button, icon then text
+			[
+				SNew(SHorizontalBox)
+
+				+ SHorizontalBox::Slot()
+				.VAlign(VAlign_Center)
+				[
+					SNew(SImage)
+					.Image(FAppStyle::Get().GetBrush("Icons.CircleArrowLeft"))
+					.ColorAndOpacity(FSlateColor::UseForeground())
+				]
+
+				+ SHorizontalBox::Slot()
+				.AutoWidth()
+				.Padding(FMargin(8.0f, 0.0f, 0.0f, 0.0f))
+				.VAlign(VAlign_Center)
+				[
+					SNew(STextBlock)
+					.Text(FText::FromString("Return to Database Selection"))
+					.Justification(ETextJustify::Center)
+				]
+			]
+		];
+}
+
+TSharedRef<SWidget> SDebuggerView::GenerateNodeDebuggerView()
+{
+	TSharedRef<SHorizontalBox> ReturnButtonView = GenerateReturnButtonView();
+	ReturnButtonView->AddSlot()
+	.VAlign(VAlign_Center)
+	.HAlign(HAlign_Fill)
+	.Padding(32, 5, 0, 0)
+	.AutoWidth()
+	[
+		SNew(SButton)
+		.VAlign(VAlign_Center)
+		.HAlign(HAlign_Fill)
+		.ButtonStyle(FAppStyle::Get(), "Button")
+		.ContentPadding( FMargin(5, 0) )
+		.OnClicked(this, &SDebuggerView::TogglePlaySelectedAssets)
+		[
+			SNew(SHorizontalBox)
+			// Icon
+			+ SHorizontalBox::Slot()
+			.AutoWidth()
+			[
+				SNew(SImage)
+				.Image_Lambda([this]
+				{
+					const bool bPlayingSelections = ViewModel.Get()->IsPlayingSelections();
+					return FSlateIcon("FAppStyle", bPlayingSelections ? "PlayWorld.StopPlaySession.Small" : "PlayWorld.PlayInViewport.Small").GetSmallIcon();
+				})
+			]
+			// Text
+			+ SHorizontalBox::Slot()
+			.Padding(FMargin(8, 0, 0, 0))
+			.AutoWidth()
+			[
+				SNew(STextBlock)
+				.Text_Lambda([this] { return ViewModel.Get()->IsPlayingSelections() ? FText::FromString("Stop Selected Asset") : FText::FromString("Play Selected Asset"); })
+				.Justification(ETextJustify::Center)
+			]
+		]
+	];
+	
+	ReturnButtonView->AddSlot()
+	.VAlign(VAlign_Center)
+	.HAlign(HAlign_Left)
+	.Padding(64, 5, 0, 0)
+	.AutoWidth()
+	[
+		SNew(SHorizontalBox)
+		
+		+ SHorizontalBox::Slot()
+		.AutoWidth()
+		.Padding(0, 5, 0, 0)
+		[
+			SNew(STextBlock)
+			.Text(FText::FromString("Asset Play Rate: "))
+		]
+		+ SHorizontalBox::Slot()
+		.AutoWidth()
+		.Padding(8, 0, 0, 0)
+		[
+			SNew(SNumericEntryBox<float>)
+			.MinValue(0.0f)
+			.MaxValue(5.0f)
+			.MinSliderValue(0.0f)
+			.MaxSliderValue(5.0f)
+			.Delta(0.01f)
+			.AllowSpin(true)
+			// Lambda to accomodate the TOptional this requires (for now)
+			.Value_Lambda([this] { return ViewModel.Get()->GetPlayRate(); })
+			.OnValueChanged(ViewModel.Get().ToSharedRef(), &FDebuggerViewModel::ChangePlayRate)	
+		]
+	];
+
+	ReturnButtonView->AddSlot()
+	.VAlign(VAlign_Center)
+	.HAlign(HAlign_Left)
+	.Padding(64, 5, 0, 0)
+	.AutoWidth()
+	[
+		SNew(SHorizontalBox)
+
+		+ SHorizontalBox::Slot()
+		.AutoWidth()
+		.Padding(0, 5, 0, 0)
+		[
+			SNew(SCheckBox)
+			.IsChecked_Lambda([this]
+			{
+				return ViewModel.Get()->IsVerbose() ? ECheckBoxState::Checked : ECheckBoxState::Unchecked; 
+			})
+			.OnCheckStateChanged_Lambda([this](ECheckBoxState State)
+			{
+				ViewModel.Get()->SetVerbose(State == ECheckBoxState::Checked);
+				UpdateViews();
+			})
+			[
+				SNew(STextBlock)
+				.Text(LOCTEXT("PoseSearchDebuggerShowVerbose", "Channels Breakdown"))
+			]
+		]
+	];
+
+	return 
+		SNew(SSplitter)
+		.Orientation(Orient_Horizontal)
+		.ResizeMode(ESplitterResizeMode::Fill)
+	
+		// Database view
+		+ SSplitter::Slot()
+		.Value(0.65f)
+		[
+			SNew(SVerticalBox)
+			
+			+ SVerticalBox::Slot()
+			.AutoHeight()
+			[
+				ReturnButtonView
+			]
+			
+			+ SVerticalBox::Slot()
+			[
+				SAssignNew(DatabaseView, SDebuggerDatabaseView)
+				.Parent(SharedThis(this))
+				.OnPoseSelectionChanged(this, &SDebuggerView::OnPoseSelectionChanged)
+			]
+		]
+
+		// Details panel view
+		+ SSplitter::Slot()
+		.Value(0.35f)
+		[
+			SAssignNew(DetailsView, SDebuggerDetailsView)
+			.Parent(SharedThis(this))
+		];
+}
+
+FName SDebuggerView::GetName() const
+{
+	static const FName DebuggerName("PoseSearchDebugger");
+	return DebuggerName;
+}
+
+uint64 SDebuggerView::GetObjectId() const
+{
+	return AnimInstanceId;
+}
+
+SDebuggerView::~SDebuggerView()
+{
+	OnViewClosed.Execute(AnimInstanceId);
+}
+
+} // namespace UE::PoseSearch
+
+#undef LOCTEXT_NAMESPACE
diff --git a/Engine/Plugins/Experimental/Animation/PoseSearch/Source/Editor/Private/PoseSearchDebuggerView.h b/Engine/Plugins/Experimental/Animation/PoseSearch/Source/Editor/Private/PoseSearchDebuggerView.h
new file mode 100644
--- /dev/null
+++ b/Engine/Plugins/Experimental/Animation/PoseSearch/Source/Editor/Private/PoseSearchDebuggerView.h
@@ -0,0 +1,171 @@
+// Copyright Epic Games, Inc. All Rights Reserved.
+
+#pragma once
+
+#include "IRewindDebuggerView.h"
+
+class IDetailsView;
+class SVerticalBox;
+class SHorizontalBox;
+class SSplitter;
+class SWidgetSwitcher;
+class UPoseSearchDatabase;
+class UPoseSearchDebuggerReflection;
+
+namespace UE::PoseSearch
+{
+
+class FDebuggerDatabaseRowData;
+class FDebuggerViewModel;
+struct FTraceMotionMatchingStateMessage;
+class SDebuggerDatabaseView;
+
+/**
+ * Details panel view widget of the PoseSearch debugger
+ */
+class SDebuggerDetailsView : public SCompoundWidget
+{
+	SLATE_BEGIN_ARGS(SDebuggerDetailsView) {}
+		SLATE_ARGUMENT(TWeakPtr<class SDebuggerView>, Parent)
+	SLATE_END_ARGS()
+
+	SDebuggerDetailsView() = default;
+	virtual ~SDebuggerDetailsView() override;
+
+	void Construct(const FArguments& InArgs);
+	void Update(const FTraceMotionMatchingStateMessage& State) const;
+
+	/** Get a const version of our reflection object */
+	const TObjectPtr<UPoseSearchDebuggerReflection> GetReflection() const { return Reflection; }
+
+private:
+	/** Update our details view object with new state information */
+	void UpdateReflection(const FTraceMotionMatchingStateMessage& State) const;
+
+	TWeakPtr<SDebuggerView> ParentDebuggerViewPtr;
+
+	/** Details widget constructed for the MM node */
+	TSharedPtr<IDetailsView> Details;
+
+	/** Last updated reflection data relative to MM state */
+	TObjectPtr<UPoseSearchDebuggerReflection> Reflection = nullptr;
+};
+
+/** Callback to relay closing of the view to destroy the debugger instance */
+DECLARE_DELEGATE_OneParam(FOnViewClosed, uint64 AnimInstanceId);
+
+/**
+ * Entire view of the PoseSearch debugger, containing all sub-widgets
+ */
+class SDebuggerView : public IRewindDebuggerView
+{
+public:
+	SLATE_BEGIN_ARGS(SDebuggerView){}
+		SLATE_ATTRIBUTE(TSharedPtr<FDebuggerViewModel>, ViewModel)
+		SLATE_EVENT(FOnViewClosed, OnViewClosed)
+	SLATE_END_ARGS()
+
+	SDebuggerView() = default;
+    virtual ~SDebuggerView() override;
+
+	void Construct(const FArguments& InArgs, uint64 InAnimInstanceId);
+	virtual void SetTimeMarker(double InTimeMarker) override;
+	virtual void Tick(const FGeometry& AllottedGeometry, const double InCurrentTime, const float InDeltaTime) override;
+	virtual FName GetName() const override;
+	virtual uint64 GetObjectId() const override;
+
+	TSharedPtr<FDebuggerViewModel> GetViewModel() const { return ViewModel.Get(); }
+	TArray<TSharedRef<FDebuggerDatabaseRowData>> GetSelectedDatabaseRows() const;
+
+private:
+	/** Called each frame to draw features of the query vector & database selections */
+	void DrawFeatures(const UWorld& DebuggerWorld, const FTraceMotionMatchingStateMessage& State, const FTransform& Transform, const USkinnedMeshComponent* Mesh) const;
+	
+	/** Check if a node selection was made, true if a node is selected */
+	bool UpdateNodeSelection();
+
+	/** Update the database and details views */
+	void UpdateViews() const;
+
+	void DrawVisualization() const;
+	
+	/** Returns an int32 appropriate to the index of our widget selector */
+	int32 SelectView() const;
+
+	/** Callback when a button in the selection view is clicked */
+	FReply OnUpdateNodeSelection(int32 InSelectedNodeId);
+
+	void OnPoseSelectionChanged(const UPoseSearchDatabase* Database, int32 PoseIdx, float Time);
+
+	/** Button interaction to toggle play / stop of the asset */
+	FReply TogglePlaySelectedAssets() const;
+
+	/** Generates the message view relaying that there is no data */
+	TSharedRef<SWidget> GenerateNoDataMessageView();
+
+	/** Generates the return button to go back to the selection mode */
+	TSharedRef<SHorizontalBox> GenerateReturnButtonView();
+
+	/** Generates the entire node debugger widget, including database and details view */
+	TSharedRef<SWidget> GenerateNodeDebuggerView();
+
+	/** Pointer to the debugger instance / model for this view */
+	TAttribute<TSharedPtr<FDebuggerViewModel>> ViewModel;
+	
+	/** Destroy the debugger instanced when closed */
+	FOnViewClosed OnViewClosed;
+
+	/** Active node being debugged */
+	int32 SelectedNodeId = -1;
+
+	/** Database view of the motion matching node */
+	TSharedPtr<SDebuggerDatabaseView> DatabaseView;
+	
+	/** Details panel for introspecting the motion matching node */
+	TSharedPtr<SDebuggerDetailsView> DetailsView;
+	
+	/** Node debugger view hosts the above two views */
+	TSharedPtr<SSplitter> NodeDebuggerView;
+
+	/** Selection view before node is selected */
+	TSharedPtr<SVerticalBox> SelectionView;
+	
+	/** Gray box occluding the debugger view when simulating */
+	TSharedPtr<SVerticalBox> SimulatingView;
+
+	/** Used to switch between views in the switcher, int32 maps to index in the SWidgetSwitcher */
+	enum ESwitcherViewType : int32
+	{
+		Selection = 0,
+		Debugger = 1,
+		StoppedMsg = 2,
+		RecordingMsg = 3,
+		NoDataMsg = 4
+	} SwitcherViewType = StoppedMsg;
+	
+	/** Contains all the above, switches between them depending on context */
+	TSharedPtr<SWidgetSwitcher> Switcher;
+
+	/** Contains the switcher, the entire debugger view */
+	TSharedPtr<SVerticalBox> DebuggerView;
+
+	/** AnimInstance this view was created for */
+	uint64 AnimInstanceId = 0;
+
+	/** Current position of the time marker */
+	double TimeMarker = -1.0;
+
+	/** Previous position of the time marker */
+	double PreviousTimeMarker = -1.0;
+
+	/** Tracks if the current time has been updated yet (delayed) */
+	bool bUpdated = false;
+
+	/** Tracks number of consecutive frames, once it reaches threshold it will update the view */
+	int32 CurrentConsecutiveFrames = 0;
+
+	/** Once the frame count has reached this value, an update will trigger for the view */
+	static constexpr int32 ConsecutiveFramesUpdateThreshold = 10;
+};
+
+} // namespace UE::PoseSearch
diff --git a/Engine/Plugins/Experimental/Animation/PoseSearch/Source/Editor/Private/PoseSearchDebuggerViewModel.cpp b/Engine/Plugins/Experimental/Animation/PoseSearch/Source/Editor/Private/PoseSearchDebuggerViewModel.cpp
new file mode 100644
--- /dev/null
+++ b/Engine/Plugins/Experimental/Animation/PoseSearch/Source/Editor/Private/PoseSearchDebuggerViewModel.cpp
@@ -0,0 +1,537 @@
+// Copyright Epic Games, Inc. All Rights Reserved.
+
+#include "PoseSearchDebuggerViewModel.h"
+#include "Animation/MirrorDataTable.h"
+#include "Engine/SkeletalMesh.h"
+#include "IAnimationProvider.h"
+#include "IGameplayProvider.h"
+#include "IRewindDebugger.h"
+#include "PoseSearch/PoseSearchDerivedData.h"
+#include "Trace/PoseSearchTraceProvider.h"
+
+namespace UE::PoseSearch
+{
+
+FDebuggerViewModel::FDebuggerViewModel(uint64 InAnimInstanceId)
+	: AnimInstanceId(InAnimInstanceId)
+{
+	Skeletons.AddDefaulted(ESkeletonIndex::Num);
+}
+
+FDebuggerViewModel::~FDebuggerViewModel()
+{
+	for (FSkeleton& Skeleton : Skeletons)
+	{
+		if (Skeleton.Actor.IsValid())
+		{
+			Skeleton.Actor->Destroy();
+		}
+	}
+
+	Skeletons.Empty();
+}
+
+const FTraceMotionMatchingStateMessage* FDebuggerViewModel::GetMotionMatchingState() const
+{
+	return ActiveMotionMatchingState;
+}
+
+const UPoseSearchDatabase* FDebuggerViewModel::GetCurrentDatabase() const
+{
+	return ActiveMotionMatchingState ? ActiveMotionMatchingState->GetCurrentDatabase() : nullptr;
+}
+
+const UPoseSearchSearchableAsset* FDebuggerViewModel::GetSearchableAsset() const
+{
+	if (ActiveMotionMatchingState)
+	{
+		return FTraceMotionMatchingState::GetObjectFromId<UPoseSearchSearchableAsset>(ActiveMotionMatchingState->SearchableAssetId);
+	}
+
+	return nullptr;
+}
+
+void FDebuggerViewModel::ShowSelectedSkeleton(const UPoseSearchDatabase* Database, int32 DbPoseIdx, float Time)
+{
+	UPoseSearchMeshComponent* Component = Skeletons[SelectedPose].Component.Get();
+	if (!Component)
+	{
+		return;
+	}
+
+	if (!FAsyncPoseSearchDatabasesManagement::RequestAsyncBuildIndex(Database, ERequestAsyncBuildFlag::ContinueRequest))
+	{
+		return;
+	}
+
+	const FPoseSearchIndex& SearchIndex = Database->GetSearchIndex();
+	const FPoseSearchIndexAsset& IndexAsset = SearchIndex.GetAssetForPose(DbPoseIdx);
+	
+	Component->ResetToStart(); 
+	bSelecting = true;
+	
+	Skeletons[SelectedPose].Type = IndexAsset.Type;
+	Skeletons[SelectedPose].Time = Time;
+	Skeletons[SelectedPose].bMirrored = IndexAsset.bMirrored;
+	Skeletons[SelectedPose].SourceDatabase = Database;
+	Skeletons[SelectedPose].AssetIdx = IndexAsset.SourceAssetIdx;
+	Skeletons[SelectedPose].BlendParameters = IndexAsset.BlendParameters;
+}
+
+void FDebuggerViewModel::ClearSelectedSkeleton()
+{
+	bSelecting = false;
+}
+
+const TArray<int32>* FDebuggerViewModel::GetNodeIds() const
+{
+	return &NodeIds;
+}
+
+int32 FDebuggerViewModel::GetNodesNum() const
+{
+	return MotionMatchingStates.Num();
+}
+
+const FTransform* FDebuggerViewModel::GetRootTransform() const
+{
+	return RootTransform;
+}
+
+bool FDebuggerViewModel::HasSearchableAssetChanged() const
+{
+	uint64 NewSearchableAssetId = ActiveMotionMatchingState ? ActiveMotionMatchingState->SearchableAssetId : 0;
+	return NewSearchableAssetId != SearchableAssetId;
+}
+
+void FDebuggerViewModel::OnUpdate()
+{
+	if (!bSkeletonsInitialized)
+	{
+		UWorld* World = RewindDebugger.Get()->GetWorldToVisualize();
+		for (FSkeleton& Skeleton : Skeletons)
+		{
+			FActorSpawnParameters ActorSpawnParameters;
+			ActorSpawnParameters.bHideFromSceneOutliner = false;
+			ActorSpawnParameters.ObjectFlags |= RF_Transient;
+			Skeleton.Actor = World->SpawnActor<AActor>(ActorSpawnParameters);
+			Skeleton.Actor->SetActorLabel(TEXT("PoseSearch"));
+			Skeleton.Component = NewObject<UPoseSearchMeshComponent>(Skeleton.Actor.Get());
+			Skeleton.Actor->AddInstanceComponent(Skeleton.Component.Get());
+			Skeleton.Component->RegisterComponentWithWorld(World);
+		}
+		FWorldDelegates::OnWorldCleanup.AddRaw(this, &FDebuggerViewModel::OnWorldCleanup);
+		bSkeletonsInitialized = true;
+	}
+
+	UpdateFromTimeline();
+}
+
+void FDebuggerViewModel::OnUpdateNodeSelection(int32 InNodeId)
+{
+	if (InNodeId == INDEX_NONE)
+	{
+		return;
+	}
+
+	ActiveMotionMatchingState = nullptr;
+
+	// Find node in all motion matching states this frame
+	const int32 NodesNum = NodeIds.Num();
+	for (int i = 0; i < NodesNum; ++i)
+	{
+		if (NodeIds[i] == InNodeId)
+		{
+			ActiveMotionMatchingState = MotionMatchingStates[i];
+			break;
+		}
+	}
+
+	if (ActiveMotionMatchingState)
+	{
+		const UPoseSearchDatabase* CurrentDatabase = ActiveMotionMatchingState->GetCurrentDatabase();
+		if (FAsyncPoseSearchDatabasesManagement::RequestAsyncBuildIndex(CurrentDatabase, ERequestAsyncBuildFlag::ContinueRequest))
+		{
+			const FPoseSearchIndex& CurrentSearchIndex = CurrentDatabase->GetSearchIndex();
+			int32 CurrentPoseIdx = ActiveMotionMatchingState->GetCurrentDatabasePoseIndex();
+			if (const FPoseSearchIndexAsset* IndexAsset = CurrentSearchIndex.GetAssetForPoseSafe(CurrentPoseIdx))
+			{
+				Skeletons[Asset].Type = IndexAsset->Type;
+				Skeletons[Asset].bMirrored = IndexAsset->bMirrored;
+				Skeletons[Asset].SourceDatabase = CurrentDatabase;
+				Skeletons[Asset].AssetIdx = IndexAsset->SourceAssetIdx;
+				Skeletons[Asset].BlendParameters = IndexAsset->BlendParameters;
+			}
+		}
+	}
+
+	uint64 NewSearchableAssetId = ActiveMotionMatchingState ? ActiveMotionMatchingState->SearchableAssetId : 0;
+	if (NewSearchableAssetId != SearchableAssetId)
+	{
+		ClearSelectedSkeleton();
+		SearchableAssetId = NewSearchableAssetId;
+	}
+}
+
+void FDebuggerViewModel::UpdatePoseSearchContext(UPoseSearchMeshComponent::FUpdateContext& InOutContext, const FSkeleton& Skeleton) const
+{
+	const FInstancedStruct* DatabaseAssetStruct = Skeleton.GetAnimationAsset();
+	if (DatabaseAssetStruct)
+	{
+		const FPoseSearchDatabaseAnimationAssetBase* DatabaseAsset = DatabaseAssetStruct->GetPtr<FPoseSearchDatabaseAnimationAssetBase>();
+		if (DatabaseAsset)
+		{
+			InOutContext.Type = DatabaseAsset->GetSearchIndexType();
+			InOutContext.StartTime = Skeletons[SelectedPose].Time;
+			InOutContext.Time = Skeletons[SelectedPose].Time;
+			InOutContext.bMirrored = Skeletons[SelectedPose].bMirrored;
+			InOutContext.bLoop = DatabaseAsset->IsLooping();
+		}
+
+		if (const FPoseSearchDatabaseSequence* DatabaseSequence = DatabaseAssetStruct->GetPtr<FPoseSearchDatabaseSequence>())
+		{
+			InOutContext.SequenceBase = DatabaseSequence->Sequence;
+		}
+		else if (const FPoseSearchDatabaseAnimComposite* DatabaseAnimComposite = DatabaseAssetStruct->GetPtr<FPoseSearchDatabaseAnimComposite>())
+		{
+			InOutContext.SequenceBase = DatabaseAnimComposite->AnimComposite;
+		}
+		else if (const FPoseSearchDatabaseBlendSpace* DatabaseBlendSpace = DatabaseAssetStruct->GetPtr<FPoseSearchDatabaseBlendSpace>())
+		{
+			InOutContext.BlendSpace = DatabaseBlendSpace->BlendSpace;
+			InOutContext.BlendParameters = Skeletons[SelectedPose].BlendParameters;
+		}
+		else
+		{
+			checkNoEntry();
+		}
+	}
+}
+
+void FDebuggerViewModel::OnDraw(FSkeletonDrawParams& DrawParams)
+{
+	const UPoseSearchDatabase* CurrentDatabase = GetCurrentDatabase();
+	if (!CurrentDatabase)
+	{
+		return;
+	}
+
+	// Returns if it is to be drawn this frame
+	auto SetDrawSkeleton = [this](UPoseSearchMeshComponent* InComponent, bool bDraw)
+	{
+		if (InComponent && InComponent->RequiredBones.IsValid())
+		{
+			const bool bIsDrawingSkeleton = InComponent->ShouldDrawDebugSkeleton();
+			if (bIsDrawingSkeleton != bDraw)
+			{
+				InComponent->SetDrawDebugSkeleton(bDraw);
+			}
+			InComponent->MarkRenderStateDirty();
+		}
+	};
+	const bool bDrawActivePose = EnumHasAnyFlags(DrawParams.Flags, ESkeletonDrawFlags::ActivePose);
+	SetDrawSkeleton(Skeletons[ActivePose].Component.Get(), bDrawActivePose);
+	// If flag is set and we are currently in a valid drawing state
+	const bool bDrawSelectedPose = EnumHasAnyFlags(DrawParams.Flags, ESkeletonDrawFlags::SelectedPose) && bSelecting;
+	SetDrawSkeleton(Skeletons[SelectedPose].Component.Get(), bDrawSelectedPose);
+
+	FillCompactPoseAndComponentRefRotations();
+
+	UPoseSearchMeshComponent::FUpdateContext UpdateContext;
+
+	UpdateContext.MirrorDataTable = CurrentDatabase->Schema->MirrorDataTable;
+	UpdateContext.CompactPoseMirrorBones = &CompactPoseMirrorBones;
+	UpdateContext.ComponentSpaceRefRotations = &ComponentSpaceRefRotations;
+
+	if (bDrawSelectedPose)
+	{
+		UPoseSearchMeshComponent* Component = Skeletons[SelectedPose].Component.Get();
+		if (Component && Component->RequiredBones.IsValid())
+		{
+			UpdatePoseSearchContext(UpdateContext, Skeletons[SelectedPose]);
+
+			if (UpdateContext.Type != ESearchIndexAssetType::Invalid)
+			{
+				Component->UpdatePose(UpdateContext);
+			}
+		}
+	}
+
+	const bool bDrawAsset = EnumHasAnyFlags(DrawParams.Flags, ESkeletonDrawFlags::Asset);
+	if (bDrawAsset && AssetData.bActive)
+	{
+		UPoseSearchMeshComponent* Component = Skeletons[Asset].Component.Get();
+		if (Component && Component->RequiredBones.IsValid())
+		{
+			SetDrawSkeleton(Component, true);
+
+			UpdatePoseSearchContext(UpdateContext, Skeletons[SelectedPose]);
+
+			if (UpdateContext.Type != ESearchIndexAssetType::Invalid)
+			{
+				Component->UpdatePose(UpdateContext);
+			}
+		}
+	}
+}
+
+void FDebuggerViewModel::UpdateFromTimeline()
+{
+	NodeIds.Empty();
+	MotionMatchingStates.Empty();
+	SkeletalMeshComponentId = 0;
+
+	// Get provider and validate
+	const TraceServices::IAnalysisSession* Session = RewindDebugger.Get()->GetAnalysisSession();
+	TraceServices::FAnalysisSessionReadScope SessionReadScope(*Session);
+
+	const FTraceProvider* PoseSearchProvider = Session->ReadProvider<FTraceProvider>(FTraceProvider::ProviderName);
+	const IAnimationProvider* AnimationProvider = Session->ReadProvider<IAnimationProvider>("AnimationProvider");
+	const IGameplayProvider* GameplayProvider = Session->ReadProvider<IGameplayProvider>("GameplayProvider");
+	if (!(PoseSearchProvider && AnimationProvider && GameplayProvider))
+	{
+		return;
+	}
+	const double TraceTime = RewindDebugger.Get()->CurrentTraceTime();
+	TraceServices::FFrame Frame;
+	ReadFrameProvider(*Session).GetFrameFromTime(TraceFrameType_Game, TraceTime, Frame);
+	PoseSearchProvider->EnumerateMotionMatchingStateTimelines(AnimInstanceId, [&](const FTraceProvider::FMotionMatchingStateTimeline& InTimeline)
+	{
+		const FTraceMotionMatchingStateMessage* Message = nullptr;
+
+		InTimeline.EnumerateEvents(Frame.StartTime, Frame.EndTime, [&Message](double InStartTime, double InEndTime, const FTraceMotionMatchingStateMessage& InMessage)
+		{
+			Message = &InMessage;
+			return TraceServices::EEventEnumerate::Stop;
+		});
+		if (Message)
+		{
+			NodeIds.Add(Message->NodeId);
+			MotionMatchingStates.Add(Message);
+			SkeletalMeshComponentId = Message->SkeletalMeshComponentId;
+		}
+	});
+	/** No active motion matching state as no messages were read */
+	if (SkeletalMeshComponentId == 0)
+	{
+		return;
+	}
+	AnimationProvider->ReadSkeletalMeshPoseTimeline(SkeletalMeshComponentId, [&](const IAnimationProvider::SkeletalMeshPoseTimeline& TimelineData, bool bHasCurves)
+	{
+		TimelineData.EnumerateEvents(Frame.StartTime, Frame.EndTime, [&](double InStartTime, double InEndTime, uint32 InDepth, const FSkeletalMeshPoseMessage& PoseMessage) -> TraceServices::EEventEnumerate
+		{
+			// Update root transform
+			RootTransform = &PoseMessage.ComponentToWorld;
+			const FSkeletalMeshInfo* SkeletalMeshInfo = AnimationProvider->FindSkeletalMeshInfo(PoseMessage.MeshId);
+			const FObjectInfo* SkeletalMeshObjectInfo = GameplayProvider->FindObjectInfo(PoseMessage.MeshId);
+			if (!SkeletalMeshInfo || !SkeletalMeshObjectInfo)
+			{
+
+				return TraceServices::EEventEnumerate::Stop;
+			}
+			UPoseSearchMeshComponent* ActiveComponent = Skeletons[ActivePose].Component.Get();
+			UPoseSearchMeshComponent* SelectedComponent = Skeletons[SelectedPose].Component.Get();
+			UPoseSearchMeshComponent* AssetComponent = Skeletons[Asset].Component.Get();
+			USkeletalMesh* SkeletalMesh = TSoftObjectPtr<USkeletalMesh>(FSoftObjectPath(SkeletalMeshObjectInfo->PathName)).LoadSynchronous();
+			if (SkeletalMesh)
+			{
+				ActiveComponent->SetSkinnedAssetAndUpdate(SkeletalMesh, true);
+				SelectedComponent->SetSkinnedAssetAndUpdate(SkeletalMesh, true);
+				AssetComponent->SetSkinnedAssetAndUpdate(SkeletalMesh, true);
+			}
+			FTransform ComponentWorldTransform;
+			// Active skeleton is simply the traced bone transforms
+			AnimationProvider->GetSkeletalMeshComponentSpacePose(PoseMessage, *SkeletalMeshInfo, ComponentWorldTransform, ActiveComponent->GetEditableComponentSpaceTransforms());
+			ActiveComponent->Initialize(ComponentWorldTransform);
+			ActiveComponent->SetDebugDrawColor(FLinearColor::Green);
+			SelectedComponent->SetDebugDrawColor(FLinearColor::Blue);
+			SelectedComponent->Initialize(ComponentWorldTransform);
+			AssetComponent->SetDebugDrawColor(FLinearColor::Red);
+			AssetComponent->Initialize(ComponentWorldTransform);
+
+			return TraceServices::EEventEnumerate::Stop;
+		});
+	});
+}
+
+void FDebuggerViewModel::UpdateAsset()
+{
+	// @todo: expose those parameters
+	static float MAX_DISTANCE_RANGE = 200.f;
+	static float MAX_TIME_RANGE = 2.f;
+
+	const UPoseSearchDatabase* Database = GetCurrentDatabase();
+	if (!Database || !IsPlayingSelections())
+	{
+		return;
+	}
+
+	FSkeleton& AssetSkeleton = Skeletons[Asset];
+	UPoseSearchMeshComponent* Component = AssetSkeleton.Component.Get();
+
+	auto RestartAsset = [&]()
+	{
+		Component->ResetToStart();
+		AssetData.AccumulatedTime = 0.0;
+		AssetSkeleton.Time = AssetData.StartTime;
+	};
+
+	const FInstancedStruct* AnimationAssetStruct = AssetSkeleton.GetAnimationAsset();
+	if (!AnimationAssetStruct)
+	{
+		checkNoEntry();
+		return;
+	}
+
+	const FPoseSearchDatabaseAnimationAssetBase* DatabaseAsset = AnimationAssetStruct->GetPtr<FPoseSearchDatabaseAnimationAssetBase>();
+	if (!DatabaseAsset)
+	{
+		checkNoEntry();
+		return;
+	}
+
+	const UAnimationAsset* AnimAsset = DatabaseAsset->GetAnimationAsset();
+	const bool bAssetLooping = DatabaseAsset->IsLooping();
+
+	const float DT = static_cast<float>(FApp::GetDeltaTime()) * AssetPlayRate;
+	const float PlayLength = AnimAsset->GetPlayLength();
+	const bool bExceededDistanceHorizon = Component->LastRootMotionDelta.GetTranslation().Size() > MAX_DISTANCE_RANGE;
+	const bool bExceededTimeHorizon = (AssetSkeleton.Time - AssetData.StartTime) > MAX_TIME_RANGE;
+	const bool bExceededHorizon = bExceededDistanceHorizon && bExceededTimeHorizon;
+	if (bAssetLooping)
+	{
+		if (bExceededHorizon)
+		{
+			// Delay before restarting the asset to give the user some idea of where it would land
+			if (AssetData.AccumulatedTime > AssetData.StopDuration)
+			{
+				RestartAsset();
+			}
+			else
+			{
+				AssetData.AccumulatedTime += DT;
+			}
+			return;
+		}
+
+		AssetSkeleton.Time += DT;
+		AssetData.AccumulatedTime += DT;
+	}
+	else
+	{
+		// Used to cap the asset, but avoid modding when updating the pose
+		static constexpr float LengthOffset = 0.001f;
+		const bool bFinishedAsset = AssetSkeleton.Time >= PlayLength - LengthOffset;
+
+		// Asset player reached end of clip or reached distance horizon of trajectory vector
+		if (bFinishedAsset || bExceededHorizon)
+		{
+			// Delay before restarting the asset to give the user some idea of where it would land
+			if (AssetData.AccumulatedTime > AssetData.StopDuration)
+			{
+				RestartAsset();
+			}
+			else
+			{
+				AssetData.AccumulatedTime += DT;
+			}
+		}
+		else
+		{
+			// If we haven't finished, update the play time capped by the anim asset (not looping)
+			AssetSkeleton.Time += DT;
+		}
+	}
+}
+
+const USkinnedMeshComponent* FDebuggerViewModel::GetMeshComponent() const
+{
+	if (Skeletons.Num() > FDebuggerViewModel::Asset)
+	{
+		return Skeletons[FDebuggerViewModel::Asset].Component.Get();
+	}
+	return nullptr;
+}
+
+void FDebuggerViewModel::FillCompactPoseAndComponentRefRotations()
+{
+	bool bResetMirrorBonesAndRotations = true;
+	const UPoseSearchDatabase* Database = GetCurrentDatabase();
+	if (Database)
+	{
+		UMirrorDataTable* MirrorDataTable = Database->Schema->MirrorDataTable;
+		if (MirrorDataTable != nullptr && Skeletons[ActivePose].Component->RequiredBones.IsValid())
+		{
+			if (CompactPoseMirrorBones.Num() == 0 || ComponentSpaceRefRotations.Num() == 0)
+			{
+				MirrorDataTable->FillCompactPoseAndComponentRefRotations(
+					Skeletons[ActivePose].Component->RequiredBones,
+					CompactPoseMirrorBones,
+					ComponentSpaceRefRotations);
+			}
+			bResetMirrorBonesAndRotations = false;
+		}
+	}
+
+	if (bResetMirrorBonesAndRotations)
+	{
+		CompactPoseMirrorBones.Reset();
+		ComponentSpaceRefRotations.Reset();
+	}
+}
+
+void FDebuggerViewModel::PlaySelection(int32 PoseIdx, float Time)
+{
+	UPoseSearchMeshComponent* Component = Skeletons[Asset].Component.Get();
+	if (!Component)
+	{
+		return;
+	}
+
+	const UPoseSearchDatabase* Database = GetCurrentDatabase();
+	if (!FAsyncPoseSearchDatabasesManagement::RequestAsyncBuildIndex(Database, ERequestAsyncBuildFlag::ContinueRequest))
+	{
+		return;
+	}
+	
+	const FPoseSearchIndexAsset& IndexAsset = Database->GetSearchIndex().GetAssetForPose(PoseIdx);
+	Component->ResetToStart();
+	
+	Skeletons[Asset].Type = IndexAsset.Type;
+	Skeletons[Asset].AssetIdx = IndexAsset.SourceAssetIdx;
+	Skeletons[Asset].Time = Time;
+	Skeletons[Asset].bMirrored = IndexAsset.bMirrored;
+	Skeletons[Asset].BlendParameters = IndexAsset.BlendParameters;
+
+	AssetData.StartTime = Time;
+	AssetData.AccumulatedTime = 0.0f;
+	AssetData.bActive = true;
+}
+
+void FDebuggerViewModel::StopSelection()
+{
+	UPoseSearchMeshComponent* Component = Skeletons[Asset].Component.Get();
+	if (!Component)
+	{
+		return;
+	}
+
+	AssetData = {};
+	// @TODO: Make more functionality rely on checking if it should draw the asset
+	Component->SetDrawDebugSkeleton(false);
+}
+void FDebuggerViewModel::OnWorldCleanup(UWorld* InWorld, bool bSessionEnded, bool bCleanupResources)
+{
+	bSkeletonsInitialized = false;
+}
+
+const FInstancedStruct* FDebuggerViewModel::FSkeleton::GetAnimationAsset() const
+{
+	if (SourceDatabase.IsValid() && SourceDatabase->AnimationAssets.IsValidIndex(AssetIdx))
+	{
+		return &SourceDatabase->GetAnimationAssetStruct(AssetIdx);
+	}
+	return nullptr;
+}
+
+} // namespace UE::PoseSearch
diff --git a/Engine/Plugins/Experimental/Animation/PoseSearch/Source/Editor/Private/PoseSearchDebuggerViewModel.h b/Engine/Plugins/Experimental/Animation/PoseSearch/Source/Editor/Private/PoseSearchDebuggerViewModel.h
new file mode 100644
--- /dev/null
+++ b/Engine/Plugins/Experimental/Animation/PoseSearch/Source/Editor/Private/PoseSearchDebuggerViewModel.h
@@ -0,0 +1,207 @@
+// Copyright Epic Games, Inc. All Rights Reserved.
+
+#pragma once
+
+#include "CoreMinimal.h"
+#include "Misc/Attribute.h"
+#include "PoseSearchMeshComponent.h"
+#include "Templates/SharedPointer.h"
+#include "UObject/WeakObjectPtrTemplates.h"
+
+class IRewindDebugger;
+
+namespace UE::PoseSearch
+{
+
+struct FTraceMotionMatchingStateMessage;
+
+/** Draw flags for the view's debug draw */
+enum class ESkeletonDrawFlags : uint32
+{
+	None			= 0,
+	ActivePose		= 1 << 0,
+	SelectedPose	= 1 << 1,
+	Asset			= 1 << 2,
+};
+ENUM_CLASS_FLAGS(ESkeletonDrawFlags)
+
+struct FSkeletonDrawParams
+{
+	ESkeletonDrawFlags Flags = ESkeletonDrawFlags::None;
+};
+
+class FDebuggerViewModel : public TSharedFromThis<FDebuggerViewModel>
+{
+public:
+	explicit FDebuggerViewModel(uint64 InAnimInstanceId);
+	virtual ~FDebuggerViewModel();
+
+	// Used for view callbacks
+    const FTraceMotionMatchingStateMessage* GetMotionMatchingState() const;
+	const UPoseSearchDatabase* GetCurrentDatabase() const;
+	const UPoseSearchSearchableAsset* GetSearchableAsset() const;
+	const TArray<int32>* GetNodeIds() const;
+	int32 GetNodesNum() const;
+	const FTransform* GetRootTransform() const;
+
+	/** Checks if Update must be called */
+	bool HasSearchableAssetChanged() const;
+
+	/** Update motion matching states for frame */
+	void OnUpdate();
+	
+	/** Updates active motion matching state based on node selection */
+	void OnUpdateNodeSelection(int32 InNodeId);
+
+	/** Updates internal Skeletal Mesh Component depending on input */
+	void OnDraw(FSkeletonDrawParams& DrawParams);
+	
+	/** Sets the selected pose skeleton*/
+	void ShowSelectedSkeleton(const UPoseSearchDatabase* Database, int32 DbPoseIdx, float Time);
+	
+	/** Clears the selected pose skeleton */
+	void ClearSelectedSkeleton();
+	
+	/** Plays the selected row upon button press */
+	void PlaySelection(int32 PoseIdx, float Time);
+
+	/** Stops the playing selection upon button press */
+	void StopSelection();
+	
+	/** If there is an asset selection playing */
+	bool IsPlayingSelections() const { return AssetData.bActive; }
+	
+	/** Play Rate of the asset selection player */
+	void ChangePlayRate(float PlayRate) { AssetPlayRate = PlayRate; }
+	
+	/** Current play rate of the asset selection player */
+	float GetPlayRate() const { return AssetPlayRate; }
+	
+	void SetVerbose(bool bVerbose) { bIsVerbose = bVerbose; }
+
+	bool IsVerbose() const { return bIsVerbose; }
+
+	/** Callback to reset debug skeletons for the active world */
+	void OnWorldCleanup(UWorld* InWorld, bool bSessionEnded, bool bCleanupResources);
+
+	/** Updates the current playing asset */
+	void UpdateAsset();
+
+	const USkinnedMeshComponent* GetMeshComponent() const;
+
+private:
+	/** Debug visualization skeleton actor */
+	struct FSkeleton
+	{
+		/** Actor object for the skeleton */
+		TWeakObjectPtr<AActor> Actor = nullptr;
+
+		/** Derived skeletal mesh for setting the skeleton in the scene */
+		TWeakObjectPtr<UPoseSearchMeshComponent> Component = nullptr;
+
+		/** Type of the asset being played */
+		ESearchIndexAssetType Type = ESearchIndexAssetType::Invalid;
+
+		/** Source database for this skeleton  */
+		TWeakObjectPtr<const UPoseSearchDatabase> SourceDatabase = nullptr;
+
+		/** Source asset for this skeleton */
+		int32 AssetIdx = 0;
+
+		/** Time in the sequence this skeleton is accessing */
+		float Time = 0.0f;
+
+		/** If this asset should be mirrored */
+		bool bMirrored = false;
+
+		/** Blend Parameters if asset is a BlendSpace */
+		FVector BlendParameters = FVector::Zero();
+
+		const FInstancedStruct* GetAnimationAsset() const;
+	};
+
+	/** Update the list of states for this frame */
+	void UpdateFromTimeline();
+
+	/** Populates arrays used for mirroring the animation pose */
+	void FillCompactPoseAndComponentRefRotations();
+
+	void UpdatePoseSearchContext(UPoseSearchMeshComponent::FUpdateContext& InOutContext, const FSkeleton& Skeleton) const;
+	
+	/** List of all Node IDs associated with motion matching states */
+	TArray<int32> NodeIds;
+	
+	/** List of all updated motion matching states per node */
+	TArray<const FTraceMotionMatchingStateMessage*> MotionMatchingStates;
+	
+	/** Currently active motion matching state based on node selection in the view */
+	const FTraceMotionMatchingStateMessage* ActiveMotionMatchingState = nullptr;
+
+	/** Active motion matching state's searchable asset */
+	uint64 SearchableAssetId = 0;
+
+	/** Current Skeletal Mesh Component Id for the AnimInstance */
+	uint64 SkeletalMeshComponentId = 0;
+
+	/** Currently active root transform on the skeletal mesh */
+	const FTransform* RootTransform = nullptr;
+
+	/** Pointer to the active rewind debugger in the scene */
+	TAttribute<const IRewindDebugger*> RewindDebugger;
+
+	/** Anim Instance associated with this debugger instance */
+	uint64 AnimInstanceId = 0;
+
+	/** Compact pose format of Mirror Bone Map */
+	TCustomBoneIndexArray<FCompactPoseBoneIndex, FCompactPoseBoneIndex> CompactPoseMirrorBones;
+
+	/** Pre-calculated component space rotations of reference pose */
+	TCustomBoneIndexArray<FQuat, FCompactPoseBoneIndex> ComponentSpaceRefRotations;
+
+
+	/** Index for each type of skeleton we store for debug visualization */
+	enum ESkeletonIndex
+	{
+		ActivePose = 0,
+		SelectedPose,
+		Asset,
+
+		Num
+	};
+
+	/** Skeleton container for each type */
+	TArray<FSkeleton, TFixedAllocator<ESkeletonIndex::Num>> Skeletons;
+
+	/** Whether the skeletons have been initialized for this world */
+	bool bSkeletonsInitialized = false;
+	
+	/** If we currently have a selection active in the view */
+	bool bSelecting = false;
+	
+	/** Data of the active playing asset */
+	struct FAsset
+	{
+		/** How long to stop upon reaching target */
+		static constexpr float StopDuration = 2.0f;
+
+		/** Time since the start of play */
+		float AccumulatedTime = 0.0f;
+
+		/** Start time of the asset */
+		float StartTime = 0.0f;
+
+		/** If we are currently playing a asset */
+		bool bActive = false;
+
+	} AssetData;
+	
+	/** Current play rate of the asset selection player */
+	float AssetPlayRate = 1.0f;
+
+	bool bIsVerbose = false;
+
+	/** Limits some public API */
+	friend class FDebugger;
+};
+
+} // namespace UE::PoseSearch
diff --git a/Engine/Plugins/Experimental/Animation/PoseSearch/Source/Editor/Private/PoseSearchMeshComponent.cpp b/Engine/Plugins/Experimental/Animation/PoseSearch/Source/Editor/Private/PoseSearchMeshComponent.cpp
new file mode 100644
--- /dev/null
+++ b/Engine/Plugins/Experimental/Animation/PoseSearch/Source/Editor/Private/PoseSearchMeshComponent.cpp
@@ -0,0 +1,171 @@
+// Copyright Epic Games, Inc. All Rights Reserved.
+
+#include "PoseSearchMeshComponent.h"
+#include "AnimationRuntime.h"
+#include "Animation/AnimRootMotionProvider.h"
+#include "Animation/BlendSpace.h"
+#include "Animation/MirrorDataTable.h"
+#include "Engine/SkinnedAsset.h"
+
+void UPoseSearchMeshComponent::Initialize(const FTransform& InComponentToWorld)
+{
+	SetComponentToWorld(InComponentToWorld);
+	const FReferenceSkeleton& SkeletalMeshRefSkeleton = GetSkinnedAsset()->GetRefSkeleton();
+
+	// set up bone visibility states as this gets skipped since we allocate the component array before registration
+	for (int32 BaseIndex = 0; BaseIndex < 2; ++BaseIndex)
+	{
+		BoneVisibilityStates[BaseIndex].SetNum(SkeletalMeshRefSkeleton.GetNum());
+		for (int32 BoneIndex = 0; BoneIndex < SkeletalMeshRefSkeleton.GetNum(); BoneIndex++)
+		{
+			BoneVisibilityStates[BaseIndex][BoneIndex] = BVS_ExplicitlyHidden;
+		}
+	}
+
+	StartingTransform = InComponentToWorld;
+	Refresh();
+}
+
+void UPoseSearchMeshComponent::Refresh()
+{
+	// Flip buffers once to copy the directly-written component space transforms
+	bNeedToFlipSpaceBaseBuffers = true;
+	bHasValidBoneTransform = false;
+	FlipEditableSpaceBases();
+	bHasValidBoneTransform = true;
+
+	InvalidateCachedBounds();
+	UpdateBounds();
+	MarkRenderTransformDirty();
+	MarkRenderDynamicDataDirty();
+	MarkRenderStateDirty();
+}
+
+void UPoseSearchMeshComponent::ResetToStart()
+{
+	SetComponentToWorld(StartingTransform);
+	Refresh();
+}
+
+void UPoseSearchMeshComponent::UpdatePose(const FUpdateContext& UpdateContext)
+{
+	FMemMark Mark(FMemStack::Get());
+
+	FCompactPose CompactPose;
+	CompactPose.SetBoneContainer(&RequiredBones);
+	FBlendedCurve Curve;
+	Curve.InitFrom(RequiredBones);
+	UE::Anim::FStackAttributeContainer Attributes;
+	FAnimationPoseData PoseData(CompactPose, Curve, Attributes);
+
+	if (UpdateContext.Type == ESearchIndexAssetType::Sequence ||
+		UpdateContext.Type == ESearchIndexAssetType::AnimComposite)
+	{
+		float AdvancedTime = UpdateContext.StartTime;
+
+		FAnimationRuntime::AdvanceTime(
+			UpdateContext.bLoop,
+			UpdateContext.Time - UpdateContext.StartTime,
+			AdvancedTime,
+			UpdateContext.SequenceBase->GetPlayLength());
+
+		FAnimExtractContext ExtractionCtx;
+		ExtractionCtx.CurrentTime = AdvancedTime;
+
+		UpdateContext.SequenceBase->GetAnimationPose(PoseData, ExtractionCtx);
+	}
+	else if (UpdateContext.Type == ESearchIndexAssetType::BlendSpace)
+	{
+		TArray<FBlendSampleData> BlendSamples;
+		int32 TriangulationIndex = 0;
+		UpdateContext.BlendSpace->GetSamplesFromBlendInput(UpdateContext.BlendParameters, BlendSamples, TriangulationIndex, true);
+		
+		float PlayLength = UpdateContext.BlendSpace->GetAnimationLengthFromSampleData(BlendSamples);
+		
+		float PreviousTime = UpdateContext.StartTime * PlayLength;
+		float CurrentTime = UpdateContext.Time * PlayLength;
+
+		float AdvancedTime = PreviousTime;
+		FAnimationRuntime::AdvanceTime(
+			UpdateContext.bLoop,
+			CurrentTime - PreviousTime,
+			CurrentTime,
+			PlayLength);
+		
+		FDeltaTimeRecord DeltaTimeRecord;
+		DeltaTimeRecord.Set(PreviousTime, AdvancedTime - PreviousTime);
+		FAnimExtractContext ExtractionCtx(static_cast<double>(AdvancedTime), true, DeltaTimeRecord, UpdateContext.bLoop);
+
+		for (int32 BlendSampleIdex = 0; BlendSampleIdex < BlendSamples.Num(); BlendSampleIdex++)
+		{
+			float Scale = BlendSamples[BlendSampleIdex].Animation->GetPlayLength() / PlayLength;
+
+			FDeltaTimeRecord BlendSampleDeltaTimeRecord;
+			BlendSampleDeltaTimeRecord.Set(DeltaTimeRecord.GetPrevious() * Scale, DeltaTimeRecord.Delta * Scale);
+
+			BlendSamples[BlendSampleIdex].DeltaTimeRecord = BlendSampleDeltaTimeRecord;
+			BlendSamples[BlendSampleIdex].PreviousTime = PreviousTime * Scale;
+			BlendSamples[BlendSampleIdex].Time = AdvancedTime * Scale;
+		}
+
+		UpdateContext.BlendSpace->GetAnimationPose(BlendSamples, ExtractionCtx, PoseData);
+	}
+	else
+	{
+		checkNoEntry();
+	}
+
+	LastRootMotionDelta = FTransform::Identity;
+
+	const UE::Anim::IAnimRootMotionProvider* RootMotionProvider = UE::Anim::IAnimRootMotionProvider::Get();
+
+	if (ensureMsgf(RootMotionProvider, TEXT("Could not get Root Motion Provider.")))
+	{
+		if (ensureMsgf(RootMotionProvider->HasRootMotion(Attributes), TEXT("Blend Space had no Root Motion Attribute.")))
+		{
+			RootMotionProvider->ExtractRootMotion(Attributes, LastRootMotionDelta);
+		}
+	}
+
+	if (UpdateContext.bMirrored)
+	{
+		FAnimationRuntime::MirrorPose(
+			CompactPose, 
+			UpdateContext.MirrorDataTable->MirrorAxis, 
+			*UpdateContext.CompactPoseMirrorBones, 
+			*UpdateContext.ComponentSpaceRefRotations);
+	}
+
+	FCSPose<FCompactPose> ComponentSpacePose;
+	ComponentSpacePose.InitPose(CompactPose);
+
+	for (const FBoneIndexType BoneIndex : RequiredBones.GetBoneIndicesArray())
+	{
+		const FTransform BoneTransform = 
+			ComponentSpacePose.GetComponentSpaceTransform(FCompactPoseBoneIndex(BoneIndex));
+
+		FSkeletonPoseBoneIndex SkeletonBoneIndex =
+			RequiredBones.GetSkeletonPoseIndexFromCompactPoseIndex(FCompactPoseBoneIndex(BoneIndex));
+		FName BoneName = 
+			RequiredBones.GetSkeletonAsset()->GetReferenceSkeleton().GetBoneName(SkeletonBoneIndex.GetInt());
+		SetBoneTransformByName(BoneName, BoneTransform, EBoneSpaces::ComponentSpace);
+	}
+
+	if (UpdateContext.bMirrored)
+	{
+		const EAxis::Type MirrorAxis = UpdateContext.MirrorDataTable->MirrorAxis;
+		FVector T = LastRootMotionDelta.GetTranslation();
+		T = FAnimationRuntime::MirrorVector(T, MirrorAxis);
+		const FQuat ReferenceRotation = (*UpdateContext.ComponentSpaceRefRotations)[FCompactPoseBoneIndex(0)];
+		FQuat Q = LastRootMotionDelta.GetRotation();
+		Q = FAnimationRuntime::MirrorQuat(Q, MirrorAxis);
+		Q *= FAnimationRuntime::MirrorQuat(ReferenceRotation, MirrorAxis).Inverse() * ReferenceRotation;
+		LastRootMotionDelta = FTransform(Q, T, LastRootMotionDelta.GetScale3D());
+	}
+
+	const FTransform ComponentTransform = LastRootMotionDelta * StartingTransform;
+
+	SetComponentToWorld(ComponentTransform);
+	FillComponentSpaceTransforms();
+	Refresh();
+}
\ No newline at end of file
diff --git a/Engine/Plugins/Experimental/Animation/PoseSearch/Source/Editor/Private/PoseSearchMeshComponent.h b/Engine/Plugins/Experimental/Animation/PoseSearch/Source/Editor/Private/PoseSearchMeshComponent.h
new file mode 100644
--- /dev/null
+++ b/Engine/Plugins/Experimental/Animation/PoseSearch/Source/Editor/Private/PoseSearchMeshComponent.h
@@ -0,0 +1,36 @@
+﻿// Copyright Epic Games, Inc. All Rights Reserved.
+
+#pragma once
+
+#include "Components/PoseableMeshComponent.h"
+#include "PoseSearch/PoseSearch.h"
+#include "PoseSearchMeshComponent.generated.h"
+
+UCLASS()
+class UPoseSearchMeshComponent : public UPoseableMeshComponent
+{
+	GENERATED_BODY()
+public:
+
+	struct FUpdateContext
+	{
+		ESearchIndexAssetType Type = ESearchIndexAssetType::Invalid;
+		const UAnimSequenceBase* SequenceBase = nullptr;
+		const UBlendSpace* BlendSpace = nullptr;
+		float StartTime = 0.0f;
+		float Time = 0.0f;
+		bool bLoop = false;
+		bool bMirrored = false;
+		FVector BlendParameters = FVector::Zero();
+		const UMirrorDataTable* MirrorDataTable = nullptr;
+		TCustomBoneIndexArray<FCompactPoseBoneIndex, FCompactPoseBoneIndex>* CompactPoseMirrorBones = nullptr;
+		TCustomBoneIndexArray<FQuat, FCompactPoseBoneIndex>* ComponentSpaceRefRotations = nullptr;
+	};
+
+	void Refresh();
+	void ResetToStart();
+	void UpdatePose(const FUpdateContext& UpdateContext);
+	void Initialize(const FTransform& InComponentToWorld);
+	FTransform StartingTransform = FTransform::Identity;
+	FTransform LastRootMotionDelta = FTransform::Identity;
+};
diff --git a/Engine/Plugins/Experimental/Animation/PoseSearch/Source/Editor/Private/SPoseSearchDatabaseAssetList.cpp b/Engine/Plugins/Experimental/Animation/PoseSearch/Source/Editor/Private/SPoseSearchDatabaseAssetList.cpp
deleted file mode 100644
diff --git a/Engine/Plugins/Experimental/Animation/PoseSearch/Source/Editor/Public/PoseSearchDatabaseAssetListItem.h b/Engine/Plugins/Experimental/Animation/PoseSearch/Source/Editor/Public/PoseSearchDatabaseAssetListItem.h
new file mode 100644
--- /dev/null
+++ b/Engine/Plugins/Experimental/Animation/PoseSearch/Source/Editor/Public/PoseSearchDatabaseAssetListItem.h
@@ -0,0 +1,59 @@
+// Copyright Epic Games, Inc. All Rights Reserved.
+
+#pragma once
+
+#include "CoreMinimal.h"
+#include "PoseSearchDatabaseAssetTreeNode.h"
+#include "Widgets/Views/STableRow.h"
+
+class FUICommandList;
+
+namespace UE::PoseSearch
+{
+	class FDatabaseViewModel;
+	class SDatabaseAssetTree;
+
+	class SDatabaseAssetListItem : public STableRow<TSharedPtr<FDatabaseAssetTreeNode>>
+	{
+	public:
+		void Construct(
+			const FArguments& InArgs,
+			const TSharedRef<FDatabaseViewModel>& InEditorViewModel,
+			const TSharedRef<STableViewBase>& OwnerTable,
+			TSharedRef<FDatabaseAssetTreeNode> InAssetTreeNode,
+			TSharedRef<FUICommandList> InCommandList,
+			TSharedPtr<SDatabaseAssetTree> InHierarchy);
+
+	protected:
+		FText GetName() const;
+		TSharedRef<SWidget> GenerateItemWidget();
+		TSharedRef<SWidget> GenerateAddButtonWidget();
+
+		const FSlateBrush* GetGroupBackgroundImage() const;
+		void ConstructGroupItem(const TSharedRef<STableViewBase>& OwnerTable);
+		void ConstructAssetItem(const TSharedRef<STableViewBase>& OwnerTable);
+
+		void OnAddSequence();
+		void OnAddBlendSpace();
+		void OnAddAnimComposite();
+
+		EVisibility GetSelectedActorIconVisbility() const;
+
+		FSlateColor GetNameTextColorAndOpacity() const;
+		FSlateColor GetLoopingColorAndOpacity() const;
+		FText GetLoopingToolTip() const;
+		FSlateColor GetRootMotionColorAndOpacity() const;
+		FText GetRootMotionOptionToolTip() const;
+		const FSlateBrush* GetMirrorOptionSlateBrush() const;
+		FText GetMirrorOptionToolTip() const;
+		
+		FText GetAssetEnabledToolTip() const;
+		ECheckBoxState GetAssetEnabledChecked() const;
+		void OnAssetIsEnabledChanged(ECheckBoxState NewCheckboxState);
+
+		TWeakPtr<FDatabaseAssetTreeNode> WeakAssetTreeNode;
+		TWeakPtr<FDatabaseViewModel> EditorViewModel;
+		TWeakPtr<SDatabaseAssetTree> SkeletonView;
+	};
+}
+
diff --git a/Engine/Plugins/Experimental/Animation/PoseSearch/Source/Editor/Public/PoseSearchDatabaseAssetTree.h b/Engine/Plugins/Experimental/Animation/PoseSearch/Source/Editor/Public/PoseSearchDatabaseAssetTree.h
new file mode 100644
--- /dev/null
+++ b/Engine/Plugins/Experimental/Animation/PoseSearch/Source/Editor/Public/PoseSearchDatabaseAssetTree.h
@@ -0,0 +1,123 @@
+// Copyright Epic Games, Inc. All Rights Reserved.
+
+#pragma once
+
+#include "PoseSearchDatabaseAssetListItem.h"
+#include "DetailColumnSizeData.h"
+#include "EditorUndoClient.h"
+#include "Widgets/Views/STreeView.h"
+#include "Widgets/SCompoundWidget.h"
+
+class FUICommandList;
+
+namespace UE::PoseSearch
+{
+	class FDatabaseViewModel;
+
+	class SDatabaseAssetTree : public SCompoundWidget, public FSelfRegisteringEditorUndoClient
+	{
+	public:
+		SLATE_BEGIN_ARGS(SDatabaseAssetTree) {}
+		SLATE_END_ARGS()
+
+		virtual ~SDatabaseAssetTree();
+
+		void Construct(const FArguments& InArgs, TSharedRef<FDatabaseViewModel> InEditorViewModel);
+
+		// SWidget interface
+		virtual FReply OnDragOver(const FGeometry& MyGeometry, const FDragDropEvent& DragDropEvent) override;
+		virtual FReply OnDrop(const FGeometry& MyGeometry, const FDragDropEvent& DragDropEvent) override;
+		virtual FReply OnKeyDown(const FGeometry& MyGeometry, const FKeyEvent& InKeyEvent) override;
+		// End SWidget interface
+
+		// Begin FEditorUndoClient interface
+		virtual bool MatchesContext(const FTransactionContext& InContext, const TArray<TPair<UObject*, FTransactionObjectEvent>>& TransactionObjectContexts) const override;
+		virtual void PostUndo(bool bSuccess) override;
+		virtual void PostRedo(bool bSuccess) override;
+		// End FEditorUndoClient interface
+		
+		void RefreshTreeView(bool bIsInitialSetup = false, bool bRecoverSelection = false);
+		void FinalizeTreeChanges(bool bRecoverSelection = false);
+
+		FDetailColumnSizeData& GetColumnSizeData() { return ColumnSizeData; }
+		
+	protected:
+		TWeakPtr<FDatabaseViewModel> EditorViewModel;
+
+		/** command list we bind to */
+		TSharedPtr<FUICommandList> CommandList;
+
+		/** tree view widget */
+		TSharedPtr<STreeView<TSharedPtr<FDatabaseAssetTreeNode>>> TreeView;
+		TSharedPtr<SWidget> TreeViewDragAndDropSuggestion;
+		TArray<TSharedPtr<FDatabaseAssetTreeNode>> RootNodes;
+		TArray<TSharedPtr<FDatabaseAssetTreeNode>> AllNodes;
+
+		TSharedRef<ITableRow> MakeTableRowWidget(
+			TSharedPtr<FDatabaseAssetTreeNode> InItem,
+			const TSharedRef<STableViewBase>& OwnerTable);
+		void HandleGetChildrenForTree(
+			TSharedPtr<FDatabaseAssetTreeNode> InNode, 
+			TArray<TSharedPtr<FDatabaseAssetTreeNode>>& OutChildren);
+
+		TOptional<EItemDropZone> OnCanAcceptDrop(
+			const FDragDropEvent& DragDropEvent, 
+			EItemDropZone DropZone, 
+			TSharedPtr<FDatabaseAssetTreeNode> TargetItem);
+
+		FReply OnAcceptDrop(
+			const FDragDropEvent& DragDropEvent,
+			EItemDropZone DropZone,
+			TSharedPtr<FDatabaseAssetTreeNode> TargetItem);
+
+		TSharedRef<SWidget> CreateAddNewMenuWidget();
+		TSharedPtr<SWidget> CreateContextMenu();
+
+		TSharedRef<SWidget> GenerateFilterBoxWidget();
+
+		FText GetFilterText() const;
+		void OnAssetFilterTextCommitted(const FText& InText, ETextCommit::Type CommitInfo);
+
+		void SetAssetFilterString(FString InString) { AssetFilterString = InString; }
+		FString GetAssetFilterString() const { return AssetFilterString; }
+
+		//String Filter For Database View
+		FString AssetFilterString;
+
+		void OnAddSequence(bool bFinalizeChanges = true);
+		void OnAddBlendSpace(bool bFinalizeChanges = true);
+		void OnAddAnimComposite(bool bFinalizeChanges = true);
+
+		void OnDeleteAsset(TSharedPtr<FDatabaseAssetTreeNode> Node, bool bFinalizeChanges = true);
+		void CreateCommandList();
+
+		/** Removes existing selected component nodes from the tree*/
+		bool CanDeleteNodes() const;
+		void OnDeleteNodes();
+
+		void EnableSelectedNodes(bool bIsEnabled);
+		void OnEnableNodes() { EnableSelectedNodes(true); }
+		void OnDisableNodes() { EnableSelectedNodes(false); }
+
+		friend SDatabaseAssetListItem;
+
+		FDetailColumnSizeData ColumnSizeData;
+		
+	protected:
+		// Called when an item is selected/deselected
+		DECLARE_MULTICAST_DELEGATE_TwoParams(
+			FOnSelectionChangedMulticaster, 
+			const TArrayView<TSharedPtr<FDatabaseAssetTreeNode>>& /* InSelectedItems */,
+			ESelectInfo::Type /* SelectInfo */);
+
+		FOnSelectionChangedMulticaster OnSelectionChanged;
+
+	public:
+		typedef FOnSelectionChangedMulticaster::FDelegate FOnSelectionChanged;
+		void RegisterOnSelectionChanged(const FOnSelectionChanged& Delegate);
+		void UnregisterOnSelectionChanged(void* Unregister);
+
+		void RecoverSelection(const TArray<TSharedPtr<FDatabaseAssetTreeNode>>& PreviouslySelectedNodes);
+	};
+}
+
diff --git a/Engine/Plugins/Experimental/Animation/PoseSearch/Source/Editor/Public/PoseSearchDatabaseAssetTreeNode.h b/Engine/Plugins/Experimental/Animation/PoseSearch/Source/Editor/Public/PoseSearchDatabaseAssetTreeNode.h
new file mode 100644
--- /dev/null
+++ b/Engine/Plugins/Experimental/Animation/PoseSearch/Source/Editor/Public/PoseSearchDatabaseAssetTreeNode.h
@@ -0,0 +1,43 @@
+// Copyright Epic Games, Inc. All Rights Reserved.
+
+#pragma once
+
+#include "CoreMinimal.h"
+#include "PoseSearch/PoseSearch.h"
+
+class FUICommandList;
+class ITableRow;
+class STableViewBase;
+
+namespace UE::PoseSearch
+{
+	class FDatabaseViewModel;
+	class SDatabaseAssetTree;
+
+	class FDatabaseAssetTreeNode : public TSharedFromThis<FDatabaseAssetTreeNode>
+	{
+
+	public:
+		FDatabaseAssetTreeNode(
+			int32 InSourceAssetIdx,
+			ESearchIndexAssetType InSourceAssetType,
+			const TSharedRef<FDatabaseViewModel>& InEditorViewModel);
+
+		TSharedRef<ITableRow> MakeTreeRowWidget(
+			const TSharedRef<STableViewBase>& InOwnerTable,
+			TSharedRef<FDatabaseAssetTreeNode> InDatabaseAssetNode,
+			TSharedRef<FUICommandList> InCommandList,
+			TSharedPtr<SDatabaseAssetTree> InHierarchy);
+
+		bool IsRootMotionEnabled() const;
+		bool IsLooping() const;
+		EPoseSearchMirrorOption GetMirrorOption() const;
+		
+		int32 SourceAssetIdx;
+		ESearchIndexAssetType SourceAssetType;
+		TSharedPtr<FDatabaseAssetTreeNode> Parent;
+		TArray<TSharedPtr<FDatabaseAssetTreeNode>> Children;
+		TWeakPtr<FDatabaseViewModel> EditorViewModel;
+	};
+}
+
diff --git a/Engine/Plugins/Experimental/Animation/PoseSearch/Source/Editor/Public/PoseSearchDatabaseEditor.h b/Engine/Plugins/Experimental/Animation/PoseSearch/Source/Editor/Public/PoseSearchDatabaseEditor.h
--- a/Engine/Plugins/Experimental/Animation/PoseSearch/Source/Editor/Public/PoseSearchDatabaseEditor.h
+++ b/Engine/Plugins/Experimental/Animation/PoseSearch/Source/Editor/Public/PoseSearchDatabaseEditor.h
@@ -5,17 +5,18 @@
 #include "CoreMinimal.h"
 #include "Toolkits/AssetEditorToolkit.h"
 #include "Misc/NotifyHook.h"
-#include "SPoseSearchDatabaseAssetList.h"
 
 class IDetailsView;
 class FToolBarBuilder;
-class UPoseSearchDatabaseSelectionReflection;
+class UPoseSearchDatabase;
 
 namespace UE::PoseSearch
 {
-	class SDatabasePreview;
+	class FDatabaseAssetTreeNode;
 	class FDatabasePreviewScene;
 	class FDatabaseViewModel;
+	class SDatabaseAssetTree;
+	class SDatabasePreview;
 
 	class FDatabaseEditor : public FAssetEditorToolkit, public FNotifyHook
 	{
diff --git a/Engine/Plugins/Experimental/Animation/PoseSearch/Source/Editor/Public/SPoseSearchDatabaseAssetList.h b/Engine/Plugins/Experimental/Animation/PoseSearch/Source/Editor/Public/SPoseSearchDatabaseAssetList.h
deleted file mode 100644
