Removed explicit default module class declaration`FModule` in PoseSearch.cpp.

diff --git a/Engine/Plugins/Experimental/Animation/PoseSearch/Source/Runtime/Private/PoseSearch.cpp b/Engine/Plugins/Experimental/Animation/PoseSearch/Source/Runtime/Private/PoseSearch.cpp
--- a/Engine/Plugins/Experimental/Animation/PoseSearch/Source/Runtime/Private/PoseSearch.cpp
+++ b/Engine/Plugins/Experimental/Animation/PoseSearch/Source/Runtime/Private/PoseSearch.cpp
@@ -2956,17 +2956,6 @@ void DrawSearchIndex(const FDebugDrawParams& DrawParams)
 #endif // ENABLE_DRAW_DEBUG
 }
 
-//////////////////////////////////////////////////////////////////////////
-// FModule
-// @todo: clean up this code
-class FModule : public IModuleInterface
-{
-public: // IModuleInterface
-	virtual void StartupModule() override {}
-	virtual void ShutdownModule() override {}
-};
 } // namespace UE::PoseSearch
 
-#undef LOCTEXT_NAMESPACE
-
-IMPLEMENT_MODULE(UE::PoseSearch::FModule, PoseSearch)
+#undef LOCTEXT_NAMESPACE
\ No newline at end of file
diff --git a/Engine/Plugins/Experimental/Animation/PoseSearch/Source/Runtime/Private/PoseSearchModule.cpp b/Engine/Plugins/Experimental/Animation/PoseSearch/Source/Runtime/Private/PoseSearchModule.cpp
new file mode 100644
--- /dev/null
+++ b/Engine/Plugins/Experimental/Animation/PoseSearch/Source/Runtime/Private/PoseSearchModule.cpp
@@ -0,0 +1,5 @@
+﻿// Copyright Epic Games, Inc. All Rights Reserved.
+
+#include "Modules/ModuleManager.h"
+
+IMPLEMENT_MODULE(FDefaultModuleImpl, PoseSearch);
\ No newline at end of file
