motion matching - deprecating FSchemaInitializer, removing unnecessary code

diff --git a/Engine/Plugins/Experimental/Animation/PoseSearch/Source/Runtime/Private/PoseSearch.cpp b/Engine/Plugins/Experimental/Animation/PoseSearch/Source/Runtime/Private/PoseSearch.cpp
--- a/Engine/Plugins/Experimental/Animation/PoseSearch/Source/Runtime/Private/PoseSearch.cpp
+++ b/Engine/Plugins/Experimental/Animation/PoseSearch/Source/Runtime/Private/PoseSearch.cpp
@@ -100,11 +100,6 @@ void CompareFeatureVectors(TConstArrayView<float> A, TConstArrayView<float> B, T
 	VR = ((VA - VB) * VW).square();
 }
 
-static inline bool IsSamplingRangeValid(FFloatInterval Range)
-{
-	return Range.IsValid() && (Range.Min >= 0.0f);
-}
-
 /**
 * Algo::LowerBound adapted to TIndexedContainerIterator for use with indexable but not necessarily contiguous containers. Used here with TRingBuffer.
 *
@@ -145,18 +140,12 @@ FORCEINLINE auto LowerBound(IteratorType First, IteratorType Last, const ValueTy
 	return Start;
 }
 
-template <typename IteratorType, typename ValueType, typename SortPredicateType>
+template <typename IteratorType, typename ValueType, typename SortPredicateType = TLess<>()>
 FORCEINLINE auto LowerBound(IteratorType First, IteratorType Last, const ValueType& Value, SortPredicateType SortPredicate) -> decltype(First.GetIndex())
 {
 	return LowerBound(First, Last, Value, FIdentityFunctor(), SortPredicate);
 }
 
-template <typename IteratorType, typename ValueType>
-FORCEINLINE auto LowerBound(IteratorType First, IteratorType Last, const ValueType& Value) -> decltype(First.GetIndex())
-{
-	return LowerBound(First, Last, Value, FIdentityFunctor(), TLess<>());
-}
-
 typedef TArray<size_t, TInlineAllocator<128>> FNonSelectableIdx;
 static void PopulateNonSelectableIdx(FNonSelectableIdx& NonSelectableIdx, FSearchContext& SearchContext, const UPoseSearchDatabase* Database, TConstArrayView<float> QueryValues)
 {
@@ -363,12 +352,6 @@ private:
 
 //////////////////////////////////////////////////////////////////////////
 // UPoseSearchFeatureChannel
-void UPoseSearchFeatureChannel::InitializeSchema(UE::PoseSearch::FSchemaInitializer& Initializer)
-{
-	ChannelIdx = Initializer.GetCurrentChannelIdx();
-	ChannelDataOffset = Initializer.GetCurrentChannelDataOffset();
-}
-
 #if WITH_EDITOR
 void UPoseSearchFeatureChannel::PopulateChannelLayoutSet(UE::PoseSearch::FFeatureChannelLayoutSet& FeatureChannelLayoutSet) const
 {
@@ -381,16 +364,6 @@ void UPoseSearchFeatureChannel::ComputeCostBreakdowns(UE::PoseSearch::ICostBreak
 }
 #endif // WITH_EDITOR
 
-//////////////////////////////////////////////////////////////////////////
-// FSchemaInitializer
-namespace UE::PoseSearch
-{
-int32 FSchemaInitializer::AddBoneReference(const FBoneReference& BoneReference)
-{
-	return BoneReferences.AddUnique(BoneReference);
-}
-} // namespace UE::PoseSearch
-
 //////////////////////////////////////////////////////////////////////////
 // UPoseSearchSchema
 void UPoseSearchSchema::Finalize(bool bRemoveEmptyChannels)
@@ -404,20 +377,17 @@ void UPoseSearchSchema::Finalize(bool bRemoveEmptyChannels)
 
 	BoneReferences.Reset();
 
-	FSchemaInitializer Initializer;
+	int32 CurrentChannelDataOffset = 0;
+
+	SchemaCardinality = 0;
 	for (int32 ChannelIdx = 0; ChannelIdx != Channels.Num(); ++ChannelIdx)
 	{
-		if (Channels[ChannelIdx].Get())
+		if (UPoseSearchFeatureChannel* Channel = Channels[ChannelIdx].Get())
 		{
-			Initializer.CurrentChannelIdx = ChannelIdx;
-			Channels[ChannelIdx]->InitializeSchema(Initializer);
+			Channel->InitializeSchema(this);
 		}
 	}
 
-	SchemaCardinality = Initializer.GetCurrentChannelDataOffset();
-
-	BoneReferences = MoveTemp(Initializer.BoneReferences);
-
 	ResolveBoneReferences();
 }
 
@@ -503,20 +473,16 @@ void UPoseSearchSchema::ResolveBoneReferences()
 	}
 }
 
-bool UPoseSearchSchema::BuildQuery(UE::PoseSearch::FSearchContext& SearchContext, FPoseSearchFeatureVectorBuilder& InOutQuery) const
+void UPoseSearchSchema::BuildQuery(UE::PoseSearch::FSearchContext& SearchContext, FPoseSearchFeatureVectorBuilder& InOutQuery) const
 {
 	QUICK_SCOPE_CYCLE_COUNTER(STAT_PoseSearch_BuildQuery);
 
 	InOutQuery.Init(this);
 
-	bool bSuccess = true;
 	for (const TObjectPtr<UPoseSearchFeatureChannel>& Channel : Channels)
 	{
-		bool bChannelSuccess = Channel->BuildQuery(SearchContext, InOutQuery);
-		bSuccess &= bChannelSuccess;
+		Channel->BuildQuery(SearchContext, InOutQuery);
 	}
-
-	return bSuccess;
 }
 
 //////////////////////////////////////////////////////////////////////////
@@ -1488,7 +1454,7 @@ void FPoseIndicesHistory::Update(const FSearchResult& SearchResult, float DeltaT
 	}
 }
 
-FTransform FSearchContext::TryGetTransformAndCacheResults(float SampleTime, const UPoseSearchSchema* Schema, int8 SchemaBoneIdx, bool& bError)
+FTransform FSearchContext::TryGetTransformAndCacheResults(float SampleTime, const UPoseSearchSchema* Schema, int8 SchemaBoneIdx)
 {
 	check(History && Schema);
 
@@ -1503,7 +1469,6 @@ FTransform FSearchContext::TryGetTransformAndCacheResults(float SampleTime, cons
 
 	if (Entry)
 	{
-		bError = false;
 		return Entry->Transform;
 	}
 
@@ -1521,11 +1486,9 @@ FTransform FSearchContext::TryGetTransformAndCacheResults(float SampleTime, cons
 				CachedEntries.Emplace(SampleTime, SampledComponentPose[NewEntryBoneIndexType], NewEntryBoneIndexType);
 			}
 
-			bError = false;
 			return SampledComponentPose[BoneIndexType];
 		}
 
-		bError = true;
 		return FTransform::Identity;
 	}
 	
@@ -1537,11 +1500,9 @@ FTransform FSearchContext::TryGetTransformAndCacheResults(float SampleTime, cons
 		NewEntry.BoneIndexType = BoneIndexType;
 		NewEntry.Transform = SampledRootTransform;
 
-		bError = false;
 		return SampledRootTransform;
 	}
 	
-	bError = true;
 	return FTransform::Identity;
 }
 
@@ -1579,20 +1540,20 @@ const FPoseSearchFeatureVectorBuilder* FSearchContext::GetCachedQuery(const UPos
 	return nullptr;
 }
 
-bool FSearchContext::GetOrBuildQuery(const UPoseSearchDatabase* Database, FPoseSearchFeatureVectorBuilder& FeatureVectorBuilder)
+void FSearchContext::GetOrBuildQuery(const UPoseSearchDatabase* Database, FPoseSearchFeatureVectorBuilder& FeatureVectorBuilder)
 {
 	const FPoseSearchFeatureVectorBuilder* CachedFeatureVectorBuilder = GetCachedQuery(Database);
 	if (CachedFeatureVectorBuilder)
 	{
 		FeatureVectorBuilder = *CachedFeatureVectorBuilder;
-		return true;
 	}
-
-	FSearchContext::FCachedQuery& NewCachedQuery = CachedQueries[CachedQueries.AddDefaulted()];
-	NewCachedQuery.Database = Database;
-	Database->BuildQuery(*this, NewCachedQuery.FeatureVectorBuilder);
-	FeatureVectorBuilder = NewCachedQuery.FeatureVectorBuilder;
-	return false;
+	else
+	{
+		FSearchContext::FCachedQuery& NewCachedQuery = CachedQueries[CachedQueries.AddDefaulted()];
+		NewCachedQuery.Database = Database;
+		Database->BuildQuery(*this, NewCachedQuery.FeatureVectorBuilder);
+		FeatureVectorBuilder = NewCachedQuery.FeatureVectorBuilder;
+	}
 }
 
 bool FSearchContext::IsCurrentResultFromDatabase(const UPoseSearchDatabase* Database) const
@@ -1995,18 +1956,7 @@ float FFeatureVectorHelper::DecodeFloatInternal(TConstArrayView<float> Values, i
 bool FDebugDrawParams::CanDraw() const
 {
 #if ENABLE_DRAW_DEBUG
-	if (!World)
-	{
-		return false;
-	}
-
-	const FPoseSearchIndex* SearchIndex = GetSearchIndex();
-	if (!SearchIndex)
-	{
-		return false;
-	}
-
-	return true;
+	return World && Database && Database->Schema && Database->Schema->IsValid();
 #else // ENABLE_DRAW_DEBUG
 	return false;
 #endif // ENABLE_DRAW_DEBUG
diff --git a/Engine/Plugins/Experimental/Animation/PoseSearch/Source/Runtime/Private/PoseSearchDerivedData.cpp b/Engine/Plugins/Experimental/Animation/PoseSearch/Source/Runtime/Private/PoseSearchDerivedData.cpp
--- a/Engine/Plugins/Experimental/Animation/PoseSearch/Source/Runtime/Private/PoseSearchDerivedData.cpp
+++ b/Engine/Plugins/Experimental/Animation/PoseSearch/Source/Runtime/Private/PoseSearchDerivedData.cpp
@@ -128,7 +128,7 @@ static TArray<float> ComputeChannelsDeviations(TConstArrayView<FPoseSearchIndexB
 			const TArray<FFeatureChannelLayoutSet::FEntry>& Entries = Pair.Value;
 			for (const FFeatureChannelLayoutSet::FEntry& Entry : Entries)
 			{
-				if (Entry.SchemaIndex == ThisSchemaIndex)
+				if (Entry.Cardinality > 0 && Entry.SchemaIndex == ThisSchemaIndex)
 				{
 					const float FeatureMeanDeviation = ComputeFeatureMeanDeviation(Entries, SearchIndexBases, Schemas);
 					// the associated data to all the Entries data is going to be used to calculate the deviation of Deviation[Entry.DataOffset] to Deviation[Entry.DataOffset + Entry.Cardinality]
diff --git a/Engine/Plugins/Experimental/Animation/PoseSearch/Source/Runtime/Private/PoseSearchFeatureChannels.cpp b/Engine/Plugins/Experimental/Animation/PoseSearch/Source/Runtime/Private/PoseSearchFeatureChannels.cpp
--- a/Engine/Plugins/Experimental/Animation/PoseSearch/Source/Runtime/Private/PoseSearchFeatureChannels.cpp
+++ b/Engine/Plugins/Experimental/Animation/PoseSearch/Source/Runtime/Private/PoseSearchFeatureChannels.cpp
@@ -318,12 +318,13 @@ class USkeleton* UPoseSearchFeatureChannel::GetSkeleton(bool& bInvalidSkeletonIs
 //////////////////////////////////////////////////////////////////////////
 // UPoseSearchFeatureChannel_Position
 
-void UPoseSearchFeatureChannel_Position::InitializeSchema(UE::PoseSearch::FSchemaInitializer& Initializer)
+void UPoseSearchFeatureChannel_Position::InitializeSchema(UPoseSearchSchema* Schema)
 {
-	Super::InitializeSchema(Initializer);
+	ChannelDataOffset = Schema->SchemaCardinality;
 	ChannelCardinality = UE::PoseSearch::FFeatureVectorHelper::EncodeVectorCardinality;
-	Initializer.SetCurrentChannelDataOffset(ChannelDataOffset + ChannelCardinality);
-	SchemaBoneIdx = Initializer.AddBoneReference(Bone);
+	Schema->SchemaCardinality += ChannelCardinality;
+
+	SchemaBoneIdx = Schema->AddBoneReference(Bone);
 }
 
 void UPoseSearchFeatureChannel_Position::FillWeights(TArray<float>& Weights) const
@@ -356,7 +357,7 @@ void UPoseSearchFeatureChannel_Position::IndexAsset(UE::PoseSearch::IAssetIndexe
 	}
 }
 
-bool UPoseSearchFeatureChannel_Position::BuildQuery(UE::PoseSearch::FSearchContext& SearchContext, FPoseSearchFeatureVectorBuilder& InOutQuery) const
+void UPoseSearchFeatureChannel_Position::BuildQuery(UE::PoseSearch::FSearchContext& SearchContext, FPoseSearchFeatureVectorBuilder& InOutQuery) const
 {
 	using namespace UE::PoseSearch;
 
@@ -371,23 +372,22 @@ bool UPoseSearchFeatureChannel_Position::BuildQuery(UE::PoseSearch::FSearchConte
 			FFeatureVectorHelper::EncodeVector(InOutQuery.EditValues(), DataOffset, SearchContext.GetCurrentResultPrevPoseVector(), SearchContext.GetCurrentResultPoseVector(), SearchContext.GetCurrentResultNextPoseVector(), LerpValue);
 		}
 		// else leave the InOutQuery set to zero since the SearchContext.History is invalid and it'll fail if we continue
-		return bSkip;
 	}
+	else
+	{
+		FTransform Transform = SearchContext.TryGetTransformAndCacheResults(SampleTimeOffset, InOutQuery.GetSchema(), SchemaBoneIdx);
 
-	bool bAnyError = false;
-	FTransform Transform = SearchContext.TryGetTransformAndCacheResults(SampleTimeOffset, InOutQuery.GetSchema(), SchemaBoneIdx, bAnyError);
+		if (!bUseSampleTimeOffsetRootBone)
+		{
+			const FTransform RootTransform = SearchContext.TryGetTransformAndCacheResults(0.f, InOutQuery.GetSchema(), FSearchContext::SchemaRootBoneIdx);
+			const FTransform RootTransformPrev = SearchContext.TryGetTransformAndCacheResults(SampleTimeOffset, InOutQuery.GetSchema(), FSearchContext::SchemaRootBoneIdx);
+			Transform = Transform * (RootTransformPrev * RootTransform.Inverse());
+		}
 
-	if (!bUseSampleTimeOffsetRootBone)
-	{
-		const FTransform RootTransform = SearchContext.TryGetTransformAndCacheResults(0.f, InOutQuery.GetSchema(), FSearchContext::SchemaRootBoneIdx, bAnyError);
-		const FTransform RootTransformPrev = SearchContext.TryGetTransformAndCacheResults(SampleTimeOffset, InOutQuery.GetSchema(), FSearchContext::SchemaRootBoneIdx, bAnyError);
-		Transform = Transform * (RootTransformPrev * RootTransform.Inverse());
+		int32 DataOffset = ChannelDataOffset;
+		FFeatureVectorHelper::EncodeVector(InOutQuery.EditValues(), DataOffset, Transform.GetTranslation());
+		check(DataOffset == ChannelDataOffset + ChannelCardinality);
 	}
-
-	int32 DataOffset = ChannelDataOffset;
-	FFeatureVectorHelper::EncodeVector(InOutQuery.EditValues(), DataOffset, Transform.GetTranslation());
-	check(DataOffset == ChannelDataOffset + ChannelCardinality);
-	return !bAnyError;
 }
 
 void UPoseSearchFeatureChannel_Position::DebugDraw(const UE::PoseSearch::FDebugDrawParams& DrawParams, TConstArrayView<float> PoseVector) const
@@ -426,12 +426,12 @@ void UPoseSearchFeatureChannel_Position::DebugDraw(const UE::PoseSearch::FDebugD
 //////////////////////////////////////////////////////////////////////////
 // UPoseSearchFeatureChannel_Heading
 
-void UPoseSearchFeatureChannel_Heading::InitializeSchema(UE::PoseSearch::FSchemaInitializer& Initializer)
+void UPoseSearchFeatureChannel_Heading::InitializeSchema(UPoseSearchSchema* Schema)
 {
-	Super::InitializeSchema(Initializer);
+	ChannelDataOffset = Schema->SchemaCardinality;
 	ChannelCardinality = UE::PoseSearch::FFeatureVectorHelper::EncodeVectorCardinality;
-	Initializer.SetCurrentChannelDataOffset(ChannelDataOffset + ChannelCardinality);
-	SchemaBoneIdx = Initializer.AddBoneReference(Bone);
+	Schema->SchemaCardinality += ChannelCardinality;
+	SchemaBoneIdx = Schema->AddBoneReference(Bone);
 }
 
 void UPoseSearchFeatureChannel_Heading::FillWeights(TArray<float>& Weights) const
@@ -481,7 +481,7 @@ void UPoseSearchFeatureChannel_Heading::IndexAsset(UE::PoseSearch::IAssetIndexer
 	}
 }
 
-bool UPoseSearchFeatureChannel_Heading::BuildQuery(UE::PoseSearch::FSearchContext& SearchContext, FPoseSearchFeatureVectorBuilder& InOutQuery) const
+void UPoseSearchFeatureChannel_Heading::BuildQuery(UE::PoseSearch::FSearchContext& SearchContext, FPoseSearchFeatureVectorBuilder& InOutQuery) const
 {
 	using namespace UE::PoseSearch;
 
@@ -497,24 +497,22 @@ bool UPoseSearchFeatureChannel_Heading::BuildQuery(UE::PoseSearch::FSearchContex
 			check(DataOffset == ChannelDataOffset + ChannelCardinality);
 		}
 		// else leave the InOutQuery set to zero since the SearchContext.History is invalid and it'll fail if we continue
-		return bSkip;
 	}
-
-	bool bAnyError = false;
-	FTransform Transform = SearchContext.TryGetTransformAndCacheResults(SampleTimeOffset, InOutQuery.GetSchema(), SchemaBoneIdx, bAnyError);
-
-	if (!bUseSampleTimeOffsetRootBone)
+	else
 	{
-		const FTransform RootTransform = SearchContext.TryGetTransformAndCacheResults(0.f, InOutQuery.GetSchema(), FSearchContext::SchemaRootBoneIdx, bAnyError);
-		const FTransform RootTransformPrev = SearchContext.TryGetTransformAndCacheResults(SampleTimeOffset, InOutQuery.GetSchema(), FSearchContext::SchemaRootBoneIdx, bAnyError);
-		Transform = Transform * (RootTransformPrev * RootTransform.Inverse());
-	}
+		FTransform Transform = SearchContext.TryGetTransformAndCacheResults(SampleTimeOffset, InOutQuery.GetSchema(), SchemaBoneIdx);
 
-	int32 DataOffset = ChannelDataOffset;
-	FFeatureVectorHelper::EncodeVector(InOutQuery.EditValues(), DataOffset, GetAxis(Transform.GetRotation()));
-	check(DataOffset == ChannelDataOffset + ChannelCardinality);
+		if (!bUseSampleTimeOffsetRootBone)
+		{
+			const FTransform RootTransform = SearchContext.TryGetTransformAndCacheResults(0.f, InOutQuery.GetSchema(), FSearchContext::SchemaRootBoneIdx);
+			const FTransform RootTransformPrev = SearchContext.TryGetTransformAndCacheResults(SampleTimeOffset, InOutQuery.GetSchema(), FSearchContext::SchemaRootBoneIdx);
+			Transform = Transform * (RootTransformPrev * RootTransform.Inverse());
+		}
 
-	return !bAnyError;
+		int32 DataOffset = ChannelDataOffset;
+		FFeatureVectorHelper::EncodeVector(InOutQuery.EditValues(), DataOffset, GetAxis(Transform.GetRotation()));
+		check(DataOffset == ChannelDataOffset + ChannelCardinality);
+	}
 }
 
 void UPoseSearchFeatureChannel_Heading::DebugDraw(const UE::PoseSearch::FDebugDrawParams& DrawParams, TConstArrayView<float> PoseVector) const
@@ -553,42 +551,38 @@ void UPoseSearchFeatureChannel_Heading::DebugDraw(const UE::PoseSearch::FDebugDr
 //////////////////////////////////////////////////////////////////////////
 // UPoseSearchFeatureChannel_Pose
 
-void UPoseSearchFeatureChannel_Pose::InitializeSchema(UE::PoseSearch::FSchemaInitializer& Initializer)
+void UPoseSearchFeatureChannel_Pose::InitializeSchema(UPoseSearchSchema* Schema)
 {
 	using namespace UE::PoseSearch;
 
-	Super::InitializeSchema(Initializer);
-
-	int32 DataOffset = ChannelDataOffset;
+	ChannelDataOffset = Schema->SchemaCardinality;
 	for (int32 ChannelBoneIdx = 0; ChannelBoneIdx != SampledBones.Num(); ++ChannelBoneIdx)
 	{
 		const FPoseSearchBone& SampledBone = SampledBones[ChannelBoneIdx];
 		if (EnumHasAnyFlags(SampledBone.Flags, EPoseSearchBoneFlags::Position))
 		{
-			DataOffset += FFeatureVectorHelper::EncodeVectorCardinality;
+			Schema->SchemaCardinality += FFeatureVectorHelper::EncodeVectorCardinality;
 		}
 		if (EnumHasAnyFlags(SampledBone.Flags, EPoseSearchBoneFlags::Rotation))
 		{
-			DataOffset += FFeatureVectorHelper::EncodeQuatCardinality;
+			Schema->SchemaCardinality += FFeatureVectorHelper::EncodeQuatCardinality;
 		}
 		if (EnumHasAnyFlags(SampledBone.Flags, EPoseSearchBoneFlags::Velocity))
 		{
-			DataOffset += FFeatureVectorHelper::EncodeVectorCardinality;
+			Schema->SchemaCardinality += FFeatureVectorHelper::EncodeVectorCardinality;
 		}
 		if (EnumHasAnyFlags(SampledBone.Flags, EPoseSearchBoneFlags::Phase))
 		{
-			DataOffset += FFeatureVectorHelper::EncodeVector2DCardinality;
+			Schema->SchemaCardinality += FFeatureVectorHelper::EncodeVector2DCardinality;
 		}
 	}
 
-	ChannelCardinality = DataOffset - ChannelDataOffset;
-
-	Initializer.SetCurrentChannelDataOffset(DataOffset);
+	ChannelCardinality = Schema->SchemaCardinality - ChannelDataOffset;
 
 	SchemaBoneIdx.Reset();
 	for (const FPoseSearchBone& Bone : SampledBones)
 	{
-		SchemaBoneIdx.Add(Initializer.AddBoneReference(Bone.Reference));
+		SchemaBoneIdx.Add(Schema->AddBoneReference(Bone.Reference));
 	}
 }
 
@@ -782,7 +776,7 @@ void UPoseSearchFeatureChannel_Pose::AddPoseFeatures(UE::PoseSearch::IAssetIndex
 	check(DataOffset == ChannelDataOffset + ChannelCardinality);
 }
 
-bool UPoseSearchFeatureChannel_Pose::BuildQuery(UE::PoseSearch::FSearchContext& SearchContext, FPoseSearchFeatureVectorBuilder& InOutQuery) const
+void UPoseSearchFeatureChannel_Pose::BuildQuery(UE::PoseSearch::FSearchContext& SearchContext, FPoseSearchFeatureVectorBuilder& InOutQuery) const
 {
 	using namespace UE::PoseSearch;
 
@@ -819,106 +813,99 @@ bool UPoseSearchFeatureChannel_Pose::BuildQuery(UE::PoseSearch::FSearchContext&
 			}
 		}
 		// else leave the InOutQuery set to zero since the SearchContext.History is invalid and it'll fail if we continue
-		return bSkip;
 	}
-
-	struct CachedTransforms
+	else
 	{
-		FTransform Current;
-		FTransform Previous;
-		bool Valid = false;
-	};
-	TArray<CachedTransforms, TInlineAllocator<32>> CachedTransforms;
-	CachedTransforms.AddUninitialized(SampledBones.Num());
-
-	bool bAnyError = false;
-	float SampleTime = 0.f;
+		struct CachedTransforms
+		{
+			FTransform Current;
+			FTransform Previous;
+			bool Valid = false; // @todo: remove this
+		};
+		TArray<CachedTransforms, TInlineAllocator<32>> CachedTransforms;
+		CachedTransforms.AddUninitialized(SampledBones.Num());
 
-	for (int32 SampledBoneIdx = 0; SampledBoneIdx != SampledBones.Num(); ++SampledBoneIdx)
-	{
-		CachedTransforms[SampledBoneIdx].Current = SearchContext.TryGetTransformAndCacheResults(SampleTime, InOutQuery.GetSchema(), SchemaBoneIdx[SampledBoneIdx], bAnyError);
-		const FPoseSearchBone& SampledBone = SampledBones[SampledBoneIdx];
+		float SampleTime = 0.f;
 
-		if (EnumHasAnyFlags(SampledBone.Flags, EPoseSearchBoneFlags::Velocity))
+		for (int32 SampledBoneIdx = 0; SampledBoneIdx != SampledBones.Num(); ++SampledBoneIdx)
 		{
-			check(SearchContext.History);
-			const float HistorySameplInterval = SearchContext.History->GetSampleTimeInterval();
-
-			CachedTransforms[SampledBoneIdx].Previous = SearchContext.TryGetTransformAndCacheResults(SampleTime - HistorySameplInterval, InOutQuery.GetSchema(), SchemaBoneIdx[SampledBoneIdx], bAnyError);
+			CachedTransforms[SampledBoneIdx].Current = SearchContext.TryGetTransformAndCacheResults(SampleTime, InOutQuery.GetSchema(), SchemaBoneIdx[SampledBoneIdx]);
+			const FPoseSearchBone& SampledBone = SampledBones[SampledBoneIdx];
 
-			if (!UE::PoseSearch::UseCharacterSpaceVelocities)
+			if (EnumHasAnyFlags(SampledBone.Flags, EPoseSearchBoneFlags::Velocity))
 			{
-				const FTransform RootTransform = SearchContext.TryGetTransformAndCacheResults(SampleTime, InOutQuery.GetSchema(), FSearchContext::SchemaRootBoneIdx, bAnyError);
-				const FTransform RootTransformPrev = SearchContext.TryGetTransformAndCacheResults(SampleTime - HistorySameplInterval, InOutQuery.GetSchema(), FSearchContext::SchemaRootBoneIdx, bAnyError);
+				check(SearchContext.History);
+				const float HistorySameplInterval = SearchContext.History->GetSampleTimeInterval();
 
-				// animation space velocity
-				CachedTransforms[SampledBoneIdx].Previous = CachedTransforms[SampledBoneIdx].Previous * (RootTransformPrev * RootTransform.Inverse());
-			}
-		}
-		CachedTransforms[SampledBoneIdx].Valid = true;
-	}
+				CachedTransforms[SampledBoneIdx].Previous = SearchContext.TryGetTransformAndCacheResults(SampleTime - HistorySameplInterval, InOutQuery.GetSchema(), SchemaBoneIdx[SampledBoneIdx]);
 
-	if (bAnyError)
-	{
-		return false;
-	}
+				if (!UE::PoseSearch::UseCharacterSpaceVelocities)
+				{
+					const FTransform RootTransform = SearchContext.TryGetTransformAndCacheResults(SampleTime, InOutQuery.GetSchema(), FSearchContext::SchemaRootBoneIdx);
+					const FTransform RootTransformPrev = SearchContext.TryGetTransformAndCacheResults(SampleTime - HistorySameplInterval, InOutQuery.GetSchema(), FSearchContext::SchemaRootBoneIdx);
 
-	int32 DataOffset = ChannelDataOffset;
-	for (int32 SampledBoneIdx = 0; SampledBoneIdx != SampledBones.Num(); ++SampledBoneIdx)
-	{
-		const FPoseSearchBone& SampledBone = SampledBones[SampledBoneIdx];
-		if (EnumHasAnyFlags(SampledBone.Flags, EPoseSearchBoneFlags::Position))
-		{
-			if (CachedTransforms[SampledBoneIdx].Valid)
-			{
-				FFeatureVectorHelper::EncodeVector(InOutQuery.EditValues(), DataOffset, CachedTransforms[SampledBoneIdx].Current.GetTranslation());
-			}
-			else
-			{
-				// preserve the InOutQuery.EditValues() and increase the DataOffset
-				DataOffset += FFeatureVectorHelper::EncodeVectorCardinality;
+					// animation space velocity
+					CachedTransforms[SampledBoneIdx].Previous = CachedTransforms[SampledBoneIdx].Previous * (RootTransformPrev * RootTransform.Inverse());
+				}
 			}
+			CachedTransforms[SampledBoneIdx].Valid = true;
 		}
 
-		if (EnumHasAnyFlags(SampledBone.Flags, EPoseSearchBoneFlags::Rotation))
+		int32 DataOffset = ChannelDataOffset;
+		for (int32 SampledBoneIdx = 0; SampledBoneIdx != SampledBones.Num(); ++SampledBoneIdx)
 		{
-			if (CachedTransforms[SampledBoneIdx].Valid)
+			const FPoseSearchBone& SampledBone = SampledBones[SampledBoneIdx];
+			if (EnumHasAnyFlags(SampledBone.Flags, EPoseSearchBoneFlags::Position))
 			{
-				FFeatureVectorHelper::EncodeQuat(InOutQuery.EditValues(), DataOffset, CachedTransforms[SampledBoneIdx].Current.GetRotation());
+				if (CachedTransforms[SampledBoneIdx].Valid)
+				{
+					FFeatureVectorHelper::EncodeVector(InOutQuery.EditValues(), DataOffset, CachedTransforms[SampledBoneIdx].Current.GetTranslation());
+				}
+				else
+				{
+					// preserve the InOutQuery.EditValues() and increase the DataOffset
+					DataOffset += FFeatureVectorHelper::EncodeVectorCardinality;
+				}
 			}
-			else
+
+			if (EnumHasAnyFlags(SampledBone.Flags, EPoseSearchBoneFlags::Rotation))
 			{
-				// preserve the InOutQuery.EditValues() and increase the DataOffset
-				DataOffset += FFeatureVectorHelper::EncodeQuatCardinality;
+				if (CachedTransforms[SampledBoneIdx].Valid)
+				{
+					FFeatureVectorHelper::EncodeQuat(InOutQuery.EditValues(), DataOffset, CachedTransforms[SampledBoneIdx].Current.GetRotation());
+				}
+				else
+				{
+					// preserve the InOutQuery.EditValues() and increase the DataOffset
+					DataOffset += FFeatureVectorHelper::EncodeQuatCardinality;
+				}
 			}
-		}
 
-		if (EnumHasAnyFlags(SampledBone.Flags, EPoseSearchBoneFlags::Velocity))
-		{
-			if (CachedTransforms[SampledBoneIdx].Valid)
+			if (EnumHasAnyFlags(SampledBone.Flags, EPoseSearchBoneFlags::Velocity))
 			{
-				const FVector LinearVelocity = (CachedTransforms[SampledBoneIdx].Current.GetTranslation() - CachedTransforms[SampledBoneIdx].Previous.GetTranslation()) / SearchContext.History->GetSampleTimeInterval();
-				FFeatureVectorHelper::EncodeVector(InOutQuery.EditValues(), DataOffset, LinearVelocity);
+				if (CachedTransforms[SampledBoneIdx].Valid)
+				{
+					const FVector LinearVelocity = (CachedTransforms[SampledBoneIdx].Current.GetTranslation() - CachedTransforms[SampledBoneIdx].Previous.GetTranslation()) / SearchContext.History->GetSampleTimeInterval();
+					FFeatureVectorHelper::EncodeVector(InOutQuery.EditValues(), DataOffset, LinearVelocity);
+				}
+				else
+				{
+					// preserve the InOutQuery.EditValues() and increase the DataOffset
+					DataOffset += FFeatureVectorHelper::EncodeVectorCardinality;
+				}
 			}
-			else
+
+			if (EnumHasAnyFlags(SampledBone.Flags, EPoseSearchBoneFlags::Phase))
 			{
-				// preserve the InOutQuery.EditValues() and increase the DataOffset
-				DataOffset += FFeatureVectorHelper::EncodeVectorCardinality;
+				DataOffset += FFeatureVectorHelper::EncodeVector2DCardinality;
+
+				// @todo: Support phase in BuildQuery
+				// FFeatureVectorHelper::EncodeVector2D(InOutQuery.EditValues(), DataOffset, ???);
 			}
 		}
 
-		if (EnumHasAnyFlags(SampledBone.Flags, EPoseSearchBoneFlags::Phase))
-		{
-			DataOffset += FFeatureVectorHelper::EncodeVector2DCardinality;
-
-			// @todo: Support phase in BuildQuery
-			// FFeatureVectorHelper::EncodeVector2D(InOutQuery.EditValues(), DataOffset, ???);
-		}
+		check(DataOffset == ChannelDataOffset + ChannelCardinality);
 	}
-
-	check(DataOffset == ChannelDataOffset + ChannelCardinality);
-
-	return true;
 }
 
 void UPoseSearchFeatureChannel_Pose::DebugDraw(const UE::PoseSearch::FDebugDrawParams& DrawParams, TConstArrayView<float> PoseVector) const
@@ -1132,56 +1119,52 @@ void UPoseSearchFeatureChannel_Trajectory::PreSave(FObjectPreSaveContext ObjectS
 	Super::PreSave(ObjectSaveContext);
 }
 
-void UPoseSearchFeatureChannel_Trajectory::InitializeSchema(UE::PoseSearch::FSchemaInitializer& Initializer)
+void UPoseSearchFeatureChannel_Trajectory::InitializeSchema(UPoseSearchSchema* Schema)
 {
 	using namespace UE::PoseSearch;
 
-	Super::InitializeSchema(Initializer);
-
-	int32 DataOffset = ChannelDataOffset;
+	ChannelDataOffset = Schema->SchemaCardinality;
 
 	for (const FPoseSearchTrajectorySample& Sample : Samples)
 	{
 		if (EnumHasAnyFlags(Sample.Flags, EPoseSearchTrajectoryFlags::Position))
 		{
-			DataOffset += FFeatureVectorHelper::EncodeVectorCardinality;
+			Schema->SchemaCardinality += FFeatureVectorHelper::EncodeVectorCardinality;
 		}
 		if (EnumHasAnyFlags(Sample.Flags, EPoseSearchTrajectoryFlags::PositionXY))
 		{
-			DataOffset += FFeatureVectorHelper::EncodeVector2DCardinality;
+			Schema->SchemaCardinality += FFeatureVectorHelper::EncodeVector2DCardinality;
 		}
 
 		if (EnumHasAnyFlags(Sample.Flags, EPoseSearchTrajectoryFlags::Velocity))
 		{
-			DataOffset += FFeatureVectorHelper::EncodeVectorCardinality;
+			Schema->SchemaCardinality += FFeatureVectorHelper::EncodeVectorCardinality;
 		}
 		if (EnumHasAnyFlags(Sample.Flags, EPoseSearchTrajectoryFlags::VelocityXY))
 		{
-			DataOffset += FFeatureVectorHelper::EncodeVector2DCardinality;
+			Schema->SchemaCardinality += FFeatureVectorHelper::EncodeVector2DCardinality;
 		}
 
 		if (EnumHasAnyFlags(Sample.Flags, EPoseSearchTrajectoryFlags::VelocityDirection))
 		{
-			DataOffset += FFeatureVectorHelper::EncodeVectorCardinality;
+			Schema->SchemaCardinality += FFeatureVectorHelper::EncodeVectorCardinality;
 		}
 		if (EnumHasAnyFlags(Sample.Flags, EPoseSearchTrajectoryFlags::VelocityDirectionXY))
 		{
-			DataOffset += FFeatureVectorHelper::EncodeVector2DCardinality;
+			Schema->SchemaCardinality += FFeatureVectorHelper::EncodeVector2DCardinality;
 		}
 
 		if (EnumHasAnyFlags(Sample.Flags, EPoseSearchTrajectoryFlags::FacingDirection))
 		{
-			DataOffset += FFeatureVectorHelper::EncodeVectorCardinality;
+			Schema->SchemaCardinality += FFeatureVectorHelper::EncodeVectorCardinality;
 		}
 		if (EnumHasAnyFlags(Sample.Flags, EPoseSearchTrajectoryFlags::FacingDirectionXY))
 		{
-			DataOffset += FFeatureVectorHelper::EncodeVector2DCardinality;
+			Schema->SchemaCardinality += FFeatureVectorHelper::EncodeVector2DCardinality;
 		}
 	}
 
-	ChannelCardinality = DataOffset - ChannelDataOffset;
-
-	Initializer.SetCurrentChannelDataOffset(DataOffset);
+	ChannelCardinality = Schema->SchemaCardinality - ChannelDataOffset;
 }
 
 void UPoseSearchFeatureChannel_Trajectory::FillWeights(TArray<float>& Weights) const
@@ -1365,14 +1348,14 @@ void UPoseSearchFeatureChannel_Trajectory::IndexAssetPrivate(const UE::PoseSearc
 	check(DataOffset == ChannelDataOffset + ChannelCardinality);
 }
 
-bool UPoseSearchFeatureChannel_Trajectory::BuildQuery(UE::PoseSearch::FSearchContext& SearchContext, FPoseSearchFeatureVectorBuilder& InOutQuery) const
+void UPoseSearchFeatureChannel_Trajectory::BuildQuery(UE::PoseSearch::FSearchContext& SearchContext, FPoseSearchFeatureVectorBuilder& InOutQuery) const
 {
 	using namespace UE::PoseSearch;
 
 	if (!SearchContext.Trajectory)
 	{
 		// @todo: do we want to reuse the SearchContext.CurrentResult data if valid?
-		return false;
+		return;
 	}
 
 	int32 NextIterStartIdx = 0;
@@ -1424,8 +1407,6 @@ bool UPoseSearchFeatureChannel_Trajectory::BuildQuery(UE::PoseSearch::FSearchCon
 		}
 	}
 	check(DataOffset == ChannelDataOffset + ChannelCardinality);
-
-	return true;
 }
 
 // lazy initialized helper to interpolate or extrapolate (linearly) UPoseSearchFeatureChannel_Trajectory trajectory positions from FPoseSearchTrajectorySample containing EPoseSearchTrajectoryFlags::Position for samples without it
@@ -2074,16 +2055,16 @@ static float ComputeCrashingLegsValue(const FVector& RightThighPos, const FVecto
 	return CrashingLegsValue;
 }
 
-void UPoseSearchFeatureChannel_FilterCrashingLegs::InitializeSchema(UE::PoseSearch::FSchemaInitializer& Initializer)
+void UPoseSearchFeatureChannel_FilterCrashingLegs::InitializeSchema(UPoseSearchSchema* Schema)
 {
-	Super::InitializeSchema(Initializer);
+	ChannelDataOffset = Schema->SchemaCardinality;
 	ChannelCardinality = UE::PoseSearch::FFeatureVectorHelper::EncodeFloatCardinality;
-	Initializer.SetCurrentChannelDataOffset(ChannelDataOffset + ChannelCardinality);
+	Schema->SchemaCardinality += ChannelCardinality;
 
-	LeftThighIdx = Initializer.AddBoneReference(LeftThigh);
-	RightThighIdx = Initializer.AddBoneReference(RightThigh);
-	LeftFootIdx = Initializer.AddBoneReference(LeftFoot);
-	RightFootIdx = Initializer.AddBoneReference(RightFoot);
+	LeftThighIdx = Schema->AddBoneReference(LeftThigh);
+	RightThighIdx = Schema->AddBoneReference(RightThigh);
+	LeftFootIdx = Schema->AddBoneReference(LeftFoot);
+	RightFootIdx = Schema->AddBoneReference(RightFoot);
 }
 
 void UPoseSearchFeatureChannel_FilterCrashingLegs::FillWeights(TArray<float>& Weights) const
@@ -2116,7 +2097,7 @@ void UPoseSearchFeatureChannel_FilterCrashingLegs::IndexAsset(UE::PoseSearch::IA
 	}
 }
 
-bool UPoseSearchFeatureChannel_FilterCrashingLegs::BuildQuery(UE::PoseSearch::FSearchContext& SearchContext, FPoseSearchFeatureVectorBuilder& InOutQuery) const
+void UPoseSearchFeatureChannel_FilterCrashingLegs::BuildQuery(UE::PoseSearch::FSearchContext& SearchContext, FPoseSearchFeatureVectorBuilder& InOutQuery) const
 {
 	using namespace UE::PoseSearch;
 
@@ -2132,19 +2113,18 @@ bool UPoseSearchFeatureChannel_FilterCrashingLegs::BuildQuery(UE::PoseSearch::FS
 			FFeatureVectorHelper::EncodeFloat(InOutQuery.EditValues(), DataOffset, SearchContext.GetCurrentResultPrevPoseVector(), SearchContext.GetCurrentResultPoseVector(), SearchContext.GetCurrentResultNextPoseVector(), LerpValue);
 		}
 		// else leave the InOutQuery set to zero since the SearchContext.History is invalid and it'll fail if we continue
-		return bSkip;
 	}
+	else
+	{
+		const float SampleTime = 0.f;
+		const FTransform LeftThighTransform = SearchContext.TryGetTransformAndCacheResults(SampleTime, InOutQuery.GetSchema(), LeftThighIdx);
+		const FTransform RightThighTransform = SearchContext.TryGetTransformAndCacheResults(SampleTime, InOutQuery.GetSchema(), RightThighIdx);
+		const FTransform LeftFootTransform = SearchContext.TryGetTransformAndCacheResults(SampleTime, InOutQuery.GetSchema(), LeftFootIdx);
+		const FTransform RightFootTransform = SearchContext.TryGetTransformAndCacheResults(SampleTime, InOutQuery.GetSchema(), RightFootIdx);
 
-	const float SampleTime = 0.f;
-	bool bAnyError = false;
-	const FTransform LeftThighTransform = SearchContext.TryGetTransformAndCacheResults(SampleTime, InOutQuery.GetSchema(), LeftThighIdx, bAnyError);
-	const FTransform RightThighTransform = SearchContext.TryGetTransformAndCacheResults(SampleTime, InOutQuery.GetSchema(), RightThighIdx, bAnyError);
-	const FTransform LeftFootTransform = SearchContext.TryGetTransformAndCacheResults(SampleTime, InOutQuery.GetSchema(), LeftFootIdx, bAnyError);
-	const FTransform RightFootTransform = SearchContext.TryGetTransformAndCacheResults(SampleTime, InOutQuery.GetSchema(), RightFootIdx, bAnyError);
-
-	const float CrashingLegsValue = ComputeCrashingLegsValue(RightThighTransform.GetTranslation(), LeftThighTransform.GetTranslation(), RightFootTransform.GetTranslation(), LeftFootTransform.GetTranslation());
-	FFeatureVectorHelper::EncodeFloat(InOutQuery.EditValues(), DataOffset, CrashingLegsValue);
-	return !bAnyError;
+		const float CrashingLegsValue = ComputeCrashingLegsValue(RightThighTransform.GetTranslation(), LeftThighTransform.GetTranslation(), RightFootTransform.GetTranslation(), LeftFootTransform.GetTranslation());
+		FFeatureVectorHelper::EncodeFloat(InOutQuery.EditValues(), DataOffset, CrashingLegsValue);
+	}
 }
 
 void UPoseSearchFeatureChannel_FilterCrashingLegs::DebugDraw(const UE::PoseSearch::FDebugDrawParams& DrawParams, TConstArrayView<float> PoseVector) const
diff --git a/Engine/Plugins/Experimental/Animation/PoseSearch/Source/Runtime/Public/PoseSearch/PoseSearch.h b/Engine/Plugins/Experimental/Animation/PoseSearch/Source/Runtime/Public/PoseSearch/PoseSearch.h
--- a/Engine/Plugins/Experimental/Animation/PoseSearch/Source/Runtime/Public/PoseSearch/PoseSearch.h
+++ b/Engine/Plugins/Experimental/Animation/PoseSearch/Source/Runtime/Public/PoseSearch/PoseSearch.h
@@ -36,11 +36,9 @@ DECLARE_LOG_CATEGORY_EXTERN(LogPoseSearch, Log, All);
 //////////////////////////////////////////////////////////////////////////
 // Forward declarations
 
-class UAnimSequence;
 class UBlendSpace;
 struct FCompactPose;
 struct FPoseContext;
-struct FReferenceSkeleton;
 class UAnimNotifyState_PoseSearchBase;
 class UPoseSearchSchema;
 
@@ -54,24 +52,14 @@ constexpr EParallelForFlags ParallelForFlags = EParallelForFlags::ForceSingleThr
 constexpr EParallelForFlags ParallelForFlags = EParallelForFlags::None;
 #endif // UE_POSE_SEARCH_FORCE_SINGLE_THREAD
 
-
-
 namespace UE::PoseSearch
 {
 class FPoseHistory;
 struct FDebugDrawParams;
-struct FSchemaInitializer;
 struct FQueryBuildingContext;
 struct FSearchContext;
 } // namespace UE::PoseSearch
 
-#if WITH_EDITOR
-namespace UE::DerivedData
-{
-	class IRequestOwner;
-} // namespace UE::DerivedData
-#endif
-
 //////////////////////////////////////////////////////////////////////////
 // Constants
 
@@ -479,12 +467,11 @@ class POSESEARCH_API UPoseSearchFeatureChannel : public UObject, public IBoneRef
 	GENERATED_BODY()
 
 public:
-	int32 GetChannelIndex() const { checkSlow(ChannelIdx >= 0); return ChannelIdx; }
 	int32 GetChannelCardinality() const { checkSlow(ChannelCardinality >= 0); return ChannelCardinality; }
 	int32 GetChannelDataOffset() const { checkSlow(ChannelDataOffset >= 0); return ChannelDataOffset; }
 
 	// Called during UPoseSearchSchema::Finalize to prepare the schema for this channel
-	virtual void InitializeSchema(UE::PoseSearch::FSchemaInitializer& Initializer);
+	virtual void InitializeSchema(UPoseSearchSchema* Schema) PURE_VIRTUAL(UPoseSearchFeatureChannel::InitializeSchema, );
 	
 	// Called at database build time to collect feature weights.
 	// Weights is sized to the cardinality of the schema and the feature channel should write
@@ -495,7 +482,7 @@ public:
 	virtual void IndexAsset(UE::PoseSearch::IAssetIndexer& Indexer, UE::PoseSearch::FAssetIndexingOutput& IndexingOutput) const PURE_VIRTUAL(UPoseSearchFeatureChannel::IndexAsset, );
 
 	// Called at runtime to add this channel's data to the query pose vector
-	virtual bool BuildQuery(UE::PoseSearch::FSearchContext& SearchContext, FPoseSearchFeatureVectorBuilder& InOutQuery) const PURE_VIRTUAL(UPoseSearchFeatureChannel::BuildQuery, return false;);
+	virtual void BuildQuery(UE::PoseSearch::FSearchContext& SearchContext, FPoseSearchFeatureVectorBuilder& InOutQuery) const PURE_VIRTUAL(UPoseSearchFeatureChannel::BuildQuery, );
 
 	// Draw this channel's data for the given pose vector
 	virtual void DebugDraw(const UE::PoseSearch::FDebugDrawParams& DrawParams, TConstArrayView<float> PoseVector) const PURE_VIRTUAL(UPoseSearchFeatureChannel::DebugDraw, );
@@ -510,12 +497,9 @@ private:
 	// Note this function is exclusively for FBoneReference details customization
 	class USkeleton* GetSkeleton(bool& bInvalidSkeletonIsError, const IPropertyHandle* PropertyHandle) override;
 
+protected:
 	friend class ::UPoseSearchSchema;
 
-	UPROPERTY(meta = (ExcludeFromHash))
-	int32 ChannelIdx = INDEX_NONE;
-
-protected:
 	UPROPERTY(meta = (ExcludeFromHash))
 	int32 ChannelDataOffset = INDEX_NONE;
 
@@ -523,30 +507,6 @@ protected:
 	int32 ChannelCardinality = INDEX_NONE;
 };
 
-namespace UE::PoseSearch
-{
-
-struct POSESEARCH_API FSchemaInitializer
-{
-public:
-	int32 AddBoneReference(const FBoneReference& BoneReference);
-
-	// Gets the index into the schema's channel array for the channel currently being initialized
-	int32 GetCurrentChannelIdx() const { return CurrentChannelIdx; }
-
-	int32 GetCurrentChannelDataOffset() const { return CurrentChannelDataOffset; }
-	void SetCurrentChannelDataOffset(int32 DataOffset) { CurrentChannelDataOffset = DataOffset; }
-
-private:
-	friend class ::UPoseSearchSchema;
-
-	int32 CurrentChannelIdx = 0;
-	int32 CurrentChannelDataOffset = 0;
-	TArray<FBoneReference> BoneReferences;
-};
-
-} // namespace UE::PoseSearch
-
 USTRUCT()
 struct FPoseSearchSchemaColorPreset
 {
@@ -562,7 +522,7 @@ struct FPoseSearchSchemaColorPreset
 /**
 * Specifies the format of a pose search index. At runtime, queries are built according to the schema for searching.
 */
-UCLASS(BlueprintType, Category = "Animation|Pose Search", Experimental, meta = (DisplayName = "Motion Database Config"))
+UCLASS(BlueprintType, Category = "Animation|Pose Search", Experimental, meta = (DisplayName = "Motion Database Config"), CollapseCategories)
 class POSESEARCH_API UPoseSearchSchema : public UDataAsset, public IBoneReferenceSkeletonProvider
 {
 	GENERATED_BODY()
@@ -647,12 +607,14 @@ public:
 	virtual void PostEditChangeProperty(struct FPropertyChangedEvent& PropertyChangedEvent) override;
 #endif
 
+	int32 AddBoneReference(const FBoneReference& BoneReference) { return BoneReferences.AddUnique(BoneReference); }
+
 private:
 	void Finalize(bool bRemoveEmptyChannels = true);
 	void ResolveBoneReferences();
 
 public:
-	bool BuildQuery(UE::PoseSearch::FSearchContext& SearchContext, FPoseSearchFeatureVectorBuilder& InOutQuery) const;
+	void BuildQuery(UE::PoseSearch::FSearchContext& SearchContext, FPoseSearchFeatureVectorBuilder& InOutQuery) const;
 };
 
 //////////////////////////////////////////////////////////////////////////
@@ -1376,14 +1338,14 @@ struct POSESEARCH_API FSearchContext
 	// can the continuing pose advance? (if not we skip evaluating it)
 	bool bCanAdvance = true;
 
-	FTransform TryGetTransformAndCacheResults(float SampleTime, const UPoseSearchSchema* Schema, int8 SchemaBoneIdx, bool& Error);
+	FTransform TryGetTransformAndCacheResults(float SampleTime, const UPoseSearchSchema* Schema, int8 SchemaBoneIdx);
 	void ClearCachedEntries();
 
 	void ResetCurrentBestCost();
 	void UpdateCurrentBestCost(const FPoseSearchCost& PoseSearchCost);
 	float GetCurrentBestTotalCost() const { return CurrentBestTotalCost; }
 
-	bool GetOrBuildQuery(const UPoseSearchDatabase* Database, FPoseSearchFeatureVectorBuilder& FeatureVectorBuilder);
+	void GetOrBuildQuery(const UPoseSearchDatabase* Database, FPoseSearchFeatureVectorBuilder& FeatureVectorBuilder);
 	const FPoseSearchFeatureVectorBuilder* GetCachedQuery(const UPoseSearchDatabase* Database) const;
 
 	bool IsCurrentResultFromDatabase(const UPoseSearchDatabase* Database) const;
diff --git a/Engine/Plugins/Experimental/Animation/PoseSearch/Source/Runtime/Public/PoseSearch/PoseSearchFeatureChannels.h b/Engine/Plugins/Experimental/Animation/PoseSearch/Source/Runtime/Public/PoseSearch/PoseSearchFeatureChannels.h
--- a/Engine/Plugins/Experimental/Animation/PoseSearch/Source/Runtime/Public/PoseSearch/PoseSearchFeatureChannels.h
+++ b/Engine/Plugins/Experimental/Animation/PoseSearch/Source/Runtime/Public/PoseSearch/PoseSearchFeatureChannels.h
@@ -26,7 +26,7 @@ enum class EInputQueryPose : uint8
 
 //////////////////////////////////////////////////////////////////////////
 // UPoseSearchFeatureChannel_Position
-UCLASS(BlueprintType, EditInlineNew)
+UCLASS(BlueprintType, EditInlineNew, meta = (DisplayName = "Position Channel"), CollapseCategories)
 class POSESEARCH_API UPoseSearchFeatureChannel_Position : public UPoseSearchFeatureChannel
 {
 	GENERATED_BODY()
@@ -56,10 +56,10 @@ public:
 	bool bUseSampleTimeOffsetRootBone = true;
 
 	// UPoseSearchFeatureChannel interface
-	virtual void InitializeSchema(UE::PoseSearch::FSchemaInitializer& Initializer) override;
+	virtual void InitializeSchema(UPoseSearchSchema* Schema) override;
 	virtual void FillWeights(TArray<float>& Weights) const override;
 	virtual void IndexAsset(UE::PoseSearch::IAssetIndexer& Indexer, UE::PoseSearch::FAssetIndexingOutput& IndexingOutput) const override;
-	virtual bool BuildQuery(UE::PoseSearch::FSearchContext& SearchContext, FPoseSearchFeatureVectorBuilder& InOutQuery) const override;
+	virtual void BuildQuery(UE::PoseSearch::FSearchContext& SearchContext, FPoseSearchFeatureVectorBuilder& InOutQuery) const override;
 	virtual void DebugDraw(const UE::PoseSearch::FDebugDrawParams& DrawParams, TConstArrayView<float> PoseVector) const override;
 };
 
@@ -77,7 +77,7 @@ enum class EHeadingAxis : uint8
 	Invalid = Num UMETA(Hidden)
 };
 
-UCLASS(BlueprintType, EditInlineNew)
+UCLASS(BlueprintType, EditInlineNew, meta = (DisplayName = "Heading Channel"), CollapseCategories)
 class POSESEARCH_API UPoseSearchFeatureChannel_Heading : public UPoseSearchFeatureChannel
 {
 	GENERATED_BODY()
@@ -110,10 +110,10 @@ public:
 	bool bUseSampleTimeOffsetRootBone = true;
 
 	// UPoseSearchFeatureChannel interface
-	virtual void InitializeSchema(UE::PoseSearch::FSchemaInitializer& Initializer) override;
+	virtual void InitializeSchema(UPoseSearchSchema* Schema) override;
 	virtual void FillWeights(TArray<float>& Weights) const override;
 	virtual void IndexAsset(UE::PoseSearch::IAssetIndexer& Indexer, UE::PoseSearch::FAssetIndexingOutput& IndexingOutput) const override;
-	virtual bool BuildQuery(UE::PoseSearch::FSearchContext& SearchContext, FPoseSearchFeatureVectorBuilder& InOutQuery) const override;
+	virtual void BuildQuery(UE::PoseSearch::FSearchContext& SearchContext, FPoseSearchFeatureVectorBuilder& InOutQuery) const override;
 	virtual void DebugDraw(const UE::PoseSearch::FDebugDrawParams& DrawParams, TConstArrayView<float> PoseVector) const override;
 
 	FVector GetAxis(const FQuat& Rotation) const;
@@ -151,7 +151,7 @@ struct POSESEARCH_API FPoseSearchBone
 };
 
 // UPoseSearchFeatureChannel_Pose
-UCLASS(BlueprintType, EditInlineNew)
+UCLASS(BlueprintType, EditInlineNew, meta = (DisplayName = "Pose Channel"), CollapseCategories)
 class POSESEARCH_API UPoseSearchFeatureChannel_Pose : public UPoseSearchFeatureChannel
 {
 	GENERATED_BODY()
@@ -170,10 +170,10 @@ public:
 	EInputQueryPose InputQueryPose = EInputQueryPose::UseContinuingPose;
 
 	// UPoseSearchFeatureChannel interface
-	virtual void InitializeSchema(UE::PoseSearch::FSchemaInitializer& Initializer) override;
+	virtual void InitializeSchema(UPoseSearchSchema* Schema) override;
 	virtual void FillWeights(TArray<float>& Weights) const override;
 	virtual void IndexAsset(UE::PoseSearch::IAssetIndexer& Indexer, UE::PoseSearch::FAssetIndexingOutput& IndexingOutput) const override;
-	virtual bool BuildQuery(UE::PoseSearch::FSearchContext& SearchContext, FPoseSearchFeatureVectorBuilder& InOutQuery) const override;
+	virtual void BuildQuery(UE::PoseSearch::FSearchContext& SearchContext, FPoseSearchFeatureVectorBuilder& InOutQuery) const override;
 	virtual void DebugDraw(const UE::PoseSearch::FDebugDrawParams& DrawParams, TConstArrayView<float> PoseVector) const override;
 
 #if WITH_EDITOR
@@ -222,7 +222,7 @@ struct POSESEARCH_API FPoseSearchTrajectorySample
 	int32 ColorPresetIndex = 0;
 };
 
-UCLASS(BlueprintType, EditInlineNew)
+UCLASS(BlueprintType, EditInlineNew, meta = (DisplayName = "Trajectory Channel"), CollapseCategories)
 class POSESEARCH_API UPoseSearchFeatureChannel_Trajectory : public UPoseSearchFeatureChannel
 {
 	GENERATED_BODY()
@@ -238,10 +238,10 @@ public:
 	virtual void PreSave(FObjectPreSaveContext ObjectSaveContext) override;
 
 	// UPoseSearchFeatureChannel interface
-	virtual void InitializeSchema(UE::PoseSearch::FSchemaInitializer& Initializer) override;
+	virtual void InitializeSchema(UPoseSearchSchema* Schema) override;
 	virtual void FillWeights(TArray<float>& Weights) const override;
 	virtual void IndexAsset(UE::PoseSearch::IAssetIndexer& Indexer, UE::PoseSearch::FAssetIndexingOutput& IndexingOutput) const override;
-	virtual bool BuildQuery(UE::PoseSearch::FSearchContext& SearchContext, FPoseSearchFeatureVectorBuilder& InOutQuery) const override;
+	virtual void BuildQuery(UE::PoseSearch::FSearchContext& SearchContext, FPoseSearchFeatureVectorBuilder& InOutQuery) const override;
 	virtual void DebugDraw(const UE::PoseSearch::FDebugDrawParams& DrawParams, TConstArrayView<float> PoseVector) const override;
 
 #if WITH_EDITOR
@@ -264,7 +264,7 @@ protected:
 // -0.5 if the left foot is exactly in front of the right foot
 // close to 1 or -1 if the feet (and so the legs) are completely crossed
 // at runtime we'll match the CrashingLegsValue and also filter by discarding pose candidates that don't respect the 'AllowedTolerance' between query and database values (happening in IsPoseValid)
-UCLASS(BlueprintType, EditInlineNew)
+UCLASS(BlueprintType, EditInlineNew, meta = (DisplayName = "CrashingLegs Channel"), CollapseCategories)
 class POSESEARCH_API UPoseSearchFeatureChannel_FilterCrashingLegs : public UPoseSearchFeatureChannel
 {
 	GENERATED_BODY()
@@ -304,13 +304,14 @@ public:
 	float AllowedTolerance = 0.3f;
 
 	// UPoseSearchFeatureChannel interface
-	virtual void InitializeSchema(UE::PoseSearch::FSchemaInitializer& Initializer) override;
+	virtual void InitializeSchema(UPoseSearchSchema* Schema) override;
 	virtual void FillWeights(TArray<float>& Weights) const override;
 	virtual void IndexAsset(UE::PoseSearch::IAssetIndexer& Indexer, UE::PoseSearch::FAssetIndexingOutput& IndexingOutput) const override;
-	virtual bool BuildQuery(UE::PoseSearch::FSearchContext& SearchContext, FPoseSearchFeatureVectorBuilder& InOutQuery) const override;
+	virtual void BuildQuery(UE::PoseSearch::FSearchContext& SearchContext, FPoseSearchFeatureVectorBuilder& InOutQuery) const override;
 	virtual void DebugDraw(const UE::PoseSearch::FDebugDrawParams& DrawParams, TConstArrayView<float> PoseVector) const override;
 
 	// IPoseFilter interface
 	virtual bool IsPoseFilterActive() const override;
 	virtual bool IsPoseValid(TConstArrayView<float> PoseValues, TConstArrayView<float> QueryValues, int32 PoseIdx, const FPoseSearchPoseMetadata& Metadata) const override;
-};
\ No newline at end of file
+};
+
