PoseSearch: fix constness when accessing instanced structs.

diff --git a/Engine/Plugins/Experimental/Animation/PoseSearch/Source/Editor/Private/PoseSearchDatabaseViewModel.cpp b/Engine/Plugins/Experimental/Animation/PoseSearch/Source/Editor/Private/PoseSearchDatabaseViewModel.cpp
--- a/Engine/Plugins/Experimental/Animation/PoseSearch/Source/Editor/Private/PoseSearchDatabaseViewModel.cpp
+++ b/Engine/Plugins/Experimental/Animation/PoseSearch/Source/Editor/Private/PoseSearchDatabaseViewModel.cpp
@@ -488,7 +488,7 @@ namespace UE::PoseSearch
 
 	void FDatabaseViewModel::SetIsEnabled(int32 AnimationAssetIndex, bool bEnabled)
 	{
-		if (FPoseSearchDatabaseAnimationAssetBase* DatabaseAnimationAsset = PoseSearchDatabase->GetAnimationAssetBase(AnimationAssetIndex))
+		if (FPoseSearchDatabaseAnimationAssetBase* DatabaseAnimationAsset = PoseSearchDatabase->GetMutableAnimationAssetBase(AnimationAssetIndex))
 		{
 			DatabaseAnimationAsset->SetIsEnabled(bEnabled);
 		}
diff --git a/Engine/Plugins/Experimental/Animation/PoseSearch/Source/Runtime/Private/PoseSearch.cpp b/Engine/Plugins/Experimental/Animation/PoseSearch/Source/Runtime/Private/PoseSearch.cpp
--- a/Engine/Plugins/Experimental/Animation/PoseSearch/Source/Runtime/Private/PoseSearch.cpp
+++ b/Engine/Plugins/Experimental/Animation/PoseSearch/Source/Runtime/Private/PoseSearch.cpp
@@ -813,7 +813,7 @@ const FInstancedStruct& UPoseSearchDatabase::GetAnimationAssetStruct(const FPose
 	return GetAnimationAssetStruct(SearchIndexAsset.SourceAssetIdx);
 }
 
-FPoseSearchDatabaseAnimationAssetBase* UPoseSearchDatabase::GetAnimationAssetBase(int32 AnimationAssetIndex) const
+const FPoseSearchDatabaseAnimationAssetBase* UPoseSearchDatabase::GetAnimationAssetBase(int32 AnimationAssetIndex) const
 {
 	if (AnimationAssets.IsValidIndex(AnimationAssetIndex))
 	{
@@ -823,11 +823,26 @@ FPoseSearchDatabaseAnimationAssetBase* UPoseSearchDatabase::GetAnimationAssetBas
 	return nullptr;
 }
 
-FPoseSearchDatabaseAnimationAssetBase* UPoseSearchDatabase::GetAnimationAssetBase(const FPoseSearchIndexAsset& SearchIndexAsset) const
+const FPoseSearchDatabaseAnimationAssetBase* UPoseSearchDatabase::GetAnimationAssetBase(const FPoseSearchIndexAsset& SearchIndexAsset) const
 {
 	return GetAnimationAssetBase(SearchIndexAsset.SourceAssetIdx);
 }
 
+FPoseSearchDatabaseAnimationAssetBase* UPoseSearchDatabase::GetMutableAnimationAssetBase(int32 AnimationAssetIndex)
+{
+	if (AnimationAssets.IsValidIndex(AnimationAssetIndex))
+	{
+		return AnimationAssets[AnimationAssetIndex].GetMutablePtr<FPoseSearchDatabaseAnimationAssetBase>();
+	}
+
+	return nullptr;
+}
+
+FPoseSearchDatabaseAnimationAssetBase* UPoseSearchDatabase::GetMutableAnimationAssetBase(const FPoseSearchIndexAsset& SearchIndexAsset)
+{
+	return GetMutableAnimationAssetBase(SearchIndexAsset.SourceAssetIdx);
+}
+
 const bool UPoseSearchDatabase::IsSourceAssetLooping(const FPoseSearchIndexAsset& SearchIndexAsset) const
 {
 	return GetAnimationAssetBase(SearchIndexAsset.SourceAssetIdx)->IsLooping();
diff --git a/Engine/Plugins/Experimental/Animation/PoseSearch/Source/Runtime/Public/PoseSearch/PoseSearch.h b/Engine/Plugins/Experimental/Animation/PoseSearch/Source/Runtime/Public/PoseSearch/PoseSearch.h
--- a/Engine/Plugins/Experimental/Animation/PoseSearch/Source/Runtime/Public/PoseSearch/PoseSearch.h
+++ b/Engine/Plugins/Experimental/Animation/PoseSearch/Source/Runtime/Public/PoseSearch/PoseSearch.h
@@ -1167,8 +1167,10 @@ public:
 
 	const FInstancedStruct& GetAnimationAssetStruct(int32 AnimationAssetIndex) const;
 	const FInstancedStruct& GetAnimationAssetStruct(const FPoseSearchIndexAsset& SearchIndexAsset) const;
-	FPoseSearchDatabaseAnimationAssetBase* GetAnimationAssetBase(int32 AnimationAssetIndex) const;
-	FPoseSearchDatabaseAnimationAssetBase* GetAnimationAssetBase(const FPoseSearchIndexAsset& SearchIndexAsset) const;
+	const FPoseSearchDatabaseAnimationAssetBase* GetAnimationAssetBase(int32 AnimationAssetIndex) const;
+	const FPoseSearchDatabaseAnimationAssetBase* GetAnimationAssetBase(const FPoseSearchIndexAsset& SearchIndexAsset) const;
+	FPoseSearchDatabaseAnimationAssetBase* GetMutableAnimationAssetBase(int32 AnimationAssetIndex);
+	FPoseSearchDatabaseAnimationAssetBase* GetMutableAnimationAssetBase(const FPoseSearchIndexAsset& SearchIndexAsset);
 	const bool IsSourceAssetLooping(const FPoseSearchIndexAsset& SearchIndexAsset) const;
 	const FString GetSourceAssetName(const FPoseSearchIndexAsset& SearchIndexAsset) const;
 	int32 GetNumberOfPrincipalComponents() const;
