PoseSearch: Another try at fixing the const GetMutableAnimationAssetBase() getter.

diff --git a/Engine/Plugins/Experimental/Animation/PoseSearch/Source/Runtime/Private/PoseSearch.cpp b/Engine/Plugins/Experimental/Animation/PoseSearch/Source/Runtime/Private/PoseSearch.cpp
--- a/Engine/Plugins/Experimental/Animation/PoseSearch/Source/Runtime/Private/PoseSearch.cpp
+++ b/Engine/Plugins/Experimental/Animation/PoseSearch/Source/Runtime/Private/PoseSearch.cpp
@@ -817,7 +817,7 @@ const FPoseSearchDatabaseAnimationAssetBase* UPoseSearchDatabase::GetAnimationAs
 {
 	if (AnimationAssets.IsValidIndex(AnimationAssetIndex))
 	{
-		return AnimationAssets[AnimationAssetIndex].GetMutablePtr<FPoseSearchDatabaseAnimationAssetBase>();
+		return AnimationAssets[AnimationAssetIndex].GetPtr<FPoseSearchDatabaseAnimationAssetBase>();
 	}
 
 	return nullptr;
