PoseSearch: One more batch of instanced struct changes.

diff --git a/Engine/Plugins/Experimental/Animation/PoseSearch/Source/Editor/Private/PoseSearchDatabaseEditorReflection.cpp b/Engine/Plugins/Experimental/Animation/PoseSearch/Source/Editor/Private/PoseSearchDatabaseEditorReflection.cpp
--- a/Engine/Plugins/Experimental/Animation/PoseSearch/Source/Editor/Private/PoseSearchDatabaseEditorReflection.cpp
+++ b/Engine/Plugins/Experimental/Animation/PoseSearch/Source/Editor/Private/PoseSearchDatabaseEditorReflection.cpp
@@ -27,7 +27,7 @@ void UPoseSearchDatabaseSequenceReflection::PostEditChangeProperty(struct FPrope
 		UPoseSearchDatabase* Database = ViewModel->GetPoseSearchDatabase();
 		if (IsValid(Database))
 		{
-			const FInstancedStruct& DatabaseAsset = Database->GetAnimationAssetStruct(WeakAssetTreeNode.Pin()->SourceAssetIdx);
+			FInstancedStruct& DatabaseAsset = Database->GetMutableAnimationAssetStruct(WeakAssetTreeNode.Pin()->SourceAssetIdx);
 			if (FPoseSearchDatabaseSequence* DatabaseSequence = DatabaseAsset.GetMutablePtr<FPoseSearchDatabaseSequence>())
 			{
 				*DatabaseSequence = Sequence;
@@ -49,7 +49,7 @@ void UPoseSearchDatabaseBlendSpaceReflection::PostEditChangeProperty(struct FPro
 		UPoseSearchDatabase* Database = ViewModel->GetPoseSearchDatabase();
 		if (IsValid(Database))
 		{
-			const FInstancedStruct& DatabaseAsset = Database->GetAnimationAssetStruct(WeakAssetTreeNode.Pin()->SourceAssetIdx);
+			FInstancedStruct& DatabaseAsset = Database->GetMutableAnimationAssetStruct(WeakAssetTreeNode.Pin()->SourceAssetIdx);
 			if (FPoseSearchDatabaseBlendSpace* DatabaseBlendSpace = DatabaseAsset.GetMutablePtr<FPoseSearchDatabaseBlendSpace>())
 			{
 				*DatabaseBlendSpace = BlendSpace;
@@ -71,7 +71,7 @@ void UPoseSearchDatabaseAnimCompositeReflection::PostEditChangeProperty(struct F
 		UPoseSearchDatabase* Database = ViewModel->GetPoseSearchDatabase();
 		if (IsValid(Database))
 		{
-			const FInstancedStruct& DatabaseAsset = Database->GetAnimationAssetStruct(WeakAssetTreeNode.Pin()->SourceAssetIdx);
+			FInstancedStruct& DatabaseAsset = Database->GetMutableAnimationAssetStruct(WeakAssetTreeNode.Pin()->SourceAssetIdx);
 			if (FPoseSearchDatabaseAnimComposite* DatabaseAnimComposite = DatabaseAsset.GetMutablePtr<FPoseSearchDatabaseAnimComposite>())
 			{
 				*DatabaseAnimComposite = AnimComposite;
diff --git a/Engine/Plugins/Experimental/Animation/PoseSearch/Source/Runtime/Private/PoseSearch.cpp b/Engine/Plugins/Experimental/Animation/PoseSearch/Source/Runtime/Private/PoseSearch.cpp
--- a/Engine/Plugins/Experimental/Animation/PoseSearch/Source/Runtime/Private/PoseSearch.cpp
+++ b/Engine/Plugins/Experimental/Animation/PoseSearch/Source/Runtime/Private/PoseSearch.cpp
@@ -813,6 +813,17 @@ const FInstancedStruct& UPoseSearchDatabase::GetAnimationAssetStruct(const FPose
 	return GetAnimationAssetStruct(SearchIndexAsset.SourceAssetIdx);
 }
 
+FInstancedStruct& UPoseSearchDatabase::GetMutableAnimationAssetStruct(int32 AnimationAssetIndex)
+{
+	check(AnimationAssets.IsValidIndex(AnimationAssetIndex));
+	return AnimationAssets[AnimationAssetIndex];
+}
+
+FInstancedStruct& UPoseSearchDatabase::GetMutableAnimationAssetStruct(const FPoseSearchIndexAsset& SearchIndexAsset)
+{
+	return GetMutableAnimationAssetStruct(SearchIndexAsset.SourceAssetIdx);
+}
+
 const FPoseSearchDatabaseAnimationAssetBase* UPoseSearchDatabase::GetAnimationAssetBase(int32 AnimationAssetIndex) const
 {
 	if (AnimationAssets.IsValidIndex(AnimationAssetIndex))
diff --git a/Engine/Plugins/Experimental/Animation/PoseSearch/Source/Runtime/Public/PoseSearch/PoseSearch.h b/Engine/Plugins/Experimental/Animation/PoseSearch/Source/Runtime/Public/PoseSearch/PoseSearch.h
--- a/Engine/Plugins/Experimental/Animation/PoseSearch/Source/Runtime/Public/PoseSearch/PoseSearch.h
+++ b/Engine/Plugins/Experimental/Animation/PoseSearch/Source/Runtime/Public/PoseSearch/PoseSearch.h
@@ -1167,6 +1167,8 @@ public:
 
 	const FInstancedStruct& GetAnimationAssetStruct(int32 AnimationAssetIndex) const;
 	const FInstancedStruct& GetAnimationAssetStruct(const FPoseSearchIndexAsset& SearchIndexAsset) const;
+	FInstancedStruct& GetMutableAnimationAssetStruct(int32 AnimationAssetIndex);
+	FInstancedStruct& GetMutableAnimationAssetStruct(const FPoseSearchIndexAsset& SearchIndexAsset);
 	const FPoseSearchDatabaseAnimationAssetBase* GetAnimationAssetBase(int32 AnimationAssetIndex) const;
 	const FPoseSearchDatabaseAnimationAssetBase* GetAnimationAssetBase(const FPoseSearchIndexAsset& SearchIndexAsset) const;
 	FPoseSearchDatabaseAnimationAssetBase* GetMutableAnimationAssetBase(int32 AnimationAssetIndex);
