motion matching - pose search plugin runtime refactoring

diff --git a/Engine/Plugins/Experimental/Animation/PoseSearch/Source/Editor/Private/PoseSearchDatabaseEditorReflection.cpp b/Engine/Plugins/Experimental/Animation/PoseSearch/Source/Editor/Private/PoseSearchDatabaseEditorReflection.cpp
--- a/Engine/Plugins/Experimental/Animation/PoseSearch/Source/Editor/Private/PoseSearchDatabaseEditorReflection.cpp
+++ b/Engine/Plugins/Experimental/Animation/PoseSearch/Source/Editor/Private/PoseSearchDatabaseEditorReflection.cpp
@@ -4,6 +4,7 @@
 #include "PoseSearchDatabaseAssetTree.h"
 #include "PoseSearchDatabaseAssetTreeNode.h"
 #include "PoseSearchDatabaseViewModel.h"
+#include "PoseSearch/PoseSearchSchema.h"
 
 #define LOCTEXT_NAMESPACE "UPoseSearchDatabaseReflection"
 
diff --git a/Engine/Plugins/Experimental/Animation/PoseSearch/Source/Editor/Private/PoseSearchDatabaseViewModel.cpp b/Engine/Plugins/Experimental/Animation/PoseSearch/Source/Editor/Private/PoseSearchDatabaseViewModel.cpp
--- a/Engine/Plugins/Experimental/Animation/PoseSearch/Source/Editor/Private/PoseSearchDatabaseViewModel.cpp
+++ b/Engine/Plugins/Experimental/Animation/PoseSearch/Source/Editor/Private/PoseSearchDatabaseViewModel.cpp
@@ -5,6 +5,7 @@
 #include "PoseSearchDatabaseAssetTreeNode.h"
 #include "PoseSearchEditor.h"
 #include "PoseSearch/PoseSearchDerivedData.h"
+#include "PoseSearch/PoseSearchSchema.h"
 #include "Modules/ModuleManager.h"
 #include "AnimPreviewInstance.h"
 #include "Animation/AnimComposite.h"
diff --git a/Engine/Plugins/Experimental/Animation/PoseSearch/Source/Editor/Private/PoseSearchDebuggerDatabaseView.cpp b/Engine/Plugins/Experimental/Animation/PoseSearch/Source/Editor/Private/PoseSearchDebuggerDatabaseView.cpp
--- a/Engine/Plugins/Experimental/Animation/PoseSearch/Source/Editor/Private/PoseSearchDebuggerDatabaseView.cpp
+++ b/Engine/Plugins/Experimental/Animation/PoseSearch/Source/Editor/Private/PoseSearchDebuggerDatabaseView.cpp
@@ -6,6 +6,8 @@
 #include "PoseSearchDebuggerView.h"
 #include "PoseSearchDebuggerViewModel.h"
 #include "PoseSearch/PoseSearchDerivedData.h"
+#include "PoseSearch/PoseSearchFeatureChannel.h"
+#include "PoseSearch/PoseSearchSchema.h"
 #include "Trace/PoseSearchTraceProvider.h"
 #include "Widgets/Input/SCheckBox.h"
 #include "Widgets/Layout/SScrollBar.h"
diff --git a/Engine/Plugins/Experimental/Animation/PoseSearch/Source/Editor/Private/PoseSearchDebuggerView.cpp b/Engine/Plugins/Experimental/Animation/PoseSearch/Source/Editor/Private/PoseSearchDebuggerView.cpp
--- a/Engine/Plugins/Experimental/Animation/PoseSearch/Source/Editor/Private/PoseSearchDebuggerView.cpp
+++ b/Engine/Plugins/Experimental/Animation/PoseSearch/Source/Editor/Private/PoseSearchDebuggerView.cpp
@@ -8,6 +8,7 @@
 #include "PoseSearchDebuggerDatabaseView.h"
 #include "PoseSearchDebuggerReflection.h"
 #include "PoseSearchDebuggerViewModel.h"
+#include "PoseSearch/PoseSearchSchema.h"
 #include "PropertyEditorModule.h"
 #include "Trace/PoseSearchTraceProvider.h"
 #include "Widgets/Input/SButton.h"
diff --git a/Engine/Plugins/Experimental/Animation/PoseSearch/Source/Editor/Private/PoseSearchDebuggerViewModel.cpp b/Engine/Plugins/Experimental/Animation/PoseSearch/Source/Editor/Private/PoseSearchDebuggerViewModel.cpp
--- a/Engine/Plugins/Experimental/Animation/PoseSearch/Source/Editor/Private/PoseSearchDebuggerViewModel.cpp
+++ b/Engine/Plugins/Experimental/Animation/PoseSearch/Source/Editor/Private/PoseSearchDebuggerViewModel.cpp
@@ -7,6 +7,7 @@
 #include "IGameplayProvider.h"
 #include "IRewindDebugger.h"
 #include "PoseSearch/PoseSearchDerivedData.h"
+#include "PoseSearch/PoseSearchSchema.h"
 #include "Trace/PoseSearchTraceProvider.h"
 
 namespace UE::PoseSearch
diff --git a/Engine/Plugins/Experimental/Animation/PoseSearch/Source/Editor/Private/PoseSearchSchemaFactory.cpp b/Engine/Plugins/Experimental/Animation/PoseSearch/Source/Editor/Private/PoseSearchSchemaFactory.cpp
--- a/Engine/Plugins/Experimental/Animation/PoseSearch/Source/Editor/Private/PoseSearchSchemaFactory.cpp
+++ b/Engine/Plugins/Experimental/Animation/PoseSearch/Source/Editor/Private/PoseSearchSchemaFactory.cpp
@@ -1,7 +1,8 @@
 // Copyright Epic Games, Inc. All Rights Reserved.
 
 #include "PoseSearchSchemaFactory.h"
-#include "PoseSearch/PoseSearch.h"
+#include "PoseSearch/PoseSearchSchema.h"
+#include "PoseSearch/PoseSearchSchema.h"
 
 #define LOCTEXT_NAMESPACE "PoseSearchEditor"
 
diff --git a/Engine/Plugins/Experimental/Animation/PoseSearch/Source/Editor/Private/PoseSearchTypeActions.cpp b/Engine/Plugins/Experimental/Animation/PoseSearch/Source/Editor/Private/PoseSearchTypeActions.cpp
--- a/Engine/Plugins/Experimental/Animation/PoseSearch/Source/Editor/Private/PoseSearchTypeActions.cpp
+++ b/Engine/Plugins/Experimental/Animation/PoseSearch/Source/Editor/Private/PoseSearchTypeActions.cpp
@@ -3,6 +3,7 @@
 #include "PoseSearchTypeActions.h"
 #include "PoseSearchDatabaseEditor.h"
 #include "PoseSearch/PoseSearch.h"
+#include "PoseSearch/PoseSearchSchema.h"
 
 #define LOCTEXT_NAMESPACE "PoseSearchTypeActions"
 
diff --git a/Engine/Plugins/Experimental/Animation/PoseSearch/Source/Runtime/Private/AnimNode_MotionMatching.cpp b/Engine/Plugins/Experimental/Animation/PoseSearch/Source/Runtime/Private/AnimNode_MotionMatching.cpp
--- a/Engine/Plugins/Experimental/Animation/PoseSearch/Source/Runtime/Private/AnimNode_MotionMatching.cpp
+++ b/Engine/Plugins/Experimental/Animation/PoseSearch/Source/Runtime/Private/AnimNode_MotionMatching.cpp
@@ -12,6 +12,7 @@
 #include "DrawDebugHelpers.h"
 #include "PoseSearch/AnimNode_PoseSearchHistoryCollector.h"
 #include "PoseSearch/PoseSearchDerivedData.h"
+#include "PoseSearch/PoseSearchSchema.h"
 #include "Trace/PoseSearchTraceLogger.h"
 
 #include UE_INLINE_GENERATED_CPP_BY_NAME(AnimNode_MotionMatching)
diff --git a/Engine/Plugins/Experimental/Animation/PoseSearch/Source/Runtime/Private/PoseSearch.cpp b/Engine/Plugins/Experimental/Animation/PoseSearch/Source/Runtime/Private/PoseSearch.cpp
--- a/Engine/Plugins/Experimental/Animation/PoseSearch/Source/Runtime/Private/PoseSearch.cpp
+++ b/Engine/Plugins/Experimental/Animation/PoseSearch/Source/Runtime/Private/PoseSearch.cpp
@@ -5,6 +5,8 @@
 #include "PoseSearch/PoseSearchAnimNotifies.h"
 #include "PoseSearch/PoseSearchDerivedData.h"
 #include "PoseSearch/PoseSearchDerivedDataKey.h"
+#include "PoseSearch/PoseSearchFeatureChannel.h"
+#include "PoseSearch/PoseSearchSchema.h"
 #include "PoseSearchEigenHelper.h"
 
 #include "Algo/BinarySearch.h"
@@ -56,17 +58,6 @@ namespace UE::PoseSearch
 //////////////////////////////////////////////////////////////////////////
 // Constants and utilities
 
-static inline float ArraySum(TConstArrayView<float> View, int32 StartIndex, int32 Offset)
-{
-	float Sum = 0.f;
-	const int32 EndIndex = StartIndex + Offset;
-	for (int i = StartIndex; i < EndIndex; ++i)
-	{
-		Sum += View[i];
-	}
-	return Sum;
-}
-
 static inline float CompareFeatureVectors(TConstArrayView<float> A, TConstArrayView<float> B, TConstArrayView<float> WeightsSqrt)
 {
 	check(A.Num() == B.Num() && A.Num() == WeightsSqrt.Num());
@@ -349,21 +340,6 @@ private:
 
 } // namespace UE::PoseSearch
 
-
-//////////////////////////////////////////////////////////////////////////
-// UPoseSearchFeatureChannel
-#if WITH_EDITOR
-void UPoseSearchFeatureChannel::PopulateChannelLayoutSet(UE::PoseSearch::FFeatureChannelLayoutSet& FeatureChannelLayoutSet) const
-{
-	FeatureChannelLayoutSet.Add(GetName(), UE::PoseSearch::FKeyBuilder(this).Finalize(), ChannelDataOffset, ChannelCardinality);
-}
-
-void UPoseSearchFeatureChannel::ComputeCostBreakdowns(UE::PoseSearch::ICostBreakDownData& CostBreakDownData, const UPoseSearchSchema* Schema) const
-{
-	CostBreakDownData.AddEntireBreakDownSection(FText::FromString(GetName()), Schema, ChannelDataOffset, ChannelCardinality);
-}
-#endif // WITH_EDITOR
-
 //////////////////////////////////////////////////////////////////////////
 // UPoseSearchSchema
 void UPoseSearchSchema::Finalize(bool bRemoveEmptyChannels)
@@ -423,68 +399,6 @@ void UPoseSearchSchema::ComputeCostBreakdowns(UE::PoseSearch::ICostBreakDownData
 }
 #endif
 
-bool UPoseSearchSchema::IsValid() const
-{
-	bool bValid = Skeleton != nullptr;
-
-	for (const FBoneReference& BoneRef : BoneReferences)
-	{
-		bValid &= BoneRef.HasValidSetup();
-	}
-
-	for (const TObjectPtr<UPoseSearchFeatureChannel>& Channel: Channels)
-	{
-		bValid &= Channel != nullptr;
-	}
-
-	bValid &= (BoneReferences.Num() == BoneIndices.Num());
-
-	return bValid;
-}
-
-void UPoseSearchSchema::ResolveBoneReferences()
-{
-	// Initialize references to obtain bone indices
-	for (FBoneReference& BoneRef : BoneReferences)
-	{
-		BoneRef.Initialize(Skeleton);
-	}
-
-	// Fill out bone index array
-	BoneIndices.SetNum(BoneReferences.Num());
-	for (int32 BoneRefIdx = 0; BoneRefIdx != BoneReferences.Num(); ++BoneRefIdx)
-	{
-		BoneIndices[BoneRefIdx] = BoneReferences[BoneRefIdx].BoneIndex;
-	}
-
-	// Build separate index array with parent indices guaranteed to be present. Sort for EnsureParentsPresent.
-	BoneIndicesWithParents = BoneIndices;
-	BoneIndicesWithParents.Sort();
-
-	if (Skeleton)
-	{
-		FAnimationRuntime::EnsureParentsPresent(BoneIndicesWithParents, Skeleton->GetReferenceSkeleton());
-	}
-
-	// BoneIndicesWithParents should at least contain the root to support mirroring root motion
-	if (BoneIndicesWithParents.Num() == 0)
-	{
-		BoneIndicesWithParents.Add(0);
-	}
-}
-
-void UPoseSearchSchema::BuildQuery(UE::PoseSearch::FSearchContext& SearchContext, FPoseSearchFeatureVectorBuilder& InOutQuery) const
-{
-	QUICK_SCOPE_CYCLE_COUNTER(STAT_PoseSearch_BuildQuery);
-
-	InOutQuery.Init(this);
-
-	for (const TObjectPtr<UPoseSearchFeatureChannel>& Channel : Channels)
-	{
-		Channel->BuildQuery(SearchContext, InOutQuery);
-	}
-}
-
 //////////////////////////////////////////////////////////////////////////
 // FPoseSearchBaseIndex
 const FPoseSearchIndexAsset& FPoseSearchIndexBase::GetAssetForPose(int32 PoseIdx) const
@@ -2148,26 +2062,6 @@ FTransform FAssetSamplingContext::MirrorTransform(const FTransform& InTransform)
 	return Result;
 }
 
-//////////////////////////////////////////////////////////////////////////
-// class ICostBreakDownData
-
-void ICostBreakDownData::AddEntireBreakDownSection(const FText& Label, const UPoseSearchSchema* Schema, int32 DataOffset, int32 Cardinality)
-{
-	BeginBreakDownSection(Label);
-
-	const int32 Count = Num();
-	for (int32 i = 0; i < Count; ++i)
-	{
-		if (IsCostVectorFromSchema(i, Schema))
-		{
-			const float CostBreakdown = ArraySum(GetCostVector(i, Schema), DataOffset, Cardinality);
-			SetCostBreakDown(CostBreakdown, i, Schema);
-		}
-	}
-
-	EndBreakDownSection(Label);
-}
-
 //////////////////////////////////////////////////////////////////////////
 // Root motion extrapolation
 
diff --git a/Engine/Plugins/Experimental/Animation/PoseSearch/Source/Runtime/Private/PoseSearchDerivedData.cpp b/Engine/Plugins/Experimental/Animation/PoseSearch/Source/Runtime/Private/PoseSearchDerivedData.cpp
--- a/Engine/Plugins/Experimental/Animation/PoseSearch/Source/Runtime/Private/PoseSearchDerivedData.cpp
+++ b/Engine/Plugins/Experimental/Animation/PoseSearch/Source/Runtime/Private/PoseSearchDerivedData.cpp
@@ -11,6 +11,8 @@
 #include "Misc/CoreDelegates.h"
 #include "PoseSearch/PoseSearchAnimNotifies.h"
 #include "PoseSearch/PoseSearchDerivedDataKey.h"
+#include "PoseSearch/PoseSearchFeatureChannel.h"
+#include "PoseSearch/PoseSearchSchema.h"
 #include "PoseSearchEigenHelper.h"
 #include "ProfilingDebugging/CookStats.h"
 #include "Serialization/BulkDataRegistry.h"
diff --git a/Engine/Plugins/Experimental/Animation/PoseSearch/Source/Runtime/Private/PoseSearchFeatureChannel.cpp b/Engine/Plugins/Experimental/Animation/PoseSearch/Source/Runtime/Private/PoseSearchFeatureChannel.cpp
new file mode 100644
--- /dev/null
+++ b/Engine/Plugins/Experimental/Animation/PoseSearch/Source/Runtime/Private/PoseSearchFeatureChannel.cpp
@@ -0,0 +1,71 @@
+// Copyright Epic Games, Inc. All Rights Reserved.
+
+#include "PoseSearch/PoseSearchFeatureChannel.h"
+#include "PoseSearch/PoseSearchDerivedDataKey.h"
+#include "PoseSearch/PoseSearchSchema.h"
+#include "Animation/Skeleton.h"
+
+//////////////////////////////////////////////////////////////////////////
+// class ICostBreakDownData
+namespace UE::PoseSearch
+{
+
+static inline float ArraySum(TConstArrayView<float> View, int32 StartIndex, int32 Offset)
+{
+	float Sum = 0.f;
+	const int32 EndIndex = StartIndex + Offset;
+	for (int i = StartIndex; i < EndIndex; ++i)
+	{
+		Sum += View[i];
+	}
+	return Sum;
+}
+
+void ICostBreakDownData::AddEntireBreakDownSection(const FText& Label, const UPoseSearchSchema* Schema, int32 DataOffset, int32 Cardinality)
+{
+	BeginBreakDownSection(Label);
+
+	const int32 Count = Num();
+	for (int32 i = 0; i < Count; ++i)
+	{
+		if (IsCostVectorFromSchema(i, Schema))
+		{
+			const float CostBreakdown = ArraySum(GetCostVector(i, Schema), DataOffset, Cardinality);
+			SetCostBreakDown(CostBreakdown, i, Schema);
+		}
+	}
+
+	EndBreakDownSection(Label);
+}
+
+} // namespace UE::PoseSearch
+
+//////////////////////////////////////////////////////////////////////////
+// UPoseSearchFeatureChannel
+#if WITH_EDITOR
+void UPoseSearchFeatureChannel::PopulateChannelLayoutSet(UE::PoseSearch::FFeatureChannelLayoutSet& FeatureChannelLayoutSet) const
+{
+	FeatureChannelLayoutSet.Add(GetName(), UE::PoseSearch::FKeyBuilder(this).Finalize(), ChannelDataOffset, ChannelCardinality);
+}
+
+void UPoseSearchFeatureChannel::ComputeCostBreakdowns(UE::PoseSearch::ICostBreakDownData& CostBreakDownData, const UPoseSearchSchema* Schema) const
+{
+	CostBreakDownData.AddEntireBreakDownSection(FText::FromString(GetName()), Schema, ChannelDataOffset, ChannelCardinality);
+}
+#endif // WITH_EDITOR
+
+class USkeleton* UPoseSearchFeatureChannel::GetSkeleton(bool& bInvalidSkeletonIsError, const IPropertyHandle* PropertyHandle)
+{
+	bInvalidSkeletonIsError = false;
+
+	UObject* Outer = GetOuter();
+	while (Outer != nullptr)
+	{
+		if (const UPoseSearchSchema* Schema = Cast<UPoseSearchSchema>(Outer))
+		{
+			return Schema->Skeleton;
+		}
+		Outer = Outer->GetOuter();
+	}
+	return nullptr;
+}
\ No newline at end of file
diff --git a/Engine/Plugins/Experimental/Animation/PoseSearch/Source/Runtime/Private/PoseSearchFeatureChannel_FilterCrashingLegs.cpp b/Engine/Plugins/Experimental/Animation/PoseSearch/Source/Runtime/Private/PoseSearchFeatureChannel_FilterCrashingLegs.cpp
new file mode 100644
--- /dev/null
+++ b/Engine/Plugins/Experimental/Animation/PoseSearch/Source/Runtime/Private/PoseSearchFeatureChannel_FilterCrashingLegs.cpp
@@ -0,0 +1,166 @@
+// Copyright Epic Games, Inc. All Rights Reserved.
+
+#include "PoseSearchFeatureChannel_FilterCrashingLegs.h"
+#include "DrawDebugHelpers.h"
+#include "PoseSearch/PoseSearch.h"
+#include "PoseSearch/PoseSearchSchema.h"
+
+static void ComputeThighsSideAndForward(const FVector& RightThighPos, const FVector& LeftThighPos, FVector& ThighsSide, FVector& ThighsForward)
+{
+	ThighsSide = (RightThighPos - LeftThighPos).GetSafeNormal2D(UE_KINDA_SMALL_NUMBER, FVector::LeftVector);
+	ThighsForward = (ThighsSide ^ FVector::UpVector).GetSafeNormal2D(UE_KINDA_SMALL_NUMBER, FVector::ForwardVector);
+}
+
+static float ComputeCrashingLegsValue(const FVector& RightThighPos, const FVector& LeftThighPos, const FVector& RightFootPos, const FVector& LeftFootPos)
+{
+	FVector ThighsSide;
+	FVector ThighsForward;
+	ComputeThighsSideAndForward(RightThighPos, LeftThighPos, ThighsSide, ThighsForward);
+
+	const FVector FeetDir = (RightFootPos - LeftFootPos).GetSafeNormal2D(UE_KINDA_SMALL_NUMBER, FVector::LeftVector);
+	const float SideDot = FeetDir | ThighsSide;
+	const float ForwardDot = FeetDir | ThighsForward;
+	const float CrashingLegsValue = FMath::Atan2(ForwardDot, SideDot) / UE_PI;
+	return CrashingLegsValue;
+}
+
+void UPoseSearchFeatureChannel_FilterCrashingLegs::InitializeSchema(UPoseSearchSchema* Schema)
+{
+	ChannelDataOffset = Schema->SchemaCardinality;
+	ChannelCardinality = UE::PoseSearch::FFeatureVectorHelper::EncodeFloatCardinality;
+	Schema->SchemaCardinality += ChannelCardinality;
+
+	LeftThighIdx = Schema->AddBoneReference(LeftThigh);
+	RightThighIdx = Schema->AddBoneReference(RightThigh);
+	LeftFootIdx = Schema->AddBoneReference(LeftFoot);
+	RightFootIdx = Schema->AddBoneReference(RightFoot);
+}
+
+void UPoseSearchFeatureChannel_FilterCrashingLegs::FillWeights(TArray<float>& Weights) const
+{
+	check(ChannelCardinality == UE::PoseSearch::FFeatureVectorHelper::EncodeFloatCardinality);
+	Weights[ChannelDataOffset] = Weight;
+}
+
+void UPoseSearchFeatureChannel_FilterCrashingLegs::IndexAsset(UE::PoseSearch::IAssetIndexer& Indexer, UE::PoseSearch::FAssetIndexingOutput& IndexingOutput) const
+{
+	using namespace UE::PoseSearch;
+
+	const FAssetIndexingContext& IndexingContext = Indexer.GetIndexingContext();
+
+	for (int32 SampleIdx = IndexingContext.BeginSampleIdx; SampleIdx != IndexingContext.EndSampleIdx; ++SampleIdx)
+	{
+		const float SubsampleTime = SampleIdx * IndexingContext.Schema->GetSamplingInterval();
+
+		bool bUnused;
+		const FTransform RightThighTransform = Indexer.GetTransformAndCacheResults(SubsampleTime, SubsampleTime, RightThighIdx, bUnused);
+		const FTransform LeftThighTransform = Indexer.GetTransformAndCacheResults(SubsampleTime, SubsampleTime, LeftThighIdx, bUnused);
+		const FTransform RightFootTransform = Indexer.GetTransformAndCacheResults(SubsampleTime, SubsampleTime, RightFootIdx, bUnused);
+		const FTransform LeftFootTransform = Indexer.GetTransformAndCacheResults(SubsampleTime, SubsampleTime, LeftFootIdx, bUnused);
+
+		const float CrashingLegsValue = ComputeCrashingLegsValue(RightThighTransform.GetTranslation(), LeftThighTransform.GetTranslation(), RightFootTransform.GetTranslation(), LeftFootTransform.GetTranslation());
+		
+		int32 DataOffset = ChannelDataOffset;
+		const int32 VectorIdx = SampleIdx - IndexingContext.BeginSampleIdx;
+		FFeatureVectorHelper::EncodeFloat(IndexingOutput.GetPoseVector(VectorIdx), DataOffset, CrashingLegsValue);
+	}
+}
+
+void UPoseSearchFeatureChannel_FilterCrashingLegs::BuildQuery(UE::PoseSearch::FSearchContext& SearchContext, FPoseSearchFeatureVectorBuilder& InOutQuery) const
+{
+	using namespace UE::PoseSearch;
+
+	const bool bIsCurrentResultValid = SearchContext.CurrentResult.IsValid();
+	const bool bSkip = InputQueryPose != EInputQueryPose::UseCharacterPose && bIsCurrentResultValid && SearchContext.CurrentResult.Database->Schema == InOutQuery.GetSchema();
+	
+	int32 DataOffset = ChannelDataOffset;
+	if (bSkip || !SearchContext.History)
+	{
+		if (bIsCurrentResultValid)
+		{
+			const float LerpValue = InputQueryPose == EInputQueryPose::UseInterpolatedContinuingPose ? SearchContext.CurrentResult.LerpValue : 0.f;
+			FFeatureVectorHelper::EncodeFloat(InOutQuery.EditValues(), DataOffset, SearchContext.GetCurrentResultPrevPoseVector(), SearchContext.GetCurrentResultPoseVector(), SearchContext.GetCurrentResultNextPoseVector(), LerpValue);
+		}
+		// else leave the InOutQuery set to zero since the SearchContext.History is invalid and it'll fail if we continue
+	}
+	else
+	{
+		const float SampleTime = 0.f;
+		const FTransform LeftThighTransform = SearchContext.TryGetTransformAndCacheResults(SampleTime, InOutQuery.GetSchema(), LeftThighIdx);
+		const FTransform RightThighTransform = SearchContext.TryGetTransformAndCacheResults(SampleTime, InOutQuery.GetSchema(), RightThighIdx);
+		const FTransform LeftFootTransform = SearchContext.TryGetTransformAndCacheResults(SampleTime, InOutQuery.GetSchema(), LeftFootIdx);
+		const FTransform RightFootTransform = SearchContext.TryGetTransformAndCacheResults(SampleTime, InOutQuery.GetSchema(), RightFootIdx);
+
+		const float CrashingLegsValue = ComputeCrashingLegsValue(RightThighTransform.GetTranslation(), LeftThighTransform.GetTranslation(), RightFootTransform.GetTranslation(), LeftFootTransform.GetTranslation());
+		FFeatureVectorHelper::EncodeFloat(InOutQuery.EditValues(), DataOffset, CrashingLegsValue);
+	}
+}
+
+void UPoseSearchFeatureChannel_FilterCrashingLegs::DebugDraw(const UE::PoseSearch::FDebugDrawParams& DrawParams, TConstArrayView<float> PoseVector) const
+{
+#if ENABLE_DRAW_DEBUG
+	using namespace UE::PoseSearch;
+
+	int32 DataOffset = ChannelDataOffset;
+	const float CrashingLegsValue = FFeatureVectorHelper::DecodeFloat(PoseVector, DataOffset);
+
+	const float LifeTime = DrawParams.DefaultLifeTime;
+	const uint8 DepthPriority = ESceneDepthPriorityGroup::SDPG_Foreground + 2;
+	const bool bPersistent = EnumHasAnyFlags(DrawParams.Flags, EDebugDrawFlags::Persistent);
+
+	const FTransform LeftThighTransform = DrawParams.Mesh != nullptr ? DrawParams.Mesh->GetSocketTransform(LeftThigh.BoneName) : DrawParams.RootTransform;
+	const FTransform RightThighTransform = DrawParams.Mesh != nullptr ? DrawParams.Mesh->GetSocketTransform(RightThigh.BoneName) : DrawParams.RootTransform;
+	const FTransform LeftFootTransform = DrawParams.Mesh != nullptr ? DrawParams.Mesh->GetSocketTransform(LeftFoot.BoneName) : DrawParams.RootTransform;
+	const FTransform RightFootTransform = DrawParams.Mesh != nullptr ? DrawParams.Mesh->GetSocketTransform(RightFoot.BoneName) : DrawParams.RootTransform;
+	
+	const float FeetDistance = (RightFootTransform.GetTranslation() - LeftFootTransform.GetTranslation()).Length();
+
+	FVector ThighsSide;
+	FVector ThighsForward;
+	ComputeThighsSideAndForward(RightThighTransform.GetTranslation(), LeftThighTransform.GetTranslation(), ThighsSide, ThighsForward);
+	const FVector CrossingLegsVector = (FMath::Cos(CrashingLegsValue * UE_PI) * ThighsSide + FMath::Sin(CrashingLegsValue * UE_PI) * ThighsForward) * FeetDistance;
+
+	auto LerpColor = [](FColor A, FColor B, float T) -> FColor
+	{
+		return FColor(
+			FMath::RoundToInt(float(A.R) * (1.f - T) + float(B.R) * T),
+			FMath::RoundToInt(float(A.G) * (1.f - T) + float(B.G) * T),
+			FMath::RoundToInt(float(A.B) * (1.f - T) + float(B.B) * T),
+			FMath::RoundToInt(float(A.A) * (1.f - T) + float(B.A) * T));
+	};
+
+	// when CrossingLegsValue is greater than .5 or lesser than -.5 we draw in Red or Orange
+	const float LerpValue = FMath::Min(CrashingLegsValue * 2.f, 1.f);
+	const FColor Color = LerpValue >= 0.f ?
+		LerpColor(FColor::Green, FColor::Red, LerpValue) :
+		LerpColor(FColor::Green, FColor::Orange, -LerpValue);
+
+	DrawDebugLine(DrawParams.World, LeftFootTransform.GetTranslation(), LeftFootTransform.GetTranslation() + CrossingLegsVector, Color, bPersistent, LifeTime, DepthPriority);
+	DrawDebugLine(DrawParams.World, RightFootTransform.GetTranslation(), RightFootTransform.GetTranslation() - CrossingLegsVector, Color, bPersistent, LifeTime, DepthPriority);
+#endif // ENABLE_DRAW_DEBUG
+}
+
+// IPoseFilter interface
+bool UPoseSearchFeatureChannel_FilterCrashingLegs::IsPoseFilterActive() const
+{
+	return true;
+}
+
+bool UPoseSearchFeatureChannel_FilterCrashingLegs::IsPoseValid(TConstArrayView<float> PoseValues, TConstArrayView<float> QueryValues, int32 PoseIdx, const FPoseSearchPoseMetadata& Metadata) const
+{
+	using namespace UE::PoseSearch;
+
+	int32 DataOffset = ChannelDataOffset;
+	const float PoseValue = FFeatureVectorHelper::DecodeFloat(PoseValues, DataOffset);
+
+	DataOffset = ChannelDataOffset;
+	const float QueryValue = FFeatureVectorHelper::DecodeFloat(QueryValues, DataOffset);
+
+	bool bIsPoseValid = true;
+	if (FMath::Abs(QueryValue - PoseValue) > AllowedTolerance)
+	{
+		bIsPoseValid = false;
+	}
+
+	return bIsPoseValid;
+}
diff --git a/Engine/Plugins/Experimental/Animation/PoseSearch/Source/Runtime/Private/PoseSearchFeatureChannel_FilterCrashingLegs.h b/Engine/Plugins/Experimental/Animation/PoseSearch/Source/Runtime/Private/PoseSearchFeatureChannel_FilterCrashingLegs.h
new file mode 100644
--- /dev/null
+++ b/Engine/Plugins/Experimental/Animation/PoseSearch/Source/Runtime/Private/PoseSearchFeatureChannel_FilterCrashingLegs.h
@@ -0,0 +1,65 @@
+// Copyright Epic Games, Inc. All Rights Reserved.
+
+#pragma once
+
+#include "PoseSearch/PoseSearchFeatureChannel.h"
+#include "BoneContainer.h"
+#include "PoseSearchFeatureChannel_FilterCrashingLegs.generated.h"
+
+// the idea is to calculate the angle between the direction from LeftThigh position to RightThigh position and the direction from LeftFoot position to RightFoot position, and divide it by PI to have values in range [-1,1]
+// the number (called 'CrashingLegsValue' calculated in ComputeCrashingLegsValue) is gonna be
+// 0 if the feet are aligned with the thighs (for example in an idle standing position)
+// 0.5 if the right foot is exactly in front of the left foot (for example when a character is running  following a line)
+// -0.5 if the left foot is exactly in front of the right foot
+// close to 1 or -1 if the feet (and so the legs) are completely crossed
+// at runtime we'll match the CrashingLegsValue and also filter by discarding pose candidates that don't respect the 'AllowedTolerance' between query and database values (happening in IsPoseValid)
+UCLASS(BlueprintType, EditInlineNew, meta = (DisplayName = "CrashingLegs Channel"), CollapseCategories)
+class POSESEARCH_API UPoseSearchFeatureChannel_FilterCrashingLegs : public UPoseSearchFeatureChannel
+{
+	GENERATED_BODY()
+
+public:
+	UPROPERTY(EditAnywhere, Category = "Settings")
+	FBoneReference LeftThigh;
+	
+	UPROPERTY(EditAnywhere, Category = "Settings")
+	FBoneReference RightThigh;
+
+	UPROPERTY(EditAnywhere, Category = "Settings")
+	FBoneReference LeftFoot;
+	
+	UPROPERTY(EditAnywhere, Category = "Settings")
+	FBoneReference RightFoot;
+
+	UPROPERTY(EditAnywhere, Category = "Settings")
+	float Weight = 0.2f;
+
+	UPROPERTY()
+	int8 LeftThighIdx;
+
+	UPROPERTY()
+	int8 RightThighIdx;
+
+	UPROPERTY()
+	int8 LeftFootIdx;
+
+	UPROPERTY()
+	int8 RightFootIdx;
+
+	UPROPERTY(EditAnywhere, Category = "Settings")
+	EInputQueryPose InputQueryPose = EInputQueryPose::UseContinuingPose;
+
+	UPROPERTY(EditAnywhere, Category = "Settings")
+	float AllowedTolerance = 0.3f;
+
+	// UPoseSearchFeatureChannel interface
+	virtual void InitializeSchema(UPoseSearchSchema* Schema) override;
+	virtual void FillWeights(TArray<float>& Weights) const override;
+	virtual void IndexAsset(UE::PoseSearch::IAssetIndexer& Indexer, UE::PoseSearch::FAssetIndexingOutput& IndexingOutput) const override;
+	virtual void BuildQuery(UE::PoseSearch::FSearchContext& SearchContext, FPoseSearchFeatureVectorBuilder& InOutQuery) const override;
+	virtual void DebugDraw(const UE::PoseSearch::FDebugDrawParams& DrawParams, TConstArrayView<float> PoseVector) const override;
+
+	// IPoseFilter interface
+	virtual bool IsPoseFilterActive() const override;
+	virtual bool IsPoseValid(TConstArrayView<float> PoseValues, TConstArrayView<float> QueryValues, int32 PoseIdx, const FPoseSearchPoseMetadata& Metadata) const override;
+};
diff --git a/Engine/Plugins/Experimental/Animation/PoseSearch/Source/Runtime/Private/PoseSearchFeatureChannel_Heading.cpp b/Engine/Plugins/Experimental/Animation/PoseSearch/Source/Runtime/Private/PoseSearchFeatureChannel_Heading.cpp
new file mode 100644
--- /dev/null
+++ b/Engine/Plugins/Experimental/Animation/PoseSearch/Source/Runtime/Private/PoseSearchFeatureChannel_Heading.cpp
@@ -0,0 +1,128 @@
+// Copyright Epic Games, Inc. All Rights Reserved.
+
+#include "PoseSearchFeatureChannel_Heading.h"
+#include "DrawDebugHelpers.h"
+#include "PoseSearch/PoseSearch.h"
+#include "PoseSearch/PoseSearchSchema.h"
+
+void UPoseSearchFeatureChannel_Heading::InitializeSchema(UPoseSearchSchema* Schema)
+{
+	ChannelDataOffset = Schema->SchemaCardinality;
+	ChannelCardinality = UE::PoseSearch::FFeatureVectorHelper::EncodeVectorCardinality;
+	Schema->SchemaCardinality += ChannelCardinality;
+	SchemaBoneIdx = Schema->AddBoneReference(Bone);
+}
+
+void UPoseSearchFeatureChannel_Heading::FillWeights(TArray<float>& Weights) const
+{
+	using namespace UE::PoseSearch;
+
+	for (int32 i = 0; i != FFeatureVectorHelper::EncodeVectorCardinality; ++i)
+	{
+		Weights[ChannelDataOffset + i] = Weight;
+	}
+}
+
+FVector UPoseSearchFeatureChannel_Heading::GetAxis(const FQuat& Rotation) const
+{
+	switch (HeadingAxis)
+	{
+	case EHeadingAxis::X:
+		return Rotation.GetAxisX();
+	case EHeadingAxis::Y:
+		return Rotation.GetAxisY();
+	case EHeadingAxis::Z:
+		return Rotation.GetAxisZ();
+	}
+
+	checkNoEntry();
+	return FVector(1.f, 0.f, 0.f);
+}
+
+void UPoseSearchFeatureChannel_Heading::IndexAsset(UE::PoseSearch::IAssetIndexer& Indexer, UE::PoseSearch::FAssetIndexingOutput& IndexingOutput) const
+{
+	using namespace UE::PoseSearch;
+
+	const FAssetIndexingContext& IndexingContext = Indexer.GetIndexingContext();
+
+	for (int32 SampleIdx = IndexingContext.BeginSampleIdx; SampleIdx != IndexingContext.EndSampleIdx; ++SampleIdx)
+	{
+		const int32 VectorIdx = SampleIdx - IndexingContext.BeginSampleIdx;
+
+		const float OriginSampleTime = FMath::Min(SampleIdx * IndexingContext.Schema->GetSamplingInterval(), IndexingContext.AssetSampler->GetPlayLength());
+		const float SubsampleTime = OriginSampleTime + SampleTimeOffset;
+
+		bool ClampedPresent;
+		const FTransform BoneTransformsPresent = Indexer.GetTransformAndCacheResults(SubsampleTime, bUseSampleTimeOffsetRootBone ? SubsampleTime : OriginSampleTime, SchemaBoneIdx, ClampedPresent);
+		int32 DataOffset = ChannelDataOffset;
+		FFeatureVectorHelper::EncodeVector(IndexingOutput.GetPoseVector(VectorIdx), DataOffset, GetAxis(BoneTransformsPresent.GetRotation()));
+		check(DataOffset == ChannelDataOffset + ChannelCardinality);
+	}
+}
+
+void UPoseSearchFeatureChannel_Heading::BuildQuery(UE::PoseSearch::FSearchContext& SearchContext, FPoseSearchFeatureVectorBuilder& InOutQuery) const
+{
+	using namespace UE::PoseSearch;
+
+	const bool bIsCurrentResultValid = SearchContext.CurrentResult.IsValid();
+	const bool bSkip = InputQueryPose != EInputQueryPose::UseCharacterPose && bIsCurrentResultValid && SearchContext.CurrentResult.Database->Schema == InOutQuery.GetSchema();
+	if (bSkip || !SearchContext.History)
+	{
+		if (bIsCurrentResultValid)
+		{
+			const float LerpValue = InputQueryPose == EInputQueryPose::UseInterpolatedContinuingPose ? SearchContext.CurrentResult.LerpValue : 0.f;
+			int32 DataOffset = ChannelDataOffset;
+			FFeatureVectorHelper::EncodeVector(InOutQuery.EditValues(), DataOffset, SearchContext.GetCurrentResultPrevPoseVector(), SearchContext.GetCurrentResultPoseVector(), SearchContext.GetCurrentResultNextPoseVector(), LerpValue, true);
+			check(DataOffset == ChannelDataOffset + ChannelCardinality);
+		}
+		// else leave the InOutQuery set to zero since the SearchContext.History is invalid and it'll fail if we continue
+	}
+	else
+	{
+		FTransform Transform = SearchContext.TryGetTransformAndCacheResults(SampleTimeOffset, InOutQuery.GetSchema(), SchemaBoneIdx);
+
+		if (!bUseSampleTimeOffsetRootBone)
+		{
+			const FTransform RootTransform = SearchContext.TryGetTransformAndCacheResults(0.f, InOutQuery.GetSchema(), FSearchContext::SchemaRootBoneIdx);
+			const FTransform RootTransformPrev = SearchContext.TryGetTransformAndCacheResults(SampleTimeOffset, InOutQuery.GetSchema(), FSearchContext::SchemaRootBoneIdx);
+			Transform = Transform * (RootTransformPrev * RootTransform.Inverse());
+		}
+
+		int32 DataOffset = ChannelDataOffset;
+		FFeatureVectorHelper::EncodeVector(InOutQuery.EditValues(), DataOffset, GetAxis(Transform.GetRotation()));
+		check(DataOffset == ChannelDataOffset + ChannelCardinality);
+	}
+}
+
+void UPoseSearchFeatureChannel_Heading::DebugDraw(const UE::PoseSearch::FDebugDrawParams& DrawParams, TConstArrayView<float> PoseVector) const
+{
+#if ENABLE_DRAW_DEBUG
+	using namespace UE::PoseSearch;
+
+	const UPoseSearchSchema* Schema = DrawParams.GetSchema();
+	check(Schema && Schema->IsValid());
+
+	const float LifeTime = DrawParams.DefaultLifeTime;
+	const uint8 DepthPriority = ESceneDepthPriorityGroup::SDPG_Foreground + 2;
+	const bool bPersistent = EnumHasAnyFlags(DrawParams.Flags, EDebugDrawFlags::Persistent);
+
+	// @todo: search in the schema for a UPoseSearchFeatureChannel_Position with the same bone as SchemaBoneIdx and decode it's data as position
+	const FVector BonePos = DrawParams.Mesh != nullptr ? DrawParams.Mesh->GetSocketTransform(Bone.BoneName).GetLocation() : DrawParams.RootTransform.GetTranslation();
+
+	int32 DataOffset = ChannelDataOffset;
+	const FVector BoneHeading = DrawParams.RootTransform.TransformPosition(FFeatureVectorHelper::DecodeVector(PoseVector, DataOffset));
+	check(DataOffset == ChannelDataOffset + ChannelCardinality);
+
+	const FColor Color = DrawParams.GetColor(ColorPresetIndex);
+
+	if (EnumHasAnyFlags(DrawParams.Flags, EDebugDrawFlags::DrawSearchIndex))
+	{
+		DrawDebugLine(DrawParams.World, BonePos, BonePos + BoneHeading, Color, bPersistent, LifeTime, DepthPriority);
+	}
+	else
+	{
+		const float AdjustedThickness = EnumHasAnyFlags(DrawParams.Flags, EDebugDrawFlags::DrawFast) ? 0.0f : 1.f;
+		DrawDebugLine(DrawParams.World, BonePos, BonePos + BoneHeading, Color, bPersistent, LifeTime, DepthPriority, AdjustedThickness);
+	}
+#endif // ENABLE_DRAW_DEBUG
+}
diff --git a/Engine/Plugins/Experimental/Animation/PoseSearch/Source/Runtime/Private/PoseSearchFeatureChannel_Heading.h b/Engine/Plugins/Experimental/Animation/PoseSearch/Source/Runtime/Private/PoseSearchFeatureChannel_Heading.h
new file mode 100644
--- /dev/null
+++ b/Engine/Plugins/Experimental/Animation/PoseSearch/Source/Runtime/Private/PoseSearchFeatureChannel_Heading.h
@@ -0,0 +1,60 @@
+// Copyright Epic Games, Inc. All Rights Reserved.
+
+#pragma once
+
+#include "PoseSearch/PoseSearchFeatureChannel.h"
+#include "BoneContainer.h"
+#include "PoseSearchFeatureChannel_Heading.generated.h"
+
+UENUM(BlueprintType)
+enum class EHeadingAxis : uint8
+{
+	X,
+	Y,
+	Z,
+
+	Num UMETA(Hidden),
+	Invalid = Num UMETA(Hidden)
+};
+
+UCLASS(BlueprintType, EditInlineNew, meta = (DisplayName = "Heading Channel"), CollapseCategories)
+class POSESEARCH_API UPoseSearchFeatureChannel_Heading : public UPoseSearchFeatureChannel
+{
+	GENERATED_BODY()
+
+public:
+	UPROPERTY(EditAnywhere, Category = "Settings")
+	FBoneReference Bone;
+
+	UPROPERTY(EditAnywhere, Category = "Settings")
+	float Weight = 1.f;
+
+	UPROPERTY(EditAnywhere, Category = "Settings")
+	float SampleTimeOffset = 0.f;
+
+	UPROPERTY(EditAnywhere, Category = "Settings")
+	EHeadingAxis HeadingAxis = EHeadingAxis::X;	
+
+	UPROPERTY()
+	int8 SchemaBoneIdx = 0;
+
+	UPROPERTY(EditAnywhere, Category = "Settings", meta = (ExcludeFromHash))
+	int32 ColorPresetIndex = 0;
+
+	UPROPERTY(EditAnywhere, Category = "Settings")
+	EInputQueryPose InputQueryPose = EInputQueryPose::UseContinuingPose;
+
+	// if UseSampleTimeOffsetRootBone is true, this UPoseSearchFeatureChannel_Position will calculate the position of Bone from the pose SampleTimeOffset seconds away from the current time pose root bone
+	// if false the calculated position of Bone will be in component space from the pose SampleTimeOffset seconds away
+	UPROPERTY(EditAnywhere, Category = "Settings")
+	bool bUseSampleTimeOffsetRootBone = true;
+
+	// UPoseSearchFeatureChannel interface
+	virtual void InitializeSchema(UPoseSearchSchema* Schema) override;
+	virtual void FillWeights(TArray<float>& Weights) const override;
+	virtual void IndexAsset(UE::PoseSearch::IAssetIndexer& Indexer, UE::PoseSearch::FAssetIndexingOutput& IndexingOutput) const override;
+	virtual void BuildQuery(UE::PoseSearch::FSearchContext& SearchContext, FPoseSearchFeatureVectorBuilder& InOutQuery) const override;
+	virtual void DebugDraw(const UE::PoseSearch::FDebugDrawParams& DrawParams, TConstArrayView<float> PoseVector) const override;
+
+	FVector GetAxis(const FQuat& Rotation) const;
+};
\ No newline at end of file
diff --git a/Engine/Plugins/Experimental/Animation/PoseSearch/Source/Runtime/Private/PoseSearchFeatureChannel_Pose.cpp b/Engine/Plugins/Experimental/Animation/PoseSearch/Source/Runtime/Private/PoseSearchFeatureChannel_Pose.cpp
new file mode 100644
--- /dev/null
+++ b/Engine/Plugins/Experimental/Animation/PoseSearch/Source/Runtime/Private/PoseSearchFeatureChannel_Pose.cpp
@@ -0,0 +1,845 @@
+// Copyright Epic Games, Inc. All Rights Reserved.
+
+#include "PoseSearchFeatureChannel_Pose.h"
+#include "DrawDebugHelpers.h"
+#include "PoseSearch/PoseSearch.h"
+#include "PoseSearch/PoseSearchDerivedDataKey.h"
+#include "PoseSearch/PoseSearchSchema.h"
+
+#define LOCTEXT_NAMESPACE "PoseSearchFeatureChannels"
+
+namespace UE::PoseSearch
+{
+
+constexpr bool UseCharacterSpaceVelocities = true;
+
+struct LocalMinMax
+{
+	enum
+	{
+		Min,
+		Max
+	} Type = Min;
+	int32 Index = 0;
+	float SignalValue = 0.f;
+};
+
+template <typename T>
+T GetValueAtIndex(int32 Sample, const TArray<T>& Values)
+{
+	const int32 Num = Values.Num();
+	check(Num > 1);
+
+	if (Sample < 0)
+	{
+		return (Values[1] - Values[0]) * Sample + Values[0];
+	}
+
+	if (Sample < Num)
+	{
+		return Values[Sample];
+	}
+
+	return (Values[Num - 1] - Values[Num - 2]) * (Sample - (Num - 1)) + Values[Num - 1];
+}
+
+static void CollectBonePositions(TArray<FVector>& BonePositions, IAssetIndexer& Indexer, int8 SchemaBoneIdx)
+{
+	const FAssetIndexingContext& IndexingContext = Indexer.GetIndexingContext();
+	const float FiniteDelta = IndexingContext.Schema->GetSamplingInterval();
+	const float SampleTimeStart = FMath::Min(IndexingContext.BeginSampleIdx * FiniteDelta, IndexingContext.AssetSampler->GetPlayLength());
+	const int32 NumSamples = IndexingContext.EndSampleIdx - IndexingContext.BeginSampleIdx;
+
+	// collecting all the bone transforms
+	BonePositions.Reset();
+	BonePositions.AddDefaulted(NumSamples);
+	for (int32 SampleIdx = 0; SampleIdx != NumSamples; ++SampleIdx)
+	{
+		const float SampleTime = SampleTimeStart + SampleIdx * FiniteDelta;
+		bool bUnused;
+		BonePositions[SampleIdx] = Indexer.GetTransformAndCacheResults(SampleTime, SampleTimeStart, SchemaBoneIdx, bUnused).GetTranslation();
+	}
+}
+
+static void CalculateSignal(const TArray<FVector>& BonePositions, TArray<float>& Signal, int32 offset = 1)
+{
+	Signal.Reset();
+	Signal.AddDefaulted(BonePositions.Num());
+
+	for (int32 SampleIdx = 0; SampleIdx != BonePositions.Num(); ++SampleIdx)
+	{
+		Signal[SampleIdx] = (GetValueAtIndex(SampleIdx + offset, BonePositions) - GetValueAtIndex(SampleIdx - offset, BonePositions)).Length();
+	}
+}
+
+static void SmoothSignal(const TArray<float>& Signal, TArray<float>& SmoothedSignal, int32 offset = 1)
+{
+	SmoothedSignal.Reset();
+	SmoothedSignal.AddDefaulted(Signal.Num());
+
+	for (int32 SampleIdx = -offset; SampleIdx != offset; ++SampleIdx)
+	{
+		SmoothedSignal[0] += GetValueAtIndex(SampleIdx, Signal);
+	}
+
+	for (int32 SampleIdx = 1; SampleIdx != Signal.Num(); ++SampleIdx)
+	{
+		SmoothedSignal[SampleIdx] = SmoothedSignal[SampleIdx - 1] - GetValueAtIndex(SampleIdx - offset - 1, Signal) + GetValueAtIndex(SampleIdx + offset, Signal);
+	}
+
+	for (int32 SampleIdx = 0; SampleIdx != Signal.Num(); ++SampleIdx)
+	{
+		SmoothedSignal[SampleIdx] /= 2 * offset + 1;
+	}
+}
+
+static void FindLocalMinMax(const TArray<float>& Signal, TArray<LocalMinMax>& MinMax)
+{
+	enum SignalState
+	{
+		Flat,
+		Ascending,
+		Descending
+	};
+
+	MinMax.Reset();
+	if (Signal.Num() > 1)
+	{
+		SignalState State = SignalState::Flat;
+		for (int32 SignalIndex = 1; SignalIndex < Signal.Num(); ++SignalIndex)
+		{
+			const int32 PrevSignalIndex = SignalIndex - 1;
+			const float PrevSignalValue = Signal[PrevSignalIndex];
+			const float SignalValue = Signal[SignalIndex];
+
+			if (State == SignalState::Flat)
+			{
+				if (SignalValue > PrevSignalValue)
+				{
+					State = SignalState::Ascending;
+				}
+				else if(SignalValue < PrevSignalValue)
+				{
+					State = SignalState::Descending;
+				}
+			}
+			else if (State == SignalState::Ascending)
+			{
+				if (SignalValue < PrevSignalValue)
+				{
+					State = SignalState::Descending;
+					
+					LocalMinMax LocalMinMax;
+					LocalMinMax.Type = LocalMinMax::Max;
+					LocalMinMax.Index = PrevSignalIndex;
+					LocalMinMax.SignalValue = Signal[LocalMinMax.Index];
+
+					check(MinMax.IsEmpty() || MinMax.Last().Type != LocalMinMax.Type);
+					MinMax.Add(LocalMinMax);
+				}
+			}
+			else // if (State == SignalState::Descending)
+			{
+				if (SignalValue > PrevSignalValue)
+				{
+					State = SignalState::Ascending;
+
+					LocalMinMax LocalMinMax;
+					LocalMinMax.Type = LocalMinMax::Min;
+					LocalMinMax.Index = PrevSignalIndex;
+					LocalMinMax.SignalValue = Signal[LocalMinMax.Index];
+
+					check(MinMax.IsEmpty() || MinMax.Last().Type != LocalMinMax.Type);
+					MinMax.Add(LocalMinMax);
+				}
+			}
+		}
+	}
+}
+
+static void ExtrapolateLocalMinMaxBoundaries(TArray<LocalMinMax>& MinMax, const TArray<float>& Signal)
+{
+	const int32 Num = MinMax.Num();
+
+	check(Signal.Num() > 0);
+
+	LocalMinMax InitialMinMax;
+	LocalMinMax FinalMinMax;
+
+	if (Num == 0)
+	{
+		const bool IsInitialMax = Signal[0] > Signal[Signal.Num() - 1];
+
+		InitialMinMax.Index = 0;
+		InitialMinMax.SignalValue = Signal[0];
+		InitialMinMax.Type = IsInitialMax ? LocalMinMax::Max : LocalMinMax::Min;
+
+		FinalMinMax.Index = Signal.Num() - 1;
+		FinalMinMax.SignalValue = Signal[Signal.Num() - 1];
+		FinalMinMax.Type = IsInitialMax ? LocalMinMax::Min : LocalMinMax::Max;
+
+		MinMax.Add(InitialMinMax);
+		MinMax.Add(FinalMinMax);
+	}
+	else
+	{
+		int32 InitialDelta = 0;
+		int32 FinalDelta = 0;
+		if (Num > 2)
+		{
+			InitialDelta = MinMax[2].Index - MinMax[1].Index;
+			FinalDelta = MinMax[Num - 2].Index - MinMax[Num - 3].Index;
+		}
+		else if (Num > 1)
+		{
+			InitialDelta = MinMax[1].Index - MinMax[0].Index;
+			FinalDelta = MinMax[Num - 1].Index - MinMax[Num - 2].Index;
+		}
+		else
+		{
+			InitialDelta = MinMax[0].Index;
+			FinalDelta = (Signal.Num() - 1) - MinMax[0].Index;
+		}
+
+		InitialMinMax.SignalValue = Num > 1 ? MinMax[1].SignalValue : Signal[0];
+		InitialMinMax.Type = MinMax[0].Type == LocalMinMax::Min ? LocalMinMax::Max : LocalMinMax::Min;
+		InitialMinMax.Index = FMath::Min(MinMax[0].Index - InitialDelta, 0);
+
+		FinalMinMax.SignalValue = Num > 1 ? MinMax[Num - 2].SignalValue : Signal[Signal.Num() - 1];
+		FinalMinMax.Type = MinMax[Num - 1].Type == LocalMinMax::Min ? LocalMinMax::Max : LocalMinMax::Min;
+		FinalMinMax.Index = FMath::Max(MinMax[Num - 1].Index + FinalDelta, Signal.Num() - 1);
+
+		// there's no point in adding an InitialMinMax if the first MinMax is at the first frame of the signal
+		if (MinMax[0].Index > 0)
+		{
+			MinMax.Insert(InitialMinMax, 0);
+		}
+
+		// there's no point in adding a FinalMinMax if the last MinMax is at the last frame of the signal
+		if (MinMax[Num - 1].Index < Signal.Num() - 1)
+		{
+			MinMax.Add(FinalMinMax);
+		}
+	}
+}
+
+static void ValidateLocalMinMax(const TArray<LocalMinMax>& MinMax)
+{
+	for (int32 i = 1; i < MinMax.Num(); ++i)
+	{
+		check(MinMax[i].Type != MinMax[i - 1].Type);
+		check(MinMax[i].Index > MinMax[i - 1].Index);
+		if (MinMax[i].Type == LocalMinMax::Min)
+		{
+			check(MinMax[i].SignalValue <= MinMax[i - 1].SignalValue);
+		}
+		else
+		{
+			check(MinMax[i].SignalValue >= MinMax[i - 1].SignalValue);
+		}
+	}
+}
+
+static void CalculatePhaseAndCertainty(int32 Index, const TArray<LocalMinMax>& MinMax, int32 SignalSize, float& Phase, float& Certainty)
+{
+	// @todo: expose them via UI
+	static float CertaintyMin = 1.f;
+	static float CertaintyMult = 0.1f;
+
+	const int32 LastIndex = MinMax.Num() - 1;
+	for (int32 i = 1; i < MinMax.Num(); ++i)
+	{
+		const int32 MinMaxIndex = MinMax[i].Index;
+		if (Index < MinMaxIndex)
+		{
+			const int32 PrevMinMaxIndex = MinMax[i - 1].Index;
+			check(MinMaxIndex > PrevMinMaxIndex);
+			const float Ratio = static_cast<float>((Index - PrevMinMaxIndex)) / static_cast<float>((MinMaxIndex - PrevMinMaxIndex));
+			const float PhaseOffset = MinMax[i - 1].Type == LocalMinMax::Min ? 0.f : 0.5f;
+			Phase = PhaseOffset + Ratio * 0.5f;
+
+			const float DeltaSignalValue = FMath::Abs(MinMax[i - 1].SignalValue - MinMax[i].SignalValue);
+			const float NextDeltaSignalValue = i < LastIndex ? FMath::Abs(MinMax[i].SignalValue - MinMax[i + 1].SignalValue) : DeltaSignalValue;
+			Certainty = CertaintyMin + (DeltaSignalValue * (1.f - Ratio) + NextDeltaSignalValue * Ratio) * CertaintyMult;
+			return;
+		}
+	}
+
+	Phase = MinMax[LastIndex].Type == LocalMinMax::Min ? 0.f : 0.5f;
+	Certainty = CertaintyMin + (LastIndex > 0 ? FMath::Abs(MinMax[LastIndex].SignalValue - MinMax[LastIndex - 1].SignalValue) : 0.f) * CertaintyMult;
+}
+
+static void CalculatePhasesFromLocalMinMax(const TArray<LocalMinMax>& MinMax, TArray<FVector2D>& Phases, int32 SignalSize)
+{
+	Phases.Reset();
+	Phases.AddDefaulted(SignalSize);
+
+	float Certainty = 1.f;
+	float Phase = 0.f;
+	for (int32 i = 0; i < SignalSize; ++i)
+	{
+		CalculatePhaseAndCertainty(i, MinMax, SignalSize, Phase, Certainty);
+		FMath::SinCos(&Phases[i].X, &Phases[i].Y, Phase * TWO_PI);
+		Phases[i] *= Certainty;
+	}
+}
+
+} // namespace UE::PoseSearch
+
+void UPoseSearchFeatureChannel_Pose::InitializeSchema(UPoseSearchSchema* Schema)
+{
+	using namespace UE::PoseSearch;
+
+	ChannelDataOffset = Schema->SchemaCardinality;
+	for (int32 ChannelBoneIdx = 0; ChannelBoneIdx != SampledBones.Num(); ++ChannelBoneIdx)
+	{
+		const FPoseSearchBone& SampledBone = SampledBones[ChannelBoneIdx];
+		if (EnumHasAnyFlags(SampledBone.Flags, EPoseSearchBoneFlags::Position))
+		{
+			Schema->SchemaCardinality += FFeatureVectorHelper::EncodeVectorCardinality;
+		}
+		if (EnumHasAnyFlags(SampledBone.Flags, EPoseSearchBoneFlags::Rotation))
+		{
+			Schema->SchemaCardinality += FFeatureVectorHelper::EncodeQuatCardinality;
+		}
+		if (EnumHasAnyFlags(SampledBone.Flags, EPoseSearchBoneFlags::Velocity))
+		{
+			Schema->SchemaCardinality += FFeatureVectorHelper::EncodeVectorCardinality;
+		}
+		if (EnumHasAnyFlags(SampledBone.Flags, EPoseSearchBoneFlags::Phase))
+		{
+			Schema->SchemaCardinality += FFeatureVectorHelper::EncodeVector2DCardinality;
+		}
+	}
+
+	ChannelCardinality = Schema->SchemaCardinality - ChannelDataOffset;
+
+	SchemaBoneIdx.Reset();
+	for (const FPoseSearchBone& Bone : SampledBones)
+	{
+		SchemaBoneIdx.Add(Schema->AddBoneReference(Bone.Reference));
+	}
+}
+
+void UPoseSearchFeatureChannel_Pose::FillWeights(TArray<float>& Weights) const
+{
+	using namespace UE::PoseSearch;
+
+	int32 DataOffset = ChannelDataOffset;
+
+	const int32 NumBones = SampledBones.Num();
+	for (int32 ChannelBoneIdx = 0; ChannelBoneIdx != NumBones; ++ChannelBoneIdx)
+	{
+		const FPoseSearchBone& SampledBone = SampledBones[ChannelBoneIdx];
+		if (EnumHasAnyFlags(SampledBone.Flags, EPoseSearchBoneFlags::Position))
+		{
+			for (int32 i = 0; i != FFeatureVectorHelper::EncodeVectorCardinality; ++i)
+			{
+				Weights[DataOffset + i] = Weight * SampledBone.Weight;
+			}
+			DataOffset += FFeatureVectorHelper::EncodeVectorCardinality;
+		}
+		if (EnumHasAnyFlags(SampledBone.Flags, EPoseSearchBoneFlags::Rotation))
+		{
+			for (int32 i = 0; i != FFeatureVectorHelper::EncodeQuatCardinality; ++i)
+			{
+				Weights[DataOffset + i] = Weight * SampledBone.Weight;
+			}
+			DataOffset += FFeatureVectorHelper::EncodeQuatCardinality;
+		}
+		if (EnumHasAnyFlags(SampledBone.Flags, EPoseSearchBoneFlags::Velocity))
+		{
+			for (int32 i = 0; i != FFeatureVectorHelper::EncodeVectorCardinality; ++i)
+			{
+				Weights[DataOffset + i] = Weight * SampledBone.Weight;
+			}
+			DataOffset += FFeatureVectorHelper::EncodeVectorCardinality;
+		}
+		if (EnumHasAnyFlags(SampledBone.Flags, EPoseSearchBoneFlags::Phase))
+		{
+			for (int32 i = 0; i != FFeatureVectorHelper::EncodeVector2DCardinality; ++i)
+			{
+				Weights[DataOffset + i] = Weight * SampledBone.Weight;
+			}
+			DataOffset += FFeatureVectorHelper::EncodeVector2DCardinality;
+		}
+	}
+
+	check(DataOffset == ChannelDataOffset + ChannelCardinality);
+}
+
+// @todo: do we really need to use double(s) in all this math?
+void UPoseSearchFeatureChannel_Pose::CalculatePhases(UE::PoseSearch::IAssetIndexer& Indexer, UE::PoseSearch::FAssetIndexingOutput& IndexingOutput, TArray<TArray<FVector2D>>& OutPhases) const
+{
+	// @todo: expose them via UI
+	static float BoneSamplingCentralDifferencesTime = 0.2f; // seconds
+	static float SmoothingWindowTime = 0.3f; // seconds
+
+	using namespace UE::PoseSearch;
+	
+	OutPhases.Reset();
+	OutPhases.AddDefaulted(SampledBones.Num());
+	
+	const float FiniteDelta = Indexer.GetIndexingContext().Schema->GetSamplingInterval();
+
+	TArray<float> Signal;
+	TArray<float> SmoothedSignal;
+	TArray<LocalMinMax> LocalMinMax;
+	TArray<FVector> BonePositions;
+	for (int32 ChannelBoneIdx = 0; ChannelBoneIdx != SampledBones.Num(); ++ChannelBoneIdx)
+	{
+		const FPoseSearchBone& SampledBone = SampledBones[ChannelBoneIdx];
+		if (EnumHasAnyFlags(SampledBone.Flags, EPoseSearchBoneFlags::Phase))
+		{
+			CollectBonePositions(BonePositions, Indexer, SchemaBoneIdx[ChannelBoneIdx]);
+
+			// @todo: have different way of calculating signals, for example: height of the bone transform, acceleration, etc?
+			const int32 BoneSamplingCentralDifferencesOffset = FMath::Max(FMath::CeilToInt(BoneSamplingCentralDifferencesTime / FiniteDelta), 1);
+			CalculateSignal(BonePositions, Signal, BoneSamplingCentralDifferencesOffset);
+
+			const int32 SmoothingWindowOffset = FMath::Max(FMath::CeilToInt(SmoothingWindowTime / FiniteDelta), 1);
+			SmoothSignal(Signal, SmoothedSignal, SmoothingWindowOffset);
+
+			FindLocalMinMax(SmoothedSignal, LocalMinMax);
+			ValidateLocalMinMax(LocalMinMax);
+
+			ExtrapolateLocalMinMaxBoundaries(LocalMinMax, SmoothedSignal);
+			ValidateLocalMinMax(LocalMinMax);
+			CalculatePhasesFromLocalMinMax(LocalMinMax, OutPhases[ChannelBoneIdx], SmoothedSignal.Num());
+		}
+	}
+}
+
+void UPoseSearchFeatureChannel_Pose::IndexAsset(UE::PoseSearch::IAssetIndexer& Indexer,  UE::PoseSearch::FAssetIndexingOutput& IndexingOutput) const
+{
+	using namespace UE::PoseSearch;
+	 
+	// Phases is an array of array with cardinality SampledBones.Num() times NumSamples (IndexingContext.EndSampleIdx - IndexingContext.BeginSampleIdx)
+	// of 2 dimensional vectors (FVector2D) representing phases in an Eucledean space with phase angle sin/cos as direction and certainty of the signal as magnitude,
+	// where certainty is a function of the amplitude of the signal used as input
+	TArray<TArray<FVector2D>> Phases;
+	CalculatePhases(Indexer, IndexingOutput, Phases);
+
+	const FAssetIndexingContext& IndexingContext = Indexer.GetIndexingContext();
+	for (int32 SampleIdx = IndexingContext.BeginSampleIdx; SampleIdx != IndexingContext.EndSampleIdx; ++SampleIdx)
+	{
+		const int32 VectorIdx = SampleIdx - IndexingContext.BeginSampleIdx;
+		AddPoseFeatures(Indexer, SampleIdx, IndexingOutput.GetPoseVector(VectorIdx), Phases);
+	}
+}
+
+void UPoseSearchFeatureChannel_Pose::AddPoseFeatures(UE::PoseSearch::IAssetIndexer& Indexer, int32 SampleIdx, TArrayView<float> FeatureVector, const TArray<TArray<FVector2D>>& Phases) const
+{
+	// This function samples the instantaneous pose at time t as well as the pose's velocity and acceleration at time t.
+	// Symmetric finite differences are used to approximate derivatives:
+	//	First symmetric derivative:   f'(t) ~ (f(t+h) - f(t-h)) / 2h
+	//	Second symmetric derivative: f''(t) ~ (f(t+h) - 2f(t) + f(t-h)) / h^2
+	// Where h is a constant time delta
+	// So this means three pose extractions are taken at time t-h, t, and t+h
+	
+	using namespace UE::PoseSearch;
+	using FSampleInfo = IAssetIndexer::FSampleInfo;
+
+	if (SampledBones.IsEmpty())
+	{
+		return;
+	}
+
+	const FAssetIndexingContext& IndexingContext = Indexer.GetIndexingContext();
+	const FAssetSamplingContext* SamplingContext = IndexingContext.SamplingContext;
+
+	const float SampleTime = FMath::Min(SampleIdx * IndexingContext.Schema->GetSamplingInterval(), IndexingContext.AssetSampler->GetPlayLength());
+
+	int32 DataOffset = ChannelDataOffset;
+	for (int32 ChannelBoneIdx = 0; ChannelBoneIdx != SampledBones.Num(); ++ChannelBoneIdx)
+	{
+		const FPoseSearchBone& SampledBone = SampledBones[ChannelBoneIdx];
+
+		// Get each bone's component transform, velocity, and acceleration and add accumulated root motion at this time offset
+		// Think of this process as freezing the character in place (at SampleTime) and then tracing the paths of their joints
+		// as they move through space from past to present to future (at times indicated by PoseSampleTimes).
+
+		if (EnumHasAnyFlags(SampledBone.Flags, EPoseSearchBoneFlags::Position))
+		{
+			bool ClampedPresent;
+			const FTransform BoneTransformsPresent = Indexer.GetTransformAndCacheResults(SampleTime, SampleTime, SchemaBoneIdx[ChannelBoneIdx], ClampedPresent);
+			FFeatureVectorHelper::EncodeVector(FeatureVector, DataOffset, BoneTransformsPresent.GetTranslation());
+		}
+
+		if (EnumHasAnyFlags(SampledBone.Flags, EPoseSearchBoneFlags::Rotation))
+		{
+			bool ClampedPresent;
+			const FTransform BoneTransformsPresent = Indexer.GetTransformAndCacheResults(SampleTime, SampleTime, SchemaBoneIdx[ChannelBoneIdx], ClampedPresent);
+			FFeatureVectorHelper::EncodeQuat(FeatureVector, DataOffset, BoneTransformsPresent.GetRotation());
+		}
+
+		if (EnumHasAnyFlags(SampledBone.Flags, EPoseSearchBoneFlags::Velocity))
+		{
+			bool ClampedPast, ClampedPresent, ClampedFuture;
+			const FTransform BoneTransformsPast = Indexer.GetTransformAndCacheResults(SampleTime - SamplingContext->FiniteDelta, UseCharacterSpaceVelocities ? SampleTime - SamplingContext->FiniteDelta : SampleTime, SchemaBoneIdx[ChannelBoneIdx], ClampedPast);
+			const FTransform BoneTransformsPresent = Indexer.GetTransformAndCacheResults(SampleTime, SampleTime, SchemaBoneIdx[ChannelBoneIdx], ClampedPresent);
+			const FTransform BoneTransformsFuture = Indexer.GetTransformAndCacheResults(SampleTime + SamplingContext->FiniteDelta, UseCharacterSpaceVelocities ? SampleTime + SamplingContext->FiniteDelta : SampleTime, SchemaBoneIdx[ChannelBoneIdx], ClampedFuture);
+
+			// We can get a better finite difference if we ignore samples that have
+			// been clamped at either side of the clip. However, if the central sample 
+			// itself is clamped, or there are no samples that are clamped, we can just 
+			// use the central difference as normal.
+			FVector LinearVelocity;
+			if (ClampedPast && !ClampedPresent && !ClampedFuture)
+			{
+				LinearVelocity = (BoneTransformsFuture.GetTranslation() - BoneTransformsPresent.GetTranslation()) / SamplingContext->FiniteDelta;
+			}
+			else if (ClampedFuture && !ClampedPresent && !ClampedPast)
+			{
+				LinearVelocity = (BoneTransformsPresent.GetTranslation() - BoneTransformsPast.GetTranslation()) / SamplingContext->FiniteDelta;
+			}
+			else
+			{
+				LinearVelocity = (BoneTransformsFuture.GetTranslation() - BoneTransformsPast.GetTranslation()) / (SamplingContext->FiniteDelta * 2.0f);
+			}
+
+			FFeatureVectorHelper::EncodeVector(FeatureVector, DataOffset, LinearVelocity);
+		}
+
+		if (EnumHasAnyFlags(SampledBone.Flags, EPoseSearchBoneFlags::Phase))
+		{
+			const int32 VectorIdx = SampleIdx - IndexingContext.BeginSampleIdx;
+			FFeatureVectorHelper::EncodeVector2D(FeatureVector, DataOffset, Phases[ChannelBoneIdx][VectorIdx]);
+		}
+	}
+
+	check(DataOffset == ChannelDataOffset + ChannelCardinality);
+}
+
+void UPoseSearchFeatureChannel_Pose::BuildQuery(UE::PoseSearch::FSearchContext& SearchContext, FPoseSearchFeatureVectorBuilder& InOutQuery) const
+{
+	using namespace UE::PoseSearch;
+
+	const bool bIsCurrentResultValid = SearchContext.CurrentResult.IsValid();
+	const bool bSkip = InputQueryPose != EInputQueryPose::UseCharacterPose && bIsCurrentResultValid && SearchContext.CurrentResult.Database->Schema == InOutQuery.GetSchema();
+	if (bSkip || !SearchContext.History)
+	{
+		if (bIsCurrentResultValid)
+		{
+			const float LerpValue = InputQueryPose == EInputQueryPose::UseInterpolatedContinuingPose ? SearchContext.CurrentResult.LerpValue : 0.f;
+			int32 DataOffset = ChannelDataOffset;
+			for (int32 SampledBoneIdx = 0; SampledBoneIdx != SampledBones.Num(); ++SampledBoneIdx)
+			{
+				const FPoseSearchBone& SampledBone = SampledBones[SampledBoneIdx];
+				if (EnumHasAnyFlags(SampledBone.Flags, EPoseSearchBoneFlags::Position))
+				{
+					FFeatureVectorHelper::EncodeVector(InOutQuery.EditValues(), DataOffset, SearchContext.GetCurrentResultPrevPoseVector(), SearchContext.GetCurrentResultPoseVector(), SearchContext.GetCurrentResultNextPoseVector(), LerpValue);
+				}
+
+				if (EnumHasAnyFlags(SampledBone.Flags, EPoseSearchBoneFlags::Rotation))
+				{
+					FFeatureVectorHelper::EncodeQuat(InOutQuery.EditValues(), DataOffset, SearchContext.GetCurrentResultPrevPoseVector(), SearchContext.GetCurrentResultPoseVector(), SearchContext.GetCurrentResultNextPoseVector(), LerpValue);
+				}
+
+				if (EnumHasAnyFlags(SampledBone.Flags, EPoseSearchBoneFlags::Velocity))
+				{
+					FFeatureVectorHelper::EncodeVector(InOutQuery.EditValues(), DataOffset, SearchContext.GetCurrentResultPrevPoseVector(), SearchContext.GetCurrentResultPoseVector(), SearchContext.GetCurrentResultNextPoseVector(), LerpValue);
+				}
+
+				if (EnumHasAnyFlags(SampledBone.Flags, EPoseSearchBoneFlags::Phase))
+				{
+					FFeatureVectorHelper::EncodeVector2D(InOutQuery.EditValues(), DataOffset, SearchContext.GetCurrentResultPrevPoseVector(), SearchContext.GetCurrentResultPoseVector(), SearchContext.GetCurrentResultNextPoseVector(), LerpValue);
+				}
+			}
+		}
+		// else leave the InOutQuery set to zero since the SearchContext.History is invalid and it'll fail if we continue
+	}
+	else
+	{
+		struct CachedTransforms
+		{
+			FTransform Current;
+			FTransform Previous;
+			bool Valid = false; // @todo: remove this
+		};
+		TArray<CachedTransforms, TInlineAllocator<32>> CachedTransforms;
+		CachedTransforms.AddUninitialized(SampledBones.Num());
+
+		float SampleTime = 0.f;
+
+		for (int32 SampledBoneIdx = 0; SampledBoneIdx != SampledBones.Num(); ++SampledBoneIdx)
+		{
+			CachedTransforms[SampledBoneIdx].Current = SearchContext.TryGetTransformAndCacheResults(SampleTime, InOutQuery.GetSchema(), SchemaBoneIdx[SampledBoneIdx]);
+			const FPoseSearchBone& SampledBone = SampledBones[SampledBoneIdx];
+
+			if (EnumHasAnyFlags(SampledBone.Flags, EPoseSearchBoneFlags::Velocity))
+			{
+				check(SearchContext.History);
+				const float HistorySameplInterval = SearchContext.History->GetSampleTimeInterval();
+
+				CachedTransforms[SampledBoneIdx].Previous = SearchContext.TryGetTransformAndCacheResults(SampleTime - HistorySameplInterval, InOutQuery.GetSchema(), SchemaBoneIdx[SampledBoneIdx]);
+
+				if (!UE::PoseSearch::UseCharacterSpaceVelocities)
+				{
+					const FTransform RootTransform = SearchContext.TryGetTransformAndCacheResults(SampleTime, InOutQuery.GetSchema(), FSearchContext::SchemaRootBoneIdx);
+					const FTransform RootTransformPrev = SearchContext.TryGetTransformAndCacheResults(SampleTime - HistorySameplInterval, InOutQuery.GetSchema(), FSearchContext::SchemaRootBoneIdx);
+
+					// animation space velocity
+					CachedTransforms[SampledBoneIdx].Previous = CachedTransforms[SampledBoneIdx].Previous * (RootTransformPrev * RootTransform.Inverse());
+				}
+			}
+			CachedTransforms[SampledBoneIdx].Valid = true;
+		}
+
+		int32 DataOffset = ChannelDataOffset;
+		for (int32 SampledBoneIdx = 0; SampledBoneIdx != SampledBones.Num(); ++SampledBoneIdx)
+		{
+			const FPoseSearchBone& SampledBone = SampledBones[SampledBoneIdx];
+			if (EnumHasAnyFlags(SampledBone.Flags, EPoseSearchBoneFlags::Position))
+			{
+				if (CachedTransforms[SampledBoneIdx].Valid)
+				{
+					FFeatureVectorHelper::EncodeVector(InOutQuery.EditValues(), DataOffset, CachedTransforms[SampledBoneIdx].Current.GetTranslation());
+				}
+				else
+				{
+					// preserve the InOutQuery.EditValues() and increase the DataOffset
+					DataOffset += FFeatureVectorHelper::EncodeVectorCardinality;
+				}
+			}
+
+			if (EnumHasAnyFlags(SampledBone.Flags, EPoseSearchBoneFlags::Rotation))
+			{
+				if (CachedTransforms[SampledBoneIdx].Valid)
+				{
+					FFeatureVectorHelper::EncodeQuat(InOutQuery.EditValues(), DataOffset, CachedTransforms[SampledBoneIdx].Current.GetRotation());
+				}
+				else
+				{
+					// preserve the InOutQuery.EditValues() and increase the DataOffset
+					DataOffset += FFeatureVectorHelper::EncodeQuatCardinality;
+				}
+			}
+
+			if (EnumHasAnyFlags(SampledBone.Flags, EPoseSearchBoneFlags::Velocity))
+			{
+				if (CachedTransforms[SampledBoneIdx].Valid)
+				{
+					const FVector LinearVelocity = (CachedTransforms[SampledBoneIdx].Current.GetTranslation() - CachedTransforms[SampledBoneIdx].Previous.GetTranslation()) / SearchContext.History->GetSampleTimeInterval();
+					FFeatureVectorHelper::EncodeVector(InOutQuery.EditValues(), DataOffset, LinearVelocity);
+				}
+				else
+				{
+					// preserve the InOutQuery.EditValues() and increase the DataOffset
+					DataOffset += FFeatureVectorHelper::EncodeVectorCardinality;
+				}
+			}
+
+			if (EnumHasAnyFlags(SampledBone.Flags, EPoseSearchBoneFlags::Phase))
+			{
+				DataOffset += FFeatureVectorHelper::EncodeVector2DCardinality;
+
+				// @todo: Support phase in BuildQuery
+				// FFeatureVectorHelper::EncodeVector2D(InOutQuery.EditValues(), DataOffset, ???);
+			}
+		}
+
+		check(DataOffset == ChannelDataOffset + ChannelCardinality);
+	}
+}
+
+void UPoseSearchFeatureChannel_Pose::DebugDraw(const UE::PoseSearch::FDebugDrawParams& DrawParams, TConstArrayView<float> PoseVector) const
+{
+#if ENABLE_DRAW_DEBUG
+	using namespace UE::PoseSearch;
+
+	const UPoseSearchSchema* Schema = DrawParams.GetSchema();
+	check(Schema && Schema->IsValid());
+
+	const float LifeTime = DrawParams.DefaultLifeTime;
+	const uint8 DepthPriority = ESceneDepthPriorityGroup::SDPG_Foreground + 2;
+	const bool bPersistent = EnumHasAnyFlags(DrawParams.Flags, EDebugDrawFlags::Persistent);
+
+	const int32 NumBones = SampledBones.Num();
+
+	if (NumBones == 0)
+	{
+		return;
+	}
+
+	int32 DataOffset = ChannelDataOffset;
+	for (int32 ChannelBoneIdx = 0; ChannelBoneIdx != NumBones; ++ChannelBoneIdx)
+	{
+		const FPoseSearchBone& SampledBone = SampledBones[ChannelBoneIdx];
+
+		FVector BonePos;
+		if (EnumHasAnyFlags(SampledBone.Flags, EPoseSearchBoneFlags::Position))
+		{
+			BonePos = FFeatureVectorHelper::DecodeVector(PoseVector, DataOffset);
+
+			const FColor Color = DrawParams.GetColor(SampledBone.ColorPresetIndex);
+
+			BonePos = DrawParams.RootTransform.TransformPosition(BonePos);
+			if (EnumHasAnyFlags(DrawParams.Flags, EDebugDrawFlags::DrawFast | EDebugDrawFlags::DrawSearchIndex))
+			{
+				DrawDebugPoint(DrawParams.World, BonePos, DrawParams.PointSize, Color, bPersistent, LifeTime, DepthPriority);
+			}
+			else
+			{
+				DrawDebugSphere(DrawParams.World, BonePos, 2.f, 8, Color, bPersistent, LifeTime, DepthPriority);
+			}
+
+			if (EnumHasAnyFlags(DrawParams.Flags, EDebugDrawFlags::DrawBoneNames))
+			{
+				DrawDebugString(
+					DrawParams.World, BonePos + FVector(0.0, 0.0, 10.0),
+					Schema->BoneReferences[SchemaBoneIdx[ChannelBoneIdx]].BoneName.ToString(),
+					nullptr, Color, LifeTime, false, 1.0f);
+			}
+		}
+		else
+		{
+			BonePos = DrawParams.Mesh != nullptr ? DrawParams.Mesh->GetSocketTransform(SampledBones[ChannelBoneIdx].Reference.BoneName).GetLocation() : DrawParams.RootTransform.GetTranslation();
+		}
+
+		if (EnumHasAnyFlags(SampledBone.Flags, EPoseSearchBoneFlags::Rotation))
+		{
+			const FQuat BoneRot = FFeatureVectorHelper::DecodeQuat(PoseVector, DataOffset);
+			// @todo: debug draw rotation
+		}
+
+		if (EnumHasAnyFlags(SampledBone.Flags, EPoseSearchBoneFlags::Velocity))
+		{
+			FVector BoneVel = FFeatureVectorHelper::DecodeVector(PoseVector, DataOffset);
+
+			const FColor Color = DrawParams.GetColor(SampledBone.ColorPresetIndex);
+
+			BoneVel *= 0.08f;
+			BoneVel = DrawParams.RootTransform.TransformVector(BoneVel);
+			FVector BoneVelDirection = BoneVel.GetSafeNormal();
+
+			if (EnumHasAnyFlags(DrawParams.Flags, EDebugDrawFlags::DrawSearchIndex))
+			{
+				DrawDebugLine(DrawParams.World, BonePos, BonePos + BoneVel, Color, bPersistent, LifeTime, DepthPriority);
+			}
+			else
+			{
+				const float AdjustedThickness = EnumHasAnyFlags(DrawParams.Flags, EDebugDrawFlags::DrawFast) ? 0.0f : 1.f;
+
+				DrawDebugLine(
+					DrawParams.World,
+					BonePos + BoneVelDirection * 2.f,
+					BonePos + BoneVel,
+					Color,
+					bPersistent,
+					LifeTime,
+					DepthPriority,
+					AdjustedThickness);
+			}
+		}
+
+		if (EnumHasAnyFlags(SampledBone.Flags, EPoseSearchBoneFlags::Phase))
+		{
+			const FVector2D Phase = FFeatureVectorHelper::DecodeVector2D(PoseVector, DataOffset);
+
+			const FColor Color = DrawParams.GetColor(SampledBone.ColorPresetIndex);
+
+			static float ScaleFactor = 1.f;
+
+			const FVector TransformXAxisVector = DrawParams.RootTransform.TransformVector(FVector::XAxisVector);
+			const FVector TransformYAxisVector = DrawParams.RootTransform.TransformVector(FVector::YAxisVector);
+			const FVector TransformZAxisVector = DrawParams.RootTransform.TransformVector(FVector::ZAxisVector);
+
+			const FVector PhaseVector = (TransformZAxisVector * Phase.X + TransformYAxisVector * Phase.Y) * ScaleFactor;
+			DrawDebugLine(DrawParams.World, BonePos, BonePos + PhaseVector, Color, bPersistent, LifeTime, DepthPriority, 0.f);
+
+			static int32 Segments = 64;
+			FMatrix CircleTransform;
+			CircleTransform.SetAxes(&TransformXAxisVector, &TransformYAxisVector, &TransformZAxisVector, &BonePos);
+			DrawDebugCircle(DrawParams.World, CircleTransform, PhaseVector.Length(), Segments, Color, bPersistent, LifeTime, DepthPriority, 0.f, false);
+		}
+	}
+
+	check(DataOffset == ChannelDataOffset + ChannelCardinality);
+#endif // ENABLE_DRAW_DEBUG
+}
+
+#if WITH_EDITOR
+void UPoseSearchFeatureChannel_Pose::PopulateChannelLayoutSet(UE::PoseSearch::FFeatureChannelLayoutSet& FeatureChannelLayoutSet) const
+{
+	using namespace UE::PoseSearch;
+	int32 DataOffset = ChannelDataOffset;
+
+	auto Add = [&FeatureChannelLayoutSet, &DataOffset](const FPoseSearchBone& SampledBone, EPoseSearchBoneFlags BoneFlag, const char* Label, int32 Cardinality)
+	{
+		FString SkeletonName = FeatureChannelLayoutSet.CurrentSchema->Skeleton->GetName();
+		FString BoneName = SampledBone.Reference.BoneName.ToString();
+
+		UE::PoseSearch::FKeyBuilder KeyBuilder;
+		KeyBuilder << SkeletonName << BoneName << BoneFlag;
+		FeatureChannelLayoutSet.Add(FString::Format(TEXT("{0} {1}"), { BoneName, Label }), KeyBuilder.Finalize(), DataOffset, Cardinality);
+
+		DataOffset += Cardinality;
+	};
+
+	for (int32 ChannelBoneIdx = 0; ChannelBoneIdx != SampledBones.Num(); ++ChannelBoneIdx)
+	{
+		const FPoseSearchBone& SampledBone = SampledBones[ChannelBoneIdx];
+		if (EnumHasAnyFlags(SampledBone.Flags, EPoseSearchBoneFlags::Position))
+		{
+			Add(SampledBone, EPoseSearchBoneFlags::Position, "Pos", FFeatureVectorHelper::EncodeVectorCardinality);
+		}
+		if (EnumHasAnyFlags(SampledBone.Flags, EPoseSearchBoneFlags::Rotation))
+		{
+			Add(SampledBone, EPoseSearchBoneFlags::Rotation, "Rot", FFeatureVectorHelper::EncodeQuatCardinality);
+		}
+		if (EnumHasAnyFlags(SampledBone.Flags, EPoseSearchBoneFlags::Velocity))
+		{
+			Add(SampledBone, EPoseSearchBoneFlags::Velocity, "Vel", FFeatureVectorHelper::EncodeVectorCardinality);
+		}
+		if (EnumHasAnyFlags(SampledBone.Flags, EPoseSearchBoneFlags::Phase))
+		{
+			Add(SampledBone, EPoseSearchBoneFlags::Phase, "Pha", FFeatureVectorHelper::EncodeVector2DCardinality);
+		}
+	}
+
+	check(DataOffset == ChannelDataOffset + ChannelCardinality);
+}
+
+void UPoseSearchFeatureChannel_Pose::ComputeCostBreakdowns(UE::PoseSearch::ICostBreakDownData& CostBreakDownData, const UPoseSearchSchema* Schema) const
+{
+	using namespace UE::PoseSearch;
+
+	CostBreakDownData.AddEntireBreakDownSection(LOCTEXT("ColumnLabelPoseChannelTotal", "Pose Total"), Schema, ChannelDataOffset, ChannelCardinality);
+
+	if (CostBreakDownData.IsVerbose())
+	{
+		int32 DataOffset = ChannelDataOffset;
+
+		for (int32 ChannelBoneIdx = 0; ChannelBoneIdx != SampledBones.Num(); ++ChannelBoneIdx)
+		{
+			const FPoseSearchBone& SampledBone = SampledBones[ChannelBoneIdx];
+			if (EnumHasAnyFlags(SampledBone.Flags, EPoseSearchBoneFlags::Position))
+			{
+				CostBreakDownData.AddEntireBreakDownSection(FText::Format(LOCTEXT("ColumnLabelPoseChannelPosition", "{0} Pos"), FText::FromName(SampledBone.Reference.BoneName)), Schema, DataOffset, FFeatureVectorHelper::EncodeVectorCardinality);
+				DataOffset += FFeatureVectorHelper::EncodeVectorCardinality;
+			}
+			if (EnumHasAnyFlags(SampledBone.Flags, EPoseSearchBoneFlags::Rotation))
+			{
+				CostBreakDownData.AddEntireBreakDownSection(FText::Format(LOCTEXT("ColumnLabelPoseChannelRotation", "{0} Rot"), FText::FromName(SampledBone.Reference.BoneName)), Schema, DataOffset, FFeatureVectorHelper::EncodeQuatCardinality);
+				DataOffset += FFeatureVectorHelper::EncodeQuatCardinality;
+			}
+			if (EnumHasAnyFlags(SampledBone.Flags, EPoseSearchBoneFlags::Velocity))
+			{
+				CostBreakDownData.AddEntireBreakDownSection(FText::Format(LOCTEXT("ColumnLabelPoseChannelVelocity", "{0} Vel"), FText::FromName(SampledBone.Reference.BoneName)), Schema, DataOffset, FFeatureVectorHelper::EncodeVectorCardinality);
+				DataOffset += FFeatureVectorHelper::EncodeVectorCardinality;
+			}
+			if (EnumHasAnyFlags(SampledBone.Flags, EPoseSearchBoneFlags::Phase))
+			{
+				CostBreakDownData.AddEntireBreakDownSection(FText::Format(LOCTEXT("ColumnLabelPoseChannelPhase", "{0} Pha"), FText::FromName(SampledBone.Reference.BoneName)), Schema, DataOffset, FFeatureVectorHelper::EncodeVector2DCardinality);
+				DataOffset += FFeatureVectorHelper::EncodeVector2DCardinality;
+			}
+		}
+
+		check(DataOffset == ChannelDataOffset + ChannelCardinality);
+	}
+}
+
+#endif // WITH_EDITOR
+
+#undef LOCTEXT_NAMESPACE
\ No newline at end of file
diff --git a/Engine/Plugins/Experimental/Animation/PoseSearch/Source/Runtime/Private/PoseSearchFeatureChannel_Pose.h b/Engine/Plugins/Experimental/Animation/PoseSearch/Source/Runtime/Private/PoseSearchFeatureChannel_Pose.h
new file mode 100644
--- /dev/null
+++ b/Engine/Plugins/Experimental/Animation/PoseSearch/Source/Runtime/Private/PoseSearchFeatureChannel_Pose.h
@@ -0,0 +1,75 @@
+// Copyright Epic Games, Inc. All Rights Reserved.
+
+#pragma once
+
+#include "PoseSearch/PoseSearchFeatureChannel.h"
+#include "BoneContainer.h"
+#include "PoseSearchFeatureChannel_Pose.generated.h"
+
+UENUM(meta = (Bitflags, UseEnumValuesAsMaskValuesInEditor = "true"))
+enum class EPoseSearchBoneFlags : uint32
+{
+	Velocity = 1 << 0,
+	Position = 1 << 1,
+	Rotation = 1 << 2,
+	Phase = 1 << 3,
+};
+ENUM_CLASS_FLAGS(EPoseSearchBoneFlags);
+constexpr bool EnumHasAnyFlags(int32 Flags, EPoseSearchBoneFlags Contains) { return (Flags & int32(Contains)) != 0; }
+inline int32& operator|=(int32& Lhs, EPoseSearchBoneFlags Rhs) { return Lhs |= int32(Rhs); }
+
+USTRUCT()
+struct POSESEARCH_API FPoseSearchBone
+{
+	GENERATED_BODY()
+
+	UPROPERTY(EditAnywhere, Category = Config)
+	FBoneReference Reference;
+
+	UPROPERTY(EditAnywhere, meta = (Bitmask, BitmaskEnum = "/Script/PoseSearch.EPoseSearchBoneFlags"), Category = Config)
+	int32 Flags = int32(EPoseSearchBoneFlags::Position);
+
+	UPROPERTY(EditAnywhere, Category = Config)
+	float Weight = 1.f;
+
+	UPROPERTY(EditAnywhere, Category = Config, meta=(ExcludeFromHash))
+	int32 ColorPresetIndex = 0;
+};
+
+// UPoseSearchFeatureChannel_Pose
+UCLASS(BlueprintType, EditInlineNew, meta = (DisplayName = "Pose Channel"), CollapseCategories)
+class POSESEARCH_API UPoseSearchFeatureChannel_Pose : public UPoseSearchFeatureChannel
+{
+	GENERATED_BODY()
+
+public:
+	UPROPERTY(EditAnywhere, Category = "Settings")
+	float Weight = 1.f;
+
+	UPROPERTY(EditAnywhere, Category = "Settings")
+	TArray<FPoseSearchBone> SampledBones;
+
+	UPROPERTY()
+	TArray<int8> SchemaBoneIdx;
+
+	UPROPERTY(EditAnywhere, Category = "Settings")
+	EInputQueryPose InputQueryPose = EInputQueryPose::UseContinuingPose;
+
+	// UPoseSearchFeatureChannel interface
+	virtual void InitializeSchema(UPoseSearchSchema* Schema) override;
+	virtual void FillWeights(TArray<float>& Weights) const override;
+	virtual void IndexAsset(UE::PoseSearch::IAssetIndexer& Indexer, UE::PoseSearch::FAssetIndexingOutput& IndexingOutput) const override;
+	virtual void BuildQuery(UE::PoseSearch::FSearchContext& SearchContext, FPoseSearchFeatureVectorBuilder& InOutQuery) const override;
+	virtual void DebugDraw(const UE::PoseSearch::FDebugDrawParams& DrawParams, TConstArrayView<float> PoseVector) const override;
+
+#if WITH_EDITOR
+	virtual void PopulateChannelLayoutSet(UE::PoseSearch::FFeatureChannelLayoutSet& FeatureChannelLayoutSet) const override;
+	virtual void ComputeCostBreakdowns(UE::PoseSearch::ICostBreakDownData& CostBreakDownData, const UPoseSearchSchema* Schema) const override;
+#endif
+
+protected:
+	void AddPoseFeatures(UE::PoseSearch::IAssetIndexer& Indexer, int32 SampleIdx, TArrayView<float> FeatureVector, const TArray<TArray<FVector2D>>& Phases) const;
+	void CalculatePhases(UE::PoseSearch::IAssetIndexer& Indexer, UE::PoseSearch::FAssetIndexingOutput& IndexingOutput, TArray<TArray<FVector2D>>& OutPhases) const;
+};
+
+
diff --git a/Engine/Plugins/Experimental/Animation/PoseSearch/Source/Runtime/Private/PoseSearchFeatureChannel_Position.cpp b/Engine/Plugins/Experimental/Animation/PoseSearch/Source/Runtime/Private/PoseSearchFeatureChannel_Position.cpp
new file mode 100644
--- /dev/null
+++ b/Engine/Plugins/Experimental/Animation/PoseSearch/Source/Runtime/Private/PoseSearchFeatureChannel_Position.cpp
@@ -0,0 +1,111 @@
+// Copyright Epic Games, Inc. All Rights Reserved.
+
+#include "PoseSearchFeatureChannel_Position.h"
+#include "DrawDebugHelpers.h"
+#include "PoseSearch/PoseSearch.h"
+#include "PoseSearch/PoseSearchSchema.h"
+
+void UPoseSearchFeatureChannel_Position::InitializeSchema(UPoseSearchSchema* Schema)
+{
+	ChannelDataOffset = Schema->SchemaCardinality;
+	ChannelCardinality = UE::PoseSearch::FFeatureVectorHelper::EncodeVectorCardinality;
+	Schema->SchemaCardinality += ChannelCardinality;
+
+	SchemaBoneIdx = Schema->AddBoneReference(Bone);
+}
+
+void UPoseSearchFeatureChannel_Position::FillWeights(TArray<float>& Weights) const
+{
+	using namespace UE::PoseSearch;
+
+	for (int32 i = 0; i != FFeatureVectorHelper::EncodeVectorCardinality; ++i)
+	{
+		Weights[ChannelDataOffset + i] = Weight;
+	}
+}
+
+void UPoseSearchFeatureChannel_Position::IndexAsset(UE::PoseSearch::IAssetIndexer& Indexer, UE::PoseSearch::FAssetIndexingOutput& IndexingOutput) const
+{
+	using namespace UE::PoseSearch;
+
+	const FAssetIndexingContext& IndexingContext = Indexer.GetIndexingContext();
+
+	for (int32 SampleIdx = IndexingContext.BeginSampleIdx; SampleIdx != IndexingContext.EndSampleIdx; ++SampleIdx)
+	{
+		const int32 VectorIdx = SampleIdx - IndexingContext.BeginSampleIdx;
+
+		const float OriginSampleTime = FMath::Min(SampleIdx * IndexingContext.Schema->GetSamplingInterval(), IndexingContext.AssetSampler->GetPlayLength());
+		const float SubsampleTime = OriginSampleTime + SampleTimeOffset;
+
+		bool ClampedPresent;
+		const FTransform BoneTransformsPresent = Indexer.GetTransformAndCacheResults(SubsampleTime, bUseSampleTimeOffsetRootBone ? SubsampleTime : OriginSampleTime, SchemaBoneIdx, ClampedPresent);
+		int32 DataOffset = ChannelDataOffset;
+		FFeatureVectorHelper::EncodeVector(IndexingOutput.GetPoseVector(VectorIdx), DataOffset, BoneTransformsPresent.GetTranslation());
+	}
+}
+
+void UPoseSearchFeatureChannel_Position::BuildQuery(UE::PoseSearch::FSearchContext& SearchContext, FPoseSearchFeatureVectorBuilder& InOutQuery) const
+{
+	using namespace UE::PoseSearch;
+
+	const bool bIsCurrentResultValid = SearchContext.CurrentResult.IsValid();
+	const bool bSkip = InputQueryPose != EInputQueryPose::UseCharacterPose && bIsCurrentResultValid && SearchContext.CurrentResult.Database->Schema == InOutQuery.GetSchema();
+	if (bSkip || !SearchContext.History)
+	{
+		if (bIsCurrentResultValid)
+		{
+			const float LerpValue = InputQueryPose == EInputQueryPose::UseInterpolatedContinuingPose ? SearchContext.CurrentResult.LerpValue : 0.f;
+			int32 DataOffset = ChannelDataOffset;
+			FFeatureVectorHelper::EncodeVector(InOutQuery.EditValues(), DataOffset, SearchContext.GetCurrentResultPrevPoseVector(), SearchContext.GetCurrentResultPoseVector(), SearchContext.GetCurrentResultNextPoseVector(), LerpValue);
+		}
+		// else leave the InOutQuery set to zero since the SearchContext.History is invalid and it'll fail if we continue
+	}
+	else
+	{
+		FTransform Transform = SearchContext.TryGetTransformAndCacheResults(SampleTimeOffset, InOutQuery.GetSchema(), SchemaBoneIdx);
+
+		if (!bUseSampleTimeOffsetRootBone)
+		{
+			const FTransform RootTransform = SearchContext.TryGetTransformAndCacheResults(0.f, InOutQuery.GetSchema(), FSearchContext::SchemaRootBoneIdx);
+			const FTransform RootTransformPrev = SearchContext.TryGetTransformAndCacheResults(SampleTimeOffset, InOutQuery.GetSchema(), FSearchContext::SchemaRootBoneIdx);
+			Transform = Transform * (RootTransformPrev * RootTransform.Inverse());
+		}
+
+		int32 DataOffset = ChannelDataOffset;
+		FFeatureVectorHelper::EncodeVector(InOutQuery.EditValues(), DataOffset, Transform.GetTranslation());
+		check(DataOffset == ChannelDataOffset + ChannelCardinality);
+	}
+}
+
+void UPoseSearchFeatureChannel_Position::DebugDraw(const UE::PoseSearch::FDebugDrawParams& DrawParams, TConstArrayView<float> PoseVector) const
+{
+#if ENABLE_DRAW_DEBUG
+	using namespace UE::PoseSearch;
+
+	const UPoseSearchSchema* Schema = DrawParams.GetSchema();
+	check(Schema && Schema->IsValid());
+
+	const float LifeTime = DrawParams.DefaultLifeTime;
+	const uint8 DepthPriority = ESceneDepthPriorityGroup::SDPG_Foreground + 2;
+	const bool bPersistent = EnumHasAnyFlags(DrawParams.Flags, EDebugDrawFlags::Persistent);
+
+	int32 DataOffset = ChannelDataOffset;
+	const FVector BonePos = DrawParams.RootTransform.TransformPosition(FFeatureVectorHelper::DecodeVector(PoseVector, DataOffset));
+
+	const FColor Color = DrawParams.GetColor(ColorPresetIndex);
+
+	if (EnumHasAnyFlags(DrawParams.Flags, EDebugDrawFlags::DrawFast | EDebugDrawFlags::DrawSearchIndex))
+	{
+		DrawDebugPoint(DrawParams.World, BonePos, DrawParams.PointSize, Color, bPersistent, LifeTime, DepthPriority);
+	}
+	else
+	{
+		DrawDebugSphere(DrawParams.World, BonePos, 2.f, 8, Color, bPersistent, LifeTime, DepthPriority);
+	}
+
+	if (EnumHasAnyFlags(DrawParams.Flags, EDebugDrawFlags::DrawBoneNames))
+	{
+		DrawDebugString(DrawParams.World, BonePos + FVector(0.0, 0.0, 10.0), Schema->BoneReferences[SchemaBoneIdx].BoneName.ToString(), nullptr, Color, LifeTime, false, 1.0f);
+	}
+#endif // ENABLE_DRAW_DEBUG
+}
diff --git a/Engine/Plugins/Experimental/Animation/PoseSearch/Source/Runtime/Private/PoseSearchFeatureChannel_Position.h b/Engine/Plugins/Experimental/Animation/PoseSearch/Source/Runtime/Private/PoseSearchFeatureChannel_Position.h
new file mode 100644
--- /dev/null
+++ b/Engine/Plugins/Experimental/Animation/PoseSearch/Source/Runtime/Private/PoseSearchFeatureChannel_Position.h
@@ -0,0 +1,44 @@
+// Copyright Epic Games, Inc. All Rights Reserved.
+
+#pragma once
+
+#include "PoseSearch/PoseSearchFeatureChannel.h"
+#include "BoneContainer.h"
+#include "PoseSearchFeatureChannel_Position.generated.h"
+
+UCLASS(BlueprintType, EditInlineNew, meta = (DisplayName = "Position Channel"), CollapseCategories)
+class POSESEARCH_API UPoseSearchFeatureChannel_Position : public UPoseSearchFeatureChannel
+{
+	GENERATED_BODY()
+
+public:
+	UPROPERTY(EditAnywhere, Category = "Settings")
+	FBoneReference Bone;
+
+	UPROPERTY(EditAnywhere, Category = "Settings")
+	float Weight = 1.f;
+
+	UPROPERTY(EditAnywhere, Category = "Settings")
+	float SampleTimeOffset = 0.f;
+
+	UPROPERTY(meta = (ExcludeFromHash))
+	int8 SchemaBoneIdx = 0;
+
+	UPROPERTY(EditAnywhere, Category = "Settings", meta = (ExcludeFromHash))
+	int32 ColorPresetIndex = 0;
+
+	UPROPERTY(EditAnywhere, Category = "Settings")
+	EInputQueryPose InputQueryPose = EInputQueryPose::UseContinuingPose;
+
+	// if UseSampleTimeOffsetRootBone is true, this UPoseSearchFeatureChannel_Position will calculate the position of Bone from the pose SampleTimeOffset seconds away from the current time pose root bone
+	// if false the calculated position of Bone will be in component space from the pose SampleTimeOffset seconds away
+	UPROPERTY(EditAnywhere, Category = "Settings")
+	bool bUseSampleTimeOffsetRootBone = true;
+
+	// UPoseSearchFeatureChannel interface
+	virtual void InitializeSchema(UPoseSearchSchema* Schema) override;
+	virtual void FillWeights(TArray<float>& Weights) const override;
+	virtual void IndexAsset(UE::PoseSearch::IAssetIndexer& Indexer, UE::PoseSearch::FAssetIndexingOutput& IndexingOutput) const override;
+	virtual void BuildQuery(UE::PoseSearch::FSearchContext& SearchContext, FPoseSearchFeatureVectorBuilder& InOutQuery) const override;
+	virtual void DebugDraw(const UE::PoseSearch::FDebugDrawParams& DrawParams, TConstArrayView<float> PoseVector) const override;
+};
diff --git a/Engine/Plugins/Experimental/Animation/PoseSearch/Source/Runtime/Private/PoseSearchFeatureChannel_Trajectory.cpp b/Engine/Plugins/Experimental/Animation/PoseSearch/Source/Runtime/Private/PoseSearchFeatureChannel_Trajectory.cpp
new file mode 100644
--- /dev/null
+++ b/Engine/Plugins/Experimental/Animation/PoseSearch/Source/Runtime/Private/PoseSearchFeatureChannel_Trajectory.cpp
@@ -0,0 +1,936 @@
+// Copyright Epic Games, Inc. All Rights Reserved.
+
+#include "PoseSearchFeatureChannel_Trajectory.h"
+#include "DrawDebugHelpers.h"
+#include "PoseSearch/PoseSearch.h"
+#include "PoseSearch/PoseSearchDerivedDataKey.h"
+#include "PoseSearch/PoseSearchSchema.h"
+
+#define LOCTEXT_NAMESPACE "PoseSearchFeatureChannels"
+
+void UPoseSearchFeatureChannel_Trajectory::PreSave(FObjectPreSaveContext ObjectSaveContext)
+{
+	Samples.Sort([](const FPoseSearchTrajectorySample& a, const FPoseSearchTrajectorySample& b)
+	{
+		return a.Offset < b.Offset;
+	});
+
+	Super::PreSave(ObjectSaveContext);
+}
+
+void UPoseSearchFeatureChannel_Trajectory::InitializeSchema(UPoseSearchSchema* Schema)
+{
+	using namespace UE::PoseSearch;
+
+	ChannelDataOffset = Schema->SchemaCardinality;
+
+	for (const FPoseSearchTrajectorySample& Sample : Samples)
+	{
+		if (EnumHasAnyFlags(Sample.Flags, EPoseSearchTrajectoryFlags::Position))
+		{
+			Schema->SchemaCardinality += FFeatureVectorHelper::EncodeVectorCardinality;
+		}
+		if (EnumHasAnyFlags(Sample.Flags, EPoseSearchTrajectoryFlags::PositionXY))
+		{
+			Schema->SchemaCardinality += FFeatureVectorHelper::EncodeVector2DCardinality;
+		}
+
+		if (EnumHasAnyFlags(Sample.Flags, EPoseSearchTrajectoryFlags::Velocity))
+		{
+			Schema->SchemaCardinality += FFeatureVectorHelper::EncodeVectorCardinality;
+		}
+		if (EnumHasAnyFlags(Sample.Flags, EPoseSearchTrajectoryFlags::VelocityXY))
+		{
+			Schema->SchemaCardinality += FFeatureVectorHelper::EncodeVector2DCardinality;
+		}
+
+		if (EnumHasAnyFlags(Sample.Flags, EPoseSearchTrajectoryFlags::VelocityDirection))
+		{
+			Schema->SchemaCardinality += FFeatureVectorHelper::EncodeVectorCardinality;
+		}
+		if (EnumHasAnyFlags(Sample.Flags, EPoseSearchTrajectoryFlags::VelocityDirectionXY))
+		{
+			Schema->SchemaCardinality += FFeatureVectorHelper::EncodeVector2DCardinality;
+		}
+
+		if (EnumHasAnyFlags(Sample.Flags, EPoseSearchTrajectoryFlags::FacingDirection))
+		{
+			Schema->SchemaCardinality += FFeatureVectorHelper::EncodeVectorCardinality;
+		}
+		if (EnumHasAnyFlags(Sample.Flags, EPoseSearchTrajectoryFlags::FacingDirectionXY))
+		{
+			Schema->SchemaCardinality += FFeatureVectorHelper::EncodeVector2DCardinality;
+		}
+	}
+
+	ChannelCardinality = Schema->SchemaCardinality - ChannelDataOffset;
+}
+
+void UPoseSearchFeatureChannel_Trajectory::FillWeights(TArray<float>& Weights) const
+{
+	using namespace UE::PoseSearch;
+
+	int32 DataOffset = ChannelDataOffset;
+	for (const FPoseSearchTrajectorySample& Sample : Samples)
+	{
+		if (EnumHasAnyFlags(Sample.Flags, EPoseSearchTrajectoryFlags::Position))
+		{
+			for (int32 i = 0; i != FFeatureVectorHelper::EncodeVectorCardinality; ++i)
+			{
+				Weights[DataOffset + i] = Weight * Sample.Weight;
+			}
+			DataOffset += FFeatureVectorHelper::EncodeVectorCardinality;
+		}
+		if (EnumHasAnyFlags(Sample.Flags, EPoseSearchTrajectoryFlags::PositionXY))
+		{
+			for (int32 i = 0; i != FFeatureVectorHelper::EncodeVector2DCardinality; ++i)
+			{
+				Weights[DataOffset + i] = Weight * Sample.Weight;
+			}
+			DataOffset += FFeatureVectorHelper::EncodeVector2DCardinality;
+		}
+
+		if (EnumHasAnyFlags(Sample.Flags, EPoseSearchTrajectoryFlags::Velocity))
+		{
+			for (int32 i = 0; i != FFeatureVectorHelper::EncodeVectorCardinality; ++i)
+			{
+				Weights[DataOffset + i] = Weight * Sample.Weight;
+			}
+			DataOffset += FFeatureVectorHelper::EncodeVectorCardinality;
+		}
+		if (EnumHasAnyFlags(Sample.Flags, EPoseSearchTrajectoryFlags::VelocityXY))
+		{
+			for (int32 i = 0; i != FFeatureVectorHelper::EncodeVector2DCardinality; ++i)
+			{
+				Weights[DataOffset + i] = Weight * Sample.Weight;
+			}
+			DataOffset += FFeatureVectorHelper::EncodeVector2DCardinality;
+		}
+
+		if (EnumHasAnyFlags(Sample.Flags, EPoseSearchTrajectoryFlags::VelocityDirection))
+		{
+			for (int32 i = 0; i != FFeatureVectorHelper::EncodeVectorCardinality; ++i)
+			{
+				Weights[DataOffset + i] = Weight * Sample.Weight;
+			}
+			DataOffset += FFeatureVectorHelper::EncodeVectorCardinality;
+		}
+		if (EnumHasAnyFlags(Sample.Flags, EPoseSearchTrajectoryFlags::VelocityDirectionXY))
+		{
+			for (int32 i = 0; i != FFeatureVectorHelper::EncodeVector2DCardinality; ++i)
+			{
+				Weights[DataOffset + i] = Weight * Sample.Weight;
+			}
+			DataOffset += FFeatureVectorHelper::EncodeVector2DCardinality;
+		}
+
+		if (EnumHasAnyFlags(Sample.Flags, EPoseSearchTrajectoryFlags::FacingDirection))
+		{
+			for (int32 i = 0; i != FFeatureVectorHelper::EncodeVectorCardinality; ++i)
+			{
+				Weights[DataOffset + i] = Weight * Sample.Weight;
+			}
+			DataOffset += FFeatureVectorHelper::EncodeVectorCardinality;
+		}
+		if (EnumHasAnyFlags(Sample.Flags, EPoseSearchTrajectoryFlags::FacingDirectionXY))
+		{
+			for (int32 i = 0; i != FFeatureVectorHelper::EncodeVector2DCardinality; ++i)
+			{
+				Weights[DataOffset + i] = Weight * Sample.Weight;
+			}
+			DataOffset += FFeatureVectorHelper::EncodeVector2DCardinality;
+		}
+	}
+
+	check(DataOffset == ChannelDataOffset + ChannelCardinality);
+}
+
+void UPoseSearchFeatureChannel_Trajectory::IndexAsset(UE::PoseSearch::IAssetIndexer& Indexer,  UE::PoseSearch::FAssetIndexingOutput& IndexingOutput) const
+{
+	const UE::PoseSearch::FAssetIndexingContext& IndexingContext = Indexer.GetIndexingContext();
+	for (int32 SampleIdx = IndexingContext.BeginSampleIdx; SampleIdx != IndexingContext.EndSampleIdx; ++SampleIdx)
+	{
+		const int32 VectorIdx = SampleIdx - IndexingContext.BeginSampleIdx;
+		IndexAssetPrivate(Indexer, SampleIdx, IndexingOutput.GetPoseVector(VectorIdx));
+	}
+}
+
+void UPoseSearchFeatureChannel_Trajectory::IndexAssetPrivate(const UE::PoseSearch::IAssetIndexer& Indexer, int32 SampleIdx, TArrayView<float> FeatureVector) const
+{
+	// This function samples the instantaneous trajectory at time t as well as the trajectory's velocity and acceleration at time t.
+	// Symmetric finite differences are used to approximate derivatives:
+	//	First symmetric derivative:   f'(t) ~ (f(t+h) - f(t-h)) / 2h
+	//	Second symmetric derivative: f''(t) ~ (f(t+h) - 2f(t) + f(t-h)) / h^2
+	// Where h is a constant time delta
+	// So this means three root motion extractions are taken at time t-h, t, and t+h
+
+	using namespace UE::PoseSearch;
+	using FSampleInfo = IAssetIndexer::FSampleInfo;
+
+	const FAssetIndexingContext& IndexingContext = Indexer.GetIndexingContext();
+	float SampleTime = FMath::Min(SampleIdx * IndexingContext.Schema->GetSamplingInterval(), IndexingContext.AssetSampler->GetPlayLength());
+	FSampleInfo Origin = Indexer.GetSampleInfo(SampleTime);
+
+	int32 DataOffset = ChannelDataOffset;
+	for (const FPoseSearchTrajectorySample& Sample : Samples)
+	{
+		const float SubsampleTime = Sample.Offset + SampleTime;
+
+		// For each pose subsample term, get the corresponding clip, accumulated root motion,
+			// and wrap the time parameter based on the clip's length.
+		const FSampleInfo SamplePast = Indexer.GetSampleInfoRelative(SubsampleTime - IndexingContext.SamplingContext->FiniteDelta, Origin);
+		const FSampleInfo SamplePresent = Indexer.GetSampleInfoRelative(SubsampleTime, Origin);
+		const FSampleInfo SampleFuture = Indexer.GetSampleInfoRelative(SubsampleTime + IndexingContext.SamplingContext->FiniteDelta, Origin);
+
+		// Mirror transforms if requested
+		const FTransform MirroredRootPast = Indexer.MirrorTransform(SamplePast.RootTransform);
+		const FTransform MirroredRootPresent = Indexer.MirrorTransform(SamplePresent.RootTransform);
+		const FTransform MirroredRootFuture = Indexer.MirrorTransform(SampleFuture.RootTransform);
+
+		// We can get a better finite difference if we ignore samples that have
+		// been clamped at either side of the clip. However, if the central sample 
+		// itself is clamped, or there are no samples that are clamped, we can just 
+		// use the central difference as normal.
+		FVector LinearVelocity;
+		if (SamplePast.bClamped && !SamplePresent.bClamped && !SampleFuture.bClamped)
+		{
+			LinearVelocity = (MirroredRootFuture.GetTranslation() - MirroredRootPresent.GetTranslation()) / IndexingContext.SamplingContext->FiniteDelta;
+		}
+		else if (SampleFuture.bClamped && !SamplePresent.bClamped && !SamplePast.bClamped)
+		{
+			LinearVelocity = (MirroredRootPresent.GetTranslation() - MirroredRootPast.GetTranslation()) / IndexingContext.SamplingContext->FiniteDelta;
+		}
+		else
+		{
+			LinearVelocity = (MirroredRootFuture.GetTranslation() - MirroredRootPast.GetTranslation()) / (IndexingContext.SamplingContext->FiniteDelta * 2.0f);
+		}
+
+		const FVector LinearVelocityDirection = LinearVelocity.GetClampedToMaxSize(1.0f);
+		const FVector FacingDirection = MirroredRootPresent.GetRotation().GetForwardVector();
+
+		if (EnumHasAnyFlags(Sample.Flags, EPoseSearchTrajectoryFlags::Position))
+		{
+			FFeatureVectorHelper::EncodeVector(FeatureVector, DataOffset, MirroredRootPresent.GetTranslation());
+		}
+		if (EnumHasAnyFlags(Sample.Flags, EPoseSearchTrajectoryFlags::PositionXY))
+		{
+			FFeatureVectorHelper::EncodeVector2D(FeatureVector, DataOffset, FVector2D(MirroredRootPresent.GetTranslation().X, MirroredRootPresent.GetTranslation().Y));
+		}
+
+		if (EnumHasAnyFlags(Sample.Flags, EPoseSearchTrajectoryFlags::Velocity))
+		{
+			FFeatureVectorHelper::EncodeVector(FeatureVector, DataOffset, LinearVelocity);
+		}
+		if (EnumHasAnyFlags(Sample.Flags, EPoseSearchTrajectoryFlags::VelocityXY))
+		{
+			FFeatureVectorHelper::EncodeVector2D(FeatureVector, DataOffset, FVector2D(LinearVelocity.X, LinearVelocity.Y));
+		}
+
+		if (EnumHasAnyFlags(Sample.Flags, EPoseSearchTrajectoryFlags::VelocityDirection))
+		{
+			FFeatureVectorHelper::EncodeVector(FeatureVector, DataOffset, LinearVelocityDirection);
+		}
+		if (EnumHasAnyFlags(Sample.Flags, EPoseSearchTrajectoryFlags::VelocityDirectionXY))
+		{
+			FFeatureVectorHelper::EncodeVector2D(FeatureVector, DataOffset, FVector2D(LinearVelocityDirection.X, LinearVelocityDirection.Y));
+		}
+
+		if (EnumHasAnyFlags(Sample.Flags, EPoseSearchTrajectoryFlags::FacingDirection))
+		{
+			FFeatureVectorHelper::EncodeVector(FeatureVector, DataOffset, FacingDirection);
+		}
+		if (EnumHasAnyFlags(Sample.Flags, EPoseSearchTrajectoryFlags::FacingDirectionXY))
+		{
+			FFeatureVectorHelper::EncodeVector2D(FeatureVector, DataOffset, FVector2D(FacingDirection.X, FacingDirection.Y).GetSafeNormal());
+		}
+	}
+	check(DataOffset == ChannelDataOffset + ChannelCardinality);
+}
+
+void UPoseSearchFeatureChannel_Trajectory::BuildQuery(UE::PoseSearch::FSearchContext& SearchContext, FPoseSearchFeatureVectorBuilder& InOutQuery) const
+{
+	using namespace UE::PoseSearch;
+
+	if (!SearchContext.Trajectory)
+	{
+		// @todo: do we want to reuse the SearchContext.CurrentResult data if valid?
+		return;
+	}
+
+	int32 NextIterStartIdx = 0;
+	int32 DataOffset = ChannelDataOffset;
+	float PreviousOffset = -FLT_MAX;
+	for (const FPoseSearchTrajectorySample& Sample : Samples)
+	{
+		// making sure Samples are sorted
+		check(Sample.Offset >= PreviousOffset);
+		const FTrajectorySample TrajectorySample = FTrajectorySampleRange::IterSampleTrajectory(SearchContext.Trajectory->Samples, Sample.Offset, NextIterStartIdx);
+
+		const FVector LinearVelocityDirection = TrajectorySample.LinearVelocity.GetClampedToMaxSize(1.0f);
+		const FVector FacingDirection = TrajectorySample.Transform.GetRotation().GetForwardVector();
+
+		if (EnumHasAnyFlags(Sample.Flags, EPoseSearchTrajectoryFlags::Position))
+		{
+			FFeatureVectorHelper::EncodeVector(InOutQuery.EditValues(), DataOffset, TrajectorySample.Transform.GetTranslation());
+		}
+		if (EnumHasAnyFlags(Sample.Flags, EPoseSearchTrajectoryFlags::PositionXY))
+		{
+			FFeatureVectorHelper::EncodeVector2D(InOutQuery.EditValues(), DataOffset, FVector2D(TrajectorySample.Transform.GetTranslation().X, TrajectorySample.Transform.GetTranslation().Y));
+		}
+
+		if (EnumHasAnyFlags(Sample.Flags, EPoseSearchTrajectoryFlags::Velocity))
+		{
+			FFeatureVectorHelper::EncodeVector(InOutQuery.EditValues(), DataOffset, TrajectorySample.LinearVelocity);
+		}
+		if (EnumHasAnyFlags(Sample.Flags, EPoseSearchTrajectoryFlags::VelocityXY))
+		{
+			FFeatureVectorHelper::EncodeVector2D(InOutQuery.EditValues(), DataOffset, FVector2D(TrajectorySample.LinearVelocity.X, TrajectorySample.LinearVelocity.Y));
+		}
+
+		if (EnumHasAnyFlags(Sample.Flags, EPoseSearchTrajectoryFlags::VelocityDirection))
+		{
+			FFeatureVectorHelper::EncodeVector(InOutQuery.EditValues(), DataOffset, LinearVelocityDirection);
+		}
+		if (EnumHasAnyFlags(Sample.Flags, EPoseSearchTrajectoryFlags::VelocityDirectionXY))
+		{
+			FFeatureVectorHelper::EncodeVector2D(InOutQuery.EditValues(), DataOffset, FVector2D(LinearVelocityDirection.X, LinearVelocityDirection.Y));
+		}
+
+		if (EnumHasAnyFlags(Sample.Flags, EPoseSearchTrajectoryFlags::FacingDirection))
+		{
+			FFeatureVectorHelper::EncodeVector(InOutQuery.EditValues(), DataOffset, FacingDirection);
+		}
+		if (EnumHasAnyFlags(Sample.Flags, EPoseSearchTrajectoryFlags::FacingDirectionXY))
+		{
+			FFeatureVectorHelper::EncodeVector2D(InOutQuery.EditValues(), DataOffset, FVector2D(FacingDirection.X, FacingDirection.Y).GetSafeNormal());
+		}
+	}
+	check(DataOffset == ChannelDataOffset + ChannelCardinality);
+}
+
+// lazy initialized helper to interpolate or extrapolate (linearly) UPoseSearchFeatureChannel_Trajectory trajectory positions from FPoseSearchTrajectorySample containing EPoseSearchTrajectoryFlags::Position for samples without it
+struct TrajectoryPositionReconstructor
+{
+	struct PositionAndOffsetSample
+	{
+		FVector Position;
+		float Offset;
+	};
+
+	TArray<PositionAndOffsetSample, TInlineAllocator<32>> PositionAndOffsetSamples;
+	bool bInitialized = false;
+
+	void Init(const UPoseSearchFeatureChannel_Trajectory& TrajectoryChannel, TConstArrayView<float> PoseVector, const FTransform& RootTransform)
+	{
+		PositionAndOffsetSamples.Reserve(TrajectoryChannel.Samples.Num() + 1);
+		bool bAddZeroOffsetSample = true;
+		int32 DataOffset = TrajectoryChannel.GetChannelDataOffset();
+		for (const FPoseSearchTrajectorySample& Sample : TrajectoryChannel.Samples)
+		{
+			if (EnumHasAnyFlags(Sample.Flags, EPoseSearchTrajectoryFlags::Position))
+			{
+				PositionAndOffsetSample PositionAndOffsetSample;
+				PositionAndOffsetSample.Position = UE::PoseSearch::FFeatureVectorHelper::DecodeVector(PoseVector, DataOffset);
+				PositionAndOffsetSample.Position = RootTransform.TransformPosition(PositionAndOffsetSample.Position);
+				PositionAndOffsetSample.Offset = Sample.Offset;
+				PositionAndOffsetSamples.Add(PositionAndOffsetSample);
+
+				if (FMath::IsNearlyZero(Sample.Offset))
+				{
+					bAddZeroOffsetSample = false;
+				}
+			}
+			if (EnumHasAnyFlags(Sample.Flags, EPoseSearchTrajectoryFlags::PositionXY))
+			{
+				const FVector2D Position2D = UE::PoseSearch::FFeatureVectorHelper::DecodeVector2D(PoseVector, DataOffset);
+
+				if (!EnumHasAnyFlags(Sample.Flags, EPoseSearchTrajectoryFlags::Position))
+				{
+					PositionAndOffsetSample PositionAndOffsetSample;
+					PositionAndOffsetSample.Position = FVector(Position2D.X, Position2D.Y, 0);
+					PositionAndOffsetSample.Position = RootTransform.TransformPosition(PositionAndOffsetSample.Position);
+					PositionAndOffsetSample.Offset = Sample.Offset;
+					PositionAndOffsetSamples.Add(PositionAndOffsetSample);
+
+					if (FMath::IsNearlyZero(Sample.Offset))
+					{
+						bAddZeroOffsetSample = false;
+					}
+				}
+			}
+
+			if (EnumHasAnyFlags(Sample.Flags, EPoseSearchTrajectoryFlags::Velocity))
+			{
+				DataOffset += UE::PoseSearch::FFeatureVectorHelper::EncodeVectorCardinality;
+			}
+			if (EnumHasAnyFlags(Sample.Flags, EPoseSearchTrajectoryFlags::VelocityXY))
+			{
+				DataOffset += UE::PoseSearch::FFeatureVectorHelper::EncodeVector2DCardinality;
+			}
+
+			if (EnumHasAnyFlags(Sample.Flags, EPoseSearchTrajectoryFlags::VelocityDirection))
+			{
+				DataOffset += UE::PoseSearch::FFeatureVectorHelper::EncodeVectorCardinality;
+			}
+			if (EnumHasAnyFlags(Sample.Flags, EPoseSearchTrajectoryFlags::VelocityDirectionXY))
+			{
+				DataOffset += UE::PoseSearch::FFeatureVectorHelper::EncodeVector2DCardinality;
+			}
+
+			if (EnumHasAnyFlags(Sample.Flags, EPoseSearchTrajectoryFlags::FacingDirection))
+			{
+				DataOffset += UE::PoseSearch::FFeatureVectorHelper::EncodeVectorCardinality;
+			}
+			if (EnumHasAnyFlags(Sample.Flags, EPoseSearchTrajectoryFlags::FacingDirectionXY))
+			{
+				DataOffset += UE::PoseSearch::FFeatureVectorHelper::EncodeVector2DCardinality;
+			}
+		}
+
+		if (bAddZeroOffsetSample)
+		{
+			PositionAndOffsetSample PositionAndOffsetSample;
+			PositionAndOffsetSample.Position = RootTransform.GetTranslation();
+			PositionAndOffsetSample.Offset = 0.f;
+			PositionAndOffsetSamples.Add(PositionAndOffsetSample);
+		}
+
+		PositionAndOffsetSamples.Sort([](const PositionAndOffsetSample& a, const PositionAndOffsetSample& b)
+		{
+			return a.Offset < b.Offset;
+		});
+
+		bInitialized = true;
+		check(DataOffset == TrajectoryChannel.GetChannelDataOffset() + TrajectoryChannel.GetChannelCardinality());
+	}
+
+	FVector GetReconstructedTrajectoryPos(const UPoseSearchFeatureChannel_Trajectory& TrajectoryChannel, TConstArrayView<float> PoseVector, const FTransform& RootTransform, float SampleOffset)
+	{
+		if (!bInitialized)
+		{
+			Init(TrajectoryChannel, PoseVector, RootTransform);
+		}
+
+		return GetReconstructedTrajectoryPos(SampleOffset);
+	}
+
+	FVector GetReconstructedTrajectoryPos(float SampleOffset) const
+	{
+		check(bInitialized);
+		check(PositionAndOffsetSamples.Num() > 0);
+		if (PositionAndOffsetSamples.Num() >= 2)
+		{
+			const int32 LowerBoundIdx = Algo::LowerBound(PositionAndOffsetSamples, SampleOffset, [](const PositionAndOffsetSample& PositionAndOffsetSample, float Value)
+				{
+					return Value > PositionAndOffsetSample.Offset;
+				});
+
+			const int32 PrevIdx = FMath::Clamp(LowerBoundIdx, 0, PositionAndOffsetSamples.Num() - 2);
+			const int32 NextIdx = PrevIdx + 1;
+
+			const float Denominator = PositionAndOffsetSamples[NextIdx].Offset - PositionAndOffsetSamples[PrevIdx].Offset;
+			if (FMath::IsNearlyZero(Denominator))
+			{
+				return PositionAndOffsetSamples[PrevIdx].Position;
+			}
+
+			const float Numerator = SampleOffset - PositionAndOffsetSamples[PrevIdx].Offset;
+			const float LerpValue = Numerator / Denominator;
+			return FMath::Lerp(PositionAndOffsetSamples[PrevIdx].Position, PositionAndOffsetSamples[NextIdx].Position, LerpValue);
+		}
+
+		return PositionAndOffsetSamples[0].Position;
+	}
+};
+
+void UPoseSearchFeatureChannel_Trajectory::DebugDraw(const UE::PoseSearch::FDebugDrawParams& DrawParams, TConstArrayView<float> PoseVector) const
+{
+#if ENABLE_DRAW_DEBUG
+	using namespace UE::PoseSearch;
+
+	const float LifeTime = DrawParams.DefaultLifeTime;
+	const uint8 DepthPriority = ESceneDepthPriorityGroup::SDPG_Foreground + 2;
+	const bool bPersistent = EnumHasAnyFlags(DrawParams.Flags, EDebugDrawFlags::Persistent);
+
+	const int32 NumSamples = Samples.Num();
+	if (NumSamples == 0)
+	{
+		return;
+	}
+
+	int32 DataOffset = ChannelDataOffset;
+	int32 SampleIdx = 0;
+	TrajectoryPositionReconstructor TrajectoryPositionReconstructor;
+	TArray<FVector, TInlineAllocator<32>> TrajSplinePos;
+	TArray<FColor, TInlineAllocator<32>> TrajSplineColor;
+	for (const FPoseSearchTrajectorySample& Sample : Samples)
+	{
+		bool bIsTrajectoryPosValid = false;
+		FVector TrajectoryPos = FVector::Zero();
+		
+		if (EnumHasAnyFlags(Sample.Flags, EPoseSearchTrajectoryFlags::Position))
+		{
+			TrajectoryPos = FFeatureVectorHelper::DecodeVector(PoseVector, DataOffset);
+			TrajectoryPos = DrawParams.RootTransform.TransformPosition(TrajectoryPos);
+			
+			bIsTrajectoryPosValid = true;
+
+			// validating TrajectoryPositionReconstructor
+			check((TrajectoryPositionReconstructor.GetReconstructedTrajectoryPos(*this, PoseVector, DrawParams.RootTransform, Sample.Offset) - TrajectoryPos).IsNearlyZero());
+
+			const FColor Color = DrawParams.GetColor(Sample.ColorPresetIndex);
+
+			if (EnumHasAnyFlags(DrawParams.Flags, EDebugDrawFlags::DrawFast | EDebugDrawFlags::DrawSearchIndex))
+			{
+				DrawDebugPoint(DrawParams.World, TrajectoryPos, DrawParams.PointSize, Color, bPersistent, LifeTime, DepthPriority);
+			}
+			else
+			{
+				DrawDebugSphere(DrawParams.World, TrajectoryPos, 2.f, 8, Color, bPersistent, LifeTime, DepthPriority);
+			}
+
+			TrajSplinePos.Add(TrajectoryPos);
+			TrajSplineColor.Add(Color);
+		}
+		if (EnumHasAnyFlags(Sample.Flags, EPoseSearchTrajectoryFlags::PositionXY))
+		{
+			FVector2D TrajectoryPos2D = FFeatureVectorHelper::DecodeVector2D(PoseVector, DataOffset);
+			if (!bIsTrajectoryPosValid)
+			{
+				TrajectoryPos = FVector(TrajectoryPos2D.X, TrajectoryPos2D.Y, 0);
+				TrajectoryPos = DrawParams.RootTransform.TransformPosition(TrajectoryPos);
+				bIsTrajectoryPosValid = true;
+
+				const FColor Color = DrawParams.GetColor(Sample.ColorPresetIndex);
+
+				if (EnumHasAnyFlags(DrawParams.Flags, EDebugDrawFlags::DrawFast | EDebugDrawFlags::DrawSearchIndex))
+				{
+					DrawDebugPoint(DrawParams.World, TrajectoryPos, DrawParams.PointSize, Color, bPersistent, LifeTime, DepthPriority);
+				}
+				else
+				{
+					DrawDebugSphere(DrawParams.World, TrajectoryPos, 2.f, 8, Color, bPersistent, LifeTime, DepthPriority);
+				}
+
+				TrajSplinePos.Add(TrajectoryPos);
+				TrajSplineColor.Add(Color);
+			}
+		}
+		
+		if (!bIsTrajectoryPosValid)
+		{
+			TrajectoryPos = TrajectoryPositionReconstructor.GetReconstructedTrajectoryPos(*this, PoseVector, DrawParams.RootTransform, Sample.Offset);
+			
+			TrajSplinePos.Add(TrajectoryPos);
+			const FColor Color = TrajSplineColor.IsEmpty() ? FColor::Black : TrajSplineColor.Last();
+			TrajSplineColor.Add(Color);
+		}
+
+		if (EnumHasAnyFlags(Sample.Flags, EPoseSearchTrajectoryFlags::Velocity))
+		{
+			FVector TrajectoryVel = FFeatureVectorHelper::DecodeVector(PoseVector, DataOffset);
+
+			const FColor Color = DrawParams.GetColor(Sample.ColorPresetIndex);
+
+			TrajectoryVel *= 0.08f;
+			TrajectoryVel = DrawParams.RootTransform.TransformVector(TrajectoryVel);
+			const FVector TrajectoryVelDirection = TrajectoryVel.GetSafeNormal();
+
+			if (EnumHasAnyFlags(DrawParams.Flags, EDebugDrawFlags::DrawSearchIndex))
+			{
+				DrawDebugLine(DrawParams.World, TrajectoryPos, TrajectoryPos + TrajectoryVel, Color, bPersistent, LifeTime, DepthPriority);
+			}
+			else
+			{
+				const float AdjustedThickness = EnumHasAnyFlags(DrawParams.Flags, EDebugDrawFlags::DrawFast) ? 0.0f : 1.f;
+
+				DrawDebugLine(
+					DrawParams.World,
+					TrajectoryPos + TrajectoryVelDirection * 2.f,
+					TrajectoryPos + TrajectoryVel,
+					Color,
+					bPersistent,
+					LifeTime,
+					DepthPriority,
+					AdjustedThickness
+				);
+			}
+		}
+		if (EnumHasAnyFlags(Sample.Flags, EPoseSearchTrajectoryFlags::VelocityXY))
+		{
+			const FVector2D TrajectoryVel2D = FFeatureVectorHelper::DecodeVector2D(PoseVector, DataOffset);
+			FVector TrajectoryVel(TrajectoryVel2D.X, TrajectoryVel2D.Y, 0.f);
+
+			const FColor Color = DrawParams.GetColor(Sample.ColorPresetIndex);
+
+			TrajectoryVel *= 0.08f;
+			TrajectoryVel = DrawParams.RootTransform.TransformVector(TrajectoryVel);
+			const FVector TrajectoryVelDirection = TrajectoryVel.GetSafeNormal();
+
+			if (EnumHasAnyFlags(DrawParams.Flags, EDebugDrawFlags::DrawSearchIndex))
+			{
+				DrawDebugLine(DrawParams.World, TrajectoryPos, TrajectoryPos + TrajectoryVel, Color, bPersistent, LifeTime, DepthPriority);
+			}
+			else
+			{
+				const float AdjustedThickness = EnumHasAnyFlags(DrawParams.Flags, EDebugDrawFlags::DrawFast) ? 0.0f : 1.f;
+
+				DrawDebugLine(
+					DrawParams.World,
+					TrajectoryPos + TrajectoryVelDirection * 2.f,
+					TrajectoryPos + TrajectoryVel,
+					Color,
+					bPersistent,
+					LifeTime,
+					DepthPriority,
+					AdjustedThickness
+				);
+			}
+		}
+
+		if (EnumHasAnyFlags(Sample.Flags, EPoseSearchTrajectoryFlags::VelocityDirection))
+		{
+			FVector TrajectoryVelDirection = FFeatureVectorHelper::DecodeVector(PoseVector, DataOffset);
+
+			const FColor Color = DrawParams.GetColor(Sample.ColorPresetIndex);
+
+			TrajectoryVelDirection = DrawParams.RootTransform.TransformVector(TrajectoryVelDirection);
+
+			if (EnumHasAnyFlags(DrawParams.Flags, EDebugDrawFlags::DrawSearchIndex))
+			{
+				DrawDebugLine(DrawParams.World, TrajectoryPos, TrajectoryPos + TrajectoryVelDirection, Color, bPersistent, LifeTime, DepthPriority);
+			}
+			else
+			{
+				const float AdjustedThickness = EnumHasAnyFlags(DrawParams.Flags, EDebugDrawFlags::DrawFast) ? 0.0f : 1.f;
+
+				DrawDebugLine(
+					DrawParams.World,
+					TrajectoryPos + TrajectoryVelDirection * 2.f,
+					TrajectoryPos + TrajectoryVelDirection * 2.f * 10.0f,
+					Color,
+					bPersistent,
+					LifeTime,
+					DepthPriority,
+					AdjustedThickness
+				);
+			}
+		}
+
+		if (EnumHasAnyFlags(Sample.Flags, EPoseSearchTrajectoryFlags::VelocityDirectionXY))
+		{
+			const FVector2D TrajectoryVelDirection2D = FFeatureVectorHelper::DecodeVector2D(PoseVector, DataOffset);
+			FVector TrajectoryVelDirection(TrajectoryVelDirection2D.X, TrajectoryVelDirection2D.Y, 0.f);
+
+			const FColor Color = DrawParams.GetColor(Sample.ColorPresetIndex);
+
+			TrajectoryVelDirection = DrawParams.RootTransform.TransformVector(TrajectoryVelDirection);
+
+			if (EnumHasAnyFlags(DrawParams.Flags, EDebugDrawFlags::DrawSearchIndex))
+			{
+				DrawDebugLine(DrawParams.World, TrajectoryPos, TrajectoryPos + TrajectoryVelDirection, Color, bPersistent, LifeTime, DepthPriority);
+			}
+			else
+			{
+				const float AdjustedThickness = EnumHasAnyFlags(DrawParams.Flags, EDebugDrawFlags::DrawFast) ? 0.0f : 1.f;
+
+				DrawDebugLine(
+					DrawParams.World,
+					TrajectoryPos + TrajectoryVelDirection * 2.f,
+					TrajectoryPos + TrajectoryVelDirection * 2.f * 10.0f,
+					Color,
+					bPersistent,
+					LifeTime,
+					DepthPriority,
+					AdjustedThickness
+				);
+			}
+		}
+
+		if (EnumHasAnyFlags(Sample.Flags, EPoseSearchTrajectoryFlags::FacingDirection))
+		{
+			FVector TrajectoryForward = FFeatureVectorHelper::DecodeVector(PoseVector, DataOffset);
+
+			const FColor Color = DrawParams.GetColor(Sample.ColorPresetIndex);
+
+			TrajectoryForward = DrawParams.RootTransform.TransformVector(TrajectoryForward);
+
+			if (EnumHasAnyFlags(DrawParams.Flags, EDebugDrawFlags::DrawSearchIndex))
+			{
+				DrawDebugLine(DrawParams.World, TrajectoryPos, TrajectoryPos + TrajectoryForward, Color, bPersistent, LifeTime, DepthPriority);
+			}
+			else
+			{
+				const float AdjustedThickness = EnumHasAnyFlags(DrawParams.Flags, EDebugDrawFlags::DrawFast) ? 0.0f : 1.f;
+
+				DrawDebugLine(
+					DrawParams.World,
+					TrajectoryPos + TrajectoryForward * 2.f,
+					TrajectoryPos + TrajectoryForward * 2.f * 10.0f,
+					Color,
+					bPersistent,
+					LifeTime,
+					DepthPriority,
+					AdjustedThickness
+				);
+			}
+		}
+
+		if (EnumHasAnyFlags(Sample.Flags, EPoseSearchTrajectoryFlags::FacingDirectionXY))
+		{
+			const FVector2D TrajectoryForward2D = FFeatureVectorHelper::DecodeVector2D(PoseVector, DataOffset);
+			FVector TrajectoryForward(TrajectoryForward2D.X, TrajectoryForward2D.Y, 0.f);
+
+			const FColor Color = DrawParams.GetColor(Sample.ColorPresetIndex);
+
+			TrajectoryForward = DrawParams.RootTransform.TransformVector(TrajectoryForward);
+
+			if (EnumHasAnyFlags(DrawParams.Flags, EDebugDrawFlags::DrawSearchIndex))
+			{
+				DrawDebugLine(DrawParams.World, TrajectoryPos, TrajectoryPos + TrajectoryForward, Color, bPersistent, LifeTime, DepthPriority);
+			}
+			else
+			{
+				const float AdjustedThickness = EnumHasAnyFlags(DrawParams.Flags, EDebugDrawFlags::DrawFast) ? 0.0f : 1.f;
+
+				DrawDebugLine(
+					DrawParams.World,
+					TrajectoryPos + TrajectoryForward * 2.f,
+					TrajectoryPos + TrajectoryForward * 2.f * 10.0f,
+					Color,
+					bPersistent,
+					LifeTime,
+					DepthPriority,
+					AdjustedThickness
+				);
+			}
+		}
+
+		if (EnumHasAnyFlags(DrawParams.Flags, EDebugDrawFlags::DrawSampleLabels))
+		{
+			const FColor Color = DrawParams.GetColor(Sample.ColorPresetIndex);
+
+			const FString SampleLabel = FString::Format(TEXT("{0}"), { SampleIdx });
+
+			static const FVector DrawDebugSampleLabelOffset = FVector(0.0f, 0.0f, 5.0f);
+			DrawDebugString(
+				DrawParams.World,
+				TrajectoryPos + DrawDebugSampleLabelOffset,
+				SampleLabel,
+				nullptr,
+				Color,
+				LifeTime,
+				false,
+				1.5f);
+		}
+
+		++SampleIdx;
+	}
+
+	DrawCentripetalCatmullRomSpline(DrawParams.World, TrajSplinePos, TrajSplineColor, 0.5f, 8.f, bPersistent, LifeTime, DepthPriority, 0.f);
+
+	check(DataOffset == ChannelDataOffset + ChannelCardinality);
+#endif // ENABLE_DRAW_DEBUG
+}
+
+#if WITH_EDITOR
+void UPoseSearchFeatureChannel_Trajectory::PopulateChannelLayoutSet(UE::PoseSearch::FFeatureChannelLayoutSet& FeatureChannelLayoutSet) const
+{
+	using namespace UE::PoseSearch;
+
+	int32 DataOffset = ChannelDataOffset;
+
+	auto Add = [&FeatureChannelLayoutSet, &DataOffset](EPoseSearchTrajectoryFlags SampleFlag, float Offset, const char* Label, int32 Cardinality)
+	{
+		FString SkeletonName = FeatureChannelLayoutSet.CurrentSchema->Skeleton->GetName();
+
+		UE::PoseSearch::FKeyBuilder KeyBuilder;
+		KeyBuilder << SkeletonName << SampleFlag << Offset;
+		FeatureChannelLayoutSet.Add(FString::Format(TEXT("Traj {0} {1}"), { Label, Offset }), KeyBuilder.Finalize(), DataOffset, Cardinality);
+
+		DataOffset += Cardinality;
+	};
+
+	for (const FPoseSearchTrajectorySample& Sample : Samples)
+	{
+		if (EnumHasAnyFlags(Sample.Flags, EPoseSearchTrajectoryFlags::Position))
+		{
+			Add(EPoseSearchTrajectoryFlags::Position, Sample.Offset, "Pos", FFeatureVectorHelper::EncodeVectorCardinality);
+		}
+		if (EnumHasAnyFlags(Sample.Flags, EPoseSearchTrajectoryFlags::PositionXY))
+		{
+			Add(EPoseSearchTrajectoryFlags::PositionXY, Sample.Offset, "PosXY", FFeatureVectorHelper::EncodeVector2DCardinality);
+		}
+
+		if (EnumHasAnyFlags(Sample.Flags, EPoseSearchTrajectoryFlags::Velocity))
+		{
+			Add(EPoseSearchTrajectoryFlags::Velocity, Sample.Offset, "Vel", FFeatureVectorHelper::EncodeVectorCardinality);
+		}
+		if (EnumHasAnyFlags(Sample.Flags, EPoseSearchTrajectoryFlags::VelocityXY))
+		{
+			Add(EPoseSearchTrajectoryFlags::VelocityXY, Sample.Offset, "VelXY", FFeatureVectorHelper::EncodeVector2DCardinality);
+		}
+
+		if (EnumHasAnyFlags(Sample.Flags, EPoseSearchTrajectoryFlags::VelocityDirection))
+		{
+			Add(EPoseSearchTrajectoryFlags::VelocityDirection, Sample.Offset, "VelDir", FFeatureVectorHelper::EncodeVectorCardinality);
+		}
+		if (EnumHasAnyFlags(Sample.Flags, EPoseSearchTrajectoryFlags::VelocityDirectionXY))
+		{
+			Add(EPoseSearchTrajectoryFlags::VelocityDirectionXY, Sample.Offset, "VelDirXY", FFeatureVectorHelper::EncodeVector2DCardinality);
+		}
+
+		if (EnumHasAnyFlags(Sample.Flags, EPoseSearchTrajectoryFlags::FacingDirection))
+		{
+			Add(EPoseSearchTrajectoryFlags::FacingDirection, Sample.Offset, "Fac", FFeatureVectorHelper::EncodeVectorCardinality);
+		}
+		if (EnumHasAnyFlags(Sample.Flags, EPoseSearchTrajectoryFlags::FacingDirectionXY))
+		{
+			Add(EPoseSearchTrajectoryFlags::FacingDirectionXY, Sample.Offset, "FacXY", FFeatureVectorHelper::EncodeVector2DCardinality);
+		}
+	}
+
+	check(DataOffset == ChannelDataOffset + ChannelCardinality);
+}
+
+void UPoseSearchFeatureChannel_Trajectory::ComputeCostBreakdowns(UE::PoseSearch::ICostBreakDownData& CostBreakDownData, const UPoseSearchSchema* Schema) const
+{
+	using namespace UE::PoseSearch;
+
+	CostBreakDownData.AddEntireBreakDownSection(LOCTEXT("ColumnLabelTrajChannelTotal", "Traj Total"), Schema, ChannelDataOffset, ChannelCardinality);
+
+	if (CostBreakDownData.IsVerbose())
+	{
+		int32 DataOffset = ChannelDataOffset;
+
+		for (const FPoseSearchTrajectorySample& Sample : Samples)
+		{
+			if (EnumHasAnyFlags(Sample.Flags, EPoseSearchTrajectoryFlags::Position))
+			{
+				CostBreakDownData.AddEntireBreakDownSection(FText::Format(LOCTEXT("ColumnLabelTrajChannelPosition", "Traj Pos {0}"), Sample.Offset), Schema, DataOffset, FFeatureVectorHelper::EncodeVectorCardinality);
+				DataOffset += FFeatureVectorHelper::EncodeVectorCardinality;
+			}
+			if (EnumHasAnyFlags(Sample.Flags, EPoseSearchTrajectoryFlags::PositionXY))
+			{
+				CostBreakDownData.AddEntireBreakDownSection(FText::Format(LOCTEXT("ColumnLabelTrajChannelPositionXY", "Traj PosXY {0}"), Sample.Offset), Schema, DataOffset, FFeatureVectorHelper::EncodeVector2DCardinality);
+				DataOffset += FFeatureVectorHelper::EncodeVector2DCardinality;
+			}
+
+			if (EnumHasAnyFlags(Sample.Flags, EPoseSearchTrajectoryFlags::Velocity))
+			{
+				CostBreakDownData.AddEntireBreakDownSection(FText::Format(LOCTEXT("ColumnLabelTrajChannelVelocity", "Traj Vel {0}"), Sample.Offset), Schema, DataOffset, FFeatureVectorHelper::EncodeVectorCardinality);
+				DataOffset += FFeatureVectorHelper::EncodeVectorCardinality;
+			}
+			if (EnumHasAnyFlags(Sample.Flags, EPoseSearchTrajectoryFlags::VelocityXY))
+			{
+				CostBreakDownData.AddEntireBreakDownSection(FText::Format(LOCTEXT("ColumnLabelTrajChannelVelocityXY", "Traj VelXY {0}"), Sample.Offset), Schema, DataOffset, FFeatureVectorHelper::EncodeVector2DCardinality);
+				DataOffset += FFeatureVectorHelper::EncodeVector2DCardinality;
+			}
+
+			if (EnumHasAnyFlags(Sample.Flags, EPoseSearchTrajectoryFlags::VelocityDirection))
+			{
+				CostBreakDownData.AddEntireBreakDownSection(FText::Format(LOCTEXT("ColumnLabelTrajChannelVelocityDirection", "Traj VelDir {0}"), Sample.Offset), Schema, DataOffset, FFeatureVectorHelper::EncodeVectorCardinality);
+				DataOffset += FFeatureVectorHelper::EncodeVectorCardinality;
+			}
+			if (EnumHasAnyFlags(Sample.Flags, EPoseSearchTrajectoryFlags::VelocityDirectionXY))
+			{
+				CostBreakDownData.AddEntireBreakDownSection(FText::Format(LOCTEXT("ColumnLabelTrajChannelVelocityDirectionXY", "Traj VelDirXY {0}"), Sample.Offset), Schema, DataOffset, FFeatureVectorHelper::EncodeVector2DCardinality);
+				DataOffset += FFeatureVectorHelper::EncodeVector2DCardinality;
+			}
+
+			if (EnumHasAnyFlags(Sample.Flags, EPoseSearchTrajectoryFlags::FacingDirection))
+			{
+				CostBreakDownData.AddEntireBreakDownSection(FText::Format(LOCTEXT("ColumnLabelTrajChannelFacingDirection", "Traj Fac {0}"), Sample.Offset), Schema, DataOffset, FFeatureVectorHelper::EncodeVectorCardinality);
+				DataOffset += FFeatureVectorHelper::EncodeVectorCardinality;
+			}
+			if (EnumHasAnyFlags(Sample.Flags, EPoseSearchTrajectoryFlags::FacingDirectionXY))
+			{
+				CostBreakDownData.AddEntireBreakDownSection(FText::Format(LOCTEXT("ColumnLabelTrajChannelFacingDirectionXY", "Traj FacXY {0}"), Sample.Offset), Schema, DataOffset, FFeatureVectorHelper::EncodeVector2DCardinality);
+				DataOffset += FFeatureVectorHelper::EncodeVector2DCardinality;
+			}
+		}
+
+		check(DataOffset == ChannelDataOffset + ChannelCardinality);
+	}
+}
+
+#endif // WITH_EDITOR
+
+bool UPoseSearchFeatureChannel_Trajectory::GetEstimatedSpeedRatio(TConstArrayView<float> QueryVector, TConstArrayView<float> PoseVector, float& EstimatedSpeedRatio) const
+{
+	float EstimatedQuerySpeed = 0.f;
+	float EstimatedPoseSpeed = 0.f;
+
+	int32 QueryDataOffset = ChannelDataOffset;
+	int32 PoseDataOffset = ChannelDataOffset;
+
+	bool bValidEstimate = false;
+	for (const FPoseSearchTrajectorySample& Sample : Samples)
+	{
+		// @todo: decode positions and use them to estimate velocities in case UPoseSearchFeatureChannel_Trajectory doens't contain EPoseSearchTrajectoryFlags::Velocity or EPoseSearchTrajectoryFlags::VelocityXY samples
+		if (EnumHasAnyFlags(Sample.Flags, EPoseSearchTrajectoryFlags::Position))
+		{
+			QueryDataOffset += UE::PoseSearch::FFeatureVectorHelper::EncodeVectorCardinality;
+			PoseDataOffset += UE::PoseSearch::FFeatureVectorHelper::EncodeVectorCardinality;
+		}
+		if (EnumHasAnyFlags(Sample.Flags, EPoseSearchTrajectoryFlags::PositionXY))
+		{
+			QueryDataOffset += UE::PoseSearch::FFeatureVectorHelper::EncodeVector2DCardinality;
+			PoseDataOffset += UE::PoseSearch::FFeatureVectorHelper::EncodeVector2DCardinality;
+		}
+
+		if (EnumHasAnyFlags(Sample.Flags, EPoseSearchTrajectoryFlags::Velocity))
+		{
+			const FVector QueryVelocity = UE::PoseSearch::FFeatureVectorHelper::DecodeVector(QueryVector, QueryDataOffset);
+			const FVector PoseVelocity = UE::PoseSearch::FFeatureVectorHelper::DecodeVector(PoseVector, PoseDataOffset);
+			EstimatedQuerySpeed += QueryVelocity.Length();
+			EstimatedPoseSpeed += PoseVelocity.Length();
+			bValidEstimate = true;
+		}
+		if (EnumHasAnyFlags(Sample.Flags, EPoseSearchTrajectoryFlags::VelocityXY))
+		{
+			const FVector2D QueryVelocity = UE::PoseSearch::FFeatureVectorHelper::DecodeVector2D(QueryVector, QueryDataOffset);
+			const FVector2D PoseVelocity = UE::PoseSearch::FFeatureVectorHelper::DecodeVector2D(PoseVector, PoseDataOffset);
+			EstimatedQuerySpeed += QueryVelocity.Length();
+			EstimatedPoseSpeed += PoseVelocity.Length();
+			bValidEstimate = true;
+		}
+
+		if (EnumHasAnyFlags(Sample.Flags, EPoseSearchTrajectoryFlags::VelocityDirection))
+		{
+			QueryDataOffset += UE::PoseSearch::FFeatureVectorHelper::EncodeVectorCardinality;
+			PoseDataOffset += UE::PoseSearch::FFeatureVectorHelper::EncodeVectorCardinality;
+		}
+		if (EnumHasAnyFlags(Sample.Flags, EPoseSearchTrajectoryFlags::VelocityDirectionXY))
+		{
+			QueryDataOffset += UE::PoseSearch::FFeatureVectorHelper::EncodeVector2DCardinality;
+			PoseDataOffset += UE::PoseSearch::FFeatureVectorHelper::EncodeVector2DCardinality;
+		}
+		if (EnumHasAnyFlags(Sample.Flags, EPoseSearchTrajectoryFlags::FacingDirection))
+		{
+			QueryDataOffset += UE::PoseSearch::FFeatureVectorHelper::EncodeVectorCardinality;
+			PoseDataOffset += UE::PoseSearch::FFeatureVectorHelper::EncodeVectorCardinality;
+		}
+		if (EnumHasAnyFlags(Sample.Flags, EPoseSearchTrajectoryFlags::FacingDirectionXY))
+		{
+			QueryDataOffset += UE::PoseSearch::FFeatureVectorHelper::EncodeVector2DCardinality;
+			PoseDataOffset += UE::PoseSearch::FFeatureVectorHelper::EncodeVector2DCardinality;
+		}
+	}
+
+	check(QueryDataOffset == ChannelDataOffset + ChannelCardinality);
+	check(PoseDataOffset == ChannelDataOffset + ChannelCardinality);
+
+	if (bValidEstimate && EstimatedPoseSpeed > UE_KINDA_SMALL_NUMBER)
+	{
+		EstimatedSpeedRatio = EstimatedQuerySpeed / EstimatedPoseSpeed;
+	}
+	else
+	{
+		EstimatedSpeedRatio = 1.f;
+	}
+
+	return bValidEstimate;
+}
+
+#undef LOCTEXT_NAMESPACE
\ No newline at end of file
diff --git a/Engine/Plugins/Experimental/Animation/PoseSearch/Source/Runtime/Private/PoseSearchFeatureChannel_Trajectory.h b/Engine/Plugins/Experimental/Animation/PoseSearch/Source/Runtime/Private/PoseSearchFeatureChannel_Trajectory.h
new file mode 100644
--- /dev/null
+++ b/Engine/Plugins/Experimental/Animation/PoseSearch/Source/Runtime/Private/PoseSearchFeatureChannel_Trajectory.h
@@ -0,0 +1,76 @@
+// Copyright Epic Games, Inc. All Rights Reserved.
+
+#pragma once
+
+#include "PoseSearch/PoseSearchFeatureChannel.h"
+#include "BoneContainer.h"
+#include "UObject/ObjectSaveContext.h"
+#include "PoseSearchFeatureChannel_Trajectory.generated.h"
+
+UENUM(meta = (Bitflags, UseEnumValuesAsMaskValuesInEditor = "true"))
+enum class EPoseSearchTrajectoryFlags : uint32
+{
+	Velocity = 1 << 0,
+	Position = 1 << 1,
+	VelocityDirection = 1 << 2,
+	FacingDirection = 1 << 3,
+	VelocityXY = 1 << 4,
+	PositionXY = 1 << 5,
+	VelocityDirectionXY = 1 << 6,
+	FacingDirectionXY = 1 << 7,
+};
+ENUM_CLASS_FLAGS(EPoseSearchTrajectoryFlags);
+constexpr bool EnumHasAnyFlags(int32 Flags, EPoseSearchTrajectoryFlags Contains) { return (Flags & int32(Contains)) != 0; }
+inline int32& operator|=(int32& Lhs, EPoseSearchTrajectoryFlags Rhs) { return Lhs |= int32(Rhs); }
+
+USTRUCT()
+struct POSESEARCH_API FPoseSearchTrajectorySample
+{
+	GENERATED_BODY()
+
+	UPROPERTY(EditAnywhere, Category = Config)
+	float Offset = 0.f; // offset in time or distance depending on UPoseSearchFeatureChannel_Trajectory.Domain
+
+	UPROPERTY(EditAnywhere, meta = (Bitmask, BitmaskEnum = "/Script/PoseSearch.EPoseSearchTrajectoryFlags"), Category = Config)
+	int32 Flags = int32(EPoseSearchTrajectoryFlags::Position);
+
+	UPROPERTY(EditAnywhere, Category = Config)
+	float Weight = 1.f;
+
+	UPROPERTY(EditAnywhere, Category = Config, meta = (ExcludeFromHash))
+	int32 ColorPresetIndex = 0;
+};
+
+UCLASS(BlueprintType, EditInlineNew, meta = (DisplayName = "Trajectory Channel"), CollapseCategories)
+class POSESEARCH_API UPoseSearchFeatureChannel_Trajectory : public UPoseSearchFeatureChannel
+{
+	GENERATED_BODY()
+
+public:
+	UPROPERTY(EditAnywhere, Category = "Settings")
+	float Weight = 1.f;
+
+	UPROPERTY(EditAnywhere, Category = "Settings")
+	TArray<FPoseSearchTrajectorySample> Samples;
+
+	// UObject interface
+	virtual void PreSave(FObjectPreSaveContext ObjectSaveContext) override;
+
+	// UPoseSearchFeatureChannel interface
+	virtual void InitializeSchema(UPoseSearchSchema* Schema) override;
+	virtual void FillWeights(TArray<float>& Weights) const override;
+	virtual void IndexAsset(UE::PoseSearch::IAssetIndexer& Indexer, UE::PoseSearch::FAssetIndexingOutput& IndexingOutput) const override;
+	virtual void BuildQuery(UE::PoseSearch::FSearchContext& SearchContext, FPoseSearchFeatureVectorBuilder& InOutQuery) const override;
+	virtual void DebugDraw(const UE::PoseSearch::FDebugDrawParams& DrawParams, TConstArrayView<float> PoseVector) const override;
+
+#if WITH_EDITOR
+	virtual void PopulateChannelLayoutSet(UE::PoseSearch::FFeatureChannelLayoutSet& FeatureChannelLayoutSet) const override;
+	virtual void ComputeCostBreakdowns(UE::PoseSearch::ICostBreakDownData& CostBreakDownData, const UPoseSearchSchema* Schema) const override;
+#endif
+
+	bool GetEstimatedSpeedRatio(TConstArrayView<float> QueryVector, TConstArrayView<float> PoseVector, float& EstimatedSpeedRatio) const;
+
+protected:
+	void IndexAssetPrivate(const UE::PoseSearch::IAssetIndexer& Indexer, int32 SampleIdx, TArrayView<float> FeatureVector) const;
+};
+
diff --git a/Engine/Plugins/Experimental/Animation/PoseSearch/Source/Runtime/Private/PoseSearchFeatureChannels.cpp b/Engine/Plugins/Experimental/Animation/PoseSearch/Source/Runtime/Private/PoseSearchFeatureChannels.cpp
deleted file mode 100644
diff --git a/Engine/Plugins/Experimental/Animation/PoseSearch/Source/Runtime/Private/PoseSearchLibrary.cpp b/Engine/Plugins/Experimental/Animation/PoseSearch/Source/Runtime/Private/PoseSearchLibrary.cpp
--- a/Engine/Plugins/Experimental/Animation/PoseSearch/Source/Runtime/Private/PoseSearchLibrary.cpp
+++ b/Engine/Plugins/Experimental/Animation/PoseSearch/Source/Runtime/Private/PoseSearchLibrary.cpp
@@ -9,8 +9,9 @@
 #include "Animation/AnimRootMotionProvider.h"
 #include "Animation/BlendSpace.h"
 #include "PoseSearch/AnimNode_MotionMatching.h"
+#include "PoseSearch/PoseSearchSchema.h"
+#include "PoseSearchFeatureChannel_Trajectory.h"
 #include "Trace/PoseSearchTraceLogger.h"
-#include "PoseSearch/PoseSearchFeatureChannels.h"
 
 #include UE_INLINE_GENERATED_CPP_BY_NAME(PoseSearchLibrary)
 
diff --git a/Engine/Plugins/Experimental/Animation/PoseSearch/Source/Runtime/Private/PoseSearchSchema.cpp b/Engine/Plugins/Experimental/Animation/PoseSearch/Source/Runtime/Private/PoseSearchSchema.cpp
new file mode 100644
--- /dev/null
+++ b/Engine/Plugins/Experimental/Animation/PoseSearch/Source/Runtime/Private/PoseSearchSchema.cpp
@@ -0,0 +1,68 @@
+// Copyright Epic Games, Inc. All Rights Reserved.
+
+#include "PoseSearch/PoseSearchSchema.h"
+#include "AnimationRuntime.h"
+#include "PoseSearch/PoseSearch.h"
+#include "PoseSearch/PoseSearchFeatureChannel.h"
+
+bool UPoseSearchSchema::IsValid() const
+{
+	bool bValid = Skeleton != nullptr;
+
+	for (const FBoneReference& BoneRef : BoneReferences)
+	{
+		bValid &= BoneRef.HasValidSetup();
+	}
+
+	for (const TObjectPtr<UPoseSearchFeatureChannel>& Channel: Channels)
+	{
+		bValid &= Channel != nullptr;
+	}
+
+	bValid &= (BoneReferences.Num() == BoneIndices.Num());
+
+	return bValid;
+}
+
+void UPoseSearchSchema::ResolveBoneReferences()
+{
+	// Initialize references to obtain bone indices
+	for (FBoneReference& BoneRef : BoneReferences)
+	{
+		BoneRef.Initialize(Skeleton);
+	}
+
+	// Fill out bone index array
+	BoneIndices.SetNum(BoneReferences.Num());
+	for (int32 BoneRefIdx = 0; BoneRefIdx != BoneReferences.Num(); ++BoneRefIdx)
+	{
+		BoneIndices[BoneRefIdx] = BoneReferences[BoneRefIdx].BoneIndex;
+	}
+
+	// Build separate index array with parent indices guaranteed to be present. Sort for EnsureParentsPresent.
+	BoneIndicesWithParents = BoneIndices;
+	BoneIndicesWithParents.Sort();
+
+	if (Skeleton)
+	{
+		FAnimationRuntime::EnsureParentsPresent(BoneIndicesWithParents, Skeleton->GetReferenceSkeleton());
+	}
+
+	// BoneIndicesWithParents should at least contain the root to support mirroring root motion
+	if (BoneIndicesWithParents.Num() == 0)
+	{
+		BoneIndicesWithParents.Add(0);
+	}
+}
+
+void UPoseSearchSchema::BuildQuery(UE::PoseSearch::FSearchContext& SearchContext, FPoseSearchFeatureVectorBuilder& InOutQuery) const
+{
+	QUICK_SCOPE_CYCLE_COUNTER(STAT_PoseSearch_BuildQuery);
+
+	InOutQuery.Init(this);
+
+	for (const TObjectPtr<UPoseSearchFeatureChannel>& Channel : Channels)
+	{
+		Channel->BuildQuery(SearchContext, InOutQuery);
+	}
+}
diff --git a/Engine/Plugins/Experimental/Animation/PoseSearch/Source/Runtime/Public/PoseSearch/PoseSearch.h b/Engine/Plugins/Experimental/Animation/PoseSearch/Source/Runtime/Public/PoseSearch/PoseSearch.h
--- a/Engine/Plugins/Experimental/Animation/PoseSearch/Source/Runtime/Public/PoseSearch/PoseSearch.h
+++ b/Engine/Plugins/Experimental/Animation/PoseSearch/Source/Runtime/Public/PoseSearch/PoseSearch.h
@@ -2,13 +2,16 @@
 
 #pragma once
 
+//#include "PoseSearchFeatureChannel.h"
+//#include "PoseSearchSchema.h"
+
+
 #include "Async/ParallelFor.h"
 #include "BonePose.h"
 #include "CoreMinimal.h"
 #include "Containers/RingBuffer.h"
 #include "Engine/DataAsset.h"
 #include "InstancedStruct.h"
-#include "IO/IoHash.h"
 #include "Modules/ModuleManager.h"
 #include "Misc/EnumClassFlags.h"
 #include "UObject/Object.h"
@@ -23,7 +26,6 @@
 #include "AlphaBlend.h"
 #include "BoneIndices.h"
 #include "GameplayTagContainer.h"
-#include "Interfaces/Interface_BoneReferenceSkeletonProvider.h"
 #include "PoseSearch/KDTree.h"
 #include "ObjectTrace.h"
 #include "PoseSearch.generated.h"
@@ -86,17 +88,6 @@ enum class EPoseSearchMode : int32
 	Invalid = Num UMETA(Hidden)
 };
 
-UENUM()
-enum class EPoseSearchDataPreprocessor : int32
-{
-	None,
-	Normalize,
-	NormalizeOnlyByDeviation,
-
-	Num UMETA(Hidden),
-	Invalid = Num UMETA(Hidden)
-};
-
 UENUM()
 enum class EPoseSearchPoseFlags : uint32
 {
@@ -361,37 +352,6 @@ private:
 	TArrayView<float> FeatureVectorTable;
 };
 
-#if WITH_EDITOR
-// data structure collecting the internal layout representation of UPoseSearchFeatureChannel,
-// so we can aggregate data from different FPoseSearchIndex and calculate mean deviation with a homogeneous data set (ComputeChannelsDeviations
-struct FFeatureChannelLayoutSet
-{
-	// data structure holding DataOffset and Cardinality to find the data of the DebugName (channel breakdown / layout) in the related SearchIndexBases[SchemaIndex]
-	struct FEntry
-	{
-		FString DebugName; // for easier debugging
-		int32 SchemaIndex = -1; // index of the associated Schemas / SearchIndexBases used as input of the algorithm
-		int32 DataOffset = -1; // data offset from the base of SearchIndexBases[SchemaIndex].Values.GetData() from where the data associated to this Item starts
-		int32 Cardinality = -1; // data cardinality
-	};
-
-	// FIoHash is the hash associated to the channel data breakdown (e.g.: it could be a single SampledBones at a specific SampleTimes for a UPoseSearchFeatureChannel_Pose)
-	TMap<FIoHash, TArray<FEntry>> EntriesMap;
-	int32 CurrentSchemaIndex = -1;
-	TWeakObjectPtr<const UPoseSearchSchema> CurrentSchema;
-
-	void Add(FString DebugName, FIoHash Id, int32 DataOffset, int32 Cardinality)
-	{
-		check(DataOffset >= 0 && Cardinality >= 0 && CurrentSchemaIndex >= 0);
-		TArray<FEntry>& Entries = EntriesMap.FindOrAdd(Id);
-
-		// making sure all the FEntry associated with the same Id have the same Cardinality
-		check(Entries.IsEmpty() || Entries[0].Cardinality == Cardinality);
-		Entries.Add({ DebugName, CurrentSchemaIndex, DataOffset, Cardinality });
-	}
-};
-#endif // WITH_EDITOR
-
 class POSESEARCH_API IAssetIndexer
 {
 public:
@@ -415,208 +375,8 @@ public:
 	virtual FTransform GetTransformAndCacheResults(float SampleTime, float OriginTime, int8 SchemaBoneIdx, bool& Clamped) = 0;
 };
 
-class POSESEARCH_API ICostBreakDownData
-{
-public:
-	virtual ~ICostBreakDownData() {}
-
-	// returns the size of the dataset
-	virtual int32 Num() const = 0;
-
-	// returns true if Index-th cost data vector is associated with Schema
-	virtual bool IsCostVectorFromSchema(int32 Index, const UPoseSearchSchema* Schema) const = 0;
-
-	// returns the Index-th cost data vector
-	virtual TConstArrayView<float> GetCostVector(int32 Index, const UPoseSearchSchema* Schema) const = 0;
-
-	// every breakdown section start by calling BeginBreakDownSection...
-	virtual void BeginBreakDownSection(const FText& Label) = 0;
-
-	// ...then add as many SetCostBreakDown into the section...
-	virtual void SetCostBreakDown(float CostBreakDown, int32 Index, const UPoseSearchSchema* Schema) = 0;
-
-	// ...to finally wrap the section up by calling EndBreakDownSection
-	virtual void EndBreakDownSection(const FText& Label) = 0;
-
-	// true if want the channel to be verbose and generate the cost breakdown labels
-	virtual bool IsVerbose() const { return true; }
-
-	// most common case implementation
-	void AddEntireBreakDownSection(const FText& Label, const UPoseSearchSchema* Schema, int32 DataOffset, int32 Cardinality);
-};
-
 } // namespace UE::PoseSearch
 
-class POSESEARCH_API IPoseFilter
-{
-public:
-	virtual ~IPoseFilter() {}
-
-	// if true this filter will be evaluated
-	virtual bool IsPoseFilterActive() const { return false; }
-	
-	// if it returns false the pose candidate will be discarded
-	virtual bool IsPoseValid(TConstArrayView<float> PoseValues, TConstArrayView<float> QueryValues, int32 PoseIdx, const FPoseSearchPoseMetadata& Metadata) const { return true; }
-};
-
-//////////////////////////////////////////////////////////////////////////
-// Feature channels interface
-UCLASS(Abstract, BlueprintType, EditInlineNew)
-class POSESEARCH_API UPoseSearchFeatureChannel : public UObject, public IBoneReferenceSkeletonProvider, public IPoseFilter
-{
-	GENERATED_BODY()
-
-public:
-	int32 GetChannelCardinality() const { checkSlow(ChannelCardinality >= 0); return ChannelCardinality; }
-	int32 GetChannelDataOffset() const { checkSlow(ChannelDataOffset >= 0); return ChannelDataOffset; }
-
-	// Called during UPoseSearchSchema::Finalize to prepare the schema for this channel
-	virtual void InitializeSchema(UPoseSearchSchema* Schema) PURE_VIRTUAL(UPoseSearchFeatureChannel::InitializeSchema, );
-	
-	// Called at database build time to collect feature weights.
-	// Weights is sized to the cardinality of the schema and the feature channel should write
-	// its weights at the channel's data offset. Channels should provide a weight for each dimension.
-	virtual void FillWeights(TArray<float>& Weights) const PURE_VIRTUAL(UPoseSearchFeatureChannel::FillWeights, );
-
-	// Called at database build time to populate pose vectors with this channel's data
-	virtual void IndexAsset(UE::PoseSearch::IAssetIndexer& Indexer, UE::PoseSearch::FAssetIndexingOutput& IndexingOutput) const PURE_VIRTUAL(UPoseSearchFeatureChannel::IndexAsset, );
-
-	// Called at runtime to add this channel's data to the query pose vector
-	virtual void BuildQuery(UE::PoseSearch::FSearchContext& SearchContext, FPoseSearchFeatureVectorBuilder& InOutQuery) const PURE_VIRTUAL(UPoseSearchFeatureChannel::BuildQuery, );
-
-	// Draw this channel's data for the given pose vector
-	virtual void DebugDraw(const UE::PoseSearch::FDebugDrawParams& DrawParams, TConstArrayView<float> PoseVector) const PURE_VIRTUAL(UPoseSearchFeatureChannel::DebugDraw, );
-
-#if WITH_EDITOR
-	virtual void PopulateChannelLayoutSet(UE::PoseSearch::FFeatureChannelLayoutSet& FeatureChannelLayoutSet) const;
-	virtual void ComputeCostBreakdowns(UE::PoseSearch::ICostBreakDownData& CostBreakDownData, const UPoseSearchSchema* Schema) const;
-#endif
-
-private:
-	// IBoneReferenceSkeletonProvider interface
-	// Note this function is exclusively for FBoneReference details customization
-	class USkeleton* GetSkeleton(bool& bInvalidSkeletonIsError, const IPropertyHandle* PropertyHandle) override;
-
-protected:
-	friend class ::UPoseSearchSchema;
-
-	UPROPERTY(meta = (ExcludeFromHash))
-	int32 ChannelDataOffset = INDEX_NONE;
-
-	UPROPERTY(meta = (ExcludeFromHash))
-	int32 ChannelCardinality = INDEX_NONE;
-};
-
-USTRUCT()
-struct FPoseSearchSchemaColorPreset
-{
-	GENERATED_BODY()
-
-	UPROPERTY(EditAnywhere, Category = "Colors", meta = (ExcludeFromHash))
-	FLinearColor Query = FLinearColor::Blue;
-
-	UPROPERTY(EditAnywhere, Category = "Colors", meta = (ExcludeFromHash))
-	FLinearColor Result = FLinearColor::Yellow;
-};
-
-/**
-* Specifies the format of a pose search index. At runtime, queries are built according to the schema for searching.
-*/
-UCLASS(BlueprintType, Category = "Animation|Pose Search", Experimental, meta = (DisplayName = "Motion Database Config"), CollapseCategories)
-class POSESEARCH_API UPoseSearchSchema : public UDataAsset, public IBoneReferenceSkeletonProvider
-{
-	GENERATED_BODY()
-
-public:
-	// @todo: used only for indexing: cache it somewhere else
-	UPROPERTY(EditAnywhere, Category = "Schema")
-	TObjectPtr<USkeleton> Skeleton;
-
-	UPROPERTY(EditAnywhere, meta = (ClampMin = "1", ClampMax = "240"), Category = "Schema")
-	int32 SampleRate = 30;
-
-	UPROPERTY(EditAnywhere, Instanced, BlueprintReadWrite, Category = "Schema")
-	TArray<TObjectPtr<UPoseSearchFeatureChannel>> Channels;
-
-	// If set, this schema will support mirroring pose search databases
-	UPROPERTY(EditAnywhere, Category = "Schema")
-	TObjectPtr<UMirrorDataTable> MirrorDataTable;
-
-	UPROPERTY(EditAnywhere, Category = "Schema")
-	EPoseSearchDataPreprocessor DataPreprocessor = EPoseSearchDataPreprocessor::Normalize;
-
-	UPROPERTY(meta = (ExcludeFromHash))
-	int32 SchemaCardinality = 0;
-
-	// @todo: used only for indexing: cache it somewhere else
-	UPROPERTY(meta = (ExcludeFromHash))
-	TArray<FBoneReference> BoneReferences;
-
-	UPROPERTY(Transient)
-	TArray<uint16> BoneIndices;
-
-	UPROPERTY(Transient)
-	TArray<uint16> BoneIndicesWithParents;
-
-	// cost added to the continuing pose from databases that uses this schema
-	UPROPERTY(EditAnywhere, Category = "Schema")
-	float ContinuingPoseCostBias = 0.f;
-
-	// base cost added to all poses from databases that uses this schema. it can be overridden by UAnimNotifyState_PoseSearchModifyCost
-	UPROPERTY(EditAnywhere, Category = "Schema")
-	float BaseCostBias = 0.f;
-
-	// If there's a mirroring mismatch between the currently playing asset and a search candidate, this cost will be 
-	// added to the candidate, making it less likely to be selected
-	UPROPERTY(EditAnywhere, Category = "Schema")
-	float MirrorMismatchCostBias = 0.f;
-
-	UPROPERTY(EditAnywhere, Category = "Schema", meta = (ExcludeFromHash))
-	TArray<FPoseSearchSchemaColorPreset> ColorPresets;
-	
-	bool IsValid () const;
-
-	int32 GetNumBones () const { return BoneIndices.Num(); }
-
-	float GetSamplingInterval() const { return 1.0f / SampleRate; }
-
-	template<typename ChannelType>
-	const ChannelType* FindFirstChannelOfType() const
-	{
-		for (const TObjectPtr<UPoseSearchFeatureChannel>& ChannelPtr : Channels)
-		{
-			if (const ChannelType* Channel = Cast<const ChannelType>(ChannelPtr.Get()))
-			{
-				return Channel;
-			}
-		}
-		return nullptr;
-	}
-
-
-	// UObject
-	virtual void PreSave(FObjectPreSaveContext ObjectSaveContext) override;
-	virtual void PostLoad() override;
-
-	// IBoneReferenceSkeletonProvider
-	class USkeleton* GetSkeleton(bool& bInvalidSkeletonIsError, const IPropertyHandle* PropertyHandle) override { bInvalidSkeletonIsError = false; return Skeleton; }
-
-#if WITH_EDITOR
-	void ComputeCostBreakdowns(UE::PoseSearch::ICostBreakDownData& CostBreakDownData) const;
-
-	virtual void PostEditChangeProperty(struct FPropertyChangedEvent& PropertyChangedEvent) override;
-#endif
-
-	int32 AddBoneReference(const FBoneReference& BoneReference) { return BoneReferences.AddUnique(BoneReference); }
-
-private:
-	void Finalize(bool bRemoveEmptyChannels = true);
-	void ResolveBoneReferences();
-
-public:
-	void BuildQuery(UE::PoseSearch::FSearchContext& SearchContext, FPoseSearchFeatureVectorBuilder& InOutQuery) const;
-};
-
 //////////////////////////////////////////////////////////////////////////
 // Search index
 
diff --git a/Engine/Plugins/Experimental/Animation/PoseSearch/Source/Runtime/Public/PoseSearch/PoseSearchFeatureChannel.h b/Engine/Plugins/Experimental/Animation/PoseSearch/Source/Runtime/Public/PoseSearch/PoseSearchFeatureChannel.h
new file mode 100644
--- /dev/null
+++ b/Engine/Plugins/Experimental/Animation/PoseSearch/Source/Runtime/Public/PoseSearch/PoseSearchFeatureChannel.h
@@ -0,0 +1,158 @@
+// Copyright Epic Games, Inc. All Rights Reserved.
+
+#pragma once
+
+#include "IO/IoHash.h"
+#include "Interfaces/Interface_BoneReferenceSkeletonProvider.h"
+#include "PoseSearchFeatureChannel.generated.h"
+
+class UPoseSearchSchema;
+struct FPoseSearchPoseMetadata;
+struct FPoseSearchFeatureVectorBuilder;
+
+UENUM(BlueprintType)
+enum class EInputQueryPose : uint8
+{
+	// use character pose to compose the query
+	UseCharacterPose,
+
+	// if available reuse continuing pose from the database to compose the query or else UseCharacterPose
+	UseContinuingPose,
+
+	// if available reuse and interpolate continuing pose from the database to compose the query or else UseCharacterPose
+	UseInterpolatedContinuingPose,
+
+	Num UMETA(Hidden),
+	Invalid = Num UMETA(Hidden)
+};
+
+namespace UE::PoseSearch
+{
+
+struct FAssetIndexingOutput;
+struct FDebugDrawParams;
+struct FSearchContext;
+class IAssetIndexer;
+
+#if WITH_EDITOR
+// data structure collecting the internal layout representation of UPoseSearchFeatureChannel,
+// so we can aggregate data from different FPoseSearchIndex and calculate mean deviation with a homogeneous data set (ComputeChannelsDeviations
+struct FFeatureChannelLayoutSet
+{
+	// data structure holding DataOffset and Cardinality to find the data of the DebugName (channel breakdown / layout) in the related SearchIndexBases[SchemaIndex]
+	struct FEntry
+	{
+		FString DebugName; // for easier debugging
+		int32 SchemaIndex = -1; // index of the associated Schemas / SearchIndexBases used as input of the algorithm
+		int32 DataOffset = -1; // data offset from the base of SearchIndexBases[SchemaIndex].Values.GetData() from where the data associated to this Item starts
+		int32 Cardinality = -1; // data cardinality
+	};
+
+	// FIoHash is the hash associated to the channel data breakdown (e.g.: it could be a single SampledBones at a specific SampleTimes for a UPoseSearchFeatureChannel_Pose)
+	TMap<FIoHash, TArray<FEntry>> EntriesMap;
+	int32 CurrentSchemaIndex = -1;
+	TWeakObjectPtr<const UPoseSearchSchema> CurrentSchema;
+
+	void Add(FString DebugName, FIoHash Id, int32 DataOffset, int32 Cardinality)
+	{
+		check(DataOffset >= 0 && Cardinality >= 0 && CurrentSchemaIndex >= 0);
+		TArray<FEntry>& Entries = EntriesMap.FindOrAdd(Id);
+
+		// making sure all the FEntry associated with the same Id have the same Cardinality
+		check(Entries.IsEmpty() || Entries[0].Cardinality == Cardinality);
+		Entries.Add({ DebugName, CurrentSchemaIndex, DataOffset, Cardinality });
+	}
+};
+#endif // WITH_EDITOR
+
+class POSESEARCH_API ICostBreakDownData
+{
+public:
+	virtual ~ICostBreakDownData() {}
+
+	// returns the size of the dataset
+	virtual int32 Num() const = 0;
+
+	// returns true if Index-th cost data vector is associated with Schema
+	virtual bool IsCostVectorFromSchema(int32 Index, const UPoseSearchSchema* Schema) const = 0;
+
+	// returns the Index-th cost data vector
+	virtual TConstArrayView<float> GetCostVector(int32 Index, const UPoseSearchSchema* Schema) const = 0;
+
+	// every breakdown section start by calling BeginBreakDownSection...
+	virtual void BeginBreakDownSection(const FText& Label) = 0;
+
+	// ...then add as many SetCostBreakDown into the section...
+	virtual void SetCostBreakDown(float CostBreakDown, int32 Index, const UPoseSearchSchema* Schema) = 0;
+
+	// ...to finally wrap the section up by calling EndBreakDownSection
+	virtual void EndBreakDownSection(const FText& Label) = 0;
+
+	// true if want the channel to be verbose and generate the cost breakdown labels
+	virtual bool IsVerbose() const { return true; }
+
+	// most common case implementation
+	void AddEntireBreakDownSection(const FText& Label, const UPoseSearchSchema* Schema, int32 DataOffset, int32 Cardinality);
+};
+
+} // namespace UE::PoseSearch
+
+class POSESEARCH_API IPoseFilter
+{
+public:
+	virtual ~IPoseFilter() {}
+
+	// if true this filter will be evaluated
+	virtual bool IsPoseFilterActive() const { return false; }
+	
+	// if it returns false the pose candidate will be discarded
+	virtual bool IsPoseValid(TConstArrayView<float> PoseValues, TConstArrayView<float> QueryValues, int32 PoseIdx, const FPoseSearchPoseMetadata& Metadata) const { return true; }
+};
+
+//////////////////////////////////////////////////////////////////////////
+// Feature channels interface
+UCLASS(Abstract, BlueprintType, EditInlineNew)
+class POSESEARCH_API UPoseSearchFeatureChannel : public UObject, public IBoneReferenceSkeletonProvider, public IPoseFilter
+{
+	GENERATED_BODY()
+
+public:
+	int32 GetChannelCardinality() const { checkSlow(ChannelCardinality >= 0); return ChannelCardinality; }
+	int32 GetChannelDataOffset() const { checkSlow(ChannelDataOffset >= 0); return ChannelDataOffset; }
+
+	// Called during UPoseSearchSchema::Finalize to prepare the schema for this channel
+	virtual void InitializeSchema(UPoseSearchSchema* Schema) PURE_VIRTUAL(UPoseSearchFeatureChannel::InitializeSchema, );
+	
+	// Called at database build time to collect feature weights.
+	// Weights is sized to the cardinality of the schema and the feature channel should write
+	// its weights at the channel's data offset. Channels should provide a weight for each dimension.
+	virtual void FillWeights(TArray<float>& Weights) const PURE_VIRTUAL(UPoseSearchFeatureChannel::FillWeights, );
+
+	// Called at database build time to populate pose vectors with this channel's data
+	virtual void IndexAsset(UE::PoseSearch::IAssetIndexer& Indexer, UE::PoseSearch::FAssetIndexingOutput& IndexingOutput) const PURE_VIRTUAL(UPoseSearchFeatureChannel::IndexAsset, );
+
+	// Called at runtime to add this channel's data to the query pose vector
+	virtual void BuildQuery(UE::PoseSearch::FSearchContext& SearchContext, FPoseSearchFeatureVectorBuilder& InOutQuery) const PURE_VIRTUAL(UPoseSearchFeatureChannel::BuildQuery, );
+
+	// Draw this channel's data for the given pose vector
+	virtual void DebugDraw(const UE::PoseSearch::FDebugDrawParams& DrawParams, TConstArrayView<float> PoseVector) const PURE_VIRTUAL(UPoseSearchFeatureChannel::DebugDraw, );
+
+#if WITH_EDITOR
+	virtual void PopulateChannelLayoutSet(UE::PoseSearch::FFeatureChannelLayoutSet& FeatureChannelLayoutSet) const;
+	virtual void ComputeCostBreakdowns(UE::PoseSearch::ICostBreakDownData& CostBreakDownData, const UPoseSearchSchema* Schema) const;
+#endif
+
+private:
+	// IBoneReferenceSkeletonProvider interface
+	// Note this function is exclusively for FBoneReference details customization
+	class USkeleton* GetSkeleton(bool& bInvalidSkeletonIsError, const IPropertyHandle* PropertyHandle) override;
+
+protected:
+	friend class ::UPoseSearchSchema;
+
+	UPROPERTY(meta = (ExcludeFromHash))
+	int32 ChannelDataOffset = INDEX_NONE;
+
+	UPROPERTY(meta = (ExcludeFromHash))
+	int32 ChannelCardinality = INDEX_NONE;
+};
\ No newline at end of file
diff --git a/Engine/Plugins/Experimental/Animation/PoseSearch/Source/Runtime/Public/PoseSearch/PoseSearchFeatureChannels.h b/Engine/Plugins/Experimental/Animation/PoseSearch/Source/Runtime/Public/PoseSearch/PoseSearchFeatureChannels.h
deleted file mode 100644
diff --git a/Engine/Plugins/Experimental/Animation/PoseSearch/Source/Runtime/Public/PoseSearch/PoseSearchSchema.h b/Engine/Plugins/Experimental/Animation/PoseSearch/Source/Runtime/Public/PoseSearch/PoseSearchSchema.h
new file mode 100644
--- /dev/null
+++ b/Engine/Plugins/Experimental/Animation/PoseSearch/Source/Runtime/Public/PoseSearch/PoseSearchSchema.h
@@ -0,0 +1,140 @@
+// Copyright Epic Games, Inc. All Rights Reserved.
+
+#pragma once
+
+#include "BoneContainer.h"
+#include "Engine/DataAsset.h"
+#include "Interfaces/Interface_BoneReferenceSkeletonProvider.h"
+#include "PoseSearchSchema.generated.h"
+
+struct FBoneReference;
+struct FPoseSearchFeatureVectorBuilder;
+class UMirrorDataTable;
+class UPoseSearchFeatureChannel;
+
+namespace UE::PoseSearch
+{
+	struct FSearchContext;
+	class ICostBreakDownData;
+} // namespace UE::PoseSearch
+
+UENUM()
+enum class EPoseSearchDataPreprocessor : int32
+{
+	None,
+	Normalize,
+	NormalizeOnlyByDeviation,
+
+	Num UMETA(Hidden),
+	Invalid = Num UMETA(Hidden)
+};
+
+USTRUCT()
+struct FPoseSearchSchemaColorPreset
+{
+	GENERATED_BODY()
+
+	UPROPERTY(EditAnywhere, Category = "Colors", meta = (ExcludeFromHash))
+	FLinearColor Query = FLinearColor::Blue;
+
+	UPROPERTY(EditAnywhere, Category = "Colors", meta = (ExcludeFromHash))
+	FLinearColor Result = FLinearColor::Yellow;
+};
+
+/**
+* Specifies the format of a pose search index. At runtime, queries are built according to the schema for searching.
+*/
+UCLASS(BlueprintType, Category = "Animation|Pose Search", Experimental, meta = (DisplayName = "Motion Database Config"), CollapseCategories)
+class POSESEARCH_API UPoseSearchSchema : public UDataAsset, public IBoneReferenceSkeletonProvider
+{
+	GENERATED_BODY()
+
+public:
+	// @todo: used only for indexing: cache it somewhere else
+	UPROPERTY(EditAnywhere, Category = "Schema")
+	TObjectPtr<USkeleton> Skeleton;
+
+	UPROPERTY(EditAnywhere, meta = (ClampMin = "1", ClampMax = "240"), Category = "Schema")
+	int32 SampleRate = 30;
+
+	UPROPERTY(EditAnywhere, Instanced, BlueprintReadWrite, Category = "Schema")
+	TArray<TObjectPtr<UPoseSearchFeatureChannel>> Channels;
+
+	// If set, this schema will support mirroring pose search databases
+	UPROPERTY(EditAnywhere, Category = "Schema")
+	TObjectPtr<UMirrorDataTable> MirrorDataTable;
+
+	UPROPERTY(EditAnywhere, Category = "Schema")
+	EPoseSearchDataPreprocessor DataPreprocessor = EPoseSearchDataPreprocessor::Normalize;
+
+	UPROPERTY(meta = (ExcludeFromHash))
+	int32 SchemaCardinality = 0;
+
+	// @todo: used only for indexing: cache it somewhere else
+	UPROPERTY(meta = (ExcludeFromHash))
+	TArray<FBoneReference> BoneReferences;
+
+	UPROPERTY(Transient)
+	TArray<uint16> BoneIndices;
+
+	UPROPERTY(Transient)
+	TArray<uint16> BoneIndicesWithParents;
+
+	// cost added to the continuing pose from databases that uses this schema
+	UPROPERTY(EditAnywhere, Category = "Schema")
+	float ContinuingPoseCostBias = 0.f;
+
+	// base cost added to all poses from databases that uses this schema. it can be overridden by UAnimNotifyState_PoseSearchModifyCost
+	UPROPERTY(EditAnywhere, Category = "Schema")
+	float BaseCostBias = 0.f;
+
+	// If there's a mirroring mismatch between the currently playing asset and a search candidate, this cost will be 
+	// added to the candidate, making it less likely to be selected
+	UPROPERTY(EditAnywhere, Category = "Schema")
+	float MirrorMismatchCostBias = 0.f;
+
+	UPROPERTY(EditAnywhere, Category = "Schema", meta = (ExcludeFromHash))
+	TArray<FPoseSearchSchemaColorPreset> ColorPresets;
+	
+	bool IsValid () const;
+
+	int32 GetNumBones () const { return BoneIndices.Num(); }
+
+	float GetSamplingInterval() const { return 1.0f / SampleRate; }
+
+	template<typename ChannelType>
+	const ChannelType* FindFirstChannelOfType() const
+	{
+		for (const TObjectPtr<UPoseSearchFeatureChannel>& ChannelPtr : Channels)
+		{
+			if (const ChannelType* Channel = Cast<const ChannelType>(ChannelPtr.Get()))
+			{
+				return Channel;
+			}
+		}
+		return nullptr;
+	}
+
+
+	// UObject
+	virtual void PreSave(FObjectPreSaveContext ObjectSaveContext) override;
+	virtual void PostLoad() override;
+
+	// IBoneReferenceSkeletonProvider
+	class USkeleton* GetSkeleton(bool& bInvalidSkeletonIsError, const IPropertyHandle* PropertyHandle) override { bInvalidSkeletonIsError = false; return Skeleton; }
+
+#if WITH_EDITOR
+	void ComputeCostBreakdowns(UE::PoseSearch::ICostBreakDownData& CostBreakDownData) const;
+
+	virtual void PostEditChangeProperty(struct FPropertyChangedEvent& PropertyChangedEvent) override;
+#endif
+
+	int32 AddBoneReference(const FBoneReference& BoneReference) { return BoneReferences.AddUnique(BoneReference); }
+
+private:
+	void Finalize(bool bRemoveEmptyChannels = true);
+	void ResolveBoneReferences();
+
+public:
+	void BuildQuery(UE::PoseSearch::FSearchContext& SearchContext, FPoseSearchFeatureVectorBuilder& InOutQuery) const;
+};
