motion matching - pose search plugin runtime refactoring

diff --git a/Engine/Plugins/Experimental/Animation/PoseSearch/Source/Editor/Private/PoseSearchCustomization.cpp b/Engine/Plugins/Experimental/Animation/PoseSearch/Source/Editor/Private/PoseSearchCustomization.cpp
--- a/Engine/Plugins/Experimental/Animation/PoseSearch/Source/Editor/Private/PoseSearchCustomization.cpp
+++ b/Engine/Plugins/Experimental/Animation/PoseSearch/Source/Editor/Private/PoseSearchCustomization.cpp
@@ -1,7 +1,6 @@
 // Copyright Epic Games, Inc. All Rights Reserved.
 
 #include "PoseSearchCustomization.h"
-#include "PoseSearch/PoseSearch.h"
 #include "Engine/GameViewportClient.h"
 #include "AssetRegistry/AssetData.h"
 #include "EditorClassUtils.h"
@@ -25,6 +24,7 @@
 #include "IDetailsView.h"
 #include "ObjectEditorUtils.h"
 #include "Animation/AnimSequence.h"
+#include "PoseSearch/PoseSearchDatabase.h"
 
 #define LOCTEXT_NAMESPACE "PoseSearchCustomization"
 
diff --git a/Engine/Plugins/Experimental/Animation/PoseSearch/Source/Editor/Private/PoseSearchDatabaseAssetListItem.cpp b/Engine/Plugins/Experimental/Animation/PoseSearch/Source/Editor/Private/PoseSearchDatabaseAssetListItem.cpp
--- a/Engine/Plugins/Experimental/Animation/PoseSearch/Source/Editor/Private/PoseSearchDatabaseAssetListItem.cpp
+++ b/Engine/Plugins/Experimental/Animation/PoseSearch/Source/Editor/Private/PoseSearchDatabaseAssetListItem.cpp
@@ -3,8 +3,6 @@
 #include "PoseSearchDatabaseAssetListItem.h"
 #include "PoseSearchDatabaseViewModel.h"
 
-#include "PoseSearch/PoseSearch.h"
-
 #include "Animation/AnimSequence.h"
 #include "Animation/AnimComposite.h"
 #include "Animation/BlendSpace.h"
@@ -15,6 +13,7 @@
 #include "AssetSelection.h"
 #include "ClassIconFinder.h"
 #include "DetailColumnSizeData.h"
+#include "PoseSearch/PoseSearchDatabase.h"
 #include "PoseSearchDatabaseAssetTree.h"
 
 #include "Widgets/Text/SRichTextBlock.h"
diff --git a/Engine/Plugins/Experimental/Animation/PoseSearch/Source/Editor/Private/PoseSearchDatabaseAssetTree.cpp b/Engine/Plugins/Experimental/Animation/PoseSearch/Source/Editor/Private/PoseSearchDatabaseAssetTree.cpp
--- a/Engine/Plugins/Experimental/Animation/PoseSearch/Source/Editor/Private/PoseSearchDatabaseAssetTree.cpp
+++ b/Engine/Plugins/Experimental/Animation/PoseSearch/Source/Editor/Private/PoseSearchDatabaseAssetTree.cpp
@@ -1,31 +1,26 @@
 // Copyright Epic Games, Inc. All Rights Reserved.
 
 #include "PoseSearchDatabaseAssetTree.h"
-#include "PoseSearchDatabaseViewModel.h"
-
-#include "PoseSearch/PoseSearch.h"
-
-#include "Animation/AnimSequence.h"
 #include "Animation/AnimComposite.h"
+#include "Animation/AnimSequence.h"
 #include "Animation/BlendSpace.h"
-
-#include "Misc/TransactionObjectEvent.h"
-#include "DragAndDrop/AssetDragDropOp.h"
-#include "Misc/FeedbackContext.h"
 #include "AssetSelection.h"
 #include "ClassIconFinder.h"
 #include "DetailColumnSizeData.h"
-
-#include "Widgets/Text/SRichTextBlock.h"
-#include "Widgets/Input/SCheckBox.h"
-#include "Widgets/Input/SSearchBox.h"
+#include "DragAndDrop/AssetDragDropOp.h"
 #include "Framework/Commands/GenericCommands.h"
 #include "Framework/MultiBox/MultiBoxBuilder.h"
+#include "Misc/FeedbackContext.h"
+#include "Misc/TransactionObjectEvent.h"
+#include "PoseSearch/PoseSearchDatabase.h"
+#include "PoseSearchDatabaseViewModel.h"
 #include "SPositiveActionButton.h"
 #include "Styling/AppStyle.h"
-
 #include "ScopedTransaction.h"
 #include "Styling/StyleColors.h"
+#include "Widgets/Text/SRichTextBlock.h"
+#include "Widgets/Input/SCheckBox.h"
+#include "Widgets/Input/SSearchBox.h"
 
 #define LOCTEXT_NAMESPACE "SDatabaseAssetTree"
 
diff --git a/Engine/Plugins/Experimental/Animation/PoseSearch/Source/Editor/Private/PoseSearchDatabaseAssetTreeNode.cpp b/Engine/Plugins/Experimental/Animation/PoseSearch/Source/Editor/Private/PoseSearchDatabaseAssetTreeNode.cpp
--- a/Engine/Plugins/Experimental/Animation/PoseSearch/Source/Editor/Private/PoseSearchDatabaseAssetTreeNode.cpp
+++ b/Engine/Plugins/Experimental/Animation/PoseSearch/Source/Editor/Private/PoseSearchDatabaseAssetTreeNode.cpp
@@ -1,34 +1,28 @@
 // Copyright Epic Games, Inc. All Rights Reserved.
 
 #include "PoseSearchDatabaseAssetTreeNode.h"
-#include "PoseSearchDatabaseViewModel.h"
-
-#include "PoseSearch/PoseSearch.h"
-
-#include "Animation/AnimSequence.h"
 #include "Animation/AnimComposite.h"
+#include "Animation/AnimSequence.h"
 #include "Animation/BlendSpace.h"
-
-#include "Misc/TransactionObjectEvent.h"
-#include "DragAndDrop/AssetDragDropOp.h"
-#include "Misc/FeedbackContext.h"
 #include "AssetSelection.h"
 #include "ClassIconFinder.h"
 #include "DetailColumnSizeData.h"
-#include "PoseSearchDatabaseAssetListItem.h"
-
-#include "Widgets/Text/SRichTextBlock.h"
-#include "Widgets/Input/SCheckBox.h"
-#include "Widgets/Input/SSearchBox.h"
-#include "Widgets/Views/SListView.h"
-
+#include "DragAndDrop/AssetDragDropOp.h"
 #include "Framework/Commands/GenericCommands.h"
 #include "Framework/MultiBox/MultiBoxBuilder.h"
+#include "Misc/FeedbackContext.h"
+#include "Misc/TransactionObjectEvent.h"
+#include "PoseSearch/PoseSearchDatabase.h"
+#include "PoseSearchDatabaseAssetListItem.h"
+#include "PoseSearchDatabaseViewModel.h"
+#include "ScopedTransaction.h"
 #include "SPositiveActionButton.h"
 #include "Styling/AppStyle.h"
-
-#include "ScopedTransaction.h"
 #include "Styling/StyleColors.h"
+#include "Widgets/Input/SCheckBox.h"
+#include "Widgets/Input/SSearchBox.h"
+#include "Widgets/Text/SRichTextBlock.h"
+#include "Widgets/Views/SListView.h"
 
 namespace UE::PoseSearch
 {
diff --git a/Engine/Plugins/Experimental/Animation/PoseSearch/Source/Editor/Private/PoseSearchDatabaseEdMode.cpp b/Engine/Plugins/Experimental/Animation/PoseSearch/Source/Editor/Private/PoseSearchDatabaseEdMode.cpp
--- a/Engine/Plugins/Experimental/Animation/PoseSearch/Source/Editor/Private/PoseSearchDatabaseEdMode.cpp
+++ b/Engine/Plugins/Experimental/Animation/PoseSearch/Source/Editor/Private/PoseSearchDatabaseEdMode.cpp
@@ -4,7 +4,6 @@
 #include "PoseSearchDatabaseViewportClient.h"
 #include "PoseSearchDatabaseEditor.h"
 #include "PoseSearchDatabaseViewModel.h"
-#include "PoseSearch/PoseSearch.h"
 
 #include "EngineUtils.h"
 
diff --git a/Engine/Plugins/Experimental/Animation/PoseSearch/Source/Editor/Private/PoseSearchDatabaseEditor.cpp b/Engine/Plugins/Experimental/Animation/PoseSearch/Source/Editor/Private/PoseSearchDatabaseEditor.cpp
--- a/Engine/Plugins/Experimental/Animation/PoseSearch/Source/Editor/Private/PoseSearchDatabaseEditor.cpp
+++ b/Engine/Plugins/Experimental/Animation/PoseSearch/Source/Editor/Private/PoseSearchDatabaseEditor.cpp
@@ -8,10 +8,10 @@
 #include "PoseSearchDatabaseViewModel.h"
 #include "PoseSearchDatabaseEditorReflection.h"
 #include "PoseSearchEditor.h"
-#include "PoseSearch/PoseSearch.h"
 #include "PoseSearch/PoseSearchDerivedData.h"
 #include "GameFramework/WorldSettings.h"
 #include "AdvancedPreviewSceneModule.h"
+#include "InstancedStruct.h"
 #include "IStructureDetailsView.h"
 #include "PropertyEditorModule.h"
 #include "Modules/ModuleManager.h"
diff --git a/Engine/Plugins/Experimental/Animation/PoseSearch/Source/Editor/Private/PoseSearchDatabaseEditorReflection.cpp b/Engine/Plugins/Experimental/Animation/PoseSearch/Source/Editor/Private/PoseSearchDatabaseEditorReflection.cpp
--- a/Engine/Plugins/Experimental/Animation/PoseSearch/Source/Editor/Private/PoseSearchDatabaseEditorReflection.cpp
+++ b/Engine/Plugins/Experimental/Animation/PoseSearch/Source/Editor/Private/PoseSearchDatabaseEditorReflection.cpp
@@ -1,6 +1,7 @@
 // Copyright Epic Games, Inc. All Rights Reserved.
 
 #include "PoseSearchDatabaseEditorReflection.h"
+#include "InstancedStruct.h"
 #include "PoseSearchDatabaseAssetTree.h"
 #include "PoseSearchDatabaseAssetTreeNode.h"
 #include "PoseSearchDatabaseViewModel.h"
diff --git a/Engine/Plugins/Experimental/Animation/PoseSearch/Source/Editor/Private/PoseSearchDatabaseFactory.cpp b/Engine/Plugins/Experimental/Animation/PoseSearch/Source/Editor/Private/PoseSearchDatabaseFactory.cpp
--- a/Engine/Plugins/Experimental/Animation/PoseSearch/Source/Editor/Private/PoseSearchDatabaseFactory.cpp
+++ b/Engine/Plugins/Experimental/Animation/PoseSearch/Source/Editor/Private/PoseSearchDatabaseFactory.cpp
@@ -1,7 +1,7 @@
 // Copyright Epic Games, Inc. All Rights Reserved.
 
 #include "PoseSearchDatabaseFactory.h"
-#include "PoseSearch/PoseSearch.h"
+#include "PoseSearch/PoseSearchDatabase.h"
 
 #define LOCTEXT_NAMESPACE "PoseSearchEditor"
 
diff --git a/Engine/Plugins/Experimental/Animation/PoseSearch/Source/Editor/Private/PoseSearchDatabasePreviewScene.cpp b/Engine/Plugins/Experimental/Animation/PoseSearch/Source/Editor/Private/PoseSearchDatabasePreviewScene.cpp
--- a/Engine/Plugins/Experimental/Animation/PoseSearch/Source/Editor/Private/PoseSearchDatabasePreviewScene.cpp
+++ b/Engine/Plugins/Experimental/Animation/PoseSearch/Source/Editor/Private/PoseSearchDatabasePreviewScene.cpp
@@ -1,16 +1,16 @@
 // Copyright Epic Games, Inc. All Rights Reserved.
 
 #include "PoseSearchDatabasePreviewScene.h"
+#include "Animation/DebugSkelMeshComponent.h"
 #include "Components/StaticMeshComponent.h"
+#include "EngineUtils.h"
+#include "GameFramework/WorldSettings.h"
+#include "PoseSearch/PoseSearchContext.h"
+#include "PoseSearch/PoseSearchDatabase.h"
+#include "PoseSearch/PoseSearchDerivedData.h"
 #include "PoseSearchDatabaseEditor.h"
 #include "PoseSearchDatabaseViewModel.h"
 
-#include "PoseSearch/PoseSearchDerivedData.h"
-
-#include "Animation/DebugSkelMeshComponent.h"
-#include "GameFramework/WorldSettings.h"
-#include "EngineUtils.h"
-
 namespace UE::PoseSearch
 {
 	FDatabasePreviewScene::FDatabasePreviewScene(
diff --git a/Engine/Plugins/Experimental/Animation/PoseSearch/Source/Editor/Private/PoseSearchDatabaseViewModel.cpp b/Engine/Plugins/Experimental/Animation/PoseSearch/Source/Editor/Private/PoseSearchDatabaseViewModel.cpp
--- a/Engine/Plugins/Experimental/Animation/PoseSearch/Source/Editor/Private/PoseSearchDatabaseViewModel.cpp
+++ b/Engine/Plugins/Experimental/Animation/PoseSearch/Source/Editor/Private/PoseSearchDatabaseViewModel.cpp
@@ -4,6 +4,8 @@
 #include "PoseSearchDatabasePreviewScene.h"
 #include "PoseSearchDatabaseAssetTreeNode.h"
 #include "PoseSearchEditor.h"
+#include "PoseSearch/PoseSearchDatabase.h"
+#include "PoseSearch/PoseSearchDefines.h"
 #include "PoseSearch/PoseSearchDerivedData.h"
 #include "PoseSearch/PoseSearchSchema.h"
 #include "Modules/ModuleManager.h"
@@ -14,6 +16,7 @@
 #include "Animation/MirrorDataTable.h"
 #include "Animation/BlendSpace.h"
 #include "EngineUtils.h"
+#include "InstancedStruct.h"
 #include "PropertyEditorModule.h"
 #include "Modules/ModuleManager.h"
 #include "Camera/CameraComponent.h"
diff --git a/Engine/Plugins/Experimental/Animation/PoseSearch/Source/Editor/Private/PoseSearchDebugger.h b/Engine/Plugins/Experimental/Animation/PoseSearch/Source/Editor/Private/PoseSearchDebugger.h
--- a/Engine/Plugins/Experimental/Animation/PoseSearch/Source/Editor/Private/PoseSearchDebugger.h
+++ b/Engine/Plugins/Experimental/Animation/PoseSearch/Source/Editor/Private/PoseSearchDebugger.h
@@ -3,12 +3,8 @@
 #pragma once
 
 #include "CoreMinimal.h"
-//#include "CustomBoneIndexArray.h"
 #include "IRewindDebuggerExtension.h"
-//#include "IRewindDebuggerView.h"
 #include "IRewindDebuggerViewCreator.h"
-//#include "PoseSearch/PoseSearch.h"
-//#include "PoseSearchMeshComponent.h"
 
 namespace TraceServices { class IAnalysisSession; }
 
diff --git a/Engine/Plugins/Experimental/Animation/PoseSearch/Source/Editor/Private/PoseSearchDebuggerDatabaseRowData.h b/Engine/Plugins/Experimental/Animation/PoseSearch/Source/Editor/Private/PoseSearchDebuggerDatabaseRowData.h
--- a/Engine/Plugins/Experimental/Animation/PoseSearch/Source/Editor/Private/PoseSearchDebuggerDatabaseRowData.h
+++ b/Engine/Plugins/Experimental/Animation/PoseSearch/Source/Editor/Private/PoseSearchDebuggerDatabaseRowData.h
@@ -2,7 +2,9 @@
 
 #pragma once
 
-#include "PoseSearch/PoseSearch.h"
+#include "PoseSearch/PoseSearchContext.h"
+#include "PoseSearch/PoseSearchIndex.h"
+#include "Templates/SharedPointer.h"
 
 namespace UE::PoseSearch
 {
diff --git a/Engine/Plugins/Experimental/Animation/PoseSearch/Source/Editor/Private/PoseSearchDebuggerDatabaseView.cpp b/Engine/Plugins/Experimental/Animation/PoseSearch/Source/Editor/Private/PoseSearchDebuggerDatabaseView.cpp
--- a/Engine/Plugins/Experimental/Animation/PoseSearch/Source/Editor/Private/PoseSearchDebuggerDatabaseView.cpp
+++ b/Engine/Plugins/Experimental/Animation/PoseSearch/Source/Editor/Private/PoseSearchDebuggerDatabaseView.cpp
@@ -2,9 +2,13 @@
 
 #include "PoseSearchDebuggerDatabaseView.h"
 #include "Algo/AllOf.h"
+#include "Animation/AnimComposite.h"
+#include "Animation/AnimSequence.h"
+#include "InstancedStruct.h"
 #include "PoseSearchDebuggerDatabaseRow.h"
 #include "PoseSearchDebuggerView.h"
 #include "PoseSearchDebuggerViewModel.h"
+#include "PoseSearch/PoseSearchDatabase.h"
 #include "PoseSearch/PoseSearchDerivedData.h"
 #include "PoseSearch/PoseSearchFeatureChannel.h"
 #include "PoseSearch/PoseSearchSchema.h"
diff --git a/Engine/Plugins/Experimental/Animation/PoseSearch/Source/Editor/Private/PoseSearchDebuggerView.cpp b/Engine/Plugins/Experimental/Animation/PoseSearch/Source/Editor/Private/PoseSearchDebuggerView.cpp
--- a/Engine/Plugins/Experimental/Animation/PoseSearch/Source/Editor/Private/PoseSearchDebuggerView.cpp
+++ b/Engine/Plugins/Experimental/Animation/PoseSearch/Source/Editor/Private/PoseSearchDebuggerView.cpp
@@ -2,13 +2,14 @@
 
 #include "PoseSearchDebuggerView.h"
 #include "Modules/ModuleManager.h"
+#include "PoseSearch/PoseSearchDatabase.h"
 #include "PoseSearch/PoseSearchDerivedData.h"
+#include "PoseSearch/PoseSearchSchema.h"
 #include "PoseSearchDebugger.h"
 #include "PoseSearchDebuggerDatabaseRowData.h"
 #include "PoseSearchDebuggerDatabaseView.h"
 #include "PoseSearchDebuggerReflection.h"
 #include "PoseSearchDebuggerViewModel.h"
-#include "PoseSearch/PoseSearchSchema.h"
 #include "PropertyEditorModule.h"
 #include "Trace/PoseSearchTraceProvider.h"
 #include "Widgets/Input/SButton.h"
diff --git a/Engine/Plugins/Experimental/Animation/PoseSearch/Source/Editor/Private/PoseSearchDebuggerViewModel.cpp b/Engine/Plugins/Experimental/Animation/PoseSearch/Source/Editor/Private/PoseSearchDebuggerViewModel.cpp
--- a/Engine/Plugins/Experimental/Animation/PoseSearch/Source/Editor/Private/PoseSearchDebuggerViewModel.cpp
+++ b/Engine/Plugins/Experimental/Animation/PoseSearch/Source/Editor/Private/PoseSearchDebuggerViewModel.cpp
@@ -1,11 +1,15 @@
 // Copyright Epic Games, Inc. All Rights Reserved.
 
 #include "PoseSearchDebuggerViewModel.h"
+#include "Animation/AnimComposite.h"
+#include "Animation/AnimSequence.h"
 #include "Animation/MirrorDataTable.h"
 #include "Engine/SkeletalMesh.h"
 #include "IAnimationProvider.h"
 #include "IGameplayProvider.h"
+#include "InstancedStruct.h"
 #include "IRewindDebugger.h"
+#include "PoseSearch/PoseSearchDatabase.h"
 #include "PoseSearch/PoseSearchDerivedData.h"
 #include "PoseSearch/PoseSearchSchema.h"
 #include "Trace/PoseSearchTraceProvider.h"
diff --git a/Engine/Plugins/Experimental/Animation/PoseSearch/Source/Editor/Private/PoseSearchDebuggerViewModel.h b/Engine/Plugins/Experimental/Animation/PoseSearch/Source/Editor/Private/PoseSearchDebuggerViewModel.h
--- a/Engine/Plugins/Experimental/Animation/PoseSearch/Source/Editor/Private/PoseSearchDebuggerViewModel.h
+++ b/Engine/Plugins/Experimental/Animation/PoseSearch/Source/Editor/Private/PoseSearchDebuggerViewModel.h
@@ -8,7 +8,10 @@
 #include "Templates/SharedPointer.h"
 #include "UObject/WeakObjectPtrTemplates.h"
 
+struct FInstancedStruct;
 class IRewindDebugger;
+class UPoseSearchDatabase;
+class UPoseSearchSearchableAsset;
 
 namespace UE::PoseSearch
 {
diff --git a/Engine/Plugins/Experimental/Animation/PoseSearch/Source/Editor/Private/PoseSearchMeshComponent.cpp b/Engine/Plugins/Experimental/Animation/PoseSearch/Source/Editor/Private/PoseSearchMeshComponent.cpp
--- a/Engine/Plugins/Experimental/Animation/PoseSearch/Source/Editor/Private/PoseSearchMeshComponent.cpp
+++ b/Engine/Plugins/Experimental/Animation/PoseSearch/Source/Editor/Private/PoseSearchMeshComponent.cpp
@@ -3,6 +3,7 @@
 #include "PoseSearchMeshComponent.h"
 #include "AnimationRuntime.h"
 #include "Animation/AnimRootMotionProvider.h"
+#include "Animation/AnimSequence.h"
 #include "Animation/BlendSpace.h"
 #include "Animation/MirrorDataTable.h"
 #include "Engine/SkinnedAsset.h"
diff --git a/Engine/Plugins/Experimental/Animation/PoseSearch/Source/Editor/Private/PoseSearchMeshComponent.h b/Engine/Plugins/Experimental/Animation/PoseSearch/Source/Editor/Private/PoseSearchMeshComponent.h
--- a/Engine/Plugins/Experimental/Animation/PoseSearch/Source/Editor/Private/PoseSearchMeshComponent.h
+++ b/Engine/Plugins/Experimental/Animation/PoseSearch/Source/Editor/Private/PoseSearchMeshComponent.h
@@ -2,10 +2,13 @@
 
 #pragma once
 
+#include "BonePose.h"
 #include "Components/PoseableMeshComponent.h"
-#include "PoseSearch/PoseSearch.h"
+#include "PoseSearch/PoseSearchIndex.h"
 #include "PoseSearchMeshComponent.generated.h"
 
+class UBlendSpace;
+
 UCLASS()
 class UPoseSearchMeshComponent : public UPoseableMeshComponent
 {
diff --git a/Engine/Plugins/Experimental/Animation/PoseSearch/Source/Editor/Private/PoseSearchTypeActions.cpp b/Engine/Plugins/Experimental/Animation/PoseSearch/Source/Editor/Private/PoseSearchTypeActions.cpp
--- a/Engine/Plugins/Experimental/Animation/PoseSearch/Source/Editor/Private/PoseSearchTypeActions.cpp
+++ b/Engine/Plugins/Experimental/Animation/PoseSearch/Source/Editor/Private/PoseSearchTypeActions.cpp
@@ -2,7 +2,7 @@
 
 #include "PoseSearchTypeActions.h"
 #include "PoseSearchDatabaseEditor.h"
-#include "PoseSearch/PoseSearch.h"
+#include "PoseSearch/PoseSearchDatabase.h"
 #include "PoseSearch/PoseSearchSchema.h"
 
 #define LOCTEXT_NAMESPACE "PoseSearchTypeActions"
diff --git a/Engine/Plugins/Experimental/Animation/PoseSearch/Source/Editor/Public/PoseSearchDatabaseAssetTreeNode.h b/Engine/Plugins/Experimental/Animation/PoseSearch/Source/Editor/Public/PoseSearchDatabaseAssetTreeNode.h
--- a/Engine/Plugins/Experimental/Animation/PoseSearch/Source/Editor/Public/PoseSearchDatabaseAssetTreeNode.h
+++ b/Engine/Plugins/Experimental/Animation/PoseSearch/Source/Editor/Public/PoseSearchDatabaseAssetTreeNode.h
@@ -3,8 +3,9 @@
 #pragma once
 
 #include "CoreMinimal.h"
-#include "PoseSearch/PoseSearch.h"
 
+enum class EPoseSearchMirrorOption;
+enum class ESearchIndexAssetType;
 class FUICommandList;
 class ITableRow;
 class STableViewBase;
diff --git a/Engine/Plugins/Experimental/Animation/PoseSearch/Source/Editor/Public/PoseSearchDatabaseEditorReflection.h b/Engine/Plugins/Experimental/Animation/PoseSearch/Source/Editor/Public/PoseSearchDatabaseEditorReflection.h
--- a/Engine/Plugins/Experimental/Animation/PoseSearch/Source/Editor/Public/PoseSearchDatabaseEditorReflection.h
+++ b/Engine/Plugins/Experimental/Animation/PoseSearch/Source/Editor/Public/PoseSearchDatabaseEditorReflection.h
@@ -3,9 +3,7 @@
 #pragma once
 
 #include "CoreMinimal.h"
-
-#include "PoseSearch/PoseSearch.h"
-
+#include "PoseSearch/PoseSearchDatabase.h"
 #include "PoseSearchDatabaseEditorReflection.generated.h"
 
 
diff --git a/Engine/Plugins/Experimental/Animation/PoseSearch/Source/Editor/Public/PoseSearchDatabaseViewModel.h b/Engine/Plugins/Experimental/Animation/PoseSearch/Source/Editor/Public/PoseSearchDatabaseViewModel.h
--- a/Engine/Plugins/Experimental/Animation/PoseSearch/Source/Editor/Public/PoseSearchDatabaseViewModel.h
+++ b/Engine/Plugins/Experimental/Animation/PoseSearch/Source/Editor/Public/PoseSearchDatabaseViewModel.h
@@ -5,9 +5,9 @@
 #include "CoreMinimal.h"
 #include "UObject/GCObject.h"
 #include "MovieSceneFwd.h"
-
+#include "PoseSearch/PoseSearchAssetSampler.h"
+#include "PoseSearch/PoseSearchIndex.h"
 #include "PoseSearchDatabasePreviewScene.h"
-#include "PoseSearch/PoseSearch.h"
 
 class UWorld;
 class UPoseSearchDatabase;
diff --git a/Engine/Plugins/Experimental/Animation/PoseSearch/Source/Editor/Public/PoseSearchEditor.h b/Engine/Plugins/Experimental/Animation/PoseSearch/Source/Editor/Public/PoseSearchEditor.h
--- a/Engine/Plugins/Experimental/Animation/PoseSearch/Source/Editor/Public/PoseSearchEditor.h
+++ b/Engine/Plugins/Experimental/Animation/PoseSearch/Source/Editor/Public/PoseSearchEditor.h
@@ -4,8 +4,6 @@
 
 #include "CoreMinimal.h"
 #include "Modules/ModuleInterface.h"
-#include "PoseSearch/PoseSearch.h"
-
 
 DECLARE_LOG_CATEGORY_EXTERN(LogPoseSearchEditor, Log, All);
 
diff --git a/Engine/Plugins/Experimental/Animation/PoseSearch/Source/Runtime/Private/AnimNode_BlendStack.cpp b/Engine/Plugins/Experimental/Animation/PoseSearch/Source/Runtime/Private/AnimNode_BlendStack.cpp
--- a/Engine/Plugins/Experimental/Animation/PoseSearch/Source/Runtime/Private/AnimNode_BlendStack.cpp
+++ b/Engine/Plugins/Experimental/Animation/PoseSearch/Source/Runtime/Private/AnimNode_BlendStack.cpp
@@ -1,12 +1,12 @@
 // Copyright Epic Games, Inc. All Rights Reserved.
 
 #include "PoseSearch/AnimNode_BlendStack.h"
-
 #include "Algo/MaxElement.h"
 #include "Animation/AnimInstanceProxy.h"
 #include "Animation/AnimComposite.h"
 #include "Animation/AnimSequence.h"
 #include "Animation/BlendSpace.h"
+#include "PoseSearch/PoseSearchDefines.h"
 
 #include UE_INLINE_GENERATED_CPP_BY_NAME(AnimNode_BlendStack)
 
diff --git a/Engine/Plugins/Experimental/Animation/PoseSearch/Source/Runtime/Private/AnimNode_MotionMatching.cpp b/Engine/Plugins/Experimental/Animation/PoseSearch/Source/Runtime/Private/AnimNode_MotionMatching.cpp
--- a/Engine/Plugins/Experimental/Animation/PoseSearch/Source/Runtime/Private/AnimNode_MotionMatching.cpp
+++ b/Engine/Plugins/Experimental/Animation/PoseSearch/Source/Runtime/Private/AnimNode_MotionMatching.cpp
@@ -11,6 +11,7 @@
 #include "Components/SkeletalMeshComponent.h"
 #include "DrawDebugHelpers.h"
 #include "PoseSearch/AnimNode_PoseSearchHistoryCollector.h"
+#include "PoseSearch/PoseSearchDatabase.h"
 #include "PoseSearch/PoseSearchDerivedData.h"
 #include "PoseSearch/PoseSearchSchema.h"
 #include "Trace/PoseSearchTraceLogger.h"
diff --git a/Engine/Plugins/Experimental/Animation/PoseSearch/Source/Runtime/Private/PoseSearch.cpp b/Engine/Plugins/Experimental/Animation/PoseSearch/Source/Runtime/Private/PoseSearch.cpp
deleted file mode 100644
diff --git a/Engine/Plugins/Experimental/Animation/PoseSearch/Source/Runtime/Private/PoseSearchAssetIndexer.cpp b/Engine/Plugins/Experimental/Animation/PoseSearch/Source/Runtime/Private/PoseSearchAssetIndexer.cpp
new file mode 100644
--- /dev/null
+++ b/Engine/Plugins/Experimental/Animation/PoseSearch/Source/Runtime/Private/PoseSearchAssetIndexer.cpp
@@ -0,0 +1,333 @@
+// Copyright Epic Games, Inc. All Rights Reserved.
+
+#include "PoseSearchAssetIndexer.h"
+
+#if WITH_EDITOR
+
+#include "AnimationRuntime.h"
+#include "Animation/MirrorDataTable.h"
+#include "PoseSearch/PoseSearchAnimNotifies.h"
+#include "PoseSearch/PoseSearchAssetSampler.h"
+#include "PoseSearch/PoseSearchDefines.h"
+#include "PoseSearch/PoseSearchFeatureChannel.h"
+#include "PoseSearch/PoseSearchSchema.h"
+
+namespace UE::PoseSearch
+{
+	
+//////////////////////////////////////////////////////////////////////////
+// FSamplingParam helpers
+struct FSamplingParam
+{
+	float WrappedParam = 0.0f;
+	int32 NumCycles = 0;
+
+	// If the animation can't loop, WrappedParam contains the clamped value and whatever is left is stored here
+	float Extrapolation = 0.0f;
+};
+
+static FSamplingParam WrapOrClampSamplingParam(bool bCanWrap, float SamplingParamExtent, float SamplingParam)
+{
+	// This is a helper function used by both time and distance sampling. A schema may specify time or distance
+	// offsets that are multiple cycles of a clip away from the current pose being sampled.
+	// And that time or distance offset may before the beginning of the clip (SamplingParam < 0.0f)
+	// or after the end of the clip (SamplingParam > SamplingParamExtent). So this function
+	// helps determine how many cycles need to be applied and what the wrapped value should be, clamping
+	// if necessary.
+
+	FSamplingParam Result;
+
+	Result.WrappedParam = SamplingParam;
+
+	const bool bIsSamplingParamExtentKindaSmall = SamplingParamExtent <= UE_KINDA_SMALL_NUMBER;
+	if (!bIsSamplingParamExtentKindaSmall && bCanWrap)
+	{
+		if (SamplingParam < 0.0f)
+		{
+			while (Result.WrappedParam < 0.0f)
+			{
+				Result.WrappedParam += SamplingParamExtent;
+				++Result.NumCycles;
+			}
+		}
+
+		else
+		{
+			while (Result.WrappedParam > SamplingParamExtent)
+			{
+				Result.WrappedParam -= SamplingParamExtent;
+				++Result.NumCycles;
+			}
+		}
+	}
+
+	const float ParamClamped = FMath::Clamp(Result.WrappedParam, 0.0f, SamplingParamExtent);
+	if (ParamClamped != Result.WrappedParam)
+	{
+		check(bIsSamplingParamExtentKindaSmall || !bCanWrap);
+		Result.Extrapolation = Result.WrappedParam - ParamClamped;
+		Result.WrappedParam = ParamClamped;
+	}
+
+	return Result;
+}
+
+//////////////////////////////////////////////////////////////////////////
+// FAssetIndexer
+void FAssetIndexer::Reset()
+{
+	Output.FirstIndexedSample = 0;
+	Output.LastIndexedSample = 0;
+	Output.NumIndexedPoses = 0;
+
+	Output.FeatureVectorTable.Reset(0);
+	Output.PoseMetadata.Reset(0);
+	Output.AllFeaturesNotAdded.Reset();
+}
+
+void FAssetIndexer::Init(const FAssetIndexingContext& InIndexingContext, const FBoneContainer& InBoneContainer)
+{
+	check(InIndexingContext.Schema);
+	check(InIndexingContext.Schema->IsValid());
+	check(InIndexingContext.AssetSampler);
+
+	BoneContainer = InBoneContainer;
+	IndexingContext = InIndexingContext;
+
+	Reset();
+
+	Output.FirstIndexedSample = FMath::FloorToInt(IndexingContext.RequestedSamplingRange.Min * IndexingContext.Schema->SampleRate);
+	Output.LastIndexedSample = FMath::Max(0, FMath::CeilToInt(IndexingContext.RequestedSamplingRange.Max * IndexingContext.Schema->SampleRate));
+	Output.NumIndexedPoses = Output.LastIndexedSample - Output.FirstIndexedSample + 1;
+
+	Output.FeatureVectorTable.SetNumZeroed(IndexingContext.Schema->SchemaCardinality * Output.NumIndexedPoses);
+	Output.PoseMetadata.SetNum(Output.NumIndexedPoses);
+}
+
+bool FAssetIndexer::Process()
+{
+	check(IndexingContext.Schema);
+	check(IndexingContext.Schema->IsValid());
+	check(IndexingContext.AssetSampler);
+
+	FMemMark Mark(FMemStack::Get());
+
+	IndexingContext.BeginSampleIdx = Output.FirstIndexedSample;
+	IndexingContext.EndSampleIdx = Output.LastIndexedSample + 1;
+
+	if (IndexingContext.Schema->SchemaCardinality > 0)
+	{
+		// Index each channel
+		for (int32 ChannelIdx = 0; ChannelIdx != IndexingContext.Schema->Channels.Num(); ++ChannelIdx)
+		{
+			IndexingContext.Schema->Channels[ChannelIdx]->IndexAsset(*this, Output.FeatureVectorTable);
+		}
+	}
+
+	// Generate pose metadata
+	for (int32 SampleIdx = IndexingContext.BeginSampleIdx; SampleIdx != IndexingContext.EndSampleIdx; ++SampleIdx)
+	{
+		const int32 PoseIdx = SampleIdx - Output.FirstIndexedSample;
+		FPoseSearchPoseMetadata& OutputPoseMetadata = Output.PoseMetadata[PoseIdx];
+		OutputPoseMetadata = GetMetadata(SampleIdx);
+	}
+
+	return true;
+}
+
+FAssetIndexer::FSampleInfo FAssetIndexer::GetSampleInfo(float SampleTime) const
+{
+	FSampleInfo Sample;
+
+	check(IndexingContext.AssetSampler);
+
+	const float PlayLength = IndexingContext.AssetSampler->GetPlayLength();
+	const bool bCanWrap = IndexingContext.AssetSampler->IsLoopable();
+
+	float MainRelativeTime = SampleTime;
+	if (SampleTime < 0.0f && bCanWrap)
+	{
+		// In this case we're sampling a loop backwards, so MainRelativeTime must adjust so the number of cycles is
+		// counted correctly.
+		MainRelativeTime += PlayLength;
+	}
+
+	const FSamplingParam SamplingParam = WrapOrClampSamplingParam(bCanWrap, PlayLength, MainRelativeTime);
+
+	Sample.Clip = IndexingContext.AssetSampler;
+
+	if (FMath::Abs(SamplingParam.Extrapolation) > SMALL_NUMBER)
+	{
+		Sample.bClamped = true;
+		Sample.ClipTime = SamplingParam.WrappedParam + SamplingParam.Extrapolation;
+		const FTransform ClipRootMotion = Sample.Clip->ExtractRootTransform(Sample.ClipTime);
+		const float ClipDistance = Sample.Clip->ExtractRootDistance(Sample.ClipTime);
+
+		Sample.RootTransform = ClipRootMotion;
+		Sample.RootDistance = ClipDistance;
+	}
+	else
+	{
+		Sample.ClipTime = SamplingParam.WrappedParam;
+
+		const FTransform RootMotionLast = IndexingContext.AssetSampler->GetTotalRootTransform();
+		float RootDistanceLast = IndexingContext.AssetSampler->GetTotalRootDistance();
+
+		// Determine how to accumulate motion for every cycle of the anim. If the sample
+		// had to be clamped, this motion will end up not getting applied below.
+		// Also invert the accumulation direction if the requested sample was wrapped backwards.
+		FTransform RootMotionPerCycle = RootMotionLast;
+		float RootDistancePerCycle = RootDistanceLast;
+		if (SampleTime < 0.0f)
+		{
+			RootMotionPerCycle = RootMotionPerCycle.Inverse();
+			RootDistancePerCycle *= -1.f;
+		}
+
+		// Find the remaining motion deltas after wrapping
+		FTransform RootMotionRemainder = Sample.Clip->ExtractRootTransform(Sample.ClipTime);
+		float RootDistanceRemainder = Sample.Clip->ExtractRootDistance(Sample.ClipTime);
+
+		// Invert motion deltas if we wrapped backwards
+		if (SampleTime < 0.0f)
+		{
+			RootMotionRemainder.SetToRelativeTransform(RootMotionLast);
+			RootDistanceRemainder = -(RootDistanceLast - RootDistanceRemainder);
+		}
+
+		Sample.RootTransform = FTransform::Identity;
+		Sample.RootDistance = 0.f;
+
+		// Note if the sample was clamped, no motion will be applied here because NumCycles will be zero
+		int32 CyclesRemaining = SamplingParam.NumCycles;
+		while (CyclesRemaining--)
+		{
+			Sample.RootTransform = RootMotionPerCycle * Sample.RootTransform;
+			Sample.RootDistance += RootDistancePerCycle;
+		}
+
+		Sample.RootTransform = RootMotionRemainder * Sample.RootTransform;
+		Sample.RootDistance += RootDistanceRemainder;
+	}
+
+	return Sample;
+}
+
+FAssetIndexer::FSampleInfo FAssetIndexer::GetSampleInfoRelative(float SampleTime, const FSampleInfo& Origin) const
+{
+	FSampleInfo Sample = GetSampleInfo(SampleTime);
+	Sample.RootTransform.SetToRelativeTransform(Origin.RootTransform);
+	Sample.RootDistance = Origin.RootDistance - Sample.RootDistance;
+	return Sample;
+}
+
+FTransform FAssetIndexer::MirrorTransform(const FTransform& Transform) const
+{
+	return IndexingContext.bMirrored ? IndexingContext.SamplingContext->MirrorTransform(Transform) : Transform;
+}
+
+FPoseSearchPoseMetadata FAssetIndexer::GetMetadata(int32 SampleIdx) const
+{
+	const float SequenceLength = IndexingContext.AssetSampler->GetPlayLength();
+	const float SampleTime = FMath::Min(SampleIdx * IndexingContext.Schema->GetSamplingInterval(), SequenceLength);
+
+	FPoseSearchPoseMetadata Metadata;
+	Metadata.CostAddend = IndexingContext.Schema->BaseCostBias;
+	Metadata.ContinuingPoseCostAddend = IndexingContext.Schema->ContinuingPoseCostBias;
+
+	TArray<UAnimNotifyState_PoseSearchBase*> NotifyStates;
+	IndexingContext.AssetSampler->ExtractPoseSearchNotifyStates(SampleTime, NotifyStates);
+	for (const UAnimNotifyState_PoseSearchBase* PoseSearchNotify : NotifyStates)
+	{
+		if (PoseSearchNotify->GetClass()->IsChildOf<UAnimNotifyState_PoseSearchBlockTransition>())
+		{
+			EnumAddFlags(Metadata.Flags, EPoseSearchPoseFlags::BlockTransition);
+		}
+		else if (PoseSearchNotify->GetClass()->IsChildOf<UAnimNotifyState_PoseSearchModifyCost>())
+		{
+			const UAnimNotifyState_PoseSearchModifyCost* ModifyCostNotify =
+				Cast<const UAnimNotifyState_PoseSearchModifyCost>(PoseSearchNotify);
+			Metadata.CostAddend = ModifyCostNotify->CostAddend;
+		}
+		else if (PoseSearchNotify->GetClass()->IsChildOf<UAnimNotifyState_PoseSearchOverrideContinuingPoseCostBias>())
+		{
+			const UAnimNotifyState_PoseSearchOverrideContinuingPoseCostBias* ContinuingPoseCostBias =
+				Cast<const UAnimNotifyState_PoseSearchOverrideContinuingPoseCostBias>(PoseSearchNotify);
+			Metadata.ContinuingPoseCostAddend = ContinuingPoseCostBias->CostAddend;
+		}
+	}
+	return Metadata;
+}
+
+FTransform FAssetIndexer::GetTransformAndCacheResults(float SampleTime, float OriginTime, int8 SchemaBoneIdx, bool& Clamped)
+{
+	// @todo: use an hashmap if we end up having too many entries
+	CachedEntry* Entry = CachedEntries.FindByPredicate([SampleTime, OriginTime](const FAssetIndexer::CachedEntry& Entry)
+		{
+			return Entry.SampleTime == SampleTime && Entry.OriginTime == OriginTime;
+		});
+
+	const FAssetSamplingContext* SamplingContext = IndexingContext.SamplingContext;
+
+	if (!Entry)
+	{
+		Entry = &CachedEntries[CachedEntries.AddDefaulted()];
+
+		Entry->SampleTime = SampleTime;
+		Entry->OriginTime = OriginTime;
+
+		if (!BoneContainer.IsValid())
+		{
+			UE_LOG(LogPoseSearch,
+				Warning,
+				TEXT("Invalid BoneContainer encountered in FAssetIndexer::GetTransformAndCacheResults. Asset: %s. Schema: %s. BoneContainerAsset: %s. NumBoneIndices: %d"),
+				*GetNameSafe(IndexingContext.AssetSampler->GetAsset()),
+				*GetNameSafe(IndexingContext.Schema),
+				*GetNameSafe(BoneContainer.GetAsset()),
+				BoneContainer.GetCompactPoseNumBones());
+		}
+
+		Entry->Pose.SetBoneContainer(&BoneContainer);
+		Entry->UnusedCurve.InitFrom(BoneContainer);
+
+		IAssetIndexer::FSampleInfo Origin = GetSampleInfo(OriginTime);
+		IAssetIndexer::FSampleInfo Sample = GetSampleInfoRelative(SampleTime, Origin);
+
+		float CurrentTime = Sample.ClipTime;
+		float PreviousTime = CurrentTime - SamplingContext->FiniteDelta;
+
+		check(Sample.Clip->IsLoopable() || PreviousTime <= Sample.Clip->GetPlayLength());
+
+		FDeltaTimeRecord DeltaTimeRecord;
+		DeltaTimeRecord.Set(PreviousTime, CurrentTime - PreviousTime);
+		FAnimExtractContext ExtractionCtx(static_cast<double>(CurrentTime), true, DeltaTimeRecord, Sample.Clip->IsLoopable());
+
+		Sample.Clip->ExtractPose(ExtractionCtx, Entry->AnimPoseData);
+
+		if (IndexingContext.bMirrored)
+		{
+			FAnimationRuntime::MirrorPose(
+				Entry->AnimPoseData.GetPose(),
+				IndexingContext.Schema->MirrorDataTable->MirrorAxis,
+				SamplingContext->CompactPoseMirrorBones,
+				SamplingContext->ComponentSpaceRefRotations
+			);
+			// Note curves and attributes are not used during the indexing process and therefore don't need to be mirrored
+		}
+
+		Entry->ComponentSpacePose.InitPose(Entry->Pose);
+		Entry->RootTransform = Sample.RootTransform;
+		Entry->Clamped = Sample.bClamped;
+	}
+
+	const FBoneReference& BoneReference = IndexingContext.Schema->BoneReferences[SchemaBoneIdx];
+	FCompactPoseBoneIndex CompactBoneIndex = BoneContainer.MakeCompactPoseIndex(FMeshPoseBoneIndex(BoneReference.BoneIndex));
+
+	const FTransform BoneTransform = Entry->ComponentSpacePose.GetComponentSpaceTransform(CompactBoneIndex) * MirrorTransform(Entry->RootTransform);
+	Clamped = Entry->Clamped;
+
+	return BoneTransform;
+}
+
+} // namespace UE::PoseSearch
+#endif // WITH_EDITOR
diff --git a/Engine/Plugins/Experimental/Animation/PoseSearch/Source/Runtime/Private/PoseSearchAssetIndexer.h b/Engine/Plugins/Experimental/Animation/PoseSearch/Source/Runtime/Private/PoseSearchAssetIndexer.h
new file mode 100644
--- /dev/null
+++ b/Engine/Plugins/Experimental/Animation/PoseSearch/Source/Runtime/Private/PoseSearchAssetIndexer.h
@@ -0,0 +1,70 @@
+// Copyright Epic Games, Inc. All Rights Reserved.
+
+#pragma once
+
+#if WITH_EDITOR
+
+#include "Animation/AnimationPoseData.h"
+#include "Animation/AnimCurveTypes.h"
+#include "Animation/AttributesRuntime.h"
+#include "BonePose.h"
+#include "PoseSearch/PoseSearchAssetIndexer.h"
+#include "PoseSearch/PoseSearchIndex.h"
+
+namespace UE::PoseSearch
+{
+struct FAssetIndexingContext;
+
+class FAssetIndexer : public IAssetIndexer
+{
+public:
+
+	struct FOutput
+	{
+		int32 FirstIndexedSample = 0;
+		int32 LastIndexedSample = 0;
+		int32 NumIndexedPoses = 0;
+
+		TArray<float> FeatureVectorTable;
+		TArray<FPoseSearchPoseMetadata> PoseMetadata;
+		TBitArray<> AllFeaturesNotAdded;
+	} Output;
+
+	void Reset();
+	void Init(const FAssetIndexingContext& IndexingContext, const FBoneContainer& InBoneContainer);
+	bool Process();
+
+public: // IAssetIndexer
+
+	const FAssetIndexingContext& GetIndexingContext() const override { return IndexingContext; }
+	FSampleInfo GetSampleInfo(float SampleTime) const override;
+	FSampleInfo GetSampleInfoRelative(float SampleTime, const FSampleInfo& Origin) const override;
+	FTransform MirrorTransform(const FTransform& Transform) const override;
+	FTransform GetTransformAndCacheResults(float SampleTime, float OriginTime, int8 SchemaBoneIdx, bool& Clamped) override;
+
+private:
+	FPoseSearchPoseMetadata GetMetadata(int32 SampleIdx) const;
+
+	struct CachedEntry
+	{
+		float SampleTime;
+		float OriginTime;
+		bool Clamped;
+
+		// @todo: minimize the Entry memory footprint
+		FTransform RootTransform;
+		FCompactPose Pose;
+		FCSPose<FCompactPose> ComponentSpacePose;
+		FBlendedCurve UnusedCurve;
+		UE::Anim::FStackAttributeContainer UnusedAtrribute;
+		FAnimationPoseData AnimPoseData = { Pose, UnusedCurve, UnusedAtrribute };
+	};
+
+	FBoneContainer BoneContainer;
+	FAssetIndexingContext IndexingContext;
+	TArray<CachedEntry> CachedEntries;
+};
+
+} // namespace UE::PoseSearch
+
+#endif // WITH_EDITOR
\ No newline at end of file
diff --git a/Engine/Plugins/Experimental/Animation/PoseSearch/Source/Runtime/Private/PoseSearchAssetSampler.cpp b/Engine/Plugins/Experimental/Animation/PoseSearch/Source/Runtime/Private/PoseSearchAssetSampler.cpp
new file mode 100644
--- /dev/null
+++ b/Engine/Plugins/Experimental/Animation/PoseSearch/Source/Runtime/Private/PoseSearchAssetSampler.cpp
@@ -0,0 +1,762 @@
+// Copyright Epic Games, Inc. All Rights Reserved.
+
+#include "PoseSearch/PoseSearchAssetSampler.h"
+#include "AnimationRuntime.h"
+#include "Animation/AnimSequence.h"
+#include "Animation/AnimRootMotionProvider.h"
+#include "Animation/BlendSpace.h"
+#include "Animation/MirrorDataTable.h"
+#include "PoseSearch/PoseSearchAnimNotifies.h"
+
+namespace UE::PoseSearch
+{
+
+//////////////////////////////////////////////////////////////////////////
+// Root motion extrapolation
+
+// Uses distance delta between NextRootDistanceIndex and NextRootDistanceIndex - 1 and extrapolates it to ExtrapolationTime
+static float ExtrapolateAccumulatedRootDistance(
+	int32 SamplingRate,
+	TConstArrayView<float> AccumulatedRootDistance,
+	int32 NextRootDistanceIndex, 
+	float ExtrapolationTime,
+	const FPoseSearchExtrapolationParameters& ExtrapolationParameters)
+{
+	check(NextRootDistanceIndex > 0 && NextRootDistanceIndex < AccumulatedRootDistance.Num());
+
+	const float DistanceDelta =
+		AccumulatedRootDistance[NextRootDistanceIndex] -
+		AccumulatedRootDistance[NextRootDistanceIndex - 1];
+	const float Speed = DistanceDelta * SamplingRate;
+	const float ExtrapolationSpeed = Speed >= ExtrapolationParameters.LinearSpeedThreshold ?
+		Speed : 0.0f;
+	const float ExtrapolatedDistance = ExtrapolationSpeed * ExtrapolationTime;
+
+	return ExtrapolatedDistance;
+}
+
+static float ExtractAccumulatedRootDistance(
+	int32 SamplingRate,
+	TConstArrayView<float> AccumulatedRootDistance,
+	float PlayLength,
+	float Time,
+	const FPoseSearchExtrapolationParameters& ExtrapolationParameters)
+{
+	const float ClampedTime = FMath::Clamp(Time, 0.0f, PlayLength);
+
+	// Find the distance sample that corresponds with the time and split into whole and partial parts
+	float IntegralDistanceSample;
+	float DistanceAlpha = FMath::Modf(ClampedTime * SamplingRate, &IntegralDistanceSample);
+	float DistanceIdx = (int32)IntegralDistanceSample;
+
+	// Verify the distance offset and any residual portion would be in bounds
+	check(DistanceIdx + (DistanceAlpha > 0.0f ? 1 : 0) < AccumulatedRootDistance.Num());
+
+	// Look up the distance and interpolate between distance samples if necessary
+	float Distance = AccumulatedRootDistance[DistanceIdx];
+	if (DistanceAlpha > 0.0f)
+	{
+		float NextDistance = AccumulatedRootDistance[DistanceIdx + 1];
+		Distance = FMath::Lerp(Distance, NextDistance, DistanceAlpha);
+	}
+
+	const float ExtrapolationTime = Time - ClampedTime;
+
+	if (ExtrapolationTime != 0.0f)
+	{
+		// If extrapolationTime is not zero, we extrapolate the beginning or the end of the animation to estimate
+		// the root distance.
+		const int32 DistIdx = (ExtrapolationTime > 0.0f) ? AccumulatedRootDistance.Num() - 1 : 1;
+		const float ExtrapolatedDistance = ExtrapolateAccumulatedRootDistance(
+			SamplingRate,
+			AccumulatedRootDistance,
+			DistIdx,
+			ExtrapolationTime,
+			ExtrapolationParameters);
+		Distance += ExtrapolatedDistance;
+	}
+
+	return Distance;
+}
+
+static FTransform ExtrapolateRootMotion(
+	FTransform SampleToExtrapolate,
+	float SampleStart, 
+	float SampleEnd, 
+	float ExtrapolationTime,
+	const FPoseSearchExtrapolationParameters& ExtrapolationParameters)
+{
+	const float SampleDelta = SampleEnd - SampleStart;
+	check(!FMath::IsNearlyZero(SampleDelta));
+
+	const FVector LinearVelocityToExtrapolate = SampleToExtrapolate.GetTranslation() / SampleDelta;
+	const float LinearSpeedToExtrapolate = LinearVelocityToExtrapolate.Size();
+	const bool bCanExtrapolateTranslation =
+		LinearSpeedToExtrapolate >= ExtrapolationParameters.LinearSpeedThreshold;
+
+	const float AngularSpeedToExtrapolateRad = SampleToExtrapolate.GetRotation().GetAngle() / SampleDelta;
+	const bool bCanExtrapolateRotation =
+		FMath::RadiansToDegrees(AngularSpeedToExtrapolateRad) >= ExtrapolationParameters.AngularSpeedThreshold;
+
+	if (!bCanExtrapolateTranslation && !bCanExtrapolateRotation)
+	{
+		return FTransform::Identity;
+	}
+
+	if (!bCanExtrapolateTranslation)
+	{
+		SampleToExtrapolate.SetTranslation(FVector::ZeroVector);
+	}
+
+	if (!bCanExtrapolateRotation)
+	{
+		SampleToExtrapolate.SetRotation(FQuat::Identity);
+	}
+
+	// converting ExtrapolationTime to a positive number to avoid dealing with the negative extrapolation and inverting
+	// transforms later on.
+	const float AbsExtrapolationTime = FMath::Abs(ExtrapolationTime);
+	const float AbsSampleDelta = FMath::Abs(SampleDelta);
+	const FTransform AbsTimeSampleToExtrapolate =
+		ExtrapolationTime >= 0.0f ? SampleToExtrapolate : SampleToExtrapolate.Inverse();
+
+	// because we're extrapolating rotation, the extrapolation must be integrated over time
+	const float SampleMultiplier = AbsExtrapolationTime / AbsSampleDelta;
+	float IntegralNumSamples;
+	float RemainingSampleFraction = FMath::Modf(SampleMultiplier, &IntegralNumSamples);
+	int32 NumSamples = (int32)IntegralNumSamples;
+
+	// adding full samples to the extrapolated root motion
+	FTransform ExtrapolatedRootMotion = FTransform::Identity;
+	for (int i = 0; i < NumSamples; ++i)
+	{
+		ExtrapolatedRootMotion = AbsTimeSampleToExtrapolate * ExtrapolatedRootMotion;
+	}
+
+	// and a blend with identity for whatever is left
+	FTransform RemainingExtrapolatedRootMotion;
+	RemainingExtrapolatedRootMotion.Blend(
+		FTransform::Identity,
+		AbsTimeSampleToExtrapolate,
+		RemainingSampleFraction);
+
+	ExtrapolatedRootMotion = RemainingExtrapolatedRootMotion * ExtrapolatedRootMotion;
+	return ExtrapolatedRootMotion;
+}
+
+//////////////////////////////////////////////////////////////////////////
+// FAssetSamplingContext
+void FAssetSamplingContext::Init(const UMirrorDataTable* InMirrorDataTable, const FBoneContainer& BoneContainer)
+{
+	MirrorDataTable = InMirrorDataTable;
+
+	if (InMirrorDataTable)
+	{
+		InMirrorDataTable->FillCompactPoseAndComponentRefRotations(BoneContainer, CompactPoseMirrorBones, ComponentSpaceRefRotations);
+	}
+	else
+	{
+		CompactPoseMirrorBones.Reset();
+		ComponentSpaceRefRotations.Reset();
+	}
+}
+
+FTransform FAssetSamplingContext::MirrorTransform(const FTransform& InTransform) const
+{
+	const EAxis::Type MirrorAxis = MirrorDataTable->MirrorAxis;
+	FVector T = InTransform.GetTranslation();
+	T = FAnimationRuntime::MirrorVector(T, MirrorAxis);
+	const FQuat ReferenceRotation = ComponentSpaceRefRotations[FCompactPoseBoneIndex(0)];
+	FQuat Q = InTransform.GetRotation();
+	Q = FAnimationRuntime::MirrorQuat(Q, MirrorAxis);
+	Q *= FAnimationRuntime::MirrorQuat(ReferenceRotation, MirrorAxis).Inverse() * ReferenceRotation;
+	FTransform Result = FTransform(Q, T, InTransform.GetScale3D());
+	return Result;
+}
+
+//////////////////////////////////////////////////////////////////////////
+// FSequenceBaseSampler
+void FSequenceBaseSampler::Init(const FInput& InInput)
+{
+	check(InInput.SequenceBase.Get());
+
+	Input = InInput;
+}
+
+void FSequenceBaseSampler::Process()
+{
+	ProcessRootDistance();
+}
+
+float FSequenceBaseSampler::GetTimeFromRootDistance(float Distance) const
+{
+	int32 NextSampleIdx = 1;
+	int32 PrevSampleIdx = 0;
+	if (Distance > 0.0f)
+	{
+		// Search for the distance value. Because the values will be extrapolated if necessary
+		// LowerBound might go past the end of the array, in which case the last valid index is used
+		int32 ClipDistanceLowerBoundIndex = Algo::LowerBound(AccumulatedRootDistance, Distance);
+		NextSampleIdx = FMath::Min(
+			ClipDistanceLowerBoundIndex,
+			AccumulatedRootDistance.Num() - 1);
+
+		// Compute distance interpolation amount
+		PrevSampleIdx = FMath::Max(0, NextSampleIdx - 1);
+	}
+
+	float NextDistance = AccumulatedRootDistance[NextSampleIdx];
+	float PrevDistance = AccumulatedRootDistance[PrevSampleIdx];
+	float DistanceSampleAlpha = FMath::GetRangePct(PrevDistance, NextDistance, Distance);
+
+	// Convert to time
+	float ClipTime = (float(NextSampleIdx) - (1.0f - DistanceSampleAlpha)) / Input.RootDistanceSamplingRate;
+	return ClipTime;
+}
+
+float FSequenceBaseSampler::GetPlayLength() const
+{
+	return Input.SequenceBase->GetPlayLength();
+}
+
+bool FSequenceBaseSampler::IsLoopable() const
+{
+	return Input.SequenceBase->bLoop;
+}
+
+void FSequenceBaseSampler::ExtractPose(const FAnimExtractContext& ExtractionCtx, FAnimationPoseData& OutAnimPoseData) const
+{
+	Input.SequenceBase->GetAnimationPose(OutAnimPoseData, ExtractionCtx);
+}
+
+FTransform FSequenceBaseSampler::ExtractRootTransform(float Time) const
+{
+	FTransform RootTransform = FTransform::Identity;
+	const UAnimSequenceBase* SequenceBase = Input.SequenceBase.Get();
+
+	if (IsLoopable())
+	{
+		FTransform LoopableRootTransform = SequenceBase->ExtractRootMotion(0.0f, Time, true);
+		return LoopableRootTransform;
+	}
+
+	const float ExtrapolationSampleTime = Input.ExtrapolationParameters.SampleTime;
+
+	const float PlayLength = SequenceBase->GetPlayLength();
+	const float ClampedTime = FMath::Clamp(Time, 0.0f, PlayLength);
+	const float ExtrapolationTime = Time - ClampedTime;
+
+	// If Time is less than zero, ExtrapolationTime will be negative. In this case, we extrapolate the beginning of the 
+	// animation to estimate where the root would be at Time
+	if (ExtrapolationTime < -SMALL_NUMBER)
+	{
+		FTransform SampleToExtrapolate = SequenceBase->ExtractRootMotionFromRange(0.0f, ExtrapolationSampleTime);
+
+		const FTransform ExtrapolatedRootMotion = ExtrapolateRootMotion(
+			SampleToExtrapolate,
+			0.0f, ExtrapolationSampleTime,
+			ExtrapolationTime,
+			Input.ExtrapolationParameters);
+		RootTransform = ExtrapolatedRootMotion;
+	}
+	else
+	{
+		RootTransform = SequenceBase->ExtractRootMotionFromRange(0.0f, ClampedTime);
+
+		// If Time is greater than PlayLength, ExtrapolationTIme will be a positive number. In this case, we extrapolate
+		// the end of the animation to estimate where the root would be at Time
+		if (ExtrapolationTime > SMALL_NUMBER)
+		{
+			FTransform SampleToExtrapolate = SequenceBase->ExtractRootMotionFromRange(PlayLength - ExtrapolationSampleTime, PlayLength);
+
+			const FTransform ExtrapolatedRootMotion = ExtrapolateRootMotion(
+				SampleToExtrapolate,
+				PlayLength - ExtrapolationSampleTime, PlayLength,
+				ExtrapolationTime,
+				Input.ExtrapolationParameters);
+			RootTransform = ExtrapolatedRootMotion * RootTransform;
+		}
+	}
+
+	return RootTransform;
+}
+
+float FSequenceBaseSampler::ExtractRootDistance(float Time) const
+{
+	return ExtractAccumulatedRootDistance(
+		Input.RootDistanceSamplingRate,
+		AccumulatedRootDistance,
+		Input.SequenceBase->GetPlayLength(),
+		Time,
+		Input.ExtrapolationParameters);
+}
+
+void FSequenceBaseSampler::ExtractPoseSearchNotifyStates(
+	float Time, 
+	TArray<UAnimNotifyState_PoseSearchBase*>& NotifyStates) const
+{
+	// getting pose search notifies in an interval of size ExtractionInterval, centered on Time
+	constexpr float ExtractionInterval = 1.0f / 120.0f;
+	FAnimNotifyContext NotifyContext;
+	Input.SequenceBase->GetAnimNotifies(Time - (ExtractionInterval * 0.5f), ExtractionInterval, NotifyContext);
+
+	// check which notifies actually overlap Time and are of the right base type
+	for (const FAnimNotifyEventReference& EventReference : NotifyContext.ActiveNotifies)
+	{
+		const FAnimNotifyEvent* NotifyEvent = EventReference.GetNotify();
+		if (!NotifyEvent)
+		{
+			continue;
+		}
+
+		if (NotifyEvent->GetTriggerTime() > Time ||
+			NotifyEvent->GetEndTriggerTime() < Time)
+		{
+			continue;
+		}
+
+		UAnimNotifyState_PoseSearchBase* PoseSearchAnimNotify = 
+			Cast<UAnimNotifyState_PoseSearchBase>(NotifyEvent->NotifyStateClass);
+		if (PoseSearchAnimNotify)
+		{
+			NotifyStates.Add(PoseSearchAnimNotify);
+		}
+	}
+}
+
+void FSequenceBaseSampler::ProcessRootDistance()
+{
+	const UAnimSequenceBase* SequenceBase = Input.SequenceBase.Get();
+	
+	// Note the distance sampling interval is independent of the schema's sampling interval
+	const float DistanceSamplingInterval = 1.0f / Input.RootDistanceSamplingRate;
+
+	const FTransform InitialRootTransform = SequenceBase->ExtractRootTrackTransform(0.0f, nullptr);
+
+	uint32 NumDistanceSamples = FMath::CeilToInt(SequenceBase->GetPlayLength() * Input.RootDistanceSamplingRate) + 1;
+	AccumulatedRootDistance.Reserve(NumDistanceSamples);
+
+	// Build a distance lookup table by sampling root motion at a fixed rate and accumulating
+	// absolute translation deltas. During indexing we'll bsearch this table and interpolate
+	// between samples in order to convert distance offsets to time offsets.
+	// See also FAssetIndexer::AddTrajectoryDistanceFeatures().
+
+	double TotalAccumulatedRootDistance = 0.0;
+	FTransform LastRootTransform = InitialRootTransform;
+	float SampleTime = 0.0f;
+	for (int32 SampleIdx = 0; SampleIdx != NumDistanceSamples; ++SampleIdx)
+	{
+		SampleTime = FMath::Min(SampleIdx * DistanceSamplingInterval, SequenceBase->GetPlayLength());
+
+		FTransform RootTransform = SequenceBase->ExtractRootTrackTransform(SampleTime, nullptr);
+		FTransform LocalRootMotion = RootTransform.GetRelativeTransform(LastRootTransform);
+		LastRootTransform = RootTransform;
+
+		TotalAccumulatedRootDistance += LocalRootMotion.GetTranslation().Size();
+		AccumulatedRootDistance.Add((float)TotalAccumulatedRootDistance);
+	}
+
+	// Verify we sampled the final frame of the clip
+	check(SampleTime == SequenceBase->GetPlayLength());
+
+	// Also emit root motion summary info to help with sample wrapping in FAssetIndexer::GetSampleInfo()
+	TotalRootTransform = LastRootTransform.GetRelativeTransform(InitialRootTransform);
+	TotalRootDistance = AccumulatedRootDistance.Last();
+}
+
+const UAnimationAsset* FSequenceBaseSampler::GetAsset() const
+{
+	return Input.SequenceBase.Get();
+}
+
+//////////////////////////////////////////////////////////////////////////
+// FBlendSpaceSampler
+void FBlendSpaceSampler::Init(const FInput& InInput)
+{
+	check(InInput.BlendSpace.Get());
+
+	Input = InInput;
+}
+
+void FBlendSpaceSampler::Process()
+{
+	FMemMark Mark(FMemStack::Get());
+
+	ProcessPlayLength();
+	ProcessRootTransform();
+	ProcessRootDistance();
+}
+
+float FBlendSpaceSampler::GetTimeFromRootDistance(float Distance) const
+{
+	int32 NextSampleIdx = 1;
+	int32 PrevSampleIdx = 0;
+	if (Distance > 0.0f)
+	{
+		// Search for the distance value. Because the values will be extrapolated if necessary
+		// LowerBound might go past the end of the array, in which case the last valid index is used
+		int32 ClipDistanceLowerBoundIndex = Algo::LowerBound(AccumulatedRootDistance, Distance);
+		NextSampleIdx = FMath::Min(
+			ClipDistanceLowerBoundIndex,
+			AccumulatedRootDistance.Num() - 1);
+
+		// Compute distance interpolation amount
+		PrevSampleIdx = FMath::Max(0, NextSampleIdx - 1);
+	}
+
+	float NextDistance = AccumulatedRootDistance[NextSampleIdx];
+	float PrevDistance = AccumulatedRootDistance[PrevSampleIdx];
+	float DistanceSampleAlpha = FMath::GetRangePct(PrevDistance, NextDistance, Distance);
+
+	// Convert to time
+	float ClipTime = (float(NextSampleIdx) - (1.0f - DistanceSampleAlpha)) / Input.RootDistanceSamplingRate;
+	return ClipTime;
+}
+
+bool FBlendSpaceSampler::IsLoopable() const
+{
+	return Input.BlendSpace->bLoop;
+}
+
+void FBlendSpaceSampler::ExtractPose(const FAnimExtractContext& ExtractionCtx, FAnimationPoseData& OutAnimPoseData) const
+{
+	TArray<FBlendSampleData> BlendSamples;
+	int32 TriangulationIndex = 0;
+	Input.BlendSpace->GetSamplesFromBlendInput(Input.BlendParameters, BlendSamples, TriangulationIndex, true);
+
+	for (int32 BlendSampleIdex = 0; BlendSampleIdex < BlendSamples.Num(); BlendSampleIdex++)
+	{
+		float Scale = BlendSamples[BlendSampleIdex].Animation->GetPlayLength() / PlayLength;
+
+		FDeltaTimeRecord BlendSampleDeltaTimeRecord;
+		BlendSampleDeltaTimeRecord.Set(ExtractionCtx.DeltaTimeRecord.GetPrevious() * Scale, ExtractionCtx.DeltaTimeRecord.Delta * Scale);
+
+		BlendSamples[BlendSampleIdex].DeltaTimeRecord = BlendSampleDeltaTimeRecord;
+		BlendSamples[BlendSampleIdex].PreviousTime = ExtractionCtx.DeltaTimeRecord.GetPrevious() * Scale;
+		BlendSamples[BlendSampleIdex].Time = ExtractionCtx.CurrentTime * Scale;
+	}
+
+	Input.BlendSpace->GetAnimationPose(BlendSamples, ExtractionCtx, OutAnimPoseData);
+}
+
+FTransform FBlendSpaceSampler::ExtractRootTransform(float Time) const
+{
+	if (IsLoopable())
+	{
+		FTransform LoopableRootTransform = ExtractBlendSpaceRootMotion(0.0f, Time, true);
+		return LoopableRootTransform;
+	}
+
+	const float ExtrapolationSampleTime = Input.ExtrapolationParameters.SampleTime;
+
+	const float ClampedTime = FMath::Clamp(Time, 0.0f, PlayLength);
+	const float ExtrapolationTime = Time - ClampedTime;
+
+	FTransform RootTransform = FTransform::Identity;
+
+	// If Time is less than zero, ExtrapolationTime will be negative. In this case, we extrapolate the beginning of the 
+	// animation to estimate where the root would be at Time
+	if (ExtrapolationTime < -SMALL_NUMBER)
+	{
+		FTransform SampleToExtrapolate = ExtractBlendSpaceRootMotionFromRange(0.0f, ExtrapolationSampleTime);
+
+		const FTransform ExtrapolatedRootMotion = ExtrapolateRootMotion(
+			SampleToExtrapolate,
+			0.0f, ExtrapolationSampleTime,
+			ExtrapolationTime,
+			Input.ExtrapolationParameters);
+		RootTransform = ExtrapolatedRootMotion;
+	}
+	else
+	{
+		RootTransform = ExtractBlendSpaceRootMotionFromRange(0.0f, ClampedTime);
+
+		// If Time is greater than PlayLength, ExtrapolationTIme will be a positive number. In this case, we extrapolate
+		// the end of the animation to estimate where the root would be at Time
+		if (ExtrapolationTime > SMALL_NUMBER)
+		{
+			FTransform SampleToExtrapolate = ExtractBlendSpaceRootMotionFromRange(PlayLength - ExtrapolationSampleTime, PlayLength);
+
+			const FTransform ExtrapolatedRootMotion = ExtrapolateRootMotion(
+				SampleToExtrapolate,
+				PlayLength - ExtrapolationSampleTime, PlayLength,
+				ExtrapolationTime,
+				Input.ExtrapolationParameters);
+			RootTransform = ExtrapolatedRootMotion * RootTransform;
+		}
+	}
+
+	return RootTransform;
+}
+
+float FBlendSpaceSampler::ExtractRootDistance(float Time) const
+{
+	return ExtractAccumulatedRootDistance(
+		Input.RootDistanceSamplingRate,
+		AccumulatedRootDistance,
+		PlayLength,
+		Time,
+		Input.ExtrapolationParameters);
+}
+
+static int32 GetHighestWeightSample(const TArray<struct FBlendSampleData>& SampleDataList)
+{
+	int32 HighestWeightIndex = 0;
+	float HighestWeight = SampleDataList[HighestWeightIndex].GetClampedWeight();
+	for (int32 I = 1; I < SampleDataList.Num(); I++)
+	{
+		if (SampleDataList[I].GetClampedWeight() > HighestWeight)
+		{
+			HighestWeightIndex = I;
+			HighestWeight = SampleDataList[I].GetClampedWeight();
+		}
+	}
+	return HighestWeightIndex;
+}
+
+void FBlendSpaceSampler::ExtractPoseSearchNotifyStates(
+	float Time,
+	TArray<UAnimNotifyState_PoseSearchBase*>& NotifyStates) const
+{
+	if (Input.BlendSpace->NotifyTriggerMode == ENotifyTriggerMode::HighestWeightedAnimation)
+	{
+		// Set up blend samples
+		TArray<FBlendSampleData> BlendSamples;
+		int32 TriangulationIndex = 0;
+		Input.BlendSpace->GetSamplesFromBlendInput(Input.BlendParameters, BlendSamples, TriangulationIndex, true);
+
+		// Find highest weighted
+		const int32 HighestWeightIndex = GetHighestWeightSample(BlendSamples);
+
+		check(HighestWeightIndex != -1);
+
+		// getting pose search notifies in an interval of size ExtractionInterval, centered on Time
+		constexpr float ExtractionInterval = 1.0f / 120.0f;
+
+		float SampleTime = Time * (BlendSamples[HighestWeightIndex].Animation->GetPlayLength() / PlayLength);
+
+		// Get notifies for highest weighted
+		FAnimNotifyContext NotifyContext;
+		BlendSamples[HighestWeightIndex].Animation->GetAnimNotifies(
+			(SampleTime - (ExtractionInterval * 0.5f)),
+			ExtractionInterval, 
+			NotifyContext);
+
+		// check which notifies actually overlap Time and are of the right base type
+		for (const FAnimNotifyEventReference& EventReference : NotifyContext.ActiveNotifies)
+		{
+			const FAnimNotifyEvent* NotifyEvent = EventReference.GetNotify();
+			if (!NotifyEvent)
+			{
+				continue;
+			}
+
+			if (NotifyEvent->GetTriggerTime() > SampleTime ||
+				NotifyEvent->GetEndTriggerTime() < SampleTime)
+			{
+				continue;
+			}
+
+			UAnimNotifyState_PoseSearchBase* PoseSearchAnimNotify =
+				Cast<UAnimNotifyState_PoseSearchBase>(NotifyEvent->NotifyStateClass);
+			if (PoseSearchAnimNotify)
+			{
+				NotifyStates.Add(PoseSearchAnimNotify);
+			}
+		}
+	}
+}
+
+FTransform FBlendSpaceSampler::ExtractBlendSpaceRootTrackTransform(float Time) const
+{
+	checkf(AccumulatedRootTransform.Num() > 0, TEXT("ProcessRootTransform must be run first"));
+
+	int32 Index = Time * Input.RootTransformSamplingRate;
+	int32 FirstIndexClamped = FMath::Clamp(Index + 0, 0, AccumulatedRootTransform.Num() - 1);
+	int32 SecondIndexClamped = FMath::Clamp(Index + 1, 0, AccumulatedRootTransform.Num() - 1);
+	float Alpha = FMath::Fmod(Time * Input.RootTransformSamplingRate, 1.0f);
+	FTransform OutputTransform;
+	OutputTransform.Blend(
+		AccumulatedRootTransform[FirstIndexClamped],
+		AccumulatedRootTransform[SecondIndexClamped],
+		Alpha);
+
+	return OutputTransform;
+}
+
+FTransform FBlendSpaceSampler::ExtractBlendSpaceRootMotionFromRange(float StartTrackPosition, float EndTrackPosition) const
+{
+	checkf(AccumulatedRootTransform.Num() > 0, TEXT("ProcessRootTransform must be run first"));
+
+	FTransform RootTransformRefPose = ExtractBlendSpaceRootTrackTransform(0.0f);
+
+	FTransform StartTransform = ExtractBlendSpaceRootTrackTransform(StartTrackPosition);
+	FTransform EndTransform = ExtractBlendSpaceRootTrackTransform(EndTrackPosition);
+
+	// Transform to Component Space
+	const FTransform RootToComponent = RootTransformRefPose.Inverse();
+	StartTransform = RootToComponent * StartTransform;
+	EndTransform = RootToComponent * EndTransform;
+
+	return EndTransform.GetRelativeTransform(StartTransform);
+}
+
+FTransform FBlendSpaceSampler::ExtractBlendSpaceRootMotion(float StartTime, float DeltaTime, bool bAllowLooping) const
+{
+	FRootMotionMovementParams RootMotionParams;
+
+	if (DeltaTime != 0.f)
+	{
+		bool const bPlayingBackwards = (DeltaTime < 0.f);
+
+		float PreviousPosition = StartTime;
+		float CurrentPosition = StartTime;
+		float DesiredDeltaMove = DeltaTime;
+
+		do
+		{
+			// Disable looping here. Advance to desired position, or beginning / end of animation 
+			const ETypeAdvanceAnim AdvanceType = FAnimationRuntime::AdvanceTime(false, DesiredDeltaMove, CurrentPosition, PlayLength);
+
+			// Verify position assumptions
+			//ensureMsgf(bPlayingBackwards ? (CurrentPosition <= PreviousPosition) : (CurrentPosition >= PreviousPosition), TEXT("in Animation %s(Skeleton %s) : bPlayingBackwards(%d), PreviousPosition(%0.2f), Current Position(%0.2f)"),
+			//	*GetName(), *GetNameSafe(GetSkeleton()), bPlayingBackwards, PreviousPosition, CurrentPosition);
+
+			RootMotionParams.Accumulate(ExtractBlendSpaceRootMotionFromRange(PreviousPosition, CurrentPosition));
+
+			// If we've hit the end of the animation, and we're allowed to loop, keep going.
+			if ((AdvanceType == ETAA_Finished) && bAllowLooping)
+			{
+				const float ActualDeltaMove = (CurrentPosition - PreviousPosition);
+				DesiredDeltaMove -= ActualDeltaMove;
+
+				PreviousPosition = bPlayingBackwards ? PlayLength : 0.f;
+				CurrentPosition = PreviousPosition;
+			}
+			else
+			{
+				break;
+			}
+		} while (true);
+	}
+
+	return RootMotionParams.GetRootMotionTransform();
+}
+
+void FBlendSpaceSampler::ProcessPlayLength()
+{
+	TArray<FBlendSampleData> BlendSamples;
+	int32 TriangulationIndex = 0;
+	Input.BlendSpace->GetSamplesFromBlendInput(Input.BlendParameters, BlendSamples, TriangulationIndex, true);
+
+	PlayLength = Input.BlendSpace->GetAnimationLengthFromSampleData(BlendSamples);
+}
+
+void FBlendSpaceSampler::ProcessRootTransform()
+{
+	// Pre-compute root motion
+
+	int32 NumRootSamples = FMath::Max(PlayLength * Input.RootTransformSamplingRate + 1, 1);
+	AccumulatedRootTransform.SetNumUninitialized(NumRootSamples);
+
+	TArray<FBlendSampleData> BlendSamples;
+	int32 TriangulationIndex = 0;
+	Input.BlendSpace->GetSamplesFromBlendInput(Input.BlendParameters, BlendSamples, TriangulationIndex, true);
+
+	FTransform RootMotionAccumulation = FTransform::Identity;
+
+	AccumulatedRootTransform[0] = RootMotionAccumulation;
+
+	for (int32 SampleIdx = 1; SampleIdx < NumRootSamples; ++SampleIdx)
+	{
+		float PreviousTime = float(SampleIdx - 1) / Input.RootTransformSamplingRate;
+		float CurrentTime = float(SampleIdx - 0) / Input.RootTransformSamplingRate;
+
+		FDeltaTimeRecord DeltaTimeRecord;
+		DeltaTimeRecord.Set(PreviousTime, CurrentTime - PreviousTime);
+		FAnimExtractContext ExtractionCtx(static_cast<double>(CurrentTime), true, DeltaTimeRecord, IsLoopable());
+
+		for (int32 BlendSampleIdex = 0; BlendSampleIdex < BlendSamples.Num(); BlendSampleIdex++)
+		{
+			float Scale = BlendSamples[BlendSampleIdex].Animation->GetPlayLength() / PlayLength;
+
+			FDeltaTimeRecord BlendSampleDeltaTimeRecord;
+			BlendSampleDeltaTimeRecord.Set(DeltaTimeRecord.GetPrevious() * Scale, DeltaTimeRecord.Delta * Scale);
+
+			BlendSamples[BlendSampleIdex].DeltaTimeRecord = BlendSampleDeltaTimeRecord;
+			BlendSamples[BlendSampleIdex].PreviousTime = PreviousTime * Scale;
+			BlendSamples[BlendSampleIdex].Time = CurrentTime * Scale;
+		}
+
+		FCompactPose Pose;
+		FBlendedCurve BlendedCurve;
+		Anim::FStackAttributeContainer StackAttributeContainer;
+		FAnimationPoseData AnimPoseData(Pose, BlendedCurve, StackAttributeContainer);
+
+		Pose.SetBoneContainer(&Input.BoneContainer);
+		BlendedCurve.InitFrom(Input.BoneContainer);
+
+		Input.BlendSpace->GetAnimationPose(BlendSamples, ExtractionCtx, AnimPoseData);
+
+		const Anim::IAnimRootMotionProvider* RootMotionProvider = Anim::IAnimRootMotionProvider::Get();
+
+		if (ensureMsgf(RootMotionProvider, TEXT("Could not get Root Motion Provider.")))
+		{
+			if (ensureMsgf(RootMotionProvider->HasRootMotion(StackAttributeContainer), TEXT("Blend Space had no Root Motion Attribute.")))
+			{
+				FTransform RootMotionDelta;
+				RootMotionProvider->ExtractRootMotion(StackAttributeContainer, RootMotionDelta);
+
+				RootMotionAccumulation = RootMotionDelta * RootMotionAccumulation;
+			}
+		}
+
+		AccumulatedRootTransform[SampleIdx] = RootMotionAccumulation;
+	}
+}
+
+void FBlendSpaceSampler::ProcessRootDistance()
+{
+	checkf(AccumulatedRootTransform.Num() > 0, TEXT("ProcessRootTransform must be run first"));
+
+	// Note the distance sampling interval is independent of the schema's sampling interval
+	const float DistanceSamplingInterval = 1.0f / Input.RootDistanceSamplingRate;
+
+	const FTransform InitialRootTransform = FTransform::Identity;
+
+	uint32 NumDistanceSamples = FMath::CeilToInt(PlayLength * Input.RootDistanceSamplingRate) + 1;
+	AccumulatedRootDistance.Reserve(NumDistanceSamples);
+
+	// Build a distance lookup table by sampling root motion at a fixed rate and accumulating
+	// absolute translation deltas. During indexing we'll bsearch this table and interpolate
+	// between samples in order to convert distance offsets to time offsets.
+	// See also FAssetIndexer::AddTrajectoryDistanceFeatures().
+	double TotalAccumulatedRootDistance = 0.0;
+	FTransform LastRootTransform = InitialRootTransform;
+	float SampleTime = 0.0f;
+	for (int32 SampleIdx = 0; SampleIdx != NumDistanceSamples; ++SampleIdx)
+	{
+		SampleTime = FMath::Min(SampleIdx * DistanceSamplingInterval, PlayLength);
+
+		FTransform RootTransform = ExtractBlendSpaceRootTrackTransform(SampleTime);
+		FTransform LocalRootMotion = RootTransform.GetRelativeTransform(LastRootTransform);
+		LastRootTransform = RootTransform;
+
+		TotalAccumulatedRootDistance += LocalRootMotion.GetTranslation().Size();
+		AccumulatedRootDistance.Add((float)TotalAccumulatedRootDistance);
+	}
+
+	// Verify we sampled the final frame of the clip
+	check(SampleTime == PlayLength);
+
+	// Also emit root motion summary info to help with sample wrapping in FAssetIndexer::GetSampleInfo()
+	TotalRootTransform = LastRootTransform.GetRelativeTransform(InitialRootTransform);
+	TotalRootDistance = AccumulatedRootDistance.Last(); 
+}
+
+const UAnimationAsset* FBlendSpaceSampler::GetAsset() const
+{
+	return Input.BlendSpace.Get();
+}
+
+} // namespace UE::PoseSearch
diff --git a/Engine/Plugins/Experimental/Animation/PoseSearch/Source/Runtime/Private/PoseSearchContext.cpp b/Engine/Plugins/Experimental/Animation/PoseSearch/Source/Runtime/Private/PoseSearchContext.cpp
new file mode 100644
--- /dev/null
+++ b/Engine/Plugins/Experimental/Animation/PoseSearch/Source/Runtime/Private/PoseSearchContext.cpp
@@ -0,0 +1,237 @@
+// Copyright Epic Games, Inc. All Rights Reserved.
+
+#include "PoseSearch/PoseSearchContext.h"
+#include "AnimationRuntime.h"
+#include "DrawDebugHelpers.h"
+#include "PoseSearch/PoseSearchDatabase.h"
+#include "PoseSearch/PoseSearchFeatureChannel.h"
+#include "PoseSearch/PoseSearchHistory.h"
+#include "PoseSearch/PoseSearchSchema.h"
+
+namespace UE::PoseSearch
+{
+	
+//////////////////////////////////////////////////////////////////////////
+// FDebugDrawParams
+bool FDebugDrawParams::CanDraw() const
+{
+#if ENABLE_DRAW_DEBUG
+	return World && Database && Database->Schema && Database->Schema->IsValid();
+#else // ENABLE_DRAW_DEBUG
+	return false;
+#endif // ENABLE_DRAW_DEBUG
+}
+
+FColor FDebugDrawParams::GetColor(int32 ColorPreset) const
+{
+#if ENABLE_DRAW_DEBUG
+	FLinearColor Color = FLinearColor::Red;
+
+	const UPoseSearchSchema* Schema = GetSchema();
+	if (!Schema || !Schema->IsValid())
+	{
+		Color = FLinearColor::Red;
+	}
+	else if (ColorPreset < 0 || ColorPreset >= Schema->ColorPresets.Num())
+	{
+		if (EnumHasAnyFlags(Flags, EDebugDrawFlags::DrawQuery))
+		{
+			Color = FLinearColor::Blue;
+		}
+		else
+		{
+			Color = FLinearColor::Green;
+		}
+	}
+	else
+	{
+		if (EnumHasAnyFlags(Flags, EDebugDrawFlags::DrawQuery))
+		{
+			Color = Schema->ColorPresets[ColorPreset].Query;
+		}
+		else
+		{
+			Color = Schema->ColorPresets[ColorPreset].Result;
+		}
+	}
+
+	return Color.ToFColor(true);
+#else // ENABLE_DRAW_DEBUG
+	return FColor::Black;
+#endif // ENABLE_DRAW_DEBUG
+}
+
+const FPoseSearchIndex* FDebugDrawParams::GetSearchIndex() const
+{
+	return Database ? &Database->GetSearchIndex() : nullptr;
+}
+
+const UPoseSearchSchema* FDebugDrawParams::GetSchema() const
+{
+	return Database ? Database->Schema : nullptr;
+}
+
+void DrawFeatureVector(const FDebugDrawParams& DrawParams, TConstArrayView<float> PoseVector)
+{
+#if ENABLE_DRAW_DEBUG
+	if (DrawParams.CanDraw())
+	{
+		const UPoseSearchSchema* Schema = DrawParams.GetSchema();
+		check(Schema);
+
+		if (PoseVector.Num() == Schema->SchemaCardinality)
+		{
+			for (int32 ChannelIdx = 0; ChannelIdx != Schema->Channels.Num(); ++ChannelIdx)
+			{
+				if (DrawParams.ChannelMask & (1 << ChannelIdx))
+				{
+					Schema->Channels[ChannelIdx]->DebugDraw(DrawParams, PoseVector);
+				}
+			}
+		}
+	}
+#endif // ENABLE_DRAW_DEBUG
+}
+
+void DrawFeatureVector(const FDebugDrawParams& DrawParams, int32 PoseIdx)
+{
+#if ENABLE_DRAW_DEBUG
+	// if we're editing the schema while in PIE with Rewind Debugger active, PoseIdx could be out of bound / stale
+	if (DrawParams.CanDraw() && PoseIdx >= 0 && PoseIdx < DrawParams.GetSearchIndex()->NumPoses)
+	{
+		DrawFeatureVector(DrawParams, DrawParams.GetSearchIndex()->GetPoseValues(PoseIdx));
+	}
+#endif // ENABLE_DRAW_DEBUG
+}
+
+//////////////////////////////////////////////////////////////////////////
+// FSearchContext
+FTransform FSearchContext::TryGetTransformAndCacheResults(float SampleTime, const UPoseSearchSchema* Schema, int8 SchemaBoneIdx)
+{
+	check(History && Schema);
+
+	static constexpr FBoneIndexType RootBoneIdx = 0xFFFF;
+	const FBoneIndexType BoneIndexType = SchemaBoneIdx >= 0 ? Schema->BoneIndices[SchemaBoneIdx] : RootBoneIdx;
+
+	// @todo: use an hashmap if we end up having too many entries
+	const FCachedEntry* Entry = CachedEntries.FindByPredicate([SampleTime, BoneIndexType](const FSearchContext::FCachedEntry& Entry)
+	{
+		return Entry.SampleTime == SampleTime && Entry.BoneIndexType == BoneIndexType;
+	});
+
+	if (Entry)
+	{
+		return Entry->Transform;
+	}
+
+	if (BoneIndexType != RootBoneIdx)
+	{
+		TArray<FTransform> SampledLocalPose;
+		if (History->TrySampleLocalPose(-SampleTime, &Schema->BoneIndicesWithParents, &SampledLocalPose, nullptr))
+		{
+			TArray<FTransform> SampledComponentPose;
+			FAnimationRuntime::FillUpComponentSpaceTransforms(Schema->Skeleton->GetReferenceSkeleton(), SampledLocalPose, SampledComponentPose);
+
+			// adding bunch of entries, without caring about adding eventual duplicates
+			for (const FBoneIndexType NewEntryBoneIndexType : Schema->BoneIndicesWithParents)
+			{
+				CachedEntries.Emplace(SampleTime, SampledComponentPose[NewEntryBoneIndexType], NewEntryBoneIndexType);
+			}
+
+			return SampledComponentPose[BoneIndexType];
+		}
+
+		return FTransform::Identity;
+	}
+	
+	FTransform SampledRootTransform;
+	if (History->TrySampleLocalPose(-SampleTime, nullptr, nullptr, &SampledRootTransform))
+	{
+		FCachedEntry& NewEntry = CachedEntries[CachedEntries.AddDefaulted()];
+		NewEntry.SampleTime = SampleTime;
+		NewEntry.BoneIndexType = BoneIndexType;
+		NewEntry.Transform = SampledRootTransform;
+
+		return SampledRootTransform;
+	}
+	
+	return FTransform::Identity;
+}
+
+void FSearchContext::ClearCachedEntries()
+{
+	CachedEntries.Reset();
+}
+
+void FSearchContext::ResetCurrentBestCost()
+{
+	CurrentBestTotalCost = MAX_flt;
+}
+
+void FSearchContext::UpdateCurrentBestCost(const FPoseSearchCost& PoseSearchCost)
+{
+	check(PoseSearchCost.IsValid());
+
+	if (PoseSearchCost.GetTotalCost() < CurrentBestTotalCost)
+	{
+		CurrentBestTotalCost = PoseSearchCost.GetTotalCost();
+	};
+}
+
+const FPoseSearchFeatureVectorBuilder* FSearchContext::GetCachedQuery(const UPoseSearchDatabase* Database) const
+{
+	const FSearchContext::FCachedQuery* CachedQuery = CachedQueries.FindByPredicate([Database](const FSearchContext::FCachedQuery& CachedQuery)
+	{
+		return CachedQuery.Database == Database;
+	});
+
+	if (CachedQuery)
+	{
+		return &CachedQuery->FeatureVectorBuilder;
+	}
+	return nullptr;
+}
+
+void FSearchContext::GetOrBuildQuery(const UPoseSearchDatabase* Database, FPoseSearchFeatureVectorBuilder& FeatureVectorBuilder)
+{
+	const FPoseSearchFeatureVectorBuilder* CachedFeatureVectorBuilder = GetCachedQuery(Database);
+	if (CachedFeatureVectorBuilder)
+	{
+		FeatureVectorBuilder = *CachedFeatureVectorBuilder;
+	}
+	else
+	{
+		FSearchContext::FCachedQuery& NewCachedQuery = CachedQueries[CachedQueries.AddDefaulted()];
+		NewCachedQuery.Database = Database;
+		Database->BuildQuery(*this, NewCachedQuery.FeatureVectorBuilder);
+		FeatureVectorBuilder = NewCachedQuery.FeatureVectorBuilder;
+	}
+}
+
+bool FSearchContext::IsCurrentResultFromDatabase(const UPoseSearchDatabase* Database) const
+{
+	return CurrentResult.IsValid() && CurrentResult.Database == Database;
+}
+
+TConstArrayView<float> FSearchContext::GetCurrentResultPrevPoseVector() const
+{
+	check(CurrentResult.IsValid());
+	const FPoseSearchIndex& SearchIndex = CurrentResult.Database->GetSearchIndex();
+	return SearchIndex.GetPoseValues(CurrentResult.PrevPoseIdx);
+}
+
+TConstArrayView<float> FSearchContext::GetCurrentResultPoseVector() const
+{
+	check(CurrentResult.IsValid());
+	const FPoseSearchIndex& SearchIndex = CurrentResult.Database->GetSearchIndex();
+	return SearchIndex.GetPoseValues(CurrentResult.PoseIdx);
+}
+
+TConstArrayView<float> FSearchContext::GetCurrentResultNextPoseVector() const
+{
+	check(CurrentResult.IsValid());
+	const FPoseSearchIndex& SearchIndex = CurrentResult.Database->GetSearchIndex();
+	return SearchIndex.GetPoseValues(CurrentResult.NextPoseIdx);
+}
+
+} // namespace UE::PoseSearch
\ No newline at end of file
diff --git a/Engine/Plugins/Experimental/Animation/PoseSearch/Source/Runtime/Private/PoseSearchDatabase.cpp b/Engine/Plugins/Experimental/Animation/PoseSearch/Source/Runtime/Private/PoseSearchDatabase.cpp
new file mode 100644
--- /dev/null
+++ b/Engine/Plugins/Experimental/Animation/PoseSearch/Source/Runtime/Private/PoseSearchDatabase.cpp
@@ -0,0 +1,906 @@
+// Copyright Epic Games, Inc. All Rights Reserved.
+
+#include "PoseSearch/PoseSearchDatabase.h"
+#include "Animation/AnimComposite.h"
+#include "Animation/AnimSequence.h"
+#include "Animation/BlendSpace.h"
+#include "Animation/BlendSpace1D.h"
+#include "InstancedStruct.h"
+#include "PoseSearch/PoseSearchContext.h"
+#include "PoseSearch/PoseSearchDefines.h"
+#include "PoseSearch/PoseSearchDerivedData.h"
+#include "PoseSearch/PoseSearchFeatureChannel.h"
+#include "PoseSearch/PoseSearchHistory.h"
+#include "PoseSearch/PoseSearchSchema.h"
+#include "PoseSearchEigenHelper.h"
+
+
+DECLARE_STATS_GROUP(TEXT("PoseSearch"), STATGROUP_PoseSearch, STATCAT_Advanced);
+DECLARE_CYCLE_STAT_EXTERN(TEXT("Search Brute Force"), STAT_PoseSearchBruteForce, STATGROUP_PoseSearch, );
+DECLARE_CYCLE_STAT_EXTERN(TEXT("Search PCA/KNN"), STAT_PoseSearchPCAKNN, STATGROUP_PoseSearch, );
+DEFINE_STAT(STAT_PoseSearchBruteForce);
+DEFINE_STAT(STAT_PoseSearchPCAKNN);
+
+namespace UE::PoseSearch
+{
+
+typedef TArray<size_t, TInlineAllocator<128>> FNonSelectableIdx;
+static void PopulateNonSelectableIdx(FNonSelectableIdx& NonSelectableIdx, FSearchContext& SearchContext, const UPoseSearchDatabase* Database, TConstArrayView<float> QueryValues)
+{
+	check(Database);
+#if UE_POSE_SEARCH_TRACE_ENABLED
+	const FPoseSearchIndex& SearchIndex = Database->GetSearchIndex();
+#endif
+
+	const FPoseSearchIndexAsset* CurrentIndexAsset = SearchContext.CurrentResult.GetSearchIndexAsset();
+	if (CurrentIndexAsset && SearchContext.IsCurrentResultFromDatabase(Database) && SearchContext.PoseJumpThresholdTime > 0.f)
+	{
+		const int32 PoseJumpIndexThreshold = FMath::FloorToInt(SearchContext.PoseJumpThresholdTime / Database->Schema->GetSamplingInterval());
+		const bool IsLooping = Database->IsSourceAssetLooping(*CurrentIndexAsset);
+
+		for (int32 i = -PoseJumpIndexThreshold; i <= -1; ++i)
+		{
+			int32 PoseIdx = SearchContext.CurrentResult.PoseIdx + i;
+			bool bIsPoseInRange = false;
+			if (IsLooping)
+			{
+				bIsPoseInRange = true;
+
+				while (PoseIdx < CurrentIndexAsset->FirstPoseIdx)
+				{
+					PoseIdx += CurrentIndexAsset->NumPoses;
+				}
+			}
+			else if (CurrentIndexAsset->IsPoseInRange(PoseIdx))
+			{
+				bIsPoseInRange = true;
+			}
+
+			if (bIsPoseInRange)
+			{
+				NonSelectableIdx.Add(PoseIdx);
+
+#if UE_POSE_SEARCH_TRACE_ENABLED
+				const FPoseSearchCost PoseCost = SearchIndex.ComparePoses(PoseIdx, SearchContext.QueryMirrorRequest, EPoseComparisonFlags::None, Database->Schema->MirrorMismatchCostBias, QueryValues);
+				SearchContext.BestCandidates.Add(PoseCost, PoseIdx, Database, EPoseCandidateFlags::DiscardedBy_PoseJumpThresholdTime);
+#endif
+			}
+			else
+			{
+				break;
+			}
+		}
+
+		for (int32 i = 0; i <= PoseJumpIndexThreshold; ++i)
+		{
+			int32 PoseIdx = SearchContext.CurrentResult.PoseIdx + i;
+			bool bIsPoseInRange = false;
+			if (IsLooping)
+			{
+				bIsPoseInRange = true;
+
+				while (PoseIdx >= CurrentIndexAsset->FirstPoseIdx + CurrentIndexAsset->NumPoses)
+				{
+					PoseIdx -= CurrentIndexAsset->NumPoses;
+				}
+			}
+			else if (CurrentIndexAsset->IsPoseInRange(PoseIdx))
+			{
+				bIsPoseInRange = true;
+			}
+
+			if (bIsPoseInRange)
+			{
+				NonSelectableIdx.Add(PoseIdx);
+
+#if UE_POSE_SEARCH_TRACE_ENABLED
+				const FPoseSearchCost PoseCost = SearchIndex.ComparePoses(PoseIdx, SearchContext.QueryMirrorRequest, EPoseComparisonFlags::None, Database->Schema->MirrorMismatchCostBias, QueryValues);
+				SearchContext.BestCandidates.Add(PoseCost, PoseIdx, Database, EPoseCandidateFlags::DiscardedBy_PoseJumpThresholdTime);
+#endif
+			}
+			else
+			{
+				break;
+			}
+		}
+	}
+
+	if (SearchContext.PoseIndicesHistory)
+	{
+		const FObjectKey DatabaseKey(Database);
+		for (auto It = SearchContext.PoseIndicesHistory->IndexToTime.CreateConstIterator(); It; ++It)
+		{
+			const FHistoricalPoseIndex& HistoricalPoseIndex = It.Key();
+			if (HistoricalPoseIndex.DatabaseKey == DatabaseKey)
+			{
+				NonSelectableIdx.Add(HistoricalPoseIndex.PoseIndex);
+
+#if UE_POSE_SEARCH_TRACE_ENABLED
+				check(HistoricalPoseIndex.PoseIndex >= 0);
+
+				// if we're editing the database and removing assets it's possible that the PoseIndicesHistory contains invalid pose indexes
+				if (HistoricalPoseIndex.PoseIndex < SearchIndex.NumPoses)
+				{
+					const FPoseSearchCost PoseCost = SearchIndex.ComparePoses(HistoricalPoseIndex.PoseIndex, SearchContext.QueryMirrorRequest, EPoseComparisonFlags::None, Database->Schema->MirrorMismatchCostBias, QueryValues);
+					SearchContext.BestCandidates.Add(PoseCost, HistoricalPoseIndex.PoseIndex, Database, EPoseCandidateFlags::DiscardedBy_PoseReselectHistory);
+				}
+#endif
+			}
+		}
+	}
+
+	NonSelectableIdx.Sort();
+}
+
+struct FPoseFilters
+{
+	FPoseFilters(const UPoseSearchSchema* Schema, TConstArrayView<size_t> NonSelectableIdx, EPoseSearchPoseFlags OverallFlags)
+	{
+		NonSelectableIdxPoseFilter.NonSelectableIdx = NonSelectableIdx;
+
+		if (EnumHasAnyFlags(OverallFlags, EPoseSearchPoseFlags::BlockTransition))
+		{
+			AllPoseFilters.Add(&BlockTransitionPoseFilter);
+		}
+
+		if (NonSelectableIdxPoseFilter.IsPoseFilterActive())
+		{
+			AllPoseFilters.Add(&NonSelectableIdxPoseFilter);
+		}
+
+		for (const IPoseFilter* ChannelPoseFilter : Schema->Channels)
+		{
+			if (ChannelPoseFilter->IsPoseFilterActive())
+			{
+				AllPoseFilters.Add(ChannelPoseFilter);
+			}
+		}
+	}
+
+	bool AreFiltersValid(const FPoseSearchIndex& SearchIndex, TConstArrayView<float> QueryValues, int32 PoseIdx, const FPoseSearchPoseMetadata& Metadata
+#if UE_POSE_SEARCH_TRACE_ENABLED
+		, UE::PoseSearch::FSearchContext& SearchContext, const UPoseSearchDatabase* Database
+#endif
+	) const
+	{
+		TConstArrayView<float> PoseValues = SearchIndex.GetPoseValues(PoseIdx);
+		for (const IPoseFilter* PoseFilter : AllPoseFilters)
+		{
+			if (!PoseFilter->IsPoseValid(PoseValues, QueryValues, PoseIdx, Metadata))
+			{
+#if UE_POSE_SEARCH_TRACE_ENABLED
+				if (PoseFilter == &NonSelectableIdxPoseFilter)
+				{
+					// candidate already added to SearchContext.BestCandidates by PopulateNonSelectableIdx
+				}
+				else if (PoseFilter == &BlockTransitionPoseFilter)
+				{
+					const FPoseSearchCost PoseCost = SearchIndex.ComparePoses(PoseIdx, SearchContext.QueryMirrorRequest, EPoseComparisonFlags::None, Database->Schema->MirrorMismatchCostBias, QueryValues);
+					SearchContext.BestCandidates.Add(PoseCost, PoseIdx, Database, EPoseCandidateFlags::DiscardedBy_BlockTransition);
+				}
+				else
+				{
+					const FPoseSearchCost PoseCost = SearchIndex.ComparePoses(PoseIdx, SearchContext.QueryMirrorRequest, EPoseComparisonFlags::None, Database->Schema->MirrorMismatchCostBias, QueryValues);
+					SearchContext.BestCandidates.Add(PoseCost, PoseIdx, Database, EPoseCandidateFlags::DiscardedBy_PoseFilter);
+				}
+#endif
+				return false;
+			}
+		}
+		return true;
+	};
+
+private:
+	struct FNonSelectableIdxPoseFilter : public IPoseFilter
+	{
+		virtual bool IsPoseFilterActive() const override
+		{
+			return !NonSelectableIdx.IsEmpty();
+		}
+
+		virtual bool IsPoseValid(TConstArrayView<float> PoseValues, TConstArrayView<float> QueryValues, int32 PoseIdx, const FPoseSearchPoseMetadata& Metadata) const override
+		{
+			return Algo::BinarySearch(NonSelectableIdx, PoseIdx) == INDEX_NONE;
+		}
+
+		TConstArrayView<size_t> NonSelectableIdx;
+	};
+
+	struct FBlockTransitionPoseFilter : public IPoseFilter
+	{
+		virtual bool IsPoseFilterActive() const override
+		{
+			return true;
+		}
+
+		virtual bool IsPoseValid(TConstArrayView<float> PoseValues, TConstArrayView<float> QueryValues, int32 PoseIdx, const FPoseSearchPoseMetadata& Metadata) const override
+		{
+			return !EnumHasAnyFlags(Metadata.Flags, EPoseSearchPoseFlags::BlockTransition);
+		}
+	};
+
+	FNonSelectableIdxPoseFilter NonSelectableIdxPoseFilter;
+	FBlockTransitionPoseFilter BlockTransitionPoseFilter;
+
+	TArray<const IPoseFilter*, TInlineAllocator<64>> AllPoseFilters;
+};
+
+} // namespace UE::PoseSearch
+
+//////////////////////////////////////////////////////////////////////////
+// FPoseSearchDatabaseSequence
+UAnimationAsset* FPoseSearchDatabaseSequence::GetAnimationAsset() const
+{
+	return Sequence;
+}
+
+UClass* FPoseSearchDatabaseSequence::GetAnimationAssetStaticClass() const
+{
+	return UAnimSequence::StaticClass();
+}
+
+bool FPoseSearchDatabaseSequence::IsLooping() const
+{
+	return Sequence ? Sequence->bLoop : false;
+}
+
+const FString FPoseSearchDatabaseSequence::GetName() const
+{
+	return Sequence ? Sequence->GetName() : FString();
+}
+
+bool FPoseSearchDatabaseSequence::IsRootMotionEnabled() const
+{
+	return Sequence ? Sequence->HasRootMotion() : false;
+}
+
+//////////////////////////////////////////////////////////////////////////
+// FPoseSearchDatabaseBlendSpace
+UAnimationAsset* FPoseSearchDatabaseBlendSpace::GetAnimationAsset() const
+{
+	return BlendSpace.Get();
+}
+
+UClass* FPoseSearchDatabaseBlendSpace::GetAnimationAssetStaticClass() const
+{
+	return UBlendSpace::StaticClass();
+}
+
+bool FPoseSearchDatabaseBlendSpace::IsLooping() const
+{
+	return BlendSpace ? BlendSpace->bLoop : false;
+}
+
+const FString FPoseSearchDatabaseBlendSpace::GetName() const
+{
+	return BlendSpace ? BlendSpace->GetName() : FString();
+}
+
+bool FPoseSearchDatabaseBlendSpace::IsRootMotionEnabled() const
+{
+	bool bIsRootMotionUsedInBlendSpace = false;
+
+	if (BlendSpace)
+	{
+		BlendSpace->ForEachImmutableSample([&bIsRootMotionUsedInBlendSpace](const FBlendSample& Sample)
+			{
+				const TObjectPtr<UAnimSequence> Sequence = Sample.Animation;
+
+				if (IsValid(Sequence) && Sequence->HasRootMotion())
+				{
+					bIsRootMotionUsedInBlendSpace = true;
+				}
+			});
+	}
+
+	return bIsRootMotionUsedInBlendSpace;
+}
+
+void FPoseSearchDatabaseBlendSpace::GetBlendSpaceParameterSampleRanges(int32& HorizontalBlendNum, int32& VerticalBlendNum) const
+{
+	check(BlendSpace);
+
+	HorizontalBlendNum = bUseGridForSampling ? BlendSpace->GetBlendParameter(0).GridNum + 1 : FMath::Max(NumberOfHorizontalSamples, 1);
+	VerticalBlendNum = BlendSpace->IsA<UBlendSpace1D>() ? 1 : bUseGridForSampling ? BlendSpace->GetBlendParameter(1).GridNum + 1 : FMath::Max(NumberOfVerticalSamples, 1);
+
+	check(HorizontalBlendNum >= 1 && VerticalBlendNum >= 1);
+}
+
+FVector FPoseSearchDatabaseBlendSpace::BlendParameterForSampleRanges(int32 HorizontalBlendIndex, int32 VerticalBlendIndex) const
+{
+	check(BlendSpace);
+
+	const bool bWrapInputOnHorizontalAxis = BlendSpace->GetBlendParameter(0).bWrapInput;
+	const bool bWrapInputOnVerticalAxis = BlendSpace->GetBlendParameter(1).bWrapInput;
+
+	int32 HorizontalBlendNum, VerticalBlendNum;
+	GetBlendSpaceParameterSampleRanges(HorizontalBlendNum, VerticalBlendNum);
+
+	if (bWrapInputOnHorizontalAxis)
+	{
+		++HorizontalBlendNum;
+	}
+
+	if (bWrapInputOnVerticalAxis)
+	{
+		++VerticalBlendNum;
+	}
+
+	const float HorizontalBlendMin = BlendSpace->GetBlendParameter(0).Min;
+	const float HorizontalBlendMax = BlendSpace->GetBlendParameter(0).Max;
+
+	const float VerticalBlendMin = BlendSpace->GetBlendParameter(1).Min;
+	const float VerticalBlendMax = BlendSpace->GetBlendParameter(1).Max;
+
+	return FVector(
+		HorizontalBlendNum > 1 ? 
+			HorizontalBlendMin + (HorizontalBlendMax - HorizontalBlendMin) * 
+			((float)HorizontalBlendIndex) / (HorizontalBlendNum - 1) : 
+		HorizontalBlendMin,
+		VerticalBlendNum > 1 ? 
+			VerticalBlendMin + (VerticalBlendMax - VerticalBlendMin) * 
+			((float)VerticalBlendIndex) / (VerticalBlendNum - 1) : 
+		VerticalBlendMin,
+		0.0f);
+}
+
+//////////////////////////////////////////////////////////////////////////
+// FPoseSearchDatabaseAnimComposite
+UAnimationAsset* FPoseSearchDatabaseAnimComposite::GetAnimationAsset() const
+{
+	return AnimComposite;
+}
+
+UClass* FPoseSearchDatabaseAnimComposite::GetAnimationAssetStaticClass() const
+{
+	return UAnimComposite::StaticClass();
+}
+
+bool FPoseSearchDatabaseAnimComposite::IsLooping() const
+{
+	return AnimComposite ? AnimComposite->bLoop : false;
+}
+
+const FString FPoseSearchDatabaseAnimComposite::GetName() const
+{
+	return AnimComposite ? AnimComposite->GetName() : FString();
+}
+
+bool FPoseSearchDatabaseAnimComposite::IsRootMotionEnabled() const
+{
+	return AnimComposite ? AnimComposite->HasRootMotion() : false;
+}
+
+//////////////////////////////////////////////////////////////////////////
+// UPoseSearchDatabase
+UPoseSearchDatabase::~UPoseSearchDatabase()
+{
+}
+
+void UPoseSearchDatabase::SetSearchIndex(const FPoseSearchIndex& SearchIndex)
+{
+	check(IsInGameThread());
+	SearchIndexPrivate = SearchIndex;
+}
+
+const FPoseSearchIndex& UPoseSearchDatabase::GetSearchIndex() const
+{
+	// making sure the search index is consistent. if it fails the calling code hasn't been protected by FAsyncPoseSearchDatabasesManagement::RequestAsyncBuildIndex
+	check(Schema && Schema->IsValid() && !SearchIndexPrivate.IsEmpty() && SearchIndexPrivate.WeightsSqrt.Num() == Schema->SchemaCardinality && SearchIndexPrivate.KDTree.Impl);
+	return SearchIndexPrivate;
+}
+
+int32 UPoseSearchDatabase::GetPoseIndexFromTime(float Time, const FPoseSearchIndexAsset& SearchIndexAsset) const
+{
+	const bool bIsLooping = IsSourceAssetLooping(SearchIndexAsset);
+	const FFloatInterval& Range = SearchIndexAsset.SamplingInterval;
+	const bool bHasPoseIndex = SearchIndexAsset.FirstPoseIdx != INDEX_NONE && SearchIndexAsset.NumPoses > 0 && (bIsLooping || Range.Contains(Time));
+	if (bHasPoseIndex)
+	{
+		int32 PoseOffset = FMath::RoundToInt(Schema->SampleRate * (Time - Range.Min));
+		
+		if (PoseOffset < 0)
+		{
+			if (bIsLooping)
+			{
+				PoseOffset = (PoseOffset % SearchIndexAsset.NumPoses) + SearchIndexAsset.NumPoses;
+			}
+			else
+			{
+				PoseOffset = 0;
+			}
+		}
+		else if (PoseOffset >= SearchIndexAsset.NumPoses)
+		{
+			if (bIsLooping)
+			{
+				PoseOffset = PoseOffset % SearchIndexAsset.NumPoses;
+			}
+			else
+			{
+				PoseOffset = SearchIndexAsset.NumPoses - 1;
+			}
+		}
+
+		int32 PoseIdx = SearchIndexAsset.FirstPoseIdx + PoseOffset;
+		return PoseIdx;
+	}
+
+	return INDEX_NONE;
+}
+
+bool UPoseSearchDatabase::GetPoseIndicesAndLerpValueFromTime(float Time, const FPoseSearchIndexAsset& SearchIndexAsset, int32& PrevPoseIdx, int32& PoseIdx, int32& NextPoseIdx, float& LerpValue) const
+{
+	PoseIdx = GetPoseIndexFromTime(Time, SearchIndexAsset);
+	if (PoseIdx == INDEX_NONE)
+	{
+		PrevPoseIdx = INDEX_NONE;
+		NextPoseIdx = INDEX_NONE;
+		LerpValue = 0.f;
+		return false;
+	}
+
+	const FFloatInterval& Range = SearchIndexAsset.SamplingInterval;
+	const float FloatPoseOffset = Schema->SampleRate * (Time - Range.Min);
+	const int32 PoseOffset = FMath::RoundToInt(FloatPoseOffset);
+	LerpValue = FloatPoseOffset - float(PoseOffset);
+
+	const float PrevTime = Time - 1.f / Schema->SampleRate;
+	const float NextTime = Time + 1.f / Schema->SampleRate;
+
+	PrevPoseIdx = GetPoseIndexFromTime(PrevTime, SearchIndexAsset);
+	if (PrevPoseIdx == INDEX_NONE)
+	{
+		PrevPoseIdx = PoseIdx;
+	}
+
+	NextPoseIdx = GetPoseIndexFromTime(NextTime, SearchIndexAsset);
+	if (NextPoseIdx == INDEX_NONE)
+	{
+		NextPoseIdx = PoseIdx;
+	}
+
+	check(LerpValue >= -0.5f && LerpValue <= 0.5f);
+
+	return true;
+}
+
+const FInstancedStruct& UPoseSearchDatabase::GetAnimationAssetStruct(int32 AnimationAssetIndex) const
+{
+	check(AnimationAssets.IsValidIndex(AnimationAssetIndex));
+	return AnimationAssets[AnimationAssetIndex];
+}
+
+const FInstancedStruct& UPoseSearchDatabase::GetAnimationAssetStruct(const FPoseSearchIndexAsset& SearchIndexAsset) const
+{
+	return GetAnimationAssetStruct(SearchIndexAsset.SourceAssetIdx);
+}
+
+FInstancedStruct& UPoseSearchDatabase::GetMutableAnimationAssetStruct(int32 AnimationAssetIndex)
+{
+	check(AnimationAssets.IsValidIndex(AnimationAssetIndex));
+	return AnimationAssets[AnimationAssetIndex];
+}
+
+FInstancedStruct& UPoseSearchDatabase::GetMutableAnimationAssetStruct(const FPoseSearchIndexAsset& SearchIndexAsset)
+{
+	return GetMutableAnimationAssetStruct(SearchIndexAsset.SourceAssetIdx);
+}
+
+const FPoseSearchDatabaseAnimationAssetBase* UPoseSearchDatabase::GetAnimationAssetBase(int32 AnimationAssetIndex) const
+{
+	if (AnimationAssets.IsValidIndex(AnimationAssetIndex))
+	{
+		return AnimationAssets[AnimationAssetIndex].GetPtr<FPoseSearchDatabaseAnimationAssetBase>();
+	}
+
+	return nullptr;
+}
+
+const FPoseSearchDatabaseAnimationAssetBase* UPoseSearchDatabase::GetAnimationAssetBase(const FPoseSearchIndexAsset& SearchIndexAsset) const
+{
+	return GetAnimationAssetBase(SearchIndexAsset.SourceAssetIdx);
+}
+
+FPoseSearchDatabaseAnimationAssetBase* UPoseSearchDatabase::GetMutableAnimationAssetBase(int32 AnimationAssetIndex)
+{
+	if (AnimationAssets.IsValidIndex(AnimationAssetIndex))
+	{
+		return AnimationAssets[AnimationAssetIndex].GetMutablePtr<FPoseSearchDatabaseAnimationAssetBase>();
+	}
+
+	return nullptr;
+}
+
+FPoseSearchDatabaseAnimationAssetBase* UPoseSearchDatabase::GetMutableAnimationAssetBase(const FPoseSearchIndexAsset& SearchIndexAsset)
+{
+	return GetMutableAnimationAssetBase(SearchIndexAsset.SourceAssetIdx);
+}
+
+const bool UPoseSearchDatabase::IsSourceAssetLooping(const FPoseSearchIndexAsset& SearchIndexAsset) const
+{
+	return GetAnimationAssetBase(SearchIndexAsset.SourceAssetIdx)->IsLooping();
+}
+
+const FString UPoseSearchDatabase::GetSourceAssetName(const FPoseSearchIndexAsset& SearchIndexAsset) const
+{
+	return GetAnimationAssetBase(SearchIndexAsset.SourceAssetIdx)->GetName();
+}
+
+int32 UPoseSearchDatabase::GetNumberOfPrincipalComponents() const
+{
+	return FMath::Min<int32>(NumberOfPrincipalComponents, Schema->SchemaCardinality);
+}
+
+bool UPoseSearchDatabase::GetSkipSearchIfPossible() const
+{
+	if (PoseSearchMode == EPoseSearchMode::PCAKDTree_Validate || PoseSearchMode == EPoseSearchMode::PCAKDTree_Compare)
+	{
+		return false;
+	}
+
+	return bSkipSearchIfPossible;
+}
+
+
+
+void UPoseSearchDatabase::PostLoad()
+{
+#if WITH_EDITORONLY_DATA
+	for (const FPoseSearchDatabaseSequence& DatabaseSequence : Sequences_DEPRECATED)
+	{
+		AnimationAssets.Add(FInstancedStruct::Make(DatabaseSequence));
+	}
+	Sequences_DEPRECATED.Empty();
+
+	for (const FPoseSearchDatabaseBlendSpace& DatabaseBlendSpace : BlendSpaces_DEPRECATED)
+	{
+		AnimationAssets.Add(FInstancedStruct::Make(DatabaseBlendSpace));
+	}
+	BlendSpaces_DEPRECATED.Empty();
+#endif // WITH_EDITORONLY_DATA
+
+#if WITH_EDITOR
+	using namespace UE::PoseSearch;
+	FAsyncPoseSearchDatabasesManagement::RequestAsyncBuildIndex(this, ERequestAsyncBuildFlag::NewRequest);
+#endif
+
+	Super::PostLoad();
+}
+
+#if WITH_EDITOR
+void UPoseSearchDatabase::RegisterOnDerivedDataRebuild(const FOnDerivedDataRebuild& Delegate)
+{
+	OnDerivedDataRebuild.Add(Delegate);
+}
+void UPoseSearchDatabase::UnregisterOnDerivedDataRebuild(void* Unregister)
+{
+	OnDerivedDataRebuild.RemoveAll(Unregister);
+}
+
+void UPoseSearchDatabase::NotifyDerivedDataRebuild() const
+{
+	OnDerivedDataRebuild.Broadcast();
+}
+
+void UPoseSearchDatabase::BeginCacheForCookedPlatformData(const ITargetPlatform* TargetPlatform)
+{
+	using namespace UE::PoseSearch;
+	Super::BeginCacheForCookedPlatformData(TargetPlatform);
+	FAsyncPoseSearchDatabasesManagement::RequestAsyncBuildIndex(this, ERequestAsyncBuildFlag::NewRequest);
+}
+
+bool UPoseSearchDatabase::IsCachedCookedPlatformDataLoaded(const ITargetPlatform* TargetPlatform)
+{
+	using namespace UE::PoseSearch;
+	check(IsInGameThread());
+	return FAsyncPoseSearchDatabasesManagement::RequestAsyncBuildIndex(this, ERequestAsyncBuildFlag::ContinueRequest);
+}
+#endif // WITH_EDITOR
+
+void UPoseSearchDatabase::PostSaveRoot(FObjectPostSaveRootContext ObjectSaveContext)
+{
+#if WITH_EDITOR
+	using namespace UE::PoseSearch;
+	if (!IsTemplate() && !ObjectSaveContext.IsProceduralSave())
+	{
+		FAsyncPoseSearchDatabasesManagement::RequestAsyncBuildIndex(this, ERequestAsyncBuildFlag::NewRequest | ERequestAsyncBuildFlag::WaitForCompletion);
+	}
+#endif
+
+	Super::PostSaveRoot(ObjectSaveContext);
+}
+
+void UPoseSearchDatabase::Serialize(FArchive& Ar)
+{
+	Super::Serialize(Ar);
+
+	if (Ar.IsFilterEditorOnly())
+	{
+		if (Ar.IsLoading() || Ar.IsCooking())
+		{
+			Ar << SearchIndexPrivate;
+		}
+	}
+}
+
+UE::PoseSearch::FSearchResult UPoseSearchDatabase::Search(UE::PoseSearch::FSearchContext& SearchContext) const
+{
+	using namespace UE::PoseSearch;
+
+	FSearchResult Result;
+
+#if WITH_EDITOR
+	if (!FAsyncPoseSearchDatabasesManagement::RequestAsyncBuildIndex(this, ERequestAsyncBuildFlag::ContinueRequest))
+	{
+		return Result;
+	}
+#endif
+
+	if (PoseSearchMode == EPoseSearchMode::BruteForce || PoseSearchMode == EPoseSearchMode::PCAKDTree_Compare)
+	{
+		Result = SearchBruteForce(SearchContext);
+	}
+
+	if (PoseSearchMode != EPoseSearchMode::BruteForce)
+	{
+#if WITH_EDITORONLY_DATA
+		FPoseSearchCost BruteForcePoseCost = Result.BruteForcePoseCost;
+#endif
+
+		Result = SearchPCAKDTree(SearchContext);
+
+#if WITH_EDITORONLY_DATA
+		Result.BruteForcePoseCost = BruteForcePoseCost;
+		if (PoseSearchMode == EPoseSearchMode::PCAKDTree_Compare)
+		{
+			check(Result.BruteForcePoseCost.GetTotalCost() <= Result.PoseCost.GetTotalCost());
+		}
+#endif
+	}
+	
+	return Result;
+}
+
+UE::PoseSearch::FSearchResult UPoseSearchDatabase::SearchPCAKDTree(UE::PoseSearch::FSearchContext& SearchContext) const
+{
+	QUICK_SCOPE_CYCLE_COUNTER(STAT_PoseSearch_PCA_KNN);
+	SCOPE_CYCLE_COUNTER(STAT_PoseSearchPCAKNN);
+
+	using namespace UE::PoseSearch;
+
+	FSearchResult Result;
+
+	const int32 NumDimensions = Schema->SchemaCardinality;
+	const FPoseSearchIndex& SearchIndex = GetSearchIndex();
+
+	const uint32 ClampedNumberOfPrincipalComponents = GetNumberOfPrincipalComponents();
+	const uint32 ClampedKDTreeQueryNumNeighbors = FMath::Clamp<uint32>(KDTreeQueryNumNeighbors, 1, SearchIndex.NumPoses);
+
+	//stack allocated temporaries
+	TArrayView<size_t> ResultIndexes((size_t*)FMemory_Alloca((ClampedKDTreeQueryNumNeighbors + 1) * sizeof(size_t)), ClampedKDTreeQueryNumNeighbors + 1);
+	TArrayView<float> ResultDistanceSqr((float*)FMemory_Alloca((ClampedKDTreeQueryNumNeighbors + 1) * sizeof(float)), ClampedKDTreeQueryNumNeighbors + 1);
+	RowMajorVectorMap WeightedQueryValues((float*)FMemory_Alloca(NumDimensions * sizeof(float)), 1, NumDimensions);
+	RowMajorVectorMap CenteredQueryValues((float*)FMemory_Alloca(NumDimensions * sizeof(float)), 1, NumDimensions);
+	RowMajorVectorMap ProjectedQueryValues((float*)FMemory_Alloca(ClampedNumberOfPrincipalComponents * sizeof(float)), 1, ClampedNumberOfPrincipalComponents);
+	
+	// KDTree in PCA space search
+	if (PoseSearchMode == EPoseSearchMode::PCAKDTree_Validate)
+	{
+		const RowMajorVectorMapConst MapWeightsSqrt(SearchIndex.WeightsSqrt.GetData(), 1, NumDimensions);
+
+		// testing the KDTree is returning the proper searches for all the original points transformed in pca space
+		for (int32 PoseIdx = 0; PoseIdx < SearchIndex.NumPoses; ++PoseIdx)
+		{
+			FKDTree::KNNResultSet ResultSet(ClampedKDTreeQueryNumNeighbors, ResultIndexes, ResultDistanceSqr);
+			TConstArrayView<float> PoseValues = SearchIndex.GetPoseValues(PoseIdx);
+
+			const RowMajorVectorMapConst Mean(SearchIndex.Mean.GetData(), 1, NumDimensions);
+			const ColMajorMatrixMapConst PCAProjectionMatrix(SearchIndex.PCAProjectionMatrix.GetData(), NumDimensions, ClampedNumberOfPrincipalComponents);
+
+			const RowMajorVectorMapConst QueryValues(PoseValues.GetData(), 1, NumDimensions);
+			WeightedQueryValues = QueryValues.array() * MapWeightsSqrt.array();
+			CenteredQueryValues.noalias() = WeightedQueryValues - Mean;
+			ProjectedQueryValues.noalias() = CenteredQueryValues * PCAProjectionMatrix;
+
+			SearchIndex.KDTree.FindNeighbors(ResultSet, ProjectedQueryValues.data());
+
+			size_t ResultIndex = 0;
+			for (; ResultIndex < ResultSet.Num(); ++ResultIndex)
+			{
+				if (PoseIdx == ResultIndexes[ResultIndex])
+				{
+					check(ResultDistanceSqr[ResultIndex] < UE_KINDA_SMALL_NUMBER);
+					break;
+				}
+			}
+			check(ResultIndex < ResultSet.Num());
+		}
+	}
+
+	SearchContext.GetOrBuildQuery(this, Result.ComposedQuery);
+
+	TConstArrayView<float> QueryValues = Result.ComposedQuery.GetValues();
+
+	const bool IsCurrentResultFromThisDatabase = SearchContext.IsCurrentResultFromDatabase(this);
+
+	// evaluating the continuing pose only if it hasn't already being evaluated and the related animation can advance
+	if (!SearchContext.bForceInterrupt && IsCurrentResultFromThisDatabase && SearchContext.bCanAdvance && !Result.ContinuingPoseCost.IsValid())
+	{
+		Result.PoseIdx = SearchContext.CurrentResult.PoseIdx;
+		Result.PoseCost = SearchIndex.ComparePoses(Result.PoseIdx, SearchContext.QueryMirrorRequest, EPoseComparisonFlags::ContinuingPose, Schema->MirrorMismatchCostBias, QueryValues);
+		Result.ContinuingPoseCost = Result.PoseCost;
+
+		if (GetSkipSearchIfPossible())
+		{
+			SearchContext.UpdateCurrentBestCost(Result.PoseCost);
+		}
+	}
+
+	// since any PoseCost calculated here is at least SearchIndex.MinCostAddend,
+	// there's no point in performing the search if CurrentBestTotalCost is already better than that
+	if (SearchContext.GetCurrentBestTotalCost() > SearchIndex.MinCostAddend)
+	{
+		FNonSelectableIdx NonSelectableIdx;
+		PopulateNonSelectableIdx(NonSelectableIdx, SearchContext, this, QueryValues);
+		FKDTree::KNNResultSet ResultSet(ClampedKDTreeQueryNumNeighbors, ResultIndexes, ResultDistanceSqr, NonSelectableIdx);
+
+		check(QueryValues.Num() == NumDimensions);
+
+		const RowMajorVectorMapConst Mean(SearchIndex.Mean.GetData(), 1, NumDimensions);
+		const ColMajorMatrixMapConst PCAProjectionMatrix(SearchIndex.PCAProjectionMatrix.GetData(), NumDimensions, ClampedNumberOfPrincipalComponents);
+
+		// transforming query values into PCA space to query the KDTree
+		const RowMajorVectorMapConst QueryValuesMap(QueryValues.GetData(), 1, NumDimensions);
+		const RowMajorVectorMapConst MapWeightsSqrt(SearchIndex.WeightsSqrt.GetData(), 1, NumDimensions);
+		WeightedQueryValues = QueryValuesMap.array() * MapWeightsSqrt.array();
+		CenteredQueryValues.noalias() = WeightedQueryValues - Mean;
+		ProjectedQueryValues.noalias() = CenteredQueryValues * PCAProjectionMatrix;
+
+		SearchIndex.KDTree.FindNeighbors(ResultSet, ProjectedQueryValues.data());
+
+		// NonSelectableIdx are already filtered out inside the kdtree search
+		const FPoseFilters PoseFilters(Schema, TConstArrayView<size_t>(), SearchIndex.OverallFlags);
+		for (size_t ResultIndex = 0; ResultIndex < ResultSet.Num(); ++ResultIndex)
+		{
+			const int32 PoseIdx = ResultIndexes[ResultIndex];
+			if (PoseFilters.AreFiltersValid(SearchIndex, QueryValues, PoseIdx, SearchIndex.PoseMetadata[PoseIdx]
+#if UE_POSE_SEARCH_TRACE_ENABLED
+				, SearchContext, this
+#endif
+			))
+			{
+				const FPoseSearchCost PoseCost = SearchIndex.ComparePoses(PoseIdx, SearchContext.QueryMirrorRequest, EPoseComparisonFlags::None, Schema->MirrorMismatchCostBias, QueryValues);
+				if (PoseCost < Result.PoseCost)
+				{
+					Result.PoseCost = PoseCost;
+					Result.PoseIdx = PoseIdx;
+				}
+
+#if UE_POSE_SEARCH_TRACE_ENABLED
+				SearchContext.BestCandidates.Add(PoseCost, PoseIdx, this, EPoseCandidateFlags::Valid_Pose);
+#endif
+			}
+		}
+
+		if (GetSkipSearchIfPossible() && Result.PoseCost.IsValid())
+		{
+			SearchContext.UpdateCurrentBestCost(Result.PoseCost);
+		}
+	}
+	else
+	{
+#if UE_POSE_SEARCH_TRACE_ENABLED
+		// calling just for reporting non selectable poses
+		FNonSelectableIdx NonSelectableIdx;
+		PopulateNonSelectableIdx(NonSelectableIdx, SearchContext, this, QueryValues);
+#endif
+	}
+
+	// finalizing Result properties
+	if (Result.PoseIdx != INDEX_NONE)
+	{
+		Result.AssetTime = SearchIndex.GetAssetTime(Result.PoseIdx, Schema->GetSamplingInterval());
+		Result.Database = this;
+	}
+
+	return Result;
+}
+
+UE::PoseSearch::FSearchResult UPoseSearchDatabase::SearchBruteForce(UE::PoseSearch::FSearchContext& SearchContext) const
+{
+	QUICK_SCOPE_CYCLE_COUNTER(STAT_PoseSearch_Brute_Force);
+	SCOPE_CYCLE_COUNTER(STAT_PoseSearchBruteForce);
+	
+	using namespace UE::PoseSearch;
+	
+	FSearchResult Result;
+
+	const FPoseSearchIndex& SearchIndex = GetSearchIndex();
+
+	SearchContext.GetOrBuildQuery(this, Result.ComposedQuery);
+	TConstArrayView<float> QueryValues = Result.ComposedQuery.GetValues();
+
+	const bool IsCurrentResultFromThisDatabase = SearchContext.IsCurrentResultFromDatabase(this);
+	if (!SearchContext.bForceInterrupt && IsCurrentResultFromThisDatabase)
+	{
+		// evaluating the continuing pose only if it hasn't already being evaluated and the related animation can advance
+		if (SearchContext.bCanAdvance && !Result.ContinuingPoseCost.IsValid())
+		{
+			Result.PoseIdx = SearchContext.CurrentResult.PoseIdx;
+			Result.PoseCost = SearchIndex.ComparePoses(Result.PoseIdx, SearchContext.QueryMirrorRequest, EPoseComparisonFlags::ContinuingPose, Schema->MirrorMismatchCostBias, QueryValues);
+			Result.ContinuingPoseCost = Result.PoseCost;
+
+			if (GetSkipSearchIfPossible())
+			{
+				SearchContext.UpdateCurrentBestCost(Result.PoseCost);
+			}
+		}
+	}
+
+	// since any PoseCost calculated here is at least SearchIndex.MinCostAddend,
+	// there's no point in performing the search if CurrentBestTotalCost is already better than that
+	if (SearchContext.GetCurrentBestTotalCost() > SearchIndex.MinCostAddend)
+	{
+		FNonSelectableIdx NonSelectableIdx;
+		PopulateNonSelectableIdx(NonSelectableIdx, SearchContext, this, QueryValues);
+		check(Algo::IsSorted(NonSelectableIdx));
+
+		const FPoseFilters PoseFilters(Schema, NonSelectableIdx, SearchIndex.OverallFlags);
+		for (int32 PoseIdx = 0; PoseIdx < SearchIndex.NumPoses; ++PoseIdx)
+		{
+			if (PoseFilters.AreFiltersValid(SearchIndex, QueryValues, PoseIdx, SearchIndex.PoseMetadata[PoseIdx]
+#if UE_POSE_SEARCH_TRACE_ENABLED
+				, SearchContext, this
+#endif
+			))
+			{
+				const FPoseSearchCost PoseCost = SearchIndex.ComparePoses(PoseIdx, SearchContext.QueryMirrorRequest, EPoseComparisonFlags::None, Schema->MirrorMismatchCostBias, QueryValues);
+				if (PoseCost < Result.PoseCost)
+				{
+					Result.PoseCost = PoseCost;
+					Result.PoseIdx = PoseIdx;
+				}
+
+#if UE_POSE_SEARCH_TRACE_ENABLED
+				if (PoseSearchMode == EPoseSearchMode::BruteForce)
+				{
+					SearchContext.BestCandidates.Add(PoseCost, PoseIdx, this, EPoseCandidateFlags::Valid_Pose);
+				}
+#endif
+			}
+		}
+
+		if (GetSkipSearchIfPossible() && Result.PoseCost.IsValid())
+		{
+			SearchContext.UpdateCurrentBestCost(Result.PoseCost);
+		}
+	}
+	else
+	{
+#if UE_POSE_SEARCH_TRACE_ENABLED
+		// calling just for reporting non selectable poses
+		FNonSelectableIdx NonSelectableIdx;
+		PopulateNonSelectableIdx(NonSelectableIdx, SearchContext, this, QueryValues);
+#endif
+	}
+
+	// finalizing Result properties
+	if (Result.PoseIdx != INDEX_NONE)
+	{
+		Result.AssetTime = SearchIndex.GetAssetTime(Result.PoseIdx, Schema->GetSamplingInterval());
+		Result.Database = this;
+	}
+
+#if WITH_EDITORONLY_DATA
+	Result.BruteForcePoseCost = Result.PoseCost; 
+#endif
+
+	return Result;
+}
+
+void UPoseSearchDatabase::BuildQuery(UE::PoseSearch::FSearchContext& SearchContext, FPoseSearchFeatureVectorBuilder& OutQuery) const
+{
+	check(Schema && Schema->IsValid());
+	Schema->BuildQuery(SearchContext, OutQuery);
+}
\ No newline at end of file
diff --git a/Engine/Plugins/Experimental/Animation/PoseSearch/Source/Runtime/Private/PoseSearchDatabaseIndexingContext.cpp b/Engine/Plugins/Experimental/Animation/PoseSearch/Source/Runtime/Private/PoseSearchDatabaseIndexingContext.cpp
new file mode 100644
--- /dev/null
+++ b/Engine/Plugins/Experimental/Animation/PoseSearch/Source/Runtime/Private/PoseSearchDatabaseIndexingContext.cpp
@@ -0,0 +1,215 @@
+// Copyright Epic Games, Inc. All Rights Reserved.
+
+#include "PoseSearchDatabaseIndexingContext.h"
+
+#if WITH_EDITOR
+
+#include "Animation/AnimComposite.h"
+#include "Animation/AnimSequence.h"
+#include "Animation/BlendSpace.h"
+#include "InstancedStruct.h"
+#include "PoseSearch/PoseSearchDatabase.h"
+#include "PoseSearch/PoseSearchDefines.h"
+#include "PoseSearch/PoseSearchSchema.h"
+
+namespace UE::PoseSearch
+{
+
+void FDatabaseIndexingContext::Prepare(const UPoseSearchDatabase* Database)
+{
+	const UPoseSearchSchema* Schema = Database->Schema;
+	check(Schema);
+
+	FBoneContainer BoneContainer;
+	BoneContainer.InitializeTo(Schema->BoneIndicesWithParents, FCurveEvaluationOption(false), *Schema->Skeleton);
+
+	TMap<const UAnimSequenceBase*, int32> SequenceSamplerMap;
+	TMap<TPair<const UBlendSpace*, FVector>, int32> BlendSpaceSamplerMap;
+
+	SamplingContext.Init(Schema->MirrorDataTable, BoneContainer);
+
+	// Prepare samplers for all animation assets.
+	for (const FInstancedStruct& DatabaseAssetStruct : Database->AnimationAssets)
+	{
+		auto AddSequenceBaseSampler = [&](const UAnimSequenceBase* Sequence)
+		{
+			if (Sequence && !SequenceSamplerMap.Contains(Sequence))
+			{
+				int32 SequenceSamplerIdx = SequenceSamplers.AddDefaulted();
+				SequenceSamplerMap.Add(Sequence, SequenceSamplerIdx);
+
+				FSequenceBaseSampler::FInput Input;
+				Input.ExtrapolationParameters = Database->ExtrapolationParameters;
+				Input.SequenceBase = Sequence;
+				SequenceSamplers[SequenceSamplerIdx].Init(Input);
+			}
+		};
+
+		if (const FPoseSearchDatabaseSequence* DatabaseSequence = DatabaseAssetStruct.GetPtr<FPoseSearchDatabaseSequence>())
+		{
+			AddSequenceBaseSampler(DatabaseSequence->Sequence);
+		}
+		else if (const FPoseSearchDatabaseAnimComposite* DatabaseAnimComposite = DatabaseAssetStruct.GetPtr<FPoseSearchDatabaseAnimComposite>())
+		{
+			AddSequenceBaseSampler(DatabaseAnimComposite->AnimComposite);
+		}
+		else if (const FPoseSearchDatabaseBlendSpace* DatabaseBlendSpace = DatabaseAssetStruct.GetPtr<FPoseSearchDatabaseBlendSpace>())
+		{
+			if (DatabaseBlendSpace->BlendSpace)
+			{
+				int32 HorizontalBlendNum, VerticalBlendNum;
+				DatabaseBlendSpace->GetBlendSpaceParameterSampleRanges(HorizontalBlendNum, VerticalBlendNum);
+
+				for (int32 HorizontalIndex = 0; HorizontalIndex < HorizontalBlendNum; HorizontalIndex++)
+				{
+					for (int32 VerticalIndex = 0; VerticalIndex < VerticalBlendNum; VerticalIndex++)
+					{
+						const FVector BlendParameters = DatabaseBlendSpace->BlendParameterForSampleRanges(HorizontalIndex, VerticalIndex);
+
+						if (!BlendSpaceSamplerMap.Contains({ DatabaseBlendSpace->BlendSpace, BlendParameters }))
+						{
+							int32 BlendSpaceSamplerIdx = BlendSpaceSamplers.AddDefaulted();
+							BlendSpaceSamplerMap.Add({ DatabaseBlendSpace->BlendSpace, BlendParameters }, BlendSpaceSamplerIdx);
+
+							FBlendSpaceSampler::FInput Input;
+							Input.BoneContainer = BoneContainer;
+							Input.ExtrapolationParameters = Database->ExtrapolationParameters;
+							Input.BlendSpace = DatabaseBlendSpace->BlendSpace;
+							Input.BlendParameters = BlendParameters;
+
+							BlendSpaceSamplers[BlendSpaceSamplerIdx].Init(Input);
+						}
+					}
+				}
+			}
+		}
+	}
+
+	TArray<IAssetSampler*, TInlineAllocator<512>> AssetSampler;
+	AssetSampler.SetNumUninitialized(SequenceSamplers.Num() + BlendSpaceSamplers.Num());
+
+	for (int i = 0; i < SequenceSamplers.Num(); ++i)
+	{
+		AssetSampler[i] = &SequenceSamplers[i];
+	}
+	for (int i = 0; i < BlendSpaceSamplers.Num(); ++i)
+	{
+		AssetSampler[i + SequenceSamplers.Num()] = &BlendSpaceSamplers[i];
+	}
+
+	ParallelFor(AssetSampler.Num(), [AssetSampler](int32 SamplerIdx) { AssetSampler[SamplerIdx]->Process(); }, ParallelForFlags);
+
+	// prepare indexers
+	Indexers.Reserve(SearchIndexBase->Assets.Num());
+
+	auto GetSequenceBaseSampler = [&](const UAnimSequenceBase* Sequence) -> const FSequenceBaseSampler*
+	{
+		return Sequence ? &SequenceSamplers[SequenceSamplerMap[Sequence]] : nullptr;
+	};
+
+	auto GetBlendSpaceSampler = [&](const UBlendSpace* BlendSpace, const FVector BlendParameters) -> const FBlendSpaceSampler*
+	{
+		return BlendSpace ? &BlendSpaceSamplers[BlendSpaceSamplerMap[{BlendSpace, BlendParameters}]] : nullptr;
+	};
+
+	Indexers.Reserve(SearchIndexBase->Assets.Num());
+
+	for (int32 AssetIdx = 0; AssetIdx != SearchIndexBase->Assets.Num(); ++AssetIdx)
+	{
+		const FPoseSearchIndexAsset& SearchIndexAsset = SearchIndexBase->Assets[AssetIdx];
+
+		FAssetIndexingContext IndexerContext;
+		IndexerContext.SamplingContext = &SamplingContext;
+		IndexerContext.Schema = Schema;
+		IndexerContext.RequestedSamplingRange = SearchIndexAsset.SamplingInterval;
+		IndexerContext.bMirrored = SearchIndexAsset.bMirrored;
+
+		const FInstancedStruct& DatabaseAsset = Database->GetAnimationAssetStruct(SearchIndexAsset.SourceAssetIdx);
+		if (const FPoseSearchDatabaseSequence* DatabaseSequence = DatabaseAsset.GetPtr<FPoseSearchDatabaseSequence>())
+		{
+			const float SequenceLength = DatabaseSequence->Sequence->GetPlayLength();
+			IndexerContext.AssetSampler = GetSequenceBaseSampler(DatabaseSequence->Sequence);
+		}
+		else if (const FPoseSearchDatabaseAnimComposite* DatabaseAnimComposite = DatabaseAsset.GetPtr<FPoseSearchDatabaseAnimComposite>())
+		{
+			IndexerContext.AssetSampler = GetSequenceBaseSampler(DatabaseAnimComposite->AnimComposite);
+		}
+		else if (const FPoseSearchDatabaseBlendSpace* DatabaseBlendSpace = DatabaseAsset.GetPtr<FPoseSearchDatabaseBlendSpace>())
+		{
+			IndexerContext.AssetSampler = GetBlendSpaceSampler(DatabaseBlendSpace->BlendSpace, SearchIndexAsset.BlendParameters);
+		}
+
+		FAssetIndexer& Indexer = Indexers.AddDefaulted_GetRef();
+		Indexer.Init(IndexerContext, BoneContainer);
+	}
+}
+
+bool FDatabaseIndexingContext::IndexAssets()
+{
+	// Index asset data
+	ParallelFor(Indexers.Num(), [this](int32 AssetIdx) { Indexers[AssetIdx].Process(); }, ParallelForFlags);
+	return true;
+}
+
+float FDatabaseIndexingContext::CalculateMinCostAddend() const
+{
+	float MinCostAddend = 0.f;
+
+	check(SearchIndexBase);
+	if (!SearchIndexBase->PoseMetadata.IsEmpty())
+	{
+		MinCostAddend = MAX_FLT;
+		for (const FPoseSearchPoseMetadata& PoseMetadata : SearchIndexBase->PoseMetadata)
+		{
+			if (PoseMetadata.CostAddend < MinCostAddend)
+			{
+				MinCostAddend = PoseMetadata.CostAddend;
+			}
+		}
+	}
+	return MinCostAddend;
+}
+
+void FDatabaseIndexingContext::JoinIndex()
+{
+	// Write index info to asset and count up total poses and storage required
+	int32 TotalPoses = 0;
+	int32 TotalFloats = 0;
+
+	check(SearchIndexBase);
+
+	// Join animation data into a single search index
+	SearchIndexBase->Values.Reset();
+	SearchIndexBase->PoseMetadata.Reset();
+	SearchIndexBase->OverallFlags = EPoseSearchPoseFlags::None;
+
+	for (int32 AssetIdx = 0; AssetIdx != SearchIndexBase->Assets.Num(); ++AssetIdx)
+	{
+		const FAssetIndexer::FOutput& Output = Indexers[AssetIdx].Output;
+
+		FPoseSearchIndexAsset& SearchIndexAsset = SearchIndexBase->Assets[AssetIdx];
+		SearchIndexAsset.NumPoses = Output.NumIndexedPoses;
+		SearchIndexAsset.FirstPoseIdx = TotalPoses;
+
+		const int32 PoseMetadataStartIdx = SearchIndexBase->PoseMetadata.Num();
+		const int32 PoseMetadataEndIdx = PoseMetadataStartIdx + Output.PoseMetadata.Num();
+
+		SearchIndexBase->Values.Append(Output.FeatureVectorTable.GetData(), Output.FeatureVectorTable.Num());
+		SearchIndexBase->PoseMetadata.Append(Output.PoseMetadata);
+
+		for (int32 i = PoseMetadataStartIdx; i < PoseMetadataEndIdx; ++i)
+		{
+			SearchIndexBase->PoseMetadata[i].AssetIndex = AssetIdx;
+			SearchIndexBase->OverallFlags |= SearchIndexBase->PoseMetadata[i].Flags;
+		}
+
+		TotalPoses += Output.NumIndexedPoses;
+		TotalFloats += Output.FeatureVectorTable.Num();
+	}
+
+	SearchIndexBase->NumPoses = TotalPoses;
+	SearchIndexBase->MinCostAddend = CalculateMinCostAddend();
+}
+
+} // namespace UE::PoseSearch
+#endif // WITH_EDITOR
diff --git a/Engine/Plugins/Experimental/Animation/PoseSearch/Source/Runtime/Private/PoseSearchDatabaseIndexingContext.h b/Engine/Plugins/Experimental/Animation/PoseSearch/Source/Runtime/Private/PoseSearchDatabaseIndexingContext.h
new file mode 100644
--- /dev/null
+++ b/Engine/Plugins/Experimental/Animation/PoseSearch/Source/Runtime/Private/PoseSearchDatabaseIndexingContext.h
@@ -0,0 +1,34 @@
+// Copyright Epic Games, Inc. All Rights Reserved.
+
+#pragma once
+
+#if WITH_EDITOR
+
+#include "PoseSearch/PoseSearchAssetSampler.h"
+#include "PoseSearchAssetIndexer.h"
+
+struct FPoseSearchIndexBase;
+class UPoseSearchDatabase;
+
+namespace UE::PoseSearch
+{
+
+struct FDatabaseIndexingContext
+{
+	FPoseSearchIndexBase* SearchIndexBase = nullptr;
+
+	FAssetSamplingContext SamplingContext;
+	TArray<FSequenceBaseSampler> SequenceSamplers; // Composite and sequence samplers
+	TArray<FBlendSpaceSampler> BlendSpaceSamplers;
+
+	TArray<FAssetIndexer> Indexers;
+
+	void Prepare(const UPoseSearchDatabase* Database);
+	bool IndexAssets();
+	void JoinIndex();
+	float CalculateMinCostAddend() const;
+};
+
+} // namespace UE::PoseSearch
+
+#endif // WITH_EDITOR
\ No newline at end of file
diff --git a/Engine/Plugins/Experimental/Animation/PoseSearch/Source/Runtime/Private/PoseSearchDatabaseSet.cpp b/Engine/Plugins/Experimental/Animation/PoseSearch/Source/Runtime/Private/PoseSearchDatabaseSet.cpp
new file mode 100644
--- /dev/null
+++ b/Engine/Plugins/Experimental/Animation/PoseSearch/Source/Runtime/Private/PoseSearchDatabaseSet.cpp
@@ -0,0 +1,102 @@
+// Copyright Epic Games, Inc. All Rights Reserved.
+
+#include "PoseSearch/PoseSearchDatabaseSet.h"
+#include "PoseSearch/PoseSearchContext.h"
+#include "PoseSearch/PoseSearchDatabase.h"
+#include "PoseSearch/PoseSearchDerivedData.h"
+#include "PoseSearch/PoseSearchSchema.h"
+
+UE::PoseSearch::FSearchResult UPoseSearchDatabaseSet::Search(UE::PoseSearch::FSearchContext& SearchContext) const
+{
+	using namespace UE::PoseSearch;
+
+	FSearchResult Result;
+	FPoseSearchCost ContinuingCost;
+#if WITH_EDITOR
+	FPoseSearchCost BruteForcePoseCost;
+#endif
+
+	// evaluating the continuing pose before all the active entries
+	const UPoseSearchDatabase* Database = SearchContext.CurrentResult.Database.Get();
+	if (bEvaluateContinuingPoseFirst &&
+		SearchContext.bCanAdvance &&
+		!SearchContext.bForceInterrupt &&
+		SearchContext.CurrentResult.IsValid()
+#if WITH_EDITOR
+		&& FAsyncPoseSearchDatabasesManagement::RequestAsyncBuildIndex(Database, ERequestAsyncBuildFlag::ContinueRequest)
+#endif
+		)
+	{
+		check(Database);
+		const FPoseSearchIndex& SearchIndex = Database->GetSearchIndex();
+		SearchContext.GetOrBuildQuery(Database, Result.ComposedQuery);
+
+		TConstArrayView<float> QueryValues = Result.ComposedQuery.GetValues();
+
+		Result.PoseIdx = SearchContext.CurrentResult.PoseIdx;
+		Result.PoseCost = SearchIndex.ComparePoses(Result.PoseIdx, SearchContext.QueryMirrorRequest, EPoseComparisonFlags::ContinuingPose, Database->Schema->MirrorMismatchCostBias, QueryValues);
+		Result.ContinuingPoseCost = Result.PoseCost;
+		ContinuingCost = Result.PoseCost;
+
+		Result.AssetTime = SearchIndex.GetAssetTime(Result.PoseIdx, Database->Schema->GetSamplingInterval());
+		Result.Database = Database;
+
+		if (Database->GetSkipSearchIfPossible())
+		{
+			SearchContext.UpdateCurrentBestCost(Result.PoseCost);
+		}
+	}
+
+	for (const FPoseSearchDatabaseSetEntry& Entry : AssetsToSearch)
+	{
+		if (!IsValid(Entry.Searchable))
+		{
+			UE_LOG(LogPoseSearch, Warning, TEXT("Invalid entry in Database Set %s"), *GetName());
+			continue;
+		}
+
+		const bool bSearchEntry =
+			!Entry.Tag.IsValid() ||
+			SearchContext.ActiveTagsContainer == nullptr ||
+			SearchContext.ActiveTagsContainer->IsEmpty() ||
+			SearchContext.ActiveTagsContainer->HasTag(Entry.Tag);
+
+		if (bSearchEntry)
+		{
+			FSearchResult EntryResult = Entry.Searchable->Search(SearchContext);
+
+			if (EntryResult.PoseCost.GetTotalCost() < Result.PoseCost.GetTotalCost())
+			{
+				Result = EntryResult;
+			}
+
+			if (EntryResult.ContinuingPoseCost.GetTotalCost() < ContinuingCost.GetTotalCost())
+			{
+				ContinuingCost = EntryResult.ContinuingPoseCost;
+			}
+#if WITH_EDITOR
+			if (EntryResult.BruteForcePoseCost.GetTotalCost() < BruteForcePoseCost.GetTotalCost())
+			{
+				BruteForcePoseCost = EntryResult.BruteForcePoseCost;
+			}
+#endif
+			if (Entry.PostSearchStatus == EPoseSearchPostSearchStatus::Stop)
+			{
+				break;
+			}
+		}
+	}
+
+	Result.ContinuingPoseCost = ContinuingCost;
+
+#if WITH_EDITOR
+	Result.BruteForcePoseCost = BruteForcePoseCost;
+#endif
+
+	if (!Result.IsValid())
+	{
+		UE_LOG(LogPoseSearch, Warning, TEXT("Invalid result searching %s"), *GetName());
+	}
+
+	return Result;
+}
diff --git a/Engine/Plugins/Experimental/Animation/PoseSearch/Source/Runtime/Private/PoseSearchDefines.cpp b/Engine/Plugins/Experimental/Animation/PoseSearch/Source/Runtime/Private/PoseSearchDefines.cpp
new file mode 100644
--- /dev/null
+++ b/Engine/Plugins/Experimental/Animation/PoseSearch/Source/Runtime/Private/PoseSearchDefines.cpp
@@ -0,0 +1,5 @@
+// Copyright Epic Games, Inc. All Rights Reserved.
+
+#include "PoseSearch/PoseSearchDefines.h"
+
+DEFINE_LOG_CATEGORY(LogPoseSearch);
diff --git a/Engine/Plugins/Experimental/Animation/PoseSearch/Source/Runtime/Private/PoseSearchDerivedData.cpp b/Engine/Plugins/Experimental/Animation/PoseSearch/Source/Runtime/Private/PoseSearchDerivedData.cpp
--- a/Engine/Plugins/Experimental/Animation/PoseSearch/Source/Runtime/Private/PoseSearchDerivedData.cpp
+++ b/Engine/Plugins/Experimental/Animation/PoseSearch/Source/Runtime/Private/PoseSearchDerivedData.cpp
@@ -3,13 +3,21 @@
 #include "PoseSearch/PoseSearchDerivedData.h"
 
 #if WITH_EDITOR
+
+#include "Animation/AnimComposite.h"
+#include "Animation/AnimSequence.h"
 #include "Animation/BlendSpace.h"
 #include "AssetRegistry/ARFilter.h"
 #include "AssetRegistry/AssetRegistryModule.h"
 #include "DerivedDataCache.h"
 #include "DerivedDataRequestOwner.h"
+#include "InstancedStruct.h"
 #include "Misc/CoreDelegates.h"
+#include "PoseSearchDatabaseIndexingContext.h"
 #include "PoseSearch/PoseSearchAnimNotifies.h"
+#include "PoseSearch/PoseSearchAssetIndexer.h"
+#include "PoseSearch/PoseSearchDatabase.h"
+#include "PoseSearch/PoseSearchDefines.h"
 #include "PoseSearch/PoseSearchDerivedDataKey.h"
 #include "PoseSearch/PoseSearchFeatureChannel.h"
 #include "PoseSearch/PoseSearchSchema.h"
@@ -146,6 +154,7 @@ static TArray<float> ComputeChannelsDeviations(TConstArrayView<FPoseSearchIndexB
 	
 	return MeanDeviations;
 }
+
 static inline FFloatInterval GetEffectiveSamplingRange(const UAnimSequenceBase* Sequence, FFloatInterval RequestedSamplingRange)
 {
 	const bool bSampleAll = (RequestedSamplingRange.Min == 0.0f) && (RequestedSamplingRange.Max == 0.0f);
@@ -508,585 +517,6 @@ static void PreprocessSearchIndexKDTree(FPoseSearchIndex& SearchIndex, int32 Num
 	}
 }
 
-//////////////////////////////////////////////////////////////////////////
-// FSamplingParam helpers
-struct FSamplingParam
-{
-	float WrappedParam = 0.0f;
-	int32 NumCycles = 0;
-
-	// If the animation can't loop, WrappedParam contains the clamped value and whatever is left is stored here
-	float Extrapolation = 0.0f;
-};
-
-static FSamplingParam WrapOrClampSamplingParam(bool bCanWrap, float SamplingParamExtent, float SamplingParam)
-{
-	// This is a helper function used by both time and distance sampling. A schema may specify time or distance
-	// offsets that are multiple cycles of a clip away from the current pose being sampled.
-	// And that time or distance offset may before the beginning of the clip (SamplingParam < 0.0f)
-	// or after the end of the clip (SamplingParam > SamplingParamExtent). So this function
-	// helps determine how many cycles need to be applied and what the wrapped value should be, clamping
-	// if necessary.
-
-	FSamplingParam Result;
-
-	Result.WrappedParam = SamplingParam;
-
-	const bool bIsSamplingParamExtentKindaSmall = SamplingParamExtent <= UE_KINDA_SMALL_NUMBER;
-	if (!bIsSamplingParamExtentKindaSmall && bCanWrap)
-	{
-		if (SamplingParam < 0.0f)
-		{
-			while (Result.WrappedParam < 0.0f)
-			{
-				Result.WrappedParam += SamplingParamExtent;
-				++Result.NumCycles;
-			}
-		}
-
-		else
-		{
-			while (Result.WrappedParam > SamplingParamExtent)
-			{
-				Result.WrappedParam -= SamplingParamExtent;
-				++Result.NumCycles;
-			}
-		}
-	}
-
-	const float ParamClamped = FMath::Clamp(Result.WrappedParam, 0.0f, SamplingParamExtent);
-	if (ParamClamped != Result.WrappedParam)
-	{
-		check(bIsSamplingParamExtentKindaSmall || !bCanWrap);
-		Result.Extrapolation = Result.WrappedParam - ParamClamped;
-		Result.WrappedParam = ParamClamped;
-	}
-
-	return Result;
-}
-
-//////////////////////////////////////////////////////////////////////////
-// FAssetIndexer
-class FAssetIndexer : public IAssetIndexer
-{
-public:
-
-	struct FOutput
-	{
-		int32 FirstIndexedSample = 0;
-		int32 LastIndexedSample = 0;
-		int32 NumIndexedPoses = 0;
-
-		TArray<float> FeatureVectorTable;
-		TArray<FPoseSearchPoseMetadata> PoseMetadata;
-		TBitArray<> AllFeaturesNotAdded;
-	} Output;
-
-	void Reset();
-	void Init(const FAssetIndexingContext& IndexingContext, const FBoneContainer& InBoneContainer);
-	bool Process();
-
-public: // IAssetIndexer
-
-	const FAssetIndexingContext& GetIndexingContext() const override { return IndexingContext; }
-	FSampleInfo GetSampleInfo(float SampleTime) const override;
-	FSampleInfo GetSampleInfoRelative(float SampleTime, const FSampleInfo& Origin) const override;
-	FTransform MirrorTransform(const FTransform& Transform) const override;
-	FTransform GetTransformAndCacheResults(float SampleTime, float OriginTime, int8 SchemaBoneIdx, bool& Clamped) override;
-
-private:
-	FPoseSearchPoseMetadata GetMetadata(int32 SampleIdx) const;
-
-	struct CachedEntry
-	{
-		float SampleTime;
-		float OriginTime;
-		bool Clamped;
-
-		// @todo: minimize the Entry memory footprint
-		FTransform RootTransform;
-		FCompactPose Pose;
-		FCSPose<FCompactPose> ComponentSpacePose;
-		FBlendedCurve UnusedCurve;
-		UE::Anim::FStackAttributeContainer UnusedAtrribute;
-		FAnimationPoseData AnimPoseData = { Pose, UnusedCurve, UnusedAtrribute };
-	};
-
-	FBoneContainer BoneContainer;
-	FAssetIndexingContext IndexingContext;
-	TArray<CachedEntry> CachedEntries;
-};
-
-void FAssetIndexer::Reset()
-{
-	Output.FirstIndexedSample = 0;
-	Output.LastIndexedSample = 0;
-	Output.NumIndexedPoses = 0;
-
-	Output.FeatureVectorTable.Reset(0);
-	Output.PoseMetadata.Reset(0);
-	Output.AllFeaturesNotAdded.Reset();
-}
-
-void FAssetIndexer::Init(const FAssetIndexingContext& InIndexingContext, const FBoneContainer& InBoneContainer)
-{
-	check(InIndexingContext.Schema);
-	check(InIndexingContext.Schema->IsValid());
-	check(InIndexingContext.AssetSampler);
-
-	BoneContainer = InBoneContainer;
-	IndexingContext = InIndexingContext;
-
-	Reset();
-
-	Output.FirstIndexedSample = FMath::FloorToInt(IndexingContext.RequestedSamplingRange.Min * IndexingContext.Schema->SampleRate);
-	Output.LastIndexedSample = FMath::Max(0, FMath::CeilToInt(IndexingContext.RequestedSamplingRange.Max * IndexingContext.Schema->SampleRate));
-	Output.NumIndexedPoses = Output.LastIndexedSample - Output.FirstIndexedSample + 1;
-
-	Output.FeatureVectorTable.SetNumZeroed(IndexingContext.Schema->SchemaCardinality * Output.NumIndexedPoses);
-	Output.PoseMetadata.SetNum(Output.NumIndexedPoses);
-}
-
-bool FAssetIndexer::Process()
-{
-	check(IndexingContext.Schema);
-	check(IndexingContext.Schema->IsValid());
-	check(IndexingContext.AssetSampler);
-
-	FMemMark Mark(FMemStack::Get());
-
-	IndexingContext.BeginSampleIdx = Output.FirstIndexedSample;
-	IndexingContext.EndSampleIdx = Output.LastIndexedSample + 1;
-
-	if (IndexingContext.Schema->SchemaCardinality > 0)
-	{
-		// Index each channel
-		FAssetIndexingOutput AssetIndexingOutput(IndexingContext.Schema->SchemaCardinality, Output.FeatureVectorTable);
-		for (int32 ChannelIdx = 0; ChannelIdx != IndexingContext.Schema->Channels.Num(); ++ChannelIdx)
-		{
-			IndexingContext.Schema->Channels[ChannelIdx]->IndexAsset(*this, AssetIndexingOutput);
-		}
-	}
-
-	// Generate pose metadata
-	for (int32 SampleIdx = IndexingContext.BeginSampleIdx; SampleIdx != IndexingContext.EndSampleIdx; ++SampleIdx)
-	{
-		const int32 PoseIdx = SampleIdx - Output.FirstIndexedSample;
-		FPoseSearchPoseMetadata& OutputPoseMetadata = Output.PoseMetadata[PoseIdx];
-		OutputPoseMetadata = GetMetadata(SampleIdx);
-	}
-
-	return true;
-}
-
-FAssetIndexer::FSampleInfo FAssetIndexer::GetSampleInfo(float SampleTime) const
-{
-	FSampleInfo Sample;
-
-	check(IndexingContext.AssetSampler);
-
-	const float PlayLength = IndexingContext.AssetSampler->GetPlayLength();
-	const bool bCanWrap = IndexingContext.AssetSampler->IsLoopable();
-
-	float MainRelativeTime = SampleTime;
-	if (SampleTime < 0.0f && bCanWrap)
-	{
-		// In this case we're sampling a loop backwards, so MainRelativeTime must adjust so the number of cycles is
-		// counted correctly.
-		MainRelativeTime += PlayLength;
-	}
-
-	const FSamplingParam SamplingParam = WrapOrClampSamplingParam(bCanWrap, PlayLength, MainRelativeTime);
-
-	Sample.Clip = IndexingContext.AssetSampler;
-
-	if (FMath::Abs(SamplingParam.Extrapolation) > SMALL_NUMBER)
-	{
-		Sample.bClamped = true;
-		Sample.ClipTime = SamplingParam.WrappedParam + SamplingParam.Extrapolation;
-		const FTransform ClipRootMotion = Sample.Clip->ExtractRootTransform(Sample.ClipTime);
-		const float ClipDistance = Sample.Clip->ExtractRootDistance(Sample.ClipTime);
-
-		Sample.RootTransform = ClipRootMotion;
-		Sample.RootDistance = ClipDistance;
-	}
-	else
-	{
-		Sample.ClipTime = SamplingParam.WrappedParam;
-
-		const FTransform RootMotionLast = IndexingContext.AssetSampler->GetTotalRootTransform();
-		float RootDistanceLast = IndexingContext.AssetSampler->GetTotalRootDistance();
-
-		// Determine how to accumulate motion for every cycle of the anim. If the sample
-		// had to be clamped, this motion will end up not getting applied below.
-		// Also invert the accumulation direction if the requested sample was wrapped backwards.
-		FTransform RootMotionPerCycle = RootMotionLast;
-		float RootDistancePerCycle = RootDistanceLast;
-		if (SampleTime < 0.0f)
-		{
-			RootMotionPerCycle = RootMotionPerCycle.Inverse();
-			RootDistancePerCycle *= -1.f;
-		}
-
-		// Find the remaining motion deltas after wrapping
-		FTransform RootMotionRemainder = Sample.Clip->ExtractRootTransform(Sample.ClipTime);
-		float RootDistanceRemainder = Sample.Clip->ExtractRootDistance(Sample.ClipTime);
-
-		// Invert motion deltas if we wrapped backwards
-		if (SampleTime < 0.0f)
-		{
-			RootMotionRemainder.SetToRelativeTransform(RootMotionLast);
-			RootDistanceRemainder = -(RootDistanceLast - RootDistanceRemainder);
-		}
-
-		Sample.RootTransform = FTransform::Identity;
-		Sample.RootDistance = 0.f;
-
-		// Note if the sample was clamped, no motion will be applied here because NumCycles will be zero
-		int32 CyclesRemaining = SamplingParam.NumCycles;
-		while (CyclesRemaining--)
-		{
-			Sample.RootTransform = RootMotionPerCycle * Sample.RootTransform;
-			Sample.RootDistance += RootDistancePerCycle;
-		}
-
-		Sample.RootTransform = RootMotionRemainder * Sample.RootTransform;
-		Sample.RootDistance += RootDistanceRemainder;
-	}
-
-	return Sample;
-}
-
-FAssetIndexer::FSampleInfo FAssetIndexer::GetSampleInfoRelative(float SampleTime, const FSampleInfo& Origin) const
-{
-	FSampleInfo Sample = GetSampleInfo(SampleTime);
-	Sample.RootTransform.SetToRelativeTransform(Origin.RootTransform);
-	Sample.RootDistance = Origin.RootDistance - Sample.RootDistance;
-	return Sample;
-}
-
-FTransform FAssetIndexer::MirrorTransform(const FTransform& Transform) const
-{
-	return IndexingContext.bMirrored ? IndexingContext.SamplingContext->MirrorTransform(Transform) : Transform;
-}
-
-FPoseSearchPoseMetadata FAssetIndexer::GetMetadata(int32 SampleIdx) const
-{
-	const float SequenceLength = IndexingContext.AssetSampler->GetPlayLength();
-	const float SampleTime = FMath::Min(SampleIdx * IndexingContext.Schema->GetSamplingInterval(), SequenceLength);
-
-	FPoseSearchPoseMetadata Metadata;
-	Metadata.CostAddend = IndexingContext.Schema->BaseCostBias;
-	Metadata.ContinuingPoseCostAddend = IndexingContext.Schema->ContinuingPoseCostBias;
-
-	TArray<UAnimNotifyState_PoseSearchBase*> NotifyStates;
-	IndexingContext.AssetSampler->ExtractPoseSearchNotifyStates(SampleTime, NotifyStates);
-	for (const UAnimNotifyState_PoseSearchBase* PoseSearchNotify : NotifyStates)
-	{
-		if (PoseSearchNotify->GetClass()->IsChildOf<UAnimNotifyState_PoseSearchBlockTransition>())
-		{
-			EnumAddFlags(Metadata.Flags, EPoseSearchPoseFlags::BlockTransition);
-		}
-		else if (PoseSearchNotify->GetClass()->IsChildOf<UAnimNotifyState_PoseSearchModifyCost>())
-		{
-			const UAnimNotifyState_PoseSearchModifyCost* ModifyCostNotify =
-				Cast<const UAnimNotifyState_PoseSearchModifyCost>(PoseSearchNotify);
-			Metadata.CostAddend = ModifyCostNotify->CostAddend;
-		}
-		else if (PoseSearchNotify->GetClass()->IsChildOf<UAnimNotifyState_PoseSearchOverrideContinuingPoseCostBias>())
-		{
-			const UAnimNotifyState_PoseSearchOverrideContinuingPoseCostBias* ContinuingPoseCostBias =
-				Cast<const UAnimNotifyState_PoseSearchOverrideContinuingPoseCostBias>(PoseSearchNotify);
-			Metadata.ContinuingPoseCostAddend = ContinuingPoseCostBias->CostAddend;
-		}
-	}
-	return Metadata;
-}
-
-FTransform FAssetIndexer::GetTransformAndCacheResults(float SampleTime, float OriginTime, int8 SchemaBoneIdx, bool& Clamped)
-{
-	// @todo: use an hashmap if we end up having too many entries
-	CachedEntry* Entry = CachedEntries.FindByPredicate([SampleTime, OriginTime](const FAssetIndexer::CachedEntry& Entry)
-		{
-			return Entry.SampleTime == SampleTime && Entry.OriginTime == OriginTime;
-		});
-
-	const FAssetSamplingContext* SamplingContext = IndexingContext.SamplingContext;
-
-	if (!Entry)
-	{
-		Entry = &CachedEntries[CachedEntries.AddDefaulted()];
-
-		Entry->SampleTime = SampleTime;
-		Entry->OriginTime = OriginTime;
-
-		if (!BoneContainer.IsValid())
-		{
-			UE_LOG(LogPoseSearch,
-				Warning,
-				TEXT("Invalid BoneContainer encountered in FAssetIndexer::GetTransformAndCacheResults. Asset: %s. Schema: %s. BoneContainerAsset: %s. NumBoneIndices: %d"),
-				*GetNameSafe(IndexingContext.AssetSampler->GetAsset()),
-				*GetNameSafe(IndexingContext.Schema),
-				*GetNameSafe(BoneContainer.GetAsset()),
-				BoneContainer.GetCompactPoseNumBones());
-		}
-
-		Entry->Pose.SetBoneContainer(&BoneContainer);
-		Entry->UnusedCurve.InitFrom(BoneContainer);
-
-		IAssetIndexer::FSampleInfo Origin = GetSampleInfo(OriginTime);
-		IAssetIndexer::FSampleInfo Sample = GetSampleInfoRelative(SampleTime, Origin);
-
-		float CurrentTime = Sample.ClipTime;
-		float PreviousTime = CurrentTime - SamplingContext->FiniteDelta;
-
-		check(Sample.Clip->IsLoopable() || PreviousTime <= Sample.Clip->GetPlayLength());
-
-		FDeltaTimeRecord DeltaTimeRecord;
-		DeltaTimeRecord.Set(PreviousTime, CurrentTime - PreviousTime);
-		FAnimExtractContext ExtractionCtx(static_cast<double>(CurrentTime), true, DeltaTimeRecord, Sample.Clip->IsLoopable());
-
-		Sample.Clip->ExtractPose(ExtractionCtx, Entry->AnimPoseData);
-
-		if (IndexingContext.bMirrored)
-		{
-			FAnimationRuntime::MirrorPose(
-				Entry->AnimPoseData.GetPose(),
-				IndexingContext.Schema->MirrorDataTable->MirrorAxis,
-				SamplingContext->CompactPoseMirrorBones,
-				SamplingContext->ComponentSpaceRefRotations
-			);
-			// Note curves and attributes are not used during the indexing process and therefore don't need to be mirrored
-		}
-
-		Entry->ComponentSpacePose.InitPose(Entry->Pose);
-		Entry->RootTransform = Sample.RootTransform;
-		Entry->Clamped = Sample.bClamped;
-	}
-
-	const FBoneReference& BoneReference = IndexingContext.Schema->BoneReferences[SchemaBoneIdx];
-	FCompactPoseBoneIndex CompactBoneIndex = BoneContainer.MakeCompactPoseIndex(FMeshPoseBoneIndex(BoneReference.BoneIndex));
-
-	const FTransform BoneTransform = Entry->ComponentSpacePose.GetComponentSpaceTransform(CompactBoneIndex) * MirrorTransform(Entry->RootTransform);
-	Clamped = Entry->Clamped;
-
-	return BoneTransform;
-}
-
-//////////////////////////////////////////////////////////////////////////
-// FDatabaseIndexingContext
-struct FDatabaseIndexingContext
-{
-	FPoseSearchIndexBase* SearchIndexBase = nullptr;
-
-	FAssetSamplingContext SamplingContext;
-	TArray<FSequenceBaseSampler> SequenceSamplers; // Composite and sequence samplers
-	TArray<FBlendSpaceSampler> BlendSpaceSamplers;
-
-	TArray<FAssetIndexer> Indexers;
-
-	void Prepare(const UPoseSearchDatabase* Database);
-	bool IndexAssets();
-	void JoinIndex();
-	float CalculateMinCostAddend() const;
-};
-
-void FDatabaseIndexingContext::Prepare(const UPoseSearchDatabase* Database)
-{
-	const UPoseSearchSchema* Schema = Database->Schema;
-	check(Schema);
-
-	FBoneContainer BoneContainer;
-	BoneContainer.InitializeTo(Schema->BoneIndicesWithParents, FCurveEvaluationOption(false), *Schema->Skeleton);
-
-	TMap<const UAnimSequenceBase*, int32> SequenceSamplerMap;
-	TMap<TPair<const UBlendSpace*, FVector>, int32> BlendSpaceSamplerMap;
-
-	SamplingContext.Init(Schema->MirrorDataTable, BoneContainer);
-
-	// Prepare samplers for all animation assets.
-	for (const FInstancedStruct& DatabaseAssetStruct : Database->AnimationAssets)
-	{
-		auto AddSequenceBaseSampler = [&](const UAnimSequenceBase* Sequence)
-		{
-			if (Sequence && !SequenceSamplerMap.Contains(Sequence))
-			{
-				int32 SequenceSamplerIdx = SequenceSamplers.AddDefaulted();
-				SequenceSamplerMap.Add(Sequence, SequenceSamplerIdx);
-
-				FSequenceBaseSampler::FInput Input;
-				Input.ExtrapolationParameters = Database->ExtrapolationParameters;
-				Input.SequenceBase = Sequence;
-				SequenceSamplers[SequenceSamplerIdx].Init(Input);
-			}
-		};
-
-		if (const FPoseSearchDatabaseSequence* DatabaseSequence = DatabaseAssetStruct.GetPtr<FPoseSearchDatabaseSequence>())
-		{
-			AddSequenceBaseSampler(DatabaseSequence->Sequence);
-		}
-		else if (const FPoseSearchDatabaseAnimComposite* DatabaseAnimComposite = DatabaseAssetStruct.GetPtr<FPoseSearchDatabaseAnimComposite>())
-		{
-			AddSequenceBaseSampler(DatabaseAnimComposite->AnimComposite);
-		}
-		else if (const FPoseSearchDatabaseBlendSpace* DatabaseBlendSpace = DatabaseAssetStruct.GetPtr<FPoseSearchDatabaseBlendSpace>())
-		{
-			if (DatabaseBlendSpace->BlendSpace)
-			{
-				int32 HorizontalBlendNum, VerticalBlendNum;
-				DatabaseBlendSpace->GetBlendSpaceParameterSampleRanges(HorizontalBlendNum, VerticalBlendNum);
-
-				for (int32 HorizontalIndex = 0; HorizontalIndex < HorizontalBlendNum; HorizontalIndex++)
-				{
-					for (int32 VerticalIndex = 0; VerticalIndex < VerticalBlendNum; VerticalIndex++)
-					{
-						const FVector BlendParameters = DatabaseBlendSpace->BlendParameterForSampleRanges(HorizontalIndex, VerticalIndex);
-
-						if (!BlendSpaceSamplerMap.Contains({ DatabaseBlendSpace->BlendSpace, BlendParameters }))
-						{
-							int32 BlendSpaceSamplerIdx = BlendSpaceSamplers.AddDefaulted();
-							BlendSpaceSamplerMap.Add({ DatabaseBlendSpace->BlendSpace, BlendParameters }, BlendSpaceSamplerIdx);
-
-							FBlendSpaceSampler::FInput Input;
-							Input.BoneContainer = BoneContainer;
-							Input.ExtrapolationParameters = Database->ExtrapolationParameters;
-							Input.BlendSpace = DatabaseBlendSpace->BlendSpace;
-							Input.BlendParameters = BlendParameters;
-
-							BlendSpaceSamplers[BlendSpaceSamplerIdx].Init(Input);
-						}
-					}
-				}
-			}
-		}
-	}
-
-	TArray<IAssetSampler*, TInlineAllocator<512>> AssetSampler;
-	AssetSampler.SetNumUninitialized(SequenceSamplers.Num() + BlendSpaceSamplers.Num());
-
-	for (int i = 0; i < SequenceSamplers.Num(); ++i)
-	{
-		AssetSampler[i] = &SequenceSamplers[i];
-	}
-	for (int i = 0; i < BlendSpaceSamplers.Num(); ++i)
-	{
-		AssetSampler[i + SequenceSamplers.Num()] = &BlendSpaceSamplers[i];
-	}
-
-	ParallelFor(AssetSampler.Num(), [AssetSampler](int32 SamplerIdx) { AssetSampler[SamplerIdx]->Process(); }, ParallelForFlags);
-
-	// prepare indexers
-	Indexers.Reserve(SearchIndexBase->Assets.Num());
-
-	auto GetSequenceBaseSampler = [&](const UAnimSequenceBase* Sequence) -> const FSequenceBaseSampler*
-	{
-		return Sequence ? &SequenceSamplers[SequenceSamplerMap[Sequence]] : nullptr;
-	};
-
-	auto GetBlendSpaceSampler = [&](const UBlendSpace* BlendSpace, const FVector BlendParameters) -> const FBlendSpaceSampler*
-	{
-		return BlendSpace ? &BlendSpaceSamplers[BlendSpaceSamplerMap[{BlendSpace, BlendParameters}]] : nullptr;
-	};
-
-	Indexers.Reserve(SearchIndexBase->Assets.Num());
-
-	for (int32 AssetIdx = 0; AssetIdx != SearchIndexBase->Assets.Num(); ++AssetIdx)
-	{
-		const FPoseSearchIndexAsset& SearchIndexAsset = SearchIndexBase->Assets[AssetIdx];
-
-		FAssetIndexingContext IndexerContext;
-		IndexerContext.SamplingContext = &SamplingContext;
-		IndexerContext.Schema = Schema;
-		IndexerContext.RequestedSamplingRange = SearchIndexAsset.SamplingInterval;
-		IndexerContext.bMirrored = SearchIndexAsset.bMirrored;
-
-		const FInstancedStruct& DatabaseAsset = Database->GetAnimationAssetStruct(SearchIndexAsset.SourceAssetIdx);
-		if (const FPoseSearchDatabaseSequence* DatabaseSequence = DatabaseAsset.GetPtr<FPoseSearchDatabaseSequence>())
-		{
-			const float SequenceLength = DatabaseSequence->Sequence->GetPlayLength();
-			IndexerContext.AssetSampler = GetSequenceBaseSampler(DatabaseSequence->Sequence);
-		}
-		else if (const FPoseSearchDatabaseAnimComposite* DatabaseAnimComposite = DatabaseAsset.GetPtr<FPoseSearchDatabaseAnimComposite>())
-		{
-			IndexerContext.AssetSampler = GetSequenceBaseSampler(DatabaseAnimComposite->AnimComposite);
-		}
-		else if (const FPoseSearchDatabaseBlendSpace* DatabaseBlendSpace = DatabaseAsset.GetPtr<FPoseSearchDatabaseBlendSpace>())
-		{
-			IndexerContext.AssetSampler = GetBlendSpaceSampler(DatabaseBlendSpace->BlendSpace, SearchIndexAsset.BlendParameters);
-		}
-
-		FAssetIndexer& Indexer = Indexers.AddDefaulted_GetRef();
-		Indexer.Init(IndexerContext, BoneContainer);
-	}
-}
-
-bool FDatabaseIndexingContext::IndexAssets()
-{
-	// Index asset data
-	ParallelFor(Indexers.Num(), [this](int32 AssetIdx) { Indexers[AssetIdx].Process(); }, ParallelForFlags);
-	return true;
-}
-
-float FDatabaseIndexingContext::CalculateMinCostAddend() const
-{
-	float MinCostAddend = 0.f;
-
-	check(SearchIndexBase);
-	if (!SearchIndexBase->PoseMetadata.IsEmpty())
-	{
-		MinCostAddend = MAX_FLT;
-		for (const FPoseSearchPoseMetadata& PoseMetadata : SearchIndexBase->PoseMetadata)
-		{
-			if (PoseMetadata.CostAddend < MinCostAddend)
-			{
-				MinCostAddend = PoseMetadata.CostAddend;
-			}
-		}
-	}
-	return MinCostAddend;
-}
-
-void FDatabaseIndexingContext::JoinIndex()
-{
-	// Write index info to asset and count up total poses and storage required
-	int32 TotalPoses = 0;
-	int32 TotalFloats = 0;
-
-	check(SearchIndexBase);
-
-	// Join animation data into a single search index
-	SearchIndexBase->Values.Reset();
-	SearchIndexBase->PoseMetadata.Reset();
-	SearchIndexBase->OverallFlags = EPoseSearchPoseFlags::None;
-
-	for (int32 AssetIdx = 0; AssetIdx != SearchIndexBase->Assets.Num(); ++AssetIdx)
-	{
-		const FAssetIndexer::FOutput& Output = Indexers[AssetIdx].Output;
-
-		FPoseSearchIndexAsset& SearchIndexAsset = SearchIndexBase->Assets[AssetIdx];
-		SearchIndexAsset.NumPoses = Output.NumIndexedPoses;
-		SearchIndexAsset.FirstPoseIdx = TotalPoses;
-
-		const int32 PoseMetadataStartIdx = SearchIndexBase->PoseMetadata.Num();
-		const int32 PoseMetadataEndIdx = PoseMetadataStartIdx + Output.PoseMetadata.Num();
-
-		SearchIndexBase->Values.Append(Output.FeatureVectorTable.GetData(), Output.FeatureVectorTable.Num());
-		SearchIndexBase->PoseMetadata.Append(Output.PoseMetadata);
-
-		for (int32 i = PoseMetadataStartIdx; i < PoseMetadataEndIdx; ++i)
-		{
-			SearchIndexBase->PoseMetadata[i].AssetIndex = AssetIdx;
-			SearchIndexBase->OverallFlags |= SearchIndexBase->PoseMetadata[i].Flags;
-		}
-
-		TotalPoses += Output.NumIndexedPoses;
-		TotalFloats += Output.FeatureVectorTable.Num();
-	}
-
-	SearchIndexBase->NumPoses = TotalPoses;
-	SearchIndexBase->MinCostAddend = CalculateMinCostAddend();
-}
-
 //////////////////////////////////////////////////////////////////////////
 // FPoseSearchDatabaseAsyncCacheTask
 struct FPoseSearchDatabaseAsyncCacheTask
diff --git a/Engine/Plugins/Experimental/Animation/PoseSearch/Source/Runtime/Private/PoseSearchFeatureChannel.cpp b/Engine/Plugins/Experimental/Animation/PoseSearch/Source/Runtime/Private/PoseSearchFeatureChannel.cpp
--- a/Engine/Plugins/Experimental/Animation/PoseSearch/Source/Runtime/Private/PoseSearchFeatureChannel.cpp
+++ b/Engine/Plugins/Experimental/Animation/PoseSearch/Source/Runtime/Private/PoseSearchFeatureChannel.cpp
@@ -5,11 +5,11 @@
 #include "PoseSearch/PoseSearchSchema.h"
 #include "Animation/Skeleton.h"
 
-//////////////////////////////////////////////////////////////////////////
-// class ICostBreakDownData
 namespace UE::PoseSearch
 {
 
+//////////////////////////////////////////////////////////////////////////
+// class ICostBreakDownData
 static inline float ArraySum(TConstArrayView<float> View, int32 StartIndex, int32 Offset)
 {
 	float Sum = 0.f;
@@ -38,6 +38,190 @@ void ICostBreakDownData::AddEntireBreakDownSection(const FText& Label, const UPo
 	EndBreakDownSection(Label);
 }
 
+//////////////////////////////////////////////////////////////////////////
+// FFeatureVectorHelper
+void FFeatureVectorHelper::EncodeQuat(TArrayView<float> Values, int32& DataOffset, const FQuat& Quat)
+{
+	const FVector X = Quat.GetAxisX();
+	const FVector Y = Quat.GetAxisY();
+
+	Values[DataOffset + 0] = X.X;
+	Values[DataOffset + 1] = X.Y;
+	Values[DataOffset + 2] = X.Z;
+	Values[DataOffset + 3] = Y.X;
+	Values[DataOffset + 4] = Y.Y;
+	Values[DataOffset + 5] = Y.Z;
+
+	DataOffset += EncodeQuatCardinality;
+}
+
+void FFeatureVectorHelper::EncodeQuat(TArrayView<float> Values, int32& DataOffset, TConstArrayView<float> PrevValues, TConstArrayView<float> CurValues, TConstArrayView<float> NextValues, float LerpValue)
+{
+	FQuat Quat = DecodeQuatInternal(CurValues, DataOffset);
+
+	// linear interpolation
+	if (!FMath::IsNearlyZero(LerpValue))
+	{
+		if (LerpValue < 0.f)
+		{
+			Quat = FQuat::Slerp(Quat, DecodeQuatInternal(PrevValues, DataOffset), -LerpValue);
+		}
+		else
+		{
+			Quat = FQuat::Slerp(Quat, DecodeQuatInternal(NextValues, DataOffset), LerpValue);
+		}
+	}
+
+	// @todo: do we need to add options for cubic interpolation?
+	EncodeQuat(Values, DataOffset, Quat);
+}
+
+FQuat FFeatureVectorHelper::DecodeQuat(TConstArrayView<float> Values, int32& DataOffset)
+{
+	const FQuat Quat = DecodeQuatInternal(Values, DataOffset);
+	DataOffset += EncodeQuatCardinality;
+	return Quat;
+}
+
+FQuat FFeatureVectorHelper::DecodeQuatInternal(TConstArrayView<float> Values, int32 DataOffset)
+{
+	const FVector X(Values[DataOffset + 0], Values[DataOffset + 1], Values[DataOffset + 2]);
+	const FVector Y(Values[DataOffset + 3], Values[DataOffset + 4], Values[DataOffset + 5]);
+	const FVector Z = FVector::CrossProduct(X, Y);
+
+	FMatrix M(FMatrix::Identity);
+	M.SetColumn(0, X);
+	M.SetColumn(1, Y);
+	M.SetColumn(2, Z);
+
+	return FQuat(M);
+}
+
+void FFeatureVectorHelper::EncodeVector(TArrayView<float> Values, int32& DataOffset, const FVector& Vector)
+{
+	Values[DataOffset + 0] = Vector.X;
+	Values[DataOffset + 1] = Vector.Y;
+	Values[DataOffset + 2] = Vector.Z;
+	DataOffset += EncodeVectorCardinality;
+}
+
+void FFeatureVectorHelper::EncodeVector(TArrayView<float> Values, int32& DataOffset, TConstArrayView<float> PrevValues, TConstArrayView<float> CurValues, TConstArrayView<float> NextValues, float LerpValue, bool bNormalize)
+{
+	FVector Vector = DecodeVectorInternal(CurValues, DataOffset);
+
+	// linear interpolation
+	if (!FMath::IsNearlyZero(LerpValue))
+	{
+		if (LerpValue < 0.f)
+		{
+			Vector = FMath::Lerp(Vector, DecodeVectorInternal(PrevValues, DataOffset), -LerpValue);
+		}
+		else
+		{
+			Vector = FMath::Lerp(Vector, DecodeVectorInternal(NextValues, DataOffset), LerpValue);
+		}
+	}
+
+	// @todo: do we need to add options for cubic interpolation?
+	if (bNormalize)
+	{
+		Vector = Vector.GetSafeNormal(UE_SMALL_NUMBER, FVector::XAxisVector);
+	}
+
+	EncodeVector(Values, DataOffset, Vector);
+}
+
+FVector FFeatureVectorHelper::DecodeVector(TConstArrayView<float> Values, int32& DataOffset)
+{
+	const FVector Vector = DecodeVectorInternal(Values, DataOffset);
+	DataOffset += EncodeVectorCardinality;
+	return Vector;
+}
+
+FVector FFeatureVectorHelper::DecodeVectorInternal(TConstArrayView<float> Values, int32 DataOffset)
+{
+	return FVector(Values[DataOffset + 0], Values[DataOffset + 1], Values[DataOffset + 2]);
+}
+
+void FFeatureVectorHelper::EncodeVector2D(TArrayView<float> Values, int32& DataOffset, const FVector2D& Vector2D)
+{
+	Values[DataOffset + 0] = Vector2D.X;
+	Values[DataOffset + 1] = Vector2D.Y;
+	DataOffset += EncodeVector2DCardinality;
+}
+
+void FFeatureVectorHelper::EncodeVector2D(TArrayView<float> Values, int32& DataOffset, TConstArrayView<float> PrevValues, TConstArrayView<float> CurValues, TConstArrayView<float> NextValues, float LerpValue)
+{
+	FVector2D Vector2D = DecodeVector2DInternal(CurValues, DataOffset);
+
+	// linear interpolation
+	if (!FMath::IsNearlyZero(LerpValue))
+	{
+		if (LerpValue < 0.f)
+		{
+			Vector2D = FMath::Lerp(Vector2D, DecodeVector2DInternal(PrevValues, DataOffset), -LerpValue);
+		}
+		else
+		{
+			Vector2D = FMath::Lerp(Vector2D, DecodeVector2DInternal(NextValues, DataOffset), LerpValue);
+		}
+	}
+
+	// @todo: do we need to add options for cubic interpolation?
+	EncodeVector2D(Values, DataOffset, Vector2D);
+}
+
+FVector2D FFeatureVectorHelper::DecodeVector2D(TConstArrayView<float> Values, int32& DataOffset)
+{
+	const FVector2D Vector2D = DecodeVector2DInternal(Values, DataOffset);
+	DataOffset += EncodeVector2DCardinality;
+	return Vector2D;
+}
+
+FVector2D FFeatureVectorHelper::DecodeVector2DInternal(TConstArrayView<float> Values, int32 DataOffset)
+{
+	return FVector2D(Values[DataOffset + 0], Values[DataOffset + 1]);
+}
+
+void FFeatureVectorHelper::EncodeFloat(TArrayView<float> Values, int32& DataOffset, const float Value)
+{
+	Values[DataOffset + 0] = Value;
+	DataOffset += EncodeFloatCardinality;
+}
+
+void FFeatureVectorHelper::EncodeFloat(TArrayView<float> Values, int32& DataOffset, TConstArrayView<float> PrevValues, TConstArrayView<float> CurValues, TConstArrayView<float> NextValues, float LerpValue)
+{
+	float Value = DecodeFloatInternal(CurValues, DataOffset);
+
+	// linear interpolation
+	if (!FMath::IsNearlyZero(LerpValue))
+	{
+		if (LerpValue < 0.f)
+		{
+			Value = FMath::Lerp(Value, DecodeFloatInternal(PrevValues, DataOffset), -LerpValue);
+		}
+		else
+		{
+			Value = FMath::Lerp(Value, DecodeFloatInternal(NextValues, DataOffset), LerpValue);
+		}
+	}
+
+	// @todo: do we need to add options for cubic interpolation?
+	EncodeFloat(Values, DataOffset, Value);
+}
+
+float FFeatureVectorHelper::DecodeFloat(TConstArrayView<float> Values, int32& DataOffset)
+{
+	const float Value = DecodeFloatInternal(Values, DataOffset);
+	DataOffset += EncodeFloatCardinality;
+	return Value;
+}
+
+float FFeatureVectorHelper::DecodeFloatInternal(TConstArrayView<float> Values, int32 DataOffset)
+{
+	return Values[DataOffset];
+}
+
 } // namespace UE::PoseSearch
 
 //////////////////////////////////////////////////////////////////////////
diff --git a/Engine/Plugins/Experimental/Animation/PoseSearch/Source/Runtime/Private/PoseSearchFeatureChannel_FilterCrashingLegs.cpp b/Engine/Plugins/Experimental/Animation/PoseSearch/Source/Runtime/Private/PoseSearchFeatureChannel_FilterCrashingLegs.cpp
--- a/Engine/Plugins/Experimental/Animation/PoseSearch/Source/Runtime/Private/PoseSearchFeatureChannel_FilterCrashingLegs.cpp
+++ b/Engine/Plugins/Experimental/Animation/PoseSearch/Source/Runtime/Private/PoseSearchFeatureChannel_FilterCrashingLegs.cpp
@@ -2,7 +2,10 @@
 
 #include "PoseSearchFeatureChannel_FilterCrashingLegs.h"
 #include "DrawDebugHelpers.h"
-#include "PoseSearch/PoseSearch.h"
+#include "PoseSearch/PoseSearchAssetIndexer.h"
+#include "PoseSearch/PoseSearchAssetSampler.h"
+#include "PoseSearch/PoseSearchContext.h"
+#include "PoseSearch/PoseSearchDatabase.h"
 #include "PoseSearch/PoseSearchSchema.h"
 
 static void ComputeThighsSideAndForward(const FVector& RightThighPos, const FVector& LeftThighPos, FVector& ThighsSide, FVector& ThighsForward)
@@ -42,7 +45,7 @@ void UPoseSearchFeatureChannel_FilterCrashingLegs::FillWeights(TArray<float>& We
 	Weights[ChannelDataOffset] = Weight;
 }
 
-void UPoseSearchFeatureChannel_FilterCrashingLegs::IndexAsset(UE::PoseSearch::IAssetIndexer& Indexer, UE::PoseSearch::FAssetIndexingOutput& IndexingOutput) const
+void UPoseSearchFeatureChannel_FilterCrashingLegs::IndexAsset(UE::PoseSearch::IAssetIndexer& Indexer, TArrayView<float> FeatureVectorTable) const
 {
 	using namespace UE::PoseSearch;
 
@@ -59,10 +62,10 @@ void UPoseSearchFeatureChannel_FilterCrashingLegs::IndexAsset(UE::PoseSearch::IA
 		const FTransform LeftFootTransform = Indexer.GetTransformAndCacheResults(SubsampleTime, SubsampleTime, LeftFootIdx, bUnused);
 
 		const float CrashingLegsValue = ComputeCrashingLegsValue(RightThighTransform.GetTranslation(), LeftThighTransform.GetTranslation(), RightFootTransform.GetTranslation(), LeftFootTransform.GetTranslation());
-		
+
 		int32 DataOffset = ChannelDataOffset;
 		const int32 VectorIdx = SampleIdx - IndexingContext.BeginSampleIdx;
-		FFeatureVectorHelper::EncodeFloat(IndexingOutput.GetPoseVector(VectorIdx), DataOffset, CrashingLegsValue);
+		FFeatureVectorHelper::EncodeFloat(IndexingContext.GetPoseVector(VectorIdx, FeatureVectorTable), DataOffset, CrashingLegsValue);
 	}
 }
 
diff --git a/Engine/Plugins/Experimental/Animation/PoseSearch/Source/Runtime/Private/PoseSearchFeatureChannel_FilterCrashingLegs.h b/Engine/Plugins/Experimental/Animation/PoseSearch/Source/Runtime/Private/PoseSearchFeatureChannel_FilterCrashingLegs.h
--- a/Engine/Plugins/Experimental/Animation/PoseSearch/Source/Runtime/Private/PoseSearchFeatureChannel_FilterCrashingLegs.h
+++ b/Engine/Plugins/Experimental/Animation/PoseSearch/Source/Runtime/Private/PoseSearchFeatureChannel_FilterCrashingLegs.h
@@ -55,7 +55,7 @@ public:
 	// UPoseSearchFeatureChannel interface
 	virtual void InitializeSchema(UPoseSearchSchema* Schema) override;
 	virtual void FillWeights(TArray<float>& Weights) const override;
-	virtual void IndexAsset(UE::PoseSearch::IAssetIndexer& Indexer, UE::PoseSearch::FAssetIndexingOutput& IndexingOutput) const override;
+	virtual void IndexAsset(UE::PoseSearch::IAssetIndexer& Indexer, TArrayView<float> FeatureVectorTable) const override;
 	virtual void BuildQuery(UE::PoseSearch::FSearchContext& SearchContext, FPoseSearchFeatureVectorBuilder& InOutQuery) const override;
 	virtual void DebugDraw(const UE::PoseSearch::FDebugDrawParams& DrawParams, TConstArrayView<float> PoseVector) const override;
 
diff --git a/Engine/Plugins/Experimental/Animation/PoseSearch/Source/Runtime/Private/PoseSearchFeatureChannel_Heading.cpp b/Engine/Plugins/Experimental/Animation/PoseSearch/Source/Runtime/Private/PoseSearchFeatureChannel_Heading.cpp
--- a/Engine/Plugins/Experimental/Animation/PoseSearch/Source/Runtime/Private/PoseSearchFeatureChannel_Heading.cpp
+++ b/Engine/Plugins/Experimental/Animation/PoseSearch/Source/Runtime/Private/PoseSearchFeatureChannel_Heading.cpp
@@ -2,7 +2,10 @@
 
 #include "PoseSearchFeatureChannel_Heading.h"
 #include "DrawDebugHelpers.h"
-#include "PoseSearch/PoseSearch.h"
+#include "PoseSearch/PoseSearchAssetIndexer.h"
+#include "PoseSearch/PoseSearchAssetSampler.h"
+#include "PoseSearch/PoseSearchContext.h"
+#include "PoseSearch/PoseSearchDatabase.h"
 #include "PoseSearch/PoseSearchSchema.h"
 
 void UPoseSearchFeatureChannel_Heading::InitializeSchema(UPoseSearchSchema* Schema)
@@ -39,7 +42,7 @@ FVector UPoseSearchFeatureChannel_Heading::GetAxis(const FQuat& Rotation) const
 	return FVector(1.f, 0.f, 0.f);
 }
 
-void UPoseSearchFeatureChannel_Heading::IndexAsset(UE::PoseSearch::IAssetIndexer& Indexer, UE::PoseSearch::FAssetIndexingOutput& IndexingOutput) const
+void UPoseSearchFeatureChannel_Heading::IndexAsset(UE::PoseSearch::IAssetIndexer& Indexer, TArrayView<float> FeatureVectorTable) const
 {
 	using namespace UE::PoseSearch;
 
@@ -55,7 +58,7 @@ void UPoseSearchFeatureChannel_Heading::IndexAsset(UE::PoseSearch::IAssetIndexer
 		bool ClampedPresent;
 		const FTransform BoneTransformsPresent = Indexer.GetTransformAndCacheResults(SubsampleTime, bUseSampleTimeOffsetRootBone ? SubsampleTime : OriginSampleTime, SchemaBoneIdx, ClampedPresent);
 		int32 DataOffset = ChannelDataOffset;
-		FFeatureVectorHelper::EncodeVector(IndexingOutput.GetPoseVector(VectorIdx), DataOffset, GetAxis(BoneTransformsPresent.GetRotation()));
+		FFeatureVectorHelper::EncodeVector(IndexingContext.GetPoseVector(VectorIdx, FeatureVectorTable), DataOffset, GetAxis(BoneTransformsPresent.GetRotation()));
 		check(DataOffset == ChannelDataOffset + ChannelCardinality);
 	}
 }
diff --git a/Engine/Plugins/Experimental/Animation/PoseSearch/Source/Runtime/Private/PoseSearchFeatureChannel_Heading.h b/Engine/Plugins/Experimental/Animation/PoseSearch/Source/Runtime/Private/PoseSearchFeatureChannel_Heading.h
--- a/Engine/Plugins/Experimental/Animation/PoseSearch/Source/Runtime/Private/PoseSearchFeatureChannel_Heading.h
+++ b/Engine/Plugins/Experimental/Animation/PoseSearch/Source/Runtime/Private/PoseSearchFeatureChannel_Heading.h
@@ -52,7 +52,7 @@ public:
 	// UPoseSearchFeatureChannel interface
 	virtual void InitializeSchema(UPoseSearchSchema* Schema) override;
 	virtual void FillWeights(TArray<float>& Weights) const override;
-	virtual void IndexAsset(UE::PoseSearch::IAssetIndexer& Indexer, UE::PoseSearch::FAssetIndexingOutput& IndexingOutput) const override;
+	virtual void IndexAsset(UE::PoseSearch::IAssetIndexer& Indexer, TArrayView<float> FeatureVectorTable) const override;
 	virtual void BuildQuery(UE::PoseSearch::FSearchContext& SearchContext, FPoseSearchFeatureVectorBuilder& InOutQuery) const override;
 	virtual void DebugDraw(const UE::PoseSearch::FDebugDrawParams& DrawParams, TConstArrayView<float> PoseVector) const override;
 
diff --git a/Engine/Plugins/Experimental/Animation/PoseSearch/Source/Runtime/Private/PoseSearchFeatureChannel_Pose.cpp b/Engine/Plugins/Experimental/Animation/PoseSearch/Source/Runtime/Private/PoseSearchFeatureChannel_Pose.cpp
--- a/Engine/Plugins/Experimental/Animation/PoseSearch/Source/Runtime/Private/PoseSearchFeatureChannel_Pose.cpp
+++ b/Engine/Plugins/Experimental/Animation/PoseSearch/Source/Runtime/Private/PoseSearchFeatureChannel_Pose.cpp
@@ -2,8 +2,12 @@
 
 #include "PoseSearchFeatureChannel_Pose.h"
 #include "DrawDebugHelpers.h"
-#include "PoseSearch/PoseSearch.h"
+#include "PoseSearch/PoseSearchAssetIndexer.h"
+#include "PoseSearch/PoseSearchAssetSampler.h"
+#include "PoseSearch/PoseSearchContext.h"
+#include "PoseSearch/PoseSearchDatabase.h"
 #include "PoseSearch/PoseSearchDerivedDataKey.h"
+#include "PoseSearch/PoseSearchHistory.h"
 #include "PoseSearch/PoseSearchSchema.h"
 
 #define LOCTEXT_NAMESPACE "PoseSearchFeatureChannels"
@@ -369,7 +373,7 @@ void UPoseSearchFeatureChannel_Pose::FillWeights(TArray<float>& Weights) const
 }
 
 // @todo: do we really need to use double(s) in all this math?
-void UPoseSearchFeatureChannel_Pose::CalculatePhases(UE::PoseSearch::IAssetIndexer& Indexer, UE::PoseSearch::FAssetIndexingOutput& IndexingOutput, TArray<TArray<FVector2D>>& OutPhases) const
+void UPoseSearchFeatureChannel_Pose::CalculatePhases(UE::PoseSearch::IAssetIndexer& Indexer, TArray<TArray<FVector2D>>& OutPhases) const
 {
 	// @todo: expose them via UI
 	static float BoneSamplingCentralDifferencesTime = 0.2f; // seconds
@@ -410,21 +414,21 @@ void UPoseSearchFeatureChannel_Pose::CalculatePhases(UE::PoseSearch::IAssetIndex
 	}
 }
 
-void UPoseSearchFeatureChannel_Pose::IndexAsset(UE::PoseSearch::IAssetIndexer& Indexer,  UE::PoseSearch::FAssetIndexingOutput& IndexingOutput) const
+void UPoseSearchFeatureChannel_Pose::IndexAsset(UE::PoseSearch::IAssetIndexer& Indexer, TArrayView<float> FeatureVectorTable) const
 {
 	using namespace UE::PoseSearch;
 	 
 	// Phases is an array of array with cardinality SampledBones.Num() times NumSamples (IndexingContext.EndSampleIdx - IndexingContext.BeginSampleIdx)
-	// of 2 dimensional vectors (FVector2D) representing phases in an Eucledean space with phase angle sin/cos as direction and certainty of the signal as magnitude,
+	// of 2 dimensional vectors (FVector2D) representing phases in an Euclidean space with phase angle sin/cos as direction and certainty of the signal as magnitude,
 	// where certainty is a function of the amplitude of the signal used as input
 	TArray<TArray<FVector2D>> Phases;
-	CalculatePhases(Indexer, IndexingOutput, Phases);
+	CalculatePhases(Indexer, Phases);
 
 	const FAssetIndexingContext& IndexingContext = Indexer.GetIndexingContext();
 	for (int32 SampleIdx = IndexingContext.BeginSampleIdx; SampleIdx != IndexingContext.EndSampleIdx; ++SampleIdx)
 	{
 		const int32 VectorIdx = SampleIdx - IndexingContext.BeginSampleIdx;
-		AddPoseFeatures(Indexer, SampleIdx, IndexingOutput.GetPoseVector(VectorIdx), Phases);
+		AddPoseFeatures(Indexer, SampleIdx, IndexingContext.GetPoseVector(VectorIdx, FeatureVectorTable), Phases);
 	}
 }
 
diff --git a/Engine/Plugins/Experimental/Animation/PoseSearch/Source/Runtime/Private/PoseSearchFeatureChannel_Pose.h b/Engine/Plugins/Experimental/Animation/PoseSearch/Source/Runtime/Private/PoseSearchFeatureChannel_Pose.h
--- a/Engine/Plugins/Experimental/Animation/PoseSearch/Source/Runtime/Private/PoseSearchFeatureChannel_Pose.h
+++ b/Engine/Plugins/Experimental/Animation/PoseSearch/Source/Runtime/Private/PoseSearchFeatureChannel_Pose.h
@@ -58,7 +58,7 @@ public:
 	// UPoseSearchFeatureChannel interface
 	virtual void InitializeSchema(UPoseSearchSchema* Schema) override;
 	virtual void FillWeights(TArray<float>& Weights) const override;
-	virtual void IndexAsset(UE::PoseSearch::IAssetIndexer& Indexer, UE::PoseSearch::FAssetIndexingOutput& IndexingOutput) const override;
+	virtual void IndexAsset(UE::PoseSearch::IAssetIndexer& Indexer, TArrayView<float> FeatureVectorTable) const override;
 	virtual void BuildQuery(UE::PoseSearch::FSearchContext& SearchContext, FPoseSearchFeatureVectorBuilder& InOutQuery) const override;
 	virtual void DebugDraw(const UE::PoseSearch::FDebugDrawParams& DrawParams, TConstArrayView<float> PoseVector) const override;
 
@@ -69,7 +69,7 @@ public:
 
 protected:
 	void AddPoseFeatures(UE::PoseSearch::IAssetIndexer& Indexer, int32 SampleIdx, TArrayView<float> FeatureVector, const TArray<TArray<FVector2D>>& Phases) const;
-	void CalculatePhases(UE::PoseSearch::IAssetIndexer& Indexer, UE::PoseSearch::FAssetIndexingOutput& IndexingOutput, TArray<TArray<FVector2D>>& OutPhases) const;
+	void CalculatePhases(UE::PoseSearch::IAssetIndexer& Indexer, TArray<TArray<FVector2D>>& OutPhases) const;
 };
 
 
diff --git a/Engine/Plugins/Experimental/Animation/PoseSearch/Source/Runtime/Private/PoseSearchFeatureChannel_Position.cpp b/Engine/Plugins/Experimental/Animation/PoseSearch/Source/Runtime/Private/PoseSearchFeatureChannel_Position.cpp
--- a/Engine/Plugins/Experimental/Animation/PoseSearch/Source/Runtime/Private/PoseSearchFeatureChannel_Position.cpp
+++ b/Engine/Plugins/Experimental/Animation/PoseSearch/Source/Runtime/Private/PoseSearchFeatureChannel_Position.cpp
@@ -2,7 +2,10 @@
 
 #include "PoseSearchFeatureChannel_Position.h"
 #include "DrawDebugHelpers.h"
-#include "PoseSearch/PoseSearch.h"
+#include "PoseSearch/PoseSearchAssetIndexer.h"
+#include "PoseSearch/PoseSearchAssetSampler.h"
+#include "PoseSearch/PoseSearchContext.h"
+#include "PoseSearch/PoseSearchDatabase.h"
 #include "PoseSearch/PoseSearchSchema.h"
 
 void UPoseSearchFeatureChannel_Position::InitializeSchema(UPoseSearchSchema* Schema)
@@ -24,7 +27,7 @@ void UPoseSearchFeatureChannel_Position::FillWeights(TArray<float>& Weights) con
 	}
 }
 
-void UPoseSearchFeatureChannel_Position::IndexAsset(UE::PoseSearch::IAssetIndexer& Indexer, UE::PoseSearch::FAssetIndexingOutput& IndexingOutput) const
+void UPoseSearchFeatureChannel_Position::IndexAsset(UE::PoseSearch::IAssetIndexer& Indexer, TArrayView<float> FeatureVectorTable) const
 {
 	using namespace UE::PoseSearch;
 
@@ -40,7 +43,7 @@ void UPoseSearchFeatureChannel_Position::IndexAsset(UE::PoseSearch::IAssetIndexe
 		bool ClampedPresent;
 		const FTransform BoneTransformsPresent = Indexer.GetTransformAndCacheResults(SubsampleTime, bUseSampleTimeOffsetRootBone ? SubsampleTime : OriginSampleTime, SchemaBoneIdx, ClampedPresent);
 		int32 DataOffset = ChannelDataOffset;
-		FFeatureVectorHelper::EncodeVector(IndexingOutput.GetPoseVector(VectorIdx), DataOffset, BoneTransformsPresent.GetTranslation());
+		FFeatureVectorHelper::EncodeVector(IndexingContext.GetPoseVector(VectorIdx, FeatureVectorTable), DataOffset, BoneTransformsPresent.GetTranslation());
 	}
 }
 
diff --git a/Engine/Plugins/Experimental/Animation/PoseSearch/Source/Runtime/Private/PoseSearchFeatureChannel_Position.h b/Engine/Plugins/Experimental/Animation/PoseSearch/Source/Runtime/Private/PoseSearchFeatureChannel_Position.h
--- a/Engine/Plugins/Experimental/Animation/PoseSearch/Source/Runtime/Private/PoseSearchFeatureChannel_Position.h
+++ b/Engine/Plugins/Experimental/Animation/PoseSearch/Source/Runtime/Private/PoseSearchFeatureChannel_Position.h
@@ -38,7 +38,7 @@ public:
 	// UPoseSearchFeatureChannel interface
 	virtual void InitializeSchema(UPoseSearchSchema* Schema) override;
 	virtual void FillWeights(TArray<float>& Weights) const override;
-	virtual void IndexAsset(UE::PoseSearch::IAssetIndexer& Indexer, UE::PoseSearch::FAssetIndexingOutput& IndexingOutput) const override;
+	virtual void IndexAsset(UE::PoseSearch::IAssetIndexer& Indexer, TArrayView<float> FeatureVectorTable) const override;
 	virtual void BuildQuery(UE::PoseSearch::FSearchContext& SearchContext, FPoseSearchFeatureVectorBuilder& InOutQuery) const override;
 	virtual void DebugDraw(const UE::PoseSearch::FDebugDrawParams& DrawParams, TConstArrayView<float> PoseVector) const override;
 };
diff --git a/Engine/Plugins/Experimental/Animation/PoseSearch/Source/Runtime/Private/PoseSearchFeatureChannel_Trajectory.cpp b/Engine/Plugins/Experimental/Animation/PoseSearch/Source/Runtime/Private/PoseSearchFeatureChannel_Trajectory.cpp
--- a/Engine/Plugins/Experimental/Animation/PoseSearch/Source/Runtime/Private/PoseSearchFeatureChannel_Trajectory.cpp
+++ b/Engine/Plugins/Experimental/Animation/PoseSearch/Source/Runtime/Private/PoseSearchFeatureChannel_Trajectory.cpp
@@ -1,8 +1,11 @@
 // Copyright Epic Games, Inc. All Rights Reserved.
 
 #include "PoseSearchFeatureChannel_Trajectory.h"
+#include "Animation/MotionTrajectoryTypes.h"
 #include "DrawDebugHelpers.h"
-#include "PoseSearch/PoseSearch.h"
+#include "PoseSearch/PoseSearchAssetIndexer.h"
+#include "PoseSearch/PoseSearchAssetSampler.h"
+#include "PoseSearch/PoseSearchContext.h"
 #include "PoseSearch/PoseSearchDerivedDataKey.h"
 #include "PoseSearch/PoseSearchSchema.h"
 
@@ -145,13 +148,13 @@ void UPoseSearchFeatureChannel_Trajectory::FillWeights(TArray<float>& Weights) c
 	check(DataOffset == ChannelDataOffset + ChannelCardinality);
 }
 
-void UPoseSearchFeatureChannel_Trajectory::IndexAsset(UE::PoseSearch::IAssetIndexer& Indexer,  UE::PoseSearch::FAssetIndexingOutput& IndexingOutput) const
+void UPoseSearchFeatureChannel_Trajectory::IndexAsset(UE::PoseSearch::IAssetIndexer& Indexer, TArrayView<float> FeatureVectorTable) const
 {
 	const UE::PoseSearch::FAssetIndexingContext& IndexingContext = Indexer.GetIndexingContext();
 	for (int32 SampleIdx = IndexingContext.BeginSampleIdx; SampleIdx != IndexingContext.EndSampleIdx; ++SampleIdx)
 	{
 		const int32 VectorIdx = SampleIdx - IndexingContext.BeginSampleIdx;
-		IndexAssetPrivate(Indexer, SampleIdx, IndexingOutput.GetPoseVector(VectorIdx));
+		IndexAssetPrivate(Indexer, SampleIdx, IndexingContext.GetPoseVector(VectorIdx, FeatureVectorTable));
 	}
 }
 
diff --git a/Engine/Plugins/Experimental/Animation/PoseSearch/Source/Runtime/Private/PoseSearchFeatureChannel_Trajectory.h b/Engine/Plugins/Experimental/Animation/PoseSearch/Source/Runtime/Private/PoseSearchFeatureChannel_Trajectory.h
--- a/Engine/Plugins/Experimental/Animation/PoseSearch/Source/Runtime/Private/PoseSearchFeatureChannel_Trajectory.h
+++ b/Engine/Plugins/Experimental/Animation/PoseSearch/Source/Runtime/Private/PoseSearchFeatureChannel_Trajectory.h
@@ -59,7 +59,7 @@ public:
 	// UPoseSearchFeatureChannel interface
 	virtual void InitializeSchema(UPoseSearchSchema* Schema) override;
 	virtual void FillWeights(TArray<float>& Weights) const override;
-	virtual void IndexAsset(UE::PoseSearch::IAssetIndexer& Indexer, UE::PoseSearch::FAssetIndexingOutput& IndexingOutput) const override;
+	virtual void IndexAsset(UE::PoseSearch::IAssetIndexer& Indexer, TArrayView<float> FeatureVectorTable) const override;
 	virtual void BuildQuery(UE::PoseSearch::FSearchContext& SearchContext, FPoseSearchFeatureVectorBuilder& InOutQuery) const override;
 	virtual void DebugDraw(const UE::PoseSearch::FDebugDrawParams& DrawParams, TConstArrayView<float> PoseVector) const override;
 
diff --git a/Engine/Plugins/Experimental/Animation/PoseSearch/Source/Runtime/Private/PoseSearchHistory.cpp b/Engine/Plugins/Experimental/Animation/PoseSearch/Source/Runtime/Private/PoseSearchHistory.cpp
new file mode 100644
--- /dev/null
+++ b/Engine/Plugins/Experimental/Animation/PoseSearch/Source/Runtime/Private/PoseSearchHistory.cpp
@@ -0,0 +1,279 @@
+// Copyright Epic Games, Inc. All Rights Reserved.
+
+#include "PoseSearch/PoseSearchHistory.h"
+#include "AnimationRuntime.h"
+#include "Animation/AnimNodeBase.h"
+#include "Animation/AnimRootMotionProvider.h"
+#include "BonePose.h"
+#include "PoseSearch/PoseSearchResult.h"
+#include "PoseSearch/PoseSearchDatabase.h"
+
+IMPLEMENT_ANIMGRAPH_MESSAGE(UE::PoseSearch::IPoseHistoryProvider);
+
+#define LOCTEXT_NAMESPACE "PoseSearch"
+
+namespace UE::PoseSearch
+{
+
+/**
+* Algo::LowerBound adapted to TIndexedContainerIterator for use with indexable but not necessarily contiguous containers. Used here with TRingBuffer.
+*
+* Performs binary search, resulting in position of the first element >= Value using predicate
+*
+* @param First TIndexedContainerIterator beginning of range to search through, must be already sorted by SortPredicate
+* @param Last TIndexedContainerIterator end of range
+* @param Value Value to look for
+* @param SortPredicate Predicate for sort comparison, defaults to <
+*
+* @returns Position of the first element >= Value, may be position after last element in range
+*/
+template <typename IteratorType, typename ValueType, typename ProjectionType, typename SortPredicateType>
+FORCEINLINE auto LowerBound(IteratorType First, IteratorType Last, const ValueType& Value, ProjectionType Projection, SortPredicateType SortPredicate) -> decltype(First.GetIndex())
+{
+	using SizeType = decltype(First.GetIndex());
+
+	check(First.GetIndex() <= Last.GetIndex());
+
+	// Current start of sequence to check
+	SizeType Start = First.GetIndex();
+
+	// Size of sequence to check
+	SizeType Size = Last.GetIndex() - Start;
+
+	// With this method, if Size is even it will do one more comparison than necessary, but because Size can be predicted by the CPU it is faster in practice
+	while (Size > 0)
+	{
+		const SizeType LeftoverSize = Size % 2;
+		Size = Size / 2;
+
+		const SizeType CheckIndex = Start + Size;
+		const SizeType StartIfLess = CheckIndex + LeftoverSize;
+
+		auto&& CheckValue = Invoke(Projection, *(First + CheckIndex));
+		Start = SortPredicate(CheckValue, Value) ? StartIfLess : Start;
+	}
+	return Start;
+}
+
+template <typename IteratorType, typename ValueType, typename SortPredicateType = TLess<>()>
+FORCEINLINE auto LowerBound(IteratorType First, IteratorType Last, const ValueType& Value, SortPredicateType SortPredicate) -> decltype(First.GetIndex())
+{
+	return LowerBound(First, Last, Value, FIdentityFunctor(), SortPredicate);
+}
+
+//////////////////////////////////////////////////////////////////////////
+// FPoseHistory
+/**
+* Fills skeleton transforms with evaluated compact pose transforms.
+* Bones that weren't evaluated are filled with the bone's reference pose.
+*/
+static void CopyCompactToSkeletonPose(const FCompactPose& Pose, TArray<FTransform>& OutLocalTransforms)
+{
+	const FBoneContainer& BoneContainer = Pose.GetBoneContainer();
+	const USkeleton* SkeletonAsset = BoneContainer.GetSkeletonAsset();
+	check(SkeletonAsset);
+
+	const FReferenceSkeleton& RefSkeleton = SkeletonAsset->GetReferenceSkeleton();
+	TArrayView<const FTransform> RefSkeletonTransforms = MakeArrayView(RefSkeleton.GetRefBonePose());
+	const int32 NumSkeletonBones = RefSkeleton.GetNum();
+
+	OutLocalTransforms.SetNum(NumSkeletonBones);
+
+	for (auto SkeletonBoneIdx = FSkeletonPoseBoneIndex(0); SkeletonBoneIdx != NumSkeletonBones; ++SkeletonBoneIdx)
+	{
+		FCompactPoseBoneIndex CompactBoneIdx = BoneContainer.GetCompactPoseIndexFromSkeletonPoseIndex(SkeletonBoneIdx);
+		OutLocalTransforms[SkeletonBoneIdx.GetInt()] = CompactBoneIdx.IsValid() ? Pose[CompactBoneIdx] : RefSkeletonTransforms[SkeletonBoneIdx.GetInt()];
+	}
+}
+
+void FPoseHistory::Init(int32 InNumPoses, float InTimeHorizon)
+{
+	Poses.Reserve(InNumPoses);
+	TimeHorizon = InTimeHorizon;
+}
+
+void FPoseHistory::Init(const FPoseHistory& History)
+{
+	Poses = History.Poses;
+	TimeHorizon = History.TimeHorizon;
+}
+
+bool FPoseHistory::TrySampleLocalPose(float SecondsAgo, const TArray<FBoneIndexType>* RequiredBones, TArray<FTransform>* LocalPose, FTransform* RootTransform) const
+{
+	const int32 NextIdx = LowerBound(Poses.begin(), Poses.end(), SecondsAgo, [](const FPose& Pose, float Value)
+	{
+		return Value < Pose.Time;
+	});
+	if (NextIdx <= 0 || NextIdx >= Poses.Num())
+	{
+		// We may not have accumulated enough poses yet
+		return false;
+	}
+
+	const int32 PrevIdx = NextIdx - 1;
+
+	const FPose& PrevPose = Poses[PrevIdx];
+	const FPose& NextPose = Poses[NextIdx];
+
+#if DO_CHECK
+	check(PrevPose.LocalTransforms.Num() == NextPose.LocalTransforms.Num());
+	FBoneIndexType MaxBoneIndexType = 0;
+	if (RequiredBones)
+	{
+		for (FBoneIndexType BoneIndexType : *RequiredBones)
+		{
+			if (BoneIndexType > MaxBoneIndexType)
+			{
+				MaxBoneIndexType = BoneIndexType;
+			}
+		}
+		check(MaxBoneIndexType < PrevPose.LocalTransforms.Num());
+	}
+#endif
+	// Compute alpha between previous and next Poses
+	const float Alpha = FMath::GetMappedRangeValueUnclamped(FVector2f(PrevPose.Time, NextPose.Time), FVector2f(0.0f, 1.0f), SecondsAgo);
+
+	// Lerp between poses by alpha to produce output local pose at requested sample time
+	if (LocalPose)
+	{
+		check(RequiredBones);
+		*LocalPose = PrevPose.LocalTransforms;
+		FAnimationRuntime::LerpBoneTransforms(*LocalPose, NextPose.LocalTransforms, Alpha, *RequiredBones);
+	}
+
+	if (RootTransform)
+	{
+		RootTransform->Blend(PrevPose.RootTransform, NextPose.RootTransform, Alpha);
+	}
+	return true;
+}
+
+bool FPoseHistory::Update(float SecondsElapsed, const FPoseContext& PoseContext, FTransform ComponentTransform, FText* OutError, ERootUpdateMode UpdateMode)
+{
+	// Age our elapsed times
+	for (FPose& Pose : Poses)
+	{
+		Pose.Time += SecondsElapsed;
+	}
+
+	if (Poses.Num() != Poses.Max())
+	{
+		// Consume every pose until the queue is full
+		Poses.Emplace();
+	}
+	else
+	{
+		// Exercise pose retention policy. We must guarantee there is always one additional pose
+		// beyond the time horizon so we can compute derivatives at the time horizon. We also
+		// want to evenly distribute poses across the entire history buffer so we only push additional
+		// poses when enough time has elapsed.
+
+		const float SampleInterval = GetSampleTimeInterval();
+
+		bool bCanEvictOldest = Poses[1].Time >= TimeHorizon + SampleInterval;
+		bool bShouldPushNewest = Poses[Poses.Num() - 2].Time >= SampleInterval;
+
+		if (bCanEvictOldest && bShouldPushNewest)
+		{
+			FPose PoseTemp = MoveTemp(Poses.First());
+			Poses.PopFront();
+			Poses.Emplace(MoveTemp(PoseTemp));
+		}
+	}
+
+	// Regardless of the retention policy, we always update the most recent pose
+	FPose& CurrentPose = Poses.Last();
+	CurrentPose.Time = 0.f;
+	CopyCompactToSkeletonPose(PoseContext.Pose, CurrentPose.LocalTransforms);
+
+	// Initialize with Previous Root Transform or Identity
+	CurrentPose.RootTransform = Poses.Num() > 1 ? Poses[Poses.Num() - 2].RootTransform : FTransform::Identity;
+	
+	// Update using either AniumRootMotionProvider or Component Transform
+	if (UpdateMode == ERootUpdateMode::RootMotionDelta)
+	{
+		const UE::Anim::IAnimRootMotionProvider* RootMotionProvider = UE::Anim::IAnimRootMotionProvider::Get();
+
+		if (RootMotionProvider)
+		{
+			if (RootMotionProvider->HasRootMotion(PoseContext.CustomAttributes))
+			{
+				FTransform RootMotionDelta = FTransform::Identity;
+				RootMotionProvider->ExtractRootMotion(PoseContext.CustomAttributes, RootMotionDelta);
+
+				CurrentPose.RootTransform = RootMotionDelta * CurrentPose.RootTransform;
+			}
+#if WITH_EDITORONLY_DATA	
+			else
+			{
+				if (OutError)
+				{
+					*OutError = LOCTEXT("PoseHistoryRootMotionProviderError",
+						"Input to Pose History has no Root Motion Attribute. Try disabling 'Use Root Motion'.");
+				}
+				return false;
+			}
+#endif
+		}
+#if WITH_EDITORONLY_DATA	
+		else
+		{
+			if (OutError)
+			{
+				*OutError = LOCTEXT("PoseHistoryRootMotionAttributeError",
+					"Could not get Root Motion Provider. Try disabling 'Use Root Motion'.");
+			}
+			return false;
+		}
+#endif
+	}
+	else if (UpdateMode == ERootUpdateMode::ComponentTransformDelta)
+	{
+		CurrentPose.RootTransform = ComponentTransform;
+	}
+	else
+	{
+		checkNoEntry();
+	}
+
+	return true;
+}
+
+float FPoseHistory::GetSampleTimeInterval() const
+{
+	// Reserve one pose for computing derivatives at the time horizon
+	return TimeHorizon / (Poses.Max() - 1);
+}
+
+//////////////////////////////////////////////////////////////////////////
+// FPoseIndicesHistory
+void FPoseIndicesHistory::Update(const FSearchResult& SearchResult, float DeltaTime, float MaxTime)
+{
+	if (MaxTime > 0.f)
+	{
+		for (auto It = IndexToTime.CreateIterator(); It; ++It)
+		{
+			It.Value() += DeltaTime;
+			if (It.Value() > MaxTime)
+			{
+				It.RemoveCurrent();
+			}
+		}
+
+		if (SearchResult.IsValid())
+		{
+			FHistoricalPoseIndex HistoricalPoseIndex;
+			HistoricalPoseIndex.PoseIndex = SearchResult.PoseIdx;
+			HistoricalPoseIndex.DatabaseKey = FObjectKey(SearchResult.Database.Get());
+			IndexToTime.Add(HistoricalPoseIndex, 0.f);
+		}
+	}
+	else
+	{
+		IndexToTime.Reset();
+	}
+}
+
+} // namespace UE::PoseSearch
+
+#undef LOCTEXT_NAMESPACE
\ No newline at end of file
diff --git a/Engine/Plugins/Experimental/Animation/PoseSearch/Source/Runtime/Private/PoseSearchIndex.cpp b/Engine/Plugins/Experimental/Animation/PoseSearch/Source/Runtime/Private/PoseSearchIndex.cpp
new file mode 100644
--- /dev/null
+++ b/Engine/Plugins/Experimental/Animation/PoseSearch/Source/Runtime/Private/PoseSearchIndex.cpp
@@ -0,0 +1,254 @@
+// Copyright Epic Games, Inc. All Rights Reserved.
+
+#include "PoseSearch/PoseSearchIndex.h"
+#include "PoseSearchEigenHelper.h"
+
+namespace UE::PoseSearch
+{
+
+static inline float CompareFeatureVectors(TConstArrayView<float> A, TConstArrayView<float> B, TConstArrayView<float> WeightsSqrt)
+{
+	check(A.Num() == B.Num() && A.Num() == WeightsSqrt.Num());
+
+	Eigen::Map<const Eigen::ArrayXf> VA(A.GetData(), A.Num());
+	Eigen::Map<const Eigen::ArrayXf> VB(B.GetData(), B.Num());
+	Eigen::Map<const Eigen::ArrayXf> VW(WeightsSqrt.GetData(), WeightsSqrt.Num());
+
+	return ((VA - VB) * VW).square().sum();
+}
+
+void CompareFeatureVectors(TConstArrayView<float> A, TConstArrayView<float> B, TConstArrayView<float> WeightsSqrt, TArrayView<float> Result)
+{
+	check(A.Num() == B.Num() && A.Num() == WeightsSqrt.Num() && A.Num() == Result.Num());
+
+	Eigen::Map<const Eigen::ArrayXf> VA(A.GetData(), A.Num());
+	Eigen::Map<const Eigen::ArrayXf> VB(B.GetData(), B.Num());
+	Eigen::Map<const Eigen::ArrayXf> VW(WeightsSqrt.GetData(), WeightsSqrt.Num());
+	Eigen::Map<Eigen::ArrayXf> VR(Result.GetData(), Result.Num());
+
+	VR = ((VA - VB) * VW).square();
+}
+
+} // namespace UE::PoseSearch
+
+//////////////////////////////////////////////////////////////////////////
+// FPoseSearchBaseIndex
+const FPoseSearchIndexAsset& FPoseSearchIndexBase::GetAssetForPose(int32 PoseIdx) const
+{
+	const int32 AssetIndex = PoseMetadata[PoseIdx].AssetIndex;
+	return Assets[AssetIndex];
+}
+
+const FPoseSearchIndexAsset* FPoseSearchIndexBase::GetAssetForPoseSafe(int32 PoseIdx) const
+{
+	if (PoseMetadata.IsValidIndex(PoseIdx))
+	{
+		const int32 AssetIndex = PoseMetadata[PoseIdx].AssetIndex;
+		if (Assets.IsValidIndex(AssetIndex))
+		{
+			return &Assets[AssetIndex];
+		}
+	}
+	return nullptr;
+}
+
+float FPoseSearchIndexBase::GetAssetTime(int32 PoseIdx, float SamplingInterval) const
+{
+	const FPoseSearchIndexAsset& Asset = GetAssetForPose(PoseIdx);
+
+	if (Asset.Type == ESearchIndexAssetType::Sequence || Asset.Type == ESearchIndexAssetType::AnimComposite)
+	{
+		const FFloatInterval SamplingRange = Asset.SamplingInterval;
+
+		float AssetTime = FMath::Min(SamplingRange.Min + SamplingInterval * (PoseIdx - Asset.FirstPoseIdx), SamplingRange.Max);
+		return AssetTime;
+	}
+
+	if (Asset.Type == ESearchIndexAssetType::BlendSpace)
+	{
+		const FFloatInterval SamplingRange = Asset.SamplingInterval;
+
+		// For BlendSpaces the AssetTime is in the range [0, 1] while the Sampling Range
+		// is in real time (seconds)
+		float AssetTime = FMath::Min(SamplingRange.Min + SamplingInterval * (PoseIdx - Asset.FirstPoseIdx), SamplingRange.Max) / (Asset.NumPoses * SamplingInterval);
+		return AssetTime;
+	}
+	
+	checkNoEntry();
+	return -1.0f;
+}
+
+bool FPoseSearchIndexBase::IsEmpty() const
+{
+	const bool bEmpty = Assets.Num() == 0 || NumPoses == 0;
+	return bEmpty;
+}
+
+void FPoseSearchIndexBase::Reset()
+{
+	FPoseSearchIndexBase Default;
+	*this = Default;
+}
+
+FArchive& operator<<(FArchive& Ar, FPoseSearchIndexBase& Index)
+{
+	int32 NumValues = 0;
+	int32 NumAssets = 0;
+
+	if (Ar.IsSaving())
+	{
+		NumValues = Index.Values.Num();
+		NumAssets = Index.Assets.Num();
+	}
+
+	Ar << Index.NumPoses;
+	Ar << NumValues;
+	Ar << NumAssets;
+	Ar << Index.OverallFlags;
+
+	if (Ar.IsLoading())
+	{
+		Index.Values.SetNumUninitialized(NumValues);
+		Index.PoseMetadata.SetNumUninitialized(Index.NumPoses);
+		Index.Assets.SetNumUninitialized(NumAssets);
+	}
+
+	if (Index.Values.Num() > 0)
+	{
+		Ar.Serialize(&Index.Values[0], Index.Values.Num() * Index.Values.GetTypeSize());
+	}
+
+	if (Index.PoseMetadata.Num() > 0)
+	{
+		Ar.Serialize(&Index.PoseMetadata[0], Index.PoseMetadata.Num() * Index.PoseMetadata.GetTypeSize());
+	}
+
+	if (Index.Assets.Num() > 0)
+	{
+		Ar.Serialize(&Index.Assets[0], Index.Assets.Num() * Index.Assets.GetTypeSize());
+	}
+
+	Ar << Index.MinCostAddend;
+
+	return Ar;
+}
+
+//////////////////////////////////////////////////////////////////////////
+// FPoseSearchIndex
+FPoseSearchIndex::FPoseSearchIndex(const FPoseSearchIndex& Other)
+	: FPoseSearchIndexBase(Other)
+	, PCAValues(Other.PCAValues)
+	, PCAProjectionMatrix(Other.PCAProjectionMatrix)
+	, Mean(Other.Mean)
+	, WeightsSqrt(Other.WeightsSqrt)
+	, KDTree(Other.KDTree)
+#if WITH_EDITORONLY_DATA
+	, PCAExplainedVariance(Other.PCAExplainedVariance)
+	, Deviation(Other.Deviation)
+#endif // WITH_EDITORONLY_DATA
+{
+	check(!PCAValues.IsEmpty() || KDTree.DataSource.PointCount == 0);
+	KDTree.DataSource.Data = PCAValues.IsEmpty() ? nullptr : PCAValues.GetData();
+}
+
+FPoseSearchIndex& FPoseSearchIndex::operator=(const FPoseSearchIndex& Other)
+{
+	if (this != &Other)
+	{
+		this->~FPoseSearchIndex();
+		new(this)FPoseSearchIndex(Other);
+	}
+	return *this;
+}
+
+void FPoseSearchIndex::Reset()
+{
+	FPoseSearchIndex Default;
+	*this = Default;
+}
+
+TConstArrayView<float> FPoseSearchIndex::GetPoseValues(int32 PoseIdx) const
+{
+	const int32 SchemaCardinality = WeightsSqrt.Num();
+	check(PoseIdx >= 0 && PoseIdx < NumPoses&& SchemaCardinality > 0);
+	const int32 ValueOffset = PoseIdx * SchemaCardinality;
+	return MakeArrayView(&Values[ValueOffset], SchemaCardinality);
+}
+
+TConstArrayView<float> FPoseSearchIndex::GetPoseValuesSafe(int32 PoseIdx) const
+{
+	if (PoseIdx >= 0 && PoseIdx < NumPoses)
+	{
+		const int32 SchemaCardinality = WeightsSqrt.Num();
+		const int32 ValueOffset = PoseIdx * SchemaCardinality;
+		return MakeArrayView(&Values[ValueOffset], SchemaCardinality);
+	}
+	return TConstArrayView<float>();
+}
+
+FPoseSearchCost FPoseSearchIndex::ComparePoses(int32 PoseIdx, EPoseSearchBooleanRequest QueryMirrorRequest, UE::PoseSearch::EPoseComparisonFlags PoseComparisonFlags, float MirrorMismatchCostBias, TConstArrayView<float> QueryValues) const
+{
+	// base dissimilarity cost representing how the associated PoseIdx differ, in a weighted way, from the query pose (QueryValues)
+	const float DissimilarityCost = UE::PoseSearch::CompareFeatureVectors(GetPoseValues(PoseIdx), QueryValues, WeightsSqrt);
+
+	// cost addend associated to a mismatch in mirror state between query and analyzed PoseIdx
+	float MirrorMismatchAddend = 0.f;
+	if (QueryMirrorRequest != EPoseSearchBooleanRequest::Indifferent)
+	{
+		const FPoseSearchIndexAsset& IndexAsset = GetAssetForPose(PoseIdx);
+		const bool bMirroringMismatch =
+			(IndexAsset.bMirrored && QueryMirrorRequest == EPoseSearchBooleanRequest::FalseValue) ||
+			(!IndexAsset.bMirrored && QueryMirrorRequest == EPoseSearchBooleanRequest::TrueValue);
+		if (bMirroringMismatch)
+		{
+			MirrorMismatchAddend = MirrorMismatchCostBias;
+		}
+	}
+
+	const FPoseSearchPoseMetadata& PoseIdxMetadata = PoseMetadata[PoseIdx];
+
+	// cost addend associated to Schema->BaseCostBias or overriden by UAnimNotifyState_PoseSearchModifyCost
+	const float NotifyAddend = PoseIdxMetadata.CostAddend;
+
+	// cost addend associated to Schema->ContinuingPoseCostBias or overriden by UAnimNotifyState_PoseSearchOverrideContinuingPoseCostBias
+	const float ContinuingPoseCostAddend = EnumHasAnyFlags(PoseComparisonFlags, UE::PoseSearch::EPoseComparisonFlags::ContinuingPose) ? PoseIdxMetadata.ContinuingPoseCostAddend : 0.f;
+
+	return FPoseSearchCost(DissimilarityCost, NotifyAddend, MirrorMismatchAddend, ContinuingPoseCostAddend);
+}
+
+FArchive& operator<<(FArchive& Ar, FPoseSearchIndex& Index)
+{
+	Ar << static_cast<FPoseSearchIndexBase&>(Index);
+
+	int32 NumPCAValues = 0;
+
+	if (Ar.IsSaving())
+	{
+		NumPCAValues = Index.PCAValues.Num();
+	}
+
+	Ar << NumPCAValues;
+
+	if (Ar.IsLoading())
+	{
+		Index.PCAValues.SetNumUninitialized(NumPCAValues);
+	}
+
+	if (Index.PCAValues.Num() > 0)
+	{
+		Ar.Serialize(&Index.PCAValues[0], Index.PCAValues.Num() * Index.PCAValues.GetTypeSize());
+	}
+
+	Ar << Index.WeightsSqrt;
+	Ar << Index.Mean;
+	Ar << Index.PCAProjectionMatrix;
+
+	Serialize(Ar, Index.KDTree, Index.PCAValues.GetData());
+
+#if WITH_EDITORONLY_DATA
+	Ar << Index.PCAExplainedVariance;
+	Ar << Index.Deviation;
+#endif // WITH_EDITORONLY_DATA
+
+	return Ar;
+}
\ No newline at end of file
diff --git a/Engine/Plugins/Experimental/Animation/PoseSearch/Source/Runtime/Private/PoseSearchLibrary.cpp b/Engine/Plugins/Experimental/Animation/PoseSearch/Source/Runtime/Private/PoseSearchLibrary.cpp
--- a/Engine/Plugins/Experimental/Animation/PoseSearch/Source/Runtime/Private/PoseSearchLibrary.cpp
+++ b/Engine/Plugins/Experimental/Animation/PoseSearch/Source/Runtime/Private/PoseSearchLibrary.cpp
@@ -8,7 +8,9 @@
 #include "Animation/AnimNode_SequencePlayer.h"
 #include "Animation/AnimRootMotionProvider.h"
 #include "Animation/BlendSpace.h"
+#include "InstancedStruct.h"
 #include "PoseSearch/AnimNode_MotionMatching.h"
+#include "PoseSearch/PoseSearchDatabase.h"
 #include "PoseSearch/PoseSearchSchema.h"
 #include "PoseSearchFeatureChannel_Trajectory.h"
 #include "Trace/PoseSearchTraceLogger.h"
diff --git a/Engine/Plugins/Experimental/Animation/PoseSearch/Source/Runtime/Private/PoseSearchResult.cpp b/Engine/Plugins/Experimental/Animation/PoseSearch/Source/Runtime/Private/PoseSearchResult.cpp
new file mode 100644
--- /dev/null
+++ b/Engine/Plugins/Experimental/Animation/PoseSearch/Source/Runtime/Private/PoseSearchResult.cpp
@@ -0,0 +1,104 @@
+// Copyright Epic Games, Inc. All Rights Reserved.
+
+#include "PoseSearch/PoseSearchResult.h"
+#include "Animation/BlendSpace.h"
+#include "InstancedStruct.h"
+#include "PoseSearch/PoseSearchDatabase.h"
+#include "PoseSearch/PoseSearchSchema.h"
+
+//////////////////////////////////////////////////////////////////////////
+// FPoseSearchFeatureVectorBuilder
+void FPoseSearchFeatureVectorBuilder::Init(const UPoseSearchSchema* InSchema)
+{
+	check(InSchema && InSchema->IsValid());
+	Schema = InSchema;
+	Values.Reset();
+	Values.SetNumZeroed(Schema->SchemaCardinality);
+}
+
+void FPoseSearchFeatureVectorBuilder::Reset()
+{
+	Schema = nullptr;
+	Values.Reset();
+}
+
+namespace UE::PoseSearch
+{
+//////////////////////////////////////////////////////////////////////////
+// FSearchResult
+void FSearchResult::Update(float NewAssetTime)
+{
+	if (!IsValid())
+	{
+		Reset();
+	}
+	else
+	{
+		const FPoseSearchIndexAsset& SearchIndexAsset = Database->GetSearchIndex().GetAssetForPose(PoseIdx);
+		const FInstancedStruct& DatabaseAsset = Database->GetAnimationAssetStruct(SearchIndexAsset);
+		if (DatabaseAsset.GetPtr<FPoseSearchDatabaseSequence>() || DatabaseAsset.GetPtr<FPoseSearchDatabaseAnimComposite>())
+		{
+			if (Database->GetPoseIndicesAndLerpValueFromTime(NewAssetTime, SearchIndexAsset, PrevPoseIdx, PoseIdx, NextPoseIdx, LerpValue))
+			{
+				AssetTime = NewAssetTime;
+			}
+			else
+			{
+				Reset();
+			}
+		}
+		else if (const FPoseSearchDatabaseBlendSpace* DatabaseBlendSpace = DatabaseAsset.GetPtr<FPoseSearchDatabaseBlendSpace>())
+		{
+			TArray<FBlendSampleData> BlendSamples;
+			int32 TriangulationIndex = 0;
+			DatabaseBlendSpace->BlendSpace->GetSamplesFromBlendInput(SearchIndexAsset.BlendParameters, BlendSamples, TriangulationIndex, true);
+
+			const float PlayLength = DatabaseBlendSpace->BlendSpace->GetAnimationLengthFromSampleData(BlendSamples);
+
+			// Asset player time for blendspaces is normalized [0, 1] so we need to convert 
+			// to a real time before we advance it
+			const float RealTime = NewAssetTime * PlayLength;
+			if (Database->GetPoseIndicesAndLerpValueFromTime(RealTime, SearchIndexAsset, PrevPoseIdx, PoseIdx, NextPoseIdx, LerpValue))
+			{
+				AssetTime = NewAssetTime;
+			}
+			else
+			{
+				Reset();
+			}
+		}
+		else
+		{
+			checkNoEntry();
+		}
+	}
+}
+
+bool FSearchResult::IsValid() const
+{
+	return PoseIdx != INDEX_NONE && Database.IsValid();
+}
+
+void FSearchResult::Reset()
+{
+	PoseIdx = INDEX_NONE;
+	Database = nullptr;
+	ComposedQuery.Reset();
+	AssetTime = 0.0f;
+}
+
+const FPoseSearchIndexAsset* FSearchResult::GetSearchIndexAsset(bool bMandatory) const
+{
+	if (bMandatory)
+	{
+		check(IsValid());
+	}
+	else if (!IsValid())
+	{
+		return nullptr;
+	}
+
+	return &Database->GetSearchIndex().GetAssetForPose(PoseIdx);
+}
+
+} // namespace UE::PoseSearch
diff --git a/Engine/Plugins/Experimental/Animation/PoseSearch/Source/Runtime/Private/PoseSearchSchema.cpp b/Engine/Plugins/Experimental/Animation/PoseSearch/Source/Runtime/Private/PoseSearchSchema.cpp
--- a/Engine/Plugins/Experimental/Animation/PoseSearch/Source/Runtime/Private/PoseSearchSchema.cpp
+++ b/Engine/Plugins/Experimental/Animation/PoseSearch/Source/Runtime/Private/PoseSearchSchema.cpp
@@ -2,8 +2,9 @@
 
 #include "PoseSearch/PoseSearchSchema.h"
 #include "AnimationRuntime.h"
-#include "PoseSearch/PoseSearch.h"
 #include "PoseSearch/PoseSearchFeatureChannel.h"
+#include "PoseSearch/PoseSearchResult.h"
+#include "UObject/ObjectSaveContext.h"
 
 bool UPoseSearchSchema::IsValid() const
 {
@@ -66,3 +67,60 @@ void UPoseSearchSchema::BuildQuery(UE::PoseSearch::FSearchContext& SearchContext
 		Channel->BuildQuery(SearchContext, InOutQuery);
 	}
 }
+
+void UPoseSearchSchema::Finalize(bool bRemoveEmptyChannels)
+{
+	using namespace UE::PoseSearch;
+
+	if (bRemoveEmptyChannels)
+	{
+		Channels.RemoveAll([](TObjectPtr<UPoseSearchFeatureChannel>& Channel) { return !Channel; });
+	}
+
+	BoneReferences.Reset();
+
+	int32 CurrentChannelDataOffset = 0;
+
+	SchemaCardinality = 0;
+	for (int32 ChannelIdx = 0; ChannelIdx != Channels.Num(); ++ChannelIdx)
+	{
+		if (UPoseSearchFeatureChannel* Channel = Channels[ChannelIdx].Get())
+		{
+			Channel->InitializeSchema(this);
+		}
+	}
+
+	ResolveBoneReferences();
+}
+
+void UPoseSearchSchema::PreSave(FObjectPreSaveContext ObjectSaveContext)
+{
+	Finalize();
+
+	Super::PreSave(ObjectSaveContext);
+}
+
+void UPoseSearchSchema::PostLoad()
+{
+	Super::PostLoad();
+	ResolveBoneReferences();
+}
+
+#if WITH_EDITOR
+void UPoseSearchSchema::PostEditChangeProperty(FPropertyChangedEvent& PropertyChangedEvent)
+{
+	Finalize(false);
+	Super::PostEditChangeProperty(PropertyChangedEvent);
+}
+
+void UPoseSearchSchema::ComputeCostBreakdowns(UE::PoseSearch::ICostBreakDownData& CostBreakDownData) const
+{
+	for (const TObjectPtr<UPoseSearchFeatureChannel>& Channel : Channels)
+	{
+		if (Channel)
+		{
+			Channel->ComputeCostBreakdowns(CostBreakDownData, this);
+		}
+	}
+}
+#endif
\ No newline at end of file
diff --git a/Engine/Plugins/Experimental/Animation/PoseSearch/Source/Runtime/Private/Trace/PoseSearchTraceLogger.cpp b/Engine/Plugins/Experimental/Animation/PoseSearch/Source/Runtime/Private/Trace/PoseSearchTraceLogger.cpp
--- a/Engine/Plugins/Experimental/Animation/PoseSearch/Source/Runtime/Private/Trace/PoseSearchTraceLogger.cpp
+++ b/Engine/Plugins/Experimental/Animation/PoseSearch/Source/Runtime/Private/Trace/PoseSearchTraceLogger.cpp
@@ -2,12 +2,13 @@
 
 #include "PoseSearchTraceLogger.h"
 #include "Animation/AnimInstanceProxy.h"
+#include "Animation/AnimNodeBase.h"
 #include "Components/SkeletalMeshComponent.h"
+#include "PoseSearch/PoseSearchDatabase.h"
+#include "Serialization/MemoryWriter.h"
 #include "Trace/Trace.inl"
-#include "Animation/AnimNodeBase.h"
 #include "TraceFilter.h"
 
-
 UE_TRACE_CHANNEL_DEFINE(PoseSearchChannel);
 
 UE_TRACE_EVENT_BEGIN(PoseSearch, MotionMatchingState)
diff --git a/Engine/Plugins/Experimental/Animation/PoseSearch/Source/Runtime/Private/Trace/PoseSearchTraceLogger.h b/Engine/Plugins/Experimental/Animation/PoseSearch/Source/Runtime/Private/Trace/PoseSearchTraceLogger.h
--- a/Engine/Plugins/Experimental/Animation/PoseSearch/Source/Runtime/Private/Trace/PoseSearchTraceLogger.h
+++ b/Engine/Plugins/Experimental/Animation/PoseSearch/Source/Runtime/Private/Trace/PoseSearchTraceLogger.h
@@ -2,7 +2,7 @@
 
 #pragma once
 
-#include "PoseSearch/PoseSearch.h"
+#include "PoseSearch/PoseSearchContext.h"
 
 UE_TRACE_CHANNEL_EXTERN(PoseSearchChannel, POSESEARCH_API);
 
diff --git a/Engine/Plugins/Experimental/Animation/PoseSearch/Source/Runtime/Public/PoseSearch/AnimNode_BlendStack.h b/Engine/Plugins/Experimental/Animation/PoseSearch/Source/Runtime/Public/PoseSearch/AnimNode_BlendStack.h
--- a/Engine/Plugins/Experimental/Animation/PoseSearch/Source/Runtime/Public/PoseSearch/AnimNode_BlendStack.h
+++ b/Engine/Plugins/Experimental/Animation/PoseSearch/Source/Runtime/Public/PoseSearch/AnimNode_BlendStack.h
@@ -7,8 +7,7 @@
 #include "AnimNodes/AnimNode_BlendSpacePlayer.h"
 #include "AnimNodes/AnimNode_Mirror.h"
 #include "Containers/Deque.h"
-#include "PoseSearch/PoseSearch.h"
-
+#include "PoseSearch/PoseSearchIndex.h"
 #include "AnimNode_BlendStack.generated.h"
 
 USTRUCT()
diff --git a/Engine/Plugins/Experimental/Animation/PoseSearch/Source/Runtime/Public/PoseSearch/AnimNode_MotionMatching.h b/Engine/Plugins/Experimental/Animation/PoseSearch/Source/Runtime/Public/PoseSearch/AnimNode_MotionMatching.h
--- a/Engine/Plugins/Experimental/Animation/PoseSearch/Source/Runtime/Public/PoseSearch/AnimNode_MotionMatching.h
+++ b/Engine/Plugins/Experimental/Animation/PoseSearch/Source/Runtime/Public/PoseSearch/AnimNode_MotionMatching.h
@@ -4,12 +4,13 @@
 
 #include "Animation/AnimNode_AssetPlayerBase.h"
 #include "Animation/MotionTrajectoryTypes.h"
+#include "GameplayTagContainer.h"
 #include "PoseSearch/AnimNode_BlendStack.h"
-#include "PoseSearch/PoseSearch.h"
 #include "PoseSearch/PoseSearchLibrary.h"
-
 #include "AnimNode_MotionMatching.generated.h"
 
+class UPoseSearchSearchableAsset;
+
 USTRUCT(BlueprintInternalUseOnly)
 struct POSESEARCH_API FAnimNode_MotionMatching : public FAnimNode_AssetPlayerBase
 {
diff --git a/Engine/Plugins/Experimental/Animation/PoseSearch/Source/Runtime/Public/PoseSearch/AnimNode_PoseSearchHistoryCollector.h b/Engine/Plugins/Experimental/Animation/PoseSearch/Source/Runtime/Public/PoseSearch/AnimNode_PoseSearchHistoryCollector.h
--- a/Engine/Plugins/Experimental/Animation/PoseSearch/Source/Runtime/Public/PoseSearch/AnimNode_PoseSearchHistoryCollector.h
+++ b/Engine/Plugins/Experimental/Animation/PoseSearch/Source/Runtime/Public/PoseSearch/AnimNode_PoseSearchHistoryCollector.h
@@ -2,8 +2,8 @@
 
 #pragma once
 
-#include "PoseSearch/PoseSearch.h"
 #include "Animation/AnimNodeBase.h"
+#include "PoseSearch/PoseSearchHistory.h"
 #include "AnimNode_PoseSearchHistoryCollector.generated.h"
 
 USTRUCT(BlueprintInternalUseOnly)
diff --git a/Engine/Plugins/Experimental/Animation/PoseSearch/Source/Runtime/Public/PoseSearch/PoseSearch.h b/Engine/Plugins/Experimental/Animation/PoseSearch/Source/Runtime/Public/PoseSearch/PoseSearch.h
deleted file mode 100644
diff --git a/Engine/Plugins/Experimental/Animation/PoseSearch/Source/Runtime/Public/PoseSearch/PoseSearchAssetIndexer.h b/Engine/Plugins/Experimental/Animation/PoseSearch/Source/Runtime/Public/PoseSearch/PoseSearchAssetIndexer.h
new file mode 100644
--- /dev/null
+++ b/Engine/Plugins/Experimental/Animation/PoseSearch/Source/Runtime/Public/PoseSearch/PoseSearchAssetIndexer.h
@@ -0,0 +1,59 @@
+// Copyright Epic Games, Inc. All Rights Reserved.
+
+#pragma once
+
+#include "CoreMinimal.h"
+#include "PoseSearch/PoseSearchSchema.h"
+
+namespace UE::PoseSearch
+{
+	
+struct FAssetSamplingContext;
+class IAssetSampler;
+
+/**
+ * Inputs for asset indexing
+ */
+struct FAssetIndexingContext
+{
+	const FAssetSamplingContext* SamplingContext = nullptr;
+	const UPoseSearchSchema* Schema = nullptr;
+	const IAssetSampler* AssetSampler = nullptr;
+	bool bMirrored = false;
+	FFloatInterval RequestedSamplingRange = FFloatInterval(0.0f, 0.0f);
+	
+	// Index this asset's data from BeginPoseIdx up to but not including EndPoseIdx
+	int32 BeginSampleIdx = 0;
+	int32 EndSampleIdx = 0;
+
+	TArrayView<float> GetPoseVector(int32 VectorIdx, TArrayView<float> FeatureVectorTable) const
+	{
+		check(Schema);
+		return MakeArrayView(&FeatureVectorTable[VectorIdx * Schema->SchemaCardinality], Schema->SchemaCardinality);
+	}
+};
+
+class POSESEARCH_API IAssetIndexer
+{
+public:
+	struct FSampleInfo
+	{
+		const IAssetSampler* Clip = nullptr;
+		FTransform RootTransform;
+		float ClipTime = 0.0f;
+		float RootDistance = 0.0f;
+		bool bClamped = false;
+
+		bool IsValid() const { return Clip != nullptr; }
+	};
+
+	virtual ~IAssetIndexer() {}
+
+	virtual const FAssetIndexingContext& GetIndexingContext() const = 0;
+	virtual FSampleInfo GetSampleInfo(float SampleTime) const = 0;
+	virtual FSampleInfo GetSampleInfoRelative(float SampleTime, const FSampleInfo& Origin) const = 0;
+	virtual FTransform MirrorTransform(const FTransform& Transform) const = 0;
+	virtual FTransform GetTransformAndCacheResults(float SampleTime, float OriginTime, int8 SchemaBoneIdx, bool& Clamped) = 0;
+};
+
+} // namespace UE::PoseSearch
diff --git a/Engine/Plugins/Experimental/Animation/PoseSearch/Source/Runtime/Public/PoseSearch/PoseSearchAssetSampler.h b/Engine/Plugins/Experimental/Animation/PoseSearch/Source/Runtime/Public/PoseSearch/PoseSearchAssetSampler.h
new file mode 100644
--- /dev/null
+++ b/Engine/Plugins/Experimental/Animation/PoseSearch/Source/Runtime/Public/PoseSearch/PoseSearchAssetSampler.h
@@ -0,0 +1,179 @@
+// Copyright Epic Games, Inc. All Rights Reserved.
+
+#pragma once
+
+#include "BonePose.h"
+#include "PoseSearchAssetSampler.generated.h"
+
+struct FAnimationPoseData;
+struct FAnimExtractContext;
+class UAnimationAsset;
+class UAnimNotifyState_PoseSearchBase;
+class UBlendSpace;
+class UMirrorDataTable;
+
+USTRUCT()
+struct POSESEARCH_API FPoseSearchExtrapolationParameters
+{
+	GENERATED_BODY()
+
+	// If the angular root motion speed in degrees is below this value, it will be treated as zero.
+	UPROPERTY(EditAnywhere, Category = "Settings")
+	float AngularSpeedThreshold = 1.0f;
+
+	// If the root motion linear speed is below this value, it will be treated as zero.
+	UPROPERTY(EditAnywhere, Category = "Settings")
+	float LinearSpeedThreshold = 1.0f;
+
+	// Time from sequence start/end used to extrapolate the trajectory.
+	UPROPERTY(EditAnywhere, Category = "Settings")
+	float SampleTime = 0.05f;
+};
+
+namespace UE::PoseSearch
+{
+
+struct POSESEARCH_API FAssetSamplingContext
+{
+	// Time delta used for computing pose derivatives
+	static constexpr float FiniteDelta = 1 / 60.0f;
+
+	// Mirror data table pointer copied from Schema for convenience
+	TObjectPtr<const UMirrorDataTable> MirrorDataTable;
+
+	// Compact pose format of Mirror Bone Map
+	TCustomBoneIndexArray<FCompactPoseBoneIndex, FCompactPoseBoneIndex> CompactPoseMirrorBones;
+
+	// Pre-calculated component space rotations of reference pose, which allows mirror to work with any joint orientation
+	// Only initialized and used when a mirroring table is specified
+	TCustomBoneIndexArray<FQuat, FCompactPoseBoneIndex> ComponentSpaceRefRotations;
+
+	void Init(const UMirrorDataTable* InMirrorDataTable, const FBoneContainer& BoneContainer);
+	FTransform MirrorTransform(const FTransform& Transform) const;
+};
+
+/**
+ * Helper interface for sampling data from animation assets
+ */
+class POSESEARCH_API IAssetSampler
+{
+public:
+	virtual ~IAssetSampler() {};
+
+	virtual float GetPlayLength() const = 0;
+	virtual bool IsLoopable() const = 0;
+
+	// Gets the time associated with a particular root distance traveled
+	virtual float GetTimeFromRootDistance(float Distance) const = 0;
+
+	// Gets the total root distance traveled 
+	virtual float GetTotalRootDistance() const = 0;
+
+	// Gets the final root transformation at the end of the asset's playback time
+	virtual FTransform GetTotalRootTransform() const = 0;
+
+	// Extracts pose for this asset for a given context
+	virtual void ExtractPose(const FAnimExtractContext& ExtractionCtx, FAnimationPoseData& OutAnimPoseData) const = 0;
+
+	// Extracts the accumulated root distance at the given time, using the extremities of the sequence to extrapolate 
+	// beyond the sequence limits when Time is less than zero or greater than the sequence length
+	virtual float ExtractRootDistance(float Time) const = 0;
+
+	// Extracts root transform at the given time, using the extremities of the sequence to extrapolate beyond the 
+	// sequence limits when Time is less than zero or greater than the sequence length.
+	virtual FTransform ExtractRootTransform(float Time) const = 0;
+
+	// Extracts notify states inheriting from UAnimNotifyState_PoseSearchBase present in the sequence at Time.
+	// The function does not empty NotifyStates before adding new notifies!
+	virtual void ExtractPoseSearchNotifyStates(float Time, TArray<UAnimNotifyState_PoseSearchBase*>& NotifyStates) const = 0;
+
+	virtual const UAnimationAsset* GetAsset() const = 0;
+
+	virtual void Process() = 0;
+};
+
+// Sampler working with UAnimSequenceBase so it can be used for UAnimSequence as well as UAnimComposite.
+struct POSESEARCH_API FSequenceBaseSampler : public IAssetSampler
+{
+public:
+	struct FInput
+	{
+		TWeakObjectPtr<const UAnimSequenceBase> SequenceBase;
+		int32 RootDistanceSamplingRate = 60;
+		FPoseSearchExtrapolationParameters ExtrapolationParameters;
+	} Input;
+
+	void Init(const FInput& Input);
+	virtual void Process() override;
+
+	virtual float GetPlayLength() const override;
+	virtual bool IsLoopable() const override;
+
+	virtual float GetTimeFromRootDistance(float Distance) const override;
+
+	virtual float GetTotalRootDistance() const override { return TotalRootDistance; }
+	virtual FTransform GetTotalRootTransform() const override { return TotalRootTransform; }
+
+	virtual void ExtractPose(const FAnimExtractContext& ExtractionCtx, FAnimationPoseData& OutAnimPoseData) const override;
+	virtual float ExtractRootDistance(float Time) const override;
+	virtual FTransform ExtractRootTransform(float Time) const override;
+	virtual void ExtractPoseSearchNotifyStates(float Time, TArray<class UAnimNotifyState_PoseSearchBase*>& NotifyStates) const override;
+	virtual const UAnimationAsset* GetAsset() const override;
+
+private:
+	float TotalRootDistance = 0.0f;
+	FTransform TotalRootTransform = FTransform::Identity;
+	TArray<float> AccumulatedRootDistance;
+
+	void ProcessRootDistance();
+};
+
+struct POSESEARCH_API FBlendSpaceSampler : public IAssetSampler
+{
+public:
+	struct FInput
+	{
+		FBoneContainer BoneContainer;
+		TWeakObjectPtr<const UBlendSpace> BlendSpace;
+		int32 RootDistanceSamplingRate = 60;
+		int32 RootTransformSamplingRate = 60;
+		FPoseSearchExtrapolationParameters ExtrapolationParameters;
+		FVector BlendParameters;
+	} Input;
+
+	void Init(const FInput& Input);
+	virtual void Process() override;
+
+	virtual float GetPlayLength() const override { return PlayLength; }
+	virtual bool IsLoopable() const override;
+
+	virtual float GetTimeFromRootDistance(float Distance) const override;
+
+	virtual float GetTotalRootDistance() const override { return TotalRootDistance; }
+	virtual FTransform GetTotalRootTransform() const override { return TotalRootTransform; }
+
+	virtual void ExtractPose(const FAnimExtractContext& ExtractionCtx, FAnimationPoseData& OutAnimPoseData) const override;
+	virtual float ExtractRootDistance(float Time) const override;
+	virtual FTransform ExtractRootTransform(float Time) const override;
+	virtual void ExtractPoseSearchNotifyStates(float Time, TArray<class UAnimNotifyState_PoseSearchBase*>& NotifyStates) const override;
+
+	virtual const UAnimationAsset* GetAsset() const override;
+
+private:
+	float PlayLength = 0.0f;
+	float TotalRootDistance = 0.0f;
+	FTransform TotalRootTransform = FTransform::Identity;
+	TArray<float> AccumulatedRootDistance;
+	TArray<FTransform> AccumulatedRootTransform;
+
+	void ProcessPlayLength();
+	void ProcessRootDistance();
+	void ProcessRootTransform();
+
+	// Extracts the pre-computed blend space root transform. ProcessRootTransform must be run first.
+	FTransform ExtractBlendSpaceRootTrackTransform(float Time) const;
+	FTransform ExtractBlendSpaceRootMotion(float StartTime, float DeltaTime, bool bAllowLooping) const;
+	FTransform ExtractBlendSpaceRootMotionFromRange(float StartTrackPosition, float EndTrackPosition) const;
+};
+
+} // namespace UE::PoseSearch
diff --git a/Engine/Plugins/Experimental/Animation/PoseSearch/Source/Runtime/Public/PoseSearch/PoseSearchContext.h b/Engine/Plugins/Experimental/Animation/PoseSearch/Source/Runtime/Public/PoseSearch/PoseSearchContext.h
new file mode 100644
--- /dev/null
+++ b/Engine/Plugins/Experimental/Animation/PoseSearch/Source/Runtime/Public/PoseSearch/PoseSearchContext.h
@@ -0,0 +1,208 @@
+// Copyright Epic Games, Inc. All Rights Reserved.
+
+#pragma once
+
+#include "PoseSearch/PoseSearchDefines.h"
+#include "PoseSearch/PoseSearchIndex.h"
+#include "PoseSearch/PoseSearchResult.h"
+
+struct FGameplayTagContainer;
+struct FTrajectorySampleRange;
+class UPoseSearchDatabase;
+class UPoseSearchSchema;
+
+namespace UE::PoseSearch
+{
+
+class FPoseHistory;
+struct FPoseIndicesHistory;
+
+enum class EDebugDrawFlags : uint32
+{
+	None = 0,
+
+	// Draw the entire search index as a point cloud
+	DrawSearchIndex = 1 << 0,
+
+	// Draw using Query colors form the schema / config
+	DrawQuery = 1 << 1,
+
+	/**
+	* Keep rendered data until the next call to FlushPersistentDebugLines().
+	* Combine with DrawSearchIndex to draw the search index only once.
+	*/
+	Persistent = 1 << 2,
+
+	// Label samples with their indices
+	DrawSampleLabels = 1 << 3,
+
+	// Draw Bone Names
+	DrawBoneNames = 1 << 5,
+
+	// Draws simpler shapes to improve performance
+	DrawFast = 1 << 6,
+};
+ENUM_CLASS_FLAGS(EDebugDrawFlags);
+
+enum class EPoseCandidateFlags : uint8
+{
+	None = 0,
+
+	Valid_Pose = 1 << 0,
+	Valid_ContinuingPose = 1 << 1,
+	Valid_CurrentPose = 1 << 2,
+
+	AnyValidMask = Valid_Pose | Valid_ContinuingPose | Valid_CurrentPose,
+
+	DiscardedBy_PoseJumpThresholdTime = 1 << 3,
+	DiscardedBy_PoseReselectHistory = 1 << 4,
+	DiscardedBy_BlockTransition = 1 << 5,
+	DiscardedBy_PoseFilter = 1 << 6,
+
+	AnyDiscardedMask = DiscardedBy_PoseJumpThresholdTime | DiscardedBy_PoseReselectHistory | DiscardedBy_BlockTransition | DiscardedBy_PoseFilter,
+};
+ENUM_CLASS_FLAGS(EPoseCandidateFlags);
+
+// @todo: FDebugDrawParams should be enclosed with #if ENABLE_DRAW_DEBUG
+struct POSESEARCH_API FDebugDrawParams
+{
+	const UWorld* World = nullptr;
+	const UPoseSearchDatabase* Database = nullptr;
+	EDebugDrawFlags Flags = EDebugDrawFlags::None;
+	uint32 ChannelMask = (uint32)-1;
+
+	float DefaultLifeTime = 5.f;
+	float PointSize = 1.f;
+
+	FTransform RootTransform = FTransform::Identity;
+
+	// Optional Mesh for gathering SocketTransform(s)
+	TWeakObjectPtr<const USkinnedMeshComponent> Mesh = nullptr;
+
+	bool CanDraw() const;
+	FColor GetColor(int32 ColorPreset) const;
+	const FPoseSearchIndex* GetSearchIndex() const;
+	const UPoseSearchSchema* GetSchema() const;
+};
+
+POSESEARCH_API void DrawFeatureVector(const FDebugDrawParams& DrawParams, TConstArrayView<float> PoseVector);
+POSESEARCH_API void DrawFeatureVector(const FDebugDrawParams& DrawParams, int32 PoseIdx);
+
+struct POSESEARCH_API FSearchContext
+{
+	EPoseSearchBooleanRequest QueryMirrorRequest = EPoseSearchBooleanRequest::Indifferent;
+	UE::PoseSearch::FDebugDrawParams DebugDrawParams;
+	UE::PoseSearch::FPoseHistory* History = nullptr;
+	const FTrajectorySampleRange* Trajectory = nullptr;
+	TObjectPtr<const USkeletalMeshComponent> OwningComponent = nullptr;
+	UE::PoseSearch::FSearchResult CurrentResult;
+	const FBoneContainer* BoneContainer = nullptr;
+	const FGameplayTagContainer* ActiveTagsContainer = nullptr;
+	float PoseJumpThresholdTime = 0.f;
+	bool bIsTracing = false;
+	bool bForceInterrupt = false;
+	// can the continuing pose advance? (if not we skip evaluating it)
+	bool bCanAdvance = true;
+
+	FTransform TryGetTransformAndCacheResults(float SampleTime, const UPoseSearchSchema* Schema, int8 SchemaBoneIdx);
+	void ClearCachedEntries();
+
+	void ResetCurrentBestCost();
+	void UpdateCurrentBestCost(const FPoseSearchCost& PoseSearchCost);
+	float GetCurrentBestTotalCost() const { return CurrentBestTotalCost; }
+
+	void GetOrBuildQuery(const UPoseSearchDatabase* Database, FPoseSearchFeatureVectorBuilder& FeatureVectorBuilder);
+	const FPoseSearchFeatureVectorBuilder* GetCachedQuery(const UPoseSearchDatabase* Database) const;
+
+	bool IsCurrentResultFromDatabase(const UPoseSearchDatabase* Database) const;
+
+	TConstArrayView<float> GetCurrentResultPrevPoseVector() const;
+	TConstArrayView<float> GetCurrentResultPoseVector() const;
+	TConstArrayView<float> GetCurrentResultNextPoseVector() const;
+
+	static constexpr int8 SchemaRootBoneIdx = -1;
+
+	const FPoseIndicesHistory* PoseIndicesHistory = nullptr;
+
+private:
+	struct FCachedEntry
+	{
+		FCachedEntry(float InSampleTime = 0.f, const FTransform& InTransform = FTransform::Identity, FBoneIndexType InBoneIndexType = -1)
+			: SampleTime(InSampleTime)
+			, Transform(InTransform)
+			, BoneIndexType(InBoneIndexType)
+		{
+		}
+
+		float SampleTime = 0.f;
+
+		// associated transform to BoneIndexType in ComponentSpace (except for the root bone stored in global space)
+		FTransform Transform = FTransform::Identity;
+
+		// if -1 it represents the root bone
+		FBoneIndexType BoneIndexType = -1;
+	};
+
+	TArray<FCachedEntry, TInlineAllocator<64>> CachedEntries;
+	
+	struct FCachedQuery
+	{
+		const UPoseSearchDatabase* Database = nullptr;
+		FPoseSearchFeatureVectorBuilder FeatureVectorBuilder;
+	};
+
+	TArray<FCachedQuery, TInlineAllocator<8>> CachedQueries;
+
+	float CurrentBestTotalCost = MAX_flt;
+
+#if UE_POSE_SEARCH_TRACE_ENABLED
+
+public:
+	struct FPoseCandidate
+	{
+		FPoseSearchCost Cost;
+		int32 PoseIdx = 0;
+		const UPoseSearchDatabase* Database = nullptr;
+		EPoseCandidateFlags PoseCandidateFlags = EPoseCandidateFlags::None;
+
+		bool operator<(const FPoseCandidate& Other) const { return Other.Cost < Cost; } // Reverse compare because BestCandidates is a max heap
+		bool operator==(const FSearchResult& SearchResult) const { return (PoseIdx == SearchResult.PoseIdx) && (Database == SearchResult.Database.Get()); }
+	};
+
+	struct FBestPoseCandidates : private TArray<FPoseCandidate>
+	{
+		typedef TArray<FPoseCandidate> Super;
+		using Super::IsEmpty;
+
+		int32 MaxPoseCandidates = 100;
+
+		void Add(const FPoseSearchCost& Cost, int32 PoseIdx, const UPoseSearchDatabase* Database, EPoseCandidateFlags PoseCandidateFlags)
+		{
+			if (Num() < MaxPoseCandidates || Cost < HeapTop().Cost)
+			{
+				while (Num() >= MaxPoseCandidates)
+				{
+					ElementType Unused;
+					Pop(Unused);
+				}
+
+				FSearchContext::FPoseCandidate PoseCandidate;
+				PoseCandidate.Cost = Cost;
+				PoseCandidate.PoseIdx = PoseIdx;
+				PoseCandidate.Database = Database;
+				PoseCandidate.PoseCandidateFlags = PoseCandidateFlags;
+				HeapPush(PoseCandidate);
+			}
+		}
+
+		void Pop(FPoseCandidate& OutItem)
+		{
+			HeapPop(OutItem, false);
+		}
+	};
+	
+	FBestPoseCandidates BestCandidates;
+#endif
+};
+
+} // namespace UE::PoseSearch
diff --git a/Engine/Plugins/Experimental/Animation/PoseSearch/Source/Runtime/Public/PoseSearch/PoseSearchCost.h b/Engine/Plugins/Experimental/Animation/PoseSearch/Source/Runtime/Public/PoseSearch/PoseSearchCost.h
new file mode 100644
--- /dev/null
+++ b/Engine/Plugins/Experimental/Animation/PoseSearch/Source/Runtime/Public/PoseSearch/PoseSearchCost.h
@@ -0,0 +1,48 @@
+// Copyright Epic Games, Inc. All Rights Reserved.
+
+#pragma once
+
+#include "PoseSearchCost.generated.h"
+
+USTRUCT()
+struct POSESEARCH_API FPoseSearchCost
+{
+	GENERATED_BODY()
+public:
+	FPoseSearchCost() = default;
+	FPoseSearchCost(float InDissimilarityCost, float InNotifyCostAddend, float InMirrorMismatchCostAddend, float InContinuingPoseCostAddend)
+	: TotalCost(InDissimilarityCost + InNotifyCostAddend + InMirrorMismatchCostAddend + InContinuingPoseCostAddend)
+	{
+#if WITH_EDITORONLY_DATA
+		NotifyCostAddend = InNotifyCostAddend;
+		MirrorMismatchCostAddend = InMirrorMismatchCostAddend;
+		ContinuingPoseCostAddend = InContinuingPoseCostAddend;
+#endif // WITH_EDITORONLY_DATA
+	}
+
+	bool IsValid() const { return TotalCost != MAX_flt; }
+	float GetTotalCost() const { return TotalCost; }
+	bool operator<(const FPoseSearchCost& Other) const { return TotalCost < Other.TotalCost; }
+
+protected:
+	UPROPERTY()
+	float TotalCost = MAX_flt;
+
+#if WITH_EDITORONLY_DATA
+public:
+
+	float GetCostAddend() const { return NotifyCostAddend + MirrorMismatchCostAddend + ContinuingPoseCostAddend; }
+
+	// Contribution from ModifyCost anim notify
+	UPROPERTY()
+	float NotifyCostAddend = 0.f;
+
+	// Contribution from mirroring cost
+	UPROPERTY()
+	float MirrorMismatchCostAddend = 0.f;
+
+	UPROPERTY()
+	float ContinuingPoseCostAddend = 0.f;
+
+#endif // WITH_EDITORONLY_DATA
+};
\ No newline at end of file
diff --git a/Engine/Plugins/Experimental/Animation/PoseSearch/Source/Runtime/Public/PoseSearch/PoseSearchDatabase.h b/Engine/Plugins/Experimental/Animation/PoseSearch/Source/Runtime/Public/PoseSearch/PoseSearchDatabase.h
new file mode 100644
--- /dev/null
+++ b/Engine/Plugins/Experimental/Animation/PoseSearch/Source/Runtime/Public/PoseSearch/PoseSearchDatabase.h
@@ -0,0 +1,282 @@
+// Copyright Epic Games, Inc. All Rights Reserved.
+
+#pragma once
+
+#include "PoseSearch/PoseSearchAssetSampler.h"
+#include "PoseSearch/PoseSearchIndex.h"
+#include "PoseSearch/PoseSearchSearchableAsset.h"
+#include "PoseSearchDatabase.generated.h"
+
+struct FInstancedStruct;
+class UAnimComposite;
+
+UENUM()
+enum class EPoseSearchMode : int32
+{
+	BruteForce,
+	PCAKDTree,
+	PCAKDTree_Validate,	// runs PCAKDTree and performs validation tests
+	PCAKDTree_Compare,	// compares BruteForce vs PCAKDTree
+
+	Num UMETA(Hidden),
+	Invalid = Num UMETA(Hidden)
+};
+
+UENUM()
+enum class EPoseSearchMirrorOption : int32
+{
+	UnmirroredOnly UMETA(DisplayName = "Original Only"),
+	MirroredOnly UMETA(DisplayName = "Mirrored Only"),
+	UnmirroredAndMirrored UMETA(DisplayName = "Original and Mirrored"),
+
+	Num UMETA(Hidden),
+	Invalid = Num UMETA(Hidden)
+};
+
+USTRUCT()
+struct FPoseSearchExcludeFromDatabaseParameters
+{
+	GENERATED_BODY()
+
+	// Excluding the beginning of sequences can help ensure an exact past trajectory is used when building the features
+	UPROPERTY(EditAnywhere, Category = "Settings")
+	float SequenceStartInterval = 0.0f;
+
+	// Excluding the end of sequences help ensure an exact future trajectory, and also prevents the selection of
+	// a sequence which will end too soon to be worth selecting.
+	UPROPERTY(EditAnywhere, Category = "Settings")
+	float SequenceEndInterval = 0.3f;
+};
+
+USTRUCT()
+struct POSESEARCH_API FPoseSearchDatabaseAnimationAssetBase
+{
+	GENERATED_BODY()
+	virtual ~FPoseSearchDatabaseAnimationAssetBase() = default;
+
+	virtual UAnimationAsset* GetAnimationAsset() const { return nullptr; }
+	virtual UClass* GetAnimationAssetStaticClass() const { return nullptr; }
+	virtual bool IsLooping() const { return false; }
+	virtual const FString GetName() const { return {}; }
+	virtual bool IsEnabled() const { return false; }
+	virtual void SetIsEnabled(bool bInIsEnabled) {}
+	virtual bool IsRootMotionEnabled() const { return false; }
+	virtual EPoseSearchMirrorOption GetMirrorOption() const { return EPoseSearchMirrorOption::Invalid; }
+	virtual ESearchIndexAssetType GetSearchIndexType() const { return ESearchIndexAssetType::Invalid; }
+};
+
+/** An entry in a UPoseSearchDatabase. */
+USTRUCT(BlueprintType, Category = "Animation|Pose Search")
+struct POSESEARCH_API FPoseSearchDatabaseSequence : public FPoseSearchDatabaseAnimationAssetBase
+{
+	GENERATED_BODY()
+	virtual ~FPoseSearchDatabaseSequence() = default;
+
+	UPROPERTY(EditAnywhere, Category="Sequence")
+	TObjectPtr<UAnimSequence> Sequence;
+
+	UPROPERTY(EditAnywhere, Category = "Sequence")
+	bool bEnabled = true;
+
+	UPROPERTY(EditAnywhere, Category="Sequence")
+	FFloatInterval SamplingRange = FFloatInterval(0.0f, 0.0f);
+
+	UPROPERTY(EditAnywhere, Category = "Sequence")
+	EPoseSearchMirrorOption MirrorOption = EPoseSearchMirrorOption::UnmirroredOnly;
+
+	UAnimationAsset* GetAnimationAsset() const override;
+	UClass* GetAnimationAssetStaticClass() const override;
+	bool IsLooping() const override;
+	const FString GetName() const override;
+	bool IsEnabled() const override { return bEnabled; }
+	void SetIsEnabled(bool bInIsEnabled) override { bEnabled = bInIsEnabled; }
+	bool IsRootMotionEnabled() const override;
+	EPoseSearchMirrorOption GetMirrorOption() const override { return MirrorOption; }
+	ESearchIndexAssetType GetSearchIndexType() const override { return ESearchIndexAssetType::Sequence; }
+};
+
+/** An blend space entry in a UPoseSearchDatabase. */
+USTRUCT(BlueprintType, Category = "Animation|Pose Search")
+struct POSESEARCH_API FPoseSearchDatabaseBlendSpace : public FPoseSearchDatabaseAnimationAssetBase
+{
+	GENERATED_BODY()
+	virtual ~FPoseSearchDatabaseBlendSpace() = default;
+
+	UPROPERTY(EditAnywhere, Category = "BlendSpace")
+	TObjectPtr<UBlendSpace> BlendSpace;
+
+	UPROPERTY(EditAnywhere, Category = "BlendSpace")
+	bool bEnabled = true;
+
+	UPROPERTY(EditAnywhere, Category = "BlendSpace")
+	EPoseSearchMirrorOption MirrorOption = EPoseSearchMirrorOption::UnmirroredOnly;
+
+	// If to use the blendspace grid locations as parameter sample locations.
+	// When enabled, NumberOfHorizontalSamples and NumberOfVerticalSamples are ignored.
+	UPROPERTY(EditAnywhere, Category = "BlendSpace")
+	bool bUseGridForSampling = false;
+
+	UPROPERTY(EditAnywhere, Category = "BlendSpace", meta = (EditCondition = "!bUseGridForSampling", EditConditionHides, ClampMin = "1", UIMin = "1", UIMax = "25"))
+	int32 NumberOfHorizontalSamples = 9;
+
+	UPROPERTY(EditAnywhere, Category = "BlendSpace", meta = (EditCondition = "!bUseGridForSampling", EditConditionHides, ClampMin = "1", UIMin = "1", UIMax = "25"))
+	int32 NumberOfVerticalSamples = 2;
+
+	UAnimationAsset* GetAnimationAsset() const override;
+	UClass* GetAnimationAssetStaticClass() const override;
+	bool IsLooping() const override;
+	const FString GetName() const override;
+	bool IsEnabled() const override { return bEnabled; }
+	void SetIsEnabled(bool bInIsEnabled) override { bEnabled = bInIsEnabled; }
+	bool IsRootMotionEnabled() const override;
+	EPoseSearchMirrorOption GetMirrorOption() const override { return MirrorOption; }
+	ESearchIndexAssetType GetSearchIndexType() const override { return ESearchIndexAssetType::BlendSpace; }
+
+	void GetBlendSpaceParameterSampleRanges(int32& HorizontalBlendNum, int32& VerticalBlendNum) const;
+	FVector BlendParameterForSampleRanges(int32 HorizontalBlendIndex, int32 VerticalBlendIndex) const;
+};
+
+/** An entry in a UPoseSearchDatabase. */
+USTRUCT(BlueprintType, Category = "Animation|Pose Search")
+struct POSESEARCH_API FPoseSearchDatabaseAnimComposite : public FPoseSearchDatabaseAnimationAssetBase
+{
+	GENERATED_BODY()
+	virtual ~FPoseSearchDatabaseAnimComposite() = default;
+
+	UPROPERTY(EditAnywhere, Category = "AnimComposite")
+	TObjectPtr<UAnimComposite> AnimComposite;
+
+	UPROPERTY(EditAnywhere, Category = "AnimComposite")
+	bool bEnabled = true;
+
+	UPROPERTY(EditAnywhere, Category = "AnimComposite")
+	FFloatInterval SamplingRange = FFloatInterval(0.0f, 0.0f);
+
+	UPROPERTY(EditAnywhere, Category = "AnimComposite")
+	EPoseSearchMirrorOption MirrorOption = EPoseSearchMirrorOption::UnmirroredOnly;
+
+	UAnimationAsset* GetAnimationAsset() const override;
+	UClass* GetAnimationAssetStaticClass() const override;
+	bool IsLooping() const override;
+	const FString GetName() const override;
+	bool IsEnabled() const override { return bEnabled; }
+	void SetIsEnabled(bool bInIsEnabled) override { bEnabled = bInIsEnabled; }
+	bool IsRootMotionEnabled() const override;
+	EPoseSearchMirrorOption GetMirrorOption() const override { return MirrorOption; }
+	ESearchIndexAssetType GetSearchIndexType() const override { return ESearchIndexAssetType::AnimComposite; }
+};
+
+/** A data asset for indexing a collection of animation sequences. */
+UCLASS(BlueprintType, Category = "Animation|Pose Search", Experimental, meta = (DisplayName = "Normalization Set"))
+class POSESEARCH_API UNormalizationSetAsset : public UDataAsset
+{
+	GENERATED_BODY()
+public:
+
+	UPROPERTY(EditAnywhere, BlueprintReadOnly, Category = "NormalizationSet")
+	TArray<TObjectPtr<const UPoseSearchDatabase>> Databases; // @todo: it should be a UPoseSearchSearchableAsset, and have the UPoseSearchSearchableAsset iterate over all it's contained UPoseSearchSearchableAsset recursively (without duplicates)
+};
+
+/** A data asset for indexing a collection of animation sequences. */
+UCLASS(BlueprintType, Category = "Animation|Pose Search", Experimental, meta = (DisplayName = "Motion Database"))
+class POSESEARCH_API UPoseSearchDatabase : public UPoseSearchSearchableAsset
+{
+	GENERATED_BODY()
+public:
+	// Motion Database Config asset to use with this database.
+	UPROPERTY(EditAnywhere, BlueprintReadOnly, Category="Database", DisplayName="Config")
+	TObjectPtr<const UPoseSearchSchema> Schema;
+
+	UPROPERTY(EditAnywhere, Category = "Database")
+	FPoseSearchExtrapolationParameters ExtrapolationParameters;
+
+	UPROPERTY(EditAnywhere, Category = "Database")
+	FPoseSearchExcludeFromDatabaseParameters ExcludeFromDatabaseParameters;
+
+#if WITH_EDITORONLY_DATA
+	// Sequences and Blendspaces are deprecated and its data will be part of the AnimationAssets.
+	// All sequences and blend spaces will be added to the AnimationAssets in PostLoad().
+	UPROPERTY()
+	TArray<FPoseSearchDatabaseSequence> Sequences_DEPRECATED;
+	UPROPERTY()
+	TArray<FPoseSearchDatabaseBlendSpace> BlendSpaces_DEPRECATED;
+#endif // WITH_EDITORONLY_DATA
+
+	UPROPERTY(EditAnywhere, Category="Database")
+	TArray<FInstancedStruct> AnimationAssets;
+
+	UPROPERTY(EditAnywhere, Category = "Performance")
+	EPoseSearchMode PoseSearchMode = EPoseSearchMode::PCAKDTree;
+
+	UPROPERTY(EditAnywhere, Category = "Performance", meta = (EditCondition = "PoseSearchMode != EPoseSearchMode::BruteForce", EditConditionHides, ClampMin = "1", ClampMax = "64", UIMin = "1", UIMax = "64"))
+	int32 NumberOfPrincipalComponents = 4;
+
+	UPROPERTY(EditAnywhere, Category = "Performance", meta = (EditCondition = "PoseSearchMode != EPoseSearchMode::BruteForce", EditConditionHides, ClampMin = "1", ClampMax = "256", UIMin = "1", UIMax = "256"))
+	int32 KDTreeMaxLeafSize = 16;
+	
+	UPROPERTY(EditAnywhere, Category = "Performance", meta = (EditCondition = "PoseSearchMode != EPoseSearchMode::BruteForce", EditConditionHides, ClampMin = "1", ClampMax = "600", UIMin = "1", UIMax = "600"))
+	int32 KDTreeQueryNumNeighbors = 200;
+
+	// if true, this database search will be skipped if cannot decrease the pose cost, and poses will not be listed into the PoseSearchDebugger
+	UPROPERTY(EditAnywhere, Category = "Performance")
+	bool bSkipSearchIfPossible = true;
+
+	UPROPERTY(EditAnywhere, BlueprintReadOnly, Category = "Database")
+	TObjectPtr<const UNormalizationSetAsset> NormalizationSet;
+
+private:
+	UPROPERTY(Transient, meta = (ExcludeFromHash))
+	FPoseSearchIndex SearchIndexPrivate; // do not use it directly. use GetSearchIndex / SetSearchIndex interact with it and validate that is ok to do so
+
+public:
+	virtual ~UPoseSearchDatabase();
+
+	void SetSearchIndex(const FPoseSearchIndex& SearchIndex);
+	const FPoseSearchIndex& GetSearchIndex() const;
+	
+	bool GetSkipSearchIfPossible() const;
+
+	int32 GetPoseIndexFromTime(float AssetTime, const FPoseSearchIndexAsset& SearchIndexAsset) const;
+	bool GetPoseIndicesAndLerpValueFromTime(float Time, const FPoseSearchIndexAsset& SearchIndexAsset, int32& PrevPoseIdx, int32& PoseIdx, int32& NextPoseIdx, float& LerpValue) const;
+
+	const FInstancedStruct& GetAnimationAssetStruct(int32 AnimationAssetIndex) const;
+	const FInstancedStruct& GetAnimationAssetStruct(const FPoseSearchIndexAsset& SearchIndexAsset) const;
+	FInstancedStruct& GetMutableAnimationAssetStruct(int32 AnimationAssetIndex);
+	FInstancedStruct& GetMutableAnimationAssetStruct(const FPoseSearchIndexAsset& SearchIndexAsset);
+	const FPoseSearchDatabaseAnimationAssetBase* GetAnimationAssetBase(int32 AnimationAssetIndex) const;
+	const FPoseSearchDatabaseAnimationAssetBase* GetAnimationAssetBase(const FPoseSearchIndexAsset& SearchIndexAsset) const;
+	FPoseSearchDatabaseAnimationAssetBase* GetMutableAnimationAssetBase(int32 AnimationAssetIndex);
+	FPoseSearchDatabaseAnimationAssetBase* GetMutableAnimationAssetBase(const FPoseSearchIndexAsset& SearchIndexAsset);
+	const bool IsSourceAssetLooping(const FPoseSearchIndexAsset& SearchIndexAsset) const;
+	const FString GetSourceAssetName(const FPoseSearchIndexAsset& SearchIndexAsset) const;
+	int32 GetNumberOfPrincipalComponents() const;
+	void BuildQuery(UE::PoseSearch::FSearchContext& SearchContext, FPoseSearchFeatureVectorBuilder& OutQuery) const;
+
+	// Begin UObject
+	virtual void PostLoad() override;
+	virtual void PostSaveRoot(FObjectPostSaveRootContext ObjectSaveContext) override;
+	virtual void Serialize(FArchive& Ar) override;
+	// End UObject
+	
+	// Begin UPoseSearchSearchableAsset
+	virtual UE::PoseSearch::FSearchResult Search(UE::PoseSearch::FSearchContext& SearchContext) const override;
+	// End UPoseSearchSearchableAsset
+
+#if WITH_EDITOR
+	virtual void BeginCacheForCookedPlatformData(const ITargetPlatform* TargetPlatform) override;
+	virtual bool IsCachedCookedPlatformDataLoaded(const ITargetPlatform* TargetPlatform) override;
+
+private:
+	DECLARE_MULTICAST_DELEGATE(FOnDerivedDataRebuildMulticaster);
+	FOnDerivedDataRebuildMulticaster OnDerivedDataRebuild;
+public:
+	typedef FOnDerivedDataRebuildMulticaster::FDelegate FOnDerivedDataRebuild;
+	void RegisterOnDerivedDataRebuild(const FOnDerivedDataRebuild& Delegate);
+	void UnregisterOnDerivedDataRebuild(void* Unregister);
+	void NotifyDerivedDataRebuild() const;
+#endif // WITH_EDITOR
+
+private:
+	UE::PoseSearch::FSearchResult SearchPCAKDTree(UE::PoseSearch::FSearchContext& SearchContext) const;
+	UE::PoseSearch::FSearchResult SearchBruteForce(UE::PoseSearch::FSearchContext& SearchContext) const;
+};
\ No newline at end of file
diff --git a/Engine/Plugins/Experimental/Animation/PoseSearch/Source/Runtime/Public/PoseSearch/PoseSearchDatabaseSet.h b/Engine/Plugins/Experimental/Animation/PoseSearch/Source/Runtime/Public/PoseSearch/PoseSearchDatabaseSet.h
new file mode 100644
--- /dev/null
+++ b/Engine/Plugins/Experimental/Animation/PoseSearch/Source/Runtime/Public/PoseSearch/PoseSearchDatabaseSet.h
@@ -0,0 +1,52 @@
+// Copyright Epic Games, Inc. All Rights Reserved.
+
+#pragma once
+
+#include "GameplayTagContainer.h"
+#include "PoseSearch/PoseSearchSearchableAsset.h"
+#include "PoseSearchDatabaseSet.generated.h"
+
+UENUM()
+enum class EPoseSearchPostSearchStatus : uint8
+{
+	// Continue looking for results 
+	Continue,
+
+	// Halt and return the best result
+	Stop
+};
+
+USTRUCT(BlueprintType, Category = "Animation|Pose Search")
+struct POSESEARCH_API FPoseSearchDatabaseSetEntry
+{
+	GENERATED_BODY()
+
+public:
+	UPROPERTY(EditAnywhere, Category = Settings)
+	TObjectPtr<UPoseSearchSearchableAsset> Searchable;
+
+	UPROPERTY(EditAnywhere, Category = Settings)
+	FGameplayTag Tag;
+
+	UPROPERTY(EditAnywhere, Category = Settings)
+	EPoseSearchPostSearchStatus PostSearchStatus = EPoseSearchPostSearchStatus::Continue;
+};
+
+/** A data asset which holds a collection searchable assets. */
+UCLASS(BlueprintType, Category = "Animation|Pose Search", Experimental)
+class POSESEARCH_API UPoseSearchDatabaseSet : public UPoseSearchSearchableAsset
+{
+	GENERATED_BODY()
+
+public:
+	UPROPERTY(EditAnywhere, Category = Settings)
+	TArray<FPoseSearchDatabaseSetEntry> AssetsToSearch;
+
+	// if there's a valid continuing pose and bEvaluateContinuingPoseFirst is true, the continuing pose will be evaluated as first search,
+	// otherwise it'll be evaluated with the related database: if the database is not active the continuing pose evaluation will be skipped
+	UPROPERTY(EditAnywhere, Category = Settings)
+	bool bEvaluateContinuingPoseFirst = true;
+
+public:
+	virtual UE::PoseSearch::FSearchResult Search(UE::PoseSearch::FSearchContext& SearchContext) const override;
+};
diff --git a/Engine/Plugins/Experimental/Animation/PoseSearch/Source/Runtime/Public/PoseSearch/PoseSearchDefines.h b/Engine/Plugins/Experimental/Animation/PoseSearch/Source/Runtime/Public/PoseSearch/PoseSearchDefines.h
new file mode 100644
--- /dev/null
+++ b/Engine/Plugins/Experimental/Animation/PoseSearch/Source/Runtime/Public/PoseSearch/PoseSearchDefines.h
@@ -0,0 +1,22 @@
+// Copyright Epic Games, Inc. All Rights Reserved.
+
+#pragma once
+
+#include "Async/ParallelFor.h"
+#include "Logging/LogMacros.h"
+#include "ObjectTrace.h"
+
+// Enable this if object tracing is enabled, mimics animation tracing
+#define UE_POSE_SEARCH_TRACE_ENABLED OBJECT_TRACE_ENABLED
+
+#ifndef UE_POSE_SEARCH_FORCE_SINGLE_THREAD
+#define UE_POSE_SEARCH_FORCE_SINGLE_THREAD 0
+#endif
+
+#if UE_POSE_SEARCH_FORCE_SINGLE_THREAD
+constexpr EParallelForFlags ParallelForFlags = EParallelForFlags::ForceSingleThread;
+#else
+constexpr EParallelForFlags ParallelForFlags = EParallelForFlags::None;
+#endif // UE_POSE_SEARCH_FORCE_SINGLE_THREAD
+
+DECLARE_LOG_CATEGORY_EXTERN(LogPoseSearch, Log, All);
diff --git a/Engine/Plugins/Experimental/Animation/PoseSearch/Source/Runtime/Public/PoseSearch/PoseSearchDerivedData.h b/Engine/Plugins/Experimental/Animation/PoseSearch/Source/Runtime/Public/PoseSearch/PoseSearchDerivedData.h
--- a/Engine/Plugins/Experimental/Animation/PoseSearch/Source/Runtime/Public/PoseSearch/PoseSearchDerivedData.h
+++ b/Engine/Plugins/Experimental/Animation/PoseSearch/Source/Runtime/Public/PoseSearch/PoseSearchDerivedData.h
@@ -3,10 +3,12 @@
 #pragma once
 
 #if WITH_EDITOR
-#include "PoseSearch/PoseSearch.h"
 #include "TickableEditorObject.h"
+#include "UObject/GCObject.h"
 #include "UObject/ObjectSaveContext.h"
 
+class UPoseSearchDatabase;
+
 namespace UE::PoseSearch
 {
 	struct FPoseSearchDatabaseAsyncCacheTask;
diff --git a/Engine/Plugins/Experimental/Animation/PoseSearch/Source/Runtime/Public/PoseSearch/PoseSearchFeatureChannel.h b/Engine/Plugins/Experimental/Animation/PoseSearch/Source/Runtime/Public/PoseSearch/PoseSearchFeatureChannel.h
--- a/Engine/Plugins/Experimental/Animation/PoseSearch/Source/Runtime/Public/PoseSearch/PoseSearchFeatureChannel.h
+++ b/Engine/Plugins/Experimental/Animation/PoseSearch/Source/Runtime/Public/PoseSearch/PoseSearchFeatureChannel.h
@@ -29,11 +29,41 @@ enum class EInputQueryPose : uint8
 namespace UE::PoseSearch
 {
 
-struct FAssetIndexingOutput;
 struct FDebugDrawParams;
 struct FSearchContext;
 class IAssetIndexer;
 
+/** Helper class for extracting and encoding features into a float buffer */
+class POSESEARCH_API FFeatureVectorHelper
+{
+public:
+	enum { EncodeQuatCardinality = 6 };
+	static void EncodeQuat(TArrayView<float> Values, int32& DataOffset, const FQuat& Quat);
+	static void EncodeQuat(TArrayView<float> Values, int32& DataOffset, TConstArrayView<float> PrevValues, TConstArrayView<float> CurValues, TConstArrayView<float> NextValues, float LerpValue);
+	static FQuat DecodeQuat(TConstArrayView<float> Values, int32& DataOffset);
+
+	enum { EncodeVectorCardinality = 3 };
+	static void EncodeVector(TArrayView<float> Values, int32& DataOffset, const FVector& Vector);
+	static void EncodeVector(TArrayView<float> Values, int32& DataOffset, TConstArrayView<float> PrevValues, TConstArrayView<float> CurValues, TConstArrayView<float> NextValues, float LerpValue, bool bNormalize = false);
+	static FVector DecodeVector(TConstArrayView<float> Values, int32& DataOffset);
+
+	enum { EncodeVector2DCardinality = 2 };
+	static void EncodeVector2D(TArrayView<float> Values, int32& DataOffset, const FVector2D& Vector2D);
+	static void EncodeVector2D(TArrayView<float> Values, int32& DataOffset, TConstArrayView<float> PrevValues, TConstArrayView<float> CurValues, TConstArrayView<float> NextValues, float LerpValue);
+	static FVector2D DecodeVector2D(TConstArrayView<float> Values, int32& DataOffset);
+
+	enum { EncodeFloatCardinality = 1 };
+	static void EncodeFloat(TArrayView<float> Values, int32& DataOffset, const float Value);
+	static void EncodeFloat(TArrayView<float> Values, int32& DataOffset, TConstArrayView<float> PrevValues, TConstArrayView<float> CurValues, TConstArrayView<float> NextValues, float LerpValue);
+	static float DecodeFloat(TConstArrayView<float> Values, int32& DataOffset);
+
+private:
+	static FQuat DecodeQuatInternal(TConstArrayView<float> Values, int32 DataOffset);
+	static FVector DecodeVectorInternal(TConstArrayView<float> Values, int32 DataOffset);
+	static FVector2D DecodeVector2DInternal(TConstArrayView<float> Values, int32 DataOffset);
+	static float DecodeFloatInternal(TConstArrayView<float> Values, int32 DataOffset);
+};
+
 #if WITH_EDITOR
 // data structure collecting the internal layout representation of UPoseSearchFeatureChannel,
 // so we can aggregate data from different FPoseSearchIndex and calculate mean deviation with a homogeneous data set (ComputeChannelsDeviations
@@ -129,7 +159,7 @@ public:
 	virtual void FillWeights(TArray<float>& Weights) const PURE_VIRTUAL(UPoseSearchFeatureChannel::FillWeights, );
 
 	// Called at database build time to populate pose vectors with this channel's data
-	virtual void IndexAsset(UE::PoseSearch::IAssetIndexer& Indexer, UE::PoseSearch::FAssetIndexingOutput& IndexingOutput) const PURE_VIRTUAL(UPoseSearchFeatureChannel::IndexAsset, );
+	virtual void IndexAsset(UE::PoseSearch::IAssetIndexer& Indexer, TArrayView<float> FeatureVectorTable) const PURE_VIRTUAL(UPoseSearchFeatureChannel::IndexAsset, );
 
 	// Called at runtime to add this channel's data to the query pose vector
 	virtual void BuildQuery(UE::PoseSearch::FSearchContext& SearchContext, FPoseSearchFeatureVectorBuilder& InOutQuery) const PURE_VIRTUAL(UPoseSearchFeatureChannel::BuildQuery, );
diff --git a/Engine/Plugins/Experimental/Animation/PoseSearch/Source/Runtime/Public/PoseSearch/PoseSearchHistory.h b/Engine/Plugins/Experimental/Animation/PoseSearch/Source/Runtime/Public/PoseSearch/PoseSearchHistory.h
new file mode 100644
--- /dev/null
+++ b/Engine/Plugins/Experimental/Animation/PoseSearch/Source/Runtime/Public/PoseSearch/PoseSearchHistory.h
@@ -0,0 +1,83 @@
+// Copyright Epic Games, Inc. All Rights Reserved.
+
+#pragma once
+
+#include "Animation/AnimNodeMessages.h"
+#include "BoneIndices.h"
+#include "Containers/RingBuffer.h"
+#include "UObject/ObjectKey.h"
+
+struct FPoseContext;
+
+namespace UE::PoseSearch
+{
+
+struct FSearchResult;
+
+/**
+* Records poses over time in a ring buffer.
+* FFeatureVectorBuilder uses this to sample from the present or past poses according to the search schema.
+*/
+class POSESEARCH_API FPoseHistory
+{
+public:
+
+	enum class ERootUpdateMode
+	{
+		RootMotionDelta,
+		ComponentTransformDelta,
+	};
+
+	void Init(int32 InNumPoses, float InTimeHorizon);
+	void Init(const FPoseHistory& History);
+	bool Update(float SecondsElapsed, const FPoseContext& PoseContext, FTransform ComponentTransform, FText* OutError, ERootUpdateMode UpdateMode = ERootUpdateMode::RootMotionDelta);
+	float GetSampleTimeInterval() const;
+	float GetTimeHorizon() const { return TimeHorizon; }
+	bool TrySampleLocalPose(float Time, const TArray<FBoneIndexType>* RequiredBones, TArray<FTransform>* LocalPose, FTransform* RootTransform) const;
+
+private:
+
+	struct FPose
+	{
+		FTransform RootTransform;
+		TArray<FTransform> LocalTransforms;
+		float Time = 0.0f;
+	};
+	TRingBuffer<FPose> Poses;
+	float TimeHorizon = 0.0f;
+};
+
+class IPoseHistoryProvider : public UE::Anim::IGraphMessage
+{
+	DECLARE_ANIMGRAPH_MESSAGE(IPoseHistoryProvider);
+public:
+	virtual const FPoseHistory& GetPoseHistory() const = 0;
+	virtual FPoseHistory& GetPoseHistory() = 0;
+};
+
+struct FHistoricalPoseIndex
+{
+	bool operator==(const FHistoricalPoseIndex& Index) const
+	{
+		return PoseIndex == Index.PoseIndex && DatabaseKey == Index.DatabaseKey;
+	}
+
+	friend FORCEINLINE uint32 GetTypeHash(const FHistoricalPoseIndex& Index)
+	{
+		return HashCombineFast(::GetTypeHash(Index.PoseIndex), GetTypeHash(Index.DatabaseKey));
+	}
+
+	int32 PoseIndex = INDEX_NONE;
+	FObjectKey DatabaseKey;
+};
+
+struct FPoseIndicesHistory
+{
+	void Update(const FSearchResult& SearchResult, float DeltaTime, float MaxTime);
+	void Reset() { IndexToTime.Reset(); }
+	TMap<FHistoricalPoseIndex, float> IndexToTime;
+};
+
+} // namespace UE::PoseSearch
+
+
diff --git a/Engine/Plugins/Experimental/Animation/PoseSearch/Source/Runtime/Public/PoseSearch/PoseSearchIndex.h b/Engine/Plugins/Experimental/Animation/PoseSearch/Source/Runtime/Public/PoseSearch/PoseSearchIndex.h
new file mode 100644
--- /dev/null
+++ b/Engine/Plugins/Experimental/Animation/PoseSearch/Source/Runtime/Public/PoseSearch/PoseSearchIndex.h
@@ -0,0 +1,217 @@
+// Copyright Epic Games, Inc. All Rights Reserved.
+
+#pragma once
+
+#include "PoseSearch/KDTree.h"
+#include "PoseSearch/PoseSearchCost.h"
+#include "PoseSearchIndex.generated.h"
+
+namespace UE::PoseSearch
+{
+
+POSESEARCH_API void CompareFeatureVectors(TConstArrayView<float> A, TConstArrayView<float> B, TConstArrayView<float> WeightsSqrt, TArrayView<float> Result);
+
+enum class EPoseComparisonFlags : int32
+{
+	None = 0,
+	ContinuingPose = 1 << 0,
+};
+ENUM_CLASS_FLAGS(EPoseComparisonFlags);
+
+} // namespace UE::PoseSearch
+
+UENUM()
+enum class ESearchIndexAssetType : int32
+{
+	Invalid,
+	Sequence,
+	BlendSpace,
+	AnimComposite,
+};
+
+UENUM()
+enum class EPoseSearchBooleanRequest : uint8
+{
+	FalseValue,
+	TrueValue,
+	Indifferent, // if this is used, there will be no cost difference between true and false results
+
+	Num UMETA(Hidden),
+	Invalid = Num UMETA(Hidden)
+};
+
+UENUM()
+enum class EPoseSearchPoseFlags : uint32
+{
+	None = 0,
+
+	// Don't return this pose as a search result
+	BlockTransition = 1 << 0,
+};
+ENUM_CLASS_FLAGS(EPoseSearchPoseFlags);
+
+/**
+ * This is kept for each pose in the search index along side the feature vector values and is used to influence the search.
+ */
+USTRUCT()
+struct POSESEARCH_API FPoseSearchPoseMetadata
+{
+	GENERATED_BODY()
+
+	UPROPERTY(meta = (NeverInHash))
+	EPoseSearchPoseFlags Flags = EPoseSearchPoseFlags::None;
+
+	// @todo: consider float16
+	UPROPERTY(meta = (NeverInHash))
+	float CostAddend = 0.0f;
+
+	// @todo: consider float16
+	UPROPERTY(meta = (NeverInHash))
+	float ContinuingPoseCostAddend = 0.0f;
+
+	// @todo: consider int16
+	UPROPERTY(meta = (NeverInHash))
+	int32 AssetIndex = INDEX_NONE;
+};
+
+/**
+* Information about a source animation asset used by a search index.
+* Some source animation entries may generate multiple FPoseSearchIndexAsset entries.
+**/
+USTRUCT()
+struct POSESEARCH_API FPoseSearchIndexAsset
+{
+	GENERATED_BODY()
+public:
+	FPoseSearchIndexAsset() {}
+
+	FPoseSearchIndexAsset(
+		ESearchIndexAssetType InType,
+		int32 InSourceAssetIdx, 
+		bool bInMirrored, 
+		const FFloatInterval& InSamplingInterval,
+		FVector InBlendParameters = FVector::Zero())
+		: Type(InType)
+		, SourceAssetIdx(InSourceAssetIdx)
+		, bMirrored(bInMirrored)
+		, BlendParameters(InBlendParameters)
+		, SamplingInterval(InSamplingInterval) {}
+
+	// Default to Sequence for now for backward compatibility but
+	// at some point we might want to change this to Invalid.
+	UPROPERTY(meta = (NeverInHash))
+	ESearchIndexAssetType Type = ESearchIndexAssetType::Sequence;
+
+	// Index of the source asset in search index's container (i.e. UPoseSearchDatabase)
+	UPROPERTY(meta = (NeverInHash))
+	int32 SourceAssetIdx = INDEX_NONE;
+
+	UPROPERTY(meta = (NeverInHash))
+	bool bMirrored = false;
+
+	UPROPERTY(meta = (NeverInHash))
+	FVector BlendParameters = FVector::Zero();
+
+	UPROPERTY(meta = (NeverInHash))
+	FFloatInterval SamplingInterval;
+
+	UPROPERTY(meta = (NeverInHash))
+	int32 FirstPoseIdx = INDEX_NONE;
+
+	UPROPERTY(meta = (NeverInHash))
+	int32 NumPoses = 0;
+
+	bool IsPoseInRange(int32 PoseIdx) const { return (PoseIdx >= FirstPoseIdx) && (PoseIdx < FirstPoseIdx + NumPoses); }
+};
+
+
+/**
+* case class for FPoseSearchIndex. building block used to gather data for data mining and calculate weights, pca, kdtree stuff
+*/
+USTRUCT()
+struct POSESEARCH_API FPoseSearchIndexBase
+{
+	GENERATED_BODY()
+
+	UPROPERTY(Category = Info, VisibleAnywhere, meta = (NeverInHash))
+	int32 NumPoses = 0;
+
+	UPROPERTY(meta = (NeverInHash))
+	TArray<float> Values;
+	
+	UPROPERTY(meta = (NeverInHash))
+	TArray<FPoseSearchPoseMetadata> PoseMetadata;
+
+	UPROPERTY(meta = (NeverInHash))
+	EPoseSearchPoseFlags OverallFlags = EPoseSearchPoseFlags::None;
+
+	UPROPERTY(meta = (NeverInHash))
+	TArray<FPoseSearchIndexAsset> Assets;
+
+	// minimum of the database metadata CostAddend: it represents the minimum cost of any search for the associated database (we'll skip the search in case the search result total cost is already less than MinCostAddend)
+	UPROPERTY(Category = Info, VisibleAnywhere, meta = (NeverInHash))
+	float MinCostAddend = -MAX_FLT;
+
+	bool IsValidPoseIndex(int32 PoseIdx) const { return PoseIdx < NumPoses; }
+	bool IsEmpty() const;
+
+	const FPoseSearchIndexAsset& GetAssetForPose(int32 PoseIdx) const;
+	const FPoseSearchIndexAsset* GetAssetForPoseSafe(int32 PoseIdx) const;
+
+	float GetAssetTime(int32 PoseIdx, float SamplingInterval) const;
+
+	void Reset();
+	
+	friend FArchive& operator<<(FArchive& Ar, FPoseSearchIndexBase& Index);
+};
+
+/**
+* A search index for animation poses. The structure of the search index is determined by its UPoseSearchSchema.
+* May represent a single animation (see UPoseSearchSequenceMetaData) or a collection (see UPoseSearchDatabase).
+*/
+USTRUCT()
+struct POSESEARCH_API FPoseSearchIndex : public FPoseSearchIndexBase
+{
+	GENERATED_BODY()
+
+	UPROPERTY(meta = (NeverInHash))
+	TArray<float> PCAValues;
+
+	UPROPERTY(Category = Info, VisibleAnywhere, meta = (NeverInHash))
+	TArray<float> PCAProjectionMatrix;
+
+	UPROPERTY(Category = Info, VisibleAnywhere, meta = (NeverInHash))
+	TArray<float> Mean;
+
+	// we store weights square roots to reduce numerical errors when CompareFeatureVectors 
+	// ((VA - VB) * VW).square().sum()
+	// instead of
+	// ((VA - VB).square() * VW).sum()
+	// since (VA - VB).square() could lead to big numbers, and VW being multiplied by the variance of the dataset
+	UPROPERTY(Category = Info, VisibleAnywhere, meta = (NeverInHash))
+	TArray<float> WeightsSqrt;
+
+	UE::PoseSearch::FKDTree KDTree;
+
+#if WITH_EDITORONLY_DATA
+	UPROPERTY(Category = Info, VisibleAnywhere, meta = (NeverInHash))
+	float PCAExplainedVariance = 0.f;
+
+	UPROPERTY(Category = Info, VisibleAnywhere, meta = (NeverInHash))
+	TArray<float> Deviation;
+#endif // WITH_EDITORONLY_DATA
+
+	FPoseSearchIndex() = default;
+	~FPoseSearchIndex() = default;
+	FPoseSearchIndex(const FPoseSearchIndex& Other); // custom copy constructor to deal with the KDTree DataSrc
+	FPoseSearchIndex(FPoseSearchIndex&& Other) = delete;
+	FPoseSearchIndex& operator=(const FPoseSearchIndex& Other); // custom equal operator to deal with the KDTree DataSrc
+	FPoseSearchIndex& operator=(FPoseSearchIndex&& Other) = delete;
+
+	void Reset();
+	TConstArrayView<float> GetPoseValues(int32 PoseIdx) const;
+	TConstArrayView<float> GetPoseValuesSafe(int32 PoseIdx) const;
+	FPoseSearchCost ComparePoses(int32 PoseIdx, EPoseSearchBooleanRequest QueryMirrorRequest, UE::PoseSearch::EPoseComparisonFlags PoseComparisonFlags, float MirrorMismatchCostBias, TConstArrayView<float> QueryValues) const;
+
+	friend FArchive& operator<<(FArchive& Ar, FPoseSearchIndex& Index);
+};
diff --git a/Engine/Plugins/Experimental/Animation/PoseSearch/Source/Runtime/Public/PoseSearch/PoseSearchLibrary.h b/Engine/Plugins/Experimental/Animation/PoseSearch/Source/Runtime/Public/PoseSearch/PoseSearchLibrary.h
--- a/Engine/Plugins/Experimental/Animation/PoseSearch/Source/Runtime/Public/PoseSearch/PoseSearchLibrary.h
+++ b/Engine/Plugins/Experimental/Animation/PoseSearch/Source/Runtime/Public/PoseSearch/PoseSearchLibrary.h
@@ -2,17 +2,27 @@
 
 #pragma once
 
+#include "AlphaBlend.h"
 #include "AnimNodes/AnimNode_SequenceEvaluator.h"
 #include "Animation/AnimExecutionContext.h"
 #include "Animation/AnimNodeReference.h"
 #include "Kismet/BlueprintFunctionLibrary.h"
 #include "Kismet/BlueprintFunctionLibrary.h"
-#include "PoseSearch/PoseSearch.h"
+#include "PoseSearch/PoseSearchHistory.h"
+#include "PoseSearch/PoseSearchResult.h"
 #include "SequenceEvaluatorLibrary.h"
 #include "SequencePlayerLibrary.h"
-
 #include "PoseSearchLibrary.generated.h"
 
+namespace UE::PoseSearch
+{
+	struct FSearchContext;
+} // namespace UE::PoseSearch
+
+struct FGameplayTagContainer;
+struct FTrajectorySampleRange;
+class UPoseSearchSearchableAsset;
+
 UENUM(BlueprintType, Category="Motion Trajectory", meta=(Bitflags, UseEnumValuesAsMaskValuesInEditor="true"))
 enum class EMotionMatchingFlags : uint8
 {
diff --git a/Engine/Plugins/Experimental/Animation/PoseSearch/Source/Runtime/Public/PoseSearch/PoseSearchResult.h b/Engine/Plugins/Experimental/Animation/PoseSearch/Source/Runtime/Public/PoseSearch/PoseSearchResult.h
new file mode 100644
--- /dev/null
+++ b/Engine/Plugins/Experimental/Animation/PoseSearch/Source/Runtime/Public/PoseSearch/PoseSearchResult.h
@@ -0,0 +1,75 @@
+// Copyright Epic Games, Inc. All Rights Reserved.
+
+#pragma once
+
+#include "PoseSearch/PoseSearchCost.h"
+#include "PoseSearchResult.generated.h"
+
+struct FPoseSearchIndexAsset;
+class UPoseSearchDatabase;
+class UPoseSearchSchema;
+
+/**
+* float buffer of features according to a UPoseSearchSchema layout.
+* FFeatureVectorBuilder is used to build search queries at runtime and for adding samples during search index construction.
+*/
+USTRUCT()
+struct POSESEARCH_API FPoseSearchFeatureVectorBuilder
+{
+	GENERATED_BODY()
+
+public:
+	void Init(const UPoseSearchSchema* Schema);
+	void Reset();
+
+	const UPoseSearchSchema* GetSchema() const { return Schema.Get(); }
+
+	TArray<float>& EditValues() { return Values; }
+	TConstArrayView<float> GetValues() const { return Values; }
+
+private:
+	UPROPERTY(Transient)
+	TWeakObjectPtr<const UPoseSearchSchema> Schema = nullptr;
+	TArray<float> Values;
+};
+
+namespace UE::PoseSearch
+{
+	
+struct FSearchResult
+{
+	// best cost of the currently selected PoseIdx (it could be equal to ContinuingPoseCost)
+	FPoseSearchCost PoseCost;
+	int32 PoseIdx = INDEX_NONE;
+
+	int32 PrevPoseIdx = INDEX_NONE;
+	int32 NextPoseIdx = INDEX_NONE;
+
+	// lerp value to find AssetTime from PrevPoseIdx -> AssetTime -> NextPoseIdx, within range [-0.5, 0.5]
+	float LerpValue = 0.f;
+
+	TWeakObjectPtr<const UPoseSearchDatabase> Database;
+	FPoseSearchFeatureVectorBuilder ComposedQuery;
+
+	// cost of the current pose with the query from database in the result, if possible
+	FPoseSearchCost ContinuingPoseCost; 
+
+	float AssetTime = 0.0f;
+
+#if WITH_EDITORONLY_DATA
+	FPoseSearchCost BruteForcePoseCost;
+#endif // WITH_EDITORONLY_DATA
+
+	// Attempts to set the internal state to match the provided asset time including updating the internal DbPoseIdx. 
+	// If the provided asset time is out of bounds for the currently playing asset then this function will reset the 
+	// state back to the default state.
+	void Update(float NewAssetTime);
+
+	bool IsValid() const;
+
+	void Reset();
+
+	const FPoseSearchIndexAsset* GetSearchIndexAsset(bool bMandatory = false) const;
+};
+
+} // namespace UE::PoseSearch
\ No newline at end of file
diff --git a/Engine/Plugins/Experimental/Animation/PoseSearch/Source/Runtime/Public/PoseSearch/PoseSearchSearchableAsset.h b/Engine/Plugins/Experimental/Animation/PoseSearch/Source/Runtime/Public/PoseSearch/PoseSearchSearchableAsset.h
new file mode 100644
--- /dev/null
+++ b/Engine/Plugins/Experimental/Animation/PoseSearch/Source/Runtime/Public/PoseSearch/PoseSearchSearchableAsset.h
@@ -0,0 +1,22 @@
+// Copyright Epic Games, Inc. All Rights Reserved.
+
+#pragma once
+
+#include "Engine/DataAsset.h"
+#include "PoseSearch/PoseSearchResult.h"
+#include "PoseSearchSearchableAsset.generated.h"
+
+namespace UE::PoseSearch
+{
+	struct FSearchContext;
+} // namespace UE::PoseSearch
+
+/** A data asset for indexing a collection of animation sequences. */
+UCLASS(Abstract, BlueprintType, Experimental)
+class POSESEARCH_API UPoseSearchSearchableAsset : public UDataAsset
+{
+	GENERATED_BODY()
+public:
+
+	virtual UE::PoseSearch::FSearchResult Search(UE::PoseSearch::FSearchContext& SearchContext) const PURE_VIRTUAL(UPoseSearchSearchableAsset::Search, return UE::PoseSearch::FSearchResult(););
+};
\ No newline at end of file
