motion matching - removing more types from FKeyBuilder to generate MM DDC keys (fix for key generation indeterminism)

diff --git a/Engine/Plugins/Experimental/Animation/PoseSearch/Source/Runtime/Private/PoseSearchDerivedDataKey.cpp b/Engine/Plugins/Experimental/Animation/PoseSearch/Source/Runtime/Private/PoseSearchDerivedDataKey.cpp
new file mode 100644
--- /dev/null
+++ b/Engine/Plugins/Experimental/Animation/PoseSearch/Source/Runtime/Private/PoseSearchDerivedDataKey.cpp
@@ -0,0 +1,223 @@
+// Copyright Epic Games, Inc. All Rights Reserved.
+
+#if WITH_EDITOR
+
+#include "PoseSearch/PoseSearchDerivedDataKey.h"
+#include "Animation/AnimBoneCompressionSettings.h"
+#include "Animation/AnimCurveCompressionSettings.h"
+#include "Animation/MirrorDataTable.h"
+#include "Animation/Skeleton.h"
+#include "Engine/SkeletalMesh.h"
+#include "Engine/SkeletalMeshSocket.h"
+#include "Engine/SkinnedAsset.h"
+#include "Engine/StreamableRenderAsset.h"
+#include "Factories/FbxAnimSequenceImportData.h"
+#include "Components/SkeletalMeshComponent.h"
+#include "PoseSearch/PoseSearchDefines.h"
+#include "UObject/DevObjectVersion.h"
+
+namespace UE::PoseSearch
+{
+
+FKeyBuilder::FKeyBuilder()
+{
+	ArIgnoreOuterRef = true;
+
+	// Set FDerivedDataKeyBuilder to be a saving archive instead of a reference collector.
+	// Reference collection causes FSoftObjectPtrs to be serialized by their weak pointer,
+	// which doesn't give a stable hash.  Serializing these to a saving archive will
+	// use a string reference instead, which is a more meaningful hash value.
+	SetIsSaving(true);
+}
+
+FKeyBuilder::FKeyBuilder(const UObject* Object, bool bUseDataVer)
+: FKeyBuilder()
+{
+	check(Object);
+
+	if (bUseDataVer)
+	{
+		// used to invalidate the key without having to change POSESEARCHDB_DERIVEDDATA_VER all the times
+		int32 POSESEARCHDB_DERIVEDDATA_VER_SMALL = 9;
+		FGuid VersionGuid = FDevSystemGuids::GetSystemGuid(FDevSystemGuids::Get().POSESEARCHDB_DERIVEDDATA_VER);
+
+		*this << VersionGuid;
+		*this << POSESEARCHDB_DERIVEDDATA_VER_SMALL;
+	}
+
+	// FKeyBuilder is a saving only archiver, and since it doesn't modify the input Object it's safe to do a const_cast 
+	UObject* NonConstObject = const_cast<UObject*>(Object);
+	*this << NonConstObject;
+}
+
+void FKeyBuilder::Seek(int64 InPos)
+{
+	checkf(InPos == Tell(), TEXT("A hash cannot be computed when serialization relies on seeking."));
+	FArchiveUObject::Seek(InPos);
+}
+
+bool FKeyBuilder::ShouldSkipProperty(const FProperty* InProperty) const
+{
+	if (InProperty == nullptr)
+	{
+		return false;
+	}
+
+	if (Super::ShouldSkipProperty(InProperty))
+	{
+		return true;
+	}
+
+	if (!InProperty->HasAllPropertyFlags(CPF_Edit)) // bIsEditAnywhereProperty
+	{
+		return true;
+	}
+
+	if (InProperty->HasAllPropertyFlags(CPF_Transient))
+	{
+		return true;
+	}
+
+	if (InProperty->HasMetaData(ExcludeFromHashName))
+	{
+		return true;
+	}
+		
+	check(!InProperty->HasMetaData(NeverInHashName));
+
+	#if UE_POSE_SEARCH_DERIVED_DATA_LOGGING
+	UE_LOG(LogPoseSearch, Log, TEXT("%s - %s"), *GetIndentation(), *InProperty->GetFullName());
+	#endif
+
+	return false;
+}
+
+void FKeyBuilder::Serialize(void* Data, int64 Length)
+{
+	const uint8* HasherData = reinterpret_cast<uint8*>(Data);
+
+	#if UE_POSE_SEARCH_DERIVED_DATA_LOGGING_VERBOSE
+	FString RawBytesString = BytesToString(HasherData, Length);
+	UE_LOG(LogPoseSearch, Log, TEXT("%s  > %s"), *GetIndentation(), *RawBytesString);
+	#endif
+
+	Hasher.Update(HasherData, Length);
+}
+
+FArchive& FKeyBuilder::operator<<(FName& Name)
+{
+	// Don't include the name of the object being serialized, since that isn't technically part of the object's state
+	if (!ObjectBeingSerialized || (Name != ObjectBeingSerialized->GetFName()))
+	{
+		// we cannot use GetTypeHash(Name) since it's bound to be non deterministic between editor restarts, so we convert the name into an FString and let the Serialize(void* Data, int64 Length) deal with it
+		FString NameString = Name.ToString();
+		*this << NameString;
+	}
+	return *this;
+}
+
+FArchive& FKeyBuilder::operator<<(class UObject*& Object)
+{
+	if (Object)
+	{
+		#if UE_POSE_SEARCH_DERIVED_DATA_LOGGING
+		++Indentation;
+		#endif
+
+		bool bAlreadyProcessed = false;
+		ObjectsAlreadySerialized.Add(Object, &bAlreadyProcessed);
+
+		// If we haven't already serialized this object
+		if (bAlreadyProcessed)
+		{
+			#if UE_POSE_SEARCH_DERIVED_DATA_LOGGING
+			UE_LOG(LogPoseSearch, Log, TEXT("%sAlreadyProcessed '%s' (%s)"), *GetIndentation(), *Object->GetName(), *Object->GetClass()->GetName());
+			#endif
+		}
+		// for specific types we only add their names to the hash
+		else if (AddNameOnly(Object))
+		{
+			#if UE_POSE_SEARCH_DERIVED_DATA_LOGGING
+			UE_LOG(LogPoseSearch, Log, TEXT("%sAddingNameOnly '%s' (%s)"), *GetIndentation(), *Object->GetName(), *Object->GetClass()->GetName());
+			#endif
+
+			FString ObjectName = GetFullNameSafe(Object);
+			*this << ObjectName;
+		}
+		else
+		{
+			// Serialize it
+			ObjectBeingSerialized = Object;
+
+			#if UE_POSE_SEARCH_DERIVED_DATA_LOGGING
+			UE_LOG(LogPoseSearch, Log, TEXT("%sBegin '%s' (%s)"), *GetIndentation(), *Object->GetName(), *Object->GetClass()->GetName());
+			#endif
+
+			const_cast<UObject*>(Object)->Serialize(*this);
+
+			#if UE_POSE_SEARCH_DERIVED_DATA_LOGGING
+			UE_LOG(LogPoseSearch, Log, TEXT("%sEnd '%s' (%s)"), *GetIndentation(), *Object->GetName(), *Object->GetClass()->GetName());
+			#endif
+
+			ObjectBeingSerialized = nullptr;
+		}
+			
+		#if UE_POSE_SEARCH_DERIVED_DATA_LOGGING
+		--Indentation;
+		#endif
+	}
+
+	return *this;
+}
+
+FString FKeyBuilder::GetArchiveName() const
+{
+	return TEXT("FDerivedDataKeyBuilder");
+}
+	
+FIoHash FKeyBuilder::Finalize() const
+{
+	// Stores a BLAKE3-160 hash, taken from the first 20 bytes of a BLAKE3-256 hash
+	return FIoHash(Hasher.Finalize());
+}
+
+const TSet<const UObject*>& FKeyBuilder::GetDependencies() const
+{
+	return ObjectsAlreadySerialized;
+}
+
+bool FKeyBuilder::AddNameOnly(class UObject* Object) const
+{
+	return
+		Cast<UAnimBoneCompressionSettings>(Object) ||
+		Cast<UAnimCurveCompressionSettings>(Object) ||
+		Cast<UFbxAnimSequenceImportData>(Object) ||
+		Cast<UMirrorDataTable>(Object) ||
+		Cast<USkeleton>(Object) ||
+		Cast<USkeletalMeshComponent>(Object) ||
+		Cast<USkeletalMeshSocket>(Object) ||
+		Cast<USkinnedMeshComponent>(Object) ||
+		Cast<UMeshComponent>(Object) ||
+		Cast<UPrimitiveComponent>(Object) ||
+		Cast<USceneComponent>(Object) ||
+		Cast<UActorComponent>(Object) ||
+		Cast<UStreamableRenderAsset>(Object) ||
+		Cast<USkinnedAsset>(Object) ||
+		Cast<USkeletalMesh>(Object);
+}
+
+#if UE_POSE_SEARCH_DERIVED_DATA_LOGGING
+FString FKeyBuilder::GetIndentation() const
+{
+	FString IndentationString;
+	for (int32 i = 0; i < Indentation; ++i)
+	{
+		IndentationString.Append(" ");
+	}
+	return IndentationString;
+}
+#endif
+
+} // namespace UE::PoseSearch
+
+#endif // WITH_EDITOR
\ No newline at end of file
diff --git a/Engine/Plugins/Experimental/Animation/PoseSearch/Source/Runtime/Public/PoseSearch/PoseSearchDerivedDataKey.h b/Engine/Plugins/Experimental/Animation/PoseSearch/Source/Runtime/Public/PoseSearch/PoseSearchDerivedDataKey.h
--- a/Engine/Plugins/Experimental/Animation/PoseSearch/Source/Runtime/Public/PoseSearch/PoseSearchDerivedDataKey.h
+++ b/Engine/Plugins/Experimental/Animation/PoseSearch/Source/Runtime/Public/PoseSearch/PoseSearchDerivedDataKey.h
@@ -4,20 +4,14 @@
 
 #if WITH_EDITOR
 
-#include "Animation/AnimBoneCompressionSettings.h"
-#include "Animation/AnimCurveCompressionSettings.h"
-#include "Animation/MirrorDataTable.h"
-#include "Containers/Set.h"
-#include "Containers/UnrealString.h"
-#include "Factories/FbxAnimSequenceImportData.h"
 #include "Hash/Blake3.h"
-#include "UObject/DevObjectVersion.h"
-#include "UObject/Object.h"
-#include "UObject/UnrealType.h"
+#include "IO/IoHash.h"
 #include "Serialization/ArchiveUObject.h"
+#include "UObject/UnrealType.h"
 
 namespace UE::PoseSearch
 {
+
 // log properties and UObjects names
 #ifndef UE_POSE_SEARCH_DERIVED_DATA_LOGGING
 	#define UE_POSE_SEARCH_DERIVED_DATA_LOGGING 0
@@ -28,8 +22,6 @@ namespace UE::PoseSearch
 	#define UE_POSE_SEARCH_DERIVED_DATA_LOGGING_VERBOSE 0
 #endif
 
-//////////////////////////////////////////////////////////////////////////
-// FKeyBuilder
 class POSESEARCH_API FKeyBuilder : public FArchiveUObject
 {
 public:
@@ -40,206 +32,30 @@ public:
 	inline static const FName ExcludeFromHashName = FName(TEXT("ExcludeFromHash"));
 	inline static const FName NeverInHashName = FName(TEXT("NeverInHash"));
 
-	FKeyBuilder()
-	{
-		ArIgnoreOuterRef = true;
-
-		// Set FDerivedDataKeyBuilder to be a saving archive instead of a reference collector.
-		// Reference collection causes FSoftObjectPtrs to be serialized by their weak pointer,
-		// which doesn't give a stable hash.  Serializing these to a saving archive will
-		// use a string reference instead, which is a more meaningful hash value.
-		SetIsSaving(true);
-	}
-
-	FKeyBuilder(const UObject* Object, bool bUseDataVer = false)
-		: FKeyBuilder()
-	{
-		check(Object);
-
-		if (bUseDataVer)
-		{
-			// used to invalidate the key without having to change POSESEARCHDB_DERIVEDDATA_VER all the times
-			int32 POSESEARCHDB_DERIVEDDATA_VER_SMALL = 9;
-			FGuid VersionGuid = FDevSystemGuids::GetSystemGuid(FDevSystemGuids::Get().POSESEARCHDB_DERIVEDDATA_VER);
-
-			*this << VersionGuid;
-			*this << POSESEARCHDB_DERIVEDDATA_VER_SMALL;
-		}
-
-		// FKeyBuilder is a saving only archiver, and since it doesn't modify the input Object it's safe to do a const_cast 
-		UObject* NonConstObject = const_cast<UObject*>(Object);
-		*this << NonConstObject;
-	}
-
-	virtual void Seek(int64 InPos) override
-	{
-		checkf(InPos == Tell(), TEXT("A hash cannot be computed when serialization relies on seeking."));
-		FArchiveUObject::Seek(InPos);
-	}
-
-	//~ Begin FArchive Interface
-	virtual bool ShouldSkipProperty(const FProperty* InProperty) const override
-	{
-		if (InProperty == nullptr)
-		{
-			return false;
-		}
-
-		if (Super::ShouldSkipProperty(InProperty))
-		{
-			return true;
-		}
-
-		if (!InProperty->HasAllPropertyFlags(CPF_Edit)) // bIsEditAnywhereProperty
-		{
-			return true;
-		}
-
-		if (InProperty->HasAllPropertyFlags(CPF_Transient))
-		{
-			return true;
-		}
-
-		if (InProperty->HasMetaData(ExcludeFromHashName))
-		{
-			return true;
-		}
-		
-		check(!InProperty->HasMetaData(NeverInHashName));
-
-		#if UE_POSE_SEARCH_DERIVED_DATA_LOGGING
-		UE_LOG(LogPoseSearch, Log, TEXT("%s - %s"), *GetIndentation(), *InProperty->GetFullName());
-		#endif
-
-		return false;
-	}
-
-	virtual void Serialize(void* Data, int64 Length) override
-	{
-		const uint8* HasherData = reinterpret_cast<uint8*>(Data);
-
-		#if UE_POSE_SEARCH_DERIVED_DATA_LOGGING_VERBOSE
-		FString RawBytesString = BytesToString(HasherData, Length);
-		UE_LOG(LogPoseSearch, Log, TEXT("%s  > %s"), *GetIndentation(), *RawBytesString);
-		#endif
-
-		Hasher.Update(HasherData, Length);
-	}
-
-	// @todo: do we need those Serialize(s)?
-	//virtual void SerializeBits(void* Bits, int64 LengthBits) override
-	//virtual void SerializeInt(uint32& Value, uint32 Max) override
-	//virtual void SerializeIntPacked(uint32& Value) override
+	FKeyBuilder();
+	FKeyBuilder(const UObject* Object, bool bUseDataVer = false);
 
 	using Super::IsSaving;
-	using Super::operator<<; // For visibility of the overloads we don't override
-
-	virtual FArchive& operator<<(FName& Name) override
-	{
-		// Don't include the name of the object being serialized, since that isn't technically part of the object's state
-		if (!ObjectBeingSerialized || (Name != ObjectBeingSerialized->GetFName()))
-		{
-			// we cannot use GetTypeHash(Name) since it's bound to be non deterministic between editor restarts, so we convert the name into an FString and let the Serialize(void* Data, int64 Length) deal with it
-			FString NameString = Name.ToString();
-			*this << NameString;
-		}
-		return *this;
-	}
-
-	virtual FArchive& operator<<(class UObject*& Object) override
-	{
-		if (Object)
-		{
-			#if UE_POSE_SEARCH_DERIVED_DATA_LOGGING
-			++Indentation;
-			#endif
-
-			bool bAlreadyProcessed = false;
-			ObjectsAlreadySerialized.Add(Object, &bAlreadyProcessed);
-
-			// If we haven't already serialized this object
-			if (bAlreadyProcessed)
-			{
-				#if UE_POSE_SEARCH_DERIVED_DATA_LOGGING
-				UE_LOG(LogPoseSearch, Log, TEXT("%sAlreadyProcessed '%s' (%s)"), *GetIndentation(), *Object->GetName(), *Object->GetClass()->GetName());
-				#endif
-			}
-			// for specific types we only add their names to the hash
-			else if (AddNameOnly(Object))
-			{
-				#if UE_POSE_SEARCH_DERIVED_DATA_LOGGING
-				UE_LOG(LogPoseSearch, Log, TEXT("%sAddingNameOnly '%s' (%s)"), *GetIndentation(), *Object->GetName(), *Object->GetClass()->GetName());
-				#endif
-
-				FString ObjectName = GetFullNameSafe(Object);
-				*this << ObjectName;
-			}
-			else
-			{
-				// Serialize it
-				ObjectBeingSerialized = Object;
+	using Super::operator<<;
 
-				#if UE_POSE_SEARCH_DERIVED_DATA_LOGGING
-				UE_LOG(LogPoseSearch, Log, TEXT("%sBegin '%s' (%s)"), *GetIndentation(), *Object->GetName(), *Object->GetClass()->GetName());
-				#endif
-
-				const_cast<UObject*>(Object)->Serialize(*this);
-
-				#if UE_POSE_SEARCH_DERIVED_DATA_LOGGING
-				UE_LOG(LogPoseSearch, Log, TEXT("%sEnd '%s' (%s)"), *GetIndentation(), *Object->GetName(), *Object->GetClass()->GetName());
-				#endif
-
-				ObjectBeingSerialized = nullptr;
-			}
-			
-			#if UE_POSE_SEARCH_DERIVED_DATA_LOGGING
-			--Indentation;
-			#endif
-		}
-
-		return *this;
-	}
-
-	virtual FString GetArchiveName() const override
-	{
-		return TEXT("FDerivedDataKeyBuilder");
-	}
+	//~ Begin FArchive Interface
+	virtual void Seek(int64 InPos) override;
+	virtual bool ShouldSkipProperty(const FProperty* InProperty) const override;
+	virtual void Serialize(void* Data, int64 Length) override;
+	virtual FArchive& operator<<(FName& Name) override;
+	virtual FArchive& operator<<(class UObject*& Object) override;
+	virtual FString GetArchiveName() const override;
 	//~ End FArchive Interface
 	
-	FIoHash Finalize() const
-	{
-		// Stores a BLAKE3-160 hash, taken from the first 20 bytes of a BLAKE3-256 hash
-		return FIoHash(Hasher.Finalize());
-	}
-
-	const TSet<const UObject*>& GetDependencies() const
-	{
-		return ObjectsAlreadySerialized;
-	}
+	FIoHash Finalize() const;
+	const TSet<const UObject*>& GetDependencies() const;
 
 protected:
-	/** to keep the key generation lightweight, we hash only the full names for these types */
-	bool AddNameOnly(class UObject* Object) const
-	{
-		return
-			Cast<UAnimBoneCompressionSettings>(Object) ||
-			Cast<UAnimCurveCompressionSettings>(Object) ||
-			Cast<UFbxAnimSequenceImportData>(Object) ||
-			Cast<UMirrorDataTable>(Object) ||
-			Cast<USkeleton>(Object);
-	}
+	// to keep the key generation lightweight, we hash only the full names for these types
+	bool AddNameOnly(class UObject* Object) const;
 
 #if UE_POSE_SEARCH_DERIVED_DATA_LOGGING
-	FString GetIndentation() const
-	{
-		FString IndentationString;
-		for (int32 i = 0; i < Indentation; ++i)
-		{
-			IndentationString.Append(" ");
-		}
-		return IndentationString;
-	}
-
+	FString GetIndentation() const;
 	int32 Indentation = 0;
 #endif
 
