[PoseSearch]
PoseSearchTraceAnalyzer - Fix non unity compile error

diff --git a/Engine/Plugins/Experimental/Animation/PoseSearch/Source/Editor/Private/Trace/PoseSearchTraceAnalyzer.cpp b/Engine/Plugins/Experimental/Animation/PoseSearch/Source/Editor/Private/Trace/PoseSearchTraceAnalyzer.cpp
--- a/Engine/Plugins/Experimental/Animation/PoseSearch/Source/Editor/Private/Trace/PoseSearchTraceAnalyzer.cpp
+++ b/Engine/Plugins/Experimental/Animation/PoseSearch/Source/Editor/Private/Trace/PoseSearchTraceAnalyzer.cpp
@@ -4,6 +4,7 @@
 
 #include "HAL/LowLevelMemTracker.h"
 #include "PoseSearchTraceProvider.h"
+#include "Serialization/MemoryReader.h"
 #include "Trace/PoseSearchTraceLogger.h"
 #include "TraceServices/Model/AnalysisSession.h"
 
