motion matching - fix for FAsyncPoseSearchDatabasesManagement::Tick when migrating assets

diff --git a/Engine/Plugins/Experimental/Animation/PoseSearch/Source/Runtime/Private/PoseSearchDerivedData.cpp b/Engine/Plugins/Experimental/Animation/PoseSearch/Source/Runtime/Private/PoseSearchDerivedData.cpp
--- a/Engine/Plugins/Experimental/Animation/PoseSearch/Source/Runtime/Private/PoseSearchDerivedData.cpp
+++ b/Engine/Plugins/Experimental/Animation/PoseSearch/Source/Runtime/Private/PoseSearchDerivedData.cpp
@@ -538,6 +538,7 @@ struct FPoseSearchDatabaseAsyncCacheTask
 	void Cancel(FCriticalSection& OuterMutex);
 	bool Poll(FCriticalSection& OuterMutex) const;
 	bool ContainsDatabase(const UPoseSearchDatabase* OtherDatabase, FCriticalSection& OuterMutex) const;
+	bool IsValid(FCriticalSection& OuterMutex) const;
 
 	~FPoseSearchDatabaseAsyncCacheTask();
 	EState GetState() const { return EState(ThreadSafeState.GetValue()); }
@@ -701,6 +702,12 @@ bool FPoseSearchDatabaseAsyncCacheTask::ContainsDatabase(const UPoseSearchDataba
 	return Database.Get() == OtherDatabase;
 }
 
+bool FPoseSearchDatabaseAsyncCacheTask::IsValid(FCriticalSection& OuterMutex) const
+{
+	FScopeLock Lock(&OuterMutex);
+	return Database.IsValid();
+}
+
 // called once the task is done:
 // if EStatus::Ok (data has been retrieved from DDC) we deserialize the payload into the local SearchIndex
 // if EStatus::Error we BuildIndex and if that's successful we 'Put' it on DDC
@@ -956,7 +963,14 @@ void FAsyncPoseSearchDatabasesManagement::Tick(float DeltaTime)
 	// iterating backwards because of the possible RemoveAtSwap 
 	for (int32 TaskIndex = Tasks.Num() - 1; TaskIndex >= 0; --TaskIndex)
 	{
-		Tasks[TaskIndex]->Update(Mutex);
+		if (!Tasks[TaskIndex]->IsValid(Mutex))
+		{
+			Tasks.RemoveAtSwap(TaskIndex, 1, false);
+		}
+		else
+		{
+			Tasks[TaskIndex]->Update(Mutex);
+		}
 			
 		// @todo: check key validity every few ticks, or perhaps delete unused for a long time Tasks
 	}
