motion matching - fix for FDatabaseViewModel::UpdatePreviewActors blend space preview

diff --git a/Engine/Plugins/Experimental/Animation/PoseSearch/Source/Editor/Private/PoseSearchDatabaseViewModel.cpp b/Engine/Plugins/Experimental/Animation/PoseSearch/Source/Editor/Private/PoseSearchDatabaseViewModel.cpp
--- a/Engine/Plugins/Experimental/Animation/PoseSearch/Source/Editor/Private/PoseSearchDatabaseViewModel.cpp
+++ b/Engine/Plugins/Experimental/Animation/PoseSearch/Source/Editor/Private/PoseSearchDatabaseViewModel.cpp
@@ -373,13 +373,13 @@ namespace UE::PoseSearch
 		TSharedPtr<FDatabasePreviewScene> PreviewScene = PreviewScenePtr.Pin();
 		for (FDatabasePreviewActor& PreviewActor : GetPreviewActors())
 		{
-			if (!PreviewActor.GetAnimPreviewInstance() ||
-				PreviewActor.IndexAssetIndex >= SearchIndex.Assets.Num())
+			UAnimPreviewInstance* AnimInstance = PreviewActor.GetAnimPreviewInstance();
+			if (!AnimInstance || PreviewActor.IndexAssetIndex >= SearchIndex.Assets.Num())
 			{
 				continue;
 			}
 
-			const UAnimationAsset* PreviewAsset = PreviewActor.GetAnimPreviewInstance()->GetAnimationAsset();
+			const UAnimationAsset* PreviewAsset = AnimInstance->GetAnimationAsset();
 			const IAssetSampler* Sampler = PreviewActor.GetSampler();
 			if (!PreviewAsset || !Sampler)
 			{
@@ -388,20 +388,21 @@ namespace UE::PoseSearch
 
 			float CurrentTime = 0.f;
 			FAnimationRuntime::AdvanceTime(false, PlayTime, CurrentTime, PreviewAsset->GetPlayLength());
+			const float CurrentScaledTime = PreviewActor.GetScaledTime(CurrentTime);
+
+			const FPoseSearchIndexAsset& IndexAsset = SearchIndex.Assets[PreviewActor.IndexAssetIndex];
+			AnimInstance->SetPosition(CurrentScaledTime);
+			AnimInstance->SetPlayRate(0.f);
+			AnimInstance->SetBlendSpacePosition(IndexAsset.BlendParameters);
 
 			FTransform RootMotion = Sampler->ExtractRootTransform(CurrentTime);
-			if (PreviewActor.GetAnimPreviewInstance()->GetMirrorDataTable())
+			if (AnimInstance->GetMirrorDataTable())
 			{
-				RootMotion = MirrorRootMotion(RootMotion, PreviewActor.GetAnimPreviewInstance()->GetMirrorDataTable());
+				RootMotion = MirrorRootMotion(RootMotion, AnimInstance->GetMirrorDataTable());
 			}
 			PreviewActor.Actor->SetActorTransform(RootMotion);
 
-			const float CurrentScaledTime = PreviewActor.GetScaledTime(CurrentTime);
-			PreviewActor.GetAnimPreviewInstance()->SetPosition(CurrentScaledTime);
-			PreviewActor.GetAnimPreviewInstance()->SetPlayRate(0.f);
-
-			const FPoseSearchIndexAsset& SearchIndexAsset = SearchIndex.Assets[PreviewActor.IndexAssetIndex];
-			PreviewActor.CurrentPoseIndex = PoseSearchDatabase->GetPoseIndexFromTime(CurrentTime, SearchIndexAsset);
+			PreviewActor.CurrentPoseIndex = PoseSearchDatabase->GetPoseIndexFromTime(CurrentTime, IndexAsset);
 		}
 	}
 
