motion matching - exposing pose channel character relative velocities, fix for position channel time offset, improved fault tolerance for MM DDC, defaulting channel invalid bones to use root transform (WIP consolidation of bones and trajectory transform calculation)

diff --git a/Engine/Plugins/Experimental/Animation/PoseSearch/Source/Runtime/Private/PoseSearchAssetIndexer.cpp b/Engine/Plugins/Experimental/Animation/PoseSearch/Source/Runtime/Private/PoseSearchAssetIndexer.cpp
--- a/Engine/Plugins/Experimental/Animation/PoseSearch/Source/Runtime/Private/PoseSearchAssetIndexer.cpp
+++ b/Engine/Plugins/Experimental/Animation/PoseSearch/Source/Runtime/Private/PoseSearchAssetIndexer.cpp
@@ -259,6 +259,8 @@ FPoseSearchPoseMetadata FAssetIndexer::GetMetadata(int32 SampleIdx) const
 	return Metadata;
 }
 
+// returns the transform in character space for the bone indexed by Schema->BoneReferences[SchemaBoneIdx] 
+// at SampleTime-OriginTime seconds ahead (or behind) the pose at OriginTime
 FTransform FAssetIndexer::GetTransformAndCacheResults(float SampleTime, float OriginTime, int8 SchemaBoneIdx, bool& Clamped)
 {
 	// @todo: use an hashmap if we end up having too many entries
@@ -295,12 +297,28 @@ FTransform FAssetIndexer::GetTransformAndCacheResults(float SampleTime, float Or
 
 		float CurrentTime = Sample.ClipTime;
 		float PreviousTime = CurrentTime - SamplingContext->FiniteDelta;
-
-		check(Sample.Clip->IsLoopable() || PreviousTime <= Sample.Clip->GetPlayLength());
+		
+		const bool bLoopable = Sample.Clip->IsLoopable();
+		const float PlayLength = Sample.Clip->GetPlayLength();
+		if (!bLoopable)
+		{
+			// if not loopable we clamp the pose at time zero or PlayLength
+			if (PreviousTime < 0.f)
+			{
+				PreviousTime = 0.f;
+				CurrentTime = FMath::Min(SamplingContext->FiniteDelta, PlayLength);
+			}
+			else if (CurrentTime > PlayLength)
+			{
+				CurrentTime = PlayLength;
+				PreviousTime = FMath::Max(PlayLength - SamplingContext->FiniteDelta, 0.f);
+			}
+		}
 
 		FDeltaTimeRecord DeltaTimeRecord;
 		DeltaTimeRecord.Set(PreviousTime, CurrentTime - PreviousTime);
-		FAnimExtractContext ExtractionCtx(static_cast<double>(CurrentTime), true, DeltaTimeRecord, Sample.Clip->IsLoopable());
+		// no need to extract root motion here, since we use the precalculated Sample.RootTransform as root transform for the Entry
+		FAnimExtractContext ExtractionCtx(static_cast<double>(CurrentTime), false, DeltaTimeRecord, bLoopable);
 
 		Sample.Clip->ExtractPose(ExtractionCtx, Entry->AnimPoseData);
 
@@ -320,10 +338,18 @@ FTransform FAssetIndexer::GetTransformAndCacheResults(float SampleTime, float Or
 		Entry->Clamped = Sample.bClamped;
 	}
 
+	FTransform BoneTransform;
 	const FBoneReference& BoneReference = IndexingContext.Schema->BoneReferences[SchemaBoneIdx];
-	FCompactPoseBoneIndex CompactBoneIndex = BoneContainer.MakeCompactPoseIndex(FMeshPoseBoneIndex(BoneReference.BoneIndex));
+	if (BoneReference.HasValidSetup())
+	{
+		FCompactPoseBoneIndex CompactBoneIndex = BoneContainer.MakeCompactPoseIndex(FMeshPoseBoneIndex(BoneReference.BoneIndex));
+		BoneTransform = Entry->ComponentSpacePose.GetComponentSpaceTransform(CompactBoneIndex) * MirrorTransform(Entry->RootTransform);
+	}
+	else
+	{
+		BoneTransform = MirrorTransform(Entry->RootTransform);
+	}
 
-	const FTransform BoneTransform = Entry->ComponentSpacePose.GetComponentSpaceTransform(CompactBoneIndex) * MirrorTransform(Entry->RootTransform);
 	Clamped = Entry->Clamped;
 
 	return BoneTransform;
diff --git a/Engine/Plugins/Experimental/Animation/PoseSearch/Source/Runtime/Private/PoseSearchContext.cpp b/Engine/Plugins/Experimental/Animation/PoseSearch/Source/Runtime/Private/PoseSearchContext.cpp
--- a/Engine/Plugins/Experimental/Animation/PoseSearch/Source/Runtime/Private/PoseSearchContext.cpp
+++ b/Engine/Plugins/Experimental/Animation/PoseSearch/Source/Runtime/Private/PoseSearchContext.cpp
@@ -111,7 +111,7 @@ FTransform FSearchContext::TryGetTransformAndCacheResults(float SampleTime, cons
 	check(History && Schema);
 
 	static constexpr FBoneIndexType RootBoneIdx = 0xFFFF;
-	const FBoneIndexType BoneIndexType = SchemaBoneIdx >= 0 ? Schema->BoneIndices[SchemaBoneIdx] : RootBoneIdx;
+	const FBoneIndexType BoneIndexType = SchemaBoneIdx >= 0 && Schema->BoneReferences[SchemaBoneIdx].HasValidSetup() ? Schema->BoneReferences[SchemaBoneIdx].BoneIndex : RootBoneIdx;
 
 	// @todo: use an hashmap if we end up having too many entries
 	const FCachedEntry* Entry = CachedEntries.FindByPredicate([SampleTime, BoneIndexType](const FSearchContext::FCachedEntry& Entry)
diff --git a/Engine/Plugins/Experimental/Animation/PoseSearch/Source/Runtime/Private/PoseSearchDerivedDataKey.cpp b/Engine/Plugins/Experimental/Animation/PoseSearch/Source/Runtime/Private/PoseSearchDerivedDataKey.cpp
--- a/Engine/Plugins/Experimental/Animation/PoseSearch/Source/Runtime/Private/PoseSearchDerivedDataKey.cpp
+++ b/Engine/Plugins/Experimental/Animation/PoseSearch/Source/Runtime/Private/PoseSearchDerivedDataKey.cpp
@@ -38,7 +38,7 @@ FKeyBuilder::FKeyBuilder(const UObject* Object, bool bUseDataVer)
 	if (bUseDataVer)
 	{
 		// used to invalidate the key without having to change POSESEARCHDB_DERIVEDDATA_VER all the times
-		int32 POSESEARCHDB_DERIVEDDATA_VER_SMALL = 9;
+		int32 POSESEARCHDB_DERIVEDDATA_VER_SMALL = 17;
 		FGuid VersionGuid = FDevSystemGuids::GetSystemGuid(FDevSystemGuids::Get().POSESEARCHDB_DERIVEDDATA_VER);
 
 		*this << VersionGuid;
diff --git a/Engine/Plugins/Experimental/Animation/PoseSearch/Source/Runtime/Private/PoseSearchFeatureChannel_Heading.cpp b/Engine/Plugins/Experimental/Animation/PoseSearch/Source/Runtime/Private/PoseSearchFeatureChannel_Heading.cpp
--- a/Engine/Plugins/Experimental/Animation/PoseSearch/Source/Runtime/Private/PoseSearchFeatureChannel_Heading.cpp
+++ b/Engine/Plugins/Experimental/Animation/PoseSearch/Source/Runtime/Private/PoseSearchFeatureChannel_Heading.cpp
@@ -56,7 +56,7 @@ void UPoseSearchFeatureChannel_Heading::IndexAsset(UE::PoseSearch::IAssetIndexer
 		const float SubsampleTime = OriginSampleTime + SampleTimeOffset;
 
 		bool ClampedPresent;
-		const FTransform BoneTransformsPresent = Indexer.GetTransformAndCacheResults(SubsampleTime, bUseSampleTimeOffsetRootBone ? SubsampleTime : OriginSampleTime, SchemaBoneIdx, ClampedPresent);
+		const FTransform BoneTransformsPresent = Indexer.GetTransformAndCacheResults(SubsampleTime, OriginSampleTime, SchemaBoneIdx, ClampedPresent);
 		int32 DataOffset = ChannelDataOffset;
 		FFeatureVectorHelper::EncodeVector(IndexingContext.GetPoseVector(VectorIdx, FeatureVectorTable), DataOffset, GetAxis(BoneTransformsPresent.GetRotation()));
 		check(DataOffset == ChannelDataOffset + ChannelCardinality);
@@ -83,13 +83,9 @@ void UPoseSearchFeatureChannel_Heading::BuildQuery(UE::PoseSearch::FSearchContex
 	else
 	{
 		FTransform Transform = SearchContext.TryGetTransformAndCacheResults(SampleTimeOffset, InOutQuery.GetSchema(), SchemaBoneIdx);
-
-		if (!bUseSampleTimeOffsetRootBone)
-		{
-			const FTransform RootTransform = SearchContext.TryGetTransformAndCacheResults(0.f, InOutQuery.GetSchema(), FSearchContext::SchemaRootBoneIdx);
-			const FTransform RootTransformPrev = SearchContext.TryGetTransformAndCacheResults(SampleTimeOffset, InOutQuery.GetSchema(), FSearchContext::SchemaRootBoneIdx);
-			Transform = Transform * (RootTransformPrev * RootTransform.Inverse());
-		}
+		const FTransform RootTransform = SearchContext.TryGetTransformAndCacheResults(0.f, InOutQuery.GetSchema(), FSearchContext::SchemaRootBoneIdx);
+		const FTransform RootTransformPrev = SearchContext.TryGetTransformAndCacheResults(SampleTimeOffset, InOutQuery.GetSchema(), FSearchContext::SchemaRootBoneIdx);
+		Transform = Transform * (RootTransformPrev * RootTransform.Inverse());
 
 		int32 DataOffset = ChannelDataOffset;
 		FFeatureVectorHelper::EncodeVector(InOutQuery.EditValues(), DataOffset, GetAxis(Transform.GetRotation()));
diff --git a/Engine/Plugins/Experimental/Animation/PoseSearch/Source/Runtime/Private/PoseSearchFeatureChannel_Heading.h b/Engine/Plugins/Experimental/Animation/PoseSearch/Source/Runtime/Private/PoseSearchFeatureChannel_Heading.h
--- a/Engine/Plugins/Experimental/Animation/PoseSearch/Source/Runtime/Private/PoseSearchFeatureChannel_Heading.h
+++ b/Engine/Plugins/Experimental/Animation/PoseSearch/Source/Runtime/Private/PoseSearchFeatureChannel_Heading.h
@@ -44,11 +44,6 @@ public:
 	UPROPERTY(EditAnywhere, Category = "Settings")
 	EInputQueryPose InputQueryPose = EInputQueryPose::UseContinuingPose;
 
-	// if UseSampleTimeOffsetRootBone is true, this UPoseSearchFeatureChannel_Position will calculate the position of Bone from the pose SampleTimeOffset seconds away from the current time pose root bone
-	// if false the calculated position of Bone will be in component space from the pose SampleTimeOffset seconds away
-	UPROPERTY(EditAnywhere, Category = "Settings")
-	bool bUseSampleTimeOffsetRootBone = true;
-
 	// UPoseSearchFeatureChannel interface
 	virtual void InitializeSchema(UPoseSearchSchema* Schema) override;
 	virtual void FillWeights(TArray<float>& Weights) const override;
diff --git a/Engine/Plugins/Experimental/Animation/PoseSearch/Source/Runtime/Private/PoseSearchFeatureChannel_Pose.cpp b/Engine/Plugins/Experimental/Animation/PoseSearch/Source/Runtime/Private/PoseSearchFeatureChannel_Pose.cpp
--- a/Engine/Plugins/Experimental/Animation/PoseSearch/Source/Runtime/Private/PoseSearchFeatureChannel_Pose.cpp
+++ b/Engine/Plugins/Experimental/Animation/PoseSearch/Source/Runtime/Private/PoseSearchFeatureChannel_Pose.cpp
@@ -15,8 +15,6 @@
 namespace UE::PoseSearch
 {
 
-constexpr bool UseCharacterSpaceVelocities = true;
-
 struct LocalMinMax
 {
 	enum
@@ -480,9 +478,9 @@ void UPoseSearchFeatureChannel_Pose::AddPoseFeatures(UE::PoseSearch::IAssetIndex
 		if (EnumHasAnyFlags(SampledBone.Flags, EPoseSearchBoneFlags::Velocity))
 		{
 			bool ClampedPast, ClampedPresent, ClampedFuture;
-			const FTransform BoneTransformsPast = Indexer.GetTransformAndCacheResults(SampleTime - SamplingContext->FiniteDelta, UseCharacterSpaceVelocities ? SampleTime - SamplingContext->FiniteDelta : SampleTime, SchemaBoneIdx[ChannelBoneIdx], ClampedPast);
+			const FTransform BoneTransformsPast = Indexer.GetTransformAndCacheResults(SampleTime - SamplingContext->FiniteDelta, bUseCharacterSpaceVelocities ? SampleTime - SamplingContext->FiniteDelta : SampleTime, SchemaBoneIdx[ChannelBoneIdx], ClampedPast);
 			const FTransform BoneTransformsPresent = Indexer.GetTransformAndCacheResults(SampleTime, SampleTime, SchemaBoneIdx[ChannelBoneIdx], ClampedPresent);
-			const FTransform BoneTransformsFuture = Indexer.GetTransformAndCacheResults(SampleTime + SamplingContext->FiniteDelta, UseCharacterSpaceVelocities ? SampleTime + SamplingContext->FiniteDelta : SampleTime, SchemaBoneIdx[ChannelBoneIdx], ClampedFuture);
+			const FTransform BoneTransformsFuture = Indexer.GetTransformAndCacheResults(SampleTime + SamplingContext->FiniteDelta, bUseCharacterSpaceVelocities ? SampleTime + SamplingContext->FiniteDelta : SampleTime, SchemaBoneIdx[ChannelBoneIdx], ClampedFuture);
 
 			// We can get a better finite difference if we ignore samples that have
 			// been clamped at either side of the clip. However, if the central sample 
@@ -578,7 +576,7 @@ void UPoseSearchFeatureChannel_Pose::BuildQuery(UE::PoseSearch::FSearchContext&
 
 				CachedTransforms[SampledBoneIdx].Previous = SearchContext.TryGetTransformAndCacheResults(SampleTime - HistorySameplInterval, InOutQuery.GetSchema(), SchemaBoneIdx[SampledBoneIdx]);
 
-				if (!UE::PoseSearch::UseCharacterSpaceVelocities)
+				if (!bUseCharacterSpaceVelocities)
 				{
 					const FTransform RootTransform = SearchContext.TryGetTransformAndCacheResults(SampleTime, InOutQuery.GetSchema(), FSearchContext::SchemaRootBoneIdx);
 					const FTransform RootTransformPrev = SearchContext.TryGetTransformAndCacheResults(SampleTime - HistorySameplInterval, InOutQuery.GetSchema(), FSearchContext::SchemaRootBoneIdx);
diff --git a/Engine/Plugins/Experimental/Animation/PoseSearch/Source/Runtime/Private/PoseSearchFeatureChannel_Pose.h b/Engine/Plugins/Experimental/Animation/PoseSearch/Source/Runtime/Private/PoseSearchFeatureChannel_Pose.h
--- a/Engine/Plugins/Experimental/Animation/PoseSearch/Source/Runtime/Private/PoseSearchFeatureChannel_Pose.h
+++ b/Engine/Plugins/Experimental/Animation/PoseSearch/Source/Runtime/Private/PoseSearchFeatureChannel_Pose.h
@@ -55,6 +55,10 @@ public:
 	UPROPERTY(EditAnywhere, Category = "Settings")
 	EInputQueryPose InputQueryPose = EInputQueryPose::UseContinuingPose;
 
+	// if bUseCharacterSpaceVelocities is true, velocities will be calculated as from the positions in character space, otherwise they will be calculated using global space positions
+	UPROPERTY(EditAnywhere, Category = "Settings")
+	bool bUseCharacterSpaceVelocities = true;
+
 	// UPoseSearchFeatureChannel interface
 	virtual void InitializeSchema(UPoseSearchSchema* Schema) override;
 	virtual void FillWeights(TArray<float>& Weights) const override;
diff --git a/Engine/Plugins/Experimental/Animation/PoseSearch/Source/Runtime/Private/PoseSearchFeatureChannel_Position.cpp b/Engine/Plugins/Experimental/Animation/PoseSearch/Source/Runtime/Private/PoseSearchFeatureChannel_Position.cpp
--- a/Engine/Plugins/Experimental/Animation/PoseSearch/Source/Runtime/Private/PoseSearchFeatureChannel_Position.cpp
+++ b/Engine/Plugins/Experimental/Animation/PoseSearch/Source/Runtime/Private/PoseSearchFeatureChannel_Position.cpp
@@ -41,7 +41,7 @@ void UPoseSearchFeatureChannel_Position::IndexAsset(UE::PoseSearch::IAssetIndexe
 		const float SubsampleTime = OriginSampleTime + SampleTimeOffset;
 
 		bool ClampedPresent;
-		const FTransform BoneTransformsPresent = Indexer.GetTransformAndCacheResults(SubsampleTime, bUseSampleTimeOffsetRootBone ? SubsampleTime : OriginSampleTime, SchemaBoneIdx, ClampedPresent);
+		const FTransform BoneTransformsPresent = Indexer.GetTransformAndCacheResults(SubsampleTime, OriginSampleTime, SchemaBoneIdx, ClampedPresent);
 		int32 DataOffset = ChannelDataOffset;
 		FFeatureVectorHelper::EncodeVector(IndexingContext.GetPoseVector(VectorIdx, FeatureVectorTable), DataOffset, BoneTransformsPresent.GetTranslation());
 	}
@@ -66,13 +66,9 @@ void UPoseSearchFeatureChannel_Position::BuildQuery(UE::PoseSearch::FSearchConte
 	else
 	{
 		FTransform Transform = SearchContext.TryGetTransformAndCacheResults(SampleTimeOffset, InOutQuery.GetSchema(), SchemaBoneIdx);
-
-		if (!bUseSampleTimeOffsetRootBone)
-		{
-			const FTransform RootTransform = SearchContext.TryGetTransformAndCacheResults(0.f, InOutQuery.GetSchema(), FSearchContext::SchemaRootBoneIdx);
-			const FTransform RootTransformPrev = SearchContext.TryGetTransformAndCacheResults(SampleTimeOffset, InOutQuery.GetSchema(), FSearchContext::SchemaRootBoneIdx);
-			Transform = Transform * (RootTransformPrev * RootTransform.Inverse());
-		}
+		const FTransform RootTransform = SearchContext.TryGetTransformAndCacheResults(0.f, InOutQuery.GetSchema(), FSearchContext::SchemaRootBoneIdx);
+		const FTransform RootTransformPrev = SearchContext.TryGetTransformAndCacheResults(SampleTimeOffset, InOutQuery.GetSchema(), FSearchContext::SchemaRootBoneIdx);
+		Transform = Transform * (RootTransformPrev * RootTransform.Inverse());
 
 		int32 DataOffset = ChannelDataOffset;
 		FFeatureVectorHelper::EncodeVector(InOutQuery.EditValues(), DataOffset, Transform.GetTranslation());
diff --git a/Engine/Plugins/Experimental/Animation/PoseSearch/Source/Runtime/Private/PoseSearchFeatureChannel_Position.h b/Engine/Plugins/Experimental/Animation/PoseSearch/Source/Runtime/Private/PoseSearchFeatureChannel_Position.h
--- a/Engine/Plugins/Experimental/Animation/PoseSearch/Source/Runtime/Private/PoseSearchFeatureChannel_Position.h
+++ b/Engine/Plugins/Experimental/Animation/PoseSearch/Source/Runtime/Private/PoseSearchFeatureChannel_Position.h
@@ -30,11 +30,6 @@ public:
 	UPROPERTY(EditAnywhere, Category = "Settings")
 	EInputQueryPose InputQueryPose = EInputQueryPose::UseContinuingPose;
 
-	// if UseSampleTimeOffsetRootBone is true, this UPoseSearchFeatureChannel_Position will calculate the position of Bone from the pose SampleTimeOffset seconds away from the current time pose root bone
-	// if false the calculated position of Bone will be in component space from the pose SampleTimeOffset seconds away
-	UPROPERTY(EditAnywhere, Category = "Settings")
-	bool bUseSampleTimeOffsetRootBone = true;
-
 	// UPoseSearchFeatureChannel interface
 	virtual void InitializeSchema(UPoseSearchSchema* Schema) override;
 	virtual void FillWeights(TArray<float>& Weights) const override;
diff --git a/Engine/Plugins/Experimental/Animation/PoseSearch/Source/Runtime/Private/PoseSearchSchema.cpp b/Engine/Plugins/Experimental/Animation/PoseSearch/Source/Runtime/Private/PoseSearchSchema.cpp
--- a/Engine/Plugins/Experimental/Animation/PoseSearch/Source/Runtime/Private/PoseSearchSchema.cpp
+++ b/Engine/Plugins/Experimental/Animation/PoseSearch/Source/Runtime/Private/PoseSearchSchema.cpp
@@ -10,47 +10,36 @@ bool UPoseSearchSchema::IsValid() const
 {
 	bool bValid = Skeleton != nullptr;
 
-	for (const FBoneReference& BoneRef : BoneReferences)
-	{
-		bValid &= BoneRef.HasValidSetup();
-	}
-
 	for (const TObjectPtr<UPoseSearchFeatureChannel>& Channel: Channels)
 	{
 		bValid &= Channel != nullptr;
 	}
 
-	bValid &= (BoneReferences.Num() == BoneIndices.Num());
-
 	return bValid;
 }
 
 void UPoseSearchSchema::ResolveBoneReferences()
 {
-	// Initialize references to obtain bone indices
-	for (FBoneReference& BoneRef : BoneReferences)
-	{
-		BoneRef.Initialize(Skeleton);
-	}
-
-	// Fill out bone index array
-	BoneIndices.SetNum(BoneReferences.Num());
-	for (int32 BoneRefIdx = 0; BoneRefIdx != BoneReferences.Num(); ++BoneRefIdx)
-	{
-		BoneIndices[BoneRefIdx] = BoneReferences[BoneRefIdx].BoneIndex;
-	}
-
-	// Build separate index array with parent indices guaranteed to be present. Sort for EnsureParentsPresent.
-	BoneIndicesWithParents = BoneIndices;
-	BoneIndicesWithParents.Sort();
-
+	BoneIndicesWithParents.Reset();
 	if (Skeleton)
 	{
+		// Initialize references to obtain bone indices and fill out bone index array
+		for (FBoneReference& BoneRef : BoneReferences)
+		{
+			BoneRef.Initialize(Skeleton);
+			if (BoneRef.HasValidSetup())
+			{
+				BoneIndicesWithParents.Add(BoneRef.BoneIndex);
+			}
+		}
+
+		// Build separate index array with parent indices guaranteed to be present. Sort for EnsureParentsPresent.
+		BoneIndicesWithParents.Sort();
 		FAnimationRuntime::EnsureParentsPresent(BoneIndicesWithParents, Skeleton->GetReferenceSkeleton());
 	}
 
 	// BoneIndicesWithParents should at least contain the root to support mirroring root motion
-	if (BoneIndicesWithParents.Num() == 0)
+	if (BoneIndicesWithParents.IsEmpty())
 	{
 		BoneIndicesWithParents.Add(0);
 	}
diff --git a/Engine/Plugins/Experimental/Animation/PoseSearch/Source/Runtime/Public/PoseSearch/PoseSearchSchema.h b/Engine/Plugins/Experimental/Animation/PoseSearch/Source/Runtime/Public/PoseSearch/PoseSearchSchema.h
--- a/Engine/Plugins/Experimental/Animation/PoseSearch/Source/Runtime/Public/PoseSearch/PoseSearchSchema.h
+++ b/Engine/Plugins/Experimental/Animation/PoseSearch/Source/Runtime/Public/PoseSearch/PoseSearchSchema.h
@@ -74,9 +74,6 @@ public:
 	UPROPERTY(meta = (ExcludeFromHash))
 	TArray<FBoneReference> BoneReferences;
 
-	UPROPERTY(Transient)
-	TArray<uint16> BoneIndices;
-
 	UPROPERTY(Transient)
 	TArray<uint16> BoneIndicesWithParents;
 
@@ -98,8 +95,6 @@ public:
 	
 	bool IsValid () const;
 
-	int32 GetNumBones () const { return BoneIndices.Num(); }
-
 	float GetSamplingInterval() const { return 1.0f / SampleRate; }
 
 	template<typename ChannelType>
