motion matching - group channel

diff --git a/Engine/Plugins/Experimental/Animation/PoseSearch/Source/Runtime/Private/PoseSearchFeatureChannel_Group.cpp b/Engine/Plugins/Experimental/Animation/PoseSearch/Source/Runtime/Private/PoseSearchFeatureChannel_Group.cpp
new file mode 100644
--- /dev/null
+++ b/Engine/Plugins/Experimental/Animation/PoseSearch/Source/Runtime/Private/PoseSearchFeatureChannel_Group.cpp
@@ -0,0 +1,97 @@
+// Copyright Epic Games, Inc. All Rights Reserved.
+
+#include "PoseSearchFeatureChannel_Group.h"
+#include "DrawDebugHelpers.h"
+#include "PoseSearch/PoseSearchSchema.h"
+
+void UPoseSearchFeatureChannel_Group::InitializeSchema(UPoseSearchSchema* Schema)
+{
+	ChannelDataOffset = Schema->SchemaCardinality;
+	for (TObjectPtr<UPoseSearchFeatureChannel>& SubChannelPtr : SubChannels)
+	{
+		if (UPoseSearchFeatureChannel* SubChannel = SubChannelPtr.Get())
+		{
+			SubChannel->InitializeSchema(Schema);
+		}
+	}
+	ChannelCardinality = Schema->SchemaCardinality - ChannelDataOffset;
+}
+
+void UPoseSearchFeatureChannel_Group::FillWeights(TArray<float>& Weights) const
+{
+	for (const TObjectPtr<UPoseSearchFeatureChannel>& SubChannelPtr : SubChannels)
+	{
+		if (const UPoseSearchFeatureChannel* SubChannel = SubChannelPtr.Get())
+		{
+			SubChannel->FillWeights(Weights);
+		}
+	}
+}
+
+void UPoseSearchFeatureChannel_Group::IndexAsset(UE::PoseSearch::IAssetIndexer& Indexer, TArrayView<float> FeatureVectorTable) const
+{
+	for (const TObjectPtr<UPoseSearchFeatureChannel>& SubChannelPtr : SubChannels)
+	{
+		if (const UPoseSearchFeatureChannel* SubChannel = SubChannelPtr.Get())
+		{
+			SubChannel->IndexAsset(Indexer, FeatureVectorTable);
+		}
+	}
+}
+
+void UPoseSearchFeatureChannel_Group::BuildQuery(UE::PoseSearch::FSearchContext& SearchContext, FPoseSearchFeatureVectorBuilder& InOutQuery) const
+{
+	for (const TObjectPtr<UPoseSearchFeatureChannel>& SubChannelPtr : SubChannels)
+	{
+		if (const UPoseSearchFeatureChannel* SubChannel = SubChannelPtr.Get())
+		{
+			SubChannel->BuildQuery(SearchContext, InOutQuery);
+		}
+	}
+}
+
+void UPoseSearchFeatureChannel_Group::DebugDraw(const UE::PoseSearch::FDebugDrawParams& DrawParams, TConstArrayView<float> PoseVector) const
+{
+#if ENABLE_DRAW_DEBUG
+	for (const TObjectPtr<UPoseSearchFeatureChannel>& SubChannelPtr : SubChannels)
+	{
+		if (const UPoseSearchFeatureChannel* SubChannel = SubChannelPtr.Get())
+		{
+			SubChannel->DebugDraw(DrawParams, PoseVector);
+		}
+	}
+#endif // ENABLE_DRAW_DEBUG
+}
+
+// IPoseFilter interface
+bool UPoseSearchFeatureChannel_Group::IsPoseFilterActive() const
+{
+	for (const TObjectPtr<UPoseSearchFeatureChannel>& SubChannelPtr : SubChannels)
+	{
+		if (const UPoseSearchFeatureChannel* SubChannel = SubChannelPtr.Get())
+		{
+			if (SubChannel->IsPoseFilterActive())
+			{
+				return true;
+			}
+		}
+	}
+
+	return false;
+}
+
+bool UPoseSearchFeatureChannel_Group::IsPoseValid(TConstArrayView<float> PoseValues, TConstArrayView<float> QueryValues, int32 PoseIdx, const FPoseSearchPoseMetadata& Metadata) const
+{
+	for (const TObjectPtr<UPoseSearchFeatureChannel>& SubChannelPtr : SubChannels)
+	{
+		if (const UPoseSearchFeatureChannel* SubChannel = SubChannelPtr.Get())
+		{
+			if (!SubChannel->IsPoseValid(PoseValues, QueryValues, PoseIdx, Metadata))
+			{
+				return false;
+			}
+		}
+	}
+
+	return true;
+}
diff --git a/Engine/Plugins/Experimental/Animation/PoseSearch/Source/Runtime/Private/PoseSearchFeatureChannel_Group.h b/Engine/Plugins/Experimental/Animation/PoseSearch/Source/Runtime/Private/PoseSearchFeatureChannel_Group.h
new file mode 100644
--- /dev/null
+++ b/Engine/Plugins/Experimental/Animation/PoseSearch/Source/Runtime/Private/PoseSearchFeatureChannel_Group.h
@@ -0,0 +1,28 @@
+// Copyright Epic Games, Inc. All Rights Reserved.
+
+#pragma once
+
+#include "PoseSearch/PoseSearchFeatureChannel.h"
+#include "BoneContainer.h"
+#include "PoseSearchFeatureChannel_Group.generated.h"
+
+UCLASS(BlueprintType, EditInlineNew, meta = (DisplayName = "Group Channel"), CollapseCategories)
+class POSESEARCH_API UPoseSearchFeatureChannel_Group : public UPoseSearchFeatureChannel
+{
+	GENERATED_BODY()
+
+public:
+	UPROPERTY(EditAnywhere, Instanced, BlueprintReadWrite, Category = "Settings")
+	TArray<TObjectPtr<UPoseSearchFeatureChannel>> SubChannels;
+
+	// UPoseSearchFeatureChannel interface
+	virtual void InitializeSchema(UPoseSearchSchema* Schema) override;
+	virtual void FillWeights(TArray<float>& Weights) const override;
+	virtual void IndexAsset(UE::PoseSearch::IAssetIndexer& Indexer, TArrayView<float> FeatureVectorTable) const override;
+	virtual void BuildQuery(UE::PoseSearch::FSearchContext& SearchContext, FPoseSearchFeatureVectorBuilder& InOutQuery) const override;
+	virtual void DebugDraw(const UE::PoseSearch::FDebugDrawParams& DrawParams, TConstArrayView<float> PoseVector) const override;
+
+	// IPoseFilter interface
+	virtual bool IsPoseFilterActive() const override;
+	virtual bool IsPoseValid(TConstArrayView<float> PoseValues, TConstArrayView<float> QueryValues, int32 PoseIdx, const FPoseSearchPoseMetadata& Metadata) const override;
+};
\ No newline at end of file
