motion matching
- fix for FFeatureVectorHelper::DecodeQuatInternal
- made FFeatureVectorHelper::EncodeQuat consistent with FFeatureVectorHelper::DecodeQuatInternal
- fix for UPoseSearchFeatureChannel_Heading::DebugDraw
- implemented UPoseSearchFeatureChannel_Pose::DebugDraw for EPoseSearchBoneFlags::Rotation

diff --git a/Engine/Plugins/Experimental/Animation/PoseSearch/Source/Runtime/Private/PoseSearchFeatureChannel.cpp b/Engine/Plugins/Experimental/Animation/PoseSearch/Source/Runtime/Private/PoseSearchFeatureChannel.cpp
--- a/Engine/Plugins/Experimental/Animation/PoseSearch/Source/Runtime/Private/PoseSearchFeatureChannel.cpp
+++ b/Engine/Plugins/Experimental/Animation/PoseSearch/Source/Runtime/Private/PoseSearchFeatureChannel.cpp
@@ -42,8 +42,12 @@ void ICostBreakDownData::AddEntireBreakDownSection(const FText& Label, const UPo
 // FFeatureVectorHelper
 void FFeatureVectorHelper::EncodeQuat(TArrayView<float> Values, int32& DataOffset, const FQuat& Quat)
 {
-	const FVector X = Quat.GetAxisX();
-	const FVector Y = Quat.GetAxisY();
+	// unitary quaternions are non Euclidean representation of rotation, so they cannot be used to calculate cost functions within the context of kdtrees, 
+	// so we convert them in a matrix, and pick 2 axis (we choose X,Y), skipping the 3rd since correlated to the cross product of the first two (this saves memory and cpu cycles)
+
+	const FMatrix M = Quat.ToMatrix();
+	const FVector X = M.GetScaledAxis(EAxis::X);
+	const FVector Y = M.GetScaledAxis(EAxis::Y);
 
 	Values[DataOffset + 0] = X.X;
 	Values[DataOffset + 1] = X.Y;
@@ -85,15 +89,13 @@ FQuat FFeatureVectorHelper::DecodeQuat(TConstArrayView<float> Values, int32& Dat
 
 FQuat FFeatureVectorHelper::DecodeQuatInternal(TConstArrayView<float> Values, int32 DataOffset)
 {
+	// reconstructing the chosen 2 axis (X,Y - from EncodeQuat) from the 6 floats from Values 
 	const FVector X(Values[DataOffset + 0], Values[DataOffset + 1], Values[DataOffset + 2]);
 	const FVector Y(Values[DataOffset + 3], Values[DataOffset + 4], Values[DataOffset + 5]);
+	// calculating the 3rd axis required to reconstruct the rotation matrix
 	const FVector Z = FVector::CrossProduct(X, Y);
-
-	FMatrix M(FMatrix::Identity);
-	M.SetColumn(0, X);
-	M.SetColumn(1, Y);
-	M.SetColumn(2, Z);
-
+	const FMatrix M(X, Y, Z, FVector::ZeroVector);
+	// converting the matrix to the unitary quaternion
 	return FQuat(M);
 }
 
diff --git a/Engine/Plugins/Experimental/Animation/PoseSearch/Source/Runtime/Private/PoseSearchFeatureChannel_Heading.cpp b/Engine/Plugins/Experimental/Animation/PoseSearch/Source/Runtime/Private/PoseSearchFeatureChannel_Heading.cpp
--- a/Engine/Plugins/Experimental/Animation/PoseSearch/Source/Runtime/Private/PoseSearchFeatureChannel_Heading.cpp
+++ b/Engine/Plugins/Experimental/Animation/PoseSearch/Source/Runtime/Private/PoseSearchFeatureChannel_Heading.cpp
@@ -58,8 +58,8 @@ void UPoseSearchFeatureChannel_Heading::IndexAsset(UE::PoseSearch::IAssetIndexer
 		bool ClampedPresent;
 		const FTransform BoneTransformsPresent = Indexer.GetTransformAndCacheResults(SubsampleTime, OriginSampleTime, SchemaBoneIdx, ClampedPresent);
 		int32 DataOffset = ChannelDataOffset;
-		FFeatureVectorHelper::EncodeVector(IndexingContext.GetPoseVector(VectorIdx, FeatureVectorTable), DataOffset, GetAxis(BoneTransformsPresent.GetRotation()));
-		check(DataOffset == ChannelDataOffset + ChannelCardinality);
+		const FVector Heading = GetAxis(BoneTransformsPresent.GetRotation());
+		FFeatureVectorHelper::EncodeVector(IndexingContext.GetPoseVector(VectorIdx, FeatureVectorTable), DataOffset, Heading);
 	}
 }
 
@@ -76,7 +76,6 @@ void UPoseSearchFeatureChannel_Heading::BuildQuery(UE::PoseSearch::FSearchContex
 			const float LerpValue = InputQueryPose == EInputQueryPose::UseInterpolatedContinuingPose ? SearchContext.CurrentResult.LerpValue : 0.f;
 			int32 DataOffset = ChannelDataOffset;
 			FFeatureVectorHelper::EncodeVector(InOutQuery.EditValues(), DataOffset, SearchContext.GetCurrentResultPrevPoseVector(), SearchContext.GetCurrentResultPoseVector(), SearchContext.GetCurrentResultNextPoseVector(), LerpValue, true);
-			check(DataOffset == ChannelDataOffset + ChannelCardinality);
 		}
 		// else leave the InOutQuery set to zero since the SearchContext.History is invalid and it'll fail if we continue
 	}
@@ -89,7 +88,6 @@ void UPoseSearchFeatureChannel_Heading::BuildQuery(UE::PoseSearch::FSearchContex
 
 		int32 DataOffset = ChannelDataOffset;
 		FFeatureVectorHelper::EncodeVector(InOutQuery.EditValues(), DataOffset, GetAxis(Transform.GetRotation()));
-		check(DataOffset == ChannelDataOffset + ChannelCardinality);
 	}
 }
 
@@ -109,19 +107,18 @@ void UPoseSearchFeatureChannel_Heading::DebugDraw(const UE::PoseSearch::FDebugDr
 	const FVector BonePos = DrawParams.Mesh != nullptr ? DrawParams.Mesh->GetSocketTransform(Bone.BoneName).GetLocation() : DrawParams.RootTransform.GetTranslation();
 
 	int32 DataOffset = ChannelDataOffset;
-	const FVector BoneHeading = DrawParams.RootTransform.TransformPosition(FFeatureVectorHelper::DecodeVector(PoseVector, DataOffset));
-	check(DataOffset == ChannelDataOffset + ChannelCardinality);
+	const FVector BoneHeading = DrawParams.RootTransform.GetRotation().RotateVector(FFeatureVectorHelper::DecodeVector(PoseVector, DataOffset));
 
 	const FColor Color = DrawParams.GetColor(ColorPresetIndex);
 
 	if (EnumHasAnyFlags(DrawParams.Flags, EDebugDrawFlags::DrawSearchIndex))
 	{
-		DrawDebugLine(DrawParams.World, BonePos, BonePos + BoneHeading, Color, bPersistent, LifeTime, DepthPriority);
+		DrawDebugLine(DrawParams.World, BonePos, BonePos + BoneHeading * 15.f, Color, bPersistent, LifeTime, DepthPriority);
 	}
 	else
 	{
 		const float AdjustedThickness = EnumHasAnyFlags(DrawParams.Flags, EDebugDrawFlags::DrawFast) ? 0.0f : 1.f;
-		DrawDebugLine(DrawParams.World, BonePos, BonePos + BoneHeading, Color, bPersistent, LifeTime, DepthPriority, AdjustedThickness);
+		DrawDebugLine(DrawParams.World, BonePos, BonePos + BoneHeading * 15.f, Color, bPersistent, LifeTime, DepthPriority, AdjustedThickness);
 	}
 #endif // ENABLE_DRAW_DEBUG
 }
diff --git a/Engine/Plugins/Experimental/Animation/PoseSearch/Source/Runtime/Private/PoseSearchFeatureChannel_Pose.cpp b/Engine/Plugins/Experimental/Animation/PoseSearch/Source/Runtime/Private/PoseSearchFeatureChannel_Pose.cpp
--- a/Engine/Plugins/Experimental/Animation/PoseSearch/Source/Runtime/Private/PoseSearchFeatureChannel_Pose.cpp
+++ b/Engine/Plugins/Experimental/Animation/PoseSearch/Source/Runtime/Private/PoseSearchFeatureChannel_Pose.cpp
@@ -497,7 +497,7 @@ void UPoseSearchFeatureChannel_Pose::AddPoseFeatures(UE::PoseSearch::IAssetIndex
 			}
 			else
 			{
-				LinearVelocity = (BoneTransformsFuture.GetTranslation() - BoneTransformsPast.GetTranslation()) / (SamplingContext->FiniteDelta * 2.0f);
+				LinearVelocity = (BoneTransformsFuture.GetTranslation() - BoneTransformsPast.GetTranslation()) / (SamplingContext->FiniteDelta * 2.f);
 			}
 
 			FFeatureVectorHelper::EncodeVector(FeatureVector, DataOffset, LinearVelocity);
@@ -691,7 +691,7 @@ void UPoseSearchFeatureChannel_Pose::DebugDraw(const UE::PoseSearch::FDebugDrawP
 				DrawDebugString(
 					DrawParams.World, BonePos + FVector(0.0, 0.0, 10.0),
 					Schema->BoneReferences[SchemaBoneIdx[ChannelBoneIdx]].BoneName.ToString(),
-					nullptr, Color, LifeTime, false, 1.0f);
+					nullptr, Color, LifeTime, false, 1.f);
 			}
 		}
 		else
@@ -702,7 +702,9 @@ void UPoseSearchFeatureChannel_Pose::DebugDraw(const UE::PoseSearch::FDebugDrawP
 		if (EnumHasAnyFlags(SampledBone.Flags, EPoseSearchBoneFlags::Rotation))
 		{
 			const FQuat BoneRot = FFeatureVectorHelper::DecodeQuat(PoseVector, DataOffset);
-			// @todo: debug draw rotation
+			const FRotator Rotator(DrawParams.RootTransform.GetRotation() * BoneRot);
+			const float AdjustedThickness = EnumHasAnyFlags(DrawParams.Flags, EDebugDrawFlags::DrawFast) ? 0.f : 1.f;
+			DrawDebugCoordinateSystem(DrawParams.World, BonePos, Rotator, 15.f, bPersistent, LifeTime, DepthPriority, AdjustedThickness);
 		}
 
 		if (EnumHasAnyFlags(SampledBone.Flags, EPoseSearchBoneFlags::Velocity))
@@ -721,7 +723,7 @@ void UPoseSearchFeatureChannel_Pose::DebugDraw(const UE::PoseSearch::FDebugDrawP
 			}
 			else
 			{
-				const float AdjustedThickness = EnumHasAnyFlags(DrawParams.Flags, EDebugDrawFlags::DrawFast) ? 0.0f : 1.f;
+				const float AdjustedThickness = EnumHasAnyFlags(DrawParams.Flags, EDebugDrawFlags::DrawFast) ? 0.f : 1.f;
 
 				DrawDebugLine(
 					DrawParams.World,
