PoseSearchDebugger - Open asset or motion database with asset paused and matching the time being displayed/used in the PoseSearchDebugger.

diff --git a/Engine/Plugins/Experimental/Animation/PoseSearch/Source/Editor/Private/PoseSearchDatabaseAssetTree.cpp b/Engine/Plugins/Experimental/Animation/PoseSearch/Source/Editor/Private/PoseSearchDatabaseAssetTree.cpp
--- a/Engine/Plugins/Experimental/Animation/PoseSearch/Source/Editor/Private/PoseSearchDatabaseAssetTree.cpp
+++ b/Engine/Plugins/Experimental/Animation/PoseSearch/Source/Editor/Private/PoseSearchDatabaseAssetTree.cpp
@@ -733,6 +733,21 @@ namespace UE::PoseSearch
 		RefreshTreeView(false, bRecoverSelection);
 		EditorViewModel.Pin()->BuildSearchIndex();
 	}
+
+	TWeakPtr<FDatabaseAssetTreeNode> SDatabaseAssetTree::SetSelectedItem(int32 SourceAssetIdx)
+	{
+		for (TSharedPtr<FDatabaseAssetTreeNode> & Node : AllNodes)
+		{
+			if (Node->SourceAssetIdx == SourceAssetIdx)
+			{
+				TreeView->ClearSelection();
+				TreeView->SetItemSelection(Node, true);
+				return Node;
+			}
+		}
+
+		return TWeakPtr<FDatabaseAssetTreeNode>();
+	}
 }
 
 #undef LOCTEXT_NAMESPACE
diff --git a/Engine/Plugins/Experimental/Animation/PoseSearch/Source/Editor/Private/PoseSearchDatabaseEditor.cpp b/Engine/Plugins/Experimental/Animation/PoseSearch/Source/Editor/Private/PoseSearchDatabaseEditor.cpp
--- a/Engine/Plugins/Experimental/Animation/PoseSearch/Source/Editor/Private/PoseSearchDatabaseEditor.cpp
+++ b/Engine/Plugins/Experimental/Animation/PoseSearch/Source/Editor/Private/PoseSearchDatabaseEditor.cpp
@@ -55,6 +55,19 @@ namespace UE::PoseSearch
 		return ViewModel.IsValid() ? ViewModel->GetPoseSearchDatabase() : nullptr;
 	}
 
+	void FDatabaseEditor::SetSelectedAsset(int32 SourceAssetIdx)
+	{
+		if (ViewModel.IsValid())
+		{
+			const TWeakPtr<FDatabaseAssetTreeNode> Node = AssetTreeWidget->SetSelectedItem(SourceAssetIdx);
+			
+			if (Node.IsValid())
+			{
+				ViewModel->SetSelectedNode(Node.Pin());
+			}
+		}
+	}
+
 	void FDatabaseEditor::BuildSearchIndex()
 	{
 		ViewModel->BuildSearchIndex();
diff --git a/Engine/Plugins/Experimental/Animation/PoseSearch/Source/Editor/Private/PoseSearchDatabaseViewModel.cpp b/Engine/Plugins/Experimental/Animation/PoseSearch/Source/Editor/Private/PoseSearchDatabaseViewModel.cpp
--- a/Engine/Plugins/Experimental/Animation/PoseSearch/Source/Editor/Private/PoseSearchDatabaseViewModel.cpp
+++ b/Engine/Plugins/Experimental/Animation/PoseSearch/Source/Editor/Private/PoseSearchDatabaseViewModel.cpp
@@ -509,6 +509,14 @@ namespace UE::PoseSearch
 		return false;
 	}
 
+	void FDatabaseViewModel::SetSelectedNode(const TSharedPtr<FDatabaseAssetTreeNode>& InSelectedNode)
+	{
+		SelectedNodes.Reset(1);
+		SelectedNodes.Push(InSelectedNode);
+		ResetPreviewActors();
+		ProcessSelectedActor(nullptr);
+	}
+
 	void FDatabaseViewModel::SetSelectedNodes(const TArrayView<TSharedPtr<FDatabaseAssetTreeNode>>& InSelectedNodes)
 	{
 		SelectedNodes = InSelectedNodes;
diff --git a/Engine/Plugins/Experimental/Animation/PoseSearch/Source/Editor/Private/PoseSearchDebuggerDatabaseColumns.h b/Engine/Plugins/Experimental/Animation/PoseSearch/Source/Editor/Private/PoseSearchDebuggerDatabaseColumns.h
--- a/Engine/Plugins/Experimental/Animation/PoseSearch/Source/Editor/Private/PoseSearchDebuggerDatabaseColumns.h
+++ b/Engine/Plugins/Experimental/Animation/PoseSearch/Source/Editor/Private/PoseSearchDebuggerDatabaseColumns.h
@@ -2,8 +2,14 @@
 
 #pragma once
 
+#include "AnimPreviewInstance.h"
 #include "Editor.h"
+#include "IAnimationEditor.h"
+#include "IPersonaToolkit.h"
+#include "PoseSearchDatabaseEditor.h"
+#include "PoseSearchDatabaseViewModel.h"
 #include "PoseSearchDebuggerDatabaseRowData.h"
+#include "Animation/DebugSkelMeshComponent.h"
 #include "Preferences/PersonaOptions.h"
 #include "Styling/AppStyle.h"
 #include "Subsystems/AssetEditorSubsystem.h"
@@ -116,7 +122,42 @@ struct FDatabaseName : IColumn
 				})
 			.OnNavigate_Lambda([RowData]()
 				{
-					GEditor->GetEditorSubsystem<UAssetEditorSubsystem>()->OpenEditorForAsset(RowData->DatabasePath);
+					UObject* Asset = nullptr;
+
+					// Load asset
+					if (UPackage* Package = LoadPackage(NULL, *RowData->DatabasePath, LOAD_NoRedirects))
+					{
+						Package->FullyLoad();
+
+						const FString AssetName = FPaths::GetBaseFilename(RowData->DatabasePath);
+						Asset = FindObject<UObject>(Package, *AssetName);
+					}
+					else
+					{
+						// Fallback for unsaved assets
+						Asset = FindObject<UObject>(nullptr, *RowData->DatabasePath);
+					}
+
+					// Open editor
+					if (Asset != nullptr)
+					{
+						if (UAssetEditorSubsystem* AssetEditorSS = GEditor->GetEditorSubsystem<UAssetEditorSubsystem>())
+						{
+							AssetEditorSS->OpenEditorForAsset(Asset);
+								
+							if (IAssetEditorInstance* Editor = AssetEditorSS->FindEditorForAsset(Asset, true))
+							{
+								if (Editor->GetEditorName() == FName("PoseSearchDatabaseEditor"))
+								{
+									FDatabaseEditor* DatabaseEditor = static_cast<FDatabaseEditor*>(Editor);
+									
+									// Open asset paused and at specific time as seen on the pose search debugger.
+									DatabaseEditor->SetSelectedAsset(RowData->DbAssetIdx);
+									DatabaseEditor->GetViewModel()->SetPlayTime(RowData->AssetTime, false);
+								}
+							}
+						}
+					}
 				});
 	}
 };
@@ -145,7 +186,44 @@ struct FAssetName : IColumn
 				})
 			.OnNavigate_Lambda([RowData]()
 				{
-					GEditor->GetEditorSubsystem<UAssetEditorSubsystem>()->OpenEditorForAsset(RowData->AssetPath);
+					UObject* Asset = nullptr;
+
+					// Load asset
+					if (UPackage* Package = LoadPackage(NULL, *RowData->AssetPath, LOAD_NoRedirects))
+					{
+						Package->FullyLoad();
+
+						const FString AssetName = FPaths::GetBaseFilename(RowData->AssetPath);
+						Asset = FindObject<UObject>(Package, *AssetName);
+					}
+					else
+					{
+						// Fallback for unsaved assets
+						Asset = FindObject<UObject>(nullptr, *RowData->AssetPath);
+					}
+
+					// Open editor
+					if (Asset != nullptr)
+					{
+						if (UAssetEditorSubsystem* AssetEditorSS = GEditor->GetEditorSubsystem<UAssetEditorSubsystem>())
+						{
+							AssetEditorSS->OpenEditorForAsset(Asset);
+							
+							if (IAssetEditorInstance* Editor = AssetEditorSS->FindEditorForAsset(Asset, true))
+							{
+								if (Editor->GetEditorName() == "AnimationEditor")
+								{
+									const IAnimationEditor* AnimationEditor = static_cast<IAnimationEditor*>(Editor);
+									const UDebugSkelMeshComponent* PreviewComponent = AnimationEditor->GetPersonaToolkit()->GetPreviewMeshComponent();
+
+									// Open asset paused and at specific time as seen on the pose search debugger.
+									PreviewComponent->PreviewInstance->SetPosition(RowData->AssetTime);
+									PreviewComponent->PreviewInstance->SetPlaying(false);
+									PreviewComponent->PreviewInstance->SetBlendSpacePosition(RowData->BlendParameters);
+								}
+							}
+						}
+					}
 				});
 	}
 };
diff --git a/Engine/Plugins/Experimental/Animation/PoseSearch/Source/Editor/Public/PoseSearchDatabaseAssetTree.h b/Engine/Plugins/Experimental/Animation/PoseSearch/Source/Editor/Public/PoseSearchDatabaseAssetTree.h
--- a/Engine/Plugins/Experimental/Animation/PoseSearch/Source/Editor/Public/PoseSearchDatabaseAssetTree.h
+++ b/Engine/Plugins/Experimental/Animation/PoseSearch/Source/Editor/Public/PoseSearchDatabaseAssetTree.h
@@ -39,6 +39,7 @@ namespace UE::PoseSearch
 		void RefreshTreeView(bool bIsInitialSetup = false, bool bRecoverSelection = false);
 		void FinalizeTreeChanges(bool bRecoverSelection = false);
 
+		TWeakPtr<FDatabaseAssetTreeNode> SetSelectedItem(int32 SourceAssetIdx);
 		FDetailColumnSizeData& GetColumnSizeData() { return ColumnSizeData; }
 		
 	protected:
diff --git a/Engine/Plugins/Experimental/Animation/PoseSearch/Source/Editor/Public/PoseSearchDatabaseEditor.h b/Engine/Plugins/Experimental/Animation/PoseSearch/Source/Editor/Public/PoseSearchDatabaseEditor.h
--- a/Engine/Plugins/Experimental/Animation/PoseSearch/Source/Editor/Public/PoseSearchDatabaseEditor.h
+++ b/Engine/Plugins/Experimental/Animation/PoseSearch/Source/Editor/Public/PoseSearchDatabaseEditor.h
@@ -39,7 +39,9 @@ namespace UE::PoseSearch
 		UPoseSearchDatabase* GetPoseSearchDatabase();
 		FDatabaseViewModel* GetViewModel() const { return ViewModel.Get(); }
 		TSharedPtr<FDatabaseViewModel> GetViewModelSharedPtr() const { return ViewModel; }
-
+		
+		void SetSelectedAsset(int32 SourceAssetIdx);
+		
 		void BuildSearchIndex();
 
 		void PreviewBackwardEnd();
diff --git a/Engine/Plugins/Experimental/Animation/PoseSearch/Source/Editor/Public/PoseSearchDatabaseViewModel.h b/Engine/Plugins/Experimental/Animation/PoseSearch/Source/Editor/Public/PoseSearchDatabaseViewModel.h
--- a/Engine/Plugins/Experimental/Animation/PoseSearch/Source/Editor/Public/PoseSearchDatabaseViewModel.h
+++ b/Engine/Plugins/Experimental/Animation/PoseSearch/Source/Editor/Public/PoseSearchDatabaseViewModel.h
@@ -109,6 +109,7 @@ namespace UE::PoseSearch
 		void SetIsEnabled(int32 AnimationAssetIndex, bool bEnabled);
 		bool IsEnabled(int32 AnimationAssetIndex) const;
 
+		void SetSelectedNode(const TSharedPtr<FDatabaseAssetTreeNode>& InSelectedNode);
 		void SetSelectedNodes(const TArrayView<TSharedPtr<FDatabaseAssetTreeNode>>& InSelectedNodes);
 		void ProcessSelectedActor(AActor* Actor);
 		
