motion matching
- FTrajectorySampleRange improvements: reduced to log complexity to GetSampleAtTime
- refactored UPoseSearchFeatureChannel_Velocity and UPoseSearchFeatureChannel_Phase from UPoseSearchFeatureChannel_Pose
- derived UPoseSearchFeatureChannel_Pose and UPoseSearchFeatureChannel_Trajectory from UPoseSearchFeatureChannel_Group
- minor motion trajectory cleanup
- small pose search debugger stability improvement

diff --git a/Engine/Plugins/Experimental/Animation/MotionTrajectory/Source/MotionTrajectory/Private/MotionTrajectory.cpp b/Engine/Plugins/Experimental/Animation/MotionTrajectory/Source/MotionTrajectory/Private/MotionTrajectory.cpp
--- a/Engine/Plugins/Experimental/Animation/MotionTrajectory/Source/MotionTrajectory/Private/MotionTrajectory.cpp
+++ b/Engine/Plugins/Experimental/Animation/MotionTrajectory/Source/MotionTrajectory/Private/MotionTrajectory.cpp
@@ -74,7 +74,7 @@ FTrajectorySampleRange UMotionTrajectoryComponent::CombineHistoryPresentPredicti
 	Trajectory.Samples.Append(History.Samples);
 	Trajectory.Samples.Add(PresentTrajectorySampleLS);
 	Trajectory.Samples.Append(Prediction.Samples);
-	Trajectory.SampleRate = SampleRate;
+
 #if WITH_EDITORONLY_DATA
 	Trajectory.DebugDrawTrajectory(bDebugDrawTrajectory, GetOwner()->GetWorld(), PreviousWorldTransform);
 #endif
@@ -196,71 +196,26 @@ FTrajectorySampleRange UMotionTrajectoryComponent::GetTrajectory() const
 {
 	const APawn* Pawn = TryGetOwnerPawn();
 	checkf(false, TEXT("UMotionTrajectoryComponent::GetTrajectory for Pawn: %s requires implementation."), Pawn ? *(Pawn->GetHumanReadableName()) : TEXT("NULL"));
-	return FTrajectorySampleRange(SampleRate);
+	return FTrajectorySampleRange();
 }
 
 FTrajectorySampleRange UMotionTrajectoryComponent::GetTrajectoryWithSettings(const FMotionTrajectorySettings& Settings, bool bIncludeHistory) const
 {
 	const APawn* Pawn = TryGetOwnerPawn();
 	checkf(false, TEXT("UMotionTrajectoryComponent::GetTrajectoryWithSettings for Pawn: %s requires implementation."), Pawn ? *(Pawn->GetHumanReadableName()) : TEXT("NULL"));
-	return FTrajectorySampleRange(SampleRate);
-}
-
-void UMotionTrajectoryComponent::SetSampleRate(int32 Rate)
-{
-	// To maintain trajectory sample coherency, changing the sample rate will cause all historical samples to be flushed
-	// In the future, we may consider leveraging the frame-rate independent uniform resampling algorithm implemented within GetHistory()
-	FlushHistory();
-	SampleRate = Rate;
+	return FTrajectorySampleRange();
 }
 
 FTrajectorySampleRange UMotionTrajectoryComponent::GetHistory() const
 {
-	if (SampleHistory.IsEmpty())
+	FTrajectorySampleRange SourceHistory;
+	const int32 HistorySampleCount = SampleHistory.Num();
+	if (HistorySampleCount > 0)
 	{
-		return FTrajectorySampleRange(SampleRate);
-	}
-	else if (bUniformSampledHistory)
-	{
-		// Establish a uniform history range that will be frame-rate independent, and resampled against the SampleRate property specified by this component
-		FTrajectorySampleRange UniformHistory;
-		UniformHistory.Samples.Reserve(SampleHistory.Num());
-		UniformHistory.SampleRate = SampleRate;
-
-		int32 InterpStartIdx = 0;
-		const float SampleFrequency = 1.f / static_cast<float>(SampleRate);
-		const float FirstNonUniformSampleTime = SampleHistory.First().AccumulatedSeconds;
-		const float FirstUniformSampleTime = FMath::CeilToFloat(FirstNonUniformSampleTime / SampleFrequency) * SampleFrequency;
-
-		check(SampleFrequency > 0.f);
-		check(FirstUniformSampleTime <= 0.f);
-
-		// Iterate the source history buffer with a uniform sample frequency
-		// The iteration exit condition is meant to guarantee that we do not include (erroneously) a instantaneous/present trajectory sample at zero time
-		for (int32 Idx = 0;; ++Idx)
-		{
-			const float SampleTime = FirstUniformSampleTime + static_cast<float>(Idx) * SampleFrequency;
-			if (SampleTime < -UE_KINDA_SMALL_NUMBER)
-			{
-				const FTrajectorySample Sample = FTrajectorySampleRange::IterSampleTrajectory(SampleHistory, SampleTime, InterpStartIdx, bSmoothInterpolation);
-				UniformHistory.Samples.Emplace(Sample);
-			}
-			else
-			{
-				break;
-			}
-		}
-
-		return UniformHistory;
-	}
-	else
-	{
-		const int32 HistorySampleCount = SampleHistory.Num();
 		const int32 PresentSampleIdx = HistorySampleCount - 1;
 
 		// Convert the source history ring buffer into a trajectory range
 		// We consider these values to be non-uniform with a variable, frame-rate dependent sample rate governed by the component tick rate
-		FTrajectorySampleRange SourceHistory;
 		SourceHistory.Samples.Reserve(HistorySampleCount);
 
 		for (const auto& Sample : SampleHistory)
@@ -268,12 +223,12 @@ FTrajectorySampleRange UMotionTrajectoryComponent::GetHistory() const
 			SourceHistory.Samples.Emplace(Sample);
 		}
 
+		// @todo: fix me - maybe don't add it in the first place?
 		// Remove the present trajectory sample, this will not be included in the history
 		check(FMath::IsNearlyZero(SourceHistory.Samples[PresentSampleIdx].AccumulatedSeconds));
 		SourceHistory.Samples.RemoveAt(PresentSampleIdx);
-
-		return SourceHistory;
 	}
+	return SourceHistory;
 }
 
 #undef LOCTEXT_NAMESPACE
diff --git a/Engine/Plugins/Experimental/Animation/MotionTrajectory/Source/MotionTrajectory/Private/MotionTrajectoryCharacterMovement.cpp b/Engine/Plugins/Experimental/Animation/MotionTrajectory/Source/MotionTrajectory/Private/MotionTrajectoryCharacterMovement.cpp
--- a/Engine/Plugins/Experimental/Animation/MotionTrajectory/Source/MotionTrajectory/Private/MotionTrajectoryCharacterMovement.cpp
+++ b/Engine/Plugins/Experimental/Animation/MotionTrajectory/Source/MotionTrajectory/Private/MotionTrajectoryCharacterMovement.cpp
@@ -12,31 +12,26 @@ FTrajectorySample UCharacterMovementTrajectoryComponent::CalcWorldSpacePresentTr
 {
 	FTrajectorySample ReturnValue;
 	
-	const APawn* Pawn = TryGetOwnerPawn();
-	if (!Pawn)
-	{
-		return ReturnValue;
-	}
-
-	const UCharacterMovementComponent* MovementComponent = Cast<UCharacterMovementComponent>(Pawn->GetMovementComponent());
-	if (MovementComponent)
+	if (const APawn* Pawn = TryGetOwnerPawn())
 	{
-		if (MovementComponent->MovementMode == EMovementMode::MOVE_Walking)
+		if (const UCharacterMovementComponent* MovementComponent = Cast<UCharacterMovementComponent>(Pawn->GetMovementComponent()))
 		{
-			FTransform ComponentWorldTransform = Pawn->GetActorTransform();
+			if (MovementComponent->MovementMode == EMovementMode::MOVE_Walking)
+			{
+				FTransform ComponentWorldTransform = Pawn->GetActorTransform();
 
-			ReturnValue.Transform = ComponentWorldTransform;
-			ReturnValue.LinearVelocity = MovementComponent->Velocity;
+				ReturnValue.Transform = ComponentWorldTransform;
+				ReturnValue.LinearVelocity = MovementComponent->Velocity;
+			}
+		}
+		else
+		{
+			UE_LOG(
+				LogMotionTrajectory,
+				Error,
+				TEXT("UCharacterMovementTrajectoryComponent expects the owner to have a CharacterMovementComponent"));
 		}
 	}
-	else
-	{
-		UE_LOG(
-			LogMotionTrajectory,
-			Error,
-			TEXT("UCharacterMovementTrajectoryComponent expects the owner to have a CharacterMovementComponent"));
-	}
-
 	return ReturnValue;
 }
 
@@ -136,37 +131,31 @@ FTrajectorySampleRange UCharacterMovementTrajectoryComponent::GetTrajectory() co
 	return GetTrajectoryWithSettings(PredictionSettings, bPredictionIncludesHistory);
 }
 
-FTrajectorySampleRange UCharacterMovementTrajectoryComponent::GetTrajectoryWithSettings(const FMotionTrajectorySettings& Settings
-	, bool bIncludeHistory) const
+FTrajectorySampleRange UCharacterMovementTrajectoryComponent::GetTrajectoryWithSettings(const FMotionTrajectorySettings& Settings, bool bIncludeHistory) const
 {
-	const APawn* Pawn = TryGetOwnerPawn();
-	if (!Pawn)
-	{
-		return FTrajectorySampleRange(SampleRate);
-	}
-
-	const UCharacterMovementComponent* MovementComponent = Cast<UCharacterMovementComponent>(Pawn->GetMovementComponent());
+	FTrajectorySampleRange Prediction;
 
-	// Currently the trajectory prediction only supports the walking movement mode of the character movement component
-	if (ensure(MovementComponent) && MovementComponent->MovementMode == EMovementMode::MOVE_Walking)
+	if (const APawn* Pawn = TryGetOwnerPawn())
 	{
-		// Step the prediction iteratively towards the specified domain horizon(s)
-		FTrajectorySampleRange Prediction(SampleRate);
-		PredictTrajectory(
-			SampleRate, 
-			MaxSamples, 
-			Settings, 
-			PresentTrajectorySampleLS, 
-			DesiredControlRotationVelocity,
-			Prediction);
-
-		// Combine past, present, and future into a uniformly sampled complete trajectory
-		return CombineHistoryPresentPrediction(bIncludeHistory, Prediction);
-	}
-	else
-	{
-		return FTrajectorySampleRange(SampleRate);
+		const UCharacterMovementComponent* MovementComponent = Cast<UCharacterMovementComponent>(Pawn->GetMovementComponent());
+
+		// Currently the trajectory prediction only supports the walking movement mode of the character movement component
+		if (ensure(MovementComponent) && MovementComponent->MovementMode == EMovementMode::MOVE_Walking)
+		{
+			// Step the prediction iteratively towards the specified domain horizon(s)
+			PredictTrajectory(
+				SampleRate,
+				MaxSamples,
+				Settings,
+				PresentTrajectorySampleLS,
+				DesiredControlRotationVelocity,
+				Prediction);
+
+			// Combine past, present, and future into a uniformly sampled complete trajectory
+			Prediction = CombineHistoryPresentPrediction(bIncludeHistory, Prediction);
+		}
 	}
+	return Prediction;
 }
 
 void UCharacterMovementTrajectoryComponent::PredictTrajectory(
@@ -177,7 +166,7 @@ void UCharacterMovementTrajectoryComponent::PredictTrajectory(
 	const FRotator& InDesiredControlRotationVelocity,
 	FTrajectorySampleRange& OutTrajectoryRange) const
 {
-	OutTrajectoryRange.SampleRate = InSampleRate;
+	check(InSampleRate > UE_KINDA_SMALL_NUMBER);
 	const float IntegrationDelta = 1.f / static_cast<float>(InSampleRate);
 
 	FTrajectorySample Sample = PresentTrajectory;
diff --git a/Engine/Plugins/Experimental/Animation/MotionTrajectory/Source/MotionTrajectory/Private/MotionTrajectoryLibrary.cpp b/Engine/Plugins/Experimental/Animation/MotionTrajectory/Source/MotionTrajectory/Private/MotionTrajectoryLibrary.cpp
--- a/Engine/Plugins/Experimental/Animation/MotionTrajectory/Source/MotionTrajectory/Private/MotionTrajectoryLibrary.cpp
+++ b/Engine/Plugins/Experimental/Animation/MotionTrajectory/Source/MotionTrajectory/Private/MotionTrajectoryLibrary.cpp
@@ -182,8 +182,7 @@ bool UMotionTrajectoryBlueprintLibrary::IsStoppingTrajectory(
 		const FVector LastLinearVelocity = Trajectory.Samples.Last().LinearVelocity;
 		const float SquaredLastLinearSpeed = LastLinearVelocity.SquaredLength();
 		
-		int StartIdx = 0;
-		const FTrajectorySample& PresentSample = FTrajectorySampleRange::IterSampleTrajectory(Trajectory.Samples, 0.f, StartIdx);
+		const FTrajectorySample PresentSample = Trajectory.GetSampleAtTime(0.f);
 		const FVector PresenLinearVelocity = PresentSample.LinearVelocity;
 		const float SquaredPresentLinearSpeed = PresenLinearVelocity.SquaredLength();
 
@@ -205,8 +204,7 @@ bool UMotionTrajectoryBlueprintLibrary::IsStartingTrajectory(
 		const FVector FirstLinearVelocity = Trajectory.Samples[0].LinearVelocity;
 		const float SquaredFirstLinearSpeed = FirstLinearVelocity.SquaredLength();
 
-		int StartIdx = 0;
-		const FTrajectorySample& PresentSample = FTrajectorySampleRange::IterSampleTrajectory(Trajectory.Samples, 0.f, StartIdx);
+		const FTrajectorySample PresentSample = Trajectory.GetSampleAtTime(0.f);
 		const FVector PresentLinearVelocity = PresentSample.LinearVelocity;
 		const float SquaredPresentLinearSpeed = PresentLinearVelocity.SquaredLength();
 
@@ -240,8 +238,7 @@ bool UMotionTrajectoryBlueprintLibrary::IsConstantSpeedTrajectory(
 		const float SquaredLastLinearSpeed = LastLinearVelocity.SquaredLength();
 		const bool LastSpeedWithinLimit = IsWithinLimit(SquaredLastLinearSpeed);
 
-		int StartIdx = 0;
-		const FTrajectorySample& PresentSample = FTrajectorySampleRange::IterSampleTrajectory(Trajectory.Samples, 0.f, StartIdx);
+		const FTrajectorySample PresentSample = Trajectory.GetSampleAtTime(0.f);
 		const FVector PresentLinearVelocity = PresentSample.LinearVelocity;
 		const float SquaredPresentLinearSpeed = PresentLinearVelocity.SquaredLength();
 		const bool PresentSpeedWithinLimit = IsWithinLimit(SquaredLastLinearSpeed);
@@ -252,361 +249,4 @@ bool UMotionTrajectoryBlueprintLibrary::IsConstantSpeedTrajectory(
 	}
 
 	return false;
-}
-
-namespace UE::MotionTrajectory
-{
-	struct FTurnEvaluator
-	{
-
-	public:
-		FTurnEvaluator(const FTrajectorySampleRange& InTrajectory, float InRotationConstraintLimit, const FVector& InTurnAxis)
-			: Trajectory(InTrajectory)
-			, RotationConstraintLimit(InRotationConstraintLimit)
-			, TurnAxis(InTurnAxis)
-		{
-			CalcRotSpeedToExtrapolate();
-			CalcTurnData();
-		}
-
-		const FTrajectorySampleRange& Trajectory;
-		float RotationConstraintLimit;
-		FVector TurnAxis;
-
-		float RotSpeedToExtrapolate;
-		TArray<float, TInlineAllocator<128>> AccumulatedTurnAxisRotations;
-		TArray<float, TInlineAllocator<128>> ImmediateTurnAxisRotations;
-
-		// Evaluates if there's a turn in the trajectory starting from SampleIdx. If the velocity and the alignment
-		// between the forward axis and the velocity match when comparing the samples at SampleIdx and the end of the
-		// trajectory, up to MaxAlignmentAngleDegrees, the function will return false. Otherwise, it returns false.
-		bool IsTurning(int32 SampleIdx, const FVector& ForwardAxis, float MaxAlignmentAngleDegrees)
-		{
-			const FTrajectorySample CurrentSample = Trajectory.Samples[SampleIdx];
-
-			const int32 LastSampleIdx = Trajectory.Samples.Num() - 1;
-			const FTrajectorySample LastSample = Trajectory.Samples[LastSampleIdx];
-
-			const float CurrentToLastRotSpeed =
-				CalcRotationSpeed(SampleIdx, LastSampleIdx);
-			const float CurrentSpeedDelta = CurrentToLastRotSpeed - RotSpeedToExtrapolate;
-			const float CurrentAngleDelta = CurrentSpeedDelta * LastSample.AccumulatedSeconds;
-
-			const FVector CurrentFwd = CurrentSample.Transform.GetRotation().RotateVector(ForwardAxis);
-			const FQuat CurrentVelToFwdRotation = FQuat::FindBetweenVectors(CurrentSample.LinearVelocity, CurrentFwd);
-
-			const FVector LastFwd = LastSample.Transform.GetRotation().RotateVector(ForwardAxis);
-			const FQuat LastVelToFwdRotation = FQuat::FindBetweenVectors(LastSample.LinearVelocity, LastFwd);
-
-			const FQuat VelAlignmentDelta = CurrentVelToFwdRotation.Inverse() * LastVelToFwdRotation;
-			const float VelAlignmentDeltaAngle = VelAlignmentDelta.GetTwistAngle(TurnAxis);
-
-			const float MaxAlignmentAngleRadians = FMath::DegreesToRadians(MaxAlignmentAngleDegrees);
-			bool bIsAligned =
-				(FMath::Abs(CurrentAngleDelta) < MaxAlignmentAngleRadians) &&
-				(FMath::Abs(VelAlignmentDeltaAngle) < MaxAlignmentAngleRadians);
-
-			return !bIsAligned;
-		}
-
-		bool FindTurnBeyondExtrapolation(float MinSharpTurnAngleDegrees)
-		{
-			check(Trajectory.Samples.Num() > 1);
-
-			const float MinSharpTurnAngleRadians = FMath::DegreesToRadians(MinSharpTurnAngleDegrees);
-
-			const int32 LastSampleIdx = Trajectory.Samples.Num() - 1;
-			const float FirstToLastRotSpeed = CalcRotationSpeed(0, LastSampleIdx);
-
-			const float SignLastDirRotSpeed = RotSpeedToExtrapolate >= 0 ? 1.0f : -1.0f; // FMath::Sign(RotSpeedToExtrapolate);
-			const float SignedFirstToLastRotSpeed = SignLastDirRotSpeed * FirstToLastRotSpeed;
-
-			// first we'll look for the most aggressive turning frame with respect to the circling speed
-			float FastestRelativeTurnSpeed = 0.0f;
-			int32 FastestRelativeTurnSampleIdx = FindFastestRelativeSpeedSample(FastestRelativeTurnSpeed);
-			if (FastestRelativeTurnSampleIdx == INDEX_NONE)
-			{
-				return false;
-			}
-
-			// Starting from the root sample containing the fastest rotation speed, grow the evaluated interval, from
-			// StartSampleIdx to EndSampleIdx, in a greedy approach. This should lead to the shortest interval that 
-			// contains a large enough turn
-			const FTrajectorySample& RootSample = Trajectory.Samples[FastestRelativeTurnSampleIdx];
-			int32 EndSampleIdx = FastestRelativeTurnSampleIdx;
-			int32 StartSampleIdx = EndSampleIdx - 1;
-			const float TurnSpeedIntervalRadius = RotSpeedToExtrapolate / 2.0f;
-			bool bDone = false;
-			while (!bDone)
-			{
-				const float CurrentTurnLength = 
-					GetAccumulatedDomainValue(Trajectory.Samples[EndSampleIdx]) - 
-					GetAccumulatedDomainValue(Trajectory.Samples[StartSampleIdx]);
-
-				const int32 PrevIdx = StartSampleIdx - 1;
-				float PrevRelativeTurn = 0.0f;
-				bool bPrevValid = IsSampleContributingToSharpTurn(
-					StartSampleIdx,
-					CurrentTurnLength,
-					FastestRelativeTurnSpeed,
-					PrevRelativeTurn);
-
-				const int32 NextIdx = EndSampleIdx + 1;
-				float NextRelativeTurn = 0.0f;
-				bool bNextValid = IsSampleContributingToSharpTurn(
-					NextIdx,
-					CurrentTurnLength,
-					FastestRelativeTurnSpeed,
-					NextRelativeTurn);
-
-				// Every iteration one of 3 things happen:
-				// 1. The interval grows 1 sample into the past
-				// 2. The interval grows 1 sample into the future
-				// 3. The evaluation ends because either a large enough turn was found or there are no more valid
-				//    samples to add to the interval
-				if (bNextValid && bPrevValid)
-				{
-					// if both end points are contributing to the sharp turn, expand the interval towards the biggest
-					// contribution
-					if (FMath::Abs(NextRelativeTurn) >= FMath::Abs(PrevRelativeTurn))
-					{
-						EndSampleIdx = NextIdx;
-					}
-					else
-					{
-						StartSampleIdx = PrevIdx;
-					}
-				}
-				else if (bNextValid)
-				{
-					EndSampleIdx = NextIdx;
-				}
-				else if (bPrevValid)
-				{
-					StartSampleIdx = PrevIdx;
-				}
-				else
-				{
-					// if none of the evaluated extremities contribute to the turn, stop increasing the turn interval.
-					bDone = true;
-				}
-
-				const float ExtrapolatedTurn =
-					RotSpeedToExtrapolate *
-					(Trajectory.Samples[EndSampleIdx].AccumulatedSeconds -
-					 Trajectory.Samples[StartSampleIdx].AccumulatedSeconds);
-				const float TotalTurn =
-					AccumulatedTurnAxisRotations[EndSampleIdx] - AccumulatedTurnAxisRotations[StartSampleIdx];
-
-				if ((SignLastDirRotSpeed * TotalTurn) > (SignLastDirRotSpeed * ExtrapolatedTurn + MinSharpTurnAngleRadians) ||
-					(SignLastDirRotSpeed * TotalTurn) < -MinSharpTurnAngleRadians)
-				{
-					// at any moment if a big enough turn is found, return true, there's no need to do anything else.
-					return true;
-				}
-			}
-
-			return false;
-		}
-
-	private:
-
-		float CalcRotationSpeed(int32 StartIdx, int32 EndIdx)
-		{
-			const FTrajectorySample& SampleA = Trajectory.Samples[StartIdx];
-			const FTrajectorySample& SampleB = Trajectory.Samples[EndIdx];
-			const float TotalTime = SampleB.AccumulatedSeconds - SampleA.AccumulatedSeconds;
-			const float TotalAngleDelta = AccumulatedTurnAxisRotations[EndIdx] - AccumulatedTurnAxisRotations[StartIdx];
-			const float RotationSpeed = TotalAngleDelta / TotalTime;
-			return RotationSpeed;
-		}
-
-		// Populate arrays with turn information. If these end up being used multiple times per frame, we should consider
-		// adding this information to FTrajectorySample (maybe deferred until needed)
-		void CalcTurnData()
-		{
-			AccumulatedTurnAxisRotations.SetNum(Trajectory.Samples.Num());
-			AccumulatedTurnAxisRotations[0] = 0.0f;
-
-			ImmediateTurnAxisRotations.SetNum(Trajectory.Samples.Num());
-			ImmediateTurnAxisRotations[0] = 0.0f;
-
-			const int32 NumSamples = Trajectory.Samples.Num();
-			for (int32 SampleIdx = 1; SampleIdx < NumSamples; ++SampleIdx)
-			{
-				const FTrajectorySample& PrevSample = Trajectory.Samples[SampleIdx - 1];
-				const FTrajectorySample& CurrSample = Trajectory.Samples[SampleIdx];
-
-				FQuat VelocityDelta = FQuat::FindBetweenVectors(PrevSample.LinearVelocity, CurrSample.LinearVelocity);
-				FVector DeltaAxis;
-				float DeltaAngle = VelocityDelta.GetTwistAngle(TurnAxis);
-
-				ImmediateTurnAxisRotations[SampleIdx] = DeltaAngle;
-				AccumulatedTurnAxisRotations[SampleIdx] = AccumulatedTurnAxisRotations[SampleIdx - 1] + DeltaAngle;
-			}
-		}
-
-		void CalcRotSpeedToExtrapolate()
-		{
-			check(Trajectory.Samples.Num() > 2);
-
-			const int32 LastSampleIdx = Trajectory.Samples.Num() - 1;
-			const FTrajectorySample& LastSample = Trajectory.Samples[LastSampleIdx];
-
-			const int32 BeforeLastSampleIdx = LastSampleIdx - 1;
-			const FTrajectorySample& BeforeLastSample = Trajectory.Samples[BeforeLastSampleIdx];
-
-			const float LastDeltaSeconds = LastSample.AccumulatedSeconds - BeforeLastSample.AccumulatedSeconds;
-
-			const FQuat LastDirectionRotation =
-				FQuat::FindBetweenVectors(BeforeLastSample.LinearVelocity, LastSample.LinearVelocity);
-			const float LastDirRotAngle = LastDirectionRotation.GetTwistAngle(TurnAxis);
-
-			RotSpeedToExtrapolate = LastDirRotAngle / LastDeltaSeconds;
-		}
-
-		float GetAccumulatedDomainValue(const FTrajectorySample& Sample)
-		{
-			return Sample.AccumulatedSeconds;
-		}
-
-		// A sample is contributing to the sharp turn if it is turning to the same side as FastestRelativeTurnSpeed,
-		// if adding the sample does not cause the turn to grow larger than the specified constraint, and the turn is
-		// either a. sharper than the extrapolated rotation speed or b. to the opposite side of the extrapolated
-		// rotation speed.
-		bool IsSampleContributingToSharpTurn(
-			int32 SampleIdx,
-			float CurrentTurnLength,
-			float FastestRelativeTurnSpeed, 
-			float& OutSampleRelativeTurn)
-		{
-			OutSampleRelativeTurn = 0.0f;
-
-			bool bIsContributing = false;
-
-			if (SampleIdx >= 1 && SampleIdx < Trajectory.Samples.Num())
-			{
-				const int32 PrevIdx = SampleIdx - 1;
-
-				const FTrajectorySample& Sample = Trajectory.Samples[SampleIdx];
-				const float StartDomainValue = GetAccumulatedDomainValue(Sample);
-
-				const FTrajectorySample& PrevSample = Trajectory.Samples[PrevIdx];
-				const float PrevDomainValue = GetAccumulatedDomainValue(PrevSample);
-
-				const float SampleLengthContribution = StartDomainValue - PrevDomainValue;
-				const float NewTurnLength = SampleLengthContribution + CurrentTurnLength;
-
-				if (NewTurnLength <= RotationConstraintLimit)
-				{
-					const float DeltaTime =
-						Trajectory.Samples[SampleIdx].AccumulatedSeconds -
-						Trajectory.Samples[PrevIdx].AccumulatedSeconds;
-					const float PrevTurnSpeed = ImmediateTurnAxisRotations[SampleIdx] / DeltaTime;
-
-					const float TurnSpeedIntervalRadius = RotSpeedToExtrapolate / 2.0f;
-					OutSampleRelativeTurn = PrevTurnSpeed - TurnSpeedIntervalRadius;
-					bIsContributing =
-						FMath::Sign(OutSampleRelativeTurn) == FMath::Sign(FastestRelativeTurnSpeed) &&
-						FMath::Abs(OutSampleRelativeTurn) > FMath::Abs(TurnSpeedIntervalRadius);
-				}
-			}
-
-			return bIsContributing;
-		}
-
-		int32 FindFastestRelativeSpeedSample(float& OutFastestRelativeTurnSpeed)
-		{
-			int32 FastestRelativeTurnSampleIdx = INDEX_NONE;
-			OutFastestRelativeTurnSpeed = 0.0f;
-			float AbsFastestRelativeTurnSpeed = 0.0f;
-
-			// we'll compare turn speeds with the center of the speed interval [0; RotSpeedToInterpolate].
-			// A turn that falls inside the interval will be disregarded.
-			const float TurnSpeedIntervalRadius = RotSpeedToExtrapolate / 2.0f;
-			const int32 NumSamples = Trajectory.Samples.Num();
-			for (int32 SampleIdx = 1; SampleIdx < NumSamples; ++SampleIdx)
-			{
-				const float DeltaTime =
-					Trajectory.Samples[SampleIdx].AccumulatedSeconds -
-					Trajectory.Samples[SampleIdx - 1].AccumulatedSeconds;
-				const float ImmediateSpeed = ImmediateTurnAxisRotations[SampleIdx] / DeltaTime;
-				const float RelativeSpeed = ImmediateSpeed - TurnSpeedIntervalRadius;
-				const float AbsRelativeSpeed = FMath::Abs(RelativeSpeed);
-				if (AbsRelativeSpeed > AbsFastestRelativeTurnSpeed)
-				{
-					FastestRelativeTurnSampleIdx = SampleIdx;
-					OutFastestRelativeTurnSpeed = RelativeSpeed;
-					AbsFastestRelativeTurnSpeed = FMath::Abs(OutFastestRelativeTurnSpeed);
-				}
-			}
-
-			if (AbsFastestRelativeTurnSpeed <= FMath::Abs(TurnSpeedIntervalRadius))
-			{
-				FastestRelativeTurnSampleIdx = INDEX_NONE;
-			}
-
-			return FastestRelativeTurnSampleIdx;
-		}
-
-	};
-} // namespace UE::MotionTrajectory
-
-
-bool UMotionTrajectoryBlueprintLibrary::IsSharpVelocityDirChange(
-	const FTrajectorySampleRange& Trajectory,
-	float MinSharpTurnAngleDegrees,
-	float RotationConstraintValue,
-	float MaxAlignmentAngleDegrees,
-	float MinLinearSpeed,
-	FVector TurnAxis,
-	FVector ForwardAxis)
-{
-	if (Trajectory.Samples.Num() < 2)
-	{
-		// not enough samples to evaluate the trajectory
-		return false;
-	}
-	const float SquaredMinLinearSpeed = MinLinearSpeed * MinLinearSpeed;
-	const int32 LastSampleIdx = Trajectory.Samples.Num() - 1;
-
-	int32 PresentIdx = 0;
-	const FTrajectorySample& PresentSample = FTrajectorySampleRange::IterSampleTrajectory(Trajectory.Samples, 0.f, PresentIdx);
-
-	if (PresentIdx >= LastSampleIdx)
-	{
-		// we need at least one more sample in the future
-		return false;
-	}
-
-	const FTrajectorySample& LastSample = Trajectory.Samples[LastSampleIdx];
-	const float SquaredLastSampleSpeed = LastSample.LinearVelocity.SquaredLength();
-
-	const FTrajectorySample& FirstSample = Trajectory.Samples[0];
-	const float SquaredFirstSampleSpeed = FirstSample.LinearVelocity.SquaredLength();
-
-	if (SquaredLastSampleSpeed < SquaredMinLinearSpeed || SquaredFirstSampleSpeed < SquaredMinLinearSpeed)
-	{
-		// trajectory end points not fast enough
-		return false;
-	}
-
-	const FTrajectorySample& BeforeLastSample = Trajectory.Samples[LastSampleIdx - 1];
-
-	const float LastDeltaSeconds = LastSample.AccumulatedSeconds - BeforeLastSample.AccumulatedSeconds;
-	if (LastDeltaSeconds < SMALL_NUMBER)
-	{
-		// delta too small to evaluate
-		return false;
-	}
-
-	UE::MotionTrajectory::FTurnEvaluator TurnEvaluator = UE::MotionTrajectory::FTurnEvaluator(Trajectory, RotationConstraintValue, TurnAxis);
-
-	if (!TurnEvaluator.IsTurning(PresentIdx, ForwardAxis, MaxAlignmentAngleDegrees))
-	{
-		return false;
-	}
-
-	const bool bIsSharpTurn = TurnEvaluator.FindTurnBeyondExtrapolation(MinSharpTurnAngleDegrees);
-	return bIsSharpTurn;
-}
+}
\ No newline at end of file
diff --git a/Engine/Plugins/Experimental/Animation/MotionTrajectory/Source/MotionTrajectory/Public/MotionTrajectory.h b/Engine/Plugins/Experimental/Animation/MotionTrajectory/Source/MotionTrajectory/Public/MotionTrajectory.h
--- a/Engine/Plugins/Experimental/Animation/MotionTrajectory/Source/MotionTrajectory/Public/MotionTrajectory.h
+++ b/Engine/Plugins/Experimental/Animation/MotionTrajectory/Source/MotionTrajectory/Public/MotionTrajectory.h
@@ -100,10 +100,6 @@ public:
 	UFUNCTION(BlueprintPure, Category="Motion Trajectory", meta=(AutoCreateRefTerm="Settings", DisplayName="Get Trajectory (With Settings)"))
 	virtual FTrajectorySampleRange GetTrajectoryWithSettings(const FMotionTrajectorySettings& Settings, bool bIncludeHistory) const;
 	
-	// Changes the history and prediction sample rate for the component
-	UFUNCTION(BlueprintCallable, Category="Motion Trajectory")
-	void SetSampleRate(int32 Rate);
-
 	// Retrieves the historical trajectory
 	UFUNCTION(BlueprintCallable, Category="Motion Trajectory")
 	FTrajectorySampleRange GetHistory() const;
@@ -116,7 +112,7 @@ public:
 	UPROPERTY(EditDefaultsOnly, BlueprintReadOnly, Category="Settings")
 	FMotionTrajectorySettings HistorySettings;
 
-	// The trajectory sampling rate for both prediction and history
+	// The trajectory sampling rate for prediction
 	UPROPERTY(EditDefaultsOnly, BlueprintReadOnly, Category="Settings", meta=(ClampMin="5", ClampMax="120"))
 	int32 SampleRate = 30;
 
@@ -130,16 +126,6 @@ public:
 	UPROPERTY(EditDefaultsOnly, BlueprintReadOnly, Category="Settings")
 	bool bPredictionIncludesHistory = true;
 
-	// Determines if history will be sampled by the specified Sample Rate (true) or the Component Tick Rate (false)
-	// Consider disabling this option when the engine frame rate and component sample rate are the same
-	UPROPERTY(EditDefaultsOnly, BlueprintReadOnly, Category="Settings")
-	bool bUniformSampledHistory = true;
-
-	// History samples will use Centripetal Catmull–Rom spline interpolation (true) (slow, accurate) instead of linear interpolation (false) (fast, approximate)
-	// Consider enabling this option when accurate trajectory sample speed at low tick rates (high interpolation usage) is important
-	UPROPERTY(EditAnywhere, BlueprintReadWrite, Category="Settings")
-	bool bSmoothInterpolation = false;
-
 #if WITH_EDITORONLY_DATA
 	// Debug drawing options for the complete trajectory range
 	// a.MotionTrajectory.Debug (0/1) to enable or disable
diff --git a/Engine/Plugins/Experimental/Animation/MotionTrajectory/Source/MotionTrajectory/Public/MotionTrajectoryLibrary.h b/Engine/Plugins/Experimental/Animation/MotionTrajectory/Source/MotionTrajectory/Public/MotionTrajectoryLibrary.h
--- a/Engine/Plugins/Experimental/Animation/MotionTrajectory/Source/MotionTrajectory/Public/MotionTrajectoryLibrary.h
+++ b/Engine/Plugins/Experimental/Animation/MotionTrajectory/Source/MotionTrajectory/Public/MotionTrajectoryLibrary.h
@@ -128,35 +128,4 @@ public:
 			const FTrajectorySampleRange& Trajectory,
 			float Speed = 0.0f,
 			float Tolerance = 0.001f);
-	/**
-	 * Returns true if the trajectory has a sharp velocity direction change. This function will compare the total 
-	 * turning of the trajectory against the extrapolation of the final angular velocity in the trajectory. This ensures
-	 * the function distinguishes sharp turns from smooth circling. Note this is not detecting facing changes,
-	 * this function will return true in a strafing turn in which the character is always facing the same direction.
-	 * 
-	 * @param Trajectory				Trajectory being evaluated.
-	 * @param MinSharpTurnAngleDegrees	How many degrees of turning must be in the trajectory that can't be explained by 
-	 *									extrapolating the angular velocity at the end of the trajectory.
-	 * @param RotationConstraintDomain	Specifies if the turn must happen unconstrained, within a given distance or
-	 *									within a given time
-	 * @param RotationConstraintValue	If RotationConstraintDomain is not None, this value will specify the constraint
-	 *									threshold
-	 * @param MaxAlignmentAngleDegrees	If the current trajectory sample velocity and facing are well aligned with the 
-	 *									last point in the trajectory future, there's no turn.
-	 * @param MinLinearSpeed			Minimum linear speed at the end points of the trajectory required for this
-	 *									function to return true, so it's possible to ignore turning starts/stops
-	 * @param TurnAxis					Specifies the axis along which turns will be evaluated.
-	 * @param ForwardAxis				Specifies the trajectory forward angle, used to evaluate alignment.
-	 *									
-	 * @return True if a sharp turn is detected, false otherwise.
-	 */
-	UFUNCTION(BlueprintPure, Category = "Motion Trajectory", meta = (BlueprintThreadSafe))
-	static bool IsSharpVelocityDirChange(
-		const FTrajectorySampleRange& Trajectory,
-		float MinSharpTurnAngleDegrees = 45.0f,
-		float RotationConstraintValue = 0.0f,
-		float MaxAlignmentAngleDegrees = 5.0f,
-		float MinLinearSpeed = 1.0f,
-		FVector TurnAxis = FVector::UpVector,
-		FVector ForwardAxis = FVector::RightVector);
 };
\ No newline at end of file
diff --git a/Engine/Plugins/Experimental/Animation/PoseSearch/Source/Editor/Private/PoseSearchDebuggerViewModel.cpp b/Engine/Plugins/Experimental/Animation/PoseSearch/Source/Editor/Private/PoseSearchDebuggerViewModel.cpp
--- a/Engine/Plugins/Experimental/Animation/PoseSearch/Source/Editor/Private/PoseSearchDebuggerViewModel.cpp
+++ b/Engine/Plugins/Experimental/Animation/PoseSearch/Source/Editor/Private/PoseSearchDebuggerViewModel.cpp
@@ -461,20 +461,21 @@ const USkinnedMeshComponent* FDebuggerViewModel::GetMeshComponent() const
 void FDebuggerViewModel::FillCompactPoseAndComponentRefRotations()
 {
 	bool bResetMirrorBonesAndRotations = true;
-	const UPoseSearchDatabase* Database = GetCurrentDatabase();
-	if (Database)
+	if (const UPoseSearchDatabase* Database = GetCurrentDatabase())
 	{
-		UMirrorDataTable* MirrorDataTable = Database->Schema->MirrorDataTable;
-		if (MirrorDataTable != nullptr && Skeletons[ActivePose].Component->RequiredBones.IsValid())
+		if (UMirrorDataTable* MirrorDataTable = Database->Schema->MirrorDataTable)
 		{
-			if (CompactPoseMirrorBones.Num() == 0 || ComponentSpaceRefRotations.Num() == 0)
+			if (UPoseSearchMeshComponent* MeshComponent = Skeletons[ActivePose].Component.Get())
 			{
-				MirrorDataTable->FillCompactPoseAndComponentRefRotations(
-					Skeletons[ActivePose].Component->RequiredBones,
-					CompactPoseMirrorBones,
-					ComponentSpaceRefRotations);
+				if (MeshComponent->RequiredBones.IsValid())
+				{
+					if (CompactPoseMirrorBones.Num() == 0 || ComponentSpaceRefRotations.Num() == 0)
+					{
+						MirrorDataTable->FillCompactPoseAndComponentRefRotations(MeshComponent->RequiredBones, CompactPoseMirrorBones, ComponentSpaceRefRotations);
+					}
+					bResetMirrorBonesAndRotations = false;
+				}
 			}
-			bResetMirrorBonesAndRotations = false;
 		}
 	}
 
diff --git a/Engine/Plugins/Experimental/Animation/PoseSearch/Source/Runtime/Private/PoseSearchAssetIndexer.cpp b/Engine/Plugins/Experimental/Animation/PoseSearch/Source/Runtime/Private/PoseSearchAssetIndexer.cpp
--- a/Engine/Plugins/Experimental/Animation/PoseSearch/Source/Runtime/Private/PoseSearchAssetIndexer.cpp
+++ b/Engine/Plugins/Experimental/Animation/PoseSearch/Source/Runtime/Private/PoseSearchAssetIndexer.cpp
@@ -118,9 +118,12 @@ bool FAssetIndexer::Process()
 	if (IndexingContext.Schema->SchemaCardinality > 0)
 	{
 		// Index each channel
-		for (int32 ChannelIdx = 0; ChannelIdx != IndexingContext.Schema->Channels.Num(); ++ChannelIdx)
+		for (const TObjectPtr<UPoseSearchFeatureChannel>& ChannelPtr : IndexingContext.Schema->Channels)
 		{
-			IndexingContext.Schema->Channels[ChannelIdx]->IndexAsset(*this, Output.FeatureVectorTable);
+			if (ChannelPtr)
+			{
+				ChannelPtr->IndexAsset(*this, Output.FeatureVectorTable);
+			}
 		}
 	}
 
diff --git a/Engine/Plugins/Experimental/Animation/PoseSearch/Source/Runtime/Private/PoseSearchContext.cpp b/Engine/Plugins/Experimental/Animation/PoseSearch/Source/Runtime/Private/PoseSearchContext.cpp
--- a/Engine/Plugins/Experimental/Animation/PoseSearch/Source/Runtime/Private/PoseSearchContext.cpp
+++ b/Engine/Plugins/Experimental/Animation/PoseSearch/Source/Runtime/Private/PoseSearchContext.cpp
@@ -71,9 +71,41 @@ const UPoseSearchSchema* FDebugDrawParams::GetSchema() const
 	return Database ? Database->Schema : nullptr;
 }
 
-void DrawFeatureVector(const FDebugDrawParams& DrawParams, TConstArrayView<float> PoseVector)
+void FDebugDrawParams::ClearCachedPositions()
+{
+	CachedPositions.Reset();
+}
+
+void FDebugDrawParams::AddCachedPosition(float TimeOffset, int8 SchemaBoneIdx, const FVector& Position)
+{
+	if (const UPoseSearchSchema* Schema = GetSchema())
+	{
+		CachedPositions.Add(TimeOffset, Schema->GetBoneIndexType(SchemaBoneIdx), Position);
+	}
+}
+
+FVector FDebugDrawParams::GetCachedPosition(float TimeOffset, int8 SchemaBoneIdx) const
+{
+	if (const UPoseSearchSchema* Schema = GetSchema())
+	{
+		if (auto CachedPosition = CachedPositions.Find(TimeOffset, Schema->GetBoneIndexType(SchemaBoneIdx)))
+		{
+			return CachedPosition->Transform;
+		}
+
+		if (Mesh.IsValid() && SchemaBoneIdx >= 0)
+		{
+			return Mesh->GetSocketTransform(Schema->BoneReferences[SchemaBoneIdx].BoneName).GetLocation();
+		}
+	}
+	return RootTransform.GetTranslation();
+}
+
+void DrawFeatureVector(FDebugDrawParams& DrawParams, TConstArrayView<float> PoseVector)
 {
 #if ENABLE_DRAW_DEBUG
+	DrawParams.ClearCachedPositions();
+
 	if (DrawParams.CanDraw())
 	{
 		const UPoseSearchSchema* Schema = DrawParams.GetSchema();
@@ -81,11 +113,19 @@ void DrawFeatureVector(const FDebugDrawParams& DrawParams, TConstArrayView<float
 
 		if (PoseVector.Num() == Schema->SchemaCardinality)
 		{
-			for (int32 ChannelIdx = 0; ChannelIdx != Schema->Channels.Num(); ++ChannelIdx)
+			for (const TObjectPtr<UPoseSearchFeatureChannel>& ChannelPtr : Schema->Channels)
+			{
+				if (ChannelPtr)
+				{
+					ChannelPtr->PreDebugDraw(DrawParams, PoseVector);
+				}
+			}
+
+			for (const TObjectPtr<UPoseSearchFeatureChannel>& ChannelPtr : Schema->Channels)
 			{
-				if (DrawParams.ChannelMask & (1 << ChannelIdx))
+				if (ChannelPtr)
 				{
-					Schema->Channels[ChannelIdx]->DebugDraw(DrawParams, PoseVector);
+					ChannelPtr->DebugDraw(DrawParams, PoseVector);
 				}
 			}
 		}
@@ -93,7 +133,7 @@ void DrawFeatureVector(const FDebugDrawParams& DrawParams, TConstArrayView<float
 #endif // ENABLE_DRAW_DEBUG
 }
 
-void DrawFeatureVector(const FDebugDrawParams& DrawParams, int32 PoseIdx)
+void DrawFeatureVector(FDebugDrawParams& DrawParams, int32 PoseIdx)
 {
 #if ENABLE_DRAW_DEBUG
 	// if we're editing the schema while in PIE with Rewind Debugger active, PoseIdx could be out of bound / stale
@@ -110,21 +150,13 @@ FTransform FSearchContext::TryGetTransformAndCacheResults(float SampleTime, cons
 {
 	check(History && Schema);
 
-	static constexpr FBoneIndexType RootBoneIdx = 0xFFFF;
-	const FBoneIndexType BoneIndexType = SchemaBoneIdx >= 0 && Schema->BoneReferences[SchemaBoneIdx].HasValidSetup() ? Schema->BoneReferences[SchemaBoneIdx].BoneIndex : RootBoneIdx;
-
-	// @todo: use an hashmap if we end up having too many entries
-	const FCachedEntry* Entry = CachedEntries.FindByPredicate([SampleTime, BoneIndexType](const FSearchContext::FCachedEntry& Entry)
+	const FBoneIndexType BoneIndexType = Schema->GetBoneIndexType(SchemaBoneIdx);
+	if (const FCachedTransform<FTransform>* CachedTransform = CachedTransforms.Find(SampleTime, BoneIndexType))
 	{
-		return Entry.SampleTime == SampleTime && Entry.BoneIndexType == BoneIndexType;
-	});
-
-	if (Entry)
-	{
-		return Entry->Transform;
+		return CachedTransform->Transform;
 	}
 
-	if (BoneIndexType != RootBoneIdx)
+	if (BoneIndexType != UPoseSearchSchema::RootBoneIdx)
 	{
 		TArray<FTransform> SampledLocalPose;
 		if (History->TrySampleLocalPose(-SampleTime, &Schema->BoneIndicesWithParents, &SampledLocalPose, nullptr))
@@ -135,7 +167,7 @@ FTransform FSearchContext::TryGetTransformAndCacheResults(float SampleTime, cons
 			// adding bunch of entries, without caring about adding eventual duplicates
 			for (const FBoneIndexType NewEntryBoneIndexType : Schema->BoneIndicesWithParents)
 			{
-				CachedEntries.Emplace(SampleTime, SampledComponentPose[NewEntryBoneIndexType], NewEntryBoneIndexType);
+				CachedTransforms.Add(SampleTime, NewEntryBoneIndexType, SampledComponentPose[NewEntryBoneIndexType]);
 			}
 
 			return SampledComponentPose[BoneIndexType];
@@ -147,11 +179,7 @@ FTransform FSearchContext::TryGetTransformAndCacheResults(float SampleTime, cons
 	FTransform SampledRootTransform;
 	if (History->TrySampleLocalPose(-SampleTime, nullptr, nullptr, &SampledRootTransform))
 	{
-		FCachedEntry& NewEntry = CachedEntries[CachedEntries.AddDefaulted()];
-		NewEntry.SampleTime = SampleTime;
-		NewEntry.BoneIndexType = BoneIndexType;
-		NewEntry.Transform = SampledRootTransform;
-
+		CachedTransforms.Add(SampleTime, BoneIndexType, SampledRootTransform);
 		return SampledRootTransform;
 	}
 	
@@ -160,7 +188,7 @@ FTransform FSearchContext::TryGetTransformAndCacheResults(float SampleTime, cons
 
 void FSearchContext::ClearCachedEntries()
 {
-	CachedEntries.Reset();
+	CachedTransforms.Reset();
 }
 
 void FSearchContext::ResetCurrentBestCost()
diff --git a/Engine/Plugins/Experimental/Animation/PoseSearch/Source/Runtime/Private/PoseSearchDatabase.cpp b/Engine/Plugins/Experimental/Animation/PoseSearch/Source/Runtime/Private/PoseSearchDatabase.cpp
--- a/Engine/Plugins/Experimental/Animation/PoseSearch/Source/Runtime/Private/PoseSearchDatabase.cpp
+++ b/Engine/Plugins/Experimental/Animation/PoseSearch/Source/Runtime/Private/PoseSearchDatabase.cpp
@@ -151,7 +151,7 @@ struct FPoseFilters
 
 		for (const IPoseFilter* ChannelPoseFilter : Schema->Channels)
 		{
-			if (ChannelPoseFilter->IsPoseFilterActive())
+			if (ChannelPoseFilter && ChannelPoseFilter->IsPoseFilterActive())
 			{
 				AllPoseFilters.Add(ChannelPoseFilter);
 			}
diff --git a/Engine/Plugins/Experimental/Animation/PoseSearch/Source/Runtime/Private/PoseSearchDerivedData.cpp b/Engine/Plugins/Experimental/Animation/PoseSearch/Source/Runtime/Private/PoseSearchDerivedData.cpp
--- a/Engine/Plugins/Experimental/Animation/PoseSearch/Source/Runtime/Private/PoseSearchDerivedData.cpp
+++ b/Engine/Plugins/Experimental/Animation/PoseSearch/Source/Runtime/Private/PoseSearchDerivedData.cpp
@@ -126,10 +126,12 @@ static TArray<float> ComputeChannelsDeviations(TConstArrayView<FPoseSearchIndexB
 
 			FeatureChannelLayoutSet.CurrentSchemaIndex = SchemaIndex;
 			FeatureChannelLayoutSet.CurrentSchema = Schema;
-			for (int ChannelIdx = 0; ChannelIdx != Schema->Channels.Num(); ++ChannelIdx)
+			for (const TObjectPtr<UPoseSearchFeatureChannel>& ChannelPtr : Schema->Channels)
 			{
-				const UPoseSearchFeatureChannel* Channel = Schema->Channels[ChannelIdx].Get();
-				Channel->PopulateChannelLayoutSet(FeatureChannelLayoutSet);
+				if (ChannelPtr)
+				{
+					ChannelPtr->PopulateChannelLayoutSet(FeatureChannelLayoutSet);
+				}
 			}
 		}
 
@@ -307,10 +309,13 @@ static void PreprocessSearchIndexWeights(FPoseSearchIndex& SearchIndex, const UP
 {
 	const int32 NumDimensions = Schema->SchemaCardinality;
 	SearchIndex.WeightsSqrt.Init(1.f, NumDimensions);
-	for (int ChannelIdx = 0; ChannelIdx != Schema->Channels.Num(); ++ChannelIdx)
+
+	for (const TObjectPtr<UPoseSearchFeatureChannel>& ChannelPtr : Schema->Channels)
 	{
-		const UPoseSearchFeatureChannel* Channel = Schema->Channels[ChannelIdx].Get();
-		Channel->FillWeights(SearchIndex.WeightsSqrt);
+		if (ChannelPtr)
+		{
+			ChannelPtr->FillWeights(SearchIndex.WeightsSqrt);
+		}
 	}
 
 	EPoseSearchDataPreprocessor DataPreprocessor = Schema->DataPreprocessor;
@@ -679,6 +684,9 @@ void FPoseSearchDatabaseAsyncCacheTask::Wait(FCriticalSection& OuterMutex)
 	if (!bFailedIndexing)
 	{
 		Database->SetSearchIndex(SearchIndex); // @todo: implement FPoseSearchIndex move ctor and assignment operator and use a MoveTemp(SearchIndex) here
+
+		check(Database->Schema && Database->Schema->IsValid() && !SearchIndex.IsEmpty() && SearchIndex.WeightsSqrt.Num() == Database->Schema->SchemaCardinality && SearchIndex.KDTree.Impl);
+
 		SetState(EState::Ended);
 		bBroadcastOnDerivedDataRebuild = true;
 	}
@@ -718,6 +726,7 @@ void FPoseSearchDatabaseAsyncCacheTask::OnGetComplete(UE::DerivedData::FCacheGet
 	const FCacheKey FullIndexKey = Response.Record.GetKey();
 
 	// The database is part of the derived data cache and up to date, skip re-building it.
+	bool bCacheCorrupted = false;
 	if (Response.Status == EStatus::Ok)
 	{
 		COOK_STAT(auto Timer = UsageStats.TimeAsyncWait());
@@ -728,16 +737,30 @@ void FPoseSearchDatabaseAsyncCacheTask::OnGetComplete(UE::DerivedData::FCacheGet
 		FMemoryReaderView Reader(RawData);
 		Reader << SearchIndex;
 
-		UE_LOG(LogPoseSearch, Log, TEXT("%s - %s BuildIndex From Cache"), *LexToString(FullIndexKey.Hash), *Database->GetName());
+		check(Database->Schema && Database->Schema->IsValid());
+		// cache can be corrupted in case the version of the derived data cache has not being updated while 
+		// developing channels that changes their cardinality without impacting any asset properties
+		// so to account for this, we just reindex the database and update the associated DDC 
+		if (!SearchIndex.IsEmpty() && SearchIndex.WeightsSqrt.Num() == Database->Schema->SchemaCardinality && SearchIndex.KDTree.Impl)
+		{
+			UE_LOG(LogPoseSearch, Log, TEXT("%s - %s BuildIndex From Cache"), *LexToString(FullIndexKey.Hash), *Database->GetName());
+		}
+		else
+		{
+			UE_LOG(LogPoseSearch, Warning, TEXT("%s - %s BuildIndex From Cache Corrupted!"), *LexToString(FullIndexKey.Hash), *Database->GetName());
+			bCacheCorrupted = true;
+		}
 
 		COOK_STAT(Timer.AddHit(RawData.GetSize()));
 	}
-	else if (Response.Status == EStatus::Canceled)
+	
+	if (Response.Status == EStatus::Canceled)
 	{
 		SearchIndex.Reset();
 		UE_LOG(LogPoseSearch, Log, TEXT("%s - %s BuildIndex Cancelled"), *LexToString(FullIndexKey.Hash), *Database->GetName());
 	}
-	else if (Response.Status == EStatus::Error)
+	
+	if (Response.Status == EStatus::Error || bCacheCorrupted)
 	{
 		// we didn't find the cached data associated to the PendingDerivedDataKey: we'll BuildIndex to update SearchIndex and "Put" the data over the DDC
 		Owner.LaunchTask(TEXT("PoseSearchDatabaseBuild"), [this, FullIndexKey]
diff --git a/Engine/Plugins/Experimental/Animation/PoseSearch/Source/Runtime/Private/PoseSearchFeatureChannel.cpp b/Engine/Plugins/Experimental/Animation/PoseSearch/Source/Runtime/Private/PoseSearchFeatureChannel.cpp
--- a/Engine/Plugins/Experimental/Animation/PoseSearch/Source/Runtime/Private/PoseSearchFeatureChannel.cpp
+++ b/Engine/Plugins/Experimental/Animation/PoseSearch/Source/Runtime/Private/PoseSearchFeatureChannel.cpp
@@ -61,18 +61,18 @@ void FFeatureVectorHelper::EncodeQuat(TArrayView<float> Values, int32& DataOffse
 
 void FFeatureVectorHelper::EncodeQuat(TArrayView<float> Values, int32& DataOffset, TConstArrayView<float> PrevValues, TConstArrayView<float> CurValues, TConstArrayView<float> NextValues, float LerpValue)
 {
-	FQuat Quat = DecodeQuatInternal(CurValues, DataOffset);
+	FQuat Quat = DecodeQuatAtOffset(CurValues, DataOffset);
 
 	// linear interpolation
 	if (!FMath::IsNearlyZero(LerpValue))
 	{
 		if (LerpValue < 0.f)
 		{
-			Quat = FQuat::Slerp(Quat, DecodeQuatInternal(PrevValues, DataOffset), -LerpValue);
+			Quat = FQuat::Slerp(Quat, DecodeQuatAtOffset(PrevValues, DataOffset), -LerpValue);
 		}
 		else
 		{
-			Quat = FQuat::Slerp(Quat, DecodeQuatInternal(NextValues, DataOffset), LerpValue);
+			Quat = FQuat::Slerp(Quat, DecodeQuatAtOffset(NextValues, DataOffset), LerpValue);
 		}
 	}
 
@@ -82,12 +82,12 @@ void FFeatureVectorHelper::EncodeQuat(TArrayView<float> Values, int32& DataOffse
 
 FQuat FFeatureVectorHelper::DecodeQuat(TConstArrayView<float> Values, int32& DataOffset)
 {
-	const FQuat Quat = DecodeQuatInternal(Values, DataOffset);
+	const FQuat Quat = DecodeQuatAtOffset(Values, DataOffset);
 	DataOffset += EncodeQuatCardinality;
 	return Quat;
 }
 
-FQuat FFeatureVectorHelper::DecodeQuatInternal(TConstArrayView<float> Values, int32 DataOffset)
+FQuat FFeatureVectorHelper::DecodeQuatAtOffset(TConstArrayView<float> Values, int32 DataOffset)
 {
 	// reconstructing the chosen 2 axis (X,Y - from EncodeQuat) from the 6 floats from Values 
 	const FVector X(Values[DataOffset + 0], Values[DataOffset + 1], Values[DataOffset + 2]);
@@ -109,18 +109,18 @@ void FFeatureVectorHelper::EncodeVector(TArrayView<float> Values, int32& DataOff
 
 void FFeatureVectorHelper::EncodeVector(TArrayView<float> Values, int32& DataOffset, TConstArrayView<float> PrevValues, TConstArrayView<float> CurValues, TConstArrayView<float> NextValues, float LerpValue, bool bNormalize)
 {
-	FVector Vector = DecodeVectorInternal(CurValues, DataOffset);
+	FVector Vector = DecodeVectorAtOffset(CurValues, DataOffset);
 
 	// linear interpolation
 	if (!FMath::IsNearlyZero(LerpValue))
 	{
 		if (LerpValue < 0.f)
 		{
-			Vector = FMath::Lerp(Vector, DecodeVectorInternal(PrevValues, DataOffset), -LerpValue);
+			Vector = FMath::Lerp(Vector, DecodeVectorAtOffset(PrevValues, DataOffset), -LerpValue);
 		}
 		else
 		{
-			Vector = FMath::Lerp(Vector, DecodeVectorInternal(NextValues, DataOffset), LerpValue);
+			Vector = FMath::Lerp(Vector, DecodeVectorAtOffset(NextValues, DataOffset), LerpValue);
 		}
 	}
 
@@ -135,12 +135,12 @@ void FFeatureVectorHelper::EncodeVector(TArrayView<float> Values, int32& DataOff
 
 FVector FFeatureVectorHelper::DecodeVector(TConstArrayView<float> Values, int32& DataOffset)
 {
-	const FVector Vector = DecodeVectorInternal(Values, DataOffset);
+	const FVector Vector = DecodeVectorAtOffset(Values, DataOffset);
 	DataOffset += EncodeVectorCardinality;
 	return Vector;
 }
 
-FVector FFeatureVectorHelper::DecodeVectorInternal(TConstArrayView<float> Values, int32 DataOffset)
+FVector FFeatureVectorHelper::DecodeVectorAtOffset(TConstArrayView<float> Values, int32 DataOffset)
 {
 	return FVector(Values[DataOffset + 0], Values[DataOffset + 1], Values[DataOffset + 2]);
 }
@@ -154,18 +154,18 @@ void FFeatureVectorHelper::EncodeVector2D(TArrayView<float> Values, int32& DataO
 
 void FFeatureVectorHelper::EncodeVector2D(TArrayView<float> Values, int32& DataOffset, TConstArrayView<float> PrevValues, TConstArrayView<float> CurValues, TConstArrayView<float> NextValues, float LerpValue)
 {
-	FVector2D Vector2D = DecodeVector2DInternal(CurValues, DataOffset);
+	FVector2D Vector2D = DecodeVector2DAtOffset(CurValues, DataOffset);
 
 	// linear interpolation
 	if (!FMath::IsNearlyZero(LerpValue))
 	{
 		if (LerpValue < 0.f)
 		{
-			Vector2D = FMath::Lerp(Vector2D, DecodeVector2DInternal(PrevValues, DataOffset), -LerpValue);
+			Vector2D = FMath::Lerp(Vector2D, DecodeVector2DAtOffset(PrevValues, DataOffset), -LerpValue);
 		}
 		else
 		{
-			Vector2D = FMath::Lerp(Vector2D, DecodeVector2DInternal(NextValues, DataOffset), LerpValue);
+			Vector2D = FMath::Lerp(Vector2D, DecodeVector2DAtOffset(NextValues, DataOffset), LerpValue);
 		}
 	}
 
@@ -175,12 +175,12 @@ void FFeatureVectorHelper::EncodeVector2D(TArrayView<float> Values, int32& DataO
 
 FVector2D FFeatureVectorHelper::DecodeVector2D(TConstArrayView<float> Values, int32& DataOffset)
 {
-	const FVector2D Vector2D = DecodeVector2DInternal(Values, DataOffset);
+	const FVector2D Vector2D = DecodeVector2DAtOffset(Values, DataOffset);
 	DataOffset += EncodeVector2DCardinality;
 	return Vector2D;
 }
 
-FVector2D FFeatureVectorHelper::DecodeVector2DInternal(TConstArrayView<float> Values, int32 DataOffset)
+FVector2D FFeatureVectorHelper::DecodeVector2DAtOffset(TConstArrayView<float> Values, int32 DataOffset)
 {
 	return FVector2D(Values[DataOffset + 0], Values[DataOffset + 1]);
 }
@@ -193,18 +193,18 @@ void FFeatureVectorHelper::EncodeFloat(TArrayView<float> Values, int32& DataOffs
 
 void FFeatureVectorHelper::EncodeFloat(TArrayView<float> Values, int32& DataOffset, TConstArrayView<float> PrevValues, TConstArrayView<float> CurValues, TConstArrayView<float> NextValues, float LerpValue)
 {
-	float Value = DecodeFloatInternal(CurValues, DataOffset);
+	float Value = DecodeFloatAtOffset(CurValues, DataOffset);
 
 	// linear interpolation
 	if (!FMath::IsNearlyZero(LerpValue))
 	{
 		if (LerpValue < 0.f)
 		{
-			Value = FMath::Lerp(Value, DecodeFloatInternal(PrevValues, DataOffset), -LerpValue);
+			Value = FMath::Lerp(Value, DecodeFloatAtOffset(PrevValues, DataOffset), -LerpValue);
 		}
 		else
 		{
-			Value = FMath::Lerp(Value, DecodeFloatInternal(NextValues, DataOffset), LerpValue);
+			Value = FMath::Lerp(Value, DecodeFloatAtOffset(NextValues, DataOffset), LerpValue);
 		}
 	}
 
@@ -214,12 +214,12 @@ void FFeatureVectorHelper::EncodeFloat(TArrayView<float> Values, int32& DataOffs
 
 float FFeatureVectorHelper::DecodeFloat(TConstArrayView<float> Values, int32& DataOffset)
 {
-	const float Value = DecodeFloatInternal(Values, DataOffset);
+	const float Value = DecodeFloatAtOffset(Values, DataOffset);
 	DataOffset += EncodeFloatCardinality;
 	return Value;
 }
 
-float FFeatureVectorHelper::DecodeFloatInternal(TConstArrayView<float> Values, int32 DataOffset)
+float FFeatureVectorHelper::DecodeFloatAtOffset(TConstArrayView<float> Values, int32 DataOffset)
 {
 	return Values[DataOffset];
 }
diff --git a/Engine/Plugins/Experimental/Animation/PoseSearch/Source/Runtime/Private/PoseSearchFeatureChannel_FilterCrashingLegs.cpp b/Engine/Plugins/Experimental/Animation/PoseSearch/Source/Runtime/Private/PoseSearchFeatureChannel_FilterCrashingLegs.cpp
--- a/Engine/Plugins/Experimental/Animation/PoseSearch/Source/Runtime/Private/PoseSearchFeatureChannel_FilterCrashingLegs.cpp
+++ b/Engine/Plugins/Experimental/Animation/PoseSearch/Source/Runtime/Private/PoseSearchFeatureChannel_FilterCrashingLegs.cpp
@@ -27,7 +27,7 @@ static float ComputeCrashingLegsValue(const FVector& RightThighPos, const FVecto
 	return CrashingLegsValue;
 }
 
-void UPoseSearchFeatureChannel_FilterCrashingLegs::InitializeSchema(UPoseSearchSchema* Schema)
+void UPoseSearchFeatureChannel_FilterCrashingLegs::Finalize(UPoseSearchSchema* Schema)
 {
 	ChannelDataOffset = Schema->SchemaCardinality;
 	ChannelCardinality = UE::PoseSearch::FFeatureVectorHelper::EncodeFloatCardinality;
@@ -104,23 +104,22 @@ void UPoseSearchFeatureChannel_FilterCrashingLegs::DebugDraw(const UE::PoseSearc
 #if ENABLE_DRAW_DEBUG
 	using namespace UE::PoseSearch;
 
-	int32 DataOffset = ChannelDataOffset;
-	const float CrashingLegsValue = FFeatureVectorHelper::DecodeFloat(PoseVector, DataOffset);
+	const float CrashingLegsValue = FFeatureVectorHelper::DecodeFloatAtOffset(PoseVector, ChannelDataOffset);
 
 	const float LifeTime = DrawParams.DefaultLifeTime;
 	const uint8 DepthPriority = ESceneDepthPriorityGroup::SDPG_Foreground + 2;
 	const bool bPersistent = EnumHasAnyFlags(DrawParams.Flags, EDebugDrawFlags::Persistent);
 
-	const FTransform LeftThighTransform = DrawParams.Mesh != nullptr ? DrawParams.Mesh->GetSocketTransform(LeftThigh.BoneName) : DrawParams.RootTransform;
-	const FTransform RightThighTransform = DrawParams.Mesh != nullptr ? DrawParams.Mesh->GetSocketTransform(RightThigh.BoneName) : DrawParams.RootTransform;
-	const FTransform LeftFootTransform = DrawParams.Mesh != nullptr ? DrawParams.Mesh->GetSocketTransform(LeftFoot.BoneName) : DrawParams.RootTransform;
-	const FTransform RightFootTransform = DrawParams.Mesh != nullptr ? DrawParams.Mesh->GetSocketTransform(RightFoot.BoneName) : DrawParams.RootTransform;
-	
-	const float FeetDistance = (RightFootTransform.GetTranslation() - LeftFootTransform.GetTranslation()).Length();
+	const FVector LeftThighPosition = DrawParams.GetCachedPosition(0.f, LeftThighIdx);
+	const FVector RightThighPosition = DrawParams.GetCachedPosition(0.f, RightThighIdx);
+	const FVector LeftFootPosition = DrawParams.GetCachedPosition(0.f, LeftFootIdx);
+	const FVector RightFootPosition = DrawParams.GetCachedPosition(0.f, RightFootIdx);
+
+	const float FeetDistance = (RightFootPosition - LeftFootPosition).Length();
 
 	FVector ThighsSide;
 	FVector ThighsForward;
-	ComputeThighsSideAndForward(RightThighTransform.GetTranslation(), LeftThighTransform.GetTranslation(), ThighsSide, ThighsForward);
+	ComputeThighsSideAndForward(RightThighPosition, LeftThighPosition, ThighsSide, ThighsForward);
 	const FVector CrossingLegsVector = (FMath::Cos(CrashingLegsValue * UE_PI) * ThighsSide + FMath::Sin(CrashingLegsValue * UE_PI) * ThighsForward) * FeetDistance;
 
 	auto LerpColor = [](FColor A, FColor B, float T) -> FColor
@@ -138,8 +137,8 @@ void UPoseSearchFeatureChannel_FilterCrashingLegs::DebugDraw(const UE::PoseSearc
 		LerpColor(FColor::Green, FColor::Red, LerpValue) :
 		LerpColor(FColor::Green, FColor::Orange, -LerpValue);
 
-	DrawDebugLine(DrawParams.World, LeftFootTransform.GetTranslation(), LeftFootTransform.GetTranslation() + CrossingLegsVector, Color, bPersistent, LifeTime, DepthPriority);
-	DrawDebugLine(DrawParams.World, RightFootTransform.GetTranslation(), RightFootTransform.GetTranslation() - CrossingLegsVector, Color, bPersistent, LifeTime, DepthPriority);
+	DrawDebugLine(DrawParams.World, LeftFootPosition, LeftFootPosition + CrossingLegsVector, Color, bPersistent, LifeTime, DepthPriority);
+	DrawDebugLine(DrawParams.World, RightFootPosition, RightFootPosition - CrossingLegsVector, Color, bPersistent, LifeTime, DepthPriority);
 #endif // ENABLE_DRAW_DEBUG
 }
 
@@ -153,11 +152,8 @@ bool UPoseSearchFeatureChannel_FilterCrashingLegs::IsPoseValid(TConstArrayView<f
 {
 	using namespace UE::PoseSearch;
 
-	int32 DataOffset = ChannelDataOffset;
-	const float PoseValue = FFeatureVectorHelper::DecodeFloat(PoseValues, DataOffset);
-
-	DataOffset = ChannelDataOffset;
-	const float QueryValue = FFeatureVectorHelper::DecodeFloat(QueryValues, DataOffset);
+	const float PoseValue = FFeatureVectorHelper::DecodeFloatAtOffset(PoseValues, ChannelDataOffset);
+	const float QueryValue = FFeatureVectorHelper::DecodeFloatAtOffset(QueryValues, ChannelDataOffset);
 
 	bool bIsPoseValid = true;
 	if (FMath::Abs(QueryValue - PoseValue) > AllowedTolerance)
diff --git a/Engine/Plugins/Experimental/Animation/PoseSearch/Source/Runtime/Private/PoseSearchFeatureChannel_FilterCrashingLegs.h b/Engine/Plugins/Experimental/Animation/PoseSearch/Source/Runtime/Private/PoseSearchFeatureChannel_FilterCrashingLegs.h
--- a/Engine/Plugins/Experimental/Animation/PoseSearch/Source/Runtime/Private/PoseSearchFeatureChannel_FilterCrashingLegs.h
+++ b/Engine/Plugins/Experimental/Animation/PoseSearch/Source/Runtime/Private/PoseSearchFeatureChannel_FilterCrashingLegs.h
@@ -53,7 +53,7 @@ public:
 	float AllowedTolerance = 0.3f;
 
 	// UPoseSearchFeatureChannel interface
-	virtual void InitializeSchema(UPoseSearchSchema* Schema) override;
+	virtual void Finalize(UPoseSearchSchema* Schema) override;
 	virtual void FillWeights(TArray<float>& Weights) const override;
 	virtual void IndexAsset(UE::PoseSearch::IAssetIndexer& Indexer, TArrayView<float> FeatureVectorTable) const override;
 	virtual void BuildQuery(UE::PoseSearch::FSearchContext& SearchContext, FPoseSearchFeatureVectorBuilder& InOutQuery) const override;
diff --git a/Engine/Plugins/Experimental/Animation/PoseSearch/Source/Runtime/Private/PoseSearchFeatureChannel_Group.cpp b/Engine/Plugins/Experimental/Animation/PoseSearch/Source/Runtime/Private/PoseSearchFeatureChannel_Group.cpp
--- a/Engine/Plugins/Experimental/Animation/PoseSearch/Source/Runtime/Private/PoseSearchFeatureChannel_Group.cpp
+++ b/Engine/Plugins/Experimental/Animation/PoseSearch/Source/Runtime/Private/PoseSearchFeatureChannel_Group.cpp
@@ -4,22 +4,22 @@
 #include "DrawDebugHelpers.h"
 #include "PoseSearch/PoseSearchSchema.h"
 
-void UPoseSearchFeatureChannel_Group::InitializeSchema(UPoseSearchSchema* Schema)
+void UPoseSearchFeatureChannel_GroupBase::Finalize(UPoseSearchSchema* Schema)
 {
 	ChannelDataOffset = Schema->SchemaCardinality;
-	for (TObjectPtr<UPoseSearchFeatureChannel>& SubChannelPtr : SubChannels)
+	for (TObjectPtr<UPoseSearchFeatureChannel>& SubChannelPtr : GetSubChannels())
 	{
 		if (UPoseSearchFeatureChannel* SubChannel = SubChannelPtr.Get())
 		{
-			SubChannel->InitializeSchema(Schema);
+			SubChannel->Finalize(Schema);
 		}
 	}
 	ChannelCardinality = Schema->SchemaCardinality - ChannelDataOffset;
 }
 
-void UPoseSearchFeatureChannel_Group::FillWeights(TArray<float>& Weights) const
+void UPoseSearchFeatureChannel_GroupBase::FillWeights(TArray<float>& Weights) const
 {
-	for (const TObjectPtr<UPoseSearchFeatureChannel>& SubChannelPtr : SubChannels)
+	for (const TObjectPtr<UPoseSearchFeatureChannel>& SubChannelPtr : GetSubChannels())
 	{
 		if (const UPoseSearchFeatureChannel* SubChannel = SubChannelPtr.Get())
 		{
@@ -28,9 +28,9 @@ void UPoseSearchFeatureChannel_Group::FillWeights(TArray<float>& Weights) const
 	}
 }
 
-void UPoseSearchFeatureChannel_Group::IndexAsset(UE::PoseSearch::IAssetIndexer& Indexer, TArrayView<float> FeatureVectorTable) const
+void UPoseSearchFeatureChannel_GroupBase::IndexAsset(UE::PoseSearch::IAssetIndexer& Indexer, TArrayView<float> FeatureVectorTable) const
 {
-	for (const TObjectPtr<UPoseSearchFeatureChannel>& SubChannelPtr : SubChannels)
+	for (const TObjectPtr<UPoseSearchFeatureChannel>& SubChannelPtr : GetSubChannels())
 	{
 		if (const UPoseSearchFeatureChannel* SubChannel = SubChannelPtr.Get())
 		{
@@ -39,9 +39,9 @@ void UPoseSearchFeatureChannel_Group::IndexAsset(UE::PoseSearch::IAssetIndexer&
 	}
 }
 
-void UPoseSearchFeatureChannel_Group::BuildQuery(UE::PoseSearch::FSearchContext& SearchContext, FPoseSearchFeatureVectorBuilder& InOutQuery) const
+void UPoseSearchFeatureChannel_GroupBase::BuildQuery(UE::PoseSearch::FSearchContext& SearchContext, FPoseSearchFeatureVectorBuilder& InOutQuery) const
 {
-	for (const TObjectPtr<UPoseSearchFeatureChannel>& SubChannelPtr : SubChannels)
+	for (const TObjectPtr<UPoseSearchFeatureChannel>& SubChannelPtr : GetSubChannels())
 	{
 		if (const UPoseSearchFeatureChannel* SubChannel = SubChannelPtr.Get())
 		{
@@ -50,10 +50,23 @@ void UPoseSearchFeatureChannel_Group::BuildQuery(UE::PoseSearch::FSearchContext&
 	}
 }
 
-void UPoseSearchFeatureChannel_Group::DebugDraw(const UE::PoseSearch::FDebugDrawParams& DrawParams, TConstArrayView<float> PoseVector) const
+void UPoseSearchFeatureChannel_GroupBase::PreDebugDraw(UE::PoseSearch::FDebugDrawParams& DrawParams, TConstArrayView<float> PoseVector) const
 {
 #if ENABLE_DRAW_DEBUG
-	for (const TObjectPtr<UPoseSearchFeatureChannel>& SubChannelPtr : SubChannels)
+	for (const TObjectPtr<UPoseSearchFeatureChannel>& SubChannelPtr : GetSubChannels())
+	{
+		if (const UPoseSearchFeatureChannel* SubChannel = SubChannelPtr.Get())
+		{
+			SubChannel->PreDebugDraw(DrawParams, PoseVector);
+		}
+	}
+#endif // ENABLE_DRAW_DEBUG
+}
+
+void UPoseSearchFeatureChannel_GroupBase::DebugDraw(const UE::PoseSearch::FDebugDrawParams& DrawParams, TConstArrayView<float> PoseVector) const
+{
+#if ENABLE_DRAW_DEBUG
+	for (const TObjectPtr<UPoseSearchFeatureChannel>& SubChannelPtr : GetSubChannels())
 	{
 		if (const UPoseSearchFeatureChannel* SubChannel = SubChannelPtr.Get())
 		{
@@ -64,9 +77,9 @@ void UPoseSearchFeatureChannel_Group::DebugDraw(const UE::PoseSearch::FDebugDraw
 }
 
 // IPoseFilter interface
-bool UPoseSearchFeatureChannel_Group::IsPoseFilterActive() const
+bool UPoseSearchFeatureChannel_GroupBase::IsPoseFilterActive() const
 {
-	for (const TObjectPtr<UPoseSearchFeatureChannel>& SubChannelPtr : SubChannels)
+	for (const TObjectPtr<UPoseSearchFeatureChannel>& SubChannelPtr : GetSubChannels())
 	{
 		if (const UPoseSearchFeatureChannel* SubChannel = SubChannelPtr.Get())
 		{
@@ -80,9 +93,9 @@ bool UPoseSearchFeatureChannel_Group::IsPoseFilterActive() const
 	return false;
 }
 
-bool UPoseSearchFeatureChannel_Group::IsPoseValid(TConstArrayView<float> PoseValues, TConstArrayView<float> QueryValues, int32 PoseIdx, const FPoseSearchPoseMetadata& Metadata) const
+bool UPoseSearchFeatureChannel_GroupBase::IsPoseValid(TConstArrayView<float> PoseValues, TConstArrayView<float> QueryValues, int32 PoseIdx, const FPoseSearchPoseMetadata& Metadata) const
 {
-	for (const TObjectPtr<UPoseSearchFeatureChannel>& SubChannelPtr : SubChannels)
+	for (const TObjectPtr<UPoseSearchFeatureChannel>& SubChannelPtr : GetSubChannels())
 	{
 		if (const UPoseSearchFeatureChannel* SubChannel = SubChannelPtr.Get())
 		{
diff --git a/Engine/Plugins/Experimental/Animation/PoseSearch/Source/Runtime/Private/PoseSearchFeatureChannel_Group.h b/Engine/Plugins/Experimental/Animation/PoseSearch/Source/Runtime/Private/PoseSearchFeatureChannel_Group.h
--- a/Engine/Plugins/Experimental/Animation/PoseSearch/Source/Runtime/Private/PoseSearchFeatureChannel_Group.h
+++ b/Engine/Plugins/Experimental/Animation/PoseSearch/Source/Runtime/Private/PoseSearchFeatureChannel_Group.h
@@ -6,23 +6,38 @@
 #include "BoneContainer.h"
 #include "PoseSearchFeatureChannel_Group.generated.h"
 
-UCLASS(BlueprintType, EditInlineNew, meta = (DisplayName = "Group Channel"), CollapseCategories)
-class POSESEARCH_API UPoseSearchFeatureChannel_Group : public UPoseSearchFeatureChannel
+// Feature channels interface
+UCLASS(Abstract, BlueprintType, EditInlineNew)
+class POSESEARCH_API UPoseSearchFeatureChannel_GroupBase : public UPoseSearchFeatureChannel
 {
 	GENERATED_BODY()
 
 public:
-	UPROPERTY(EditAnywhere, Instanced, BlueprintReadWrite, Category = "Settings")
-	TArray<TObjectPtr<UPoseSearchFeatureChannel>> SubChannels;
+	virtual TArrayView<TObjectPtr<UPoseSearchFeatureChannel>> GetSubChannels() PURE_VIRTUAL(UPoseSearchFeatureChannel::GetSubChannels, return TArrayView<TObjectPtr<UPoseSearchFeatureChannel>>(); );
+	virtual TConstArrayView<TObjectPtr<UPoseSearchFeatureChannel>> GetSubChannels() const PURE_VIRTUAL(UPoseSearchFeatureChannel::GetSubChannels, return TConstArrayView<TObjectPtr<UPoseSearchFeatureChannel>>(); );
 
 	// UPoseSearchFeatureChannel interface
-	virtual void InitializeSchema(UPoseSearchSchema* Schema) override;
+	virtual void Finalize(UPoseSearchSchema* Schema) override;
 	virtual void FillWeights(TArray<float>& Weights) const override;
 	virtual void IndexAsset(UE::PoseSearch::IAssetIndexer& Indexer, TArrayView<float> FeatureVectorTable) const override;
 	virtual void BuildQuery(UE::PoseSearch::FSearchContext& SearchContext, FPoseSearchFeatureVectorBuilder& InOutQuery) const override;
+	virtual void PreDebugDraw(UE::PoseSearch::FDebugDrawParams& DrawParams, TConstArrayView<float> PoseVector) const override;
 	virtual void DebugDraw(const UE::PoseSearch::FDebugDrawParams& DrawParams, TConstArrayView<float> PoseVector) const override;
 
 	// IPoseFilter interface
 	virtual bool IsPoseFilterActive() const override;
 	virtual bool IsPoseValid(TConstArrayView<float> PoseValues, TConstArrayView<float> QueryValues, int32 PoseIdx, const FPoseSearchPoseMetadata& Metadata) const override;
+};
+
+UCLASS(BlueprintType, EditInlineNew, meta = (DisplayName = "Group Channel"), CollapseCategories)
+class POSESEARCH_API UPoseSearchFeatureChannel_Group : public UPoseSearchFeatureChannel_GroupBase
+{
+	GENERATED_BODY()
+
+	virtual TArrayView<TObjectPtr<UPoseSearchFeatureChannel>> GetSubChannels() override { return SubChannels; }
+	virtual TConstArrayView<TObjectPtr<UPoseSearchFeatureChannel>> GetSubChannels() const override { return SubChannels; }
+
+public:
+	UPROPERTY(EditAnywhere, Instanced, BlueprintReadWrite, Category = "SubChannels")
+	TArray<TObjectPtr<UPoseSearchFeatureChannel>> SubChannels;
 };
\ No newline at end of file
diff --git a/Engine/Plugins/Experimental/Animation/PoseSearch/Source/Runtime/Private/PoseSearchFeatureChannel_Heading.cpp b/Engine/Plugins/Experimental/Animation/PoseSearch/Source/Runtime/Private/PoseSearchFeatureChannel_Heading.cpp
--- a/Engine/Plugins/Experimental/Animation/PoseSearch/Source/Runtime/Private/PoseSearchFeatureChannel_Heading.cpp
+++ b/Engine/Plugins/Experimental/Animation/PoseSearch/Source/Runtime/Private/PoseSearchFeatureChannel_Heading.cpp
@@ -1,6 +1,7 @@
 // Copyright Epic Games, Inc. All Rights Reserved.
 
 #include "PoseSearchFeatureChannel_Heading.h"
+#include "Animation/MotionTrajectoryTypes.h"
 #include "DrawDebugHelpers.h"
 #include "PoseSearch/PoseSearchAssetIndexer.h"
 #include "PoseSearch/PoseSearchAssetSampler.h"
@@ -8,7 +9,7 @@
 #include "PoseSearch/PoseSearchDatabase.h"
 #include "PoseSearch/PoseSearchSchema.h"
 
-void UPoseSearchFeatureChannel_Heading::InitializeSchema(UPoseSearchSchema* Schema)
+void UPoseSearchFeatureChannel_Heading::Finalize(UPoseSearchSchema* Schema)
 {
 	ChannelDataOffset = Schema->SchemaCardinality;
 	ChannelCardinality = UE::PoseSearch::FFeatureVectorHelper::EncodeVectorCardinality;
@@ -81,10 +82,28 @@ void UPoseSearchFeatureChannel_Heading::BuildQuery(UE::PoseSearch::FSearchContex
 	}
 	else
 	{
-		FTransform Transform = SearchContext.TryGetTransformAndCacheResults(SampleTimeOffset, InOutQuery.GetSchema(), SchemaBoneIdx);
-		const FTransform RootTransform = SearchContext.TryGetTransformAndCacheResults(0.f, InOutQuery.GetSchema(), FSearchContext::SchemaRootBoneIdx);
-		const FTransform RootTransformPrev = SearchContext.TryGetTransformAndCacheResults(SampleTimeOffset, InOutQuery.GetSchema(), FSearchContext::SchemaRootBoneIdx);
-		Transform = Transform * (RootTransformPrev * RootTransform.Inverse());
+		FTransform Transform;
+		if (InOutQuery.GetSchema()->BoneReferences[SchemaBoneIdx].HasValidSetup())
+		{
+			// calculating the Transform in component space for the bone indexed by SchemaBoneIdx
+			Transform = SearchContext.TryGetTransformAndCacheResults(SampleTimeOffset, InOutQuery.GetSchema(), SchemaBoneIdx);
+
+			// if the SampleTimeOffset is not zero we calculate the delta root bone between the root at current time (zero) and the root at SampleTimeOffset (in the past)
+			// so we can offset the Transform by this amount
+			if (!FMath::IsNearlyZero(SampleTimeOffset, UE_KINDA_SMALL_NUMBER))
+			{
+				const FTransform RootTransform = SearchContext.TryGetTransformAndCacheResults(0.f, InOutQuery.GetSchema(), FSearchContext::SchemaRootBoneIdx);
+				const FTransform RootTransformPrev = SearchContext.TryGetTransformAndCacheResults(SampleTimeOffset, InOutQuery.GetSchema(), FSearchContext::SchemaRootBoneIdx);
+				Transform = Transform * (RootTransformPrev * RootTransform.Inverse());
+			}
+		}
+		else
+		{
+			check(SearchContext.Trajectory);
+			// @todo: make this call consistent with Transform = SearchContext.TryGetTransformAndCacheResults(SampleTimeOffset, InOutQuery.GetSchema(), FSearchContext::SchemaRootBoneIdx);
+			const FTrajectorySample TrajectorySample = SearchContext.Trajectory->GetSampleAtTime(SampleTimeOffset);
+			Transform = TrajectorySample.Transform;
+		}
 
 		int32 DataOffset = ChannelDataOffset;
 		FFeatureVectorHelper::EncodeVector(InOutQuery.EditValues(), DataOffset, GetAxis(Transform.GetRotation()));
@@ -96,20 +115,12 @@ void UPoseSearchFeatureChannel_Heading::DebugDraw(const UE::PoseSearch::FDebugDr
 #if ENABLE_DRAW_DEBUG
 	using namespace UE::PoseSearch;
 
-	const UPoseSearchSchema* Schema = DrawParams.GetSchema();
-	check(Schema && Schema->IsValid());
-
 	const float LifeTime = DrawParams.DefaultLifeTime;
 	const uint8 DepthPriority = ESceneDepthPriorityGroup::SDPG_Foreground + 2;
 	const bool bPersistent = EnumHasAnyFlags(DrawParams.Flags, EDebugDrawFlags::Persistent);
-
-	// @todo: search in the schema for a UPoseSearchFeatureChannel_Position with the same bone as SchemaBoneIdx and decode it's data as position
-	const FVector BonePos = DrawParams.Mesh != nullptr ? DrawParams.Mesh->GetSocketTransform(Bone.BoneName).GetLocation() : DrawParams.RootTransform.GetTranslation();
-
-	int32 DataOffset = ChannelDataOffset;
-	const FVector BoneHeading = DrawParams.RootTransform.GetRotation().RotateVector(FFeatureVectorHelper::DecodeVector(PoseVector, DataOffset));
-
 	const FColor Color = DrawParams.GetColor(ColorPresetIndex);
+	const FVector BoneHeading = DrawParams.RootTransform.GetRotation().RotateVector(FFeatureVectorHelper::DecodeVectorAtOffset(PoseVector, ChannelDataOffset));
+	const FVector BonePos = DrawParams.GetCachedPosition(SampleTimeOffset, SchemaBoneIdx);
 
 	if (EnumHasAnyFlags(DrawParams.Flags, EDebugDrawFlags::DrawSearchIndex))
 	{
diff --git a/Engine/Plugins/Experimental/Animation/PoseSearch/Source/Runtime/Private/PoseSearchFeatureChannel_Heading.h b/Engine/Plugins/Experimental/Animation/PoseSearch/Source/Runtime/Private/PoseSearchFeatureChannel_Heading.h
--- a/Engine/Plugins/Experimental/Animation/PoseSearch/Source/Runtime/Private/PoseSearchFeatureChannel_Heading.h
+++ b/Engine/Plugins/Experimental/Animation/PoseSearch/Source/Runtime/Private/PoseSearchFeatureChannel_Heading.h
@@ -45,7 +45,7 @@ public:
 	EInputQueryPose InputQueryPose = EInputQueryPose::UseContinuingPose;
 
 	// UPoseSearchFeatureChannel interface
-	virtual void InitializeSchema(UPoseSearchSchema* Schema) override;
+	virtual void Finalize(UPoseSearchSchema* Schema) override;
 	virtual void FillWeights(TArray<float>& Weights) const override;
 	virtual void IndexAsset(UE::PoseSearch::IAssetIndexer& Indexer, TArrayView<float> FeatureVectorTable) const override;
 	virtual void BuildQuery(UE::PoseSearch::FSearchContext& SearchContext, FPoseSearchFeatureVectorBuilder& InOutQuery) const override;
diff --git a/Engine/Plugins/Experimental/Animation/PoseSearch/Source/Runtime/Private/PoseSearchFeatureChannel_Phase.cpp b/Engine/Plugins/Experimental/Animation/PoseSearch/Source/Runtime/Private/PoseSearchFeatureChannel_Phase.cpp
new file mode 100644
--- /dev/null
+++ b/Engine/Plugins/Experimental/Animation/PoseSearch/Source/Runtime/Private/PoseSearchFeatureChannel_Phase.cpp
@@ -0,0 +1,398 @@
+// Copyright Epic Games, Inc. All Rights Reserved.
+
+#include "PoseSearchFeatureChannel_Phase.h"
+#include "DrawDebugHelpers.h"
+#include "PoseSearch/PoseSearchAssetIndexer.h"
+#include "PoseSearch/PoseSearchAssetSampler.h"
+#include "PoseSearch/PoseSearchContext.h"
+#include "PoseSearch/PoseSearchDatabase.h"
+#include "PoseSearch/PoseSearchSchema.h"
+
+namespace UE::PoseSearch
+{
+	struct LocalMinMax
+	{
+		enum
+		{
+			Min,
+			Max
+		} Type = Min;
+		int32 Index = 0;
+		float SignalValue = 0.f;
+	};
+
+	template <typename T>
+	T GetValueAtIndex(int32 Sample, const TArray<T>& Values)
+	{
+		const int32 Num = Values.Num();
+		check(Num > 1);
+
+		if (Sample < 0)
+		{
+			return (Values[1] - Values[0]) * Sample + Values[0];
+		}
+
+		if (Sample < Num)
+		{
+			return Values[Sample];
+		}
+
+		return (Values[Num - 1] - Values[Num - 2]) * (Sample - (Num - 1)) + Values[Num - 1];
+	}
+
+	static void CollectBonePositions(TArray<FVector>& BonePositions, IAssetIndexer& Indexer, int8 SchemaBoneIdx)
+	{
+		const FAssetIndexingContext& IndexingContext = Indexer.GetIndexingContext();
+		const float FiniteDelta = IndexingContext.Schema->GetSamplingInterval();
+		const float SampleTimeStart = FMath::Min(IndexingContext.BeginSampleIdx * FiniteDelta, IndexingContext.AssetSampler->GetPlayLength());
+		const int32 NumSamples = IndexingContext.EndSampleIdx - IndexingContext.BeginSampleIdx;
+
+		// collecting all the bone transforms
+		BonePositions.Reset();
+		BonePositions.AddDefaulted(NumSamples);
+		for (int32 SampleIdx = 0; SampleIdx != NumSamples; ++SampleIdx)
+		{
+			const float SampleTime = SampleTimeStart + SampleIdx * FiniteDelta;
+			bool bUnused;
+			BonePositions[SampleIdx] = Indexer.GetTransformAndCacheResults(SampleTime, SampleTimeStart, SchemaBoneIdx, bUnused).GetTranslation();
+		}
+	}
+
+	static void CalculateSignal(const TArray<FVector>& BonePositions, TArray<float>& Signal, int32 offset = 1)
+	{
+		Signal.Reset();
+		Signal.AddDefaulted(BonePositions.Num());
+
+		for (int32 SampleIdx = 0; SampleIdx != BonePositions.Num(); ++SampleIdx)
+		{
+			Signal[SampleIdx] = (GetValueAtIndex(SampleIdx + offset, BonePositions) - GetValueAtIndex(SampleIdx - offset, BonePositions)).Length();
+		}
+	}
+
+	static void SmoothSignal(const TArray<float>& Signal, TArray<float>& SmoothedSignal, int32 offset = 1)
+	{
+		SmoothedSignal.Reset();
+		SmoothedSignal.AddDefaulted(Signal.Num());
+
+		for (int32 SampleIdx = -offset; SampleIdx != offset; ++SampleIdx)
+		{
+			SmoothedSignal[0] += GetValueAtIndex(SampleIdx, Signal);
+		}
+
+		for (int32 SampleIdx = 1; SampleIdx != Signal.Num(); ++SampleIdx)
+		{
+			SmoothedSignal[SampleIdx] = SmoothedSignal[SampleIdx - 1] - GetValueAtIndex(SampleIdx - offset - 1, Signal) + GetValueAtIndex(SampleIdx + offset, Signal);
+		}
+
+		for (int32 SampleIdx = 0; SampleIdx != Signal.Num(); ++SampleIdx)
+		{
+			SmoothedSignal[SampleIdx] /= 2 * offset + 1;
+		}
+	}
+
+	static void FindLocalMinMax(const TArray<float>& Signal, TArray<LocalMinMax>& MinMax)
+	{
+		enum SignalState
+		{
+			Flat,
+			Ascending,
+			Descending
+		};
+
+		MinMax.Reset();
+		if (Signal.Num() > 1)
+		{
+			SignalState State = SignalState::Flat;
+			for (int32 SignalIndex = 1; SignalIndex < Signal.Num(); ++SignalIndex)
+			{
+				const int32 PrevSignalIndex = SignalIndex - 1;
+				const float PrevSignalValue = Signal[PrevSignalIndex];
+				const float SignalValue = Signal[SignalIndex];
+
+				if (State == SignalState::Flat)
+				{
+					if (SignalValue > PrevSignalValue)
+					{
+						State = SignalState::Ascending;
+					}
+					else if (SignalValue < PrevSignalValue)
+					{
+						State = SignalState::Descending;
+					}
+				}
+				else if (State == SignalState::Ascending)
+				{
+					if (SignalValue < PrevSignalValue)
+					{
+						State = SignalState::Descending;
+
+						LocalMinMax LocalMinMax;
+						LocalMinMax.Type = LocalMinMax::Max;
+						LocalMinMax.Index = PrevSignalIndex;
+						LocalMinMax.SignalValue = Signal[LocalMinMax.Index];
+
+						check(MinMax.IsEmpty() || MinMax.Last().Type != LocalMinMax.Type);
+						MinMax.Add(LocalMinMax);
+					}
+				}
+				else // if (State == SignalState::Descending)
+				{
+					if (SignalValue > PrevSignalValue)
+					{
+						State = SignalState::Ascending;
+
+						LocalMinMax LocalMinMax;
+						LocalMinMax.Type = LocalMinMax::Min;
+						LocalMinMax.Index = PrevSignalIndex;
+						LocalMinMax.SignalValue = Signal[LocalMinMax.Index];
+
+						check(MinMax.IsEmpty() || MinMax.Last().Type != LocalMinMax.Type);
+						MinMax.Add(LocalMinMax);
+					}
+				}
+			}
+		}
+	}
+
+	static void ExtrapolateLocalMinMaxBoundaries(TArray<LocalMinMax>& MinMax, const TArray<float>& Signal)
+	{
+		const int32 Num = MinMax.Num();
+
+		check(Signal.Num() > 0);
+
+		LocalMinMax InitialMinMax;
+		LocalMinMax FinalMinMax;
+
+		if (Num == 0)
+		{
+			const bool IsInitialMax = Signal[0] > Signal[Signal.Num() - 1];
+
+			InitialMinMax.Index = 0;
+			InitialMinMax.SignalValue = Signal[0];
+			InitialMinMax.Type = IsInitialMax ? LocalMinMax::Max : LocalMinMax::Min;
+
+			FinalMinMax.Index = Signal.Num() - 1;
+			FinalMinMax.SignalValue = Signal[Signal.Num() - 1];
+			FinalMinMax.Type = IsInitialMax ? LocalMinMax::Min : LocalMinMax::Max;
+
+			MinMax.Add(InitialMinMax);
+			MinMax.Add(FinalMinMax);
+		}
+		else
+		{
+			int32 InitialDelta = 0;
+			int32 FinalDelta = 0;
+			if (Num > 2)
+			{
+				InitialDelta = MinMax[2].Index - MinMax[1].Index;
+				FinalDelta = MinMax[Num - 2].Index - MinMax[Num - 3].Index;
+			}
+			else if (Num > 1)
+			{
+				InitialDelta = MinMax[1].Index - MinMax[0].Index;
+				FinalDelta = MinMax[Num - 1].Index - MinMax[Num - 2].Index;
+			}
+			else
+			{
+				InitialDelta = MinMax[0].Index;
+				FinalDelta = (Signal.Num() - 1) - MinMax[0].Index;
+			}
+
+			InitialMinMax.SignalValue = Num > 1 ? MinMax[1].SignalValue : Signal[0];
+			InitialMinMax.Type = MinMax[0].Type == LocalMinMax::Min ? LocalMinMax::Max : LocalMinMax::Min;
+			InitialMinMax.Index = FMath::Min(MinMax[0].Index - InitialDelta, 0);
+
+			FinalMinMax.SignalValue = Num > 1 ? MinMax[Num - 2].SignalValue : Signal[Signal.Num() - 1];
+			FinalMinMax.Type = MinMax[Num - 1].Type == LocalMinMax::Min ? LocalMinMax::Max : LocalMinMax::Min;
+			FinalMinMax.Index = FMath::Max(MinMax[Num - 1].Index + FinalDelta, Signal.Num() - 1);
+
+			// there's no point in adding an InitialMinMax if the first MinMax is at the first frame of the signal
+			if (MinMax[0].Index > 0)
+			{
+				MinMax.Insert(InitialMinMax, 0);
+			}
+
+			// there's no point in adding a FinalMinMax if the last MinMax is at the last frame of the signal
+			if (MinMax[Num - 1].Index < Signal.Num() - 1)
+			{
+				MinMax.Add(FinalMinMax);
+			}
+		}
+	}
+
+	static void ValidateLocalMinMax(const TArray<LocalMinMax>& MinMax)
+	{
+		for (int32 i = 1; i < MinMax.Num(); ++i)
+		{
+			check(MinMax[i].Type != MinMax[i - 1].Type);
+			check(MinMax[i].Index > MinMax[i - 1].Index);
+			if (MinMax[i].Type == LocalMinMax::Min)
+			{
+				check(MinMax[i].SignalValue <= MinMax[i - 1].SignalValue);
+			}
+			else
+			{
+				check(MinMax[i].SignalValue >= MinMax[i - 1].SignalValue);
+			}
+		}
+	}
+
+	static void CalculatePhaseAndCertainty(int32 Index, const TArray<LocalMinMax>& MinMax, int32 SignalSize, float& Phase, float& Certainty)
+	{
+		// @todo: expose them via UI
+		static float CertaintyMin = 1.f;
+		static float CertaintyMult = 0.1f;
+
+		const int32 LastIndex = MinMax.Num() - 1;
+		for (int32 i = 1; i < MinMax.Num(); ++i)
+		{
+			const int32 MinMaxIndex = MinMax[i].Index;
+			if (Index < MinMaxIndex)
+			{
+				const int32 PrevMinMaxIndex = MinMax[i - 1].Index;
+				check(MinMaxIndex > PrevMinMaxIndex);
+				const float Ratio = static_cast<float>((Index - PrevMinMaxIndex)) / static_cast<float>((MinMaxIndex - PrevMinMaxIndex));
+				const float PhaseOffset = MinMax[i - 1].Type == LocalMinMax::Min ? 0.f : 0.5f;
+				Phase = PhaseOffset + Ratio * 0.5f;
+
+				const float DeltaSignalValue = FMath::Abs(MinMax[i - 1].SignalValue - MinMax[i].SignalValue);
+				const float NextDeltaSignalValue = i < LastIndex ? FMath::Abs(MinMax[i].SignalValue - MinMax[i + 1].SignalValue) : DeltaSignalValue;
+				Certainty = CertaintyMin + (DeltaSignalValue * (1.f - Ratio) + NextDeltaSignalValue * Ratio) * CertaintyMult;
+				return;
+			}
+		}
+
+		Phase = MinMax[LastIndex].Type == LocalMinMax::Min ? 0.f : 0.5f;
+		Certainty = CertaintyMin + (LastIndex > 0 ? FMath::Abs(MinMax[LastIndex].SignalValue - MinMax[LastIndex - 1].SignalValue) : 0.f) * CertaintyMult;
+	}
+
+	static void CalculatePhasesFromLocalMinMax(const TArray<LocalMinMax>& MinMax, TArray<FVector2D>& Phases, int32 SignalSize)
+	{
+		Phases.Reset();
+		Phases.AddDefaulted(SignalSize);
+
+		float Certainty = 1.f;
+		float Phase = 0.f;
+		for (int32 i = 0; i < SignalSize; ++i)
+		{
+			CalculatePhaseAndCertainty(i, MinMax, SignalSize, Phase, Certainty);
+			FMath::SinCos(&Phases[i].X, &Phases[i].Y, Phase * TWO_PI);
+			Phases[i] *= Certainty;
+		}
+	}
+
+} // namespace UE::PoseSearch
+
+void UPoseSearchFeatureChannel_Phase::Finalize(UPoseSearchSchema* Schema)
+{
+	ChannelDataOffset = Schema->SchemaCardinality;
+	ChannelCardinality = UE::PoseSearch::FFeatureVectorHelper::EncodeVector2DCardinality;
+	Schema->SchemaCardinality += ChannelCardinality;
+	SchemaBoneIdx = Schema->AddBoneReference(Bone);
+}
+
+void UPoseSearchFeatureChannel_Phase::FillWeights(TArray<float>& Weights) const
+{
+	using namespace UE::PoseSearch;
+
+	for (int32 i = 0; i != FFeatureVectorHelper::EncodeVector2DCardinality; ++i)
+	{
+		Weights[ChannelDataOffset + i] = Weight;
+	}
+}
+
+void UPoseSearchFeatureChannel_Phase::IndexAsset(UE::PoseSearch::IAssetIndexer& Indexer, TArrayView<float> FeatureVectorTable) const
+{
+	using namespace UE::PoseSearch;
+
+	// @todo: do we really need to use double(s) in all this math?
+	// @todo: expose them via UI
+	static float BoneSamplingCentralDifferencesTime = 0.2f; // seconds
+	static float SmoothingWindowTime = 0.3f; // seconds
+
+
+	const FAssetIndexingContext& IndexingContext = Indexer.GetIndexingContext();
+	const UPoseSearchSchema* Schema = IndexingContext.Schema;
+	const float FiniteDelta = Schema->GetSamplingInterval();
+	
+	TArray<FVector2D> Phases;
+	TArray<float> Signal;
+	TArray<float> SmoothedSignal;
+	TArray<LocalMinMax> LocalMinMax;
+	TArray<FVector> BonePositions;
+	
+	CollectBonePositions(BonePositions, Indexer, SchemaBoneIdx);
+
+	// @todo: have different way of calculating signals, for example: height of the bone transform, acceleration, etc?
+	const int32 BoneSamplingCentralDifferencesOffset = FMath::Max(FMath::CeilToInt(BoneSamplingCentralDifferencesTime / FiniteDelta), 1);
+	CalculateSignal(BonePositions, Signal, BoneSamplingCentralDifferencesOffset);
+
+	const int32 SmoothingWindowOffset = FMath::Max(FMath::CeilToInt(SmoothingWindowTime / FiniteDelta), 1);
+	SmoothSignal(Signal, SmoothedSignal, SmoothingWindowOffset);
+
+	FindLocalMinMax(SmoothedSignal, LocalMinMax);
+	ValidateLocalMinMax(LocalMinMax);
+
+	ExtrapolateLocalMinMaxBoundaries(LocalMinMax, SmoothedSignal);
+	ValidateLocalMinMax(LocalMinMax);
+	CalculatePhasesFromLocalMinMax(LocalMinMax, Phases, SmoothedSignal.Num());
+
+	for (int32 SampleIdx = IndexingContext.BeginSampleIdx; SampleIdx != IndexingContext.EndSampleIdx; ++SampleIdx)
+	{
+		const int32 VectorIdx = SampleIdx - IndexingContext.BeginSampleIdx;
+		int32 DataOffset = ChannelDataOffset;
+		FFeatureVectorHelper::EncodeVector2D(IndexingContext.GetPoseVector(VectorIdx, FeatureVectorTable), DataOffset, Phases[VectorIdx]);
+	}
+}
+
+void UPoseSearchFeatureChannel_Phase::BuildQuery(UE::PoseSearch::FSearchContext& SearchContext, FPoseSearchFeatureVectorBuilder& InOutQuery) const
+{
+	using namespace UE::PoseSearch;
+
+	const bool bIsCurrentResultValid = SearchContext.CurrentResult.IsValid();
+	const bool bSkip = InputQueryPose != EInputQueryPose::UseCharacterPose && bIsCurrentResultValid && SearchContext.CurrentResult.Database->Schema == InOutQuery.GetSchema();
+	if (bSkip || !SearchContext.History)
+	{
+		if (bIsCurrentResultValid)
+		{
+			const float LerpValue = InputQueryPose == EInputQueryPose::UseInterpolatedContinuingPose ? SearchContext.CurrentResult.LerpValue : 0.f;
+			int32 DataOffset = ChannelDataOffset;
+			FFeatureVectorHelper::EncodeVector2D(InOutQuery.EditValues(), DataOffset, SearchContext.GetCurrentResultPrevPoseVector(), SearchContext.GetCurrentResultPoseVector(), SearchContext.GetCurrentResultNextPoseVector(), LerpValue);
+		}
+		// else leave the InOutQuery set to zero since the SearchContext.History is invalid and it'll fail if we continue
+	}
+	else
+	{
+		// @todo: Support phase in BuildQuery
+		// FFeatureVectorHelper::EncodeVector2D(InOutQuery.EditValues(), DataOffset, ???);
+	}
+}
+
+void UPoseSearchFeatureChannel_Phase::DebugDraw(const UE::PoseSearch::FDebugDrawParams& DrawParams, TConstArrayView<float> PoseVector) const
+{
+#if ENABLE_DRAW_DEBUG
+	using namespace UE::PoseSearch;
+
+	static float ScaleFactor = 1.f;
+
+	const float LifeTime = DrawParams.DefaultLifeTime;
+	const uint8 DepthPriority = ESceneDepthPriorityGroup::SDPG_Foreground + 2;
+	const bool bPersistent = EnumHasAnyFlags(DrawParams.Flags, EDebugDrawFlags::Persistent);
+	const FColor Color = DrawParams.GetColor(ColorPresetIndex);
+
+	const FVector2D Phase = FFeatureVectorHelper::DecodeVector2DAtOffset(PoseVector, ChannelDataOffset);
+	const FVector BonePos = DrawParams.GetCachedPosition(0.f, SchemaBoneIdx);
+
+	const FVector TransformXAxisVector = DrawParams.RootTransform.TransformVector(FVector::XAxisVector);
+	const FVector TransformYAxisVector = DrawParams.RootTransform.TransformVector(FVector::YAxisVector);
+	const FVector TransformZAxisVector = DrawParams.RootTransform.TransformVector(FVector::ZAxisVector);
+
+	const FVector PhaseVector = (TransformZAxisVector * Phase.X + TransformYAxisVector * Phase.Y) * ScaleFactor;
+	DrawDebugLine(DrawParams.World, BonePos, BonePos + PhaseVector, Color, bPersistent, LifeTime, DepthPriority, 0.f);
+
+	static int32 Segments = 32;
+	FMatrix CircleTransform;
+	CircleTransform.SetAxes(&TransformXAxisVector, &TransformYAxisVector, &TransformZAxisVector, &BonePos);
+	DrawDebugCircle(DrawParams.World, CircleTransform, PhaseVector.Length(), Segments, Color, bPersistent, LifeTime, DepthPriority, 0.f, false);
+#endif // ENABLE_DRAW_DEBUG
+}
diff --git a/Engine/Plugins/Experimental/Animation/PoseSearch/Source/Runtime/Private/PoseSearchFeatureChannel_Phase.h b/Engine/Plugins/Experimental/Animation/PoseSearch/Source/Runtime/Private/PoseSearchFeatureChannel_Phase.h
new file mode 100644
--- /dev/null
+++ b/Engine/Plugins/Experimental/Animation/PoseSearch/Source/Runtime/Private/PoseSearchFeatureChannel_Phase.h
@@ -0,0 +1,36 @@
+// Copyright Epic Games, Inc. All Rights Reserved.
+
+#pragma once
+
+#include "PoseSearch/PoseSearchFeatureChannel.h"
+#include "BoneContainer.h"
+#include "PoseSearchFeatureChannel_Phase.generated.h"
+
+UCLASS(BlueprintType, EditInlineNew, meta = (DisplayName = "Phase Channel"), CollapseCategories)
+class POSESEARCH_API UPoseSearchFeatureChannel_Phase : public UPoseSearchFeatureChannel
+{
+	GENERATED_BODY()
+
+public:
+	UPROPERTY(EditAnywhere, Category = "Settings")
+	FBoneReference Bone;
+
+	UPROPERTY(EditAnywhere, Category = "Settings")
+	float Weight = 1.f;
+
+	UPROPERTY()
+	int8 SchemaBoneIdx = 0;
+
+	UPROPERTY(EditAnywhere, Category = "Settings", meta = (ExcludeFromHash))
+	int32 ColorPresetIndex = 0;
+
+	UPROPERTY(EditAnywhere, Category = "Settings")
+	EInputQueryPose InputQueryPose = EInputQueryPose::UseContinuingPose;
+
+	// UPoseSearchFeatureChannel interface
+	virtual void Finalize(UPoseSearchSchema* Schema) override;
+	virtual void FillWeights(TArray<float>& Weights) const override;
+	virtual void IndexAsset(UE::PoseSearch::IAssetIndexer& Indexer, TArrayView<float> FeatureVectorTable) const override;
+	virtual void BuildQuery(UE::PoseSearch::FSearchContext& SearchContext, FPoseSearchFeatureVectorBuilder& InOutQuery) const override;
+	virtual void DebugDraw(const UE::PoseSearch::FDebugDrawParams& DrawParams, TConstArrayView<float> PoseVector) const override;
+};
\ No newline at end of file
diff --git a/Engine/Plugins/Experimental/Animation/PoseSearch/Source/Runtime/Private/PoseSearchFeatureChannel_Pose.cpp b/Engine/Plugins/Experimental/Animation/PoseSearch/Source/Runtime/Private/PoseSearchFeatureChannel_Pose.cpp
--- a/Engine/Plugins/Experimental/Animation/PoseSearch/Source/Runtime/Private/PoseSearchFeatureChannel_Pose.cpp
+++ b/Engine/Plugins/Experimental/Animation/PoseSearch/Source/Runtime/Private/PoseSearchFeatureChannel_Pose.cpp
@@ -1,846 +1,179 @@
 // Copyright Epic Games, Inc. All Rights Reserved.
 
 #include "PoseSearchFeatureChannel_Pose.h"
-#include "DrawDebugHelpers.h"
-#include "PoseSearch/PoseSearchAssetIndexer.h"
-#include "PoseSearch/PoseSearchAssetSampler.h"
-#include "PoseSearch/PoseSearchContext.h"
-#include "PoseSearch/PoseSearchDatabase.h"
-#include "PoseSearch/PoseSearchDerivedDataKey.h"
-#include "PoseSearch/PoseSearchHistory.h"
+#include "Animation/Skeleton.h"
 #include "PoseSearch/PoseSearchSchema.h"
+#include "PoseSearch/PoseSearchDerivedDataKey.h"
+#include "PoseSearchFeatureChannel_Heading.h"
+#include "PoseSearchFeatureChannel_Phase.h"
+#include "PoseSearchFeatureChannel_Position.h"
+#include "PoseSearchFeatureChannel_Velocity.h"
 
 #define LOCTEXT_NAMESPACE "PoseSearchFeatureChannels"
 
-namespace UE::PoseSearch
-{
-
-struct LocalMinMax
-{
-	enum
-	{
-		Min,
-		Max
-	} Type = Min;
-	int32 Index = 0;
-	float SignalValue = 0.f;
-};
-
-template <typename T>
-T GetValueAtIndex(int32 Sample, const TArray<T>& Values)
-{
-	const int32 Num = Values.Num();
-	check(Num > 1);
-
-	if (Sample < 0)
-	{
-		return (Values[1] - Values[0]) * Sample + Values[0];
-	}
-
-	if (Sample < Num)
-	{
-		return Values[Sample];
-	}
-
-	return (Values[Num - 1] - Values[Num - 2]) * (Sample - (Num - 1)) + Values[Num - 1];
-}
-
-static void CollectBonePositions(TArray<FVector>& BonePositions, IAssetIndexer& Indexer, int8 SchemaBoneIdx)
-{
-	const FAssetIndexingContext& IndexingContext = Indexer.GetIndexingContext();
-	const float FiniteDelta = IndexingContext.Schema->GetSamplingInterval();
-	const float SampleTimeStart = FMath::Min(IndexingContext.BeginSampleIdx * FiniteDelta, IndexingContext.AssetSampler->GetPlayLength());
-	const int32 NumSamples = IndexingContext.EndSampleIdx - IndexingContext.BeginSampleIdx;
-
-	// collecting all the bone transforms
-	BonePositions.Reset();
-	BonePositions.AddDefaulted(NumSamples);
-	for (int32 SampleIdx = 0; SampleIdx != NumSamples; ++SampleIdx)
-	{
-		const float SampleTime = SampleTimeStart + SampleIdx * FiniteDelta;
-		bool bUnused;
-		BonePositions[SampleIdx] = Indexer.GetTransformAndCacheResults(SampleTime, SampleTimeStart, SchemaBoneIdx, bUnused).GetTranslation();
-	}
-}
-
-static void CalculateSignal(const TArray<FVector>& BonePositions, TArray<float>& Signal, int32 offset = 1)
-{
-	Signal.Reset();
-	Signal.AddDefaulted(BonePositions.Num());
-
-	for (int32 SampleIdx = 0; SampleIdx != BonePositions.Num(); ++SampleIdx)
-	{
-		Signal[SampleIdx] = (GetValueAtIndex(SampleIdx + offset, BonePositions) - GetValueAtIndex(SampleIdx - offset, BonePositions)).Length();
-	}
-}
-
-static void SmoothSignal(const TArray<float>& Signal, TArray<float>& SmoothedSignal, int32 offset = 1)
-{
-	SmoothedSignal.Reset();
-	SmoothedSignal.AddDefaulted(Signal.Num());
-
-	for (int32 SampleIdx = -offset; SampleIdx != offset; ++SampleIdx)
-	{
-		SmoothedSignal[0] += GetValueAtIndex(SampleIdx, Signal);
-	}
-
-	for (int32 SampleIdx = 1; SampleIdx != Signal.Num(); ++SampleIdx)
-	{
-		SmoothedSignal[SampleIdx] = SmoothedSignal[SampleIdx - 1] - GetValueAtIndex(SampleIdx - offset - 1, Signal) + GetValueAtIndex(SampleIdx + offset, Signal);
-	}
-
-	for (int32 SampleIdx = 0; SampleIdx != Signal.Num(); ++SampleIdx)
-	{
-		SmoothedSignal[SampleIdx] /= 2 * offset + 1;
-	}
-}
-
-static void FindLocalMinMax(const TArray<float>& Signal, TArray<LocalMinMax>& MinMax)
-{
-	enum SignalState
-	{
-		Flat,
-		Ascending,
-		Descending
-	};
-
-	MinMax.Reset();
-	if (Signal.Num() > 1)
-	{
-		SignalState State = SignalState::Flat;
-		for (int32 SignalIndex = 1; SignalIndex < Signal.Num(); ++SignalIndex)
-		{
-			const int32 PrevSignalIndex = SignalIndex - 1;
-			const float PrevSignalValue = Signal[PrevSignalIndex];
-			const float SignalValue = Signal[SignalIndex];
-
-			if (State == SignalState::Flat)
-			{
-				if (SignalValue > PrevSignalValue)
-				{
-					State = SignalState::Ascending;
-				}
-				else if(SignalValue < PrevSignalValue)
-				{
-					State = SignalState::Descending;
-				}
-			}
-			else if (State == SignalState::Ascending)
-			{
-				if (SignalValue < PrevSignalValue)
-				{
-					State = SignalState::Descending;
-					
-					LocalMinMax LocalMinMax;
-					LocalMinMax.Type = LocalMinMax::Max;
-					LocalMinMax.Index = PrevSignalIndex;
-					LocalMinMax.SignalValue = Signal[LocalMinMax.Index];
-
-					check(MinMax.IsEmpty() || MinMax.Last().Type != LocalMinMax.Type);
-					MinMax.Add(LocalMinMax);
-				}
-			}
-			else // if (State == SignalState::Descending)
-			{
-				if (SignalValue > PrevSignalValue)
-				{
-					State = SignalState::Ascending;
-
-					LocalMinMax LocalMinMax;
-					LocalMinMax.Type = LocalMinMax::Min;
-					LocalMinMax.Index = PrevSignalIndex;
-					LocalMinMax.SignalValue = Signal[LocalMinMax.Index];
-
-					check(MinMax.IsEmpty() || MinMax.Last().Type != LocalMinMax.Type);
-					MinMax.Add(LocalMinMax);
-				}
-			}
-		}
-	}
-}
-
-static void ExtrapolateLocalMinMaxBoundaries(TArray<LocalMinMax>& MinMax, const TArray<float>& Signal)
-{
-	const int32 Num = MinMax.Num();
-
-	check(Signal.Num() > 0);
-
-	LocalMinMax InitialMinMax;
-	LocalMinMax FinalMinMax;
-
-	if (Num == 0)
-	{
-		const bool IsInitialMax = Signal[0] > Signal[Signal.Num() - 1];
-
-		InitialMinMax.Index = 0;
-		InitialMinMax.SignalValue = Signal[0];
-		InitialMinMax.Type = IsInitialMax ? LocalMinMax::Max : LocalMinMax::Min;
-
-		FinalMinMax.Index = Signal.Num() - 1;
-		FinalMinMax.SignalValue = Signal[Signal.Num() - 1];
-		FinalMinMax.Type = IsInitialMax ? LocalMinMax::Min : LocalMinMax::Max;
-
-		MinMax.Add(InitialMinMax);
-		MinMax.Add(FinalMinMax);
-	}
-	else
-	{
-		int32 InitialDelta = 0;
-		int32 FinalDelta = 0;
-		if (Num > 2)
-		{
-			InitialDelta = MinMax[2].Index - MinMax[1].Index;
-			FinalDelta = MinMax[Num - 2].Index - MinMax[Num - 3].Index;
-		}
-		else if (Num > 1)
-		{
-			InitialDelta = MinMax[1].Index - MinMax[0].Index;
-			FinalDelta = MinMax[Num - 1].Index - MinMax[Num - 2].Index;
-		}
-		else
-		{
-			InitialDelta = MinMax[0].Index;
-			FinalDelta = (Signal.Num() - 1) - MinMax[0].Index;
-		}
-
-		InitialMinMax.SignalValue = Num > 1 ? MinMax[1].SignalValue : Signal[0];
-		InitialMinMax.Type = MinMax[0].Type == LocalMinMax::Min ? LocalMinMax::Max : LocalMinMax::Min;
-		InitialMinMax.Index = FMath::Min(MinMax[0].Index - InitialDelta, 0);
-
-		FinalMinMax.SignalValue = Num > 1 ? MinMax[Num - 2].SignalValue : Signal[Signal.Num() - 1];
-		FinalMinMax.Type = MinMax[Num - 1].Type == LocalMinMax::Min ? LocalMinMax::Max : LocalMinMax::Min;
-		FinalMinMax.Index = FMath::Max(MinMax[Num - 1].Index + FinalDelta, Signal.Num() - 1);
-
-		// there's no point in adding an InitialMinMax if the first MinMax is at the first frame of the signal
-		if (MinMax[0].Index > 0)
-		{
-			MinMax.Insert(InitialMinMax, 0);
-		}
-
-		// there's no point in adding a FinalMinMax if the last MinMax is at the last frame of the signal
-		if (MinMax[Num - 1].Index < Signal.Num() - 1)
-		{
-			MinMax.Add(FinalMinMax);
-		}
-	}
-}
-
-static void ValidateLocalMinMax(const TArray<LocalMinMax>& MinMax)
-{
-	for (int32 i = 1; i < MinMax.Num(); ++i)
-	{
-		check(MinMax[i].Type != MinMax[i - 1].Type);
-		check(MinMax[i].Index > MinMax[i - 1].Index);
-		if (MinMax[i].Type == LocalMinMax::Min)
-		{
-			check(MinMax[i].SignalValue <= MinMax[i - 1].SignalValue);
-		}
-		else
-		{
-			check(MinMax[i].SignalValue >= MinMax[i - 1].SignalValue);
-		}
-	}
-}
-
-static void CalculatePhaseAndCertainty(int32 Index, const TArray<LocalMinMax>& MinMax, int32 SignalSize, float& Phase, float& Certainty)
-{
-	// @todo: expose them via UI
-	static float CertaintyMin = 1.f;
-	static float CertaintyMult = 0.1f;
-
-	const int32 LastIndex = MinMax.Num() - 1;
-	for (int32 i = 1; i < MinMax.Num(); ++i)
-	{
-		const int32 MinMaxIndex = MinMax[i].Index;
-		if (Index < MinMaxIndex)
-		{
-			const int32 PrevMinMaxIndex = MinMax[i - 1].Index;
-			check(MinMaxIndex > PrevMinMaxIndex);
-			const float Ratio = static_cast<float>((Index - PrevMinMaxIndex)) / static_cast<float>((MinMaxIndex - PrevMinMaxIndex));
-			const float PhaseOffset = MinMax[i - 1].Type == LocalMinMax::Min ? 0.f : 0.5f;
-			Phase = PhaseOffset + Ratio * 0.5f;
-
-			const float DeltaSignalValue = FMath::Abs(MinMax[i - 1].SignalValue - MinMax[i].SignalValue);
-			const float NextDeltaSignalValue = i < LastIndex ? FMath::Abs(MinMax[i].SignalValue - MinMax[i + 1].SignalValue) : DeltaSignalValue;
-			Certainty = CertaintyMin + (DeltaSignalValue * (1.f - Ratio) + NextDeltaSignalValue * Ratio) * CertaintyMult;
-			return;
-		}
-	}
-
-	Phase = MinMax[LastIndex].Type == LocalMinMax::Min ? 0.f : 0.5f;
-	Certainty = CertaintyMin + (LastIndex > 0 ? FMath::Abs(MinMax[LastIndex].SignalValue - MinMax[LastIndex - 1].SignalValue) : 0.f) * CertaintyMult;
-}
-
-static void CalculatePhasesFromLocalMinMax(const TArray<LocalMinMax>& MinMax, TArray<FVector2D>& Phases, int32 SignalSize)
-{
-	Phases.Reset();
-	Phases.AddDefaulted(SignalSize);
-
-	float Certainty = 1.f;
-	float Phase = 0.f;
-	for (int32 i = 0; i < SignalSize; ++i)
-	{
-		CalculatePhaseAndCertainty(i, MinMax, SignalSize, Phase, Certainty);
-		FMath::SinCos(&Phases[i].X, &Phases[i].Y, Phase * TWO_PI);
-		Phases[i] *= Certainty;
-	}
-}
-
-} // namespace UE::PoseSearch
-
-void UPoseSearchFeatureChannel_Pose::InitializeSchema(UPoseSearchSchema* Schema)
-{
-	using namespace UE::PoseSearch;
-
-	ChannelDataOffset = Schema->SchemaCardinality;
-	for (int32 ChannelBoneIdx = 0; ChannelBoneIdx != SampledBones.Num(); ++ChannelBoneIdx)
-	{
-		const FPoseSearchBone& SampledBone = SampledBones[ChannelBoneIdx];
-		if (EnumHasAnyFlags(SampledBone.Flags, EPoseSearchBoneFlags::Position))
-		{
-			Schema->SchemaCardinality += FFeatureVectorHelper::EncodeVectorCardinality;
-		}
-		if (EnumHasAnyFlags(SampledBone.Flags, EPoseSearchBoneFlags::Rotation))
-		{
-			Schema->SchemaCardinality += FFeatureVectorHelper::EncodeQuatCardinality;
-		}
-		if (EnumHasAnyFlags(SampledBone.Flags, EPoseSearchBoneFlags::Velocity))
-		{
-			Schema->SchemaCardinality += FFeatureVectorHelper::EncodeVectorCardinality;
-		}
-		if (EnumHasAnyFlags(SampledBone.Flags, EPoseSearchBoneFlags::Phase))
-		{
-			Schema->SchemaCardinality += FFeatureVectorHelper::EncodeVector2DCardinality;
-		}
-	}
-
-	ChannelCardinality = Schema->SchemaCardinality - ChannelDataOffset;
-
-	SchemaBoneIdx.Reset();
-	for (const FPoseSearchBone& Bone : SampledBones)
-	{
-		SchemaBoneIdx.Add(Schema->AddBoneReference(Bone.Reference));
-	}
-}
-
-void UPoseSearchFeatureChannel_Pose::FillWeights(TArray<float>& Weights) const
-{
-	using namespace UE::PoseSearch;
-
-	int32 DataOffset = ChannelDataOffset;
-
-	const int32 NumBones = SampledBones.Num();
-	for (int32 ChannelBoneIdx = 0; ChannelBoneIdx != NumBones; ++ChannelBoneIdx)
-	{
-		const FPoseSearchBone& SampledBone = SampledBones[ChannelBoneIdx];
-		if (EnumHasAnyFlags(SampledBone.Flags, EPoseSearchBoneFlags::Position))
-		{
-			for (int32 i = 0; i != FFeatureVectorHelper::EncodeVectorCardinality; ++i)
-			{
-				Weights[DataOffset + i] = Weight * SampledBone.Weight;
-			}
-			DataOffset += FFeatureVectorHelper::EncodeVectorCardinality;
-		}
-		if (EnumHasAnyFlags(SampledBone.Flags, EPoseSearchBoneFlags::Rotation))
-		{
-			for (int32 i = 0; i != FFeatureVectorHelper::EncodeQuatCardinality; ++i)
-			{
-				Weights[DataOffset + i] = Weight * SampledBone.Weight;
-			}
-			DataOffset += FFeatureVectorHelper::EncodeQuatCardinality;
-		}
-		if (EnumHasAnyFlags(SampledBone.Flags, EPoseSearchBoneFlags::Velocity))
-		{
-			for (int32 i = 0; i != FFeatureVectorHelper::EncodeVectorCardinality; ++i)
-			{
-				Weights[DataOffset + i] = Weight * SampledBone.Weight;
-			}
-			DataOffset += FFeatureVectorHelper::EncodeVectorCardinality;
-		}
-		if (EnumHasAnyFlags(SampledBone.Flags, EPoseSearchBoneFlags::Phase))
-		{
-			for (int32 i = 0; i != FFeatureVectorHelper::EncodeVector2DCardinality; ++i)
-			{
-				Weights[DataOffset + i] = Weight * SampledBone.Weight;
-			}
-			DataOffset += FFeatureVectorHelper::EncodeVector2DCardinality;
-		}
-	}
-
-	check(DataOffset == ChannelDataOffset + ChannelCardinality);
-}
-
-// @todo: do we really need to use double(s) in all this math?
-void UPoseSearchFeatureChannel_Pose::CalculatePhases(UE::PoseSearch::IAssetIndexer& Indexer, TArray<TArray<FVector2D>>& OutPhases) const
-{
-	// @todo: expose them via UI
-	static float BoneSamplingCentralDifferencesTime = 0.2f; // seconds
-	static float SmoothingWindowTime = 0.3f; // seconds
-
-	using namespace UE::PoseSearch;
-	
-	OutPhases.Reset();
-	OutPhases.AddDefaulted(SampledBones.Num());
-	
-	const float FiniteDelta = Indexer.GetIndexingContext().Schema->GetSamplingInterval();
-
-	TArray<float> Signal;
-	TArray<float> SmoothedSignal;
-	TArray<LocalMinMax> LocalMinMax;
-	TArray<FVector> BonePositions;
-	for (int32 ChannelBoneIdx = 0; ChannelBoneIdx != SampledBones.Num(); ++ChannelBoneIdx)
-	{
-		const FPoseSearchBone& SampledBone = SampledBones[ChannelBoneIdx];
-		if (EnumHasAnyFlags(SampledBone.Flags, EPoseSearchBoneFlags::Phase))
-		{
-			CollectBonePositions(BonePositions, Indexer, SchemaBoneIdx[ChannelBoneIdx]);
-
-			// @todo: have different way of calculating signals, for example: height of the bone transform, acceleration, etc?
-			const int32 BoneSamplingCentralDifferencesOffset = FMath::Max(FMath::CeilToInt(BoneSamplingCentralDifferencesTime / FiniteDelta), 1);
-			CalculateSignal(BonePositions, Signal, BoneSamplingCentralDifferencesOffset);
-
-			const int32 SmoothingWindowOffset = FMath::Max(FMath::CeilToInt(SmoothingWindowTime / FiniteDelta), 1);
-			SmoothSignal(Signal, SmoothedSignal, SmoothingWindowOffset);
-
-			FindLocalMinMax(SmoothedSignal, LocalMinMax);
-			ValidateLocalMinMax(LocalMinMax);
-
-			ExtrapolateLocalMinMaxBoundaries(LocalMinMax, SmoothedSignal);
-			ValidateLocalMinMax(LocalMinMax);
-			CalculatePhasesFromLocalMinMax(LocalMinMax, OutPhases[ChannelBoneIdx], SmoothedSignal.Num());
-		}
-	}
-}
-
-void UPoseSearchFeatureChannel_Pose::IndexAsset(UE::PoseSearch::IAssetIndexer& Indexer, TArrayView<float> FeatureVectorTable) const
+void UPoseSearchFeatureChannel_Pose::Finalize(UPoseSearchSchema* Schema)
 {
-	using namespace UE::PoseSearch;
-	 
-	// Phases is an array of array with cardinality SampledBones.Num() times NumSamples (IndexingContext.EndSampleIdx - IndexingContext.BeginSampleIdx)
-	// of 2 dimensional vectors (FVector2D) representing phases in an Euclidean space with phase angle sin/cos as direction and certainty of the signal as magnitude,
-	// where certainty is a function of the amplitude of the signal used as input
-	TArray<TArray<FVector2D>> Phases;
-	CalculatePhases(Indexer, Phases);
-
-	const FAssetIndexingContext& IndexingContext = Indexer.GetIndexingContext();
-	for (int32 SampleIdx = IndexingContext.BeginSampleIdx; SampleIdx != IndexingContext.EndSampleIdx; ++SampleIdx)
-	{
-		const int32 VectorIdx = SampleIdx - IndexingContext.BeginSampleIdx;
-		AddPoseFeatures(Indexer, SampleIdx, IndexingContext.GetPoseVector(VectorIdx, FeatureVectorTable), Phases);
-	}
-}
-
-void UPoseSearchFeatureChannel_Pose::AddPoseFeatures(UE::PoseSearch::IAssetIndexer& Indexer, int32 SampleIdx, TArrayView<float> FeatureVector, const TArray<TArray<FVector2D>>& Phases) const
-{
-	// This function samples the instantaneous pose at time t as well as the pose's velocity and acceleration at time t.
-	// Symmetric finite differences are used to approximate derivatives:
-	//	First symmetric derivative:   f'(t) ~ (f(t+h) - f(t-h)) / 2h
-	//	Second symmetric derivative: f''(t) ~ (f(t+h) - 2f(t) + f(t-h)) / h^2
-	// Where h is a constant time delta
-	// So this means three pose extractions are taken at time t-h, t, and t+h
-	
-	using namespace UE::PoseSearch;
-	using FSampleInfo = IAssetIndexer::FSampleInfo;
-
-	if (SampledBones.IsEmpty())
-	{
-		return;
-	}
-
-	const FAssetIndexingContext& IndexingContext = Indexer.GetIndexingContext();
-	const FAssetSamplingContext* SamplingContext = IndexingContext.SamplingContext;
+	SubChannels.Reset();
 
-	const float SampleTime = FMath::Min(SampleIdx * IndexingContext.Schema->GetSamplingInterval(), IndexingContext.AssetSampler->GetPlayLength());
-
-	int32 DataOffset = ChannelDataOffset;
 	for (int32 ChannelBoneIdx = 0; ChannelBoneIdx != SampledBones.Num(); ++ChannelBoneIdx)
 	{
 		const FPoseSearchBone& SampledBone = SampledBones[ChannelBoneIdx];
-
-		// Get each bone's component transform, velocity, and acceleration and add accumulated root motion at this time offset
-		// Think of this process as freezing the character in place (at SampleTime) and then tracing the paths of their joints
-		// as they move through space from past to present to future (at times indicated by PoseSampleTimes).
-
 		if (EnumHasAnyFlags(SampledBone.Flags, EPoseSearchBoneFlags::Position))
 		{
-			bool ClampedPresent;
-			const FTransform BoneTransformsPresent = Indexer.GetTransformAndCacheResults(SampleTime, SampleTime, SchemaBoneIdx[ChannelBoneIdx], ClampedPresent);
-			FFeatureVectorHelper::EncodeVector(FeatureVector, DataOffset, BoneTransformsPresent.GetTranslation());
+			UPoseSearchFeatureChannel_Position* Position = NewObject<UPoseSearchFeatureChannel_Position>();
+			Position->Bone = SampledBone.Reference;
+			Position->Weight = SampledBone.Weight * Weight;
+			Position->SampleTimeOffset = 0.f;
+			Position->ColorPresetIndex = SampledBone.ColorPresetIndex;
+			Position->InputQueryPose = InputQueryPose;
+			Position->SetFlags(RF_Transient);
+			SubChannels.Add(Position);
 		}
 
 		if (EnumHasAnyFlags(SampledBone.Flags, EPoseSearchBoneFlags::Rotation))
 		{
-			bool ClampedPresent;
-			const FTransform BoneTransformsPresent = Indexer.GetTransformAndCacheResults(SampleTime, SampleTime, SchemaBoneIdx[ChannelBoneIdx], ClampedPresent);
-			FFeatureVectorHelper::EncodeQuat(FeatureVector, DataOffset, BoneTransformsPresent.GetRotation());
+			UPoseSearchFeatureChannel_Heading* HeadingX = NewObject<UPoseSearchFeatureChannel_Heading>();
+			HeadingX->Bone = SampledBone.Reference;
+			HeadingX->Weight = SampledBone.Weight * Weight;
+			HeadingX->SampleTimeOffset = 0.f;
+			HeadingX->HeadingAxis = EHeadingAxis::X;
+			HeadingX->ColorPresetIndex = SampledBone.ColorPresetIndex;
+			HeadingX->InputQueryPose = InputQueryPose;
+			HeadingX->SetFlags(RF_Transient);
+			SubChannels.Add(HeadingX);
+
+			UPoseSearchFeatureChannel_Heading* HeadingY = NewObject<UPoseSearchFeatureChannel_Heading>();
+			HeadingY->Bone = SampledBone.Reference;
+			HeadingY->Weight = SampledBone.Weight * Weight;
+			HeadingY->SampleTimeOffset = 0.f;
+			HeadingY->HeadingAxis = EHeadingAxis::Y;
+			HeadingY->ColorPresetIndex = SampledBone.ColorPresetIndex;
+			HeadingY->InputQueryPose = InputQueryPose;
+			HeadingY->SetFlags(RF_Transient);
+			SubChannels.Add(HeadingY);
 		}
 
 		if (EnumHasAnyFlags(SampledBone.Flags, EPoseSearchBoneFlags::Velocity))
 		{
-			bool ClampedPast, ClampedPresent, ClampedFuture;
-			const FTransform BoneTransformsPast = Indexer.GetTransformAndCacheResults(SampleTime - SamplingContext->FiniteDelta, bUseCharacterSpaceVelocities ? SampleTime - SamplingContext->FiniteDelta : SampleTime, SchemaBoneIdx[ChannelBoneIdx], ClampedPast);
-			const FTransform BoneTransformsPresent = Indexer.GetTransformAndCacheResults(SampleTime, SampleTime, SchemaBoneIdx[ChannelBoneIdx], ClampedPresent);
-			const FTransform BoneTransformsFuture = Indexer.GetTransformAndCacheResults(SampleTime + SamplingContext->FiniteDelta, bUseCharacterSpaceVelocities ? SampleTime + SamplingContext->FiniteDelta : SampleTime, SchemaBoneIdx[ChannelBoneIdx], ClampedFuture);
-
-			// We can get a better finite difference if we ignore samples that have
-			// been clamped at either side of the clip. However, if the central sample 
-			// itself is clamped, or there are no samples that are clamped, we can just 
-			// use the central difference as normal.
-			FVector LinearVelocity;
-			if (ClampedPast && !ClampedPresent && !ClampedFuture)
-			{
-				LinearVelocity = (BoneTransformsFuture.GetTranslation() - BoneTransformsPresent.GetTranslation()) / SamplingContext->FiniteDelta;
-			}
-			else if (ClampedFuture && !ClampedPresent && !ClampedPast)
-			{
-				LinearVelocity = (BoneTransformsPresent.GetTranslation() - BoneTransformsPast.GetTranslation()) / SamplingContext->FiniteDelta;
-			}
-			else
-			{
-				LinearVelocity = (BoneTransformsFuture.GetTranslation() - BoneTransformsPast.GetTranslation()) / (SamplingContext->FiniteDelta * 2.f);
-			}
-
-			FFeatureVectorHelper::EncodeVector(FeatureVector, DataOffset, LinearVelocity);
+			UPoseSearchFeatureChannel_Velocity* Velocity = NewObject<UPoseSearchFeatureChannel_Velocity>();
+			Velocity->Bone = SampledBone.Reference;
+			Velocity->Weight = SampledBone.Weight * Weight;
+			Velocity->SampleTimeOffset = 0.f;
+			Velocity->ColorPresetIndex = SampledBone.ColorPresetIndex;
+			Velocity->InputQueryPose = InputQueryPose;
+			Velocity->bUseCharacterSpaceVelocities = bUseCharacterSpaceVelocities;
+			Velocity->SetFlags(RF_Transient);
+			SubChannels.Add(Velocity);
 		}
 
 		if (EnumHasAnyFlags(SampledBone.Flags, EPoseSearchBoneFlags::Phase))
 		{
-			const int32 VectorIdx = SampleIdx - IndexingContext.BeginSampleIdx;
-			FFeatureVectorHelper::EncodeVector2D(FeatureVector, DataOffset, Phases[ChannelBoneIdx][VectorIdx]);
+			UPoseSearchFeatureChannel_Phase* Phase = NewObject<UPoseSearchFeatureChannel_Phase>();
+			Phase->Bone = SampledBone.Reference;
+			Phase->Weight = SampledBone.Weight * Weight;
+			Phase->ColorPresetIndex = SampledBone.ColorPresetIndex;
+			Phase->InputQueryPose = InputQueryPose;
+			Phase->SetFlags(RF_Transient);
+			SubChannels.Add(Phase);
 		}
 	}
 
-	check(DataOffset == ChannelDataOffset + ChannelCardinality);
+	Super::Finalize(Schema);
 }
 
-void UPoseSearchFeatureChannel_Pose::BuildQuery(UE::PoseSearch::FSearchContext& SearchContext, FPoseSearchFeatureVectorBuilder& InOutQuery) const
+#if WITH_EDITOR
+void UPoseSearchFeatureChannel_Pose::PopulateChannelLayoutSet(UE::PoseSearch::FFeatureChannelLayoutSet& FeatureChannelLayoutSet) const
 {
-	using namespace UE::PoseSearch;
-
-	const bool bIsCurrentResultValid = SearchContext.CurrentResult.IsValid();
-	const bool bSkip = InputQueryPose != EInputQueryPose::UseCharacterPose && bIsCurrentResultValid && SearchContext.CurrentResult.Database->Schema == InOutQuery.GetSchema();
-	if (bSkip || !SearchContext.History)
-	{
-		if (bIsCurrentResultValid)
-		{
-			const float LerpValue = InputQueryPose == EInputQueryPose::UseInterpolatedContinuingPose ? SearchContext.CurrentResult.LerpValue : 0.f;
-			int32 DataOffset = ChannelDataOffset;
-			for (int32 SampledBoneIdx = 0; SampledBoneIdx != SampledBones.Num(); ++SampledBoneIdx)
-			{
-				const FPoseSearchBone& SampledBone = SampledBones[SampledBoneIdx];
-				if (EnumHasAnyFlags(SampledBone.Flags, EPoseSearchBoneFlags::Position))
-				{
-					FFeatureVectorHelper::EncodeVector(InOutQuery.EditValues(), DataOffset, SearchContext.GetCurrentResultPrevPoseVector(), SearchContext.GetCurrentResultPoseVector(), SearchContext.GetCurrentResultNextPoseVector(), LerpValue);
-				}
-
-				if (EnumHasAnyFlags(SampledBone.Flags, EPoseSearchBoneFlags::Rotation))
-				{
-					FFeatureVectorHelper::EncodeQuat(InOutQuery.EditValues(), DataOffset, SearchContext.GetCurrentResultPrevPoseVector(), SearchContext.GetCurrentResultPoseVector(), SearchContext.GetCurrentResultNextPoseVector(), LerpValue);
-				}
-
-				if (EnumHasAnyFlags(SampledBone.Flags, EPoseSearchBoneFlags::Velocity))
-				{
-					FFeatureVectorHelper::EncodeVector(InOutQuery.EditValues(), DataOffset, SearchContext.GetCurrentResultPrevPoseVector(), SearchContext.GetCurrentResultPoseVector(), SearchContext.GetCurrentResultNextPoseVector(), LerpValue);
-				}
-
-				if (EnumHasAnyFlags(SampledBone.Flags, EPoseSearchBoneFlags::Phase))
-				{
-					FFeatureVectorHelper::EncodeVector2D(InOutQuery.EditValues(), DataOffset, SearchContext.GetCurrentResultPrevPoseVector(), SearchContext.GetCurrentResultPoseVector(), SearchContext.GetCurrentResultNextPoseVector(), LerpValue);
-				}
-			}
-		}
-		// else leave the InOutQuery set to zero since the SearchContext.History is invalid and it'll fail if we continue
-	}
-	else
+	for (const TObjectPtr<UPoseSearchFeatureChannel>& SubChannelPtr : GetSubChannels())
 	{
-		struct CachedTransforms
+		FString BoneName;
+		EPoseSearchBoneFlags BoneFlag;
+		const char* Label;
+		const UPoseSearchFeatureChannel* Channel = SubChannelPtr.Get();
+		if (const UPoseSearchFeatureChannel_Position* Position = Cast<UPoseSearchFeatureChannel_Position>(Channel))
 		{
-			FTransform Current;
-			FTransform Previous;
-			bool Valid = false; // @todo: remove this
-		};
-		TArray<CachedTransforms, TInlineAllocator<32>> CachedTransforms;
-		CachedTransforms.AddUninitialized(SampledBones.Num());
-
-		float SampleTime = 0.f;
-
-		for (int32 SampledBoneIdx = 0; SampledBoneIdx != SampledBones.Num(); ++SampledBoneIdx)
-		{
-			CachedTransforms[SampledBoneIdx].Current = SearchContext.TryGetTransformAndCacheResults(SampleTime, InOutQuery.GetSchema(), SchemaBoneIdx[SampledBoneIdx]);
-			const FPoseSearchBone& SampledBone = SampledBones[SampledBoneIdx];
-
-			if (EnumHasAnyFlags(SampledBone.Flags, EPoseSearchBoneFlags::Velocity))
-			{
-				check(SearchContext.History);
-				const float HistorySameplInterval = SearchContext.History->GetSampleTimeInterval();
-
-				CachedTransforms[SampledBoneIdx].Previous = SearchContext.TryGetTransformAndCacheResults(SampleTime - HistorySameplInterval, InOutQuery.GetSchema(), SchemaBoneIdx[SampledBoneIdx]);
-
-				if (!bUseCharacterSpaceVelocities)
-				{
-					const FTransform RootTransform = SearchContext.TryGetTransformAndCacheResults(SampleTime, InOutQuery.GetSchema(), FSearchContext::SchemaRootBoneIdx);
-					const FTransform RootTransformPrev = SearchContext.TryGetTransformAndCacheResults(SampleTime - HistorySameplInterval, InOutQuery.GetSchema(), FSearchContext::SchemaRootBoneIdx);
-
-					// animation space velocity
-					CachedTransforms[SampledBoneIdx].Previous = CachedTransforms[SampledBoneIdx].Previous * (RootTransformPrev * RootTransform.Inverse());
-				}
-			}
-			CachedTransforms[SampledBoneIdx].Valid = true;
+			BoneName = Position->Bone.BoneName.ToString();
+			BoneFlag = EPoseSearchBoneFlags::Position;
+			Label = "Pos";
 		}
-
-		int32 DataOffset = ChannelDataOffset;
-		for (int32 SampledBoneIdx = 0; SampledBoneIdx != SampledBones.Num(); ++SampledBoneIdx)
+		else if (const UPoseSearchFeatureChannel_Heading* Heading = Cast<UPoseSearchFeatureChannel_Heading>(Channel))
 		{
-			const FPoseSearchBone& SampledBone = SampledBones[SampledBoneIdx];
-			if (EnumHasAnyFlags(SampledBone.Flags, EPoseSearchBoneFlags::Position))
-			{
-				if (CachedTransforms[SampledBoneIdx].Valid)
-				{
-					FFeatureVectorHelper::EncodeVector(InOutQuery.EditValues(), DataOffset, CachedTransforms[SampledBoneIdx].Current.GetTranslation());
-				}
-				else
-				{
-					// preserve the InOutQuery.EditValues() and increase the DataOffset
-					DataOffset += FFeatureVectorHelper::EncodeVectorCardinality;
-				}
-			}
-
-			if (EnumHasAnyFlags(SampledBone.Flags, EPoseSearchBoneFlags::Rotation))
-			{
-				if (CachedTransforms[SampledBoneIdx].Valid)
-				{
-					FFeatureVectorHelper::EncodeQuat(InOutQuery.EditValues(), DataOffset, CachedTransforms[SampledBoneIdx].Current.GetRotation());
-				}
-				else
-				{
-					// preserve the InOutQuery.EditValues() and increase the DataOffset
-					DataOffset += FFeatureVectorHelper::EncodeQuatCardinality;
-				}
-			}
+			BoneName = Heading->Bone.BoneName.ToString();
+			BoneFlag = EPoseSearchBoneFlags::Rotation;
 
-			if (EnumHasAnyFlags(SampledBone.Flags, EPoseSearchBoneFlags::Velocity))
+			if (Heading->HeadingAxis == EHeadingAxis::X)
 			{
-				if (CachedTransforms[SampledBoneIdx].Valid)
-				{
-					const FVector LinearVelocity = (CachedTransforms[SampledBoneIdx].Current.GetTranslation() - CachedTransforms[SampledBoneIdx].Previous.GetTranslation()) / SearchContext.History->GetSampleTimeInterval();
-					FFeatureVectorHelper::EncodeVector(InOutQuery.EditValues(), DataOffset, LinearVelocity);
-				}
-				else
-				{
-					// preserve the InOutQuery.EditValues() and increase the DataOffset
-					DataOffset += FFeatureVectorHelper::EncodeVectorCardinality;
-				}
-			}
-
-			if (EnumHasAnyFlags(SampledBone.Flags, EPoseSearchBoneFlags::Phase))
-			{
-				DataOffset += FFeatureVectorHelper::EncodeVector2DCardinality;
-
-				// @todo: Support phase in BuildQuery
-				// FFeatureVectorHelper::EncodeVector2D(InOutQuery.EditValues(), DataOffset, ???);
-			}
-		}
-
-		check(DataOffset == ChannelDataOffset + ChannelCardinality);
-	}
-}
-
-void UPoseSearchFeatureChannel_Pose::DebugDraw(const UE::PoseSearch::FDebugDrawParams& DrawParams, TConstArrayView<float> PoseVector) const
-{
-#if ENABLE_DRAW_DEBUG
-	using namespace UE::PoseSearch;
-
-	const UPoseSearchSchema* Schema = DrawParams.GetSchema();
-	check(Schema && Schema->IsValid());
-
-	const float LifeTime = DrawParams.DefaultLifeTime;
-	const uint8 DepthPriority = ESceneDepthPriorityGroup::SDPG_Foreground + 2;
-	const bool bPersistent = EnumHasAnyFlags(DrawParams.Flags, EDebugDrawFlags::Persistent);
-
-	const int32 NumBones = SampledBones.Num();
-
-	if (NumBones == 0)
-	{
-		return;
-	}
-
-	int32 DataOffset = ChannelDataOffset;
-	for (int32 ChannelBoneIdx = 0; ChannelBoneIdx != NumBones; ++ChannelBoneIdx)
-	{
-		const FPoseSearchBone& SampledBone = SampledBones[ChannelBoneIdx];
-
-		FVector BonePos;
-		if (EnumHasAnyFlags(SampledBone.Flags, EPoseSearchBoneFlags::Position))
-		{
-			BonePos = FFeatureVectorHelper::DecodeVector(PoseVector, DataOffset);
-
-			const FColor Color = DrawParams.GetColor(SampledBone.ColorPresetIndex);
-
-			BonePos = DrawParams.RootTransform.TransformPosition(BonePos);
-			if (EnumHasAnyFlags(DrawParams.Flags, EDebugDrawFlags::DrawFast | EDebugDrawFlags::DrawSearchIndex))
-			{
-				DrawDebugPoint(DrawParams.World, BonePos, DrawParams.PointSize, Color, bPersistent, LifeTime, DepthPriority);
+				Label = "HdX";
 			}
 			else
 			{
-				DrawDebugSphere(DrawParams.World, BonePos, 2.f, 8, Color, bPersistent, LifeTime, DepthPriority);
-			}
-
-			if (EnumHasAnyFlags(DrawParams.Flags, EDebugDrawFlags::DrawBoneNames))
-			{
-				DrawDebugString(
-					DrawParams.World, BonePos + FVector(0.0, 0.0, 10.0),
-					Schema->BoneReferences[SchemaBoneIdx[ChannelBoneIdx]].BoneName.ToString(),
-					nullptr, Color, LifeTime, false, 1.f);
+				check(Heading->HeadingAxis == EHeadingAxis::Y);
+				Label = "HdY";
 			}
 		}
-		else
+		else if (const UPoseSearchFeatureChannel_Velocity* Velocity = Cast<UPoseSearchFeatureChannel_Velocity>(Channel))
 		{
-			BonePos = DrawParams.Mesh != nullptr ? DrawParams.Mesh->GetSocketTransform(SampledBones[ChannelBoneIdx].Reference.BoneName).GetLocation() : DrawParams.RootTransform.GetTranslation();
+			BoneName = Velocity->Bone.BoneName.ToString();
+			BoneFlag = EPoseSearchBoneFlags::Velocity;
+			Label = "Vel";
 		}
-
-		if (EnumHasAnyFlags(SampledBone.Flags, EPoseSearchBoneFlags::Rotation))
+		else if (const UPoseSearchFeatureChannel_Phase* Phase = Cast<UPoseSearchFeatureChannel_Phase>(Channel))
 		{
-			const FQuat BoneRot = FFeatureVectorHelper::DecodeQuat(PoseVector, DataOffset);
-			const FRotator Rotator(DrawParams.RootTransform.GetRotation() * BoneRot);
-			const float AdjustedThickness = EnumHasAnyFlags(DrawParams.Flags, EDebugDrawFlags::DrawFast) ? 0.f : 1.f;
-			DrawDebugCoordinateSystem(DrawParams.World, BonePos, Rotator, 15.f, bPersistent, LifeTime, DepthPriority, AdjustedThickness);
+			BoneName = Phase->Bone.BoneName.ToString();
+			BoneFlag = EPoseSearchBoneFlags::Phase;
+			Label = "Pha";
 		}
-
-		if (EnumHasAnyFlags(SampledBone.Flags, EPoseSearchBoneFlags::Velocity))
+		else
 		{
-			FVector BoneVel = FFeatureVectorHelper::DecodeVector(PoseVector, DataOffset);
-
-			const FColor Color = DrawParams.GetColor(SampledBone.ColorPresetIndex);
-
-			BoneVel *= 0.08f;
-			BoneVel = DrawParams.RootTransform.TransformVector(BoneVel);
-			FVector BoneVelDirection = BoneVel.GetSafeNormal();
-
-			if (EnumHasAnyFlags(DrawParams.Flags, EDebugDrawFlags::DrawSearchIndex))
-			{
-				DrawDebugLine(DrawParams.World, BonePos, BonePos + BoneVel, Color, bPersistent, LifeTime, DepthPriority);
-			}
-			else
-			{
-				const float AdjustedThickness = EnumHasAnyFlags(DrawParams.Flags, EDebugDrawFlags::DrawFast) ? 0.f : 1.f;
-
-				DrawDebugLine(
-					DrawParams.World,
-					BonePos + BoneVelDirection * 2.f,
-					BonePos + BoneVel,
-					Color,
-					bPersistent,
-					LifeTime,
-					DepthPriority,
-					AdjustedThickness);
-			}
+			checkNoEntry();
+			continue;
 		}
-
-		if (EnumHasAnyFlags(SampledBone.Flags, EPoseSearchBoneFlags::Phase))
-		{
-			const FVector2D Phase = FFeatureVectorHelper::DecodeVector2D(PoseVector, DataOffset);
-
-			const FColor Color = DrawParams.GetColor(SampledBone.ColorPresetIndex);
-
-			static float ScaleFactor = 1.f;
-
-			const FVector TransformXAxisVector = DrawParams.RootTransform.TransformVector(FVector::XAxisVector);
-			const FVector TransformYAxisVector = DrawParams.RootTransform.TransformVector(FVector::YAxisVector);
-			const FVector TransformZAxisVector = DrawParams.RootTransform.TransformVector(FVector::ZAxisVector);
-
-			const FVector PhaseVector = (TransformZAxisVector * Phase.X + TransformYAxisVector * Phase.Y) * ScaleFactor;
-			DrawDebugLine(DrawParams.World, BonePos, BonePos + PhaseVector, Color, bPersistent, LifeTime, DepthPriority, 0.f);
-
-			static int32 Segments = 64;
-			FMatrix CircleTransform;
-			CircleTransform.SetAxes(&TransformXAxisVector, &TransformYAxisVector, &TransformZAxisVector, &BonePos);
-			DrawDebugCircle(DrawParams.World, CircleTransform, PhaseVector.Length(), Segments, Color, bPersistent, LifeTime, DepthPriority, 0.f, false);
-		}
-	}
-
-	check(DataOffset == ChannelDataOffset + ChannelCardinality);
-#endif // ENABLE_DRAW_DEBUG
-}
-
-#if WITH_EDITOR
-void UPoseSearchFeatureChannel_Pose::PopulateChannelLayoutSet(UE::PoseSearch::FFeatureChannelLayoutSet& FeatureChannelLayoutSet) const
-{
-	using namespace UE::PoseSearch;
-	int32 DataOffset = ChannelDataOffset;
-
-	auto Add = [&FeatureChannelLayoutSet, &DataOffset](const FPoseSearchBone& SampledBone, EPoseSearchBoneFlags BoneFlag, const char* Label, int32 Cardinality)
-	{
+		
 		FString SkeletonName = FeatureChannelLayoutSet.CurrentSchema->Skeleton->GetName();
-		FString BoneName = SampledBone.Reference.BoneName.ToString();
-
 		UE::PoseSearch::FKeyBuilder KeyBuilder;
 		KeyBuilder << SkeletonName << BoneName << BoneFlag;
-		FeatureChannelLayoutSet.Add(FString::Format(TEXT("{0} {1}"), { BoneName, Label }), KeyBuilder.Finalize(), DataOffset, Cardinality);
-
-		DataOffset += Cardinality;
-	};
-
-	for (int32 ChannelBoneIdx = 0; ChannelBoneIdx != SampledBones.Num(); ++ChannelBoneIdx)
-	{
-		const FPoseSearchBone& SampledBone = SampledBones[ChannelBoneIdx];
-		if (EnumHasAnyFlags(SampledBone.Flags, EPoseSearchBoneFlags::Position))
-		{
-			Add(SampledBone, EPoseSearchBoneFlags::Position, "Pos", FFeatureVectorHelper::EncodeVectorCardinality);
-		}
-		if (EnumHasAnyFlags(SampledBone.Flags, EPoseSearchBoneFlags::Rotation))
-		{
-			Add(SampledBone, EPoseSearchBoneFlags::Rotation, "Rot", FFeatureVectorHelper::EncodeQuatCardinality);
-		}
-		if (EnumHasAnyFlags(SampledBone.Flags, EPoseSearchBoneFlags::Velocity))
-		{
-			Add(SampledBone, EPoseSearchBoneFlags::Velocity, "Vel", FFeatureVectorHelper::EncodeVectorCardinality);
-		}
-		if (EnumHasAnyFlags(SampledBone.Flags, EPoseSearchBoneFlags::Phase))
-		{
-			Add(SampledBone, EPoseSearchBoneFlags::Phase, "Pha", FFeatureVectorHelper::EncodeVector2DCardinality);
-		}
+		FeatureChannelLayoutSet.Add(FString::Format(TEXT("{0} {1}"), { BoneName, Label }), KeyBuilder.Finalize(), Channel->GetChannelDataOffset(), Channel->GetChannelCardinality());
 	}
-
-	check(DataOffset == ChannelDataOffset + ChannelCardinality);
 }
 
 void UPoseSearchFeatureChannel_Pose::ComputeCostBreakdowns(UE::PoseSearch::ICostBreakDownData& CostBreakDownData, const UPoseSearchSchema* Schema) const
 {
-	using namespace UE::PoseSearch;
-
 	CostBreakDownData.AddEntireBreakDownSection(LOCTEXT("ColumnLabelPoseChannelTotal", "Pose Total"), Schema, ChannelDataOffset, ChannelCardinality);
 
 	if (CostBreakDownData.IsVerbose())
 	{
-		int32 DataOffset = ChannelDataOffset;
-
-		for (int32 ChannelBoneIdx = 0; ChannelBoneIdx != SampledBones.Num(); ++ChannelBoneIdx)
+		for (const TObjectPtr<UPoseSearchFeatureChannel>& SubChannelPtr : GetSubChannels())
 		{
-			const FPoseSearchBone& SampledBone = SampledBones[ChannelBoneIdx];
-			if (EnumHasAnyFlags(SampledBone.Flags, EPoseSearchBoneFlags::Position))
+			const UPoseSearchFeatureChannel* Channel = SubChannelPtr.Get();
+			if (const UPoseSearchFeatureChannel_Position* Position = Cast<UPoseSearchFeatureChannel_Position>(Channel))
 			{
-				CostBreakDownData.AddEntireBreakDownSection(FText::Format(LOCTEXT("ColumnLabelPoseChannelPosition", "{0} Pos"), FText::FromName(SampledBone.Reference.BoneName)), Schema, DataOffset, FFeatureVectorHelper::EncodeVectorCardinality);
-				DataOffset += FFeatureVectorHelper::EncodeVectorCardinality;
+				CostBreakDownData.AddEntireBreakDownSection(FText::Format(LOCTEXT("ColumnLabelPoseChannelPosition", "{0} Pos"), FText::FromName(Position->Bone.BoneName)), Schema, Position->GetChannelDataOffset(), Position->GetChannelCardinality());
 			}
-			if (EnumHasAnyFlags(SampledBone.Flags, EPoseSearchBoneFlags::Rotation))
+			else if (const UPoseSearchFeatureChannel_Heading* Heading = Cast<UPoseSearchFeatureChannel_Heading>(Channel))
 			{
-				CostBreakDownData.AddEntireBreakDownSection(FText::Format(LOCTEXT("ColumnLabelPoseChannelRotation", "{0} Rot"), FText::FromName(SampledBone.Reference.BoneName)), Schema, DataOffset, FFeatureVectorHelper::EncodeQuatCardinality);
-				DataOffset += FFeatureVectorHelper::EncodeQuatCardinality;
+				if (Heading->HeadingAxis == EHeadingAxis::X)
+				{
+					CostBreakDownData.AddEntireBreakDownSection(FText::Format(LOCTEXT("ColumnLabelPoseChannelHeading", "{0} HdX"), FText::FromName(Heading->Bone.BoneName)), Schema, Heading->GetChannelDataOffset(), Heading->GetChannelCardinality());
+				}
+				else
+				{
+					check(Heading->HeadingAxis == EHeadingAxis::Y);
+					CostBreakDownData.AddEntireBreakDownSection(FText::Format(LOCTEXT("ColumnLabelPoseChannelHeading", "{0} HdY"), FText::FromName(Heading->Bone.BoneName)), Schema, Heading->GetChannelDataOffset(), Heading->GetChannelCardinality());
+				}
 			}
-			if (EnumHasAnyFlags(SampledBone.Flags, EPoseSearchBoneFlags::Velocity))
+			else if (const UPoseSearchFeatureChannel_Velocity* Velocity = Cast<UPoseSearchFeatureChannel_Velocity>(Channel))
 			{
-				CostBreakDownData.AddEntireBreakDownSection(FText::Format(LOCTEXT("ColumnLabelPoseChannelVelocity", "{0} Vel"), FText::FromName(SampledBone.Reference.BoneName)), Schema, DataOffset, FFeatureVectorHelper::EncodeVectorCardinality);
-				DataOffset += FFeatureVectorHelper::EncodeVectorCardinality;
+				CostBreakDownData.AddEntireBreakDownSection(FText::Format(LOCTEXT("ColumnLabelPoseChannelVelocity", "{0} Vel"), FText::FromName(Velocity->Bone.BoneName)), Schema, Velocity->GetChannelDataOffset(), Velocity->GetChannelCardinality());
 			}
-			if (EnumHasAnyFlags(SampledBone.Flags, EPoseSearchBoneFlags::Phase))
+			else if (const UPoseSearchFeatureChannel_Phase* Phase = Cast<UPoseSearchFeatureChannel_Phase>(Channel))
 			{
-				CostBreakDownData.AddEntireBreakDownSection(FText::Format(LOCTEXT("ColumnLabelPoseChannelPhase", "{0} Pha"), FText::FromName(SampledBone.Reference.BoneName)), Schema, DataOffset, FFeatureVectorHelper::EncodeVector2DCardinality);
-				DataOffset += FFeatureVectorHelper::EncodeVector2DCardinality;
+				CostBreakDownData.AddEntireBreakDownSection(FText::Format(LOCTEXT("ColumnLabelPoseChannelPhase", "{0} Pha"), FText::FromName(Phase->Bone.BoneName)), Schema, Phase->GetChannelDataOffset(), Phase->GetChannelCardinality());
+			}
+			else
+			{
+				checkNoEntry();
 			}
 		}
-
-		check(DataOffset == ChannelDataOffset + ChannelCardinality);
 	}
 }
 
diff --git a/Engine/Plugins/Experimental/Animation/PoseSearch/Source/Runtime/Private/PoseSearchFeatureChannel_Pose.h b/Engine/Plugins/Experimental/Animation/PoseSearch/Source/Runtime/Private/PoseSearchFeatureChannel_Pose.h
--- a/Engine/Plugins/Experimental/Animation/PoseSearch/Source/Runtime/Private/PoseSearchFeatureChannel_Pose.h
+++ b/Engine/Plugins/Experimental/Animation/PoseSearch/Source/Runtime/Private/PoseSearchFeatureChannel_Pose.h
@@ -2,8 +2,7 @@
 
 #pragma once
 
-#include "PoseSearch/PoseSearchFeatureChannel.h"
-#include "BoneContainer.h"
+#include "PoseSearchFeatureChannel_Group.h"
 #include "PoseSearchFeatureChannel_Pose.generated.h"
 
 UENUM(meta = (Bitflags, UseEnumValuesAsMaskValuesInEditor = "true"))
@@ -38,7 +37,7 @@ struct POSESEARCH_API FPoseSearchBone
 
 // UPoseSearchFeatureChannel_Pose
 UCLASS(BlueprintType, EditInlineNew, meta = (DisplayName = "Pose Channel"), CollapseCategories)
-class POSESEARCH_API UPoseSearchFeatureChannel_Pose : public UPoseSearchFeatureChannel
+class POSESEARCH_API UPoseSearchFeatureChannel_Pose : public UPoseSearchFeatureChannel_GroupBase
 {
 	GENERATED_BODY()
 
@@ -55,25 +54,24 @@ public:
 	UPROPERTY(EditAnywhere, Category = "Settings")
 	EInputQueryPose InputQueryPose = EInputQueryPose::UseContinuingPose;
 
-	// if bUseCharacterSpaceVelocities is true, velocities will be calculated as from the positions in character space, otherwise they will be calculated using global space positions
+	// if bUseCharacterSpaceVelocities is true, velocities will be calculated from the positions in character space, otherwise they will be calculated using global space positions
 	UPROPERTY(EditAnywhere, Category = "Settings")
 	bool bUseCharacterSpaceVelocities = true;
 
+	UPROPERTY(Transient)
+	TArray<TObjectPtr<UPoseSearchFeatureChannel>> SubChannels;
+
+	// UPoseSearchFeatureChannel_GroupBase interface
+	virtual TArrayView<TObjectPtr<UPoseSearchFeatureChannel>> GetSubChannels() override { return SubChannels; }
+	virtual TConstArrayView<TObjectPtr<UPoseSearchFeatureChannel>> GetSubChannels() const override { return SubChannels; }
+
 	// UPoseSearchFeatureChannel interface
-	virtual void InitializeSchema(UPoseSearchSchema* Schema) override;
-	virtual void FillWeights(TArray<float>& Weights) const override;
-	virtual void IndexAsset(UE::PoseSearch::IAssetIndexer& Indexer, TArrayView<float> FeatureVectorTable) const override;
-	virtual void BuildQuery(UE::PoseSearch::FSearchContext& SearchContext, FPoseSearchFeatureVectorBuilder& InOutQuery) const override;
-	virtual void DebugDraw(const UE::PoseSearch::FDebugDrawParams& DrawParams, TConstArrayView<float> PoseVector) const override;
+	virtual void Finalize(UPoseSearchSchema* Schema) override;
 
 #if WITH_EDITOR
 	virtual void PopulateChannelLayoutSet(UE::PoseSearch::FFeatureChannelLayoutSet& FeatureChannelLayoutSet) const override;
 	virtual void ComputeCostBreakdowns(UE::PoseSearch::ICostBreakDownData& CostBreakDownData, const UPoseSearchSchema* Schema) const override;
 #endif
-
-protected:
-	void AddPoseFeatures(UE::PoseSearch::IAssetIndexer& Indexer, int32 SampleIdx, TArrayView<float> FeatureVector, const TArray<TArray<FVector2D>>& Phases) const;
-	void CalculatePhases(UE::PoseSearch::IAssetIndexer& Indexer, TArray<TArray<FVector2D>>& OutPhases) const;
 };
 
 
diff --git a/Engine/Plugins/Experimental/Animation/PoseSearch/Source/Runtime/Private/PoseSearchFeatureChannel_Position.cpp b/Engine/Plugins/Experimental/Animation/PoseSearch/Source/Runtime/Private/PoseSearchFeatureChannel_Position.cpp
--- a/Engine/Plugins/Experimental/Animation/PoseSearch/Source/Runtime/Private/PoseSearchFeatureChannel_Position.cpp
+++ b/Engine/Plugins/Experimental/Animation/PoseSearch/Source/Runtime/Private/PoseSearchFeatureChannel_Position.cpp
@@ -1,6 +1,7 @@
 // Copyright Epic Games, Inc. All Rights Reserved.
 
 #include "PoseSearchFeatureChannel_Position.h"
+#include "Animation/MotionTrajectoryTypes.h"
 #include "DrawDebugHelpers.h"
 #include "PoseSearch/PoseSearchAssetIndexer.h"
 #include "PoseSearch/PoseSearchAssetSampler.h"
@@ -8,7 +9,7 @@
 #include "PoseSearch/PoseSearchDatabase.h"
 #include "PoseSearch/PoseSearchSchema.h"
 
-void UPoseSearchFeatureChannel_Position::InitializeSchema(UPoseSearchSchema* Schema)
+void UPoseSearchFeatureChannel_Position::Finalize(UPoseSearchSchema* Schema)
 {
 	ChannelDataOffset = Schema->SchemaCardinality;
 	ChannelCardinality = UE::PoseSearch::FFeatureVectorHelper::EncodeVectorCardinality;
@@ -51,6 +52,7 @@ void UPoseSearchFeatureChannel_Position::BuildQuery(UE::PoseSearch::FSearchConte
 {
 	using namespace UE::PoseSearch;
 
+	check(InOutQuery.GetSchema());
 	const bool bIsCurrentResultValid = SearchContext.CurrentResult.IsValid();
 	const bool bSkip = InputQueryPose != EInputQueryPose::UseCharacterPose && bIsCurrentResultValid && SearchContext.CurrentResult.Database->Schema == InOutQuery.GetSchema();
 	if (bSkip || !SearchContext.History)
@@ -65,34 +67,57 @@ void UPoseSearchFeatureChannel_Position::BuildQuery(UE::PoseSearch::FSearchConte
 	}
 	else
 	{
-		FTransform Transform = SearchContext.TryGetTransformAndCacheResults(SampleTimeOffset, InOutQuery.GetSchema(), SchemaBoneIdx);
-		const FTransform RootTransform = SearchContext.TryGetTransformAndCacheResults(0.f, InOutQuery.GetSchema(), FSearchContext::SchemaRootBoneIdx);
-		const FTransform RootTransformPrev = SearchContext.TryGetTransformAndCacheResults(SampleTimeOffset, InOutQuery.GetSchema(), FSearchContext::SchemaRootBoneIdx);
-		Transform = Transform * (RootTransformPrev * RootTransform.Inverse());
+		FTransform Transform;
+		if (InOutQuery.GetSchema()->BoneReferences[SchemaBoneIdx].HasValidSetup())
+		{
+			// calculating the Transform in component space for the bone indexed by SchemaBoneIdx
+			Transform = SearchContext.TryGetTransformAndCacheResults(SampleTimeOffset, InOutQuery.GetSchema(), SchemaBoneIdx);
+
+			// if the SampleTimeOffset is not zero we calculate the delta root bone between the root at current time (zero) and the root at SampleTimeOffset (in the past)
+			// so we can offset the Transform by this amount
+			if (!FMath::IsNearlyZero(SampleTimeOffset, UE_KINDA_SMALL_NUMBER))
+			{
+				const FTransform RootTransform = SearchContext.TryGetTransformAndCacheResults(0.f, InOutQuery.GetSchema(), FSearchContext::SchemaRootBoneIdx);
+				const FTransform RootTransformPrev = SearchContext.TryGetTransformAndCacheResults(SampleTimeOffset, InOutQuery.GetSchema(), FSearchContext::SchemaRootBoneIdx);
+				Transform = Transform * (RootTransformPrev * RootTransform.Inverse());
+			}
+		}
+		else
+		{
+			check(SearchContext.Trajectory);
+			// @todo: make this call consistent with Transform = SearchContext.TryGetTransformAndCacheResults(SampleTimeOffset, InOutQuery.GetSchema(), FSearchContext::SchemaRootBoneIdx);
+			const FTrajectorySample TrajectorySample = SearchContext.Trajectory->GetSampleAtTime(SampleTimeOffset);
+			Transform = TrajectorySample.Transform;
+		}
 
 		int32 DataOffset = ChannelDataOffset;
 		FFeatureVectorHelper::EncodeVector(InOutQuery.EditValues(), DataOffset, Transform.GetTranslation());
-		check(DataOffset == ChannelDataOffset + ChannelCardinality);
 	}
 }
 
-void UPoseSearchFeatureChannel_Position::DebugDraw(const UE::PoseSearch::FDebugDrawParams& DrawParams, TConstArrayView<float> PoseVector) const
+void UPoseSearchFeatureChannel_Position::PreDebugDraw(UE::PoseSearch::FDebugDrawParams& DrawParams, TConstArrayView<float> PoseVector) const
 {
 #if ENABLE_DRAW_DEBUG
 	using namespace UE::PoseSearch;
+	const FVector BonePos = DrawParams.RootTransform.TransformPosition(FFeatureVectorHelper::DecodeVectorAtOffset(PoseVector, ChannelDataOffset));
+	DrawParams.AddCachedPosition(SampleTimeOffset, SchemaBoneIdx, BonePos);
+#endif // ENABLE_DRAW_DEBUG
+}
 
-	const UPoseSearchSchema* Schema = DrawParams.GetSchema();
-	check(Schema && Schema->IsValid());
+void UPoseSearchFeatureChannel_Position::DebugDraw(const UE::PoseSearch::FDebugDrawParams& DrawParams, TConstArrayView<float> PoseVector) const
+{
+#if ENABLE_DRAW_DEBUG
+	using namespace UE::PoseSearch;
 
 	const float LifeTime = DrawParams.DefaultLifeTime;
 	const uint8 DepthPriority = ESceneDepthPriorityGroup::SDPG_Foreground + 2;
 	const bool bPersistent = EnumHasAnyFlags(DrawParams.Flags, EDebugDrawFlags::Persistent);
-
-	int32 DataOffset = ChannelDataOffset;
-	const FVector BonePos = DrawParams.RootTransform.TransformPosition(FFeatureVectorHelper::DecodeVector(PoseVector, DataOffset));
-
 	const FColor Color = DrawParams.GetColor(ColorPresetIndex);
 
+	const FVector BonePos = DrawParams.RootTransform.TransformPosition(FFeatureVectorHelper::DecodeVectorAtOffset(PoseVector, ChannelDataOffset));
+	// validating DrawParams.AddCachedPosition 
+	check(DrawParams.GetCachedPosition(SampleTimeOffset, SchemaBoneIdx) == BonePos);
+
 	if (EnumHasAnyFlags(DrawParams.Flags, EDebugDrawFlags::DrawFast | EDebugDrawFlags::DrawSearchIndex))
 	{
 		DrawDebugPoint(DrawParams.World, BonePos, DrawParams.PointSize, Color, bPersistent, LifeTime, DepthPriority);
@@ -104,6 +129,8 @@ void UPoseSearchFeatureChannel_Position::DebugDraw(const UE::PoseSearch::FDebugD
 
 	if (EnumHasAnyFlags(DrawParams.Flags, EDebugDrawFlags::DrawBoneNames))
 	{
+		const UPoseSearchSchema* Schema = DrawParams.GetSchema();
+		check(Schema && Schema->IsValid());
 		DrawDebugString(DrawParams.World, BonePos + FVector(0.0, 0.0, 10.0), Schema->BoneReferences[SchemaBoneIdx].BoneName.ToString(), nullptr, Color, LifeTime, false, 1.0f);
 	}
 #endif // ENABLE_DRAW_DEBUG
diff --git a/Engine/Plugins/Experimental/Animation/PoseSearch/Source/Runtime/Private/PoseSearchFeatureChannel_Position.h b/Engine/Plugins/Experimental/Animation/PoseSearch/Source/Runtime/Private/PoseSearchFeatureChannel_Position.h
--- a/Engine/Plugins/Experimental/Animation/PoseSearch/Source/Runtime/Private/PoseSearchFeatureChannel_Position.h
+++ b/Engine/Plugins/Experimental/Animation/PoseSearch/Source/Runtime/Private/PoseSearchFeatureChannel_Position.h
@@ -31,9 +31,10 @@ public:
 	EInputQueryPose InputQueryPose = EInputQueryPose::UseContinuingPose;
 
 	// UPoseSearchFeatureChannel interface
-	virtual void InitializeSchema(UPoseSearchSchema* Schema) override;
+	virtual void Finalize(UPoseSearchSchema* Schema) override;
 	virtual void FillWeights(TArray<float>& Weights) const override;
 	virtual void IndexAsset(UE::PoseSearch::IAssetIndexer& Indexer, TArrayView<float> FeatureVectorTable) const override;
 	virtual void BuildQuery(UE::PoseSearch::FSearchContext& SearchContext, FPoseSearchFeatureVectorBuilder& InOutQuery) const override;
+	virtual void PreDebugDraw(UE::PoseSearch::FDebugDrawParams& DrawParams, TConstArrayView<float> PoseVector) const override;
 	virtual void DebugDraw(const UE::PoseSearch::FDebugDrawParams& DrawParams, TConstArrayView<float> PoseVector) const override;
 };
diff --git a/Engine/Plugins/Experimental/Animation/PoseSearch/Source/Runtime/Private/PoseSearchFeatureChannel_Trajectory.cpp b/Engine/Plugins/Experimental/Animation/PoseSearch/Source/Runtime/Private/PoseSearchFeatureChannel_Trajectory.cpp
--- a/Engine/Plugins/Experimental/Animation/PoseSearch/Source/Runtime/Private/PoseSearchFeatureChannel_Trajectory.cpp
+++ b/Engine/Plugins/Experimental/Animation/PoseSearch/Source/Runtime/Private/PoseSearchFeatureChannel_Trajectory.cpp
@@ -1,939 +1,247 @@
 // Copyright Epic Games, Inc. All Rights Reserved.
 
 #include "PoseSearchFeatureChannel_Trajectory.h"
-#include "Animation/MotionTrajectoryTypes.h"
+#include "Animation/Skeleton.h"
 #include "DrawDebugHelpers.h"
-#include "PoseSearch/PoseSearchAssetIndexer.h"
-#include "PoseSearch/PoseSearchAssetSampler.h"
 #include "PoseSearch/PoseSearchContext.h"
 #include "PoseSearch/PoseSearchDerivedDataKey.h"
 #include "PoseSearch/PoseSearchSchema.h"
+#include "PoseSearchFeatureChannel_Heading.h"
+#include "PoseSearchFeatureChannel_Phase.h"
+#include "PoseSearchFeatureChannel_Position.h"
+#include "PoseSearchFeatureChannel_Velocity.h"
 
 #define LOCTEXT_NAMESPACE "PoseSearchFeatureChannels"
 
 void UPoseSearchFeatureChannel_Trajectory::PreSave(FObjectPreSaveContext ObjectSaveContext)
 {
 	Samples.Sort([](const FPoseSearchTrajectorySample& a, const FPoseSearchTrajectorySample& b)
-	{
-		return a.Offset < b.Offset;
-	});
+		{
+			return a.Offset < b.Offset;
+		});
 
 	Super::PreSave(ObjectSaveContext);
 }
 
-void UPoseSearchFeatureChannel_Trajectory::InitializeSchema(UPoseSearchSchema* Schema)
+void UPoseSearchFeatureChannel_Trajectory::Finalize(UPoseSearchSchema* Schema)
 {
-	using namespace UE::PoseSearch;
-
-	ChannelDataOffset = Schema->SchemaCardinality;
+	SubChannels.Reset();
 
 	for (const FPoseSearchTrajectorySample& Sample : Samples)
 	{
-		if (EnumHasAnyFlags(Sample.Flags, EPoseSearchTrajectoryFlags::Position))
+		// @todo: implement PositionXY properly as 2 dimension channel
+		if (EnumHasAnyFlags(Sample.Flags, EPoseSearchTrajectoryFlags::Position | EPoseSearchTrajectoryFlags::PositionXY))
 		{
-			Schema->SchemaCardinality += FFeatureVectorHelper::EncodeVectorCardinality;
-		}
-		if (EnumHasAnyFlags(Sample.Flags, EPoseSearchTrajectoryFlags::PositionXY))
-		{
-			Schema->SchemaCardinality += FFeatureVectorHelper::EncodeVector2DCardinality;
+			UPoseSearchFeatureChannel_Position* Position = NewObject<UPoseSearchFeatureChannel_Position>();
+			Position->Weight = Sample.Weight * Weight;
+			Position->SampleTimeOffset = Sample.Offset;
+			Position->ColorPresetIndex = Sample.ColorPresetIndex;
+			Position->InputQueryPose = EInputQueryPose::UseCharacterPose;
+			Position->SetFlags(RF_Transient);
+			SubChannels.Add(Position);
 		}
 
-		if (EnumHasAnyFlags(Sample.Flags, EPoseSearchTrajectoryFlags::Velocity))
+		// @todo: implement VelocityXY properly as 2 dimension channel
+		if (EnumHasAnyFlags(Sample.Flags, EPoseSearchTrajectoryFlags::Velocity | EPoseSearchTrajectoryFlags::VelocityXY))
 		{
-			Schema->SchemaCardinality += FFeatureVectorHelper::EncodeVectorCardinality;
-		}
-		if (EnumHasAnyFlags(Sample.Flags, EPoseSearchTrajectoryFlags::VelocityXY))
-		{
-			Schema->SchemaCardinality += FFeatureVectorHelper::EncodeVector2DCardinality;
+			UPoseSearchFeatureChannel_Velocity* Velocity = NewObject<UPoseSearchFeatureChannel_Velocity>();
+			Velocity->Weight = Sample.Weight * Weight;
+			Velocity->SampleTimeOffset = Sample.Offset;
+			Velocity->ColorPresetIndex = Sample.ColorPresetIndex;
+			Velocity->InputQueryPose = EInputQueryPose::UseCharacterPose;
+			Velocity->bUseCharacterSpaceVelocities = false;
+			Velocity->SetFlags(RF_Transient);
+			SubChannels.Add(Velocity);
 		}
 
-		if (EnumHasAnyFlags(Sample.Flags, EPoseSearchTrajectoryFlags::VelocityDirection))
-		{
-			Schema->SchemaCardinality += FFeatureVectorHelper::EncodeVectorCardinality;
-		}
-		if (EnumHasAnyFlags(Sample.Flags, EPoseSearchTrajectoryFlags::VelocityDirectionXY))
+		// @todo: implement VelocityDirectionXY properly as 2 dimension channel
+		if (EnumHasAnyFlags(Sample.Flags, EPoseSearchTrajectoryFlags::VelocityDirection | EPoseSearchTrajectoryFlags::VelocityDirectionXY))
 		{
-			Schema->SchemaCardinality += FFeatureVectorHelper::EncodeVector2DCardinality;
+			UPoseSearchFeatureChannel_Velocity* Velocity = NewObject<UPoseSearchFeatureChannel_Velocity>();
+			Velocity->Weight = Sample.Weight * Weight;
+			Velocity->SampleTimeOffset = Sample.Offset;
+			Velocity->ColorPresetIndex = Sample.ColorPresetIndex;
+			Velocity->InputQueryPose = EInputQueryPose::UseCharacterPose;
+			Velocity->bUseCharacterSpaceVelocities = false;
+			Velocity->bNormalize = true;
+			Velocity->SetFlags(RF_Transient);
+			SubChannels.Add(Velocity);
 		}
 
-		if (EnumHasAnyFlags(Sample.Flags, EPoseSearchTrajectoryFlags::FacingDirection))
-		{
-			Schema->SchemaCardinality += FFeatureVectorHelper::EncodeVectorCardinality;
-		}
-		if (EnumHasAnyFlags(Sample.Flags, EPoseSearchTrajectoryFlags::FacingDirectionXY))
+		// @todo: implement FacingDirectionXY properly as 2 dimension channel
+		if (EnumHasAnyFlags(Sample.Flags, EPoseSearchTrajectoryFlags::FacingDirection | EPoseSearchTrajectoryFlags::FacingDirectionXY))
 		{
-			Schema->SchemaCardinality += FFeatureVectorHelper::EncodeVector2DCardinality;
+			UPoseSearchFeatureChannel_Heading* Heading = NewObject<UPoseSearchFeatureChannel_Heading>();
+			Heading->Weight = Sample.Weight * Weight;
+			Heading->SampleTimeOffset = Sample.Offset;
+			Heading->ColorPresetIndex = Sample.ColorPresetIndex;
+			Heading->InputQueryPose = EInputQueryPose::UseCharacterPose;
+			Heading->SetFlags(RF_Transient);
+			SubChannels.Add(Heading);
 		}
 	}
 
-	ChannelCardinality = Schema->SchemaCardinality - ChannelDataOffset;
+	Super::Finalize(Schema);
 }
 
-void UPoseSearchFeatureChannel_Trajectory::FillWeights(TArray<float>& Weights) const
+void UPoseSearchFeatureChannel_Trajectory::DebugDraw(const UE::PoseSearch::FDebugDrawParams& DrawParams, TConstArrayView<float> PoseVector) const
 {
+#if ENABLE_DRAW_DEBUG
 	using namespace UE::PoseSearch;
 
-	int32 DataOffset = ChannelDataOffset;
-	for (const FPoseSearchTrajectorySample& Sample : Samples)
+	TArray<const UPoseSearchFeatureChannel_Position*, TInlineAllocator<32>> Positions;
+	for (const TObjectPtr<UPoseSearchFeatureChannel>& SubChannelPtr : GetSubChannels())
 	{
-		if (EnumHasAnyFlags(Sample.Flags, EPoseSearchTrajectoryFlags::Position))
+		if (const UPoseSearchFeatureChannel_Position* Position = Cast<UPoseSearchFeatureChannel_Position>(SubChannelPtr.Get()))
 		{
-			for (int32 i = 0; i != FFeatureVectorHelper::EncodeVectorCardinality; ++i)
-			{
-				Weights[DataOffset + i] = Weight * Sample.Weight;
-			}
-			DataOffset += FFeatureVectorHelper::EncodeVectorCardinality;
-		}
-		if (EnumHasAnyFlags(Sample.Flags, EPoseSearchTrajectoryFlags::PositionXY))
-		{
-			for (int32 i = 0; i != FFeatureVectorHelper::EncodeVector2DCardinality; ++i)
-			{
-				Weights[DataOffset + i] = Weight * Sample.Weight;
-			}
-			DataOffset += FFeatureVectorHelper::EncodeVector2DCardinality;
+			Positions.Add(Position);
 		}
+	}
 
-		if (EnumHasAnyFlags(Sample.Flags, EPoseSearchTrajectoryFlags::Velocity))
-		{
-			for (int32 i = 0; i != FFeatureVectorHelper::EncodeVectorCardinality; ++i)
-			{
-				Weights[DataOffset + i] = Weight * Sample.Weight;
-			}
-			DataOffset += FFeatureVectorHelper::EncodeVectorCardinality;
-		}
-		if (EnumHasAnyFlags(Sample.Flags, EPoseSearchTrajectoryFlags::VelocityXY))
-		{
-			for (int32 i = 0; i != FFeatureVectorHelper::EncodeVector2DCardinality; ++i)
+	if (Positions.Num() >= 2)
+	{
+		Positions.Sort([](const UPoseSearchFeatureChannel_Position& a, const UPoseSearchFeatureChannel_Position& b)
 			{
-				Weights[DataOffset + i] = Weight * Sample.Weight;
-			}
-			DataOffset += FFeatureVectorHelper::EncodeVector2DCardinality;
-		}
+				return a.SampleTimeOffset < b.SampleTimeOffset;
+			});
 
-		if (EnumHasAnyFlags(Sample.Flags, EPoseSearchTrajectoryFlags::VelocityDirection))
-		{
-			for (int32 i = 0; i != FFeatureVectorHelper::EncodeVectorCardinality; ++i)
-			{
-				Weights[DataOffset + i] = Weight * Sample.Weight;
-			}
-			DataOffset += FFeatureVectorHelper::EncodeVectorCardinality;
-		}
-		if (EnumHasAnyFlags(Sample.Flags, EPoseSearchTrajectoryFlags::VelocityDirectionXY))
+		// big enough negative number that prevents PrevTimeOffset * CurrTimeOffset being infinite (there will never be UPoseSearchFeatureChannel_Trajectory trying to match 1000 seconds in the past)
+		float PrevTimeOffset = -1000.f;
+		TArray<FVector, TInlineAllocator<32>> TrajSplinePos;
+		TArray<FColor, TInlineAllocator<32>> TrajSplineColor;
+		for (int32 i = 0; i < Positions.Num(); ++i)
 		{
-			for (int32 i = 0; i != FFeatureVectorHelper::EncodeVector2DCardinality; ++i)
-			{
-				Weights[DataOffset + i] = Weight * Sample.Weight;
-			}
-			DataOffset += FFeatureVectorHelper::EncodeVector2DCardinality;
-		}
+			const float CurrTimeOffset = Positions[i]->SampleTimeOffset;
+			const int32 CurrColorPresetIndex = Positions[i]->ColorPresetIndex;
 
-		if (EnumHasAnyFlags(Sample.Flags, EPoseSearchTrajectoryFlags::FacingDirection))
-		{
-			for (int32 i = 0; i != FFeatureVectorHelper::EncodeVectorCardinality; ++i)
-			{
-				Weights[DataOffset + i] = Weight * Sample.Weight;
-			}
-			DataOffset += FFeatureVectorHelper::EncodeVectorCardinality;
-		}
-		if (EnumHasAnyFlags(Sample.Flags, EPoseSearchTrajectoryFlags::FacingDirectionXY))
-		{
-			for (int32 i = 0; i != FFeatureVectorHelper::EncodeVector2DCardinality; ++i)
+			if (PrevTimeOffset * CurrTimeOffset < UE_KINDA_SMALL_NUMBER)
 			{
-				Weights[DataOffset + i] = Weight * Sample.Weight;
+				// we jumped from negative to positive time offset without having a zero time offset. so we add the zero
+				TrajSplinePos.Add(DrawParams.GetCachedPosition(0.f));
+				TrajSplineColor.Add(DrawParams.GetColor(CurrColorPresetIndex));
 			}
-			DataOffset += FFeatureVectorHelper::EncodeVector2DCardinality;
-		}
-	}
-
-	check(DataOffset == ChannelDataOffset + ChannelCardinality);
-}
-
-void UPoseSearchFeatureChannel_Trajectory::IndexAsset(UE::PoseSearch::IAssetIndexer& Indexer, TArrayView<float> FeatureVectorTable) const
-{
-	const UE::PoseSearch::FAssetIndexingContext& IndexingContext = Indexer.GetIndexingContext();
-	for (int32 SampleIdx = IndexingContext.BeginSampleIdx; SampleIdx != IndexingContext.EndSampleIdx; ++SampleIdx)
-	{
-		const int32 VectorIdx = SampleIdx - IndexingContext.BeginSampleIdx;
-		IndexAssetPrivate(Indexer, SampleIdx, IndexingContext.GetPoseVector(VectorIdx, FeatureVectorTable));
-	}
-}
-
-void UPoseSearchFeatureChannel_Trajectory::IndexAssetPrivate(const UE::PoseSearch::IAssetIndexer& Indexer, int32 SampleIdx, TArrayView<float> FeatureVector) const
-{
-	// This function samples the instantaneous trajectory at time t as well as the trajectory's velocity and acceleration at time t.
-	// Symmetric finite differences are used to approximate derivatives:
-	//	First symmetric derivative:   f'(t) ~ (f(t+h) - f(t-h)) / 2h
-	//	Second symmetric derivative: f''(t) ~ (f(t+h) - 2f(t) + f(t-h)) / h^2
-	// Where h is a constant time delta
-	// So this means three root motion extractions are taken at time t-h, t, and t+h
 
-	using namespace UE::PoseSearch;
-	using FSampleInfo = IAssetIndexer::FSampleInfo;
-
-	const FAssetIndexingContext& IndexingContext = Indexer.GetIndexingContext();
-	float SampleTime = FMath::Min(SampleIdx * IndexingContext.Schema->GetSamplingInterval(), IndexingContext.AssetSampler->GetPlayLength());
-	FSampleInfo Origin = Indexer.GetSampleInfo(SampleTime);
-
-	int32 DataOffset = ChannelDataOffset;
-	for (const FPoseSearchTrajectorySample& Sample : Samples)
-	{
-		const float SubsampleTime = Sample.Offset + SampleTime;
-
-		// For each pose subsample term, get the corresponding clip, accumulated root motion,
-			// and wrap the time parameter based on the clip's length.
-		const FSampleInfo SamplePast = Indexer.GetSampleInfoRelative(SubsampleTime - IndexingContext.SamplingContext->FiniteDelta, Origin);
-		const FSampleInfo SamplePresent = Indexer.GetSampleInfoRelative(SubsampleTime, Origin);
-		const FSampleInfo SampleFuture = Indexer.GetSampleInfoRelative(SubsampleTime + IndexingContext.SamplingContext->FiniteDelta, Origin);
+			TrajSplinePos.Add(DrawParams.GetCachedPosition(CurrTimeOffset));
+			TrajSplineColor.Add(DrawParams.GetColor(CurrColorPresetIndex));
 
-		// Mirror transforms if requested
-		const FTransform MirroredRootPast = Indexer.MirrorTransform(SamplePast.RootTransform);
-		const FTransform MirroredRootPresent = Indexer.MirrorTransform(SamplePresent.RootTransform);
-		const FTransform MirroredRootFuture = Indexer.MirrorTransform(SampleFuture.RootTransform);
-
-		// We can get a better finite difference if we ignore samples that have
-		// been clamped at either side of the clip. However, if the central sample 
-		// itself is clamped, or there are no samples that are clamped, we can just 
-		// use the central difference as normal.
-		FVector LinearVelocity;
-		if (SamplePast.bClamped && !SamplePresent.bClamped && !SampleFuture.bClamped)
-		{
-			LinearVelocity = (MirroredRootFuture.GetTranslation() - MirroredRootPresent.GetTranslation()) / IndexingContext.SamplingContext->FiniteDelta;
-		}
-		else if (SampleFuture.bClamped && !SamplePresent.bClamped && !SamplePast.bClamped)
-		{
-			LinearVelocity = (MirroredRootPresent.GetTranslation() - MirroredRootPast.GetTranslation()) / IndexingContext.SamplingContext->FiniteDelta;
-		}
-		else
-		{
-			LinearVelocity = (MirroredRootFuture.GetTranslation() - MirroredRootPast.GetTranslation()) / (IndexingContext.SamplingContext->FiniteDelta * 2.0f);
-		}
-
-		const FVector LinearVelocityDirection = LinearVelocity.GetClampedToMaxSize(1.0f);
-		const FVector FacingDirection = MirroredRootPresent.GetRotation().GetForwardVector();
-
-		if (EnumHasAnyFlags(Sample.Flags, EPoseSearchTrajectoryFlags::Position))
-		{
-			FFeatureVectorHelper::EncodeVector(FeatureVector, DataOffset, MirroredRootPresent.GetTranslation());
-		}
-		if (EnumHasAnyFlags(Sample.Flags, EPoseSearchTrajectoryFlags::PositionXY))
-		{
-			FFeatureVectorHelper::EncodeVector2D(FeatureVector, DataOffset, FVector2D(MirroredRootPresent.GetTranslation().X, MirroredRootPresent.GetTranslation().Y));
-		}
-
-		if (EnumHasAnyFlags(Sample.Flags, EPoseSearchTrajectoryFlags::Velocity))
-		{
-			FFeatureVectorHelper::EncodeVector(FeatureVector, DataOffset, LinearVelocity);
-		}
-		if (EnumHasAnyFlags(Sample.Flags, EPoseSearchTrajectoryFlags::VelocityXY))
-		{
-			FFeatureVectorHelper::EncodeVector2D(FeatureVector, DataOffset, FVector2D(LinearVelocity.X, LinearVelocity.Y));
-		}
-
-		if (EnumHasAnyFlags(Sample.Flags, EPoseSearchTrajectoryFlags::VelocityDirection))
-		{
-			FFeatureVectorHelper::EncodeVector(FeatureVector, DataOffset, LinearVelocityDirection);
-		}
-		if (EnumHasAnyFlags(Sample.Flags, EPoseSearchTrajectoryFlags::VelocityDirectionXY))
-		{
-			FFeatureVectorHelper::EncodeVector2D(FeatureVector, DataOffset, FVector2D(LinearVelocityDirection.X, LinearVelocityDirection.Y));
-		}
-
-		if (EnumHasAnyFlags(Sample.Flags, EPoseSearchTrajectoryFlags::FacingDirection))
-		{
-			FFeatureVectorHelper::EncodeVector(FeatureVector, DataOffset, FacingDirection);
+			PrevTimeOffset = CurrTimeOffset;
 		}
-		if (EnumHasAnyFlags(Sample.Flags, EPoseSearchTrajectoryFlags::FacingDirectionXY))
-		{
-			FFeatureVectorHelper::EncodeVector2D(FeatureVector, DataOffset, FVector2D(FacingDirection.X, FacingDirection.Y).GetSafeNormal());
-		}
-	}
-	check(DataOffset == ChannelDataOffset + ChannelCardinality);
-}
 
-void UPoseSearchFeatureChannel_Trajectory::BuildQuery(UE::PoseSearch::FSearchContext& SearchContext, FPoseSearchFeatureVectorBuilder& InOutQuery) const
-{
-	using namespace UE::PoseSearch;
+		const float LifeTime = DrawParams.DefaultLifeTime;
+		const uint8 DepthPriority = ESceneDepthPriorityGroup::SDPG_Foreground + 2;
+		const bool bPersistent = EnumHasAnyFlags(DrawParams.Flags, EDebugDrawFlags::Persistent);
 
-	if (!SearchContext.Trajectory)
-	{
-		// @todo: do we want to reuse the SearchContext.CurrentResult data if valid?
-		return;
+		DrawCentripetalCatmullRomSpline(DrawParams.World, TrajSplinePos, TrajSplineColor, 0.5f, 8.f, bPersistent, LifeTime, DepthPriority, 0.f);
 	}
 
-	int32 NextIterStartIdx = 0;
-	int32 DataOffset = ChannelDataOffset;
-	float PreviousOffset = -FLT_MAX;
-	for (const FPoseSearchTrajectorySample& Sample : Samples)
-	{
-		// making sure Samples are sorted
-		check(Sample.Offset >= PreviousOffset);
-		const FTrajectorySample TrajectorySample = FTrajectorySampleRange::IterSampleTrajectory(SearchContext.Trajectory->Samples, Sample.Offset, NextIterStartIdx);
-
-		const FVector LinearVelocityDirection = TrajectorySample.LinearVelocity.GetClampedToMaxSize(1.0f);
-		const FVector FacingDirection = TrajectorySample.Transform.GetRotation().GetForwardVector();
-
-		if (EnumHasAnyFlags(Sample.Flags, EPoseSearchTrajectoryFlags::Position))
-		{
-			FFeatureVectorHelper::EncodeVector(InOutQuery.EditValues(), DataOffset, TrajectorySample.Transform.GetTranslation());
-		}
-		if (EnumHasAnyFlags(Sample.Flags, EPoseSearchTrajectoryFlags::PositionXY))
-		{
-			FFeatureVectorHelper::EncodeVector2D(InOutQuery.EditValues(), DataOffset, FVector2D(TrajectorySample.Transform.GetTranslation().X, TrajectorySample.Transform.GetTranslation().Y));
-		}
-
-		if (EnumHasAnyFlags(Sample.Flags, EPoseSearchTrajectoryFlags::Velocity))
-		{
-			FFeatureVectorHelper::EncodeVector(InOutQuery.EditValues(), DataOffset, TrajectorySample.LinearVelocity);
-		}
-		if (EnumHasAnyFlags(Sample.Flags, EPoseSearchTrajectoryFlags::VelocityXY))
-		{
-			FFeatureVectorHelper::EncodeVector2D(InOutQuery.EditValues(), DataOffset, FVector2D(TrajectorySample.LinearVelocity.X, TrajectorySample.LinearVelocity.Y));
-		}
-
-		if (EnumHasAnyFlags(Sample.Flags, EPoseSearchTrajectoryFlags::VelocityDirection))
-		{
-			FFeatureVectorHelper::EncodeVector(InOutQuery.EditValues(), DataOffset, LinearVelocityDirection);
-		}
-		if (EnumHasAnyFlags(Sample.Flags, EPoseSearchTrajectoryFlags::VelocityDirectionXY))
-		{
-			FFeatureVectorHelper::EncodeVector2D(InOutQuery.EditValues(), DataOffset, FVector2D(LinearVelocityDirection.X, LinearVelocityDirection.Y));
-		}
-
-		if (EnumHasAnyFlags(Sample.Flags, EPoseSearchTrajectoryFlags::FacingDirection))
-		{
-			FFeatureVectorHelper::EncodeVector(InOutQuery.EditValues(), DataOffset, FacingDirection);
-		}
-		if (EnumHasAnyFlags(Sample.Flags, EPoseSearchTrajectoryFlags::FacingDirectionXY))
-		{
-			FFeatureVectorHelper::EncodeVector2D(InOutQuery.EditValues(), DataOffset, FVector2D(FacingDirection.X, FacingDirection.Y).GetSafeNormal());
-		}
-	}
-	check(DataOffset == ChannelDataOffset + ChannelCardinality);
+	Super::DebugDraw(DrawParams, PoseVector);
+#endif // ENABLE_DRAW_DEBUG
 }
 
-// lazy initialized helper to interpolate or extrapolate (linearly) UPoseSearchFeatureChannel_Trajectory trajectory positions from FPoseSearchTrajectorySample containing EPoseSearchTrajectoryFlags::Position for samples without it
-struct TrajectoryPositionReconstructor
+#if WITH_EDITOR
+void UPoseSearchFeatureChannel_Trajectory::PopulateChannelLayoutSet(UE::PoseSearch::FFeatureChannelLayoutSet& FeatureChannelLayoutSet) const
 {
-	struct PositionAndOffsetSample
+	for (const TObjectPtr<UPoseSearchFeatureChannel>& SubChannelPtr : GetSubChannels())
 	{
-		FVector Position;
 		float Offset;
-	};
-
-	TArray<PositionAndOffsetSample, TInlineAllocator<32>> PositionAndOffsetSamples;
-	bool bInitialized = false;
-
-	void Init(const UPoseSearchFeatureChannel_Trajectory& TrajectoryChannel, TConstArrayView<float> PoseVector, const FTransform& RootTransform)
-	{
-		PositionAndOffsetSamples.Reserve(TrajectoryChannel.Samples.Num() + 1);
-		bool bAddZeroOffsetSample = true;
-		int32 DataOffset = TrajectoryChannel.GetChannelDataOffset();
-		for (const FPoseSearchTrajectorySample& Sample : TrajectoryChannel.Samples)
+		EPoseSearchTrajectoryFlags SampleFlag;
+		const char* Label;
+		const UPoseSearchFeatureChannel* Channel = SubChannelPtr.Get();
+		if (const UPoseSearchFeatureChannel_Position* Position = Cast<UPoseSearchFeatureChannel_Position>(Channel))
 		{
-			if (EnumHasAnyFlags(Sample.Flags, EPoseSearchTrajectoryFlags::Position))
-			{
-				PositionAndOffsetSample PositionAndOffsetSample;
-				PositionAndOffsetSample.Position = UE::PoseSearch::FFeatureVectorHelper::DecodeVector(PoseVector, DataOffset);
-				PositionAndOffsetSample.Position = RootTransform.TransformPosition(PositionAndOffsetSample.Position);
-				PositionAndOffsetSample.Offset = Sample.Offset;
-				PositionAndOffsetSamples.Add(PositionAndOffsetSample);
-
-				if (FMath::IsNearlyZero(Sample.Offset))
-				{
-					bAddZeroOffsetSample = false;
-				}
-			}
-			if (EnumHasAnyFlags(Sample.Flags, EPoseSearchTrajectoryFlags::PositionXY))
-			{
-				const FVector2D Position2D = UE::PoseSearch::FFeatureVectorHelper::DecodeVector2D(PoseVector, DataOffset);
-
-				if (!EnumHasAnyFlags(Sample.Flags, EPoseSearchTrajectoryFlags::Position))
-				{
-					PositionAndOffsetSample PositionAndOffsetSample;
-					PositionAndOffsetSample.Position = FVector(Position2D.X, Position2D.Y, 0);
-					PositionAndOffsetSample.Position = RootTransform.TransformPosition(PositionAndOffsetSample.Position);
-					PositionAndOffsetSample.Offset = Sample.Offset;
-					PositionAndOffsetSamples.Add(PositionAndOffsetSample);
-
-					if (FMath::IsNearlyZero(Sample.Offset))
-					{
-						bAddZeroOffsetSample = false;
-					}
-				}
-			}
-
-			if (EnumHasAnyFlags(Sample.Flags, EPoseSearchTrajectoryFlags::Velocity))
-			{
-				DataOffset += UE::PoseSearch::FFeatureVectorHelper::EncodeVectorCardinality;
-			}
-			if (EnumHasAnyFlags(Sample.Flags, EPoseSearchTrajectoryFlags::VelocityXY))
-			{
-				DataOffset += UE::PoseSearch::FFeatureVectorHelper::EncodeVector2DCardinality;
-			}
-
-			if (EnumHasAnyFlags(Sample.Flags, EPoseSearchTrajectoryFlags::VelocityDirection))
-			{
-				DataOffset += UE::PoseSearch::FFeatureVectorHelper::EncodeVectorCardinality;
-			}
-			if (EnumHasAnyFlags(Sample.Flags, EPoseSearchTrajectoryFlags::VelocityDirectionXY))
-			{
-				DataOffset += UE::PoseSearch::FFeatureVectorHelper::EncodeVector2DCardinality;
-			}
-
-			if (EnumHasAnyFlags(Sample.Flags, EPoseSearchTrajectoryFlags::FacingDirection))
-			{
-				DataOffset += UE::PoseSearch::FFeatureVectorHelper::EncodeVectorCardinality;
-			}
-			if (EnumHasAnyFlags(Sample.Flags, EPoseSearchTrajectoryFlags::FacingDirectionXY))
-			{
-				DataOffset += UE::PoseSearch::FFeatureVectorHelper::EncodeVector2DCardinality;
-			}
+			Offset = Position->SampleTimeOffset;
+			SampleFlag = EPoseSearchTrajectoryFlags::Position;
+			Label = "Pos";
 		}
-
-		if (bAddZeroOffsetSample)
+		else if (const UPoseSearchFeatureChannel_Velocity* Velocity = Cast<UPoseSearchFeatureChannel_Velocity>(Channel))
 		{
-			PositionAndOffsetSample PositionAndOffsetSample;
-			PositionAndOffsetSample.Position = RootTransform.GetTranslation();
-			PositionAndOffsetSample.Offset = 0.f;
-			PositionAndOffsetSamples.Add(PositionAndOffsetSample);
-		}
-
-		PositionAndOffsetSamples.Sort([](const PositionAndOffsetSample& a, const PositionAndOffsetSample& b)
-		{
-			return a.Offset < b.Offset;
-		});
-
-		bInitialized = true;
-		check(DataOffset == TrajectoryChannel.GetChannelDataOffset() + TrajectoryChannel.GetChannelCardinality());
-	}
-
-	FVector GetReconstructedTrajectoryPos(const UPoseSearchFeatureChannel_Trajectory& TrajectoryChannel, TConstArrayView<float> PoseVector, const FTransform& RootTransform, float SampleOffset)
-	{
-		if (!bInitialized)
-		{
-			Init(TrajectoryChannel, PoseVector, RootTransform);
-		}
-
-		return GetReconstructedTrajectoryPos(SampleOffset);
-	}
-
-	FVector GetReconstructedTrajectoryPos(float SampleOffset) const
-	{
-		check(bInitialized);
-		check(PositionAndOffsetSamples.Num() > 0);
-		if (PositionAndOffsetSamples.Num() >= 2)
-		{
-			const int32 LowerBoundIdx = Algo::LowerBound(PositionAndOffsetSamples, SampleOffset, [](const PositionAndOffsetSample& PositionAndOffsetSample, float Value)
-				{
-					return Value > PositionAndOffsetSample.Offset;
-				});
-
-			const int32 PrevIdx = FMath::Clamp(LowerBoundIdx, 0, PositionAndOffsetSamples.Num() - 2);
-			const int32 NextIdx = PrevIdx + 1;
-
-			const float Denominator = PositionAndOffsetSamples[NextIdx].Offset - PositionAndOffsetSamples[PrevIdx].Offset;
-			if (FMath::IsNearlyZero(Denominator))
+			Offset = Velocity->SampleTimeOffset;
+			if (Velocity->bNormalize)
 			{
-				return PositionAndOffsetSamples[PrevIdx].Position;
-			}
-
-			const float Numerator = SampleOffset - PositionAndOffsetSamples[PrevIdx].Offset;
-			const float LerpValue = Numerator / Denominator;
-			return FMath::Lerp(PositionAndOffsetSamples[PrevIdx].Position, PositionAndOffsetSamples[NextIdx].Position, LerpValue);
-		}
-
-		return PositionAndOffsetSamples[0].Position;
-	}
-};
-
-void UPoseSearchFeatureChannel_Trajectory::DebugDraw(const UE::PoseSearch::FDebugDrawParams& DrawParams, TConstArrayView<float> PoseVector) const
-{
-#if ENABLE_DRAW_DEBUG
-	using namespace UE::PoseSearch;
-
-	const float LifeTime = DrawParams.DefaultLifeTime;
-	const uint8 DepthPriority = ESceneDepthPriorityGroup::SDPG_Foreground + 2;
-	const bool bPersistent = EnumHasAnyFlags(DrawParams.Flags, EDebugDrawFlags::Persistent);
-
-	const int32 NumSamples = Samples.Num();
-	if (NumSamples == 0)
-	{
-		return;
-	}
-
-	int32 DataOffset = ChannelDataOffset;
-	int32 SampleIdx = 0;
-	TrajectoryPositionReconstructor TrajectoryPositionReconstructor;
-	TArray<FVector, TInlineAllocator<32>> TrajSplinePos;
-	TArray<FColor, TInlineAllocator<32>> TrajSplineColor;
-	for (const FPoseSearchTrajectorySample& Sample : Samples)
-	{
-		bool bIsTrajectoryPosValid = false;
-		FVector TrajectoryPos = FVector::Zero();
-		
-		if (EnumHasAnyFlags(Sample.Flags, EPoseSearchTrajectoryFlags::Position))
-		{
-			TrajectoryPos = FFeatureVectorHelper::DecodeVector(PoseVector, DataOffset);
-			TrajectoryPos = DrawParams.RootTransform.TransformPosition(TrajectoryPos);
-			
-			bIsTrajectoryPosValid = true;
-
-			// validating TrajectoryPositionReconstructor
-			check((TrajectoryPositionReconstructor.GetReconstructedTrajectoryPos(*this, PoseVector, DrawParams.RootTransform, Sample.Offset) - TrajectoryPos).IsNearlyZero());
-
-			const FColor Color = DrawParams.GetColor(Sample.ColorPresetIndex);
-
-			if (EnumHasAnyFlags(DrawParams.Flags, EDebugDrawFlags::DrawFast | EDebugDrawFlags::DrawSearchIndex))
-			{
-				DrawDebugPoint(DrawParams.World, TrajectoryPos, DrawParams.PointSize, Color, bPersistent, LifeTime, DepthPriority);
-			}
-			else
-			{
-				DrawDebugSphere(DrawParams.World, TrajectoryPos, 2.f, 8, Color, bPersistent, LifeTime, DepthPriority);
-			}
-
-			TrajSplinePos.Add(TrajectoryPos);
-			TrajSplineColor.Add(Color);
-		}
-		if (EnumHasAnyFlags(Sample.Flags, EPoseSearchTrajectoryFlags::PositionXY))
-		{
-			FVector2D TrajectoryPos2D = FFeatureVectorHelper::DecodeVector2D(PoseVector, DataOffset);
-			if (!bIsTrajectoryPosValid)
-			{
-				TrajectoryPos = FVector(TrajectoryPos2D.X, TrajectoryPos2D.Y, 0);
-				TrajectoryPos = DrawParams.RootTransform.TransformPosition(TrajectoryPos);
-				bIsTrajectoryPosValid = true;
-
-				const FColor Color = DrawParams.GetColor(Sample.ColorPresetIndex);
-
-				if (EnumHasAnyFlags(DrawParams.Flags, EDebugDrawFlags::DrawFast | EDebugDrawFlags::DrawSearchIndex))
-				{
-					DrawDebugPoint(DrawParams.World, TrajectoryPos, DrawParams.PointSize, Color, bPersistent, LifeTime, DepthPriority);
-				}
-				else
-				{
-					DrawDebugSphere(DrawParams.World, TrajectoryPos, 2.f, 8, Color, bPersistent, LifeTime, DepthPriority);
-				}
-
-				TrajSplinePos.Add(TrajectoryPos);
-				TrajSplineColor.Add(Color);
-			}
-		}
-		
-		if (!bIsTrajectoryPosValid)
-		{
-			TrajectoryPos = TrajectoryPositionReconstructor.GetReconstructedTrajectoryPos(*this, PoseVector, DrawParams.RootTransform, Sample.Offset);
-			
-			TrajSplinePos.Add(TrajectoryPos);
-			const FColor Color = TrajSplineColor.IsEmpty() ? FColor::Black : TrajSplineColor.Last();
-			TrajSplineColor.Add(Color);
-		}
-
-		if (EnumHasAnyFlags(Sample.Flags, EPoseSearchTrajectoryFlags::Velocity))
-		{
-			FVector TrajectoryVel = FFeatureVectorHelper::DecodeVector(PoseVector, DataOffset);
-
-			const FColor Color = DrawParams.GetColor(Sample.ColorPresetIndex);
-
-			TrajectoryVel *= 0.08f;
-			TrajectoryVel = DrawParams.RootTransform.TransformVector(TrajectoryVel);
-			const FVector TrajectoryVelDirection = TrajectoryVel.GetSafeNormal();
-
-			if (EnumHasAnyFlags(DrawParams.Flags, EDebugDrawFlags::DrawSearchIndex))
-			{
-				DrawDebugLine(DrawParams.World, TrajectoryPos, TrajectoryPos + TrajectoryVel, Color, bPersistent, LifeTime, DepthPriority);
+				SampleFlag = EPoseSearchTrajectoryFlags::VelocityDirection;
+				Label = "VelDir";
 			}
 			else
 			{
-				const float AdjustedThickness = EnumHasAnyFlags(DrawParams.Flags, EDebugDrawFlags::DrawFast) ? 0.0f : 1.f;
-
-				DrawDebugLine(
-					DrawParams.World,
-					TrajectoryPos + TrajectoryVelDirection * 2.f,
-					TrajectoryPos + TrajectoryVel,
-					Color,
-					bPersistent,
-					LifeTime,
-					DepthPriority,
-					AdjustedThickness
-				);
+				SampleFlag = EPoseSearchTrajectoryFlags::Velocity;
+				Label = "Vel";
 			}
 		}
-		if (EnumHasAnyFlags(Sample.Flags, EPoseSearchTrajectoryFlags::VelocityXY))
+		else if (const UPoseSearchFeatureChannel_Heading* Heading = Cast<UPoseSearchFeatureChannel_Heading>(Channel))
 		{
-			const FVector2D TrajectoryVel2D = FFeatureVectorHelper::DecodeVector2D(PoseVector, DataOffset);
-			FVector TrajectoryVel(TrajectoryVel2D.X, TrajectoryVel2D.Y, 0.f);
-
-			const FColor Color = DrawParams.GetColor(Sample.ColorPresetIndex);
-
-			TrajectoryVel *= 0.08f;
-			TrajectoryVel = DrawParams.RootTransform.TransformVector(TrajectoryVel);
-			const FVector TrajectoryVelDirection = TrajectoryVel.GetSafeNormal();
-
-			if (EnumHasAnyFlags(DrawParams.Flags, EDebugDrawFlags::DrawSearchIndex))
-			{
-				DrawDebugLine(DrawParams.World, TrajectoryPos, TrajectoryPos + TrajectoryVel, Color, bPersistent, LifeTime, DepthPriority);
-			}
-			else
-			{
-				const float AdjustedThickness = EnumHasAnyFlags(DrawParams.Flags, EDebugDrawFlags::DrawFast) ? 0.0f : 1.f;
-
-				DrawDebugLine(
-					DrawParams.World,
-					TrajectoryPos + TrajectoryVelDirection * 2.f,
-					TrajectoryPos + TrajectoryVel,
-					Color,
-					bPersistent,
-					LifeTime,
-					DepthPriority,
-					AdjustedThickness
-				);
-			}
+			Offset = Heading->SampleTimeOffset;
+			SampleFlag = EPoseSearchTrajectoryFlags::FacingDirection;
+			Label = "Fac";
 		}
-
-		if (EnumHasAnyFlags(Sample.Flags, EPoseSearchTrajectoryFlags::VelocityDirection))
-		{
-			FVector TrajectoryVelDirection = FFeatureVectorHelper::DecodeVector(PoseVector, DataOffset);
-
-			const FColor Color = DrawParams.GetColor(Sample.ColorPresetIndex);
-
-			TrajectoryVelDirection = DrawParams.RootTransform.TransformVector(TrajectoryVelDirection);
-
-			if (EnumHasAnyFlags(DrawParams.Flags, EDebugDrawFlags::DrawSearchIndex))
-			{
-				DrawDebugLine(DrawParams.World, TrajectoryPos, TrajectoryPos + TrajectoryVelDirection, Color, bPersistent, LifeTime, DepthPriority);
-			}
-			else
-			{
-				const float AdjustedThickness = EnumHasAnyFlags(DrawParams.Flags, EDebugDrawFlags::DrawFast) ? 0.0f : 1.f;
-
-				DrawDebugLine(
-					DrawParams.World,
-					TrajectoryPos + TrajectoryVelDirection * 2.f,
-					TrajectoryPos + TrajectoryVelDirection * 2.f * 10.0f,
-					Color,
-					bPersistent,
-					LifeTime,
-					DepthPriority,
-					AdjustedThickness
-				);
-			}
-		}
-
-		if (EnumHasAnyFlags(Sample.Flags, EPoseSearchTrajectoryFlags::VelocityDirectionXY))
-		{
-			const FVector2D TrajectoryVelDirection2D = FFeatureVectorHelper::DecodeVector2D(PoseVector, DataOffset);
-			FVector TrajectoryVelDirection(TrajectoryVelDirection2D.X, TrajectoryVelDirection2D.Y, 0.f);
-
-			const FColor Color = DrawParams.GetColor(Sample.ColorPresetIndex);
-
-			TrajectoryVelDirection = DrawParams.RootTransform.TransformVector(TrajectoryVelDirection);
-
-			if (EnumHasAnyFlags(DrawParams.Flags, EDebugDrawFlags::DrawSearchIndex))
-			{
-				DrawDebugLine(DrawParams.World, TrajectoryPos, TrajectoryPos + TrajectoryVelDirection, Color, bPersistent, LifeTime, DepthPriority);
-			}
-			else
-			{
-				const float AdjustedThickness = EnumHasAnyFlags(DrawParams.Flags, EDebugDrawFlags::DrawFast) ? 0.0f : 1.f;
-
-				DrawDebugLine(
-					DrawParams.World,
-					TrajectoryPos + TrajectoryVelDirection * 2.f,
-					TrajectoryPos + TrajectoryVelDirection * 2.f * 10.0f,
-					Color,
-					bPersistent,
-					LifeTime,
-					DepthPriority,
-					AdjustedThickness
-				);
-			}
-		}
-
-		if (EnumHasAnyFlags(Sample.Flags, EPoseSearchTrajectoryFlags::FacingDirection))
-		{
-			FVector TrajectoryForward = FFeatureVectorHelper::DecodeVector(PoseVector, DataOffset);
-
-			const FColor Color = DrawParams.GetColor(Sample.ColorPresetIndex);
-
-			TrajectoryForward = DrawParams.RootTransform.TransformVector(TrajectoryForward);
-
-			if (EnumHasAnyFlags(DrawParams.Flags, EDebugDrawFlags::DrawSearchIndex))
-			{
-				DrawDebugLine(DrawParams.World, TrajectoryPos, TrajectoryPos + TrajectoryForward, Color, bPersistent, LifeTime, DepthPriority);
-			}
-			else
-			{
-				const float AdjustedThickness = EnumHasAnyFlags(DrawParams.Flags, EDebugDrawFlags::DrawFast) ? 0.0f : 1.f;
-
-				DrawDebugLine(
-					DrawParams.World,
-					TrajectoryPos + TrajectoryForward * 2.f,
-					TrajectoryPos + TrajectoryForward * 2.f * 10.0f,
-					Color,
-					bPersistent,
-					LifeTime,
-					DepthPriority,
-					AdjustedThickness
-				);
-			}
-		}
-
-		if (EnumHasAnyFlags(Sample.Flags, EPoseSearchTrajectoryFlags::FacingDirectionXY))
-		{
-			const FVector2D TrajectoryForward2D = FFeatureVectorHelper::DecodeVector2D(PoseVector, DataOffset);
-			FVector TrajectoryForward(TrajectoryForward2D.X, TrajectoryForward2D.Y, 0.f);
-
-			const FColor Color = DrawParams.GetColor(Sample.ColorPresetIndex);
-
-			TrajectoryForward = DrawParams.RootTransform.TransformVector(TrajectoryForward);
-
-			if (EnumHasAnyFlags(DrawParams.Flags, EDebugDrawFlags::DrawSearchIndex))
-			{
-				DrawDebugLine(DrawParams.World, TrajectoryPos, TrajectoryPos + TrajectoryForward, Color, bPersistent, LifeTime, DepthPriority);
-			}
-			else
-			{
-				const float AdjustedThickness = EnumHasAnyFlags(DrawParams.Flags, EDebugDrawFlags::DrawFast) ? 0.0f : 1.f;
-
-				DrawDebugLine(
-					DrawParams.World,
-					TrajectoryPos + TrajectoryForward * 2.f,
-					TrajectoryPos + TrajectoryForward * 2.f * 10.0f,
-					Color,
-					bPersistent,
-					LifeTime,
-					DepthPriority,
-					AdjustedThickness
-				);
-			}
-		}
-
-		if (EnumHasAnyFlags(DrawParams.Flags, EDebugDrawFlags::DrawSampleLabels))
+		else
 		{
-			const FColor Color = DrawParams.GetColor(Sample.ColorPresetIndex);
-
-			const FString SampleLabel = FString::Format(TEXT("{0}"), { SampleIdx });
-
-			static const FVector DrawDebugSampleLabelOffset = FVector(0.0f, 0.0f, 5.0f);
-			DrawDebugString(
-				DrawParams.World,
-				TrajectoryPos + DrawDebugSampleLabelOffset,
-				SampleLabel,
-				nullptr,
-				Color,
-				LifeTime,
-				false,
-				1.5f);
+			checkNoEntry();
+			continue;
 		}
 
-		++SampleIdx;
-	}
-
-	DrawCentripetalCatmullRomSpline(DrawParams.World, TrajSplinePos, TrajSplineColor, 0.5f, 8.f, bPersistent, LifeTime, DepthPriority, 0.f);
-
-	check(DataOffset == ChannelDataOffset + ChannelCardinality);
-#endif // ENABLE_DRAW_DEBUG
-}
-
-#if WITH_EDITOR
-void UPoseSearchFeatureChannel_Trajectory::PopulateChannelLayoutSet(UE::PoseSearch::FFeatureChannelLayoutSet& FeatureChannelLayoutSet) const
-{
-	using namespace UE::PoseSearch;
-
-	int32 DataOffset = ChannelDataOffset;
-
-	auto Add = [&FeatureChannelLayoutSet, &DataOffset](EPoseSearchTrajectoryFlags SampleFlag, float Offset, const char* Label, int32 Cardinality)
-	{
 		FString SkeletonName = FeatureChannelLayoutSet.CurrentSchema->Skeleton->GetName();
-
 		UE::PoseSearch::FKeyBuilder KeyBuilder;
 		KeyBuilder << SkeletonName << SampleFlag << Offset;
-		FeatureChannelLayoutSet.Add(FString::Format(TEXT("Traj {0} {1}"), { Label, Offset }), KeyBuilder.Finalize(), DataOffset, Cardinality);
-
-		DataOffset += Cardinality;
-	};
-
-	for (const FPoseSearchTrajectorySample& Sample : Samples)
-	{
-		if (EnumHasAnyFlags(Sample.Flags, EPoseSearchTrajectoryFlags::Position))
-		{
-			Add(EPoseSearchTrajectoryFlags::Position, Sample.Offset, "Pos", FFeatureVectorHelper::EncodeVectorCardinality);
-		}
-		if (EnumHasAnyFlags(Sample.Flags, EPoseSearchTrajectoryFlags::PositionXY))
-		{
-			Add(EPoseSearchTrajectoryFlags::PositionXY, Sample.Offset, "PosXY", FFeatureVectorHelper::EncodeVector2DCardinality);
-		}
-
-		if (EnumHasAnyFlags(Sample.Flags, EPoseSearchTrajectoryFlags::Velocity))
-		{
-			Add(EPoseSearchTrajectoryFlags::Velocity, Sample.Offset, "Vel", FFeatureVectorHelper::EncodeVectorCardinality);
-		}
-		if (EnumHasAnyFlags(Sample.Flags, EPoseSearchTrajectoryFlags::VelocityXY))
-		{
-			Add(EPoseSearchTrajectoryFlags::VelocityXY, Sample.Offset, "VelXY", FFeatureVectorHelper::EncodeVector2DCardinality);
-		}
-
-		if (EnumHasAnyFlags(Sample.Flags, EPoseSearchTrajectoryFlags::VelocityDirection))
-		{
-			Add(EPoseSearchTrajectoryFlags::VelocityDirection, Sample.Offset, "VelDir", FFeatureVectorHelper::EncodeVectorCardinality);
-		}
-		if (EnumHasAnyFlags(Sample.Flags, EPoseSearchTrajectoryFlags::VelocityDirectionXY))
-		{
-			Add(EPoseSearchTrajectoryFlags::VelocityDirectionXY, Sample.Offset, "VelDirXY", FFeatureVectorHelper::EncodeVector2DCardinality);
-		}
-
-		if (EnumHasAnyFlags(Sample.Flags, EPoseSearchTrajectoryFlags::FacingDirection))
-		{
-			Add(EPoseSearchTrajectoryFlags::FacingDirection, Sample.Offset, "Fac", FFeatureVectorHelper::EncodeVectorCardinality);
-		}
-		if (EnumHasAnyFlags(Sample.Flags, EPoseSearchTrajectoryFlags::FacingDirectionXY))
-		{
-			Add(EPoseSearchTrajectoryFlags::FacingDirectionXY, Sample.Offset, "FacXY", FFeatureVectorHelper::EncodeVector2DCardinality);
-		}
+		FeatureChannelLayoutSet.Add(FString::Format(TEXT("Traj {0} {1}"), { Label, Offset }), KeyBuilder.Finalize(), Channel->GetChannelDataOffset(), Channel->GetChannelCardinality());
 	}
-
-	check(DataOffset == ChannelDataOffset + ChannelCardinality);
 }
 
 void UPoseSearchFeatureChannel_Trajectory::ComputeCostBreakdowns(UE::PoseSearch::ICostBreakDownData& CostBreakDownData, const UPoseSearchSchema* Schema) const
 {
-	using namespace UE::PoseSearch;
-
 	CostBreakDownData.AddEntireBreakDownSection(LOCTEXT("ColumnLabelTrajChannelTotal", "Traj Total"), Schema, ChannelDataOffset, ChannelCardinality);
 
 	if (CostBreakDownData.IsVerbose())
 	{
-		int32 DataOffset = ChannelDataOffset;
-
-		for (const FPoseSearchTrajectorySample& Sample : Samples)
+		for (const TObjectPtr<UPoseSearchFeatureChannel>& SubChannelPtr : GetSubChannels())
 		{
-			if (EnumHasAnyFlags(Sample.Flags, EPoseSearchTrajectoryFlags::Position))
-			{
-				CostBreakDownData.AddEntireBreakDownSection(FText::Format(LOCTEXT("ColumnLabelTrajChannelPosition", "Traj Pos {0}"), Sample.Offset), Schema, DataOffset, FFeatureVectorHelper::EncodeVectorCardinality);
-				DataOffset += FFeatureVectorHelper::EncodeVectorCardinality;
-			}
-			if (EnumHasAnyFlags(Sample.Flags, EPoseSearchTrajectoryFlags::PositionXY))
-			{
-				CostBreakDownData.AddEntireBreakDownSection(FText::Format(LOCTEXT("ColumnLabelTrajChannelPositionXY", "Traj PosXY {0}"), Sample.Offset), Schema, DataOffset, FFeatureVectorHelper::EncodeVector2DCardinality);
-				DataOffset += FFeatureVectorHelper::EncodeVector2DCardinality;
-			}
-
-			if (EnumHasAnyFlags(Sample.Flags, EPoseSearchTrajectoryFlags::Velocity))
-			{
-				CostBreakDownData.AddEntireBreakDownSection(FText::Format(LOCTEXT("ColumnLabelTrajChannelVelocity", "Traj Vel {0}"), Sample.Offset), Schema, DataOffset, FFeatureVectorHelper::EncodeVectorCardinality);
-				DataOffset += FFeatureVectorHelper::EncodeVectorCardinality;
-			}
-			if (EnumHasAnyFlags(Sample.Flags, EPoseSearchTrajectoryFlags::VelocityXY))
+			const UPoseSearchFeatureChannel* Channel = SubChannelPtr.Get();
+			if (const UPoseSearchFeatureChannel_Position* Position = Cast<UPoseSearchFeatureChannel_Position>(Channel))
 			{
-				CostBreakDownData.AddEntireBreakDownSection(FText::Format(LOCTEXT("ColumnLabelTrajChannelVelocityXY", "Traj VelXY {0}"), Sample.Offset), Schema, DataOffset, FFeatureVectorHelper::EncodeVector2DCardinality);
-				DataOffset += FFeatureVectorHelper::EncodeVector2DCardinality;
+				CostBreakDownData.AddEntireBreakDownSection(FText::Format(LOCTEXT("ColumnLabelTrajChannelPosition", "Traj Pos {0}"), Position->SampleTimeOffset), Schema, Position->GetChannelDataOffset(), Position->GetChannelCardinality());
 			}
-
-			if (EnumHasAnyFlags(Sample.Flags, EPoseSearchTrajectoryFlags::VelocityDirection))
+			else if (const UPoseSearchFeatureChannel_Velocity* Velocity = Cast<UPoseSearchFeatureChannel_Velocity>(Channel))
 			{
-				CostBreakDownData.AddEntireBreakDownSection(FText::Format(LOCTEXT("ColumnLabelTrajChannelVelocityDirection", "Traj VelDir {0}"), Sample.Offset), Schema, DataOffset, FFeatureVectorHelper::EncodeVectorCardinality);
-				DataOffset += FFeatureVectorHelper::EncodeVectorCardinality;
+				CostBreakDownData.AddEntireBreakDownSection(FText::Format(LOCTEXT("ColumnLabelTrajChannelVelocity", "Traj Vel {0}"), Velocity->SampleTimeOffset), Schema, Velocity->GetChannelDataOffset(), Velocity->GetChannelCardinality());
 			}
-			if (EnumHasAnyFlags(Sample.Flags, EPoseSearchTrajectoryFlags::VelocityDirectionXY))
+			else if (const UPoseSearchFeatureChannel_Heading* Heading = Cast<UPoseSearchFeatureChannel_Heading>(Channel))
 			{
-				CostBreakDownData.AddEntireBreakDownSection(FText::Format(LOCTEXT("ColumnLabelTrajChannelVelocityDirectionXY", "Traj VelDirXY {0}"), Sample.Offset), Schema, DataOffset, FFeatureVectorHelper::EncodeVector2DCardinality);
-				DataOffset += FFeatureVectorHelper::EncodeVector2DCardinality;
+				CostBreakDownData.AddEntireBreakDownSection(FText::Format(LOCTEXT("ColumnLabelTrajChannelFacingDirection", "Traj Fac {0}"), Heading->SampleTimeOffset), Schema, Heading->GetChannelDataOffset(), Heading->GetChannelCardinality());
 			}
-
-			if (EnumHasAnyFlags(Sample.Flags, EPoseSearchTrajectoryFlags::FacingDirection))
-			{
-				CostBreakDownData.AddEntireBreakDownSection(FText::Format(LOCTEXT("ColumnLabelTrajChannelFacingDirection", "Traj Fac {0}"), Sample.Offset), Schema, DataOffset, FFeatureVectorHelper::EncodeVectorCardinality);
-				DataOffset += FFeatureVectorHelper::EncodeVectorCardinality;
-			}
-			if (EnumHasAnyFlags(Sample.Flags, EPoseSearchTrajectoryFlags::FacingDirectionXY))
+			else
 			{
-				CostBreakDownData.AddEntireBreakDownSection(FText::Format(LOCTEXT("ColumnLabelTrajChannelFacingDirectionXY", "Traj FacXY {0}"), Sample.Offset), Schema, DataOffset, FFeatureVectorHelper::EncodeVector2DCardinality);
-				DataOffset += FFeatureVectorHelper::EncodeVector2DCardinality;
+				checkNoEntry();
 			}
 		}
-
-		check(DataOffset == ChannelDataOffset + ChannelCardinality);
 	}
 }
 
 #endif // WITH_EDITOR
 
-bool UPoseSearchFeatureChannel_Trajectory::GetEstimatedSpeedRatio(TConstArrayView<float> QueryVector, TConstArrayView<float> PoseVector, float& EstimatedSpeedRatio) const
+float UPoseSearchFeatureChannel_Trajectory::GetEstimatedSpeedRatio(TConstArrayView<float> QueryVector, TConstArrayView<float> PoseVector) const
 {
 	float EstimatedQuerySpeed = 0.f;
 	float EstimatedPoseSpeed = 0.f;
 
-	int32 QueryDataOffset = ChannelDataOffset;
-	int32 PoseDataOffset = ChannelDataOffset;
-
-	bool bValidEstimate = false;
-	for (const FPoseSearchTrajectorySample& Sample : Samples)
+	for (const TObjectPtr<UPoseSearchFeatureChannel>& SubChannelPtr : GetSubChannels())
 	{
-		// @todo: decode positions and use them to estimate velocities in case UPoseSearchFeatureChannel_Trajectory doens't contain EPoseSearchTrajectoryFlags::Velocity or EPoseSearchTrajectoryFlags::VelocityXY samples
-		if (EnumHasAnyFlags(Sample.Flags, EPoseSearchTrajectoryFlags::Position))
+		if (const UPoseSearchFeatureChannel_Velocity* Velocity = Cast<UPoseSearchFeatureChannel_Velocity>(SubChannelPtr.Get()))
 		{
-			QueryDataOffset += UE::PoseSearch::FFeatureVectorHelper::EncodeVectorCardinality;
-			PoseDataOffset += UE::PoseSearch::FFeatureVectorHelper::EncodeVectorCardinality;
-		}
-		if (EnumHasAnyFlags(Sample.Flags, EPoseSearchTrajectoryFlags::PositionXY))
-		{
-			QueryDataOffset += UE::PoseSearch::FFeatureVectorHelper::EncodeVector2DCardinality;
-			PoseDataOffset += UE::PoseSearch::FFeatureVectorHelper::EncodeVector2DCardinality;
-		}
-
-		if (EnumHasAnyFlags(Sample.Flags, EPoseSearchTrajectoryFlags::Velocity))
-		{
-			const FVector QueryVelocity = UE::PoseSearch::FFeatureVectorHelper::DecodeVector(QueryVector, QueryDataOffset);
-			const FVector PoseVelocity = UE::PoseSearch::FFeatureVectorHelper::DecodeVector(PoseVector, PoseDataOffset);
-			EstimatedQuerySpeed += QueryVelocity.Length();
-			EstimatedPoseSpeed += PoseVelocity.Length();
-			bValidEstimate = true;
-		}
-		if (EnumHasAnyFlags(Sample.Flags, EPoseSearchTrajectoryFlags::VelocityXY))
-		{
-			const FVector2D QueryVelocity = UE::PoseSearch::FFeatureVectorHelper::DecodeVector2D(QueryVector, QueryDataOffset);
-			const FVector2D PoseVelocity = UE::PoseSearch::FFeatureVectorHelper::DecodeVector2D(PoseVector, PoseDataOffset);
-			EstimatedQuerySpeed += QueryVelocity.Length();
-			EstimatedPoseSpeed += PoseVelocity.Length();
-			bValidEstimate = true;
-		}
-
-		if (EnumHasAnyFlags(Sample.Flags, EPoseSearchTrajectoryFlags::VelocityDirection))
-		{
-			QueryDataOffset += UE::PoseSearch::FFeatureVectorHelper::EncodeVectorCardinality;
-			PoseDataOffset += UE::PoseSearch::FFeatureVectorHelper::EncodeVectorCardinality;
-		}
-		if (EnumHasAnyFlags(Sample.Flags, EPoseSearchTrajectoryFlags::VelocityDirectionXY))
-		{
-			QueryDataOffset += UE::PoseSearch::FFeatureVectorHelper::EncodeVector2DCardinality;
-			PoseDataOffset += UE::PoseSearch::FFeatureVectorHelper::EncodeVector2DCardinality;
-		}
-		if (EnumHasAnyFlags(Sample.Flags, EPoseSearchTrajectoryFlags::FacingDirection))
-		{
-			QueryDataOffset += UE::PoseSearch::FFeatureVectorHelper::EncodeVectorCardinality;
-			PoseDataOffset += UE::PoseSearch::FFeatureVectorHelper::EncodeVectorCardinality;
-		}
-		if (EnumHasAnyFlags(Sample.Flags, EPoseSearchTrajectoryFlags::FacingDirectionXY))
-		{
-			QueryDataOffset += UE::PoseSearch::FFeatureVectorHelper::EncodeVector2DCardinality;
-			PoseDataOffset += UE::PoseSearch::FFeatureVectorHelper::EncodeVector2DCardinality;
+			if (!Velocity->bNormalize)
+			{
+				const FVector QueryVelocity = UE::PoseSearch::FFeatureVectorHelper::DecodeVectorAtOffset(QueryVector, Velocity->GetChannelDataOffset());
+				const FVector PoseVelocity = UE::PoseSearch::FFeatureVectorHelper::DecodeVectorAtOffset(PoseVector, Velocity->GetChannelDataOffset());
+				EstimatedQuerySpeed += QueryVelocity.Length();
+				EstimatedPoseSpeed += PoseVelocity.Length();
+			}
 		}
 	}
 
-	check(QueryDataOffset == ChannelDataOffset + ChannelCardinality);
-	check(PoseDataOffset == ChannelDataOffset + ChannelCardinality);
-
-	if (bValidEstimate && EstimatedPoseSpeed > UE_KINDA_SMALL_NUMBER)
-	{
-		EstimatedSpeedRatio = EstimatedQuerySpeed / EstimatedPoseSpeed;
-	}
-	else
+	if (EstimatedPoseSpeed > UE_KINDA_SMALL_NUMBER)
 	{
-		EstimatedSpeedRatio = 1.f;
+		return EstimatedQuerySpeed / EstimatedPoseSpeed;
 	}
 
-	return bValidEstimate;
+	return 1.f;
 }
 
 #undef LOCTEXT_NAMESPACE
\ No newline at end of file
diff --git a/Engine/Plugins/Experimental/Animation/PoseSearch/Source/Runtime/Private/PoseSearchFeatureChannel_Trajectory.h b/Engine/Plugins/Experimental/Animation/PoseSearch/Source/Runtime/Private/PoseSearchFeatureChannel_Trajectory.h
--- a/Engine/Plugins/Experimental/Animation/PoseSearch/Source/Runtime/Private/PoseSearchFeatureChannel_Trajectory.h
+++ b/Engine/Plugins/Experimental/Animation/PoseSearch/Source/Runtime/Private/PoseSearchFeatureChannel_Trajectory.h
@@ -2,11 +2,12 @@
 
 #pragma once
 
-#include "PoseSearch/PoseSearchFeatureChannel.h"
-#include "BoneContainer.h"
+#include "PoseSearchFeatureChannel_Group.h"
 #include "UObject/ObjectSaveContext.h"
 #include "PoseSearchFeatureChannel_Trajectory.generated.h"
 
+struct FTrajectorySampleRange;
+
 UENUM(meta = (Bitflags, UseEnumValuesAsMaskValuesInEditor = "true"))
 enum class EPoseSearchTrajectoryFlags : uint32
 {
@@ -40,9 +41,8 @@ struct POSESEARCH_API FPoseSearchTrajectorySample
 	UPROPERTY(EditAnywhere, Category = Config, meta = (ExcludeFromHash))
 	int32 ColorPresetIndex = 0;
 };
-
 UCLASS(BlueprintType, EditInlineNew, meta = (DisplayName = "Trajectory Channel"), CollapseCategories)
-class POSESEARCH_API UPoseSearchFeatureChannel_Trajectory : public UPoseSearchFeatureChannel
+class POSESEARCH_API UPoseSearchFeatureChannel_Trajectory : public UPoseSearchFeatureChannel_GroupBase
 {
 	GENERATED_BODY()
 
@@ -53,24 +53,25 @@ public:
 	UPROPERTY(EditAnywhere, Category = "Settings")
 	TArray<FPoseSearchTrajectorySample> Samples;
 
+	UPROPERTY(Transient)
+	TArray<TObjectPtr<UPoseSearchFeatureChannel>> SubChannels;
+
+	// UPoseSearchFeatureChannel_GroupBase interface
+	virtual TArrayView<TObjectPtr<UPoseSearchFeatureChannel>> GetSubChannels() override { return SubChannels; }
+	virtual TConstArrayView<TObjectPtr<UPoseSearchFeatureChannel>> GetSubChannels() const override { return SubChannels; }
+
 	// UObject interface
 	virtual void PreSave(FObjectPreSaveContext ObjectSaveContext) override;
 
 	// UPoseSearchFeatureChannel interface
-	virtual void InitializeSchema(UPoseSearchSchema* Schema) override;
-	virtual void FillWeights(TArray<float>& Weights) const override;
-	virtual void IndexAsset(UE::PoseSearch::IAssetIndexer& Indexer, TArrayView<float> FeatureVectorTable) const override;
-	virtual void BuildQuery(UE::PoseSearch::FSearchContext& SearchContext, FPoseSearchFeatureVectorBuilder& InOutQuery) const override;
+	virtual void Finalize(UPoseSearchSchema* Schema) override;
 	virtual void DebugDraw(const UE::PoseSearch::FDebugDrawParams& DrawParams, TConstArrayView<float> PoseVector) const override;
 
+	float GetEstimatedSpeedRatio(TConstArrayView<float> QueryVector, TConstArrayView<float> PoseVector) const;
+
 #if WITH_EDITOR
 	virtual void PopulateChannelLayoutSet(UE::PoseSearch::FFeatureChannelLayoutSet& FeatureChannelLayoutSet) const override;
 	virtual void ComputeCostBreakdowns(UE::PoseSearch::ICostBreakDownData& CostBreakDownData, const UPoseSearchSchema* Schema) const override;
 #endif
-
-	bool GetEstimatedSpeedRatio(TConstArrayView<float> QueryVector, TConstArrayView<float> PoseVector, float& EstimatedSpeedRatio) const;
-
-protected:
-	void IndexAssetPrivate(const UE::PoseSearch::IAssetIndexer& Indexer, int32 SampleIdx, TArrayView<float> FeatureVector) const;
 };
 
diff --git a/Engine/Plugins/Experimental/Animation/PoseSearch/Source/Runtime/Private/PoseSearchFeatureChannel_Velocity.cpp b/Engine/Plugins/Experimental/Animation/PoseSearch/Source/Runtime/Private/PoseSearchFeatureChannel_Velocity.cpp
new file mode 100644
--- /dev/null
+++ b/Engine/Plugins/Experimental/Animation/PoseSearch/Source/Runtime/Private/PoseSearchFeatureChannel_Velocity.cpp
@@ -0,0 +1,163 @@
+// Copyright Epic Games, Inc. All Rights Reserved.
+
+#include "PoseSearchFeatureChannel_Velocity.h"
+#include "Animation/MotionTrajectoryTypes.h"
+#include "DrawDebugHelpers.h"
+#include "PoseSearch/PoseSearchAssetIndexer.h"
+#include "PoseSearch/PoseSearchAssetSampler.h"
+#include "PoseSearch/PoseSearchContext.h"
+#include "PoseSearch/PoseSearchDatabase.h"
+#include "PoseSearch/PoseSearchHistory.h"
+#include "PoseSearch/PoseSearchSchema.h"
+
+void UPoseSearchFeatureChannel_Velocity::Finalize(UPoseSearchSchema* Schema)
+{
+	ChannelDataOffset = Schema->SchemaCardinality;
+	ChannelCardinality = UE::PoseSearch::FFeatureVectorHelper::EncodeVectorCardinality;
+	Schema->SchemaCardinality += ChannelCardinality;
+
+	SchemaBoneIdx = Schema->AddBoneReference(Bone);
+}
+
+void UPoseSearchFeatureChannel_Velocity::FillWeights(TArray<float>& Weights) const
+{
+	using namespace UE::PoseSearch;
+
+	for (int32 i = 0; i != FFeatureVectorHelper::EncodeVectorCardinality; ++i)
+	{
+		Weights[ChannelDataOffset + i] = Weight;
+	}
+}
+
+void UPoseSearchFeatureChannel_Velocity::IndexAsset(UE::PoseSearch::IAssetIndexer& Indexer, TArrayView<float> FeatureVectorTable) const
+{
+	using namespace UE::PoseSearch;
+
+	const FAssetIndexingContext& IndexingContext = Indexer.GetIndexingContext();
+	const FAssetSamplingContext* SamplingContext = IndexingContext.SamplingContext;
+	check(SamplingContext->FiniteDelta > UE_KINDA_SMALL_NUMBER);
+
+	for (int32 SampleIdx = IndexingContext.BeginSampleIdx; SampleIdx != IndexingContext.EndSampleIdx; ++SampleIdx)
+	{
+		const int32 VectorIdx = SampleIdx - IndexingContext.BeginSampleIdx;
+
+		const float OriginSampleTime = FMath::Min(SampleIdx * IndexingContext.Schema->GetSamplingInterval(), IndexingContext.AssetSampler->GetPlayLength());
+		const float SubsampleTime = OriginSampleTime + SampleTimeOffset;
+
+		bool ClampedPast, ClampedPresent, ClampedFuture;
+		const FTransform BoneTransformsPast = Indexer.GetTransformAndCacheResults(SubsampleTime - SamplingContext->FiniteDelta, bUseCharacterSpaceVelocities ? OriginSampleTime - SamplingContext->FiniteDelta : OriginSampleTime, SchemaBoneIdx, ClampedPast);
+		const FTransform BoneTransformsPresent = Indexer.GetTransformAndCacheResults(SubsampleTime, OriginSampleTime, SchemaBoneIdx, ClampedPresent);
+		const FTransform BoneTransformsFuture = Indexer.GetTransformAndCacheResults(SubsampleTime + SamplingContext->FiniteDelta, bUseCharacterSpaceVelocities ? OriginSampleTime + SamplingContext->FiniteDelta : OriginSampleTime, SchemaBoneIdx, ClampedFuture);
+
+		// We can get a better finite difference if we ignore samples that have
+		// been clamped at either side of the clip. However, if the central sample 
+		// itself is clamped, or there are no samples that are clamped, we can just 
+		// use the central difference as normal.
+		FVector LinearVelocity;
+		if (ClampedPast && !ClampedPresent && !ClampedFuture)
+		{
+			LinearVelocity = (BoneTransformsFuture.GetTranslation() - BoneTransformsPresent.GetTranslation()) / SamplingContext->FiniteDelta;
+		}
+		else if (ClampedFuture && !ClampedPresent && !ClampedPast)
+		{
+			LinearVelocity = (BoneTransformsPresent.GetTranslation() - BoneTransformsPast.GetTranslation()) / SamplingContext->FiniteDelta;
+		}
+		else
+		{
+			LinearVelocity = (BoneTransformsFuture.GetTranslation() - BoneTransformsPast.GetTranslation()) / (SamplingContext->FiniteDelta * 2.f);
+		}
+
+		if (bNormalize)
+		{
+			LinearVelocity = LinearVelocity.GetClampedToMaxSize(1.f);
+		}
+
+		int32 DataOffset = ChannelDataOffset;
+		FFeatureVectorHelper::EncodeVector(IndexingContext.GetPoseVector(VectorIdx, FeatureVectorTable), DataOffset, LinearVelocity);
+	}
+}
+
+void UPoseSearchFeatureChannel_Velocity::BuildQuery(UE::PoseSearch::FSearchContext& SearchContext, FPoseSearchFeatureVectorBuilder& InOutQuery) const
+{
+	using namespace UE::PoseSearch;
+
+	const bool bIsCurrentResultValid = SearchContext.CurrentResult.IsValid();
+	const bool bSkip = InputQueryPose != EInputQueryPose::UseCharacterPose && bIsCurrentResultValid && SearchContext.CurrentResult.Database->Schema == InOutQuery.GetSchema();
+	if (bSkip || !SearchContext.History)
+	{
+		if (bIsCurrentResultValid)
+		{
+			const float LerpValue = InputQueryPose == EInputQueryPose::UseInterpolatedContinuingPose ? SearchContext.CurrentResult.LerpValue : 0.f;
+			int32 DataOffset = ChannelDataOffset;
+			FFeatureVectorHelper::EncodeVector(InOutQuery.EditValues(), DataOffset, SearchContext.GetCurrentResultPrevPoseVector(), SearchContext.GetCurrentResultPoseVector(), SearchContext.GetCurrentResultNextPoseVector(), LerpValue);
+		}
+		// else leave the InOutQuery set to zero since the SearchContext.History is invalid and it'll fail if we continue
+	}
+	else
+	{
+		FVector LinearVelocity;
+		if (InOutQuery.GetSchema()->BoneReferences[SchemaBoneIdx].HasValidSetup())
+		{
+			check(SearchContext.History);
+			const float HistorySameplInterval = SearchContext.History->GetSampleTimeInterval();
+			check(HistorySameplInterval > UE_KINDA_SMALL_NUMBER);
+
+			// @todo: optimize this code for SchemaBoneIdx == FSearchContext::SchemaRootBoneIdx (!InOutQuery.GetSchema()->BoneReferences[SchemaBoneIdx].HasValidSetup())
+			// calculating the Transforms in component space for the bone indexed by SchemaBoneIdx
+			FTransform TransformCurrent = SearchContext.TryGetTransformAndCacheResults(SampleTimeOffset, InOutQuery.GetSchema(), SchemaBoneIdx);
+			FTransform TransformPrevious = SearchContext.TryGetTransformAndCacheResults(SampleTimeOffset - HistorySameplInterval, InOutQuery.GetSchema(), SchemaBoneIdx);
+
+			// if we want the velocity not in character space we need to calculate the root offset delta between the SampleTimeOffset and the SampleTimeOffset - HistorySameplInterval to apply to the TransformPrevious
+			if (!bUseCharacterSpaceVelocities)
+			{
+				const FTransform RootTransform = SearchContext.TryGetTransformAndCacheResults(SampleTimeOffset, InOutQuery.GetSchema(), FSearchContext::SchemaRootBoneIdx);
+				const FTransform RootTransformPrev = SearchContext.TryGetTransformAndCacheResults(SampleTimeOffset - HistorySameplInterval, InOutQuery.GetSchema(), FSearchContext::SchemaRootBoneIdx);
+				TransformPrevious = TransformPrevious * (RootTransformPrev * RootTransform.Inverse());
+			}
+
+			LinearVelocity = (TransformCurrent.GetTranslation() - TransformPrevious.GetTranslation()) / HistorySameplInterval;
+		}
+		else
+		{
+			check(SearchContext.Trajectory);
+			// @todo: make this call consistent with Transform = SearchContext.TryGetTransformAndCacheResults(SampleTimeOffset, InOutQuery.GetSchema(), FSearchContext::SchemaRootBoneIdx);
+			const FTrajectorySample TrajectorySample = SearchContext.Trajectory->GetSampleAtTime(SampleTimeOffset);
+			LinearVelocity = TrajectorySample.LinearVelocity;
+		}
+
+		if (bNormalize)
+		{
+			LinearVelocity = LinearVelocity.GetClampedToMaxSize(1.f);
+		}
+
+		int32 DataOffset = ChannelDataOffset;
+		FFeatureVectorHelper::EncodeVector(InOutQuery.EditValues(), DataOffset, LinearVelocity);
+	}
+}
+
+void UPoseSearchFeatureChannel_Velocity::DebugDraw(const UE::PoseSearch::FDebugDrawParams& DrawParams, TConstArrayView<float> PoseVector) const
+{
+#if ENABLE_DRAW_DEBUG
+	using namespace UE::PoseSearch;
+
+	const float LifeTime = DrawParams.DefaultLifeTime;
+	const uint8 DepthPriority = ESceneDepthPriorityGroup::SDPG_Foreground + 2;
+	const bool bPersistent = EnumHasAnyFlags(DrawParams.Flags, EDebugDrawFlags::Persistent);
+	const FColor Color = DrawParams.GetColor(ColorPresetIndex);
+	const float LinearVelocityScale = bNormalize ? 15.f : 0.08f;
+
+	const FVector LinearVelocity = DrawParams.RootTransform.TransformVector(FFeatureVectorHelper::DecodeVectorAtOffset(PoseVector, ChannelDataOffset));
+	const FVector BoneVelDirection = LinearVelocity.GetSafeNormal();
+	const FVector BonePos = DrawParams.GetCachedPosition(SampleTimeOffset, SchemaBoneIdx);
+
+	if (EnumHasAnyFlags(DrawParams.Flags, EDebugDrawFlags::DrawSearchIndex))
+	{
+		DrawDebugLine(DrawParams.World, BonePos, BonePos + LinearVelocity * LinearVelocityScale, Color, bPersistent, LifeTime, DepthPriority);
+	}
+	else
+	{
+		const float AdjustedThickness = EnumHasAnyFlags(DrawParams.Flags, EDebugDrawFlags::DrawFast) ? 0.f : 1.f;
+		DrawDebugLine(DrawParams.World, BonePos + BoneVelDirection * 2.f, BonePos + LinearVelocity * LinearVelocityScale, Color, bPersistent, LifeTime, DepthPriority, AdjustedThickness);
+	}
+#endif // ENABLE_DRAW_DEBUG
+}
diff --git a/Engine/Plugins/Experimental/Animation/PoseSearch/Source/Runtime/Private/PoseSearchFeatureChannel_Velocity.h b/Engine/Plugins/Experimental/Animation/PoseSearch/Source/Runtime/Private/PoseSearchFeatureChannel_Velocity.h
new file mode 100644
--- /dev/null
+++ b/Engine/Plugins/Experimental/Animation/PoseSearch/Source/Runtime/Private/PoseSearchFeatureChannel_Velocity.h
@@ -0,0 +1,47 @@
+// Copyright Epic Games, Inc. All Rights Reserved.
+
+#pragma once
+
+#include "PoseSearch/PoseSearchFeatureChannel.h"
+#include "BoneContainer.h"
+#include "PoseSearchFeatureChannel_Velocity.generated.h"
+
+UCLASS(BlueprintType, EditInlineNew, meta = (DisplayName = "Velocity Channel"), CollapseCategories)
+class POSESEARCH_API UPoseSearchFeatureChannel_Velocity : public UPoseSearchFeatureChannel
+{
+	GENERATED_BODY()
+
+public:
+	UPROPERTY(EditAnywhere, Category = "Settings")
+	FBoneReference Bone;
+
+	UPROPERTY(EditAnywhere, Category = "Settings")
+	float Weight = 1.f;
+
+	UPROPERTY(EditAnywhere, Category = "Settings")
+	float SampleTimeOffset = 0.f;
+
+	UPROPERTY(meta = (ExcludeFromHash))
+	int8 SchemaBoneIdx = 0;
+
+	UPROPERTY(EditAnywhere, Category = "Settings", meta = (ExcludeFromHash))
+	int32 ColorPresetIndex = 0;
+
+	UPROPERTY(EditAnywhere, Category = "Settings")
+	EInputQueryPose InputQueryPose = EInputQueryPose::UseContinuingPose;
+
+	// if bUseCharacterSpaceVelocities is true, velocities will be calculated from the positions in character space, otherwise they will be calculated using global space positions
+	UPROPERTY(EditAnywhere, Category = "Settings")
+	bool bUseCharacterSpaceVelocities = true;
+
+	// if bNormalize is true, velocity will be clamped to a maximum size of one (UPoseSearchFeatureChannel_Velocity will be matching the "direction" of the velocity)
+	UPROPERTY(EditAnywhere, Category = "Settings")
+	bool bNormalize = false;
+
+	// UPoseSearchFeatureChannel interface
+	virtual void Finalize(UPoseSearchSchema* Schema) override;
+	virtual void FillWeights(TArray<float>& Weights) const override;
+	virtual void IndexAsset(UE::PoseSearch::IAssetIndexer& Indexer, TArrayView<float> FeatureVectorTable) const override;
+	virtual void BuildQuery(UE::PoseSearch::FSearchContext& SearchContext, FPoseSearchFeatureVectorBuilder& InOutQuery) const override;
+	virtual void DebugDraw(const UE::PoseSearch::FDebugDrawParams& DrawParams, TConstArrayView<float> PoseVector) const override;
+};
diff --git a/Engine/Plugins/Experimental/Animation/PoseSearch/Source/Runtime/Private/PoseSearchLibrary.cpp b/Engine/Plugins/Experimental/Animation/PoseSearch/Source/Runtime/Private/PoseSearchLibrary.cpp
--- a/Engine/Plugins/Experimental/Animation/PoseSearch/Source/Runtime/Private/PoseSearchLibrary.cpp
+++ b/Engine/Plugins/Experimental/Animation/PoseSearch/Source/Runtime/Private/PoseSearchLibrary.cpp
@@ -234,9 +234,8 @@ static void TraceMotionMatchingState(
 	if (DeltaTime > SMALL_NUMBER)
 	{
 		// simulation
-		int32 FirstIdx = 0;
-		const FTrajectorySample PrevSample = FTrajectorySampleRange::IterSampleTrajectory(Trajectory.Samples, -DeltaTime, FirstIdx);
-		const FTrajectorySample CurrSample = FTrajectorySampleRange::IterSampleTrajectory(Trajectory.Samples, 0.f, FirstIdx);
+		const FTrajectorySample PrevSample = Trajectory.GetSampleAtTime(-DeltaTime);
+		const FTrajectorySample CurrSample = Trajectory.GetSampleAtTime(0.f);
 
 		const FTransform SimDelta = CurrSample.Transform.GetRelativeTransform(PrevSample.Transform);
 
@@ -379,18 +378,9 @@ void FMotionMatchingState::UpdateWantedPlayRate(const UE::PoseSearch::FSearchCon
 				{
 					TConstArrayView<float> QueryData = PoseSearchFeatureVectorBuilder->GetValues();
 					TConstArrayView<float> ResultData = CurrentSearchResult.Database->GetSearchIndex().GetPoseValues(CurrentSearchResult.PoseIdx);
-					float EstimatedSpeedRatio = 1.f;
-					if (TrajectoryChannel->GetEstimatedSpeedRatio(QueryData, ResultData, EstimatedSpeedRatio))
-					{
-						check(Settings.PlayRateMin <= Settings.PlayRateMax);
-						WantedPlayRate = FMath::Clamp(EstimatedSpeedRatio, Settings.PlayRateMin, Settings.PlayRateMax);
-					}
-					else
-					{
-						UE_LOG(LogPoseSearch, Warning,
-							TEXT("Couldn't update the WantedPlayRate in FMotionMatchingState::UpdateWantedPlayRate, because Schema '%s' UPoseSearchFeatureChannel_Trajectory::GetEstimatedSpeedRatio failed to compute, because of missing EPoseSearchTrajectoryFlags::Velocity samples"),
-							*GetNameSafe(CurrentSearchResult.Database->Schema));
-					}
+					const float EstimatedSpeedRatio = TrajectoryChannel->GetEstimatedSpeedRatio(QueryData, ResultData);
+					check(Settings.PlayRateMin <= Settings.PlayRateMax);
+					WantedPlayRate = FMath::Clamp(EstimatedSpeedRatio, Settings.PlayRateMin, Settings.PlayRateMax);
 				}
 				else
 				{
diff --git a/Engine/Plugins/Experimental/Animation/PoseSearch/Source/Runtime/Private/PoseSearchSchema.cpp b/Engine/Plugins/Experimental/Animation/PoseSearch/Source/Runtime/Private/PoseSearchSchema.cpp
--- a/Engine/Plugins/Experimental/Animation/PoseSearch/Source/Runtime/Private/PoseSearchSchema.cpp
+++ b/Engine/Plugins/Experimental/Animation/PoseSearch/Source/Runtime/Private/PoseSearchSchema.cpp
@@ -8,14 +8,7 @@
 
 bool UPoseSearchSchema::IsValid() const
 {
-	bool bValid = Skeleton != nullptr;
-
-	for (const TObjectPtr<UPoseSearchFeatureChannel>& Channel: Channels)
-	{
-		bValid &= Channel != nullptr;
-	}
-
-	return bValid;
+	return Skeleton != nullptr;
 }
 
 void UPoseSearchSchema::ResolveBoneReferences()
@@ -51,31 +44,31 @@ void UPoseSearchSchema::BuildQuery(UE::PoseSearch::FSearchContext& SearchContext
 
 	InOutQuery.Init(this);
 
-	for (const TObjectPtr<UPoseSearchFeatureChannel>& Channel : Channels)
+	for (const TObjectPtr<UPoseSearchFeatureChannel>& ChannelPtr : Channels)
 	{
-		Channel->BuildQuery(SearchContext, InOutQuery);
+		if (ChannelPtr)
+		{
+			ChannelPtr->BuildQuery(SearchContext, InOutQuery);
+		}
 	}
 }
 
-void UPoseSearchSchema::Finalize(bool bRemoveEmptyChannels)
+FBoneIndexType UPoseSearchSchema::GetBoneIndexType(int8 SchemaBoneIdx) const
 {
-	using namespace UE::PoseSearch;
-
-	if (bRemoveEmptyChannels)
-	{
-		Channels.RemoveAll([](TObjectPtr<UPoseSearchFeatureChannel>& Channel) { return !Channel; });
-	}
+	return SchemaBoneIdx >= 0 && BoneReferences[SchemaBoneIdx].HasValidSetup() ? BoneReferences[SchemaBoneIdx].BoneIndex : RootBoneIdx;
+}
 
+void UPoseSearchSchema::Finalize()
+{
 	BoneReferences.Reset();
 
-	int32 CurrentChannelDataOffset = 0;
-
 	SchemaCardinality = 0;
-	for (int32 ChannelIdx = 0; ChannelIdx != Channels.Num(); ++ChannelIdx)
+
+	for (const TObjectPtr<UPoseSearchFeatureChannel>& ChannelPtr : Channels)
 	{
-		if (UPoseSearchFeatureChannel* Channel = Channels[ChannelIdx].Get())
+		if (ChannelPtr)
 		{
-			Channel->InitializeSchema(this);
+			ChannelPtr->Finalize(this);
 		}
 	}
 
@@ -85,30 +78,40 @@ void UPoseSearchSchema::Finalize(bool bRemoveEmptyChannels)
 void UPoseSearchSchema::PreSave(FObjectPreSaveContext ObjectSaveContext)
 {
 	Finalize();
-
 	Super::PreSave(ObjectSaveContext);
 }
 
 void UPoseSearchSchema::PostLoad()
 {
 	Super::PostLoad();
-	ResolveBoneReferences();
+	Finalize();
+}
+
+int32 UPoseSearchSchema::AddBoneReference(const FBoneReference& BoneReference)
+{
+	return BoneReferences.AddUnique(BoneReference);
+}
+
+USkeleton* UPoseSearchSchema::GetSkeleton(bool& bInvalidSkeletonIsError, const IPropertyHandle* PropertyHandle)
+{
+	bInvalidSkeletonIsError = false;
+	return Skeleton;
 }
 
 #if WITH_EDITOR
 void UPoseSearchSchema::PostEditChangeProperty(FPropertyChangedEvent& PropertyChangedEvent)
 {
-	Finalize(false);
+	Finalize();
 	Super::PostEditChangeProperty(PropertyChangedEvent);
 }
 
 void UPoseSearchSchema::ComputeCostBreakdowns(UE::PoseSearch::ICostBreakDownData& CostBreakDownData) const
 {
-	for (const TObjectPtr<UPoseSearchFeatureChannel>& Channel : Channels)
+	for (const TObjectPtr<UPoseSearchFeatureChannel>& ChannelPtr : Channels)
 	{
-		if (Channel)
+		if (ChannelPtr)
 		{
-			Channel->ComputeCostBreakdowns(CostBreakDownData, this);
+			ChannelPtr->ComputeCostBreakdowns(CostBreakDownData, this);
 		}
 	}
 }
diff --git a/Engine/Plugins/Experimental/Animation/PoseSearch/Source/Runtime/Public/PoseSearch/PoseSearchContext.h b/Engine/Plugins/Experimental/Animation/PoseSearch/Source/Runtime/Public/PoseSearch/PoseSearchContext.h
--- a/Engine/Plugins/Experimental/Animation/PoseSearch/Source/Runtime/Public/PoseSearch/PoseSearchContext.h
+++ b/Engine/Plugins/Experimental/Animation/PoseSearch/Source/Runtime/Public/PoseSearch/PoseSearchContext.h
@@ -63,13 +63,64 @@ enum class EPoseCandidateFlags : uint8
 };
 ENUM_CLASS_FLAGS(EPoseCandidateFlags);
 
+template<typename FTransformType>
+struct POSESEARCH_API FCachedTransform
+{
+	FCachedTransform()
+		: SampleTime(0.f)
+		, BoneIndexType(-1)
+		, Transform()
+	{
+	}
+
+	FCachedTransform(float InSampleTime, FBoneIndexType InBoneIndexType, const FTransformType& InTransform)
+		: SampleTime(InSampleTime)
+		, BoneIndexType(InBoneIndexType)
+		, Transform(InTransform)
+	{
+	}
+
+	float SampleTime;
+	
+	// if -1 it represents the root bone
+	FBoneIndexType BoneIndexType;
+
+	// associated transform to BoneIndexType in ComponentSpace (except for the root bone stored in global space)
+	FTransformType Transform;
+};
+
+template<typename FTransformType>
+struct FCachedTransforms
+{
+	const FCachedTransform<FTransformType>* Find(float SampleTime, FBoneIndexType BoneIndexType) const
+	{
+		// @todo: use an hashmap if we end up having too many entries
+		return CachedTransforms.FindByPredicate([SampleTime, BoneIndexType](const FCachedTransform<FTransformType>& CachedTransform)
+			{
+				return CachedTransform.SampleTime == SampleTime && CachedTransform.BoneIndexType == BoneIndexType;
+			});
+	}
+
+	void Add(float SampleTime, FBoneIndexType BoneIndexType, const FTransformType& Transform)
+	{
+		CachedTransforms.Emplace(SampleTime, BoneIndexType, Transform);
+	}
+
+	void Reset()
+	{
+		CachedTransforms.Reset();
+	}
+
+private:
+	TArray<FCachedTransform<FTransformType>, TInlineAllocator<64>> CachedTransforms;
+};
+
 // @todo: FDebugDrawParams should be enclosed with #if ENABLE_DRAW_DEBUG
 struct POSESEARCH_API FDebugDrawParams
 {
 	const UWorld* World = nullptr;
 	const UPoseSearchDatabase* Database = nullptr;
 	EDebugDrawFlags Flags = EDebugDrawFlags::None;
-	uint32 ChannelMask = (uint32)-1;
 
 	float DefaultLifeTime = 5.f;
 	float PointSize = 1.f;
@@ -83,10 +134,17 @@ struct POSESEARCH_API FDebugDrawParams
 	FColor GetColor(int32 ColorPreset) const;
 	const FPoseSearchIndex* GetSearchIndex() const;
 	const UPoseSearchSchema* GetSchema() const;
+
+	void ClearCachedPositions();
+	void AddCachedPosition(float TimeOffset, int8 SchemaBoneIdx, const FVector& Position);
+	FVector GetCachedPosition(float TimeOffset, int8 SchemaBoneIdx = -1) const;
+
+private:
+	FCachedTransforms<FVector> CachedPositions;
 };
 
-POSESEARCH_API void DrawFeatureVector(const FDebugDrawParams& DrawParams, TConstArrayView<float> PoseVector);
-POSESEARCH_API void DrawFeatureVector(const FDebugDrawParams& DrawParams, int32 PoseIdx);
+POSESEARCH_API void DrawFeatureVector(FDebugDrawParams& DrawParams, TConstArrayView<float> PoseVector);
+POSESEARCH_API void DrawFeatureVector(FDebugDrawParams& DrawParams, int32 PoseIdx);
 
 struct POSESEARCH_API FSearchContext
 {
@@ -125,28 +183,11 @@ struct POSESEARCH_API FSearchContext
 	const FPoseIndicesHistory* PoseIndicesHistory = nullptr;
 
 private:
-	struct FCachedEntry
-	{
-		FCachedEntry(float InSampleTime = 0.f, const FTransform& InTransform = FTransform::Identity, FBoneIndexType InBoneIndexType = -1)
-			: SampleTime(InSampleTime)
-			, Transform(InTransform)
-			, BoneIndexType(InBoneIndexType)
-		{
-		}
-
-		float SampleTime = 0.f;
-
-		// associated transform to BoneIndexType in ComponentSpace (except for the root bone stored in global space)
-		FTransform Transform = FTransform::Identity;
-
-		// if -1 it represents the root bone
-		FBoneIndexType BoneIndexType = -1;
-	};
-
-	TArray<FCachedEntry, TInlineAllocator<64>> CachedEntries;
+	FCachedTransforms<FTransform> CachedTransforms;
 	
 	struct FCachedQuery
 	{
+		// @todo: why do we need to hold on to a Database? isn't the FeatureVectorBuilder.Schema enough?
 		const UPoseSearchDatabase* Database = nullptr;
 		FPoseSearchFeatureVectorBuilder FeatureVectorBuilder;
 	};
diff --git a/Engine/Plugins/Experimental/Animation/PoseSearch/Source/Runtime/Public/PoseSearch/PoseSearchFeatureChannel.h b/Engine/Plugins/Experimental/Animation/PoseSearch/Source/Runtime/Public/PoseSearch/PoseSearchFeatureChannel.h
--- a/Engine/Plugins/Experimental/Animation/PoseSearch/Source/Runtime/Public/PoseSearch/PoseSearchFeatureChannel.h
+++ b/Engine/Plugins/Experimental/Animation/PoseSearch/Source/Runtime/Public/PoseSearch/PoseSearchFeatureChannel.h
@@ -41,27 +41,25 @@ public:
 	static void EncodeQuat(TArrayView<float> Values, int32& DataOffset, const FQuat& Quat);
 	static void EncodeQuat(TArrayView<float> Values, int32& DataOffset, TConstArrayView<float> PrevValues, TConstArrayView<float> CurValues, TConstArrayView<float> NextValues, float LerpValue);
 	static FQuat DecodeQuat(TConstArrayView<float> Values, int32& DataOffset);
+	static FQuat DecodeQuatAtOffset(TConstArrayView<float> Values, int32 DataOffset);
 
 	enum { EncodeVectorCardinality = 3 };
 	static void EncodeVector(TArrayView<float> Values, int32& DataOffset, const FVector& Vector);
 	static void EncodeVector(TArrayView<float> Values, int32& DataOffset, TConstArrayView<float> PrevValues, TConstArrayView<float> CurValues, TConstArrayView<float> NextValues, float LerpValue, bool bNormalize = false);
 	static FVector DecodeVector(TConstArrayView<float> Values, int32& DataOffset);
+	static FVector DecodeVectorAtOffset(TConstArrayView<float> Values, int32 DataOffset);
 
 	enum { EncodeVector2DCardinality = 2 };
 	static void EncodeVector2D(TArrayView<float> Values, int32& DataOffset, const FVector2D& Vector2D);
 	static void EncodeVector2D(TArrayView<float> Values, int32& DataOffset, TConstArrayView<float> PrevValues, TConstArrayView<float> CurValues, TConstArrayView<float> NextValues, float LerpValue);
 	static FVector2D DecodeVector2D(TConstArrayView<float> Values, int32& DataOffset);
+	static FVector2D DecodeVector2DAtOffset(TConstArrayView<float> Values, int32 DataOffset);
 
 	enum { EncodeFloatCardinality = 1 };
 	static void EncodeFloat(TArrayView<float> Values, int32& DataOffset, const float Value);
 	static void EncodeFloat(TArrayView<float> Values, int32& DataOffset, TConstArrayView<float> PrevValues, TConstArrayView<float> CurValues, TConstArrayView<float> NextValues, float LerpValue);
 	static float DecodeFloat(TConstArrayView<float> Values, int32& DataOffset);
-
-private:
-	static FQuat DecodeQuatInternal(TConstArrayView<float> Values, int32 DataOffset);
-	static FVector DecodeVectorInternal(TConstArrayView<float> Values, int32 DataOffset);
-	static FVector2D DecodeVector2DInternal(TConstArrayView<float> Values, int32 DataOffset);
-	static float DecodeFloatInternal(TConstArrayView<float> Values, int32 DataOffset);
+	static float DecodeFloatAtOffset(TConstArrayView<float> Values, int32 DataOffset);
 };
 
 #if WITH_EDITOR
@@ -151,7 +149,7 @@ public:
 	int32 GetChannelDataOffset() const { checkSlow(ChannelDataOffset >= 0); return ChannelDataOffset; }
 
 	// Called during UPoseSearchSchema::Finalize to prepare the schema for this channel
-	virtual void InitializeSchema(UPoseSearchSchema* Schema) PURE_VIRTUAL(UPoseSearchFeatureChannel::InitializeSchema, );
+	virtual void Finalize(UPoseSearchSchema* Schema) PURE_VIRTUAL(UPoseSearchFeatureChannel::Finalize, );
 	
 	// Called at database build time to collect feature weights.
 	// Weights is sized to the cardinality of the schema and the feature channel should write
@@ -164,6 +162,9 @@ public:
 	// Called at runtime to add this channel's data to the query pose vector
 	virtual void BuildQuery(UE::PoseSearch::FSearchContext& SearchContext, FPoseSearchFeatureVectorBuilder& InOutQuery) const PURE_VIRTUAL(UPoseSearchFeatureChannel::BuildQuery, );
 
+	// API called before DebugDraw to collect shared channel informations such as decoded positions form the PoseVector
+	virtual void PreDebugDraw(UE::PoseSearch::FDebugDrawParams& DrawParams, TConstArrayView<float> PoseVector) const {}
+
 	// Draw this channel's data for the given pose vector
 	virtual void DebugDraw(const UE::PoseSearch::FDebugDrawParams& DrawParams, TConstArrayView<float> PoseVector) const PURE_VIRTUAL(UPoseSearchFeatureChannel::DebugDraw, );
 
@@ -180,9 +181,9 @@ private:
 protected:
 	friend class ::UPoseSearchSchema;
 
-	UPROPERTY(meta = (ExcludeFromHash))
+	UPROPERTY(Transient)
 	int32 ChannelDataOffset = INDEX_NONE;
 
-	UPROPERTY(meta = (ExcludeFromHash))
+	UPROPERTY(Transient)
 	int32 ChannelCardinality = INDEX_NONE;
 };
\ No newline at end of file
diff --git a/Engine/Plugins/Experimental/Animation/PoseSearch/Source/Runtime/Public/PoseSearch/PoseSearchHistory.h b/Engine/Plugins/Experimental/Animation/PoseSearch/Source/Runtime/Public/PoseSearch/PoseSearchHistory.h
--- a/Engine/Plugins/Experimental/Animation/PoseSearch/Source/Runtime/Public/PoseSearch/PoseSearchHistory.h
+++ b/Engine/Plugins/Experimental/Animation/PoseSearch/Source/Runtime/Public/PoseSearch/PoseSearchHistory.h
@@ -39,7 +39,7 @@ private:
 
 	struct FPose
 	{
-		FTransform RootTransform;
+		FTransform RootTransform; // @todo: remove RootTransform: this should be unnecessary, since FTrajectorySampleRange already contains the past as well as the prediction for the RootTransform
 		TArray<FTransform> LocalTransforms;
 		float Time = 0.0f;
 	};
diff --git a/Engine/Plugins/Experimental/Animation/PoseSearch/Source/Runtime/Public/PoseSearch/PoseSearchSchema.h b/Engine/Plugins/Experimental/Animation/PoseSearch/Source/Runtime/Public/PoseSearch/PoseSearchSchema.h
--- a/Engine/Plugins/Experimental/Animation/PoseSearch/Source/Runtime/Public/PoseSearch/PoseSearchSchema.h
+++ b/Engine/Plugins/Experimental/Animation/PoseSearch/Source/Runtime/Public/PoseSearch/PoseSearchSchema.h
@@ -67,11 +67,11 @@ public:
 	UPROPERTY(EditAnywhere, Category = "Schema")
 	EPoseSearchDataPreprocessor DataPreprocessor = EPoseSearchDataPreprocessor::Normalize;
 
-	UPROPERTY(meta = (ExcludeFromHash))
+	UPROPERTY(Transient)
 	int32 SchemaCardinality = 0;
 
 	// @todo: used only for indexing: cache it somewhere else
-	UPROPERTY(meta = (ExcludeFromHash))
+	UPROPERTY(Transient)
 	TArray<FBoneReference> BoneReferences;
 
 	UPROPERTY(Transient)
@@ -110,13 +110,14 @@ public:
 		return nullptr;
 	}
 
-
 	// UObject
 	virtual void PreSave(FObjectPreSaveContext ObjectSaveContext) override;
 	virtual void PostLoad() override;
 
+	int32 AddBoneReference(const FBoneReference& BoneReference);
+
 	// IBoneReferenceSkeletonProvider
-	class USkeleton* GetSkeleton(bool& bInvalidSkeletonIsError, const IPropertyHandle* PropertyHandle) override { bInvalidSkeletonIsError = false; return Skeleton; }
+	USkeleton* GetSkeleton(bool& bInvalidSkeletonIsError, const IPropertyHandle* PropertyHandle) override;
 
 #if WITH_EDITOR
 	void ComputeCostBreakdowns(UE::PoseSearch::ICostBreakDownData& CostBreakDownData) const;
@@ -124,12 +125,12 @@ public:
 	virtual void PostEditChangeProperty(struct FPropertyChangedEvent& PropertyChangedEvent) override;
 #endif
 
-	int32 AddBoneReference(const FBoneReference& BoneReference) { return BoneReferences.AddUnique(BoneReference); }
+	void BuildQuery(UE::PoseSearch::FSearchContext& SearchContext, FPoseSearchFeatureVectorBuilder& InOutQuery) const;
+
+	static constexpr FBoneIndexType RootBoneIdx = 0xFFFF;
+	FBoneIndexType GetBoneIndexType(int8 SchemaBoneIdx) const;
 
 private:
-	void Finalize(bool bRemoveEmptyChannels = true);
+	void Finalize();
 	void ResolveBoneReferences();
-
-public:
-	void BuildQuery(UE::PoseSearch::FSearchContext& SearchContext, FPoseSearchFeatureVectorBuilder& InOutQuery) const;
 };
diff --git a/Engine/Source/Runtime/Engine/Private/Animation/MotionTrajectoryTypes.cpp b/Engine/Source/Runtime/Engine/Private/Animation/MotionTrajectoryTypes.cpp
--- a/Engine/Source/Runtime/Engine/Private/Animation/MotionTrajectoryTypes.cpp
+++ b/Engine/Source/Runtime/Engine/Private/Animation/MotionTrajectoryTypes.cpp
@@ -156,6 +156,38 @@ void FTrajectorySampleRange::TransformReferenceFrame(const FTransform& Transform
 	}
 }
 
+FTrajectorySample FTrajectorySampleRange::GetSampleAtTime(float Time, bool bExtrapolate) const
+{
+	const int32 Num = Samples.Num();
+	if (Num > 1)
+	{
+		const int32 LowerBoundIdx = Algo::LowerBound(Samples, Time, [](const FTrajectorySample& TrajectorySample, float Value)
+			{
+				return Value > TrajectorySample.AccumulatedSeconds;
+			});
+
+		const int32 NextIdx = FMath::Clamp(LowerBoundIdx, 1, Samples.Num() - 1);
+		const int32 PrevIdx = NextIdx - 1;
+
+		const float Denominator = Samples[NextIdx].AccumulatedSeconds - Samples[PrevIdx].AccumulatedSeconds;
+		if (!FMath::IsNearlyZero(Denominator))
+		{
+			const float Numerator = Time - Samples[PrevIdx].AccumulatedSeconds;
+			const float LerpValue = bExtrapolate ? Numerator / Denominator : FMath::Clamp(Numerator / Denominator, 0.f, 1.f);
+			return Samples[PrevIdx].Lerp(Samples[NextIdx], LerpValue);
+		}
+
+		return Samples[PrevIdx];
+	}
+
+	if (Num > 0)
+	{
+		return Samples[0];
+	}
+
+	return FTrajectorySample();
+}
+
 void FTrajectorySampleRange::DebugDrawTrajectory(bool bEnable
 	, const UWorld* World
 	, const FTransform& WorldTransform
diff --git a/Engine/Source/Runtime/Engine/Public/Animation/MotionTrajectoryTypes.h b/Engine/Source/Runtime/Engine/Public/Animation/MotionTrajectoryTypes.h
--- a/Engine/Source/Runtime/Engine/Public/Animation/MotionTrajectoryTypes.h
+++ b/Engine/Source/Runtime/Engine/Public/Animation/MotionTrajectoryTypes.h
@@ -2,8 +2,8 @@
 
 #pragma once
 
+#include "Algo/BinarySearch.h"
 #include "UObject/ObjectMacros.h"
-
 #include "MotionTrajectoryTypes.generated.h"
 
 USTRUCT(BlueprintType, Category="Motion Trajectory")
@@ -49,28 +49,10 @@ struct ENGINE_API FTrajectorySampleRange
 {
 	GENERATED_BODY()
 
-	// Debug rendering contants
-	static constexpr FLinearColor DebugDefaultPredictionColor{ 0.f, 1.f, 0.f };
-	static constexpr FLinearColor DebugDefaultHistoryColor{ 0.f, 0.f, 1.f };
-	static constexpr float DebugDefaultTransformScale = 10.f;
-	static constexpr float DebugDefaultTransformThickness = 2.f;
-	static constexpr float DebugDefaultArrowScale = 0.025f;
-	static constexpr float DebugDefaultArrowSize = 40.f;
-	static constexpr float DebugDefaultArrowThickness = 2.f;
-
 	// Linearly ordered container for past, present, and future trajectory samples
 	UPROPERTY(EditAnywhere, BlueprintReadWrite, Category="Motion Trajectory")
 	TArray<FTrajectorySample> Samples;
 
-	// Per-second sample rate of the trajectory sample container
-	UPROPERTY(EditAnywhere, BlueprintReadWrite, Category="Motion Trajectory")
-	int32 SampleRate = 30;
-
-	// Constructors
-	FTrajectorySampleRange() = default;
-
-	FTrajectorySampleRange(int32 Rate) : Samples(), SampleRate(Rate) { }
-
 	// Removes history samples from trajectory (retains present and future)
 	void RemoveHistory();
 
@@ -89,74 +71,17 @@ struct ENGINE_API FTrajectorySampleRange
 	// Determine if all trajectory samples are default values
 	bool HasOnlyZeroSamples() const;
 
+	FTrajectorySample GetSampleAtTime(float Time, bool bExtrapolate = false) const;
+
 	// Debug draw in-world trajectory samples and optional corresponding information
 	void DebugDrawTrajectory(bool bEnable
 		, const UWorld* World
 		, const FTransform& WorldTransform
-		, const FLinearColor PredictionColor = DebugDefaultPredictionColor
-		, const FLinearColor HistoryColor = DebugDefaultHistoryColor
-		, float TransformScale = DebugDefaultTransformScale
-		, float TransformThickness = DebugDefaultTransformThickness
-		, float VelArrowScale = DebugDefaultArrowScale
-		, float VelArrowSize = DebugDefaultArrowSize
-		, float VelArrowThickness = DebugDefaultArrowThickness) const;
-
-	// Iterator for precise subsampling of the trajectory
-	template <typename Container> static FTrajectorySample IterSampleTrajectory(const Container& Samples, float DomainValue, int32& InitialIdx, bool bSmoothInterp = false)
-	{
-		auto SafeNextIndex = [](float Idx, float Num)
-		{
-			return FMath::Min(Idx + 1, Num - 1);
-		};
-
-		auto SafePrevIndex = [](float Idx)
-		{
-			return Idx > 0 ? Idx - 1 : 0;
-		};
-
-		// Iterative invocations of this function should use the NextIdx output parameter as a method of short circuiting a full linear traversal of the range
-		for (int32 Idx = InitialIdx, Num = Samples.Num(); Idx < Num; ++Idx)
-		{
-			const FTrajectorySample& NextSample = Samples[Idx];
-			const float NextDomainValue = NextSample.AccumulatedSeconds;
-
-			// Continue traversing the samples until an appropriate "right-hand" side domain value is found for interpolation 
-			if (NextDomainValue < DomainValue)
-			{
-				continue;
-			}
-			else
-			{
-				// Range check to disallow a right-hand 0 index from using -1 as a left index
-				InitialIdx = SafePrevIndex(Idx);
-
-				// Find the interpolation factor of P between: [P1 .. P .. P2] by mapping the domain values to [0 ... 1]
-				// Note: NextIdx is biased to equal P1 rather than P2 to account for cases where subsequent trajectory iterations may require subsampling between [P1 ... P2] again
-				const FTrajectorySample& InitialSample = Samples[InitialIdx];
-				const float PrevDomainValue = InitialSample.AccumulatedSeconds;
-				const float Alpha = FMath::GetMappedRangeValueUnclamped(FVector2f(PrevDomainValue, NextDomainValue), FVector2f(0.f, 1.f), DomainValue);
-
-				FTrajectorySample InterpSample;
-				if (bSmoothInterp)
-				{
-					// Apply Centripetal Catmull–Rom spline interpolation when four valid control points can be established during sub-sampling
-					const int32 PrevIdx = SafePrevIndex(InitialIdx);
-					const int32 NextIdx = SafeNextIndex(Idx, Num);
-
-					if (PrevIdx != InitialIdx && InitialIdx != Idx && Idx != NextIdx)
-					{
-						InterpSample = InitialSample.SmoothInterp(Samples[PrevIdx], NextSample, Samples[NextIdx], Alpha);
-						return InterpSample;
-					}
-				}
-
-				// For now, we will by default defer to linear interpolation
-				InterpSample = InitialSample.Lerp(NextSample, Alpha);
-				return InterpSample;
-			}
-		}
-
-		// Sampling beyond the available range will extrapolate as the last sample
-		return !Samples.IsEmpty() ? Samples.Last() : FTrajectorySample();
-	}
+		, const FLinearColor PredictionColor = { 0.f, 1.f, 0.f }
+		, const FLinearColor HistoryColor = { 0.f, 0.f, 1.f }
+		, float TransformScale = 10.f
+		, float TransformThickness = 2.f
+		, float VelArrowScale = 0.025f
+		, float VelArrowSize = 40.f
+		, float VelArrowThickness = 2.f) const;
 };
\ No newline at end of file
