motion matching - deprecating UPoseSearchFeatureChannel ComputeCostBreakdowns and PopulateChannelLayoutSet

diff --git a/Engine/Plugins/Experimental/Animation/PoseSearch/Source/Editor/Private/PoseSearchDebuggerDatabaseView.cpp b/Engine/Plugins/Experimental/Animation/PoseSearch/Source/Editor/Private/PoseSearchDebuggerDatabaseView.cpp
--- a/Engine/Plugins/Experimental/Animation/PoseSearch/Source/Editor/Private/PoseSearchDebuggerDatabaseView.cpp
+++ b/Engine/Plugins/Experimental/Animation/PoseSearch/Source/Editor/Private/PoseSearchDebuggerDatabaseView.cpp
@@ -24,108 +24,134 @@ namespace UE::PoseSearch
 
 void SDebuggerDatabaseView::Update(const FTraceMotionMatchingStateMessage& State)
 {
-	class SCostBreakDownData : public ICostBreakDownData
+	class SCostBreakDownData
 	{
 	public:
-		SCostBreakDownData(TArray<TSharedRef<FDebuggerDatabaseRowData>>& UnfilteredDatabaseRows_, bool bIsVerbose_)
-			: UnfilteredDatabaseRows(UnfilteredDatabaseRows_)
-			, bIsVerbose(bIsVerbose_)
+		SCostBreakDownData(const TArray<FTraceMotionMatchingStateDatabaseEntry>& DatabaseEntries, bool bIsVerbose)
 		{
+			// processing all the DatabaseEntries to collect the LabelToChannels
+			for (const FTraceMotionMatchingStateDatabaseEntry& DbEntry : DatabaseEntries)
+			{
+				const UPoseSearchDatabase* Database = FTraceMotionMatchingState::GetObjectFromId<UPoseSearchDatabase>(DbEntry.DatabaseId);
+				if (FAsyncPoseSearchDatabasesManagement::RequestAsyncBuildIndex(Database, ERequestAsyncBuildFlag::ContinueRequest))
+				{
+					for (const TObjectPtr<UPoseSearchFeatureChannel>& ChannelPtr : Database->Schema->Channels)
+					{
+						if (const UPoseSearchFeatureChannel* Channel = ChannelPtr.Get())
+						{
+							AnalyzeChannelRecursively(Channel, bIsVerbose);
+						}
+					}
+				}
+			}
 		}
 
-		// begin ICostBreakDownData
-		virtual ~SCostBreakDownData()
+		void ProcessData(TArray<TSharedRef<FDebuggerDatabaseRowData>>& InOutUnfilteredDatabaseRows) const
 		{
-			check(CurrentLabeIndex == -1);
-		}
+			for (TSharedRef<FDebuggerDatabaseRowData>& UnfilteredDatabaseRowRef : InOutUnfilteredDatabaseRows)
+			{
+				FDebuggerDatabaseRowData& UnfilteredDatabaseRow = UnfilteredDatabaseRowRef.Get();
+				UnfilteredDatabaseRow.CostBreakdowns.AddDefaulted(LabelToChannels.Num());
 
-		virtual int32 Num() const override
-		{
-			return UnfilteredDatabaseRows.Num();
-		}
+				for (int32 LabelToChannelIndex = 0; LabelToChannelIndex < LabelToChannels.Num(); ++LabelToChannelIndex)
+				{
+					const FLabelToChannels& LabelToChannel = LabelToChannels[LabelToChannelIndex];
 
-		virtual bool IsCostVectorFromSchema(int32 Index, const UPoseSearchSchema* Schema) const override
-		{
-			return UnfilteredDatabaseRows[Index]->SourceDatabase->Schema == Schema;
+					// there should only be at most one channel per schema with the unique label,
+					// but we'll keep this generic allowing multiple channels from the same schema having the same label.
+					// the cost will be the sum of all the channels cost
+					float CostBreakdown = 0.f;
+					for (const UPoseSearchFeatureChannel* Channel : LabelToChannel.Channels)
+					{
+						// checking if the row is associated to the Channel
+						if (UnfilteredDatabaseRow.SourceDatabase->Schema == Channel->GetSchema())
+						{
+							CostBreakdown += ArraySum(UnfilteredDatabaseRow.CostVector, Channel->GetChannelDataOffset(), Channel->GetChannelCardinality());
+						}
+					}
+					UnfilteredDatabaseRow.CostBreakdowns[LabelToChannelIndex] = CostBreakdown;
+				}
+			}
 		}
 
-		virtual TConstArrayView<float> GetCostVector(int32 Index, const UPoseSearchSchema* Schema) const override
+		bool AreLabelsEqualTo(const TArray<FText>& OtherLabels) const
 		{
-			check(UnfilteredDatabaseRows[Index]->SourceDatabase->Schema == Schema);
-			return UnfilteredDatabaseRows[Index]->CostVector;
-		}
+			if (LabelToChannels.Num() != OtherLabels.Num())
+			{
+				return false;
+			}
 
-		virtual void BeginBreakDownSection(const FText& Label) override
-		{
-			check(CurrentLabeIndex == -1);
-			CurrentLabeIndex = GetOrAddLabelIndex(Label);
-		}
+			for (int32 i = 0; i < LabelToChannels.Num(); ++i)
+			{
+				if (!LabelToChannels[i].Label.EqualTo(OtherLabels[i]))
+				{
+					return false;
+				}
+			}
 
-		virtual void SetCostBreakDown(float CostBreakDown, int32 Index, const UPoseSearchSchema* Schema) override
-		{
-			FDebuggerDatabaseRowData& UnfilteredDatabaseRow = UnfilteredDatabaseRows[Index].Get();
-			check(UnfilteredDatabaseRow.SourceDatabase->Schema == Schema);
-			TArray<float>& CostBreakdowns = UnfilteredDatabaseRow.CostBreakdowns;
-			CostBreakdowns.SetNumZeroed(CurrentLabeIndex + 1, false);
-			CostBreakdowns[CurrentLabeIndex] = CostBreakDown;
+			return true;
 		}
 
-		virtual void EndBreakDownSection(const FText& Label) override
+		const TArray<FText> GetLabels() const
 		{
-			check(CurrentLabeIndex == GetOrAddLabelIndex(Label));
-			CurrentLabeIndex = INDEX_NONE;
+			TArray<FText> Labels;
+			for (const FLabelToChannels& LabelToChannel : LabelToChannels)
+			{
+				Labels.Add(LabelToChannel.Label);
+			}
+			return Labels;
 		}
 
-		virtual bool IsVerbose() const override
+	private:
+		void AnalyzeChannelRecursively(const UPoseSearchFeatureChannel* Channel, bool bIsVerbose)
 		{
-			return bIsVerbose;
-		}
-		// end ICostBreakDownData
+			const FText Label = FText::FromString(Channel->GetLabel());
 
-		int32 GetOrAddLabelIndex(const FText& Label)
-		{
-			for (int32 i = 0; i < Labels.Num(); ++i)
+			bool bLabelFound = false;
+			for (int32 i = 0; i < LabelToChannels.Num(); ++i)
 			{
-				if (Labels[i].EqualTo(Label))
+				if (LabelToChannels[i].Label.EqualTo(Label))
 				{
-					return i;
+					LabelToChannels[i].Channels.AddUnique(Channel);
+					bLabelFound = true;
 				}
 			}
-			Labels.Add(Label);
-			return Labels.Num() - 1;
-		}
-
-		bool AreLabelsEqualTo(TArray<FText> OtherLabels)
-		{
-			if (Labels.Num() != OtherLabels.Num())
+			if (!bLabelFound)
 			{
-				return false;
+				LabelToChannels.AddDefaulted();
+				LabelToChannels.Last().Label = Label;
+				LabelToChannels.Last().Channels.Add(Channel);
 			}
 
-			for (int32 i = 0; i < Labels.Num(); ++i)
+			if (bIsVerbose)
 			{
-				if (!Labels[i].EqualTo(OtherLabels[i]))
+				for (const TObjectPtr<UPoseSearchFeatureChannel>& SubChannelPtr : Channel->GetSubChannels())
 				{
-					return false;
+					if (const UPoseSearchFeatureChannel* SubChannel = SubChannelPtr.Get())
+					{
+						AnalyzeChannelRecursively(SubChannel, bIsVerbose);
+					}
 				}
 			}
-
-			return true;
 		}
 
-		const TArray<FText>& GetLabels() const
+		static float ArraySum(TConstArrayView<float> View, int32 StartIndex, int32 Offset)
 		{
-			return Labels;
+			float Sum = 0.f;
+			const int32 EndIndex = StartIndex + Offset;
+			for (int i = StartIndex; i < EndIndex; ++i)
+			{
+				Sum += View[i];
+			}
+			return Sum;
 		}
 
-	private:
-		TArray<TSharedRef<FDebuggerDatabaseRowData>>& UnfilteredDatabaseRows;
-
-		TArray<FText> Labels;
-
-		int32 CurrentLabeIndex = INDEX_NONE;
-
-		bool bIsVerbose = true;
+		struct FLabelToChannels
+		{
+			FText Label;
+			TArray<const UPoseSearchFeatureChannel*> Channels; // NoTe: channels can be from different schemas
+		};
+		TArray<FLabelToChannels> LabelToChannels;
 	};
 
 	// row cost color palette
@@ -210,17 +236,10 @@ void SDebuggerDatabaseView::Update(const FTraceMotionMatchingStateMessage& State
 		}
 	}
 
-	SCostBreakDownData CostBreakDownData(UnfilteredDatabaseRows, bIsVerbose);
+	SCostBreakDownData CostBreakDownData(State.DatabaseEntries, bIsVerbose);
 	if (!UnfilteredDatabaseRows.IsEmpty())
 	{
-		for (const FTraceMotionMatchingStateDatabaseEntry& DbEntry : State.DatabaseEntries)
-		{
-			const UPoseSearchDatabase* Database = FTraceMotionMatchingState::GetObjectFromId<UPoseSearchDatabase>(DbEntry.DatabaseId);
-			if (FAsyncPoseSearchDatabasesManagement::RequestAsyncBuildIndex(Database, ERequestAsyncBuildFlag::ContinueRequest))
-			{
-				Database->Schema->ComputeCostBreakdowns(CostBreakDownData);
-			}
-		}
+		CostBreakDownData.ProcessData(UnfilteredDatabaseRows);
 
 		// calculating breakdowns min max and colors
 		TArray<float> MinCostBreakdowns;
diff --git a/Engine/Plugins/Experimental/Animation/PoseSearch/Source/Runtime/Private/PoseSearchDerivedData.cpp b/Engine/Plugins/Experimental/Animation/PoseSearch/Source/Runtime/Private/PoseSearchDerivedData.cpp
--- a/Engine/Plugins/Experimental/Animation/PoseSearch/Source/Runtime/Private/PoseSearchDerivedData.cpp
+++ b/Engine/Plugins/Experimental/Animation/PoseSearch/Source/Runtime/Private/PoseSearchDerivedData.cpp
@@ -39,6 +39,61 @@ static FCookStatsManager::FAutoRegisterCallback RegisterCookStats([](FCookStatsM
 	});
 #endif
 
+// data structure collecting the internal layout representation of UPoseSearchFeatureChannel,
+// so we can aggregate data from different FPoseSearchIndex and calculate mean deviation with a homogeneous data set (ComputeChannelsDeviations
+struct FFeatureChannelLayoutSet
+{
+	// data structure holding DataOffset and Cardinality to find the data of the DebugName (channel breakdown / layout) in the related SearchIndexBases[SchemaIndex]
+	struct FEntry
+	{
+		FString DebugName; // for easier debugging
+		int32 SchemaIndex = -1; // index of the associated Schemas / SearchIndexBases used as input of the algorithm
+		int32 DataOffset = -1; // data offset from the base of SearchIndexBases[SchemaIndex].Values.GetData() from where the data associated to this Item starts
+		int32 Cardinality = -1; // data cardinality
+	};
+
+	// FIoHash is the hash associated to the channel data breakdown (e.g.: it could be a single SampledBones at a specific SampleTimes for a UPoseSearchFeatureChannel_Pose)
+	TMap<FIoHash, TArray<FEntry>> EntriesMap;
+	int32 CurrentSchemaIndex = -1;
+	TWeakObjectPtr<const UPoseSearchSchema> CurrentSchema;
+
+	void Add(FString DebugName, FIoHash IoHash, int32 DataOffset, int32 Cardinality)
+	{
+		check(DataOffset >= 0 && Cardinality >= 0 && CurrentSchemaIndex >= 0);
+		TArray<FEntry>& Entries = EntriesMap.FindOrAdd(IoHash);
+
+		// making sure all the FEntry associated with the same IoHash have the same Cardinality
+		check(Entries.IsEmpty() || Entries[0].Cardinality == Cardinality);
+		Entries.Add({ DebugName, CurrentSchemaIndex, DataOffset, Cardinality });
+	}
+
+	void AnalyzeChannelRecursively(const UPoseSearchFeatureChannel* Channel)
+	{
+		const TConstArrayView<TObjectPtr<UPoseSearchFeatureChannel>> SubChannels = Channel->GetSubChannels();
+		if (SubChannels.Num() == 0)
+		{
+			// @todo: figure out if Label + SkeletonName is enough to identify this channel, if better performances are needed
+			// FString Label = Channel->GetLabel();
+			// FString SkeletonName = Channel->GetSchema()->Skeleton->GetName();
+			// UE::PoseSearch::FKeyBuilder KeyBuilder;
+			// KeyBuilder << SkeletonName << Label;
+
+			Add(Channel->GetLabel(), UE::PoseSearch::FKeyBuilder(Channel).Finalize(), Channel->GetChannelDataOffset(), Channel->GetChannelCardinality());
+		}
+		else
+		{
+			// the channel is a group channel, so we AnalyzeChannelRecursively
+			for (const TObjectPtr<UPoseSearchFeatureChannel>& SubChannelPtr : Channel->GetSubChannels())
+			{
+				if (const UPoseSearchFeatureChannel* SubChannel = SubChannelPtr.Get())
+				{
+					AnalyzeChannelRecursively(SubChannel);
+				}
+			}
+		}
+	}
+};
+
 static float ComputeFeatureMeanDeviation(TConstArrayView<FFeatureChannelLayoutSet::FEntry> Entries, TConstArrayView<FPoseSearchIndexBase> SearchIndexBases, TConstArrayView<const UPoseSearchSchema*> Schemas)
 {
 	check(Schemas.Num() == SearchIndexBases.Num());
@@ -128,9 +183,9 @@ static TArray<float> ComputeChannelsDeviations(TConstArrayView<FPoseSearchIndexB
 			FeatureChannelLayoutSet.CurrentSchema = Schema;
 			for (const TObjectPtr<UPoseSearchFeatureChannel>& ChannelPtr : Schema->Channels)
 			{
-				if (ChannelPtr)
+				if (const UPoseSearchFeatureChannel* Channel = ChannelPtr.Get())
 				{
-					ChannelPtr->PopulateChannelLayoutSet(FeatureChannelLayoutSet);
+					FeatureChannelLayoutSet.AnalyzeChannelRecursively(Channel);
 				}
 			}
 		}
diff --git a/Engine/Plugins/Experimental/Animation/PoseSearch/Source/Runtime/Private/PoseSearchFeatureChannel.cpp b/Engine/Plugins/Experimental/Animation/PoseSearch/Source/Runtime/Private/PoseSearchFeatureChannel.cpp
--- a/Engine/Plugins/Experimental/Animation/PoseSearch/Source/Runtime/Private/PoseSearchFeatureChannel.cpp
+++ b/Engine/Plugins/Experimental/Animation/PoseSearch/Source/Runtime/Private/PoseSearchFeatureChannel.cpp
@@ -8,36 +8,6 @@
 namespace UE::PoseSearch
 {
 
-//////////////////////////////////////////////////////////////////////////
-// class ICostBreakDownData
-static inline float ArraySum(TConstArrayView<float> View, int32 StartIndex, int32 Offset)
-{
-	float Sum = 0.f;
-	const int32 EndIndex = StartIndex + Offset;
-	for (int i = StartIndex; i < EndIndex; ++i)
-	{
-		Sum += View[i];
-	}
-	return Sum;
-}
-
-void ICostBreakDownData::AddEntireBreakDownSection(const FText& Label, const UPoseSearchSchema* Schema, int32 DataOffset, int32 Cardinality)
-{
-	BeginBreakDownSection(Label);
-
-	const int32 Count = Num();
-	for (int32 i = 0; i < Count; ++i)
-	{
-		if (IsCostVectorFromSchema(i, Schema))
-		{
-			const float CostBreakdown = ArraySum(GetCostVector(i, Schema), DataOffset, Cardinality);
-			SetCostBreakDown(CostBreakdown, i, Schema);
-		}
-	}
-
-	EndBreakDownSection(Label);
-}
-
 //////////////////////////////////////////////////////////////////////////
 // FFeatureVectorHelper
 void FFeatureVectorHelper::EncodeQuat(TArrayView<float> Values, int32& DataOffset, const FQuat& Quat)
@@ -229,29 +199,43 @@ float FFeatureVectorHelper::DecodeFloatAtOffset(TConstArrayView<float> Values, i
 //////////////////////////////////////////////////////////////////////////
 // UPoseSearchFeatureChannel
 #if WITH_EDITOR
-void UPoseSearchFeatureChannel::PopulateChannelLayoutSet(UE::PoseSearch::FFeatureChannelLayoutSet& FeatureChannelLayoutSet) const
+FString UPoseSearchFeatureChannel::GetLabel() const
 {
-	FeatureChannelLayoutSet.Add(GetName(), UE::PoseSearch::FKeyBuilder(this).Finalize(), ChannelDataOffset, ChannelCardinality);
-}
-
-void UPoseSearchFeatureChannel::ComputeCostBreakdowns(UE::PoseSearch::ICostBreakDownData& CostBreakDownData, const UPoseSearchSchema* Schema) const
-{
-	CostBreakDownData.AddEntireBreakDownSection(FText::FromString(GetName()), Schema, ChannelDataOffset, ChannelCardinality);
+	TStringBuilder<256> Label;
+	if (const UPoseSearchFeatureChannel* OuterChannel = Cast<UPoseSearchFeatureChannel>(GetOuter()))
+	{
+		Label.Append(OuterChannel->GetLabel());
+		Label.Append(TEXT("_"));
+	}
+	Label.Append(GetName());
+	return Label.ToString();
 }
-#endif // WITH_EDITOR
 
-class USkeleton* UPoseSearchFeatureChannel::GetSkeleton(bool& bInvalidSkeletonIsError, const IPropertyHandle* PropertyHandle)
+const UPoseSearchSchema* UPoseSearchFeatureChannel::GetSchema() const
 {
-	bInvalidSkeletonIsError = false;
-
 	UObject* Outer = GetOuter();
 	while (Outer != nullptr)
 	{
 		if (const UPoseSearchSchema* Schema = Cast<UPoseSearchSchema>(Outer))
 		{
-			return Schema->Skeleton;
+			return Schema;
 		}
 		Outer = Outer->GetOuter();
 	}
+	checkNoEntry();
+	return nullptr;
+}
+#endif // WITH_EDITOR
+
+USkeleton* UPoseSearchFeatureChannel::GetSkeleton(bool& bInvalidSkeletonIsError, const IPropertyHandle* PropertyHandle)
+{
+	bInvalidSkeletonIsError = false;
+#if WITH_EDITOR
+	return GetSchema()->Skeleton;
+#else
+	checkNoEntry();
 	return nullptr;
-}
\ No newline at end of file
+#endif
+}
+
+
diff --git a/Engine/Plugins/Experimental/Animation/PoseSearch/Source/Runtime/Private/PoseSearchFeatureChannel_FilterCrashingLegs.cpp b/Engine/Plugins/Experimental/Animation/PoseSearch/Source/Runtime/Private/PoseSearchFeatureChannel_FilterCrashingLegs.cpp
--- a/Engine/Plugins/Experimental/Animation/PoseSearch/Source/Runtime/Private/PoseSearchFeatureChannel_FilterCrashingLegs.cpp
+++ b/Engine/Plugins/Experimental/Animation/PoseSearch/Source/Runtime/Private/PoseSearchFeatureChannel_FilterCrashingLegs.cpp
@@ -163,3 +163,10 @@ bool UPoseSearchFeatureChannel_FilterCrashingLegs::IsPoseValid(TConstArrayView<f
 
 	return bIsPoseValid;
 }
+
+#if WITH_EDITOR
+FString UPoseSearchFeatureChannel_FilterCrashingLegs::GetLabel() const
+{
+	return Super::GetLabel();
+}
+#endif
\ No newline at end of file
diff --git a/Engine/Plugins/Experimental/Animation/PoseSearch/Source/Runtime/Private/PoseSearchFeatureChannel_FilterCrashingLegs.h b/Engine/Plugins/Experimental/Animation/PoseSearch/Source/Runtime/Private/PoseSearchFeatureChannel_FilterCrashingLegs.h
--- a/Engine/Plugins/Experimental/Animation/PoseSearch/Source/Runtime/Private/PoseSearchFeatureChannel_FilterCrashingLegs.h
+++ b/Engine/Plugins/Experimental/Animation/PoseSearch/Source/Runtime/Private/PoseSearchFeatureChannel_FilterCrashingLegs.h
@@ -58,6 +58,9 @@ public:
 	virtual void IndexAsset(UE::PoseSearch::IAssetIndexer& Indexer, TArrayView<float> FeatureVectorTable) const override;
 	virtual void BuildQuery(UE::PoseSearch::FSearchContext& SearchContext, FPoseSearchFeatureVectorBuilder& InOutQuery) const override;
 	virtual void DebugDraw(const UE::PoseSearch::FDebugDrawParams& DrawParams, TConstArrayView<float> PoseVector) const override;
+#if WITH_EDITOR
+	virtual FString GetLabel() const;
+#endif
 
 	// IPoseFilter interface
 	virtual bool IsPoseFilterActive() const override;
diff --git a/Engine/Plugins/Experimental/Animation/PoseSearch/Source/Runtime/Private/PoseSearchFeatureChannel_Group.cpp b/Engine/Plugins/Experimental/Animation/PoseSearch/Source/Runtime/Private/PoseSearchFeatureChannel_Group.cpp
--- a/Engine/Plugins/Experimental/Animation/PoseSearch/Source/Runtime/Private/PoseSearchFeatureChannel_Group.cpp
+++ b/Engine/Plugins/Experimental/Animation/PoseSearch/Source/Runtime/Private/PoseSearchFeatureChannel_Group.cpp
@@ -108,3 +108,17 @@ bool UPoseSearchFeatureChannel_GroupBase::IsPoseValid(TConstArrayView<float> Pos
 
 	return true;
 }
+
+#if WITH_EDITOR
+FString UPoseSearchFeatureChannel_Group::GetLabel() const
+{
+	TStringBuilder<256> Label;
+	if (const UPoseSearchFeatureChannel* OuterChannel = Cast<UPoseSearchFeatureChannel>(GetOuter()))
+	{
+		Label.Append(OuterChannel->GetLabel());
+		Label.Append(TEXT("_"));
+	}
+	Label.Append(TEXT("Group"));
+	return Label.ToString();
+}
+#endif // WITH_EDITOR
diff --git a/Engine/Plugins/Experimental/Animation/PoseSearch/Source/Runtime/Private/PoseSearchFeatureChannel_Group.h b/Engine/Plugins/Experimental/Animation/PoseSearch/Source/Runtime/Private/PoseSearchFeatureChannel_Group.h
--- a/Engine/Plugins/Experimental/Animation/PoseSearch/Source/Runtime/Private/PoseSearchFeatureChannel_Group.h
+++ b/Engine/Plugins/Experimental/Animation/PoseSearch/Source/Runtime/Private/PoseSearchFeatureChannel_Group.h
@@ -13,9 +13,6 @@ class POSESEARCH_API UPoseSearchFeatureChannel_GroupBase : public UPoseSearchFea
 	GENERATED_BODY()
 
 public:
-	virtual TArrayView<TObjectPtr<UPoseSearchFeatureChannel>> GetSubChannels() PURE_VIRTUAL(UPoseSearchFeatureChannel::GetSubChannels, return TArrayView<TObjectPtr<UPoseSearchFeatureChannel>>(); );
-	virtual TConstArrayView<TObjectPtr<UPoseSearchFeatureChannel>> GetSubChannels() const PURE_VIRTUAL(UPoseSearchFeatureChannel::GetSubChannels, return TConstArrayView<TObjectPtr<UPoseSearchFeatureChannel>>(); );
-
 	// UPoseSearchFeatureChannel interface
 	virtual void Finalize(UPoseSearchSchema* Schema) override;
 	virtual void FillWeights(TArray<float>& Weights) const override;
@@ -37,6 +34,10 @@ class POSESEARCH_API UPoseSearchFeatureChannel_Group : public UPoseSearchFeature
 	virtual TArrayView<TObjectPtr<UPoseSearchFeatureChannel>> GetSubChannels() override { return SubChannels; }
 	virtual TConstArrayView<TObjectPtr<UPoseSearchFeatureChannel>> GetSubChannels() const override { return SubChannels; }
 
+#if WITH_EDITOR
+	virtual FString GetLabel() const;
+#endif
+
 public:
 	UPROPERTY(EditAnywhere, Instanced, BlueprintReadWrite, Category = "SubChannels")
 	TArray<TObjectPtr<UPoseSearchFeatureChannel>> SubChannels;
diff --git a/Engine/Plugins/Experimental/Animation/PoseSearch/Source/Runtime/Private/PoseSearchFeatureChannel_Heading.cpp b/Engine/Plugins/Experimental/Animation/PoseSearch/Source/Runtime/Private/PoseSearchFeatureChannel_Heading.cpp
--- a/Engine/Plugins/Experimental/Animation/PoseSearch/Source/Runtime/Private/PoseSearchFeatureChannel_Heading.cpp
+++ b/Engine/Plugins/Experimental/Animation/PoseSearch/Source/Runtime/Private/PoseSearchFeatureChannel_Heading.cpp
@@ -133,3 +133,39 @@ void UPoseSearchFeatureChannel_Heading::DebugDraw(const UE::PoseSearch::FDebugDr
 	}
 #endif // ENABLE_DRAW_DEBUG
 }
+
+#if WITH_EDITOR
+FString UPoseSearchFeatureChannel_Heading::GetLabel() const
+{
+	TStringBuilder<256> Label;
+	if (const UPoseSearchFeatureChannel* OuterChannel = Cast<UPoseSearchFeatureChannel>(GetOuter()))
+	{
+		Label.Append(OuterChannel->GetLabel());
+		Label.Append(TEXT("_"));
+	}
+
+	Label.Append(TEXT("Head"));
+	switch (HeadingAxis)
+	{
+	case EHeadingAxis::X:
+		Label.Append(TEXT("X"));
+		break;
+	case EHeadingAxis::Y:
+		Label.Append(TEXT("Y"));
+		break;
+	case EHeadingAxis::Z:
+		Label.Append(TEXT("Z"));
+		break;
+	}
+
+	const FBoneReference& BoneReference = GetSchema()->BoneReferences[SchemaBoneIdx];
+	if (BoneReference.HasValidSetup())
+	{
+		Label.Append(TEXT("_"));
+		Label.Append(BoneReference.BoneName.ToString());
+	}
+
+	Label.Appendf(TEXT(" %.1f"), SampleTimeOffset);
+	return Label.ToString();
+}
+#endif
\ No newline at end of file
diff --git a/Engine/Plugins/Experimental/Animation/PoseSearch/Source/Runtime/Private/PoseSearchFeatureChannel_Heading.h b/Engine/Plugins/Experimental/Animation/PoseSearch/Source/Runtime/Private/PoseSearchFeatureChannel_Heading.h
--- a/Engine/Plugins/Experimental/Animation/PoseSearch/Source/Runtime/Private/PoseSearchFeatureChannel_Heading.h
+++ b/Engine/Plugins/Experimental/Animation/PoseSearch/Source/Runtime/Private/PoseSearchFeatureChannel_Heading.h
@@ -50,6 +50,9 @@ public:
 	virtual void IndexAsset(UE::PoseSearch::IAssetIndexer& Indexer, TArrayView<float> FeatureVectorTable) const override;
 	virtual void BuildQuery(UE::PoseSearch::FSearchContext& SearchContext, FPoseSearchFeatureVectorBuilder& InOutQuery) const override;
 	virtual void DebugDraw(const UE::PoseSearch::FDebugDrawParams& DrawParams, TConstArrayView<float> PoseVector) const override;
+#if WITH_EDITOR
+	virtual FString GetLabel() const;
+#endif
 
 	FVector GetAxis(const FQuat& Rotation) const;
 };
\ No newline at end of file
diff --git a/Engine/Plugins/Experimental/Animation/PoseSearch/Source/Runtime/Private/PoseSearchFeatureChannel_Phase.cpp b/Engine/Plugins/Experimental/Animation/PoseSearch/Source/Runtime/Private/PoseSearchFeatureChannel_Phase.cpp
--- a/Engine/Plugins/Experimental/Animation/PoseSearch/Source/Runtime/Private/PoseSearchFeatureChannel_Phase.cpp
+++ b/Engine/Plugins/Experimental/Animation/PoseSearch/Source/Runtime/Private/PoseSearchFeatureChannel_Phase.cpp
@@ -396,3 +396,25 @@ void UPoseSearchFeatureChannel_Phase::DebugDraw(const UE::PoseSearch::FDebugDraw
 	DrawDebugCircle(DrawParams.World, CircleTransform, PhaseVector.Length(), Segments, Color, bPersistent, LifeTime, DepthPriority, 0.f, false);
 #endif // ENABLE_DRAW_DEBUG
 }
+
+#if WITH_EDITOR
+FString UPoseSearchFeatureChannel_Phase::GetLabel() const
+{
+	TStringBuilder<256> Label;
+	if (const UPoseSearchFeatureChannel* OuterChannel = Cast<UPoseSearchFeatureChannel>(GetOuter()))
+	{
+		Label.Append(OuterChannel->GetLabel());
+		Label.Append(TEXT("_"));
+	}
+
+	Label.Append(TEXT("Pha"));
+	const FBoneReference& BoneReference = GetSchema()->BoneReferences[SchemaBoneIdx];
+	if (BoneReference.HasValidSetup())
+	{
+		Label.Append(TEXT("_"));
+		Label.Append(BoneReference.BoneName.ToString());
+	}
+
+	return Label.ToString();
+}
+#endif
\ No newline at end of file
diff --git a/Engine/Plugins/Experimental/Animation/PoseSearch/Source/Runtime/Private/PoseSearchFeatureChannel_Phase.h b/Engine/Plugins/Experimental/Animation/PoseSearch/Source/Runtime/Private/PoseSearchFeatureChannel_Phase.h
--- a/Engine/Plugins/Experimental/Animation/PoseSearch/Source/Runtime/Private/PoseSearchFeatureChannel_Phase.h
+++ b/Engine/Plugins/Experimental/Animation/PoseSearch/Source/Runtime/Private/PoseSearchFeatureChannel_Phase.h
@@ -33,4 +33,7 @@ public:
 	virtual void IndexAsset(UE::PoseSearch::IAssetIndexer& Indexer, TArrayView<float> FeatureVectorTable) const override;
 	virtual void BuildQuery(UE::PoseSearch::FSearchContext& SearchContext, FPoseSearchFeatureVectorBuilder& InOutQuery) const override;
 	virtual void DebugDraw(const UE::PoseSearch::FDebugDrawParams& DrawParams, TConstArrayView<float> PoseVector) const override;
+#if WITH_EDITOR
+	virtual FString GetLabel() const;
+#endif
 };
\ No newline at end of file
diff --git a/Engine/Plugins/Experimental/Animation/PoseSearch/Source/Runtime/Private/PoseSearchFeatureChannel_Pose.cpp b/Engine/Plugins/Experimental/Animation/PoseSearch/Source/Runtime/Private/PoseSearchFeatureChannel_Pose.cpp
--- a/Engine/Plugins/Experimental/Animation/PoseSearch/Source/Runtime/Private/PoseSearchFeatureChannel_Pose.cpp
+++ b/Engine/Plugins/Experimental/Animation/PoseSearch/Source/Runtime/Private/PoseSearchFeatureChannel_Pose.cpp
@@ -20,60 +20,55 @@ void UPoseSearchFeatureChannel_Pose::Finalize(UPoseSearchSchema* Schema)
 		const FPoseSearchBone& SampledBone = SampledBones[ChannelBoneIdx];
 		if (EnumHasAnyFlags(SampledBone.Flags, EPoseSearchBoneFlags::Position))
 		{
-			UPoseSearchFeatureChannel_Position* Position = NewObject<UPoseSearchFeatureChannel_Position>();
+			UPoseSearchFeatureChannel_Position* Position = NewObject<UPoseSearchFeatureChannel_Position>(this, NAME_None, RF_Transient);
 			Position->Bone = SampledBone.Reference;
 			Position->Weight = SampledBone.Weight * Weight;
 			Position->SampleTimeOffset = 0.f;
 			Position->ColorPresetIndex = SampledBone.ColorPresetIndex;
 			Position->InputQueryPose = InputQueryPose;
-			Position->SetFlags(RF_Transient);
 			SubChannels.Add(Position);
 		}
 
 		if (EnumHasAnyFlags(SampledBone.Flags, EPoseSearchBoneFlags::Rotation))
 		{
-			UPoseSearchFeatureChannel_Heading* HeadingX = NewObject<UPoseSearchFeatureChannel_Heading>();
+			UPoseSearchFeatureChannel_Heading* HeadingX = NewObject<UPoseSearchFeatureChannel_Heading>(this, NAME_None, RF_Transient);
 			HeadingX->Bone = SampledBone.Reference;
 			HeadingX->Weight = SampledBone.Weight * Weight;
 			HeadingX->SampleTimeOffset = 0.f;
 			HeadingX->HeadingAxis = EHeadingAxis::X;
 			HeadingX->ColorPresetIndex = SampledBone.ColorPresetIndex;
 			HeadingX->InputQueryPose = InputQueryPose;
-			HeadingX->SetFlags(RF_Transient);
 			SubChannels.Add(HeadingX);
 
-			UPoseSearchFeatureChannel_Heading* HeadingY = NewObject<UPoseSearchFeatureChannel_Heading>();
+			UPoseSearchFeatureChannel_Heading* HeadingY = NewObject<UPoseSearchFeatureChannel_Heading>(this, NAME_None, RF_Transient);
 			HeadingY->Bone = SampledBone.Reference;
 			HeadingY->Weight = SampledBone.Weight * Weight;
 			HeadingY->SampleTimeOffset = 0.f;
 			HeadingY->HeadingAxis = EHeadingAxis::Y;
 			HeadingY->ColorPresetIndex = SampledBone.ColorPresetIndex;
 			HeadingY->InputQueryPose = InputQueryPose;
-			HeadingY->SetFlags(RF_Transient);
 			SubChannels.Add(HeadingY);
 		}
 
 		if (EnumHasAnyFlags(SampledBone.Flags, EPoseSearchBoneFlags::Velocity))
 		{
-			UPoseSearchFeatureChannel_Velocity* Velocity = NewObject<UPoseSearchFeatureChannel_Velocity>();
+			UPoseSearchFeatureChannel_Velocity* Velocity = NewObject<UPoseSearchFeatureChannel_Velocity>(this, NAME_None, RF_Transient);
 			Velocity->Bone = SampledBone.Reference;
 			Velocity->Weight = SampledBone.Weight * Weight;
 			Velocity->SampleTimeOffset = 0.f;
 			Velocity->ColorPresetIndex = SampledBone.ColorPresetIndex;
 			Velocity->InputQueryPose = InputQueryPose;
 			Velocity->bUseCharacterSpaceVelocities = bUseCharacterSpaceVelocities;
-			Velocity->SetFlags(RF_Transient);
 			SubChannels.Add(Velocity);
 		}
 
 		if (EnumHasAnyFlags(SampledBone.Flags, EPoseSearchBoneFlags::Phase))
 		{
-			UPoseSearchFeatureChannel_Phase* Phase = NewObject<UPoseSearchFeatureChannel_Phase>();
+			UPoseSearchFeatureChannel_Phase* Phase = NewObject<UPoseSearchFeatureChannel_Phase>(this, NAME_None, RF_Transient);
 			Phase->Bone = SampledBone.Reference;
 			Phase->Weight = SampledBone.Weight * Weight;
 			Phase->ColorPresetIndex = SampledBone.ColorPresetIndex;
 			Phase->InputQueryPose = InputQueryPose;
-			Phase->SetFlags(RF_Transient);
 			SubChannels.Add(Phase);
 		}
 	}
@@ -82,101 +77,17 @@ void UPoseSearchFeatureChannel_Pose::Finalize(UPoseSearchSchema* Schema)
 }
 
 #if WITH_EDITOR
-void UPoseSearchFeatureChannel_Pose::PopulateChannelLayoutSet(UE::PoseSearch::FFeatureChannelLayoutSet& FeatureChannelLayoutSet) const
+FString UPoseSearchFeatureChannel_Pose::GetLabel() const
 {
-	for (const TObjectPtr<UPoseSearchFeatureChannel>& SubChannelPtr : GetSubChannels())
+	TStringBuilder<256> Label;
+	if (const UPoseSearchFeatureChannel* OuterChannel = Cast<UPoseSearchFeatureChannel>(GetOuter()))
 	{
-		FString BoneName;
-		EPoseSearchBoneFlags BoneFlag;
-		const char* Label;
-		const UPoseSearchFeatureChannel* Channel = SubChannelPtr.Get();
-		if (const UPoseSearchFeatureChannel_Position* Position = Cast<UPoseSearchFeatureChannel_Position>(Channel))
-		{
-			BoneName = Position->Bone.BoneName.ToString();
-			BoneFlag = EPoseSearchBoneFlags::Position;
-			Label = "Pos";
-		}
-		else if (const UPoseSearchFeatureChannel_Heading* Heading = Cast<UPoseSearchFeatureChannel_Heading>(Channel))
-		{
-			BoneName = Heading->Bone.BoneName.ToString();
-			BoneFlag = EPoseSearchBoneFlags::Rotation;
-
-			if (Heading->HeadingAxis == EHeadingAxis::X)
-			{
-				Label = "HdX";
-			}
-			else
-			{
-				check(Heading->HeadingAxis == EHeadingAxis::Y);
-				Label = "HdY";
-			}
-		}
-		else if (const UPoseSearchFeatureChannel_Velocity* Velocity = Cast<UPoseSearchFeatureChannel_Velocity>(Channel))
-		{
-			BoneName = Velocity->Bone.BoneName.ToString();
-			BoneFlag = EPoseSearchBoneFlags::Velocity;
-			Label = "Vel";
-		}
-		else if (const UPoseSearchFeatureChannel_Phase* Phase = Cast<UPoseSearchFeatureChannel_Phase>(Channel))
-		{
-			BoneName = Phase->Bone.BoneName.ToString();
-			BoneFlag = EPoseSearchBoneFlags::Phase;
-			Label = "Pha";
-		}
-		else
-		{
-			checkNoEntry();
-			continue;
-		}
-		
-		FString SkeletonName = FeatureChannelLayoutSet.CurrentSchema->Skeleton->GetName();
-		UE::PoseSearch::FKeyBuilder KeyBuilder;
-		KeyBuilder << SkeletonName << BoneName << BoneFlag;
-		FeatureChannelLayoutSet.Add(FString::Format(TEXT("{0} {1}"), { BoneName, Label }), KeyBuilder.Finalize(), Channel->GetChannelDataOffset(), Channel->GetChannelCardinality());
-	}
-}
-
-void UPoseSearchFeatureChannel_Pose::ComputeCostBreakdowns(UE::PoseSearch::ICostBreakDownData& CostBreakDownData, const UPoseSearchSchema* Schema) const
-{
-	CostBreakDownData.AddEntireBreakDownSection(LOCTEXT("ColumnLabelPoseChannelTotal", "Pose Total"), Schema, ChannelDataOffset, ChannelCardinality);
-
-	if (CostBreakDownData.IsVerbose())
-	{
-		for (const TObjectPtr<UPoseSearchFeatureChannel>& SubChannelPtr : GetSubChannels())
-		{
-			const UPoseSearchFeatureChannel* Channel = SubChannelPtr.Get();
-			if (const UPoseSearchFeatureChannel_Position* Position = Cast<UPoseSearchFeatureChannel_Position>(Channel))
-			{
-				CostBreakDownData.AddEntireBreakDownSection(FText::Format(LOCTEXT("ColumnLabelPoseChannelPosition", "{0} Pos"), FText::FromName(Position->Bone.BoneName)), Schema, Position->GetChannelDataOffset(), Position->GetChannelCardinality());
-			}
-			else if (const UPoseSearchFeatureChannel_Heading* Heading = Cast<UPoseSearchFeatureChannel_Heading>(Channel))
-			{
-				if (Heading->HeadingAxis == EHeadingAxis::X)
-				{
-					CostBreakDownData.AddEntireBreakDownSection(FText::Format(LOCTEXT("ColumnLabelPoseChannelHeading", "{0} HdX"), FText::FromName(Heading->Bone.BoneName)), Schema, Heading->GetChannelDataOffset(), Heading->GetChannelCardinality());
-				}
-				else
-				{
-					check(Heading->HeadingAxis == EHeadingAxis::Y);
-					CostBreakDownData.AddEntireBreakDownSection(FText::Format(LOCTEXT("ColumnLabelPoseChannelHeading", "{0} HdY"), FText::FromName(Heading->Bone.BoneName)), Schema, Heading->GetChannelDataOffset(), Heading->GetChannelCardinality());
-				}
-			}
-			else if (const UPoseSearchFeatureChannel_Velocity* Velocity = Cast<UPoseSearchFeatureChannel_Velocity>(Channel))
-			{
-				CostBreakDownData.AddEntireBreakDownSection(FText::Format(LOCTEXT("ColumnLabelPoseChannelVelocity", "{0} Vel"), FText::FromName(Velocity->Bone.BoneName)), Schema, Velocity->GetChannelDataOffset(), Velocity->GetChannelCardinality());
-			}
-			else if (const UPoseSearchFeatureChannel_Phase* Phase = Cast<UPoseSearchFeatureChannel_Phase>(Channel))
-			{
-				CostBreakDownData.AddEntireBreakDownSection(FText::Format(LOCTEXT("ColumnLabelPoseChannelPhase", "{0} Pha"), FText::FromName(Phase->Bone.BoneName)), Schema, Phase->GetChannelDataOffset(), Phase->GetChannelCardinality());
-			}
-			else
-			{
-				checkNoEntry();
-			}
-		}
+		Label.Append(OuterChannel->GetLabel());
+		Label.Append(TEXT("_"));
 	}
+	Label.Append(TEXT("Pose"));
+	return Label.ToString();
 }
-
 #endif // WITH_EDITOR
 
 #undef LOCTEXT_NAMESPACE
\ No newline at end of file
diff --git a/Engine/Plugins/Experimental/Animation/PoseSearch/Source/Runtime/Private/PoseSearchFeatureChannel_Pose.h b/Engine/Plugins/Experimental/Animation/PoseSearch/Source/Runtime/Private/PoseSearchFeatureChannel_Pose.h
--- a/Engine/Plugins/Experimental/Animation/PoseSearch/Source/Runtime/Private/PoseSearchFeatureChannel_Pose.h
+++ b/Engine/Plugins/Experimental/Animation/PoseSearch/Source/Runtime/Private/PoseSearchFeatureChannel_Pose.h
@@ -69,8 +69,7 @@ public:
 	virtual void Finalize(UPoseSearchSchema* Schema) override;
 
 #if WITH_EDITOR
-	virtual void PopulateChannelLayoutSet(UE::PoseSearch::FFeatureChannelLayoutSet& FeatureChannelLayoutSet) const override;
-	virtual void ComputeCostBreakdowns(UE::PoseSearch::ICostBreakDownData& CostBreakDownData, const UPoseSearchSchema* Schema) const override;
+	virtual FString GetLabel() const;
 #endif
 };
 
diff --git a/Engine/Plugins/Experimental/Animation/PoseSearch/Source/Runtime/Private/PoseSearchFeatureChannel_Position.cpp b/Engine/Plugins/Experimental/Animation/PoseSearch/Source/Runtime/Private/PoseSearchFeatureChannel_Position.cpp
--- a/Engine/Plugins/Experimental/Animation/PoseSearch/Source/Runtime/Private/PoseSearchFeatureChannel_Position.cpp
+++ b/Engine/Plugins/Experimental/Animation/PoseSearch/Source/Runtime/Private/PoseSearchFeatureChannel_Position.cpp
@@ -135,3 +135,26 @@ void UPoseSearchFeatureChannel_Position::DebugDraw(const UE::PoseSearch::FDebugD
 	}
 #endif // ENABLE_DRAW_DEBUG
 }
+
+#if WITH_EDITOR
+FString UPoseSearchFeatureChannel_Position::GetLabel() const
+{
+	TStringBuilder<256> Label;
+	if (const UPoseSearchFeatureChannel* OuterChannel = Cast<UPoseSearchFeatureChannel>(GetOuter()))
+	{
+		Label.Append(OuterChannel->GetLabel());
+		Label.Append(TEXT("_"));
+	}
+
+	Label.Append(TEXT("Pos"));
+	const FBoneReference& BoneReference = GetSchema()->BoneReferences[SchemaBoneIdx];
+	if (BoneReference.HasValidSetup())
+	{
+		Label.Append(TEXT("_"));
+		Label.Append(BoneReference.BoneName.ToString());
+	}
+
+	Label.Appendf(TEXT(" %.1f"), SampleTimeOffset);
+	return Label.ToString();
+}
+#endif
\ No newline at end of file
diff --git a/Engine/Plugins/Experimental/Animation/PoseSearch/Source/Runtime/Private/PoseSearchFeatureChannel_Position.h b/Engine/Plugins/Experimental/Animation/PoseSearch/Source/Runtime/Private/PoseSearchFeatureChannel_Position.h
--- a/Engine/Plugins/Experimental/Animation/PoseSearch/Source/Runtime/Private/PoseSearchFeatureChannel_Position.h
+++ b/Engine/Plugins/Experimental/Animation/PoseSearch/Source/Runtime/Private/PoseSearchFeatureChannel_Position.h
@@ -37,4 +37,7 @@ public:
 	virtual void BuildQuery(UE::PoseSearch::FSearchContext& SearchContext, FPoseSearchFeatureVectorBuilder& InOutQuery) const override;
 	virtual void PreDebugDraw(UE::PoseSearch::FDebugDrawParams& DrawParams, TConstArrayView<float> PoseVector) const override;
 	virtual void DebugDraw(const UE::PoseSearch::FDebugDrawParams& DrawParams, TConstArrayView<float> PoseVector) const override;
+#if WITH_EDITOR
+	virtual FString GetLabel() const;
+#endif
 };
diff --git a/Engine/Plugins/Experimental/Animation/PoseSearch/Source/Runtime/Private/PoseSearchFeatureChannel_Trajectory.cpp b/Engine/Plugins/Experimental/Animation/PoseSearch/Source/Runtime/Private/PoseSearchFeatureChannel_Trajectory.cpp
--- a/Engine/Plugins/Experimental/Animation/PoseSearch/Source/Runtime/Private/PoseSearchFeatureChannel_Trajectory.cpp
+++ b/Engine/Plugins/Experimental/Animation/PoseSearch/Source/Runtime/Private/PoseSearchFeatureChannel_Trajectory.cpp
@@ -32,51 +32,47 @@ void UPoseSearchFeatureChannel_Trajectory::Finalize(UPoseSearchSchema* Schema)
 		// @todo: implement PositionXY properly as 2 dimension channel
 		if (EnumHasAnyFlags(Sample.Flags, EPoseSearchTrajectoryFlags::Position | EPoseSearchTrajectoryFlags::PositionXY))
 		{
-			UPoseSearchFeatureChannel_Position* Position = NewObject<UPoseSearchFeatureChannel_Position>();
+			UPoseSearchFeatureChannel_Position* Position = NewObject<UPoseSearchFeatureChannel_Position>(this, NAME_None, RF_Transient);
 			Position->Weight = Sample.Weight * Weight;
 			Position->SampleTimeOffset = Sample.Offset;
 			Position->ColorPresetIndex = Sample.ColorPresetIndex;
 			Position->InputQueryPose = EInputQueryPose::UseCharacterPose;
-			Position->SetFlags(RF_Transient);
 			SubChannels.Add(Position);
 		}
 
 		// @todo: implement VelocityXY properly as 2 dimension channel
 		if (EnumHasAnyFlags(Sample.Flags, EPoseSearchTrajectoryFlags::Velocity | EPoseSearchTrajectoryFlags::VelocityXY))
 		{
-			UPoseSearchFeatureChannel_Velocity* Velocity = NewObject<UPoseSearchFeatureChannel_Velocity>();
+			UPoseSearchFeatureChannel_Velocity* Velocity = NewObject<UPoseSearchFeatureChannel_Velocity>(this, NAME_None, RF_Transient);
 			Velocity->Weight = Sample.Weight * Weight;
 			Velocity->SampleTimeOffset = Sample.Offset;
 			Velocity->ColorPresetIndex = Sample.ColorPresetIndex;
 			Velocity->InputQueryPose = EInputQueryPose::UseCharacterPose;
 			Velocity->bUseCharacterSpaceVelocities = false;
-			Velocity->SetFlags(RF_Transient);
 			SubChannels.Add(Velocity);
 		}
 
 		// @todo: implement VelocityDirectionXY properly as 2 dimension channel
 		if (EnumHasAnyFlags(Sample.Flags, EPoseSearchTrajectoryFlags::VelocityDirection | EPoseSearchTrajectoryFlags::VelocityDirectionXY))
 		{
-			UPoseSearchFeatureChannel_Velocity* Velocity = NewObject<UPoseSearchFeatureChannel_Velocity>();
+			UPoseSearchFeatureChannel_Velocity* Velocity = NewObject<UPoseSearchFeatureChannel_Velocity>(this, NAME_None, RF_Transient);
 			Velocity->Weight = Sample.Weight * Weight;
 			Velocity->SampleTimeOffset = Sample.Offset;
 			Velocity->ColorPresetIndex = Sample.ColorPresetIndex;
 			Velocity->InputQueryPose = EInputQueryPose::UseCharacterPose;
 			Velocity->bUseCharacterSpaceVelocities = false;
 			Velocity->bNormalize = true;
-			Velocity->SetFlags(RF_Transient);
 			SubChannels.Add(Velocity);
 		}
 
 		// @todo: implement FacingDirectionXY properly as 2 dimension channel
 		if (EnumHasAnyFlags(Sample.Flags, EPoseSearchTrajectoryFlags::FacingDirection | EPoseSearchTrajectoryFlags::FacingDirectionXY))
 		{
-			UPoseSearchFeatureChannel_Heading* Heading = NewObject<UPoseSearchFeatureChannel_Heading>();
+			UPoseSearchFeatureChannel_Heading* Heading = NewObject<UPoseSearchFeatureChannel_Heading>(this, NAME_None, RF_Transient);
 			Heading->Weight = Sample.Weight * Weight;
 			Heading->SampleTimeOffset = Sample.Offset;
 			Heading->ColorPresetIndex = Sample.ColorPresetIndex;
 			Heading->InputQueryPose = EInputQueryPose::UseCharacterPose;
-			Heading->SetFlags(RF_Transient);
 			SubChannels.Add(Heading);
 		}
 	}
@@ -139,82 +135,17 @@ void UPoseSearchFeatureChannel_Trajectory::DebugDraw(const UE::PoseSearch::FDebu
 }
 
 #if WITH_EDITOR
-void UPoseSearchFeatureChannel_Trajectory::PopulateChannelLayoutSet(UE::PoseSearch::FFeatureChannelLayoutSet& FeatureChannelLayoutSet) const
+FString UPoseSearchFeatureChannel_Trajectory::GetLabel() const
 {
-	for (const TObjectPtr<UPoseSearchFeatureChannel>& SubChannelPtr : GetSubChannels())
+	TStringBuilder<256> Label;
+	if (const UPoseSearchFeatureChannel* OuterChannel = Cast<UPoseSearchFeatureChannel>(GetOuter()))
 	{
-		float Offset;
-		EPoseSearchTrajectoryFlags SampleFlag;
-		const char* Label;
-		const UPoseSearchFeatureChannel* Channel = SubChannelPtr.Get();
-		if (const UPoseSearchFeatureChannel_Position* Position = Cast<UPoseSearchFeatureChannel_Position>(Channel))
-		{
-			Offset = Position->SampleTimeOffset;
-			SampleFlag = EPoseSearchTrajectoryFlags::Position;
-			Label = "Pos";
-		}
-		else if (const UPoseSearchFeatureChannel_Velocity* Velocity = Cast<UPoseSearchFeatureChannel_Velocity>(Channel))
-		{
-			Offset = Velocity->SampleTimeOffset;
-			if (Velocity->bNormalize)
-			{
-				SampleFlag = EPoseSearchTrajectoryFlags::VelocityDirection;
-				Label = "VelDir";
-			}
-			else
-			{
-				SampleFlag = EPoseSearchTrajectoryFlags::Velocity;
-				Label = "Vel";
-			}
-		}
-		else if (const UPoseSearchFeatureChannel_Heading* Heading = Cast<UPoseSearchFeatureChannel_Heading>(Channel))
-		{
-			Offset = Heading->SampleTimeOffset;
-			SampleFlag = EPoseSearchTrajectoryFlags::FacingDirection;
-			Label = "Fac";
-		}
-		else
-		{
-			checkNoEntry();
-			continue;
-		}
-
-		FString SkeletonName = FeatureChannelLayoutSet.CurrentSchema->Skeleton->GetName();
-		UE::PoseSearch::FKeyBuilder KeyBuilder;
-		KeyBuilder << SkeletonName << SampleFlag << Offset;
-		FeatureChannelLayoutSet.Add(FString::Format(TEXT("Traj {0} {1}"), { Label, Offset }), KeyBuilder.Finalize(), Channel->GetChannelDataOffset(), Channel->GetChannelCardinality());
+		Label.Append(OuterChannel->GetLabel());
+		Label.Append(TEXT("_"));
 	}
+	Label.Append(TEXT("Traj"));
+	return Label.ToString();
 }
-
-void UPoseSearchFeatureChannel_Trajectory::ComputeCostBreakdowns(UE::PoseSearch::ICostBreakDownData& CostBreakDownData, const UPoseSearchSchema* Schema) const
-{
-	CostBreakDownData.AddEntireBreakDownSection(LOCTEXT("ColumnLabelTrajChannelTotal", "Traj Total"), Schema, ChannelDataOffset, ChannelCardinality);
-
-	if (CostBreakDownData.IsVerbose())
-	{
-		for (const TObjectPtr<UPoseSearchFeatureChannel>& SubChannelPtr : GetSubChannels())
-		{
-			const UPoseSearchFeatureChannel* Channel = SubChannelPtr.Get();
-			if (const UPoseSearchFeatureChannel_Position* Position = Cast<UPoseSearchFeatureChannel_Position>(Channel))
-			{
-				CostBreakDownData.AddEntireBreakDownSection(FText::Format(LOCTEXT("ColumnLabelTrajChannelPosition", "Traj Pos {0}"), Position->SampleTimeOffset), Schema, Position->GetChannelDataOffset(), Position->GetChannelCardinality());
-			}
-			else if (const UPoseSearchFeatureChannel_Velocity* Velocity = Cast<UPoseSearchFeatureChannel_Velocity>(Channel))
-			{
-				CostBreakDownData.AddEntireBreakDownSection(FText::Format(LOCTEXT("ColumnLabelTrajChannelVelocity", "Traj Vel {0}"), Velocity->SampleTimeOffset), Schema, Velocity->GetChannelDataOffset(), Velocity->GetChannelCardinality());
-			}
-			else if (const UPoseSearchFeatureChannel_Heading* Heading = Cast<UPoseSearchFeatureChannel_Heading>(Channel))
-			{
-				CostBreakDownData.AddEntireBreakDownSection(FText::Format(LOCTEXT("ColumnLabelTrajChannelFacingDirection", "Traj Fac {0}"), Heading->SampleTimeOffset), Schema, Heading->GetChannelDataOffset(), Heading->GetChannelCardinality());
-			}
-			else
-			{
-				checkNoEntry();
-			}
-		}
-	}
-}
-
 #endif // WITH_EDITOR
 
 float UPoseSearchFeatureChannel_Trajectory::GetEstimatedSpeedRatio(TConstArrayView<float> QueryVector, TConstArrayView<float> PoseVector) const
diff --git a/Engine/Plugins/Experimental/Animation/PoseSearch/Source/Runtime/Private/PoseSearchFeatureChannel_Trajectory.h b/Engine/Plugins/Experimental/Animation/PoseSearch/Source/Runtime/Private/PoseSearchFeatureChannel_Trajectory.h
--- a/Engine/Plugins/Experimental/Animation/PoseSearch/Source/Runtime/Private/PoseSearchFeatureChannel_Trajectory.h
+++ b/Engine/Plugins/Experimental/Animation/PoseSearch/Source/Runtime/Private/PoseSearchFeatureChannel_Trajectory.h
@@ -66,12 +66,10 @@ public:
 	// UPoseSearchFeatureChannel interface
 	virtual void Finalize(UPoseSearchSchema* Schema) override;
 	virtual void DebugDraw(const UE::PoseSearch::FDebugDrawParams& DrawParams, TConstArrayView<float> PoseVector) const override;
-
-	float GetEstimatedSpeedRatio(TConstArrayView<float> QueryVector, TConstArrayView<float> PoseVector) const;
-
 #if WITH_EDITOR
-	virtual void PopulateChannelLayoutSet(UE::PoseSearch::FFeatureChannelLayoutSet& FeatureChannelLayoutSet) const override;
-	virtual void ComputeCostBreakdowns(UE::PoseSearch::ICostBreakDownData& CostBreakDownData, const UPoseSearchSchema* Schema) const override;
+	virtual FString GetLabel() const;
 #endif
+
+	float GetEstimatedSpeedRatio(TConstArrayView<float> QueryVector, TConstArrayView<float> PoseVector) const;
 };
 
diff --git a/Engine/Plugins/Experimental/Animation/PoseSearch/Source/Runtime/Private/PoseSearchFeatureChannel_Velocity.cpp b/Engine/Plugins/Experimental/Animation/PoseSearch/Source/Runtime/Private/PoseSearchFeatureChannel_Velocity.cpp
--- a/Engine/Plugins/Experimental/Animation/PoseSearch/Source/Runtime/Private/PoseSearchFeatureChannel_Velocity.cpp
+++ b/Engine/Plugins/Experimental/Animation/PoseSearch/Source/Runtime/Private/PoseSearchFeatureChannel_Velocity.cpp
@@ -161,3 +161,31 @@ void UPoseSearchFeatureChannel_Velocity::DebugDraw(const UE::PoseSearch::FDebugD
 	}
 #endif // ENABLE_DRAW_DEBUG
 }
+
+#if WITH_EDITOR
+FString UPoseSearchFeatureChannel_Velocity::GetLabel() const
+{
+	TStringBuilder<256> Label;
+	if (const UPoseSearchFeatureChannel* OuterChannel = Cast<UPoseSearchFeatureChannel>(GetOuter()))
+	{
+		Label.Append(OuterChannel->GetLabel());
+		Label.Append(TEXT("_"));
+	}
+
+	Label.Append(TEXT("Vel"));
+	if (bNormalize)
+	{
+		Label.Append(TEXT("Dir"));
+	}
+
+	const FBoneReference& BoneReference = GetSchema()->BoneReferences[SchemaBoneIdx];
+	if (BoneReference.HasValidSetup())
+	{
+		Label.Append(TEXT("_"));
+		Label.Append(BoneReference.BoneName.ToString());
+	}
+
+	Label.Appendf(TEXT(" %.1f"), SampleTimeOffset);
+	return Label.ToString();
+}
+#endif
\ No newline at end of file
diff --git a/Engine/Plugins/Experimental/Animation/PoseSearch/Source/Runtime/Private/PoseSearchFeatureChannel_Velocity.h b/Engine/Plugins/Experimental/Animation/PoseSearch/Source/Runtime/Private/PoseSearchFeatureChannel_Velocity.h
--- a/Engine/Plugins/Experimental/Animation/PoseSearch/Source/Runtime/Private/PoseSearchFeatureChannel_Velocity.h
+++ b/Engine/Plugins/Experimental/Animation/PoseSearch/Source/Runtime/Private/PoseSearchFeatureChannel_Velocity.h
@@ -44,4 +44,7 @@ public:
 	virtual void IndexAsset(UE::PoseSearch::IAssetIndexer& Indexer, TArrayView<float> FeatureVectorTable) const override;
 	virtual void BuildQuery(UE::PoseSearch::FSearchContext& SearchContext, FPoseSearchFeatureVectorBuilder& InOutQuery) const override;
 	virtual void DebugDraw(const UE::PoseSearch::FDebugDrawParams& DrawParams, TConstArrayView<float> PoseVector) const override;
+#if WITH_EDITOR
+	virtual FString GetLabel() const;
+#endif
 };
diff --git a/Engine/Plugins/Experimental/Animation/PoseSearch/Source/Runtime/Private/PoseSearchSchema.cpp b/Engine/Plugins/Experimental/Animation/PoseSearch/Source/Runtime/Private/PoseSearchSchema.cpp
--- a/Engine/Plugins/Experimental/Animation/PoseSearch/Source/Runtime/Private/PoseSearchSchema.cpp
+++ b/Engine/Plugins/Experimental/Animation/PoseSearch/Source/Runtime/Private/PoseSearchSchema.cpp
@@ -105,14 +105,4 @@ void UPoseSearchSchema::PostEditChangeProperty(FPropertyChangedEvent& PropertyCh
 	Super::PostEditChangeProperty(PropertyChangedEvent);
 }
 
-void UPoseSearchSchema::ComputeCostBreakdowns(UE::PoseSearch::ICostBreakDownData& CostBreakDownData) const
-{
-	for (const TObjectPtr<UPoseSearchFeatureChannel>& ChannelPtr : Channels)
-	{
-		if (ChannelPtr)
-		{
-			ChannelPtr->ComputeCostBreakdowns(CostBreakDownData, this);
-		}
-	}
-}
 #endif
\ No newline at end of file
diff --git a/Engine/Plugins/Experimental/Animation/PoseSearch/Source/Runtime/Public/PoseSearch/PoseSearchFeatureChannel.h b/Engine/Plugins/Experimental/Animation/PoseSearch/Source/Runtime/Public/PoseSearch/PoseSearchFeatureChannel.h
--- a/Engine/Plugins/Experimental/Animation/PoseSearch/Source/Runtime/Public/PoseSearch/PoseSearchFeatureChannel.h
+++ b/Engine/Plugins/Experimental/Animation/PoseSearch/Source/Runtime/Public/PoseSearch/PoseSearchFeatureChannel.h
@@ -62,67 +62,6 @@ public:
 	static float DecodeFloatAtOffset(TConstArrayView<float> Values, int32 DataOffset);
 };
 
-#if WITH_EDITOR
-// data structure collecting the internal layout representation of UPoseSearchFeatureChannel,
-// so we can aggregate data from different FPoseSearchIndex and calculate mean deviation with a homogeneous data set (ComputeChannelsDeviations
-struct FFeatureChannelLayoutSet
-{
-	// data structure holding DataOffset and Cardinality to find the data of the DebugName (channel breakdown / layout) in the related SearchIndexBases[SchemaIndex]
-	struct FEntry
-	{
-		FString DebugName; // for easier debugging
-		int32 SchemaIndex = -1; // index of the associated Schemas / SearchIndexBases used as input of the algorithm
-		int32 DataOffset = -1; // data offset from the base of SearchIndexBases[SchemaIndex].Values.GetData() from where the data associated to this Item starts
-		int32 Cardinality = -1; // data cardinality
-	};
-
-	// FIoHash is the hash associated to the channel data breakdown (e.g.: it could be a single SampledBones at a specific SampleTimes for a UPoseSearchFeatureChannel_Pose)
-	TMap<FIoHash, TArray<FEntry>> EntriesMap;
-	int32 CurrentSchemaIndex = -1;
-	TWeakObjectPtr<const UPoseSearchSchema> CurrentSchema;
-
-	void Add(FString DebugName, FIoHash Id, int32 DataOffset, int32 Cardinality)
-	{
-		check(DataOffset >= 0 && Cardinality >= 0 && CurrentSchemaIndex >= 0);
-		TArray<FEntry>& Entries = EntriesMap.FindOrAdd(Id);
-
-		// making sure all the FEntry associated with the same Id have the same Cardinality
-		check(Entries.IsEmpty() || Entries[0].Cardinality == Cardinality);
-		Entries.Add({ DebugName, CurrentSchemaIndex, DataOffset, Cardinality });
-	}
-};
-#endif // WITH_EDITOR
-
-class POSESEARCH_API ICostBreakDownData
-{
-public:
-	virtual ~ICostBreakDownData() {}
-
-	// returns the size of the dataset
-	virtual int32 Num() const = 0;
-
-	// returns true if Index-th cost data vector is associated with Schema
-	virtual bool IsCostVectorFromSchema(int32 Index, const UPoseSearchSchema* Schema) const = 0;
-
-	// returns the Index-th cost data vector
-	virtual TConstArrayView<float> GetCostVector(int32 Index, const UPoseSearchSchema* Schema) const = 0;
-
-	// every breakdown section start by calling BeginBreakDownSection...
-	virtual void BeginBreakDownSection(const FText& Label) = 0;
-
-	// ...then add as many SetCostBreakDown into the section...
-	virtual void SetCostBreakDown(float CostBreakDown, int32 Index, const UPoseSearchSchema* Schema) = 0;
-
-	// ...to finally wrap the section up by calling EndBreakDownSection
-	virtual void EndBreakDownSection(const FText& Label) = 0;
-
-	// true if want the channel to be verbose and generate the cost breakdown labels
-	virtual bool IsVerbose() const { return true; }
-
-	// most common case implementation
-	void AddEntireBreakDownSection(const FText& Label, const UPoseSearchSchema* Schema, int32 DataOffset, int32 Cardinality);
-};
-
 } // namespace UE::PoseSearch
 
 class POSESEARCH_API IPoseFilter
@@ -168,15 +107,20 @@ public:
 	// Draw this channel's data for the given pose vector
 	virtual void DebugDraw(const UE::PoseSearch::FDebugDrawParams& DrawParams, TConstArrayView<float> PoseVector) const PURE_VIRTUAL(UPoseSearchFeatureChannel::DebugDraw, );
 
+	// UPoseSearchFeatureChannels can hold sub channels
+	virtual TArrayView<TObjectPtr<UPoseSearchFeatureChannel>> GetSubChannels() { return TArrayView<TObjectPtr<UPoseSearchFeatureChannel>>(); }
+	virtual TConstArrayView<TObjectPtr<UPoseSearchFeatureChannel>> GetSubChannels() const { return TConstArrayView<TObjectPtr<UPoseSearchFeatureChannel>>(); }
+
 #if WITH_EDITOR
-	virtual void PopulateChannelLayoutSet(UE::PoseSearch::FFeatureChannelLayoutSet& FeatureChannelLayoutSet) const;
-	virtual void ComputeCostBreakdowns(UE::PoseSearch::ICostBreakDownData& CostBreakDownData, const UPoseSearchSchema* Schema) const;
+	// returns the FString used editor side to identify this UPoseSearchFeatureChannel (for instance in the pose search debugger)
+	virtual FString GetLabel() const;
+	const UPoseSearchSchema* GetSchema() const;
 #endif
 
 private:
 	// IBoneReferenceSkeletonProvider interface
 	// Note this function is exclusively for FBoneReference details customization
-	class USkeleton* GetSkeleton(bool& bInvalidSkeletonIsError, const IPropertyHandle* PropertyHandle) override;
+	USkeleton* GetSkeleton(bool& bInvalidSkeletonIsError, const IPropertyHandle* PropertyHandle) override;
 
 protected:
 	friend class ::UPoseSearchSchema;
diff --git a/Engine/Plugins/Experimental/Animation/PoseSearch/Source/Runtime/Public/PoseSearch/PoseSearchSchema.h b/Engine/Plugins/Experimental/Animation/PoseSearch/Source/Runtime/Public/PoseSearch/PoseSearchSchema.h
--- a/Engine/Plugins/Experimental/Animation/PoseSearch/Source/Runtime/Public/PoseSearch/PoseSearchSchema.h
+++ b/Engine/Plugins/Experimental/Animation/PoseSearch/Source/Runtime/Public/PoseSearch/PoseSearchSchema.h
@@ -15,7 +15,6 @@ class UPoseSearchFeatureChannel;
 namespace UE::PoseSearch
 {
 	struct FSearchContext;
-	class ICostBreakDownData;
 } // namespace UE::PoseSearch
 
 UENUM()
@@ -120,8 +119,6 @@ public:
 	USkeleton* GetSkeleton(bool& bInvalidSkeletonIsError, const IPropertyHandle* PropertyHandle) override;
 
 #if WITH_EDITOR
-	void ComputeCostBreakdowns(UE::PoseSearch::ICostBreakDownData& CostBreakDownData) const;
-
 	virtual void PostEditChangeProperty(struct FPropertyChangedEvent& PropertyChangedEvent) override;
 #endif
 
