motion matching - removing UPoseSearchFeatureChannel_Trajectory::PreSave Samples sorting

diff --git a/Engine/Plugins/Experimental/Animation/PoseSearch/Source/Runtime/Private/PoseSearchFeatureChannel_Trajectory.cpp b/Engine/Plugins/Experimental/Animation/PoseSearch/Source/Runtime/Private/PoseSearchFeatureChannel_Trajectory.cpp
--- a/Engine/Plugins/Experimental/Animation/PoseSearch/Source/Runtime/Private/PoseSearchFeatureChannel_Trajectory.cpp
+++ b/Engine/Plugins/Experimental/Animation/PoseSearch/Source/Runtime/Private/PoseSearchFeatureChannel_Trajectory.cpp
@@ -13,16 +13,6 @@
 
 #define LOCTEXT_NAMESPACE "PoseSearchFeatureChannels"
 
-void UPoseSearchFeatureChannel_Trajectory::PreSave(FObjectPreSaveContext ObjectSaveContext)
-{
-	Samples.Sort([](const FPoseSearchTrajectorySample& a, const FPoseSearchTrajectorySample& b)
-		{
-			return a.Offset < b.Offset;
-		});
-
-	Super::PreSave(ObjectSaveContext);
-}
-
 void UPoseSearchFeatureChannel_Trajectory::Finalize(UPoseSearchSchema* Schema)
 {
 	SubChannels.Reset();
diff --git a/Engine/Plugins/Experimental/Animation/PoseSearch/Source/Runtime/Private/PoseSearchFeatureChannel_Trajectory.h b/Engine/Plugins/Experimental/Animation/PoseSearch/Source/Runtime/Private/PoseSearchFeatureChannel_Trajectory.h
--- a/Engine/Plugins/Experimental/Animation/PoseSearch/Source/Runtime/Private/PoseSearchFeatureChannel_Trajectory.h
+++ b/Engine/Plugins/Experimental/Animation/PoseSearch/Source/Runtime/Private/PoseSearchFeatureChannel_Trajectory.h
@@ -60,9 +60,6 @@ public:
 	virtual TArrayView<TObjectPtr<UPoseSearchFeatureChannel>> GetSubChannels() override { return SubChannels; }
 	virtual TConstArrayView<TObjectPtr<UPoseSearchFeatureChannel>> GetSubChannels() const override { return SubChannels; }
 
-	// UObject interface
-	virtual void PreSave(FObjectPreSaveContext ObjectSaveContext) override;
-
 	// UPoseSearchFeatureChannel interface
 	virtual void Finalize(UPoseSearchSchema* Schema) override;
 	virtual void DebugDraw(const UE::PoseSearch::FDebugDrawParams& DrawParams, TConstArrayView<float> PoseVector) const override;
