motion matching - removing unused IAssetIndexer::FSampleInfo::RootDistance

diff --git a/Engine/Plugins/Experimental/Animation/PoseSearch/Source/Runtime/Private/PoseSearchAssetIndexer.cpp b/Engine/Plugins/Experimental/Animation/PoseSearch/Source/Runtime/Private/PoseSearchAssetIndexer.cpp
--- a/Engine/Plugins/Experimental/Animation/PoseSearch/Source/Runtime/Private/PoseSearchAssetIndexer.cpp
+++ b/Engine/Plugins/Experimental/Animation/PoseSearch/Source/Runtime/Private/PoseSearchAssetIndexer.cpp
@@ -163,11 +163,7 @@ FAssetIndexer::FSampleInfo FAssetIndexer::GetSampleInfo(float SampleTime) const
 	{
 		Sample.bClamped = true;
 		Sample.ClipTime = SamplingParam.WrappedParam + SamplingParam.Extrapolation;
-		const FTransform ClipRootMotion = Sample.Clip->ExtractRootTransform(Sample.ClipTime);
-		const float ClipDistance = Sample.Clip->ExtractRootDistance(Sample.ClipTime);
-
-		Sample.RootTransform = ClipRootMotion;
-		Sample.RootDistance = ClipDistance;
+		Sample.RootTransform = Sample.Clip->ExtractRootTransform(Sample.ClipTime);
 	}
 	else
 	{
@@ -180,37 +176,31 @@ FAssetIndexer::FSampleInfo FAssetIndexer::GetSampleInfo(float SampleTime) const
 		// had to be clamped, this motion will end up not getting applied below.
 		// Also invert the accumulation direction if the requested sample was wrapped backwards.
 		FTransform RootMotionPerCycle = RootMotionLast;
-		float RootDistancePerCycle = RootDistanceLast;
+		
 		if (SampleTime < 0.0f)
 		{
 			RootMotionPerCycle = RootMotionPerCycle.Inverse();
-			RootDistancePerCycle *= -1.f;
 		}
 
 		// Find the remaining motion deltas after wrapping
 		FTransform RootMotionRemainder = Sample.Clip->ExtractRootTransform(Sample.ClipTime);
-		float RootDistanceRemainder = Sample.Clip->ExtractRootDistance(Sample.ClipTime);
 
 		// Invert motion deltas if we wrapped backwards
 		if (SampleTime < 0.0f)
 		{
 			RootMotionRemainder.SetToRelativeTransform(RootMotionLast);
-			RootDistanceRemainder = -(RootDistanceLast - RootDistanceRemainder);
 		}
 
 		Sample.RootTransform = FTransform::Identity;
-		Sample.RootDistance = 0.f;
 
 		// Note if the sample was clamped, no motion will be applied here because NumCycles will be zero
 		int32 CyclesRemaining = SamplingParam.NumCycles;
 		while (CyclesRemaining--)
 		{
 			Sample.RootTransform = RootMotionPerCycle * Sample.RootTransform;
-			Sample.RootDistance += RootDistancePerCycle;
 		}
 
 		Sample.RootTransform = RootMotionRemainder * Sample.RootTransform;
-		Sample.RootDistance += RootDistanceRemainder;
 	}
 
 	return Sample;
@@ -220,7 +210,6 @@ FAssetIndexer::FSampleInfo FAssetIndexer::GetSampleInfoRelative(float SampleTime
 {
 	FSampleInfo Sample = GetSampleInfo(SampleTime);
 	Sample.RootTransform.SetToRelativeTransform(Origin.RootTransform);
-	Sample.RootDistance = Origin.RootDistance - Sample.RootDistance;
 	return Sample;
 }
 
diff --git a/Engine/Plugins/Experimental/Animation/PoseSearch/Source/Runtime/Public/PoseSearch/PoseSearchAssetIndexer.h b/Engine/Plugins/Experimental/Animation/PoseSearch/Source/Runtime/Public/PoseSearch/PoseSearchAssetIndexer.h
--- a/Engine/Plugins/Experimental/Animation/PoseSearch/Source/Runtime/Public/PoseSearch/PoseSearchAssetIndexer.h
+++ b/Engine/Plugins/Experimental/Animation/PoseSearch/Source/Runtime/Public/PoseSearch/PoseSearchAssetIndexer.h
@@ -41,7 +41,6 @@ public:
 		const IAssetSampler* Clip = nullptr;
 		FTransform RootTransform;
 		float ClipTime = 0.0f;
-		float RootDistance = 0.0f;
 		bool bClamped = false;
 
 		bool IsValid() const { return Clip != nullptr; }
