motion matching - fix for ObjectSerializationError in FArchive& operator<<(FArchive& Ar, FPoseSearchIndex& Index)

diff --git a/Engine/Plugins/Experimental/Animation/PoseSearch/Source/Runtime/Private/PoseSearchIndex.cpp b/Engine/Plugins/Experimental/Animation/PoseSearch/Source/Runtime/Private/PoseSearchIndex.cpp
--- a/Engine/Plugins/Experimental/Animation/PoseSearch/Source/Runtime/Private/PoseSearchIndex.cpp
+++ b/Engine/Plugins/Experimental/Animation/PoseSearch/Source/Runtime/Private/PoseSearchIndex.cpp
@@ -245,10 +245,8 @@ FArchive& operator<<(FArchive& Ar, FPoseSearchIndex& Index)
 
 	Serialize(Ar, Index.KDTree, Index.PCAValues.GetData());
 
-#if WITH_EDITORONLY_DATA
 	Ar << Index.PCAExplainedVariance;
 	Ar << Index.Deviation;
-#endif // WITH_EDITORONLY_DATA
 
 	return Ar;
 }
\ No newline at end of file
diff --git a/Engine/Plugins/Experimental/Animation/PoseSearch/Source/Runtime/Public/PoseSearch/PoseSearchIndex.h b/Engine/Plugins/Experimental/Animation/PoseSearch/Source/Runtime/Public/PoseSearch/PoseSearchIndex.h
--- a/Engine/Plugins/Experimental/Animation/PoseSearch/Source/Runtime/Public/PoseSearch/PoseSearchIndex.h
+++ b/Engine/Plugins/Experimental/Animation/PoseSearch/Source/Runtime/Public/PoseSearch/PoseSearchIndex.h
@@ -193,13 +193,11 @@ struct POSESEARCH_API FPoseSearchIndex : public FPoseSearchIndexBase
 
 	UE::PoseSearch::FKDTree KDTree;
 
-#if WITH_EDITORONLY_DATA
 	UPROPERTY(Category = Info, VisibleAnywhere, meta = (NeverInHash))
 	float PCAExplainedVariance = 0.f;
 
 	UPROPERTY(Category = Info, VisibleAnywhere, meta = (NeverInHash))
 	TArray<float> Deviation;
-#endif // WITH_EDITORONLY_DATA
 
 	FPoseSearchIndex() = default;
 	~FPoseSearchIndex() = default;
