motion matching - FKeyBuilder::AddNameOnly improvements

diff --git a/Engine/Plugins/Experimental/Animation/PoseSearch/Source/Runtime/Private/PoseSearchDerivedDataKey.cpp b/Engine/Plugins/Experimental/Animation/PoseSearch/Source/Runtime/Private/PoseSearchDerivedDataKey.cpp
--- a/Engine/Plugins/Experimental/Animation/PoseSearch/Source/Runtime/Private/PoseSearchDerivedDataKey.cpp
+++ b/Engine/Plugins/Experimental/Animation/PoseSearch/Source/Runtime/Private/PoseSearchDerivedDataKey.cpp
@@ -188,22 +188,21 @@ const TSet<const UObject*>& FKeyBuilder::GetDependencies() const
 
 bool FKeyBuilder::AddNameOnly(class UObject* Object) const
 {
+	check(Object);
+	const UClass* ObjectClass = Object->GetClass();
+	check(ObjectClass);
+
 	return
-		Cast<UAnimBoneCompressionSettings>(Object) ||
-		Cast<UAnimCurveCompressionSettings>(Object) ||
-		Cast<UFbxAnimSequenceImportData>(Object) ||
-		Cast<UMirrorDataTable>(Object) ||
-		Cast<USkeleton>(Object) ||
-		Cast<USkeletalMeshComponent>(Object) ||
-		Cast<USkeletalMeshSocket>(Object) ||
-		Cast<USkinnedMeshComponent>(Object) ||
-		Cast<UMeshComponent>(Object) ||
-		Cast<UPrimitiveComponent>(Object) ||
-		Cast<USceneComponent>(Object) ||
-		Cast<UActorComponent>(Object) ||
-		Cast<UStreamableRenderAsset>(Object) ||
-		Cast<USkinnedAsset>(Object) ||
-		Cast<USkeletalMesh>(Object);
+		ObjectClass->IsChildOf(UActorComponent::StaticClass()) ||
+		ObjectClass->IsChildOf(UAnimBoneCompressionSettings::StaticClass()) ||
+		ObjectClass->IsChildOf(UAnimCurveCompressionSettings::StaticClass()) ||
+		ObjectClass->IsChildOf(UAssetImportData::StaticClass()) ||
+		ObjectClass->IsChildOf(UFunction::StaticClass()) ||
+		ObjectClass->IsChildOf(UMirrorDataTable::StaticClass()) ||
+		ObjectClass->IsChildOf(USkeletalMesh::StaticClass()) ||
+		ObjectClass->IsChildOf(USkeletalMeshSocket::StaticClass()) ||
+		ObjectClass->IsChildOf(USkeleton::StaticClass()) ||
+		ObjectClass->IsChildOf(UStreamableRenderAsset::StaticClass());
 }
 
 #if UE_POSE_SEARCH_DERIVED_DATA_LOGGING
