motion matching - position / velocity / heading channels support for invalid pose history

diff --git a/Engine/Plugins/Experimental/Animation/PoseSearch/Source/Runtime/Private/PoseSearchFeatureChannel_Heading.cpp b/Engine/Plugins/Experimental/Animation/PoseSearch/Source/Runtime/Private/PoseSearchFeatureChannel_Heading.cpp
--- a/Engine/Plugins/Experimental/Animation/PoseSearch/Source/Runtime/Private/PoseSearchFeatureChannel_Heading.cpp
+++ b/Engine/Plugins/Experimental/Animation/PoseSearch/Source/Runtime/Private/PoseSearchFeatureChannel_Heading.cpp
@@ -70,12 +70,14 @@ void UPoseSearchFeatureChannel_Heading::BuildQuery(UE::PoseSearch::FSearchContex
 
 	const bool bIsCurrentResultValid = SearchContext.CurrentResult.IsValid();
 	const bool bSkip = InputQueryPose != EInputQueryPose::UseCharacterPose && bIsCurrentResultValid && SearchContext.CurrentResult.Database->Schema == InOutQuery.GetSchema();
-	if (bSkip || !SearchContext.History)
+	const bool bBoneValid = InOutQuery.GetSchema()->BoneReferences[SchemaBoneIdx].HasValidSetup();
+	if (bSkip || (!SearchContext.History && bBoneValid))
 	{
 		if (bIsCurrentResultValid)
 		{
 			const float LerpValue = InputQueryPose == EInputQueryPose::UseInterpolatedContinuingPose ? SearchContext.CurrentResult.LerpValue : 0.f;
 			int32 DataOffset = ChannelDataOffset;
+			// @todo: we should normalize if LerpValue != 0
 			FFeatureVectorHelper::EncodeVector(InOutQuery.EditValues(), DataOffset, SearchContext.GetCurrentResultPrevPoseVector(), SearchContext.GetCurrentResultPoseVector(), SearchContext.GetCurrentResultNextPoseVector(), LerpValue, true);
 		}
 		// else leave the InOutQuery set to zero since the SearchContext.History is invalid and it'll fail if we continue
@@ -83,7 +85,7 @@ void UPoseSearchFeatureChannel_Heading::BuildQuery(UE::PoseSearch::FSearchContex
 	else
 	{
 		FTransform Transform;
-		if (InOutQuery.GetSchema()->BoneReferences[SchemaBoneIdx].HasValidSetup())
+		if (bBoneValid)
 		{
 			// calculating the Transform in component space for the bone indexed by SchemaBoneIdx
 			Transform = SearchContext.TryGetTransformAndCacheResults(SampleTimeOffset, InOutQuery.GetSchema(), SchemaBoneIdx);
diff --git a/Engine/Plugins/Experimental/Animation/PoseSearch/Source/Runtime/Private/PoseSearchFeatureChannel_Position.cpp b/Engine/Plugins/Experimental/Animation/PoseSearch/Source/Runtime/Private/PoseSearchFeatureChannel_Position.cpp
--- a/Engine/Plugins/Experimental/Animation/PoseSearch/Source/Runtime/Private/PoseSearchFeatureChannel_Position.cpp
+++ b/Engine/Plugins/Experimental/Animation/PoseSearch/Source/Runtime/Private/PoseSearchFeatureChannel_Position.cpp
@@ -55,7 +55,8 @@ void UPoseSearchFeatureChannel_Position::BuildQuery(UE::PoseSearch::FSearchConte
 	check(InOutQuery.GetSchema());
 	const bool bIsCurrentResultValid = SearchContext.CurrentResult.IsValid();
 	const bool bSkip = InputQueryPose != EInputQueryPose::UseCharacterPose && bIsCurrentResultValid && SearchContext.CurrentResult.Database->Schema == InOutQuery.GetSchema();
-	if (bSkip || !SearchContext.History)
+	const bool bBoneValid = InOutQuery.GetSchema()->BoneReferences[SchemaBoneIdx].HasValidSetup();
+	if (bSkip || (!SearchContext.History && bBoneValid))
 	{
 		if (bIsCurrentResultValid)
 		{
@@ -68,7 +69,7 @@ void UPoseSearchFeatureChannel_Position::BuildQuery(UE::PoseSearch::FSearchConte
 	else
 	{
 		FTransform Transform;
-		if (InOutQuery.GetSchema()->BoneReferences[SchemaBoneIdx].HasValidSetup())
+		if (bBoneValid)
 		{
 			// calculating the Transform in component space for the bone indexed by SchemaBoneIdx
 			Transform = SearchContext.TryGetTransformAndCacheResults(SampleTimeOffset, InOutQuery.GetSchema(), SchemaBoneIdx);
diff --git a/Engine/Plugins/Experimental/Animation/PoseSearch/Source/Runtime/Private/PoseSearchFeatureChannel_Velocity.cpp b/Engine/Plugins/Experimental/Animation/PoseSearch/Source/Runtime/Private/PoseSearchFeatureChannel_Velocity.cpp
--- a/Engine/Plugins/Experimental/Animation/PoseSearch/Source/Runtime/Private/PoseSearchFeatureChannel_Velocity.cpp
+++ b/Engine/Plugins/Experimental/Animation/PoseSearch/Source/Runtime/Private/PoseSearchFeatureChannel_Velocity.cpp
@@ -83,12 +83,14 @@ void UPoseSearchFeatureChannel_Velocity::BuildQuery(UE::PoseSearch::FSearchConte
 
 	const bool bIsCurrentResultValid = SearchContext.CurrentResult.IsValid();
 	const bool bSkip = InputQueryPose != EInputQueryPose::UseCharacterPose && bIsCurrentResultValid && SearchContext.CurrentResult.Database->Schema == InOutQuery.GetSchema();
-	if (bSkip || !SearchContext.History)
+	const bool bBoneValid = InOutQuery.GetSchema()->BoneReferences[SchemaBoneIdx].HasValidSetup();
+	if (bSkip || (!SearchContext.History && bBoneValid))
 	{
 		if (bIsCurrentResultValid)
 		{
 			const float LerpValue = InputQueryPose == EInputQueryPose::UseInterpolatedContinuingPose ? SearchContext.CurrentResult.LerpValue : 0.f;
 			int32 DataOffset = ChannelDataOffset;
+			// @todo: we should normalize if bNormalize && LerpValue != 0
 			FFeatureVectorHelper::EncodeVector(InOutQuery.EditValues(), DataOffset, SearchContext.GetCurrentResultPrevPoseVector(), SearchContext.GetCurrentResultPoseVector(), SearchContext.GetCurrentResultNextPoseVector(), LerpValue);
 		}
 		// else leave the InOutQuery set to zero since the SearchContext.History is invalid and it'll fail if we continue
@@ -96,7 +98,7 @@ void UPoseSearchFeatureChannel_Velocity::BuildQuery(UE::PoseSearch::FSearchConte
 	else
 	{
 		FVector LinearVelocity;
-		if (InOutQuery.GetSchema()->BoneReferences[SchemaBoneIdx].HasValidSetup())
+		if (bBoneValid)
 		{
 			check(SearchContext.History);
 			const float HistorySameplInterval = SearchContext.History->GetSampleTimeInterval();
