motion matching - listening to OnPackageReloaded to clear up DDC FPoseSearchDatabaseAsyncCacheTask

diff --git a/Engine/Plugins/Experimental/Animation/PoseSearch/Source/Runtime/Private/PoseSearchDerivedData.cpp b/Engine/Plugins/Experimental/Animation/PoseSearch/Source/Runtime/Private/PoseSearchDerivedData.cpp
--- a/Engine/Plugins/Experimental/Animation/PoseSearch/Source/Runtime/Private/PoseSearchDerivedData.cpp
+++ b/Engine/Plugins/Experimental/Animation/PoseSearch/Source/Runtime/Private/PoseSearchDerivedData.cpp
@@ -25,6 +25,7 @@
 #include "ProfilingDebugging/CookStats.h"
 #include "Serialization/BulkDataRegistry.h"
 #include "UObject/NoExportTypes.h"
+#include "UObject/PackageReload.h"
 
 namespace UE::PoseSearch
 {
@@ -599,6 +600,8 @@ struct FPoseSearchDatabaseAsyncCacheTask
 	bool Poll(FCriticalSection& OuterMutex) const;
 	bool ContainsDatabase(const UPoseSearchDatabase* OtherDatabase, FCriticalSection& OuterMutex) const;
 	bool IsValid(FCriticalSection& OuterMutex) const;
+	const FIoHash& GetDerivedDataKey() const { return DerivedDataKey; }
+	const UPoseSearchDatabase* GetDatabase() const { return Database.Get(); }
 
 	~FPoseSearchDatabaseAsyncCacheTask();
 	EState GetState() const { return EState(ThreadSafeState.GetValue()); }
@@ -694,13 +697,17 @@ void FPoseSearchDatabaseAsyncCacheTask::Cancel(FCriticalSection& OuterMutex)
 
 bool FPoseSearchDatabaseAsyncCacheTask::CancelIfDependsOn(const UObject* Object, FCriticalSection& OuterMutex)
 {
-	FScopeLock Lock(&OuterMutex);
-
-	// DatabaseDependencies is updated only in StartNewRequestIfNeeded when there are no active requests, so it's thread safe to access it 
-	if (DatabaseDependencies.Contains(Object))
+	if (Object)
 	{
-		Cancel(OuterMutex);
-		return true;
+		FScopeLock Lock(&OuterMutex);
+
+		// DatabaseDependencies is updated only in StartNewRequestIfNeeded when there are no active requests, so it's thread safe to access it 
+		if (DatabaseDependencies.Contains(Object))
+		{
+			UE_LOG(LogPoseSearch, Log, TEXT("%s - %s Cancelled because of %s"), *LexToString(DerivedDataKey), *Database->GetName(), *Object->GetName());
+			Cancel(OuterMutex);
+			return true;
+		}
 	}
 	return false;
 }
@@ -996,6 +1003,8 @@ FAsyncPoseSearchDatabasesManagement::FAsyncPoseSearchDatabasesManagement()
 	FScopeLock Lock(&Mutex);
 
 	OnObjectModifiedHandle = FCoreUObjectDelegates::OnObjectModified.AddRaw(this, &FAsyncPoseSearchDatabasesManagement::OnObjectModified);
+	OnPackageReloadedHandle = FCoreUObjectDelegates::OnPackageReloaded.AddRaw(this, &FAsyncPoseSearchDatabasesManagement::OnPackageReloaded);
+
 	FCoreDelegates::OnPreExit.AddRaw(this, &FAsyncPoseSearchDatabasesManagement::Shutdown);
 }
 
@@ -1024,12 +1033,36 @@ void FAsyncPoseSearchDatabasesManagement::OnObjectModified(UObject* Object)
 	}
 }
 
+void FAsyncPoseSearchDatabasesManagement::OnPackageReloaded(const EPackageReloadPhase InPackageReloadPhase, FPackageReloadedEvent* InPackageReloadedEvent)
+{
+	if (InPackageReloadPhase == EPackageReloadPhase::PostPackageFixup && InPackageReloadedEvent)
+	{
+		FScopeLock Lock(&Mutex);
+
+		// @todo: figure out why we don't find the correct dependency into InPackageReloadedEvent->GetRepointedObjects()
+		//		  for now we invalidate all the DDC cache to be on the safe side
+		//for (const TPair<UObject*, UObject*>& Pair : InPackageReloadedEvent->GetRepointedObjects())
+		//{
+		//	OnObjectModified(Pair.Key);
+		//}
+
+		for (TUniquePtr<FPoseSearchDatabaseAsyncCacheTask>& TaskPtr : Tasks)
+		{
+			UE_LOG(LogPoseSearch, Log, TEXT("%s - %s Cancelled because of OnPackageReloaded"), *LexToString(TaskPtr->GetDerivedDataKey()), *TaskPtr->GetDatabase()->GetName());
+		}
+		Tasks.Reset();
+	}
+}
+
 void FAsyncPoseSearchDatabasesManagement::Shutdown()
 {
 	FScopeLock Lock(&Mutex);
 
 	FCoreUObjectDelegates::OnObjectModified.Remove(OnObjectModifiedHandle);
 	OnObjectModifiedHandle.Reset();
+
+	FCoreUObjectDelegates::OnPackageReloaded.Remove(OnPackageReloadedHandle);
+	OnPackageReloadedHandle.Reset();
 }
 
 void FAsyncPoseSearchDatabasesManagement::Tick(float DeltaTime)
diff --git a/Engine/Plugins/Experimental/Animation/PoseSearch/Source/Runtime/Public/PoseSearch/PoseSearchDerivedData.h b/Engine/Plugins/Experimental/Animation/PoseSearch/Source/Runtime/Public/PoseSearch/PoseSearchDerivedData.h
--- a/Engine/Plugins/Experimental/Animation/PoseSearch/Source/Runtime/Public/PoseSearch/PoseSearchDerivedData.h
+++ b/Engine/Plugins/Experimental/Animation/PoseSearch/Source/Runtime/Public/PoseSearch/PoseSearchDerivedData.h
@@ -35,6 +35,7 @@ namespace UE::PoseSearch
 		static FAsyncPoseSearchDatabasesManagement& Get();
 
 		void OnObjectModified(UObject* Object);
+		void OnPackageReloaded(const EPackageReloadPhase InPackageReloadPhase, FPackageReloadedEvent* InPackageReloadedEvent);
 
 		void Shutdown();
 		void StartQueuedTasks(int32 MaxActiveTasks);
@@ -56,6 +57,7 @@ namespace UE::PoseSearch
 		
 		FPoseSearchDatabaseAsyncCacheTasks& Tasks;
 		FDelegateHandle OnObjectModifiedHandle;
+		FDelegateHandle OnPackageReloadedHandle;
 		
 		static FCriticalSection Mutex;
 	};
