motion matching - fix for SDebuggerView::GetSelectedDatabaseRows trying to draw databases requiring indexing

diff --git a/Engine/Plugins/Experimental/Animation/PoseSearch/Source/Editor/Private/PoseSearchDebuggerView.cpp b/Engine/Plugins/Experimental/Animation/PoseSearch/Source/Editor/Private/PoseSearchDebuggerView.cpp
--- a/Engine/Plugins/Experimental/Animation/PoseSearch/Source/Editor/Private/PoseSearchDebuggerView.cpp
+++ b/Engine/Plugins/Experimental/Animation/PoseSearch/Source/Editor/Private/PoseSearchDebuggerView.cpp
@@ -378,18 +378,19 @@ void SDebuggerView::DrawFeatures(
 		const UPoseSearchDatabase* CurrentDatabase = ViewModel.Get()->GetCurrentDatabase();
 		if (CurrentDatabase)
 		{
+			// Set shared state
+			FDebugDrawParams DrawParams;
+			DrawParams.World = &DebuggerWorld;
+			DrawParams.RootTransform = Transform;
+			DrawParams.DefaultLifeTime = 0.0f; // Single frame render
+			DrawParams.Mesh = Mesh;
+
 			for (const FTraceMotionMatchingStateDatabaseEntry& DbEntry : State.DatabaseEntries)
 			{
 				const UPoseSearchDatabase* Database = FTraceMotionMatchingState::GetObjectFromId<UPoseSearchDatabase>(DbEntry.DatabaseId);
 				if (Database && Database == CurrentDatabase && DbEntry.QueryVector.Num() == Database->Schema->SchemaCardinality && 
 					FAsyncPoseSearchDatabasesManagement::RequestAsyncBuildIndex(CurrentDatabase, ERequestAsyncBuildFlag::ContinueRequest))
 				{
-					// Set shared state
-					FDebugDrawParams DrawParams;
-					DrawParams.World = &DebuggerWorld;
-					DrawParams.RootTransform = Transform;
-					DrawParams.DefaultLifeTime = 0.0f; // Single frame render
-					DrawParams.Mesh = Mesh;
 					DrawParams.Database = CurrentDatabase;
 					SetDrawFlags(DrawParams, Reflection ? Reflection->QueryDrawOptions : FPoseSearchDebuggerFeatureDrawOptions());
 					EnumAddFlags(DrawParams.Flags, EDebugDrawFlags::DrawQuery);
@@ -418,7 +419,10 @@ void SDebuggerView::DrawFeatures(
 		for (const TSharedRef<FDebuggerDatabaseRowData>& Row : SelectedRows)
 		{
 			DrawParams.Database = Row->SourceDatabase.Get();
-			DrawFeatureVector(DrawParams, Row->PoseIdx);
+			if (DrawParams.Database && FAsyncPoseSearchDatabasesManagement::RequestAsyncBuildIndex(DrawParams.Database, ERequestAsyncBuildFlag::ContinueRequest))
+			{
+				DrawFeatureVector(DrawParams, Row->PoseIdx);
+			}
 		}
 	}
 
@@ -431,15 +435,19 @@ void SDebuggerView::DrawFeatures(
 		check(ActiveRows.Num() < 2);
 
 		if (!ActiveRows.IsEmpty())
-		{		
-			// Use the motion-matching state's pose idx, as the active row may be update-throttled at this point
-			FDebugDrawParams DrawParams;
-			DrawParams.World = &DebuggerWorld;
-			DrawParams.RootTransform = Transform;
-			DrawParams.DefaultLifeTime = 0.0f; // Single frame render
-			DrawParams.Mesh = Mesh;
-			DrawParams.Database = ActiveRows[0]->SourceDatabase.Get();
-			DrawFeatureVector(DrawParams, ActiveRows[0]->PoseIdx);
+		{
+			const UPoseSearchDatabase* Database = ActiveRows[0]->SourceDatabase.Get();
+			if (Database && FAsyncPoseSearchDatabasesManagement::RequestAsyncBuildIndex(Database, ERequestAsyncBuildFlag::ContinueRequest))
+			{
+				// Use the motion-matching state's pose idx, as the active row may be update-throttled at this point
+				FDebugDrawParams DrawParams;
+				DrawParams.World = &DebuggerWorld;
+				DrawParams.RootTransform = Transform;
+				DrawParams.DefaultLifeTime = 0.0f; // Single frame render
+				DrawParams.Mesh = Mesh;
+				DrawParams.Database = Database;
+				DrawFeatureVector(DrawParams, ActiveRows[0]->PoseIdx);
+			}
 		}
 	}
 
@@ -453,13 +461,17 @@ void SDebuggerView::DrawFeatures(
 
 		if (!ContinuingRows.IsEmpty())
 		{
-			FDebugDrawParams DrawParams;
-			DrawParams.World = &DebuggerWorld;
-			DrawParams.RootTransform = Transform;
-			DrawParams.DefaultLifeTime = 0.0f; // Single frame render
-			DrawParams.Mesh = Mesh;
-			DrawParams.Database = ContinuingRows[0]->SourceDatabase.Get();
-			DrawFeatureVector(DrawParams, ContinuingRows[0]->PoseIdx);
+			const UPoseSearchDatabase* Database = ContinuingRows[0]->SourceDatabase.Get();
+			if (Database && FAsyncPoseSearchDatabasesManagement::RequestAsyncBuildIndex(Database, ERequestAsyncBuildFlag::ContinueRequest))
+			{
+				FDebugDrawParams DrawParams;
+				DrawParams.World = &DebuggerWorld;
+				DrawParams.RootTransform = Transform;
+				DrawParams.DefaultLifeTime = 0.0f; // Single frame render
+				DrawParams.Mesh = Mesh;
+				DrawParams.Database = Database;
+				DrawFeatureVector(DrawParams, ContinuingRows[0]->PoseIdx);
+			}
 		}
 	}
 
