motion matching - delaying DDC requests after all dependent assets are post loaded

diff --git a/Engine/Plugins/Experimental/Animation/PoseSearch/Source/Runtime/Private/PoseSearchDerivedData.cpp b/Engine/Plugins/Experimental/Animation/PoseSearch/Source/Runtime/Private/PoseSearchDerivedData.cpp
--- a/Engine/Plugins/Experimental/Animation/PoseSearch/Source/Runtime/Private/PoseSearchDerivedData.cpp
+++ b/Engine/Plugins/Experimental/Animation/PoseSearch/Source/Runtime/Private/PoseSearchDerivedData.cpp
@@ -584,10 +584,11 @@ struct FPoseSearchDatabaseAsyncCacheTask
 {
 	enum class EState
 	{
-		Prestarted,
-		Cancelled,
-		Ended,
-		Failed
+		Notstarted,	// key generation failed (not all the asset has been post loaded). It'll be retried to StartNewRequestIfNeeded the next Update
+		Prestarted,	// key has been successfully generated and we kicked the DDC get
+		Cancelled,	// the task has been cancelled
+		Ended,		// the task has ended successfully
+		Failed		// the task has ended unsuccessfully
 	};
 
 	// these methods MUST be protected by FPoseSearchDatabaseAsyncCacheTask::Mutex! and to make sure we pass the mutex as input param
@@ -622,7 +623,7 @@ private:
 	FIoHash DerivedDataKey = FIoHash::Zero;
 	TSet<TWeakObjectPtr<const UObject>> DatabaseDependencies; // @todo: make this const
 		
-	FThreadSafeCounter ThreadSafeState = int32(EState::Prestarted);
+	FThreadSafeCounter ThreadSafeState = int32(EState::Notstarted);
 	bool bBroadcastOnDerivedDataRebuild = false;
 };
 
@@ -656,31 +657,41 @@ void FPoseSearchDatabaseAsyncCacheTask::StartNewRequestIfNeeded(FCriticalSection
 
 	// composing the key
 	const FKeyBuilder KeyBuilder(Database.Get(), true);
-	const FIoHash NewDerivedDataKey(KeyBuilder.Finalize());
-	const bool bHasKeyChanged = NewDerivedDataKey != DerivedDataKey;
-	if (bHasKeyChanged)
+	if (KeyBuilder.AnyAssetNotReady())
 	{
-		DerivedDataKey = NewDerivedDataKey;
-
-		DatabaseDependencies.Reset();
-		for (const UObject* Dependency : KeyBuilder.GetDependencies())
+		DerivedDataKey = FIoHash::Zero;
+		SetState(EState::Notstarted);
+	
+		UE_LOG(LogPoseSearch, Log, TEXT("Delaying DDC until dependents post load  - %s"), *LexToString(DerivedDataKey), *Database->GetName());
+	}
+	else
+	{
+		const FIoHash NewDerivedDataKey(KeyBuilder.Finalize());
+		const bool bHasKeyChanged = NewDerivedDataKey != DerivedDataKey;
+		if (bHasKeyChanged)
 		{
-			DatabaseDependencies.Add(Dependency);
-		}
+			DerivedDataKey = NewDerivedDataKey;
 
-		SetState(EState::Prestarted);
+			DatabaseDependencies.Reset();
+			for (const UObject* Dependency : KeyBuilder.GetDependencies())
+			{
+				DatabaseDependencies.Add(Dependency);
+			}
 
-		UE_LOG(LogPoseSearch, Log, TEXT("%s - %s BeginCache"), *LexToString(DerivedDataKey), *Database->GetName());
+			SetState(EState::Prestarted);
 
-		TArray<FCacheGetRequest> CacheRequests;
-		const FCacheKey CacheKey{ Bucket, DerivedDataKey };
-		CacheRequests.Add({ { Database->GetPathName() }, CacheKey, ECachePolicy::Default });
+			UE_LOG(LogPoseSearch, Log, TEXT("%s - %s BeginCache"), *LexToString(DerivedDataKey), *Database->GetName());
 
-		Owner = FRequestOwner(EPriority::Normal);
-		GetCache().Get(CacheRequests, Owner, [this](FCacheGetResponse&& Response)
-			{
-				OnGetComplete(MoveTemp(Response));
-			});
+			TArray<FCacheGetRequest> CacheRequests;
+			const FCacheKey CacheKey{ Bucket, DerivedDataKey };
+			CacheRequests.Add({ { Database->GetPathName() }, CacheKey, ECachePolicy::Default });
+
+			Owner = FRequestOwner(EPriority::Normal);
+			GetCache().Get(CacheRequests, Owner, [this](FCacheGetResponse&& Response)
+				{
+					OnGetComplete(MoveTemp(Response));
+				});
+		}
 	}
 }
 
@@ -720,6 +731,11 @@ void FPoseSearchDatabaseAsyncCacheTask::Update(FCriticalSection& OuterMutex)
 
 	check(GetState() != EState::Cancelled); // otherwise FPoseSearchDatabaseAsyncCacheTask should have been already removed
 
+	if (GetState() == EState::Notstarted)
+	{
+		StartNewRequestIfNeeded(OuterMutex);
+	}
+
 	if (GetState() == EState::Prestarted && Poll(OuterMutex))
 	{
 		// task is done: we need to update the state form Prestarted to Ended/Failed
@@ -1129,28 +1145,10 @@ bool FAsyncPoseSearchDatabasesManagement::RequestAsyncBuildIndex(const UPoseSear
 			{
 				if (Task->GetState() == FPoseSearchDatabaseAsyncCacheTask::EState::Prestarted)
 				{
-					if (EnumHasAnyFlags(Flag, ERequestAsyncBuildFlag::WaitPreviousRequest))
-					{
-						Task->Wait(Mutex);
-					}
-					else
-					{
-						Task->Cancel(Mutex);
-					}
+					Task->Cancel(Mutex);
 				}
-
 				Task->StartNewRequestIfNeeded(Mutex);
 			}
-			else // if (EnumHasAnyFlags(Flag, ERequestAsyncBuildFlag::ContinueRequest))
-			{
-				if (Task->GetState() == FPoseSearchDatabaseAsyncCacheTask::EState::Prestarted)
-				{
-					if (EnumHasAnyFlags(Flag, ERequestAsyncBuildFlag::WaitPreviousRequest))
-					{
-						Task->Wait(Mutex);
-					}
-				}
-			}
 			break;
 		}
 	}
@@ -1162,9 +1160,13 @@ bool FAsyncPoseSearchDatabasesManagement::RequestAsyncBuildIndex(const UPoseSear
 		Task = This.Tasks.Last().Get();
 	}
 
-	if (EnumHasAnyFlags(Flag, ERequestAsyncBuildFlag::WaitForCompletion) && Task->GetState() == FPoseSearchDatabaseAsyncCacheTask::EState::Prestarted)
+	if (EnumHasAnyFlags(Flag, ERequestAsyncBuildFlag::WaitForCompletion))
 	{
-		Task->Wait(Mutex);
+		check(Task->GetState() != FPoseSearchDatabaseAsyncCacheTask::EState::Notstarted);
+		if (Task->GetState() == FPoseSearchDatabaseAsyncCacheTask::EState::Prestarted)
+		{
+			Task->Wait(Mutex);
+		}
 	}
 
 	return Task->GetState() == FPoseSearchDatabaseAsyncCacheTask::EState::Ended;
diff --git a/Engine/Plugins/Experimental/Animation/PoseSearch/Source/Runtime/Private/PoseSearchDerivedDataKey.cpp b/Engine/Plugins/Experimental/Animation/PoseSearch/Source/Runtime/Private/PoseSearchDerivedDataKey.cpp
--- a/Engine/Plugins/Experimental/Animation/PoseSearch/Source/Runtime/Private/PoseSearchDerivedDataKey.cpp
+++ b/Engine/Plugins/Experimental/Animation/PoseSearch/Source/Runtime/Private/PoseSearchDerivedDataKey.cpp
@@ -5,6 +5,7 @@
 #include "PoseSearch/PoseSearchDerivedDataKey.h"
 #include "Animation/AnimBoneCompressionSettings.h"
 #include "Animation/AnimCurveCompressionSettings.h"
+#include "Animation/AnimData/IAnimationDataModel.h"
 #include "Animation/MirrorDataTable.h"
 #include "Animation/Skeleton.h"
 #include "Engine/SkeletalMesh.h"
@@ -21,6 +22,7 @@ namespace UE::PoseSearch
 
 FKeyBuilder::FKeyBuilder()
 {
+	bAnyAssetNotReady = false;
 	ArIgnoreOuterRef = true;
 
 	// Set FDerivedDataKeyBuilder to be a saving archive instead of a reference collector.
@@ -34,6 +36,7 @@ FKeyBuilder::FKeyBuilder(const UObject* Object, bool bUseDataVer)
 : FKeyBuilder()
 {
 	check(Object);
+	bAnyAssetNotReady = false;
 
 	if (bUseDataVer)
 	{
@@ -65,21 +68,25 @@ bool FKeyBuilder::ShouldSkipProperty(const FProperty* InProperty) const
 
 	if (Super::ShouldSkipProperty(InProperty))
 	{
-		return true;
-	}
-
-	if (!InProperty->HasAllPropertyFlags(CPF_Edit)) // bIsEditAnywhereProperty
-	{
+		#if UE_POSE_SEARCH_DERIVED_DATA_LOGGING
+		UE_LOG(LogPoseSearch, Log, TEXT("%s x %s (ShouldSkipProperty)"), *GetIndentation(), *InProperty->GetFullName());
+		#endif
 		return true;
 	}
 
 	if (InProperty->HasAllPropertyFlags(CPF_Transient))
 	{
+		#if UE_POSE_SEARCH_DERIVED_DATA_LOGGING
+		UE_LOG(LogPoseSearch, Log, TEXT("%s x %s (Transient)"), *GetIndentation(), *InProperty->GetFullName());
+		#endif
 		return true;
 	}
 
 	if (InProperty->HasMetaData(ExcludeFromHashName))
 	{
+		#if UE_POSE_SEARCH_DERIVED_DATA_LOGGING
+		UE_LOG(LogPoseSearch, Log, TEXT("%s x %s (ExcludeFromHash)"), *GetIndentation(), *InProperty->GetFullName());
+		#endif
 		return true;
 	}
 		
@@ -120,51 +127,67 @@ FArchive& FKeyBuilder::operator<<(class UObject*& Object)
 {
 	if (Object)
 	{
-		#if UE_POSE_SEARCH_DERIVED_DATA_LOGGING
-		++Indentation;
-		#endif
-
-		bool bAlreadyProcessed = false;
-		ObjectsAlreadySerialized.Add(Object, &bAlreadyProcessed);
-
-		// If we haven't already serialized this object
-		if (bAlreadyProcessed)
-		{
-			#if UE_POSE_SEARCH_DERIVED_DATA_LOGGING
-			UE_LOG(LogPoseSearch, Log, TEXT("%sAlreadyProcessed '%s' (%s)"), *GetIndentation(), *Object->GetName(), *Object->GetClass()->GetName());
-			#endif
-		}
-		// for specific types we only add their names to the hash
-		else if (AddNameOnly(Object))
+		if (Object->HasAnyFlags(RF_NeedPostLoad))
 		{
-			#if UE_POSE_SEARCH_DERIVED_DATA_LOGGING
-			UE_LOG(LogPoseSearch, Log, TEXT("%sAddingNameOnly '%s' (%s)"), *GetIndentation(), *Object->GetName(), *Object->GetClass()->GetName());
-			#endif
-
-			FString ObjectName = GetFullNameSafe(Object);
-			*this << ObjectName;
+			bAnyAssetNotReady = true;
 		}
 		else
 		{
-			// Serialize it
-			ObjectBeingSerialized = Object;
-
 			#if UE_POSE_SEARCH_DERIVED_DATA_LOGGING
-			UE_LOG(LogPoseSearch, Log, TEXT("%sBegin '%s' (%s)"), *GetIndentation(), *Object->GetName(), *Object->GetClass()->GetName());
+			++Indentation;
 			#endif
 
-			const_cast<UObject*>(Object)->Serialize(*this);
+			if (Object->HasAnyFlags(RF_Transient))
+			{
+				#if UE_POSE_SEARCH_DERIVED_DATA_LOGGING
+				UE_LOG(LogPoseSearch, Log, TEXT("%sTransient '%s' (%s)"), *GetIndentation(), *Object->GetName(), *Object->GetClass()->GetName());
+				#endif
+			}
+			else
+			{
+				bool bAlreadyProcessed = false;
+				ObjectsAlreadySerialized.Add(Object, &bAlreadyProcessed);
+
+				// If we haven't already serialized this object
+				if (bAlreadyProcessed)
+				{
+					#if UE_POSE_SEARCH_DERIVED_DATA_LOGGING
+					UE_LOG(LogPoseSearch, Log, TEXT("%sAlreadyProcessed '%s' (%s)"), *GetIndentation(), *Object->GetName(), *Object->GetClass()->GetName());
+					#endif
+				}
+				// for specific types we only add their names to the hash
+				else if (AddNameOnly(Object))
+				{
+					#if UE_POSE_SEARCH_DERIVED_DATA_LOGGING
+					UE_LOG(LogPoseSearch, Log, TEXT("%sAddingNameOnly '%s' (%s)"), *GetIndentation(), *Object->GetName(), *Object->GetClass()->GetName());
+					#endif
+
+					FString ObjectName = GetFullNameSafe(Object);
+					*this << ObjectName;
+				}
+				else
+				{
+					const UObject* PreviousObjectBeingSerialized = ObjectBeingSerialized;
+					ObjectBeingSerialized = Object;
+
+					#if UE_POSE_SEARCH_DERIVED_DATA_LOGGING
+					UE_LOG(LogPoseSearch, Log, TEXT("%sBegin '%s' (%s)"), *GetIndentation(), *Object->GetName(), *Object->GetClass()->GetName());
+					#endif
+
+					const_cast<UObject*>(Object)->Serialize(*this);
+
+					#if UE_POSE_SEARCH_DERIVED_DATA_LOGGING
+					UE_LOG(LogPoseSearch, Log, TEXT("%sEnd '%s' (%s)"), *GetIndentation(), *Object->GetName(), *Object->GetClass()->GetName());
+					#endif
+
+					ObjectBeingSerialized = PreviousObjectBeingSerialized;
+				}
+			}
 
 			#if UE_POSE_SEARCH_DERIVED_DATA_LOGGING
-			UE_LOG(LogPoseSearch, Log, TEXT("%sEnd '%s' (%s)"), *GetIndentation(), *Object->GetName(), *Object->GetClass()->GetName());
+			--Indentation;
 			#endif
-
-			ObjectBeingSerialized = nullptr;
 		}
-			
-		#if UE_POSE_SEARCH_DERIVED_DATA_LOGGING
-		--Indentation;
-		#endif
 	}
 
 	return *this;
@@ -175,8 +198,14 @@ FString FKeyBuilder::GetArchiveName() const
 	return TEXT("FDerivedDataKeyBuilder");
 }
 	
+bool FKeyBuilder::AnyAssetNotReady() const
+{
+	return bAnyAssetNotReady;
+}
+
 FIoHash FKeyBuilder::Finalize() const
 {
+	check(!bAnyAssetNotReady); // otherwise key can be non deterministic
 	// Stores a BLAKE3-160 hash, taken from the first 20 bytes of a BLAKE3-256 hash
 	return FIoHash(Hasher.Finalize());
 }
@@ -188,21 +217,18 @@ const TSet<const UObject*>& FKeyBuilder::GetDependencies() const
 
 bool FKeyBuilder::AddNameOnly(class UObject* Object) const
 {
-	check(Object);
-	const UClass* ObjectClass = Object->GetClass();
-	check(ObjectClass);
-
 	return
-		ObjectClass->IsChildOf(UActorComponent::StaticClass()) ||
-		ObjectClass->IsChildOf(UAnimBoneCompressionSettings::StaticClass()) ||
-		ObjectClass->IsChildOf(UAnimCurveCompressionSettings::StaticClass()) ||
-		ObjectClass->IsChildOf(UAssetImportData::StaticClass()) ||
-		ObjectClass->IsChildOf(UFunction::StaticClass()) ||
-		ObjectClass->IsChildOf(UMirrorDataTable::StaticClass()) ||
-		ObjectClass->IsChildOf(USkeletalMesh::StaticClass()) ||
-		ObjectClass->IsChildOf(USkeletalMeshSocket::StaticClass()) ||
-		ObjectClass->IsChildOf(USkeleton::StaticClass()) ||
-		ObjectClass->IsChildOf(UStreamableRenderAsset::StaticClass());
+		Cast<IAnimationDataModel>(Object) ||
+		Cast<UActorComponent>(Object) ||
+		Cast<UAnimBoneCompressionSettings>(Object) ||
+		Cast<UAnimCurveCompressionSettings>(Object) ||
+		Cast<UAssetImportData>(Object) ||
+		Cast<UFunction>(Object) ||
+		Cast<UMirrorDataTable>(Object) ||
+		Cast<USkeletalMesh>(Object) ||
+		Cast<USkeletalMeshSocket>(Object) ||
+		Cast<USkeleton>(Object) ||
+		Cast<UStreamableRenderAsset>(Object);
 }
 
 #if UE_POSE_SEARCH_DERIVED_DATA_LOGGING
diff --git a/Engine/Plugins/Experimental/Animation/PoseSearch/Source/Runtime/Public/PoseSearch/PoseSearchDerivedData.h b/Engine/Plugins/Experimental/Animation/PoseSearch/Source/Runtime/Public/PoseSearch/PoseSearchDerivedData.h
--- a/Engine/Plugins/Experimental/Animation/PoseSearch/Source/Runtime/Public/PoseSearch/PoseSearchDerivedData.h
+++ b/Engine/Plugins/Experimental/Animation/PoseSearch/Source/Runtime/Public/PoseSearch/PoseSearchDerivedData.h
@@ -19,7 +19,6 @@ namespace UE::PoseSearch
 		NewRequest = 1 << 0,			// generates new key and kick off a task to get updated data (it'll cancel for eventual previous Database request, unless WaitPreviousRequest)
 		ContinueRequest = 1 << 1,		// make sure there's associated data to the Database (doesn't have to be up to date)
 		WaitForCompletion = 1 << 2,		// wait the termination of the NewRequest or ContinueRequest
-		WaitPreviousRequest = 1 << 3	// if NewRequest will WaitPreviousRequest
 	};
 	ENUM_CLASS_FLAGS(ERequestAsyncBuildFlag);
 
diff --git a/Engine/Plugins/Experimental/Animation/PoseSearch/Source/Runtime/Public/PoseSearch/PoseSearchDerivedDataKey.h b/Engine/Plugins/Experimental/Animation/PoseSearch/Source/Runtime/Public/PoseSearch/PoseSearchDerivedDataKey.h
--- a/Engine/Plugins/Experimental/Animation/PoseSearch/Source/Runtime/Public/PoseSearch/PoseSearchDerivedDataKey.h
+++ b/Engine/Plugins/Experimental/Animation/PoseSearch/Source/Runtime/Public/PoseSearch/PoseSearchDerivedDataKey.h
@@ -38,15 +38,16 @@ public:
 	using Super::IsSaving;
 	using Super::operator<<;
 
-	//~ Begin FArchive Interface
+	// Begin FArchive Interface
 	virtual void Seek(int64 InPos) override;
 	virtual bool ShouldSkipProperty(const FProperty* InProperty) const override;
 	virtual void Serialize(void* Data, int64 Length) override;
 	virtual FArchive& operator<<(FName& Name) override;
 	virtual FArchive& operator<<(class UObject*& Object) override;
 	virtual FString GetArchiveName() const override;
-	//~ End FArchive Interface
+	// End FArchive Interface
 	
+	bool AnyAssetNotReady() const;
 	FIoHash Finalize() const;
 	const TSet<const UObject*>& GetDependencies() const;
 
@@ -61,11 +62,14 @@ protected:
 
 	HashBuilderType Hasher;
 
-	/** Set of objects that have already been serialized */
+	// Set of objects that have already been serialized
 	TSet<const UObject*> ObjectsAlreadySerialized;
 
-	/** Object currently being serialized */
+	// Object currently being serialized
 	const UObject* ObjectBeingSerialized = nullptr;
+
+	// true if some dependent assets are not ready (fully loaded)
+	bool bAnyAssetNotReady = false;
 };
 
 } // namespace UE::PoseSearch
