motion matching - IAssetSampler, FSequenceBaseSampler, and FBlendSpaceSampler cleanup of unused APIs (GetTimeFromRootDistance, GetTotalRootDistance, ExtractRootDistance)

diff --git a/Engine/Plugins/Experimental/Animation/PoseSearch/Source/Runtime/Private/PoseSearchAssetIndexer.cpp b/Engine/Plugins/Experimental/Animation/PoseSearch/Source/Runtime/Private/PoseSearchAssetIndexer.cpp
--- a/Engine/Plugins/Experimental/Animation/PoseSearch/Source/Runtime/Private/PoseSearchAssetIndexer.cpp
+++ b/Engine/Plugins/Experimental/Animation/PoseSearch/Source/Runtime/Private/PoseSearchAssetIndexer.cpp
@@ -168,36 +168,38 @@ FAssetIndexer::FSampleInfo FAssetIndexer::GetSampleInfo(float SampleTime) const
 	else
 	{
 		Sample.ClipTime = SamplingParam.WrappedParam;
-
-		const FTransform RootMotionLast = IndexingContext.AssetSampler->GetTotalRootTransform();
-		float RootDistanceLast = IndexingContext.AssetSampler->GetTotalRootDistance();
-
-		// Determine how to accumulate motion for every cycle of the anim. If the sample
-		// had to be clamped, this motion will end up not getting applied below.
-		// Also invert the accumulation direction if the requested sample was wrapped backwards.
-		FTransform RootMotionPerCycle = RootMotionLast;
-		
-		if (SampleTime < 0.0f)
-		{
-			RootMotionPerCycle = RootMotionPerCycle.Inverse();
-		}
+		Sample.RootTransform = FTransform::Identity;
 
 		// Find the remaining motion deltas after wrapping
 		FTransform RootMotionRemainder = Sample.Clip->ExtractRootTransform(Sample.ClipTime);
 
-		// Invert motion deltas if we wrapped backwards
-		if (SampleTime < 0.0f)
+		const bool bNegativeSampleTime = SampleTime < 0.f;
+		if (SamplingParam.NumCycles > 0 || bNegativeSampleTime)
 		{
-			RootMotionRemainder.SetToRelativeTransform(RootMotionLast);
-		}
+			const FTransform RootMotionLast = IndexingContext.AssetSampler->GetTotalRootTransform();
 
-		Sample.RootTransform = FTransform::Identity;
+			// Determine how to accumulate motion for every cycle of the anim. If the sample
+			// had to be clamped, this motion will end up not getting applied below.
+			// Also invert the accumulation direction if the requested sample was wrapped backwards.
+			FTransform RootMotionPerCycle = RootMotionLast;
 
-		// Note if the sample was clamped, no motion will be applied here because NumCycles will be zero
-		int32 CyclesRemaining = SamplingParam.NumCycles;
-		while (CyclesRemaining--)
-		{
-			Sample.RootTransform = RootMotionPerCycle * Sample.RootTransform;
+			if (bNegativeSampleTime)
+			{
+				RootMotionPerCycle = RootMotionPerCycle.Inverse();
+			}
+			
+			// Invert motion deltas if we wrapped backwards
+			if (bNegativeSampleTime)
+			{
+				RootMotionRemainder.SetToRelativeTransform(RootMotionLast);
+			}
+
+			// Note if the sample was clamped, no motion will be applied here because NumCycles will be zero
+			int32 CyclesRemaining = SamplingParam.NumCycles;
+			while (CyclesRemaining--)
+			{
+				Sample.RootTransform = RootMotionPerCycle * Sample.RootTransform;
+			}
 		}
 
 		Sample.RootTransform = RootMotionRemainder * Sample.RootTransform;
diff --git a/Engine/Plugins/Experimental/Animation/PoseSearch/Source/Runtime/Private/PoseSearchAssetSampler.cpp b/Engine/Plugins/Experimental/Animation/PoseSearch/Source/Runtime/Private/PoseSearchAssetSampler.cpp
--- a/Engine/Plugins/Experimental/Animation/PoseSearch/Source/Runtime/Private/PoseSearchAssetSampler.cpp
+++ b/Engine/Plugins/Experimental/Animation/PoseSearch/Source/Runtime/Private/PoseSearchAssetSampler.cpp
@@ -14,71 +14,6 @@ namespace UE::PoseSearch
 //////////////////////////////////////////////////////////////////////////
 // Root motion extrapolation
 
-// Uses distance delta between NextRootDistanceIndex and NextRootDistanceIndex - 1 and extrapolates it to ExtrapolationTime
-static float ExtrapolateAccumulatedRootDistance(
-	int32 SamplingRate,
-	TConstArrayView<float> AccumulatedRootDistance,
-	int32 NextRootDistanceIndex, 
-	float ExtrapolationTime,
-	const FPoseSearchExtrapolationParameters& ExtrapolationParameters)
-{
-	check(NextRootDistanceIndex > 0 && NextRootDistanceIndex < AccumulatedRootDistance.Num());
-
-	const float DistanceDelta =
-		AccumulatedRootDistance[NextRootDistanceIndex] -
-		AccumulatedRootDistance[NextRootDistanceIndex - 1];
-	const float Speed = DistanceDelta * SamplingRate;
-	const float ExtrapolationSpeed = Speed >= ExtrapolationParameters.LinearSpeedThreshold ?
-		Speed : 0.0f;
-	const float ExtrapolatedDistance = ExtrapolationSpeed * ExtrapolationTime;
-
-	return ExtrapolatedDistance;
-}
-
-static float ExtractAccumulatedRootDistance(
-	int32 SamplingRate,
-	TConstArrayView<float> AccumulatedRootDistance,
-	float PlayLength,
-	float Time,
-	const FPoseSearchExtrapolationParameters& ExtrapolationParameters)
-{
-	const float ClampedTime = FMath::Clamp(Time, 0.0f, PlayLength);
-
-	// Find the distance sample that corresponds with the time and split into whole and partial parts
-	float IntegralDistanceSample;
-	float DistanceAlpha = FMath::Modf(ClampedTime * SamplingRate, &IntegralDistanceSample);
-	float DistanceIdx = (int32)IntegralDistanceSample;
-
-	// Verify the distance offset and any residual portion would be in bounds
-	check(DistanceIdx + (DistanceAlpha > 0.0f ? 1 : 0) < AccumulatedRootDistance.Num());
-
-	// Look up the distance and interpolate between distance samples if necessary
-	float Distance = AccumulatedRootDistance[DistanceIdx];
-	if (DistanceAlpha > 0.0f)
-	{
-		float NextDistance = AccumulatedRootDistance[DistanceIdx + 1];
-		Distance = FMath::Lerp(Distance, NextDistance, DistanceAlpha);
-	}
-
-	const float ExtrapolationTime = Time - ClampedTime;
-
-	if (ExtrapolationTime != 0.0f)
-	{
-		// If extrapolationTime is not zero, we extrapolate the beginning or the end of the animation to estimate
-		// the root distance.
-		const int32 DistIdx = (ExtrapolationTime > 0.0f) ? AccumulatedRootDistance.Num() - 1 : 1;
-		const float ExtrapolatedDistance = ExtrapolateAccumulatedRootDistance(
-			SamplingRate,
-			AccumulatedRootDistance,
-			DistIdx,
-			ExtrapolationTime,
-			ExtrapolationParameters);
-		Distance += ExtrapolatedDistance;
-	}
-
-	return Distance;
-}
-
 static FTransform ExtrapolateRootMotion(
 	FTransform SampleToExtrapolate,
 	float SampleStart, 
@@ -185,33 +120,6 @@ void FSequenceBaseSampler::Init(const FInput& InInput)
 
 void FSequenceBaseSampler::Process()
 {
-	ProcessRootDistance();
-}
-
-float FSequenceBaseSampler::GetTimeFromRootDistance(float Distance) const
-{
-	int32 NextSampleIdx = 1;
-	int32 PrevSampleIdx = 0;
-	if (Distance > 0.0f)
-	{
-		// Search for the distance value. Because the values will be extrapolated if necessary
-		// LowerBound might go past the end of the array, in which case the last valid index is used
-		int32 ClipDistanceLowerBoundIndex = Algo::LowerBound(AccumulatedRootDistance, Distance);
-		NextSampleIdx = FMath::Min(
-			ClipDistanceLowerBoundIndex,
-			AccumulatedRootDistance.Num() - 1);
-
-		// Compute distance interpolation amount
-		PrevSampleIdx = FMath::Max(0, NextSampleIdx - 1);
-	}
-
-	float NextDistance = AccumulatedRootDistance[NextSampleIdx];
-	float PrevDistance = AccumulatedRootDistance[PrevSampleIdx];
-	float DistanceSampleAlpha = FMath::GetRangePct(PrevDistance, NextDistance, Distance);
-
-	// Convert to time
-	float ClipTime = (float(NextSampleIdx) - (1.0f - DistanceSampleAlpha)) / Input.RootDistanceSamplingRate;
-	return ClipTime;
 }
 
 float FSequenceBaseSampler::GetPlayLength() const
@@ -229,6 +137,14 @@ void FSequenceBaseSampler::ExtractPose(const FAnimExtractContext& ExtractionCtx,
 	Input.SequenceBase->GetAnimationPose(OutAnimPoseData, ExtractionCtx);
 }
 
+FTransform FSequenceBaseSampler::GetTotalRootTransform() const
+{
+	const FTransform InitialRootTransform = Input.SequenceBase->ExtractRootTrackTransform(0.f, nullptr);
+	const FTransform LastRootTransform = Input.SequenceBase->ExtractRootTrackTransform(Input.SequenceBase->GetPlayLength(), nullptr);
+	const FTransform TotalRootTransform = LastRootTransform.GetRelativeTransform(InitialRootTransform);
+	return TotalRootTransform;
+}
+
 FTransform FSequenceBaseSampler::ExtractRootTransform(float Time) const
 {
 	FTransform RootTransform = FTransform::Identity;
@@ -236,61 +152,52 @@ FTransform FSequenceBaseSampler::ExtractRootTransform(float Time) const
 
 	if (IsLoopable())
 	{
-		FTransform LoopableRootTransform = SequenceBase->ExtractRootMotion(0.0f, Time, true);
-		return LoopableRootTransform;
-	}
-
-	const float ExtrapolationSampleTime = Input.ExtrapolationParameters.SampleTime;
-
-	const float PlayLength = SequenceBase->GetPlayLength();
-	const float ClampedTime = FMath::Clamp(Time, 0.0f, PlayLength);
-	const float ExtrapolationTime = Time - ClampedTime;
-
-	// If Time is less than zero, ExtrapolationTime will be negative. In this case, we extrapolate the beginning of the 
-	// animation to estimate where the root would be at Time
-	if (ExtrapolationTime < -SMALL_NUMBER)
-	{
-		FTransform SampleToExtrapolate = SequenceBase->ExtractRootMotionFromRange(0.0f, ExtrapolationSampleTime);
-
-		const FTransform ExtrapolatedRootMotion = ExtrapolateRootMotion(
-			SampleToExtrapolate,
-			0.0f, ExtrapolationSampleTime,
-			ExtrapolationTime,
-			Input.ExtrapolationParameters);
-		RootTransform = ExtrapolatedRootMotion;
+		RootTransform = SequenceBase->ExtractRootMotion(0.0f, Time, true);
 	}
 	else
 	{
-		RootTransform = SequenceBase->ExtractRootMotionFromRange(0.0f, ClampedTime);
+		const float ExtrapolationSampleTime = Input.ExtrapolationParameters.SampleTime;
+
+		const float PlayLength = SequenceBase->GetPlayLength();
+		const float ClampedTime = FMath::Clamp(Time, 0.0f, PlayLength);
+		const float ExtrapolationTime = Time - ClampedTime;
 
-		// If Time is greater than PlayLength, ExtrapolationTIme will be a positive number. In this case, we extrapolate
-		// the end of the animation to estimate where the root would be at Time
-		if (ExtrapolationTime > SMALL_NUMBER)
+		// If Time is less than zero, ExtrapolationTime will be negative. In this case, we extrapolate the beginning of the 
+		// animation to estimate where the root would be at Time
+		if (ExtrapolationTime < -SMALL_NUMBER)
 		{
-			FTransform SampleToExtrapolate = SequenceBase->ExtractRootMotionFromRange(PlayLength - ExtrapolationSampleTime, PlayLength);
+			FTransform SampleToExtrapolate = SequenceBase->ExtractRootMotionFromRange(0.0f, ExtrapolationSampleTime);
 
 			const FTransform ExtrapolatedRootMotion = ExtrapolateRootMotion(
 				SampleToExtrapolate,
-				PlayLength - ExtrapolationSampleTime, PlayLength,
+				0.0f, ExtrapolationSampleTime,
 				ExtrapolationTime,
 				Input.ExtrapolationParameters);
-			RootTransform = ExtrapolatedRootMotion * RootTransform;
+			RootTransform = ExtrapolatedRootMotion;
+		}
+		else
+		{
+			RootTransform = SequenceBase->ExtractRootMotionFromRange(0.0f, ClampedTime);
+
+			// If Time is greater than PlayLength, ExtrapolationTIme will be a positive number. In this case, we extrapolate
+			// the end of the animation to estimate where the root would be at Time
+			if (ExtrapolationTime > SMALL_NUMBER)
+			{
+				FTransform SampleToExtrapolate = SequenceBase->ExtractRootMotionFromRange(PlayLength - ExtrapolationSampleTime, PlayLength);
+
+				const FTransform ExtrapolatedRootMotion = ExtrapolateRootMotion(
+					SampleToExtrapolate,
+					PlayLength - ExtrapolationSampleTime, PlayLength,
+					ExtrapolationTime,
+					Input.ExtrapolationParameters);
+				RootTransform = ExtrapolatedRootMotion * RootTransform;
+			}
 		}
 	}
 
 	return RootTransform;
 }
 
-float FSequenceBaseSampler::ExtractRootDistance(float Time) const
-{
-	return ExtractAccumulatedRootDistance(
-		Input.RootDistanceSamplingRate,
-		AccumulatedRootDistance,
-		Input.SequenceBase->GetPlayLength(),
-		Time,
-		Input.ExtrapolationParameters);
-}
-
 void FSequenceBaseSampler::ExtractPoseSearchNotifyStates(
 	float Time, 
 	TArray<UAnimNotifyState_PoseSearchBase*>& NotifyStates) const
@@ -324,46 +231,6 @@ void FSequenceBaseSampler::ExtractPoseSearchNotifyStates(
 	}
 }
 
-void FSequenceBaseSampler::ProcessRootDistance()
-{
-	const UAnimSequenceBase* SequenceBase = Input.SequenceBase.Get();
-	
-	// Note the distance sampling interval is independent of the schema's sampling interval
-	const float DistanceSamplingInterval = 1.0f / Input.RootDistanceSamplingRate;
-
-	const FTransform InitialRootTransform = SequenceBase->ExtractRootTrackTransform(0.0f, nullptr);
-
-	uint32 NumDistanceSamples = FMath::CeilToInt(SequenceBase->GetPlayLength() * Input.RootDistanceSamplingRate) + 1;
-	AccumulatedRootDistance.Reserve(NumDistanceSamples);
-
-	// Build a distance lookup table by sampling root motion at a fixed rate and accumulating
-	// absolute translation deltas. During indexing we'll bsearch this table and interpolate
-	// between samples in order to convert distance offsets to time offsets.
-	// See also FAssetIndexer::AddTrajectoryDistanceFeatures().
-
-	double TotalAccumulatedRootDistance = 0.0;
-	FTransform LastRootTransform = InitialRootTransform;
-	float SampleTime = 0.0f;
-	for (int32 SampleIdx = 0; SampleIdx != NumDistanceSamples; ++SampleIdx)
-	{
-		SampleTime = FMath::Min(SampleIdx * DistanceSamplingInterval, SequenceBase->GetPlayLength());
-
-		FTransform RootTransform = SequenceBase->ExtractRootTrackTransform(SampleTime, nullptr);
-		FTransform LocalRootMotion = RootTransform.GetRelativeTransform(LastRootTransform);
-		LastRootTransform = RootTransform;
-
-		TotalAccumulatedRootDistance += LocalRootMotion.GetTranslation().Size();
-		AccumulatedRootDistance.Add((float)TotalAccumulatedRootDistance);
-	}
-
-	// Verify we sampled the final frame of the clip
-	check(SampleTime == SequenceBase->GetPlayLength());
-
-	// Also emit root motion summary info to help with sample wrapping in FAssetIndexer::GetSampleInfo()
-	TotalRootTransform = LastRootTransform.GetRelativeTransform(InitialRootTransform);
-	TotalRootDistance = AccumulatedRootDistance.Last();
-}
-
 const UAnimationAsset* FSequenceBaseSampler::GetAsset() const
 {
 	return Input.SequenceBase.Get();
@@ -384,33 +251,6 @@ void FBlendSpaceSampler::Process()
 
 	ProcessPlayLength();
 	ProcessRootTransform();
-	ProcessRootDistance();
-}
-
-float FBlendSpaceSampler::GetTimeFromRootDistance(float Distance) const
-{
-	int32 NextSampleIdx = 1;
-	int32 PrevSampleIdx = 0;
-	if (Distance > 0.0f)
-	{
-		// Search for the distance value. Because the values will be extrapolated if necessary
-		// LowerBound might go past the end of the array, in which case the last valid index is used
-		int32 ClipDistanceLowerBoundIndex = Algo::LowerBound(AccumulatedRootDistance, Distance);
-		NextSampleIdx = FMath::Min(
-			ClipDistanceLowerBoundIndex,
-			AccumulatedRootDistance.Num() - 1);
-
-		// Compute distance interpolation amount
-		PrevSampleIdx = FMath::Max(0, NextSampleIdx - 1);
-	}
-
-	float NextDistance = AccumulatedRootDistance[NextSampleIdx];
-	float PrevDistance = AccumulatedRootDistance[PrevSampleIdx];
-	float DistanceSampleAlpha = FMath::GetRangePct(PrevDistance, NextDistance, Distance);
-
-	// Convert to time
-	float ClipTime = (float(NextSampleIdx) - (1.0f - DistanceSampleAlpha)) / Input.RootDistanceSamplingRate;
-	return ClipTime;
 }
 
 bool FBlendSpaceSampler::IsLoopable() const
@@ -439,66 +279,65 @@ void FBlendSpaceSampler::ExtractPose(const FAnimExtractContext& ExtractionCtx, F
 	Input.BlendSpace->GetAnimationPose(BlendSamples, ExtractionCtx, OutAnimPoseData);
 }
 
-FTransform FBlendSpaceSampler::ExtractRootTransform(float Time) const
+FTransform FBlendSpaceSampler::GetTotalRootTransform() const
 {
-	if (IsLoopable())
-	{
-		FTransform LoopableRootTransform = ExtractBlendSpaceRootMotion(0.0f, Time, true);
-		return LoopableRootTransform;
-	}
-
-	const float ExtrapolationSampleTime = Input.ExtrapolationParameters.SampleTime;
-
-	const float ClampedTime = FMath::Clamp(Time, 0.0f, PlayLength);
-	const float ExtrapolationTime = Time - ClampedTime;
+	const FTransform InitialRootTransform = ExtractBlendSpaceRootTrackTransform(0.f);
+	const FTransform LastRootTransform = ExtractBlendSpaceRootTrackTransform(PlayLength);
+	const FTransform TotalRootTransform = LastRootTransform.GetRelativeTransform(InitialRootTransform);
+	return TotalRootTransform;
+}
 
+FTransform FBlendSpaceSampler::ExtractRootTransform(float Time) const
+{
 	FTransform RootTransform = FTransform::Identity;
-
-	// If Time is less than zero, ExtrapolationTime will be negative. In this case, we extrapolate the beginning of the 
-	// animation to estimate where the root would be at Time
-	if (ExtrapolationTime < -SMALL_NUMBER)
+	if (IsLoopable())
 	{
-		FTransform SampleToExtrapolate = ExtractBlendSpaceRootMotionFromRange(0.0f, ExtrapolationSampleTime);
-
-		const FTransform ExtrapolatedRootMotion = ExtrapolateRootMotion(
-			SampleToExtrapolate,
-			0.0f, ExtrapolationSampleTime,
-			ExtrapolationTime,
-			Input.ExtrapolationParameters);
-		RootTransform = ExtrapolatedRootMotion;
+		RootTransform = ExtractBlendSpaceRootMotion(0.0f, Time, true);
 	}
 	else
 	{
-		RootTransform = ExtractBlendSpaceRootMotionFromRange(0.0f, ClampedTime);
+		const float ExtrapolationSampleTime = Input.ExtrapolationParameters.SampleTime;
 
-		// If Time is greater than PlayLength, ExtrapolationTIme will be a positive number. In this case, we extrapolate
-		// the end of the animation to estimate where the root would be at Time
-		if (ExtrapolationTime > SMALL_NUMBER)
+		const float ClampedTime = FMath::Clamp(Time, 0.0f, PlayLength);
+		const float ExtrapolationTime = Time - ClampedTime;
+
+
+		// If Time is less than zero, ExtrapolationTime will be negative. In this case, we extrapolate the beginning of the 
+		// animation to estimate where the root would be at Time
+		if (ExtrapolationTime < -SMALL_NUMBER)
 		{
-			FTransform SampleToExtrapolate = ExtractBlendSpaceRootMotionFromRange(PlayLength - ExtrapolationSampleTime, PlayLength);
+			FTransform SampleToExtrapolate = ExtractBlendSpaceRootMotionFromRange(0.0f, ExtrapolationSampleTime);
 
 			const FTransform ExtrapolatedRootMotion = ExtrapolateRootMotion(
 				SampleToExtrapolate,
-				PlayLength - ExtrapolationSampleTime, PlayLength,
+				0.0f, ExtrapolationSampleTime,
 				ExtrapolationTime,
 				Input.ExtrapolationParameters);
-			RootTransform = ExtrapolatedRootMotion * RootTransform;
+			RootTransform = ExtrapolatedRootMotion;
+		}
+		else
+		{
+			RootTransform = ExtractBlendSpaceRootMotionFromRange(0.0f, ClampedTime);
+
+			// If Time is greater than PlayLength, ExtrapolationTIme will be a positive number. In this case, we extrapolate
+			// the end of the animation to estimate where the root would be at Time
+			if (ExtrapolationTime > SMALL_NUMBER)
+			{
+				FTransform SampleToExtrapolate = ExtractBlendSpaceRootMotionFromRange(PlayLength - ExtrapolationSampleTime, PlayLength);
+
+				const FTransform ExtrapolatedRootMotion = ExtrapolateRootMotion(
+					SampleToExtrapolate,
+					PlayLength - ExtrapolationSampleTime, PlayLength,
+					ExtrapolationTime,
+					Input.ExtrapolationParameters);
+				RootTransform = ExtrapolatedRootMotion * RootTransform;
+			}
 		}
 	}
 
 	return RootTransform;
 }
 
-float FBlendSpaceSampler::ExtractRootDistance(float Time) const
-{
-	return ExtractAccumulatedRootDistance(
-		Input.RootDistanceSamplingRate,
-		AccumulatedRootDistance,
-		PlayLength,
-		Time,
-		Input.ExtrapolationParameters);
-}
-
 static int32 GetHighestWeightSample(const TArray<struct FBlendSampleData>& SampleDataList)
 {
 	int32 HighestWeightIndex = 0;
@@ -571,16 +410,12 @@ FTransform FBlendSpaceSampler::ExtractBlendSpaceRootTrackTransform(float Time) c
 {
 	checkf(AccumulatedRootTransform.Num() > 0, TEXT("ProcessRootTransform must be run first"));
 
-	int32 Index = Time * Input.RootTransformSamplingRate;
-	int32 FirstIndexClamped = FMath::Clamp(Index + 0, 0, AccumulatedRootTransform.Num() - 1);
-	int32 SecondIndexClamped = FMath::Clamp(Index + 1, 0, AccumulatedRootTransform.Num() - 1);
-	float Alpha = FMath::Fmod(Time * Input.RootTransformSamplingRate, 1.0f);
+	const int32 Index = Time * Input.RootTransformSamplingRate;
+	const int32 FirstIndexClamped = FMath::Clamp(Index + 0, 0, AccumulatedRootTransform.Num() - 1);
+	const int32 SecondIndexClamped = FMath::Clamp(Index + 1, 0, AccumulatedRootTransform.Num() - 1);
+	const float Alpha = FMath::Fmod(Time * Input.RootTransformSamplingRate, 1.0f);
 	FTransform OutputTransform;
-	OutputTransform.Blend(
-		AccumulatedRootTransform[FirstIndexClamped],
-		AccumulatedRootTransform[SecondIndexClamped],
-		Alpha);
-
+	OutputTransform.Blend(AccumulatedRootTransform[FirstIndexClamped], AccumulatedRootTransform[SecondIndexClamped], Alpha);
 	return OutputTransform;
 }
 
@@ -715,45 +550,6 @@ void FBlendSpaceSampler::ProcessRootTransform()
 	}
 }
 
-void FBlendSpaceSampler::ProcessRootDistance()
-{
-	checkf(AccumulatedRootTransform.Num() > 0, TEXT("ProcessRootTransform must be run first"));
-
-	// Note the distance sampling interval is independent of the schema's sampling interval
-	const float DistanceSamplingInterval = 1.0f / Input.RootDistanceSamplingRate;
-
-	const FTransform InitialRootTransform = FTransform::Identity;
-
-	uint32 NumDistanceSamples = FMath::CeilToInt(PlayLength * Input.RootDistanceSamplingRate) + 1;
-	AccumulatedRootDistance.Reserve(NumDistanceSamples);
-
-	// Build a distance lookup table by sampling root motion at a fixed rate and accumulating
-	// absolute translation deltas. During indexing we'll bsearch this table and interpolate
-	// between samples in order to convert distance offsets to time offsets.
-	// See also FAssetIndexer::AddTrajectoryDistanceFeatures().
-	double TotalAccumulatedRootDistance = 0.0;
-	FTransform LastRootTransform = InitialRootTransform;
-	float SampleTime = 0.0f;
-	for (int32 SampleIdx = 0; SampleIdx != NumDistanceSamples; ++SampleIdx)
-	{
-		SampleTime = FMath::Min(SampleIdx * DistanceSamplingInterval, PlayLength);
-
-		FTransform RootTransform = ExtractBlendSpaceRootTrackTransform(SampleTime);
-		FTransform LocalRootMotion = RootTransform.GetRelativeTransform(LastRootTransform);
-		LastRootTransform = RootTransform;
-
-		TotalAccumulatedRootDistance += LocalRootMotion.GetTranslation().Size();
-		AccumulatedRootDistance.Add((float)TotalAccumulatedRootDistance);
-	}
-
-	// Verify we sampled the final frame of the clip
-	check(SampleTime == PlayLength);
-
-	// Also emit root motion summary info to help with sample wrapping in FAssetIndexer::GetSampleInfo()
-	TotalRootTransform = LastRootTransform.GetRelativeTransform(InitialRootTransform);
-	TotalRootDistance = AccumulatedRootDistance.Last(); 
-}
-
 const UAnimationAsset* FBlendSpaceSampler::GetAsset() const
 {
 	return Input.BlendSpace.Get();
diff --git a/Engine/Plugins/Experimental/Animation/PoseSearch/Source/Runtime/Public/PoseSearch/PoseSearchAssetSampler.h b/Engine/Plugins/Experimental/Animation/PoseSearch/Source/Runtime/Public/PoseSearch/PoseSearchAssetSampler.h
--- a/Engine/Plugins/Experimental/Animation/PoseSearch/Source/Runtime/Public/PoseSearch/PoseSearchAssetSampler.h
+++ b/Engine/Plugins/Experimental/Animation/PoseSearch/Source/Runtime/Public/PoseSearch/PoseSearchAssetSampler.h
@@ -63,22 +63,12 @@ public:
 	virtual float GetPlayLength() const = 0;
 	virtual bool IsLoopable() const = 0;
 
-	// Gets the time associated with a particular root distance traveled
-	virtual float GetTimeFromRootDistance(float Distance) const = 0;
-
-	// Gets the total root distance traveled 
-	virtual float GetTotalRootDistance() const = 0;
-
 	// Gets the final root transformation at the end of the asset's playback time
 	virtual FTransform GetTotalRootTransform() const = 0;
 
 	// Extracts pose for this asset for a given context
 	virtual void ExtractPose(const FAnimExtractContext& ExtractionCtx, FAnimationPoseData& OutAnimPoseData) const = 0;
 
-	// Extracts the accumulated root distance at the given time, using the extremities of the sequence to extrapolate 
-	// beyond the sequence limits when Time is less than zero or greater than the sequence length
-	virtual float ExtractRootDistance(float Time) const = 0;
-
 	// Extracts root transform at the given time, using the extremities of the sequence to extrapolate beyond the 
 	// sequence limits when Time is less than zero or greater than the sequence length.
 	virtual FTransform ExtractRootTransform(float Time) const = 0;
@@ -99,7 +89,6 @@ public:
 	struct FInput
 	{
 		TWeakObjectPtr<const UAnimSequenceBase> SequenceBase;
-		int32 RootDistanceSamplingRate = 60;
 		FPoseSearchExtrapolationParameters ExtrapolationParameters;
 	} Input;
 
@@ -109,23 +98,12 @@ public:
 	virtual float GetPlayLength() const override;
 	virtual bool IsLoopable() const override;
 
-	virtual float GetTimeFromRootDistance(float Distance) const override;
-
-	virtual float GetTotalRootDistance() const override { return TotalRootDistance; }
-	virtual FTransform GetTotalRootTransform() const override { return TotalRootTransform; }
+	virtual FTransform GetTotalRootTransform() const override;
 
 	virtual void ExtractPose(const FAnimExtractContext& ExtractionCtx, FAnimationPoseData& OutAnimPoseData) const override;
-	virtual float ExtractRootDistance(float Time) const override;
 	virtual FTransform ExtractRootTransform(float Time) const override;
 	virtual void ExtractPoseSearchNotifyStates(float Time, TArray<class UAnimNotifyState_PoseSearchBase*>& NotifyStates) const override;
 	virtual const UAnimationAsset* GetAsset() const override;
-
-private:
-	float TotalRootDistance = 0.0f;
-	FTransform TotalRootTransform = FTransform::Identity;
-	TArray<float> AccumulatedRootDistance;
-
-	void ProcessRootDistance();
 };
 
 struct POSESEARCH_API FBlendSpaceSampler : public IAssetSampler
@@ -135,8 +113,7 @@ public:
 	{
 		FBoneContainer BoneContainer;
 		TWeakObjectPtr<const UBlendSpace> BlendSpace;
-		int32 RootDistanceSamplingRate = 60;
-		int32 RootTransformSamplingRate = 60;
+		int32 RootTransformSamplingRate = 30;
 		FPoseSearchExtrapolationParameters ExtrapolationParameters;
 		FVector BlendParameters;
 	} Input;
@@ -147,13 +124,9 @@ public:
 	virtual float GetPlayLength() const override { return PlayLength; }
 	virtual bool IsLoopable() const override;
 
-	virtual float GetTimeFromRootDistance(float Distance) const override;
-
-	virtual float GetTotalRootDistance() const override { return TotalRootDistance; }
-	virtual FTransform GetTotalRootTransform() const override { return TotalRootTransform; }
+	virtual FTransform GetTotalRootTransform() const override;
 
 	virtual void ExtractPose(const FAnimExtractContext& ExtractionCtx, FAnimationPoseData& OutAnimPoseData) const override;
-	virtual float ExtractRootDistance(float Time) const override;
 	virtual FTransform ExtractRootTransform(float Time) const override;
 	virtual void ExtractPoseSearchNotifyStates(float Time, TArray<class UAnimNotifyState_PoseSearchBase*>& NotifyStates) const override;
 
@@ -161,13 +134,9 @@ public:
 
 private:
 	float PlayLength = 0.0f;
-	float TotalRootDistance = 0.0f;
-	FTransform TotalRootTransform = FTransform::Identity;
-	TArray<float> AccumulatedRootDistance;
 	TArray<FTransform> AccumulatedRootTransform;
 
 	void ProcessPlayLength();
-	void ProcessRootDistance();
 	void ProcessRootTransform();
 
 	// Extracts the pre-computed blend space root transform. ProcessRootTransform must be run first.
