motion matching - FDatabaseViewModel::RespawnPreviewActors cleanup

diff --git a/Engine/Plugins/Experimental/Animation/PoseSearch/Source/Editor/Private/PoseSearchDatabaseViewModel.cpp b/Engine/Plugins/Experimental/Animation/PoseSearch/Source/Editor/Private/PoseSearchDatabaseViewModel.cpp
--- a/Engine/Plugins/Experimental/Animation/PoseSearch/Source/Editor/Private/PoseSearchDatabaseViewModel.cpp
+++ b/Engine/Plugins/Experimental/Animation/PoseSearch/Source/Editor/Private/PoseSearchDatabaseViewModel.cpp
@@ -148,55 +148,31 @@ namespace UE::PoseSearch
 			return;
 		}
 
-		TSet<int32> AssociatedSequencesAssetIndices;
-		TSet<int32> AssociatedBlendSpacesAssetIndices;
-		for (const TSharedPtr<FDatabaseAssetTreeNode>& SelectedNode : SelectedNodes)
-		{
-			if (SelectedNode->SourceAssetType == ESearchIndexAssetType::Sequence || SelectedNode->SourceAssetType == ESearchIndexAssetType::AnimComposite)
-			{
-				AssociatedSequencesAssetIndices.Add(SelectedNode->SourceAssetIdx);
-			}
-			else if (SelectedNode->SourceAssetType == ESearchIndexAssetType::BlendSpace)
-			{
-				AssociatedBlendSpacesAssetIndices.Add(SelectedNode->SourceAssetIdx);
-			}
-		}
-
 		if (FAsyncPoseSearchDatabasesManagement::RequestAsyncBuildIndex(PoseSearchDatabase, ERequestAsyncBuildFlag::ContinueRequest))
 		{
 			FBoneContainer BoneContainer;
 			BoneContainer.InitializeTo(PoseSearchDatabase->Schema->BoneIndicesWithParents, FCurveEvaluationOption(false), *PoseSearchDatabase->Schema->Skeleton);
 
+			TSet<int32> AssociatedAssetIndices;
+			for (const TSharedPtr<FDatabaseAssetTreeNode>& SelectedNode : SelectedNodes)
+			{
+				AssociatedAssetIndices.Add(SelectedNode->SourceAssetIdx);
+			}
+
 			const FPoseSearchIndex& SearchIndex = PoseSearchDatabase->GetSearchIndex();
 			for (int32 IndexAssetIndex = 0; IndexAssetIndex < SearchIndex.Assets.Num(); ++IndexAssetIndex)
 			{
 				const FPoseSearchIndexAsset& IndexAsset = SearchIndex.Assets[IndexAssetIndex];
-
-				bool bSpawn = false;
-				if (AnimationPreviewMode == EAnimationPreviewMode::OriginalAndMirrored || !IndexAsset.bMirrored)
-				{
-					bool bIsAssociatedToSelection = false;
-					if (IndexAsset.Type == ESearchIndexAssetType::Sequence || IndexAsset.Type == ESearchIndexAssetType::AnimComposite)
-					{
-						bSpawn = AssociatedSequencesAssetIndices.Contains(IndexAsset.SourceAssetIdx);
-					}
-					else if (IndexAsset.Type == ESearchIndexAssetType::BlendSpace)
-					{
-						bSpawn = AssociatedBlendSpacesAssetIndices.Contains(IndexAsset.SourceAssetIdx);
-					}
-				}
-
-				if (bSpawn)
+				if ((AnimationPreviewMode == EAnimationPreviewMode::OriginalAndMirrored || !IndexAsset.bMirrored) &&
+					AssociatedAssetIndices.Contains(IndexAsset.SourceAssetIdx))
 				{
 					FDatabasePreviewActor PreviewActor = SpawnPreviewActor(IndexAssetIndex, BoneContainer);
 					if (PreviewActor.IsValid())
 					{
-						const UAnimationAsset* AnimationAsset = PreviewActor.GetAnimPreviewInstance()->GetAnimationAsset();
-						if (AnimationAsset)
+						if (const UAnimationAsset* AnimationAsset = PreviewActor.GetAnimPreviewInstance()->GetAnimationAsset())
 						{
 							MaxPreviewPlayLength = FMath::Max(MaxPreviewPlayLength, AnimationAsset->GetPlayLength());
 						}
-						
 						PreviewActors.Add(PreviewActor);
 					}
 				}
