motion matching - fix for crash while editing databases with a paused PIE

diff --git a/Engine/Plugins/Experimental/Animation/PoseSearch/Source/Runtime/Private/AnimNode_MotionMatching.cpp b/Engine/Plugins/Experimental/Animation/PoseSearch/Source/Runtime/Private/AnimNode_MotionMatching.cpp
--- a/Engine/Plugins/Experimental/Animation/PoseSearch/Source/Runtime/Private/AnimNode_MotionMatching.cpp
+++ b/Engine/Plugins/Experimental/Animation/PoseSearch/Source/Runtime/Private/AnimNode_MotionMatching.cpp
@@ -81,9 +81,22 @@ void FAnimNode_MotionMatching::UpdateAssetPlayer(const FAnimationUpdateContext&
 	else
 	{
 #if WITH_EDITOR
-		// in case we're still indexing MotionMatchingState.CurrentSearchResult.Database we Reset the MotionMatchingState
-		if (!FAsyncPoseSearchDatabasesManagement::RequestAsyncBuildIndex(MotionMatchingState.CurrentSearchResult.Database.Get(), ERequestAsyncBuildFlag::ContinueRequest))
+		if (FAsyncPoseSearchDatabasesManagement::RequestAsyncBuildIndex(MotionMatchingState.CurrentSearchResult.Database.Get(), ERequestAsyncBuildFlag::ContinueRequest))
+		{
+			const UPoseSearchDatabase* Database = MotionMatchingState.CurrentSearchResult.Database.Get();
+			const FPoseSearchIndex& SearchIndex = Database->GetSearchIndex();
+			if (!SearchIndex.IsValidPoseIndex(MotionMatchingState.CurrentSearchResult.PrevPoseIdx) ||
+				!SearchIndex.IsValidPoseIndex(MotionMatchingState.CurrentSearchResult.NextPoseIdx) ||
+				Database->Schema != MotionMatchingState.CurrentSearchResult.ComposedQuery.GetSchema())
+			{
+				// MotionMatchingState is out of sync with Database: we need to reset the MM state. This could happen if PIE is paused, and we edit the database,
+				// so FAnimNode_MotionMatching::UpdateAssetPlayer is never called and FAsyncPoseSearchDatabasesManagement::RequestAsyncBuildIndex never returns false here
+				MotionMatchingState.Reset();
+			}
+		}
+		else
 		{
+			// we're still indexing MotionMatchingState.CurrentSearchResult.Database, so we Reset the MotionMatchingState
 			MotionMatchingState.Reset();
 		}
 #endif // WITH_EDITOR
