motion matching - database indexing speed up and improved memory footprint

diff --git a/Engine/Plugins/Experimental/Animation/PoseSearch/Source/Runtime/Private/PoseSearchAssetIndexer.cpp b/Engine/Plugins/Experimental/Animation/PoseSearch/Source/Runtime/Private/PoseSearchAssetIndexer.cpp
--- a/Engine/Plugins/Experimental/Animation/PoseSearch/Source/Runtime/Private/PoseSearchAssetIndexer.cpp
+++ b/Engine/Plugins/Experimental/Animation/PoseSearch/Source/Runtime/Private/PoseSearchAssetIndexer.cpp
@@ -208,13 +208,6 @@ FAssetIndexer::FSampleInfo FAssetIndexer::GetSampleInfo(float SampleTime) const
 	return Sample;
 }
 
-FAssetIndexer::FSampleInfo FAssetIndexer::GetSampleInfoRelative(float SampleTime, const FSampleInfo& Origin) const
-{
-	FSampleInfo Sample = GetSampleInfo(SampleTime);
-	Sample.RootTransform.SetToRelativeTransform(Origin.RootTransform);
-	return Sample;
-}
-
 FTransform FAssetIndexer::MirrorTransform(const FTransform& Transform) const
 {
 	return IndexingContext.bMirrored ? IndexingContext.SamplingContext->MirrorTransform(Transform) : Transform;
@@ -253,45 +246,33 @@ FPoseSearchPoseMetadata FAssetIndexer::GetMetadata(int32 SampleIdx) const
 	return Metadata;
 }
 
-// returns the transform in character space for the bone indexed by Schema->BoneReferences[SchemaBoneIdx] 
-// at SampleTime-OriginTime seconds ahead (or behind) the pose at OriginTime
-FTransform FAssetIndexer::GetTransformAndCacheResults(float SampleTime, float OriginTime, int8 SchemaBoneIdx, bool& Clamped)
+FAssetIndexer::CachedEntry& FAssetIndexer::GetEntry(float SampleTime)
 {
-	// @todo: use an hashmap if we end up having too many entries
-	CachedEntry* Entry = CachedEntries.FindByPredicate([SampleTime, OriginTime](const FAssetIndexer::CachedEntry& Entry)
-		{
-			return Entry.SampleTime == SampleTime && Entry.OriginTime == OriginTime;
-		});
+	using namespace UE::Anim;
 
+	CachedEntry* Entry = CachedEntries.Find(SampleTime);
 	const FAssetSamplingContext* SamplingContext = IndexingContext.SamplingContext;
 
 	if (!Entry)
 	{
-		Entry = &CachedEntries[CachedEntries.AddDefaulted()];
-
+		Entry = &CachedEntries.Add(SampleTime);
 		Entry->SampleTime = SampleTime;
-		Entry->OriginTime = OriginTime;
 
 		if (!BoneContainer.IsValid())
 		{
 			UE_LOG(LogPoseSearch,
 				Warning,
-				TEXT("Invalid BoneContainer encountered in FAssetIndexer::GetTransformAndCacheResults. Asset: %s. Schema: %s. BoneContainerAsset: %s. NumBoneIndices: %d"),
+				TEXT("Invalid BoneContainer encountered in FAssetIndexer::GetEntry. Asset: %s. Schema: %s. BoneContainerAsset: %s. NumBoneIndices: %d"),
 				*GetNameSafe(IndexingContext.AssetSampler->GetAsset()),
 				*GetNameSafe(IndexingContext.Schema),
 				*GetNameSafe(BoneContainer.GetAsset()),
 				BoneContainer.GetCompactPoseNumBones());
 		}
 
-		Entry->Pose.SetBoneContainer(&BoneContainer);
-		Entry->UnusedCurve.InitFrom(BoneContainer);
-
-		IAssetIndexer::FSampleInfo Origin = GetSampleInfo(OriginTime);
-		IAssetIndexer::FSampleInfo Sample = GetSampleInfoRelative(SampleTime, Origin);
-
+		const IAssetIndexer::FSampleInfo Sample = GetSampleInfo(SampleTime);
 		float CurrentTime = Sample.ClipTime;
 		float PreviousTime = CurrentTime - SamplingContext->FiniteDelta;
-		
+
 		const bool bLoopable = Sample.Clip->IsLoopable();
 		const float PlayLength = Sample.Clip->GetPlayLength();
 		if (!bLoopable)
@@ -314,12 +295,20 @@ FTransform FAssetIndexer::GetTransformAndCacheResults(float SampleTime, float Or
 		// no need to extract root motion here, since we use the precalculated Sample.RootTransform as root transform for the Entry
 		FAnimExtractContext ExtractionCtx(static_cast<double>(CurrentTime), false, DeltaTimeRecord, bLoopable);
 
-		Sample.Clip->ExtractPose(ExtractionCtx, Entry->AnimPoseData);
+		FCompactPose Pose;
+		FBlendedCurve UnusedCurve;
+		FStackAttributeContainer UnusedAtrribute;
+		FAnimationPoseData AnimPoseData = { Pose, UnusedCurve, UnusedAtrribute };
+
+		UnusedCurve.InitFrom(BoneContainer);
+		Pose.SetBoneContainer(&BoneContainer);
+
+		Sample.Clip->ExtractPose(ExtractionCtx, AnimPoseData);
 
 		if (IndexingContext.bMirrored)
 		{
 			FAnimationRuntime::MirrorPose(
-				Entry->AnimPoseData.GetPose(),
+				AnimPoseData.GetPose(),
 				IndexingContext.Schema->MirrorDataTable->MirrorAxis,
 				SamplingContext->CompactPoseMirrorBones,
 				SamplingContext->ComponentSpaceRefRotations
@@ -327,25 +316,63 @@ FTransform FAssetIndexer::GetTransformAndCacheResults(float SampleTime, float Or
 			// Note curves and attributes are not used during the indexing process and therefore don't need to be mirrored
 		}
 
-		Entry->ComponentSpacePose.InitPose(Entry->Pose);
+		Entry->ComponentSpacePose.InitPose(MoveTemp(Pose));
 		Entry->RootTransform = Sample.RootTransform;
 		Entry->Clamped = Sample.bClamped;
 	}
 
-	FTransform BoneTransform;
-	const FBoneReference& BoneReference = IndexingContext.Schema->BoneReferences[SchemaBoneIdx];
-	if (BoneReference.HasValidSetup())
+	return *Entry;
+}
+
+// returns the transform in component space for the bone indexed by Schema->BoneReferences[SchemaBoneIdx] at SampleTime seconds
+FTransform FAssetIndexer::GetComponentSpaceTransform(float SampleTime, bool& Clamped, int8 SchemaBoneIdx)
+{
+	CachedEntry& Entry = GetEntry(SampleTime);
+	Clamped = Entry.Clamped;
+
+	if (SchemaBoneIdx >= 0 && IndexingContext.Schema->BoneReferences[SchemaBoneIdx].HasValidSetup())
 	{
-		FCompactPoseBoneIndex CompactBoneIndex = BoneContainer.MakeCompactPoseIndex(FMeshPoseBoneIndex(BoneReference.BoneIndex));
-		BoneTransform = Entry->ComponentSpacePose.GetComponentSpaceTransform(CompactBoneIndex) * MirrorTransform(Entry->RootTransform);
+		return CalculateComponentSpaceTransform(Entry, SchemaBoneIdx);
 	}
-	else
+
+	return FTransform::Identity;
+}
+
+// returns the transform in animation space for the bone indexed by Schema->BoneReferences[SchemaBoneIdx] at SampleTime seconds
+FTransform FAssetIndexer::GetTransform(float SampleTime, bool& Clamped, int8 SchemaBoneIdx)
+{
+	CachedEntry& Entry = GetEntry(SampleTime);
+	Clamped = Entry.Clamped;
+
+	const FTransform MirroredRootTransform = MirrorTransform(Entry.RootTransform);
+	if (SchemaBoneIdx >= 0 && IndexingContext.Schema->BoneReferences[SchemaBoneIdx].HasValidSetup())
 	{
-		BoneTransform = MirrorTransform(Entry->RootTransform);
+		return CalculateComponentSpaceTransform(Entry, SchemaBoneIdx) * MirroredRootTransform;
 	}
 
-	Clamped = Entry->Clamped;
+	return MirroredRootTransform;
+}
+
+FTransform FAssetIndexer::CalculateComponentSpaceTransform(FAssetIndexer::CachedEntry& Entry, int8 SchemaBoneIdx)
+{
+	const FBoneReference& BoneReference = IndexingContext.Schema->BoneReferences[SchemaBoneIdx];
+	const FCompactPoseBoneIndex CompactBoneIndex = BoneContainer.MakeCompactPoseIndex(FMeshPoseBoneIndex(BoneReference.BoneIndex));
+	return Entry.ComponentSpacePose.GetComponentSpaceTransform(CompactBoneIndex);
+}
+
+// returns the transform in component space for the bone indexed by Schema->BoneReferences[SchemaBoneIdx] 
+// at SampleTime-OriginTime seconds ahead (or behind) the pose at OriginTime
+FTransform FAssetIndexer::GetComponentSpaceTransform(float SampleTime, float OriginTime, bool& Clamped, int8 SchemaBoneIdx)
+{
+	if (SampleTime == OriginTime)
+	{
+		return GetComponentSpaceTransform(SampleTime, Clamped, SchemaBoneIdx);
+	}
 
+	bool Unused = false;
+	const FTransform RootBoneTransform = GetTransform(OriginTime, Unused);
+	FTransform BoneTransform = GetTransform(SampleTime, Clamped, SchemaBoneIdx);
+	BoneTransform.SetToRelativeTransform(RootBoneTransform);
 	return BoneTransform;
 }
 
diff --git a/Engine/Plugins/Experimental/Animation/PoseSearch/Source/Runtime/Private/PoseSearchAssetIndexer.h b/Engine/Plugins/Experimental/Animation/PoseSearch/Source/Runtime/Private/PoseSearchAssetIndexer.h
--- a/Engine/Plugins/Experimental/Animation/PoseSearch/Source/Runtime/Private/PoseSearchAssetIndexer.h
+++ b/Engine/Plugins/Experimental/Animation/PoseSearch/Source/Runtime/Private/PoseSearchAssetIndexer.h
@@ -37,32 +37,29 @@ public:
 public: // IAssetIndexer
 
 	const FAssetIndexingContext& GetIndexingContext() const override { return IndexingContext; }
-	FSampleInfo GetSampleInfo(float SampleTime) const override;
-	FSampleInfo GetSampleInfoRelative(float SampleTime, const FSampleInfo& Origin) const override;
-	FTransform MirrorTransform(const FTransform& Transform) const override;
-	FTransform GetTransformAndCacheResults(float SampleTime, float OriginTime, int8 SchemaBoneIdx, bool& Clamped) override;
+	FTransform GetTransform(float SampleTime, bool& Clamped, int8 SchemaBoneIdx = RootSchemaBoneIdx) override;
+	FTransform GetComponentSpaceTransform(float SampleTime, bool& Clamped, int8 SchemaBoneIdx = RootSchemaBoneIdx) override;
+	FTransform GetComponentSpaceTransform(float SampleTime, float OriginTime, bool& Clamped, int8 SchemaBoneIdx = RootSchemaBoneIdx) override;
 
 private:
-	FPoseSearchPoseMetadata GetMetadata(int32 SampleIdx) const;
-
 	struct CachedEntry
 	{
 		float SampleTime;
-		float OriginTime;
 		bool Clamped;
 
-		// @todo: minimize the Entry memory footprint
 		FTransform RootTransform;
-		FCompactPose Pose;
 		FCSPose<FCompactPose> ComponentSpacePose;
-		FBlendedCurve UnusedCurve;
-		UE::Anim::FStackAttributeContainer UnusedAtrribute;
-		FAnimationPoseData AnimPoseData = { Pose, UnusedCurve, UnusedAtrribute };
 	};
 
+	FSampleInfo GetSampleInfo(float SampleTime) const;
+	FPoseSearchPoseMetadata GetMetadata(int32 SampleIdx) const;
+	FTransform MirrorTransform(const FTransform& Transform) const;
+	CachedEntry& GetEntry(float SampleTime);
+	FTransform CalculateComponentSpaceTransform(CachedEntry& Entry, int8 SchemaBoneIdx);
+
 	FBoneContainer BoneContainer;
 	FAssetIndexingContext IndexingContext;
-	TArray<CachedEntry> CachedEntries;
+	TMap<float, CachedEntry> CachedEntries;
 };
 
 } // namespace UE::PoseSearch
diff --git a/Engine/Plugins/Experimental/Animation/PoseSearch/Source/Runtime/Private/PoseSearchContext.cpp b/Engine/Plugins/Experimental/Animation/PoseSearch/Source/Runtime/Private/PoseSearchContext.cpp
--- a/Engine/Plugins/Experimental/Animation/PoseSearch/Source/Runtime/Private/PoseSearchContext.cpp
+++ b/Engine/Plugins/Experimental/Animation/PoseSearch/Source/Runtime/Private/PoseSearchContext.cpp
@@ -156,7 +156,7 @@ FTransform FSearchContext::TryGetTransformAndCacheResults(float SampleTime, cons
 		return CachedTransform->Transform;
 	}
 
-	if (BoneIndexType != UPoseSearchSchema::RootBoneIdx)
+	if (BoneIndexType != RootBoneIndexType)
 	{
 		TArray<FTransform> SampledLocalPose;
 		if (History->TrySampleLocalPose(-SampleTime, &Schema->BoneIndicesWithParents, &SampledLocalPose, nullptr))
diff --git a/Engine/Plugins/Experimental/Animation/PoseSearch/Source/Runtime/Private/PoseSearchDerivedDataKey.cpp b/Engine/Plugins/Experimental/Animation/PoseSearch/Source/Runtime/Private/PoseSearchDerivedDataKey.cpp
--- a/Engine/Plugins/Experimental/Animation/PoseSearch/Source/Runtime/Private/PoseSearchDerivedDataKey.cpp
+++ b/Engine/Plugins/Experimental/Animation/PoseSearch/Source/Runtime/Private/PoseSearchDerivedDataKey.cpp
@@ -41,7 +41,7 @@ FKeyBuilder::FKeyBuilder(const UObject* Object, bool bUseDataVer)
 	if (bUseDataVer)
 	{
 		// used to invalidate the key without having to change POSESEARCHDB_DERIVEDDATA_VER all the times
-		int32 POSESEARCHDB_DERIVEDDATA_VER_SMALL = 17;
+		int32 POSESEARCHDB_DERIVEDDATA_VER_SMALL = 29;
 		FGuid VersionGuid = FDevSystemGuids::GetSystemGuid(FDevSystemGuids::Get().POSESEARCHDB_DERIVEDDATA_VER);
 
 		*this << VersionGuid;
diff --git a/Engine/Plugins/Experimental/Animation/PoseSearch/Source/Runtime/Private/PoseSearchFeatureChannel_FilterCrashingLegs.cpp b/Engine/Plugins/Experimental/Animation/PoseSearch/Source/Runtime/Private/PoseSearchFeatureChannel_FilterCrashingLegs.cpp
--- a/Engine/Plugins/Experimental/Animation/PoseSearch/Source/Runtime/Private/PoseSearchFeatureChannel_FilterCrashingLegs.cpp
+++ b/Engine/Plugins/Experimental/Animation/PoseSearch/Source/Runtime/Private/PoseSearchFeatureChannel_FilterCrashingLegs.cpp
@@ -56,10 +56,10 @@ void UPoseSearchFeatureChannel_FilterCrashingLegs::IndexAsset(UE::PoseSearch::IA
 		const float SubsampleTime = SampleIdx * IndexingContext.Schema->GetSamplingInterval();
 
 		bool bUnused;
-		const FTransform RightThighTransform = Indexer.GetTransformAndCacheResults(SubsampleTime, SubsampleTime, RightThighIdx, bUnused);
-		const FTransform LeftThighTransform = Indexer.GetTransformAndCacheResults(SubsampleTime, SubsampleTime, LeftThighIdx, bUnused);
-		const FTransform RightFootTransform = Indexer.GetTransformAndCacheResults(SubsampleTime, SubsampleTime, RightFootIdx, bUnused);
-		const FTransform LeftFootTransform = Indexer.GetTransformAndCacheResults(SubsampleTime, SubsampleTime, LeftFootIdx, bUnused);
+		const FTransform RightThighTransform = Indexer.GetComponentSpaceTransform(SubsampleTime, bUnused, RightThighIdx);
+		const FTransform LeftThighTransform = Indexer.GetComponentSpaceTransform(SubsampleTime, bUnused, LeftThighIdx);
+		const FTransform RightFootTransform = Indexer.GetComponentSpaceTransform(SubsampleTime, bUnused, RightFootIdx);
+		const FTransform LeftFootTransform = Indexer.GetComponentSpaceTransform(SubsampleTime, bUnused, LeftFootIdx);
 
 		const float CrashingLegsValue = ComputeCrashingLegsValue(RightThighTransform.GetTranslation(), LeftThighTransform.GetTranslation(), RightFootTransform.GetTranslation(), LeftFootTransform.GetTranslation());
 
diff --git a/Engine/Plugins/Experimental/Animation/PoseSearch/Source/Runtime/Private/PoseSearchFeatureChannel_Heading.cpp b/Engine/Plugins/Experimental/Animation/PoseSearch/Source/Runtime/Private/PoseSearchFeatureChannel_Heading.cpp
--- a/Engine/Plugins/Experimental/Animation/PoseSearch/Source/Runtime/Private/PoseSearchFeatureChannel_Heading.cpp
+++ b/Engine/Plugins/Experimental/Animation/PoseSearch/Source/Runtime/Private/PoseSearchFeatureChannel_Heading.cpp
@@ -57,7 +57,7 @@ void UPoseSearchFeatureChannel_Heading::IndexAsset(UE::PoseSearch::IAssetIndexer
 		const float SubsampleTime = OriginSampleTime + SampleTimeOffset;
 
 		bool ClampedPresent;
-		const FTransform BoneTransformsPresent = Indexer.GetTransformAndCacheResults(SubsampleTime, OriginSampleTime, SchemaBoneIdx, ClampedPresent);
+		const FTransform BoneTransformsPresent = Indexer.GetComponentSpaceTransform(SubsampleTime, OriginSampleTime, ClampedPresent, SchemaBoneIdx);
 		int32 DataOffset = ChannelDataOffset;
 		const FVector Heading = GetAxis(BoneTransformsPresent.GetRotation());
 		FFeatureVectorHelper::EncodeVector(IndexingContext.GetPoseVector(VectorIdx, FeatureVectorTable), DataOffset, Heading);
@@ -94,15 +94,15 @@ void UPoseSearchFeatureChannel_Heading::BuildQuery(UE::PoseSearch::FSearchContex
 			// so we can offset the Transform by this amount
 			if (!FMath::IsNearlyZero(SampleTimeOffset, UE_KINDA_SMALL_NUMBER))
 			{
-				const FTransform RootTransform = SearchContext.TryGetTransformAndCacheResults(0.f, InOutQuery.GetSchema(), FSearchContext::SchemaRootBoneIdx);
-				const FTransform RootTransformPrev = SearchContext.TryGetTransformAndCacheResults(SampleTimeOffset, InOutQuery.GetSchema(), FSearchContext::SchemaRootBoneIdx);
+				const FTransform RootTransform = SearchContext.TryGetTransformAndCacheResults(0.f, InOutQuery.GetSchema());
+				const FTransform RootTransformPrev = SearchContext.TryGetTransformAndCacheResults(SampleTimeOffset, InOutQuery.GetSchema());
 				Transform = Transform * (RootTransformPrev * RootTransform.Inverse());
 			}
 		}
 		else
 		{
 			check(SearchContext.Trajectory);
-			// @todo: make this call consistent with Transform = SearchContext.TryGetTransformAndCacheResults(SampleTimeOffset, InOutQuery.GetSchema(), FSearchContext::SchemaRootBoneIdx);
+			// @todo: make this call consistent with Transform = SearchContext.TryGetTransformAndCacheResults(SampleTimeOffset, InOutQuery.GetSchema());
 			const FTrajectorySample TrajectorySample = SearchContext.Trajectory->GetSampleAtTime(SampleTimeOffset);
 			Transform = TrajectorySample.Transform;
 		}
diff --git a/Engine/Plugins/Experimental/Animation/PoseSearch/Source/Runtime/Private/PoseSearchFeatureChannel_Heading.h b/Engine/Plugins/Experimental/Animation/PoseSearch/Source/Runtime/Private/PoseSearchFeatureChannel_Heading.h
--- a/Engine/Plugins/Experimental/Animation/PoseSearch/Source/Runtime/Private/PoseSearchFeatureChannel_Heading.h
+++ b/Engine/Plugins/Experimental/Animation/PoseSearch/Source/Runtime/Private/PoseSearchFeatureChannel_Heading.h
@@ -35,7 +35,7 @@ public:
 	UPROPERTY(EditAnywhere, Category = "Settings")
 	EHeadingAxis HeadingAxis = EHeadingAxis::X;	
 
-	UPROPERTY()
+	UPROPERTY(Transient)
 	int8 SchemaBoneIdx = 0;
 
 	UPROPERTY(EditAnywhere, Category = "Settings", meta = (ExcludeFromHash))
diff --git a/Engine/Plugins/Experimental/Animation/PoseSearch/Source/Runtime/Private/PoseSearchFeatureChannel_Phase.cpp b/Engine/Plugins/Experimental/Animation/PoseSearch/Source/Runtime/Private/PoseSearchFeatureChannel_Phase.cpp
--- a/Engine/Plugins/Experimental/Animation/PoseSearch/Source/Runtime/Private/PoseSearchFeatureChannel_Phase.cpp
+++ b/Engine/Plugins/Experimental/Animation/PoseSearch/Source/Runtime/Private/PoseSearchFeatureChannel_Phase.cpp
@@ -54,7 +54,7 @@ namespace UE::PoseSearch
 		{
 			const float SampleTime = SampleTimeStart + SampleIdx * FiniteDelta;
 			bool bUnused;
-			BonePositions[SampleIdx] = Indexer.GetTransformAndCacheResults(SampleTime, SampleTimeStart, SchemaBoneIdx, bUnused).GetTranslation();
+			BonePositions[SampleIdx] = Indexer.GetComponentSpaceTransform(SampleTime, SampleTimeStart, bUnused, SchemaBoneIdx).GetTranslation();
 		}
 	}
 
diff --git a/Engine/Plugins/Experimental/Animation/PoseSearch/Source/Runtime/Private/PoseSearchFeatureChannel_Phase.h b/Engine/Plugins/Experimental/Animation/PoseSearch/Source/Runtime/Private/PoseSearchFeatureChannel_Phase.h
--- a/Engine/Plugins/Experimental/Animation/PoseSearch/Source/Runtime/Private/PoseSearchFeatureChannel_Phase.h
+++ b/Engine/Plugins/Experimental/Animation/PoseSearch/Source/Runtime/Private/PoseSearchFeatureChannel_Phase.h
@@ -18,7 +18,7 @@ public:
 	UPROPERTY(EditAnywhere, Category = "Settings")
 	float Weight = 1.f;
 
-	UPROPERTY()
+	UPROPERTY(Transient)
 	int8 SchemaBoneIdx = 0;
 
 	UPROPERTY(EditAnywhere, Category = "Settings", meta = (ExcludeFromHash))
diff --git a/Engine/Plugins/Experimental/Animation/PoseSearch/Source/Runtime/Private/PoseSearchFeatureChannel_Position.cpp b/Engine/Plugins/Experimental/Animation/PoseSearch/Source/Runtime/Private/PoseSearchFeatureChannel_Position.cpp
--- a/Engine/Plugins/Experimental/Animation/PoseSearch/Source/Runtime/Private/PoseSearchFeatureChannel_Position.cpp
+++ b/Engine/Plugins/Experimental/Animation/PoseSearch/Source/Runtime/Private/PoseSearchFeatureChannel_Position.cpp
@@ -42,7 +42,7 @@ void UPoseSearchFeatureChannel_Position::IndexAsset(UE::PoseSearch::IAssetIndexe
 		const float SubsampleTime = OriginSampleTime + SampleTimeOffset;
 
 		bool ClampedPresent;
-		const FTransform BoneTransformsPresent = Indexer.GetTransformAndCacheResults(SubsampleTime, OriginSampleTime, SchemaBoneIdx, ClampedPresent);
+		const FTransform BoneTransformsPresent = Indexer.GetComponentSpaceTransform(SubsampleTime, OriginSampleTime, ClampedPresent, SchemaBoneIdx);
 		int32 DataOffset = ChannelDataOffset;
 		FFeatureVectorHelper::EncodeVector(IndexingContext.GetPoseVector(VectorIdx, FeatureVectorTable), DataOffset, BoneTransformsPresent.GetTranslation());
 	}
@@ -78,15 +78,15 @@ void UPoseSearchFeatureChannel_Position::BuildQuery(UE::PoseSearch::FSearchConte
 			// so we can offset the Transform by this amount
 			if (!FMath::IsNearlyZero(SampleTimeOffset, UE_KINDA_SMALL_NUMBER))
 			{
-				const FTransform RootTransform = SearchContext.TryGetTransformAndCacheResults(0.f, InOutQuery.GetSchema(), FSearchContext::SchemaRootBoneIdx);
-				const FTransform RootTransformPrev = SearchContext.TryGetTransformAndCacheResults(SampleTimeOffset, InOutQuery.GetSchema(), FSearchContext::SchemaRootBoneIdx);
+				const FTransform RootTransform = SearchContext.TryGetTransformAndCacheResults(0.f, InOutQuery.GetSchema());
+				const FTransform RootTransformPrev = SearchContext.TryGetTransformAndCacheResults(SampleTimeOffset, InOutQuery.GetSchema());
 				Transform = Transform * (RootTransformPrev * RootTransform.Inverse());
 			}
 		}
 		else
 		{
 			check(SearchContext.Trajectory);
-			// @todo: make this call consistent with Transform = SearchContext.TryGetTransformAndCacheResults(SampleTimeOffset, InOutQuery.GetSchema(), FSearchContext::SchemaRootBoneIdx);
+			// @todo: make this call consistent with Transform = SearchContext.TryGetTransformAndCacheResults(SampleTimeOffset, InOutQuery.GetSchema());
 			const FTrajectorySample TrajectorySample = SearchContext.Trajectory->GetSampleAtTime(SampleTimeOffset);
 			Transform = TrajectorySample.Transform;
 		}
diff --git a/Engine/Plugins/Experimental/Animation/PoseSearch/Source/Runtime/Private/PoseSearchFeatureChannel_Position.h b/Engine/Plugins/Experimental/Animation/PoseSearch/Source/Runtime/Private/PoseSearchFeatureChannel_Position.h
--- a/Engine/Plugins/Experimental/Animation/PoseSearch/Source/Runtime/Private/PoseSearchFeatureChannel_Position.h
+++ b/Engine/Plugins/Experimental/Animation/PoseSearch/Source/Runtime/Private/PoseSearchFeatureChannel_Position.h
@@ -21,7 +21,7 @@ public:
 	UPROPERTY(EditAnywhere, Category = "Settings")
 	float SampleTimeOffset = 0.f;
 
-	UPROPERTY(meta = (ExcludeFromHash))
+	UPROPERTY(Transient)
 	int8 SchemaBoneIdx = 0;
 
 	UPROPERTY(EditAnywhere, Category = "Settings", meta = (ExcludeFromHash))
diff --git a/Engine/Plugins/Experimental/Animation/PoseSearch/Source/Runtime/Private/PoseSearchFeatureChannel_Velocity.cpp b/Engine/Plugins/Experimental/Animation/PoseSearch/Source/Runtime/Private/PoseSearchFeatureChannel_Velocity.cpp
--- a/Engine/Plugins/Experimental/Animation/PoseSearch/Source/Runtime/Private/PoseSearchFeatureChannel_Velocity.cpp
+++ b/Engine/Plugins/Experimental/Animation/PoseSearch/Source/Runtime/Private/PoseSearchFeatureChannel_Velocity.cpp
@@ -45,9 +45,9 @@ void UPoseSearchFeatureChannel_Velocity::IndexAsset(UE::PoseSearch::IAssetIndexe
 		const float SubsampleTime = OriginSampleTime + SampleTimeOffset;
 
 		bool ClampedPast, ClampedPresent, ClampedFuture;
-		const FTransform BoneTransformsPast = Indexer.GetTransformAndCacheResults(SubsampleTime - SamplingContext->FiniteDelta, bUseCharacterSpaceVelocities ? OriginSampleTime - SamplingContext->FiniteDelta : OriginSampleTime, SchemaBoneIdx, ClampedPast);
-		const FTransform BoneTransformsPresent = Indexer.GetTransformAndCacheResults(SubsampleTime, OriginSampleTime, SchemaBoneIdx, ClampedPresent);
-		const FTransform BoneTransformsFuture = Indexer.GetTransformAndCacheResults(SubsampleTime + SamplingContext->FiniteDelta, bUseCharacterSpaceVelocities ? OriginSampleTime + SamplingContext->FiniteDelta : OriginSampleTime, SchemaBoneIdx, ClampedFuture);
+		const FTransform BoneTransformsPast = Indexer.GetComponentSpaceTransform(SubsampleTime - SamplingContext->FiniteDelta, bUseCharacterSpaceVelocities ? OriginSampleTime - SamplingContext->FiniteDelta : OriginSampleTime, ClampedPast, SchemaBoneIdx);
+		const FTransform BoneTransformsPresent = Indexer.GetComponentSpaceTransform(SubsampleTime, OriginSampleTime, ClampedPresent, SchemaBoneIdx);
+		const FTransform BoneTransformsFuture = Indexer.GetComponentSpaceTransform(SubsampleTime + SamplingContext->FiniteDelta, bUseCharacterSpaceVelocities ? OriginSampleTime + SamplingContext->FiniteDelta : OriginSampleTime, ClampedFuture, SchemaBoneIdx);
 
 		// We can get a better finite difference if we ignore samples that have
 		// been clamped at either side of the clip. However, if the central sample 
@@ -112,8 +112,8 @@ void UPoseSearchFeatureChannel_Velocity::BuildQuery(UE::PoseSearch::FSearchConte
 			// if we want the velocity not in character space we need to calculate the root offset delta between the SampleTimeOffset and the SampleTimeOffset - HistorySameplInterval to apply to the TransformPrevious
 			if (!bUseCharacterSpaceVelocities)
 			{
-				const FTransform RootTransform = SearchContext.TryGetTransformAndCacheResults(SampleTimeOffset, InOutQuery.GetSchema(), FSearchContext::SchemaRootBoneIdx);
-				const FTransform RootTransformPrev = SearchContext.TryGetTransformAndCacheResults(SampleTimeOffset - HistorySameplInterval, InOutQuery.GetSchema(), FSearchContext::SchemaRootBoneIdx);
+				const FTransform RootTransform = SearchContext.TryGetTransformAndCacheResults(SampleTimeOffset, InOutQuery.GetSchema());
+				const FTransform RootTransformPrev = SearchContext.TryGetTransformAndCacheResults(SampleTimeOffset - HistorySameplInterval, InOutQuery.GetSchema());
 				TransformPrevious = TransformPrevious * (RootTransformPrev * RootTransform.Inverse());
 			}
 
@@ -122,7 +122,7 @@ void UPoseSearchFeatureChannel_Velocity::BuildQuery(UE::PoseSearch::FSearchConte
 		else
 		{
 			check(SearchContext.Trajectory);
-			// @todo: make this call consistent with Transform = SearchContext.TryGetTransformAndCacheResults(SampleTimeOffset, InOutQuery.GetSchema(), FSearchContext::SchemaRootBoneIdx);
+			// @todo: make this call consistent with Transform = SearchContext.TryGetTransformAndCacheResults(SampleTimeOffset, InOutQuery.GetSchema());
 			const FTrajectorySample TrajectorySample = SearchContext.Trajectory->GetSampleAtTime(SampleTimeOffset);
 			LinearVelocity = TrajectorySample.LinearVelocity;
 		}
diff --git a/Engine/Plugins/Experimental/Animation/PoseSearch/Source/Runtime/Private/PoseSearchFeatureChannel_Velocity.h b/Engine/Plugins/Experimental/Animation/PoseSearch/Source/Runtime/Private/PoseSearchFeatureChannel_Velocity.h
--- a/Engine/Plugins/Experimental/Animation/PoseSearch/Source/Runtime/Private/PoseSearchFeatureChannel_Velocity.h
+++ b/Engine/Plugins/Experimental/Animation/PoseSearch/Source/Runtime/Private/PoseSearchFeatureChannel_Velocity.h
@@ -21,7 +21,7 @@ public:
 	UPROPERTY(EditAnywhere, Category = "Settings")
 	float SampleTimeOffset = 0.f;
 
-	UPROPERTY(meta = (ExcludeFromHash))
+	UPROPERTY(Transient)
 	int8 SchemaBoneIdx = 0;
 
 	UPROPERTY(EditAnywhere, Category = "Settings", meta = (ExcludeFromHash))
diff --git a/Engine/Plugins/Experimental/Animation/PoseSearch/Source/Runtime/Private/PoseSearchSchema.cpp b/Engine/Plugins/Experimental/Animation/PoseSearch/Source/Runtime/Private/PoseSearchSchema.cpp
--- a/Engine/Plugins/Experimental/Animation/PoseSearch/Source/Runtime/Private/PoseSearchSchema.cpp
+++ b/Engine/Plugins/Experimental/Animation/PoseSearch/Source/Runtime/Private/PoseSearchSchema.cpp
@@ -2,6 +2,7 @@
 
 #include "PoseSearch/PoseSearchSchema.h"
 #include "AnimationRuntime.h"
+#include "PoseSearch/PoseSearchDefines.h"
 #include "PoseSearch/PoseSearchFeatureChannel.h"
 #include "PoseSearch/PoseSearchResult.h"
 #include "UObject/ObjectSaveContext.h"
@@ -55,7 +56,7 @@ void UPoseSearchSchema::BuildQuery(UE::PoseSearch::FSearchContext& SearchContext
 
 FBoneIndexType UPoseSearchSchema::GetBoneIndexType(int8 SchemaBoneIdx) const
 {
-	return SchemaBoneIdx >= 0 && BoneReferences[SchemaBoneIdx].HasValidSetup() ? BoneReferences[SchemaBoneIdx].BoneIndex : RootBoneIdx;
+	return SchemaBoneIdx >= 0 && BoneReferences[SchemaBoneIdx].HasValidSetup() ? BoneReferences[SchemaBoneIdx].BoneIndex : RootBoneIndexType;
 }
 
 void UPoseSearchSchema::Finalize()
@@ -87,9 +88,11 @@ void UPoseSearchSchema::PostLoad()
 	Finalize();
 }
 
-int32 UPoseSearchSchema::AddBoneReference(const FBoneReference& BoneReference)
+int8 UPoseSearchSchema::AddBoneReference(const FBoneReference& BoneReference)
 {
-	return BoneReferences.AddUnique(BoneReference);
+	const int32 SchemaBoneIdx = BoneReferences.AddUnique(BoneReference);
+	check(SchemaBoneIdx >= 0 && SchemaBoneIdx < 128);
+	return int8(SchemaBoneIdx);
 }
 
 USkeleton* UPoseSearchSchema::GetSkeleton(bool& bInvalidSkeletonIsError, const IPropertyHandle* PropertyHandle)
diff --git a/Engine/Plugins/Experimental/Animation/PoseSearch/Source/Runtime/Public/PoseSearch/PoseSearchAssetIndexer.h b/Engine/Plugins/Experimental/Animation/PoseSearch/Source/Runtime/Public/PoseSearch/PoseSearchAssetIndexer.h
--- a/Engine/Plugins/Experimental/Animation/PoseSearch/Source/Runtime/Public/PoseSearch/PoseSearchAssetIndexer.h
+++ b/Engine/Plugins/Experimental/Animation/PoseSearch/Source/Runtime/Public/PoseSearch/PoseSearchAssetIndexer.h
@@ -3,6 +3,7 @@
 #pragma once
 
 #include "CoreMinimal.h"
+#include "PoseSearch/PoseSearchDefines.h"
 #include "PoseSearch/PoseSearchSchema.h"
 
 namespace UE::PoseSearch
@@ -49,10 +50,9 @@ public:
 	virtual ~IAssetIndexer() {}
 
 	virtual const FAssetIndexingContext& GetIndexingContext() const = 0;
-	virtual FSampleInfo GetSampleInfo(float SampleTime) const = 0;
-	virtual FSampleInfo GetSampleInfoRelative(float SampleTime, const FSampleInfo& Origin) const = 0;
-	virtual FTransform MirrorTransform(const FTransform& Transform) const = 0;
-	virtual FTransform GetTransformAndCacheResults(float SampleTime, float OriginTime, int8 SchemaBoneIdx, bool& Clamped) = 0;
+	virtual FTransform GetTransform(float SampleTime, bool& Clamped, int8 SchemaBoneIdx = RootSchemaBoneIdx) = 0;
+	virtual FTransform GetComponentSpaceTransform(float SampleTime, bool& Clamped, int8 SchemaBoneIdx = RootSchemaBoneIdx) = 0;
+	virtual FTransform GetComponentSpaceTransform(float SampleTime, float OriginTime, bool& Clamped, int8 SchemaBoneIdx = RootSchemaBoneIdx) = 0;
 };
 
 } // namespace UE::PoseSearch
diff --git a/Engine/Plugins/Experimental/Animation/PoseSearch/Source/Runtime/Public/PoseSearch/PoseSearchContext.h b/Engine/Plugins/Experimental/Animation/PoseSearch/Source/Runtime/Public/PoseSearch/PoseSearchContext.h
--- a/Engine/Plugins/Experimental/Animation/PoseSearch/Source/Runtime/Public/PoseSearch/PoseSearchContext.h
+++ b/Engine/Plugins/Experimental/Animation/PoseSearch/Source/Runtime/Public/PoseSearch/PoseSearchContext.h
@@ -162,7 +162,7 @@ struct POSESEARCH_API FSearchContext
 	// can the continuing pose advance? (if not we skip evaluating it)
 	bool bCanAdvance = true;
 
-	FTransform TryGetTransformAndCacheResults(float SampleTime, const UPoseSearchSchema* Schema, int8 SchemaBoneIdx);
+	FTransform TryGetTransformAndCacheResults(float SampleTime, const UPoseSearchSchema* Schema, int8 SchemaBoneIdx = RootSchemaBoneIdx);
 	void ClearCachedEntries();
 
 	void ResetCurrentBestCost();
@@ -178,8 +178,6 @@ struct POSESEARCH_API FSearchContext
 	TConstArrayView<float> GetCurrentResultPoseVector() const;
 	TConstArrayView<float> GetCurrentResultNextPoseVector() const;
 
-	static constexpr int8 SchemaRootBoneIdx = -1;
-
 	const FPoseIndicesHistory* PoseIndicesHistory = nullptr;
 
 private:
diff --git a/Engine/Plugins/Experimental/Animation/PoseSearch/Source/Runtime/Public/PoseSearch/PoseSearchDefines.h b/Engine/Plugins/Experimental/Animation/PoseSearch/Source/Runtime/Public/PoseSearch/PoseSearchDefines.h
--- a/Engine/Plugins/Experimental/Animation/PoseSearch/Source/Runtime/Public/PoseSearch/PoseSearchDefines.h
+++ b/Engine/Plugins/Experimental/Animation/PoseSearch/Source/Runtime/Public/PoseSearch/PoseSearchDefines.h
@@ -20,3 +20,8 @@ constexpr EParallelForFlags ParallelForFlags = EParallelForFlags::None;
 #endif // UE_POSE_SEARCH_FORCE_SINGLE_THREAD
 
 DECLARE_LOG_CATEGORY_EXTERN(LogPoseSearch, Log, All);
+
+static constexpr int8 RootSchemaBoneIdx = -1;
+static constexpr FBoneIndexType RootBoneIndexType = 0xFFFF;
+
+
diff --git a/Engine/Plugins/Experimental/Animation/PoseSearch/Source/Runtime/Public/PoseSearch/PoseSearchSchema.h b/Engine/Plugins/Experimental/Animation/PoseSearch/Source/Runtime/Public/PoseSearch/PoseSearchSchema.h
--- a/Engine/Plugins/Experimental/Animation/PoseSearch/Source/Runtime/Public/PoseSearch/PoseSearchSchema.h
+++ b/Engine/Plugins/Experimental/Animation/PoseSearch/Source/Runtime/Public/PoseSearch/PoseSearchSchema.h
@@ -113,7 +113,7 @@ public:
 	virtual void PreSave(FObjectPreSaveContext ObjectSaveContext) override;
 	virtual void PostLoad() override;
 
-	int32 AddBoneReference(const FBoneReference& BoneReference);
+	int8 AddBoneReference(const FBoneReference& BoneReference);
 
 	// IBoneReferenceSkeletonProvider
 	USkeleton* GetSkeleton(bool& bInvalidSkeletonIsError, const IPropertyHandle* PropertyHandle) override;
@@ -124,7 +124,6 @@ public:
 
 	void BuildQuery(UE::PoseSearch::FSearchContext& SearchContext, FPoseSearchFeatureVectorBuilder& InOutQuery) const;
 
-	static constexpr FBoneIndexType RootBoneIdx = 0xFFFF;
 	FBoneIndexType GetBoneIndexType(int8 SchemaBoneIdx) const;
 
 private:
