motion matching - build query API cleanup (to mimic the editor side indexing of FAssetIndexer)

diff --git a/Engine/Plugins/Experimental/Animation/PoseSearch/Source/Runtime/Private/PoseSearchContext.cpp b/Engine/Plugins/Experimental/Animation/PoseSearch/Source/Runtime/Private/PoseSearchContext.cpp
--- a/Engine/Plugins/Experimental/Animation/PoseSearch/Source/Runtime/Private/PoseSearchContext.cpp
+++ b/Engine/Plugins/Experimental/Animation/PoseSearch/Source/Runtime/Private/PoseSearchContext.cpp
@@ -146,18 +146,36 @@ void DrawFeatureVector(FDebugDrawParams& DrawParams, int32 PoseIdx)
 
 //////////////////////////////////////////////////////////////////////////
 // FSearchContext
-FTransform FSearchContext::TryGetTransformAndCacheResults(float SampleTime, const UPoseSearchSchema* Schema, int8 SchemaBoneIdx)
+
+FTransform FSearchContext::GetTransform(float SampleTime, const UPoseSearchSchema* Schema, int8 SchemaBoneIdx)
 {
-	check(History && Schema);
+	// collecting the RootTransform from the FPoseHistory
+	FTransform RootTransform = FTransform::Identity;
+	History->TrySampleLocalPose(-SampleTime, nullptr, nullptr, &RootTransform);
 
 	const FBoneIndexType BoneIndexType = Schema->GetBoneIndexType(SchemaBoneIdx);
-	if (const FCachedTransform<FTransform>* CachedTransform = CachedTransforms.Find(SampleTime, BoneIndexType))
+	if (BoneIndexType != RootBoneIndexType)
 	{
-		return CachedTransform->Transform;
+		const FTransform BoneTransform = GetComponentSpaceTransform(SampleTime, Schema, SchemaBoneIdx);
+		return BoneTransform * RootTransform;
 	}
 
+	return RootTransform;
+}
+
+FTransform FSearchContext::GetComponentSpaceTransform(float SampleTime, const UPoseSearchSchema* Schema, int8 SchemaBoneIdx)
+{
+	check(History && Schema);
+
+	const FBoneIndexType BoneIndexType = Schema->GetBoneIndexType(SchemaBoneIdx);
 	if (BoneIndexType != RootBoneIndexType)
 	{
+		if (const FCachedTransform<FTransform>* CachedTransform = CachedTransforms.Find(SampleTime, BoneIndexType))
+		{
+			return CachedTransform->Transform;
+		}
+	
+		// collecting the local bone transforms from the FPoseHistory
 		TArray<FTransform> SampledLocalPose;
 		if (History->TrySampleLocalPose(-SampleTime, &Schema->BoneIndicesWithParents, &SampledLocalPose, nullptr))
 		{
@@ -172,18 +190,22 @@ FTransform FSearchContext::TryGetTransformAndCacheResults(float SampleTime, cons
 
 			return SampledComponentPose[BoneIndexType];
 		}
-
-		return FTransform::Identity;
 	}
-	
-	FTransform SampledRootTransform;
-	if (History->TrySampleLocalPose(-SampleTime, nullptr, nullptr, &SampledRootTransform))
+
+	return FTransform::Identity;
+}
+
+FTransform FSearchContext::GetComponentSpaceTransform(float SampleTime, float OriginTime, const UPoseSearchSchema* Schema, int8 SchemaBoneIdx)
+{
+	if (SampleTime == OriginTime)
 	{
-		CachedTransforms.Add(SampleTime, BoneIndexType, SampledRootTransform);
-		return SampledRootTransform;
+		return GetComponentSpaceTransform(SampleTime, Schema, SchemaBoneIdx);
 	}
-	
-	return FTransform::Identity;
+
+	const FTransform RootBoneTransform = GetTransform(OriginTime, Schema);
+	FTransform BoneTransform = GetTransform(SampleTime, Schema, SchemaBoneIdx);
+	BoneTransform.SetToRelativeTransform(RootBoneTransform);
+	return BoneTransform;
 }
 
 void FSearchContext::ClearCachedEntries()
diff --git a/Engine/Plugins/Experimental/Animation/PoseSearch/Source/Runtime/Private/PoseSearchFeatureChannel_FilterCrashingLegs.cpp b/Engine/Plugins/Experimental/Animation/PoseSearch/Source/Runtime/Private/PoseSearchFeatureChannel_FilterCrashingLegs.cpp
--- a/Engine/Plugins/Experimental/Animation/PoseSearch/Source/Runtime/Private/PoseSearchFeatureChannel_FilterCrashingLegs.cpp
+++ b/Engine/Plugins/Experimental/Animation/PoseSearch/Source/Runtime/Private/PoseSearchFeatureChannel_FilterCrashingLegs.cpp
@@ -89,10 +89,10 @@ void UPoseSearchFeatureChannel_FilterCrashingLegs::BuildQuery(UE::PoseSearch::FS
 	else
 	{
 		const float SampleTime = 0.f;
-		const FTransform LeftThighTransform = SearchContext.TryGetTransformAndCacheResults(SampleTime, InOutQuery.GetSchema(), LeftThighIdx);
-		const FTransform RightThighTransform = SearchContext.TryGetTransformAndCacheResults(SampleTime, InOutQuery.GetSchema(), RightThighIdx);
-		const FTransform LeftFootTransform = SearchContext.TryGetTransformAndCacheResults(SampleTime, InOutQuery.GetSchema(), LeftFootIdx);
-		const FTransform RightFootTransform = SearchContext.TryGetTransformAndCacheResults(SampleTime, InOutQuery.GetSchema(), RightFootIdx);
+		const FTransform LeftThighTransform = SearchContext.GetComponentSpaceTransform(SampleTime, InOutQuery.GetSchema(), LeftThighIdx);
+		const FTransform RightThighTransform = SearchContext.GetComponentSpaceTransform(SampleTime, InOutQuery.GetSchema(), RightThighIdx);
+		const FTransform LeftFootTransform = SearchContext.GetComponentSpaceTransform(SampleTime, InOutQuery.GetSchema(), LeftFootIdx);
+		const FTransform RightFootTransform = SearchContext.GetComponentSpaceTransform(SampleTime, InOutQuery.GetSchema(), RightFootIdx);
 
 		const float CrashingLegsValue = ComputeCrashingLegsValue(RightThighTransform.GetTranslation(), LeftThighTransform.GetTranslation(), RightFootTransform.GetTranslation(), LeftFootTransform.GetTranslation());
 		FFeatureVectorHelper::EncodeFloat(InOutQuery.EditValues(), DataOffset, CrashingLegsValue);
diff --git a/Engine/Plugins/Experimental/Animation/PoseSearch/Source/Runtime/Private/PoseSearchFeatureChannel_Heading.cpp b/Engine/Plugins/Experimental/Animation/PoseSearch/Source/Runtime/Private/PoseSearchFeatureChannel_Heading.cpp
--- a/Engine/Plugins/Experimental/Animation/PoseSearch/Source/Runtime/Private/PoseSearchFeatureChannel_Heading.cpp
+++ b/Engine/Plugins/Experimental/Animation/PoseSearch/Source/Runtime/Private/PoseSearchFeatureChannel_Heading.cpp
@@ -88,16 +88,7 @@ void UPoseSearchFeatureChannel_Heading::BuildQuery(UE::PoseSearch::FSearchContex
 		if (bBoneValid)
 		{
 			// calculating the Transform in component space for the bone indexed by SchemaBoneIdx
-			Transform = SearchContext.TryGetTransformAndCacheResults(SampleTimeOffset, InOutQuery.GetSchema(), SchemaBoneIdx);
-
-			// if the SampleTimeOffset is not zero we calculate the delta root bone between the root at current time (zero) and the root at SampleTimeOffset (in the past)
-			// so we can offset the Transform by this amount
-			if (!FMath::IsNearlyZero(SampleTimeOffset, UE_KINDA_SMALL_NUMBER))
-			{
-				const FTransform RootTransform = SearchContext.TryGetTransformAndCacheResults(0.f, InOutQuery.GetSchema());
-				const FTransform RootTransformPrev = SearchContext.TryGetTransformAndCacheResults(SampleTimeOffset, InOutQuery.GetSchema());
-				Transform = Transform * (RootTransformPrev * RootTransform.Inverse());
-			}
+			Transform = SearchContext.GetComponentSpaceTransform(SampleTimeOffset, 0.f, InOutQuery.GetSchema(), SchemaBoneIdx);
 		}
 		else
 		{
diff --git a/Engine/Plugins/Experimental/Animation/PoseSearch/Source/Runtime/Private/PoseSearchFeatureChannel_Position.cpp b/Engine/Plugins/Experimental/Animation/PoseSearch/Source/Runtime/Private/PoseSearchFeatureChannel_Position.cpp
--- a/Engine/Plugins/Experimental/Animation/PoseSearch/Source/Runtime/Private/PoseSearchFeatureChannel_Position.cpp
+++ b/Engine/Plugins/Experimental/Animation/PoseSearch/Source/Runtime/Private/PoseSearchFeatureChannel_Position.cpp
@@ -72,16 +72,7 @@ void UPoseSearchFeatureChannel_Position::BuildQuery(UE::PoseSearch::FSearchConte
 		if (bBoneValid)
 		{
 			// calculating the Transform in component space for the bone indexed by SchemaBoneIdx
-			Transform = SearchContext.TryGetTransformAndCacheResults(SampleTimeOffset, InOutQuery.GetSchema(), SchemaBoneIdx);
-
-			// if the SampleTimeOffset is not zero we calculate the delta root bone between the root at current time (zero) and the root at SampleTimeOffset (in the past)
-			// so we can offset the Transform by this amount
-			if (!FMath::IsNearlyZero(SampleTimeOffset, UE_KINDA_SMALL_NUMBER))
-			{
-				const FTransform RootTransform = SearchContext.TryGetTransformAndCacheResults(0.f, InOutQuery.GetSchema());
-				const FTransform RootTransformPrev = SearchContext.TryGetTransformAndCacheResults(SampleTimeOffset, InOutQuery.GetSchema());
-				Transform = Transform * (RootTransformPrev * RootTransform.Inverse());
-			}
+			Transform = SearchContext.GetComponentSpaceTransform(SampleTimeOffset, 0.f, InOutQuery.GetSchema(), SchemaBoneIdx);
 		}
 		else
 		{
diff --git a/Engine/Plugins/Experimental/Animation/PoseSearch/Source/Runtime/Private/PoseSearchFeatureChannel_Velocity.cpp b/Engine/Plugins/Experimental/Animation/PoseSearch/Source/Runtime/Private/PoseSearchFeatureChannel_Velocity.cpp
--- a/Engine/Plugins/Experimental/Animation/PoseSearch/Source/Runtime/Private/PoseSearchFeatureChannel_Velocity.cpp
+++ b/Engine/Plugins/Experimental/Animation/PoseSearch/Source/Runtime/Private/PoseSearchFeatureChannel_Velocity.cpp
@@ -106,16 +106,8 @@ void UPoseSearchFeatureChannel_Velocity::BuildQuery(UE::PoseSearch::FSearchConte
 
 			// @todo: optimize this code for SchemaBoneIdx == FSearchContext::SchemaRootBoneIdx (!InOutQuery.GetSchema()->BoneReferences[SchemaBoneIdx].HasValidSetup())
 			// calculating the Transforms in component space for the bone indexed by SchemaBoneIdx
-			FTransform TransformCurrent = SearchContext.TryGetTransformAndCacheResults(SampleTimeOffset, InOutQuery.GetSchema(), SchemaBoneIdx);
-			FTransform TransformPrevious = SearchContext.TryGetTransformAndCacheResults(SampleTimeOffset - HistorySameplInterval, InOutQuery.GetSchema(), SchemaBoneIdx);
-
-			// if we want the velocity not in character space we need to calculate the root offset delta between the SampleTimeOffset and the SampleTimeOffset - HistorySameplInterval to apply to the TransformPrevious
-			if (!bUseCharacterSpaceVelocities)
-			{
-				const FTransform RootTransform = SearchContext.TryGetTransformAndCacheResults(SampleTimeOffset, InOutQuery.GetSchema());
-				const FTransform RootTransformPrev = SearchContext.TryGetTransformAndCacheResults(SampleTimeOffset - HistorySameplInterval, InOutQuery.GetSchema());
-				TransformPrevious = TransformPrevious * (RootTransformPrev * RootTransform.Inverse());
-			}
+			const FTransform TransformCurrent = SearchContext.GetComponentSpaceTransform(SampleTimeOffset, 0.f, InOutQuery.GetSchema(), SchemaBoneIdx);
+			const FTransform TransformPrevious = SearchContext.GetComponentSpaceTransform(SampleTimeOffset - HistorySameplInterval, -HistorySameplInterval, InOutQuery.GetSchema(), SchemaBoneIdx);
 
 			LinearVelocity = (TransformCurrent.GetTranslation() - TransformPrevious.GetTranslation()) / HistorySameplInterval;
 		}
diff --git a/Engine/Plugins/Experimental/Animation/PoseSearch/Source/Runtime/Public/PoseSearch/PoseSearchContext.h b/Engine/Plugins/Experimental/Animation/PoseSearch/Source/Runtime/Public/PoseSearch/PoseSearchContext.h
--- a/Engine/Plugins/Experimental/Animation/PoseSearch/Source/Runtime/Public/PoseSearch/PoseSearchContext.h
+++ b/Engine/Plugins/Experimental/Animation/PoseSearch/Source/Runtime/Public/PoseSearch/PoseSearchContext.h
@@ -162,7 +162,10 @@ struct POSESEARCH_API FSearchContext
 	// can the continuing pose advance? (if not we skip evaluating it)
 	bool bCanAdvance = true;
 
-	FTransform TryGetTransformAndCacheResults(float SampleTime, const UPoseSearchSchema* Schema, int8 SchemaBoneIdx = RootSchemaBoneIdx);
+	FTransform GetTransform(float SampleTime, const UPoseSearchSchema* Schema, int8 SchemaBoneIdx = RootSchemaBoneIdx);
+	FTransform GetComponentSpaceTransform(float SampleTime, const UPoseSearchSchema* Schema, int8 SchemaBoneIdx = RootSchemaBoneIdx);
+	FTransform GetComponentSpaceTransform(float SampleTime, float OriginTime, const UPoseSearchSchema* Schema, int8 SchemaBoneIdx = RootSchemaBoneIdx);
+
 	void ClearCachedEntries();
 
 	void ResetCurrentBestCost();
@@ -181,6 +184,7 @@ struct POSESEARCH_API FSearchContext
 	const FPoseIndicesHistory* PoseIndicesHistory = nullptr;
 
 private:
+	// transforms cached in component space
 	FCachedTransforms<FTransform> CachedTransforms;
 	
 	struct FCachedQuery
