motion matching - cached query as functions of schemas, not databases

diff --git a/Engine/Plugins/Experimental/Animation/PoseSearch/Source/Runtime/Private/PoseSearchContext.cpp b/Engine/Plugins/Experimental/Animation/PoseSearch/Source/Runtime/Private/PoseSearchContext.cpp
--- a/Engine/Plugins/Experimental/Animation/PoseSearch/Source/Runtime/Private/PoseSearchContext.cpp
+++ b/Engine/Plugins/Experimental/Animation/PoseSearch/Source/Runtime/Private/PoseSearchContext.cpp
@@ -228,33 +228,33 @@ void FSearchContext::UpdateCurrentBestCost(const FPoseSearchCost& PoseSearchCost
 	};
 }
 
-const FPoseSearchFeatureVectorBuilder* FSearchContext::GetCachedQuery(const UPoseSearchDatabase* Database) const
+const FPoseSearchFeatureVectorBuilder* FSearchContext::GetCachedQuery(const UPoseSearchSchema* Schema) const
 {
-	const FSearchContext::FCachedQuery* CachedQuery = CachedQueries.FindByPredicate([Database](const FSearchContext::FCachedQuery& CachedQuery)
+	const FPoseSearchFeatureVectorBuilder* CachedQuery = CachedQueries.FindByPredicate([Schema](const FPoseSearchFeatureVectorBuilder& CachedQuery)
 	{
-		return CachedQuery.Database == Database;
+		return CachedQuery.GetSchema() == Schema;
 	});
 
 	if (CachedQuery)
 	{
-		return &CachedQuery->FeatureVectorBuilder;
+		return CachedQuery;
 	}
 	return nullptr;
 }
 
-void FSearchContext::GetOrBuildQuery(const UPoseSearchDatabase* Database, FPoseSearchFeatureVectorBuilder& FeatureVectorBuilder)
+void FSearchContext::GetOrBuildQuery(const UPoseSearchSchema* Schema, FPoseSearchFeatureVectorBuilder& FeatureVectorBuilder)
 {
-	const FPoseSearchFeatureVectorBuilder* CachedFeatureVectorBuilder = GetCachedQuery(Database);
+	check(Schema && Schema->IsValid());
+	const FPoseSearchFeatureVectorBuilder* CachedFeatureVectorBuilder = GetCachedQuery(Schema);
 	if (CachedFeatureVectorBuilder)
 	{
 		FeatureVectorBuilder = *CachedFeatureVectorBuilder;
 	}
 	else
 	{
-		FSearchContext::FCachedQuery& NewCachedQuery = CachedQueries[CachedQueries.AddDefaulted()];
-		NewCachedQuery.Database = Database;
-		Database->BuildQuery(*this, NewCachedQuery.FeatureVectorBuilder);
-		FeatureVectorBuilder = NewCachedQuery.FeatureVectorBuilder;
+		FPoseSearchFeatureVectorBuilder& NewCachedQuery = CachedQueries[CachedQueries.AddDefaulted()];
+		Schema->BuildQuery(*this, NewCachedQuery);
+		FeatureVectorBuilder = NewCachedQuery;
 	}
 }
 
diff --git a/Engine/Plugins/Experimental/Animation/PoseSearch/Source/Runtime/Private/PoseSearchDatabase.cpp b/Engine/Plugins/Experimental/Animation/PoseSearch/Source/Runtime/Private/PoseSearchDatabase.cpp
--- a/Engine/Plugins/Experimental/Animation/PoseSearch/Source/Runtime/Private/PoseSearchDatabase.cpp
+++ b/Engine/Plugins/Experimental/Animation/PoseSearch/Source/Runtime/Private/PoseSearchDatabase.cpp
@@ -719,7 +719,7 @@ UE::PoseSearch::FSearchResult UPoseSearchDatabase::SearchPCAKDTree(UE::PoseSearc
 		}
 	}
 
-	SearchContext.GetOrBuildQuery(this, Result.ComposedQuery);
+	SearchContext.GetOrBuildQuery(Schema, Result.ComposedQuery);
 
 	TConstArrayView<float> QueryValues = Result.ComposedQuery.GetValues();
 
@@ -819,7 +819,7 @@ UE::PoseSearch::FSearchResult UPoseSearchDatabase::SearchBruteForce(UE::PoseSear
 
 	const FPoseSearchIndex& SearchIndex = GetSearchIndex();
 
-	SearchContext.GetOrBuildQuery(this, Result.ComposedQuery);
+	SearchContext.GetOrBuildQuery(Schema, Result.ComposedQuery);
 	TConstArrayView<float> QueryValues = Result.ComposedQuery.GetValues();
 
 	const bool IsCurrentResultFromThisDatabase = SearchContext.IsCurrentResultFromDatabase(this);
@@ -898,10 +898,4 @@ UE::PoseSearch::FSearchResult UPoseSearchDatabase::SearchBruteForce(UE::PoseSear
 #endif
 
 	return Result;
-}
-
-void UPoseSearchDatabase::BuildQuery(UE::PoseSearch::FSearchContext& SearchContext, FPoseSearchFeatureVectorBuilder& OutQuery) const
-{
-	check(Schema && Schema->IsValid());
-	Schema->BuildQuery(SearchContext, OutQuery);
 }
\ No newline at end of file
diff --git a/Engine/Plugins/Experimental/Animation/PoseSearch/Source/Runtime/Private/PoseSearchDatabaseSet.cpp b/Engine/Plugins/Experimental/Animation/PoseSearch/Source/Runtime/Private/PoseSearchDatabaseSet.cpp
--- a/Engine/Plugins/Experimental/Animation/PoseSearch/Source/Runtime/Private/PoseSearchDatabaseSet.cpp
+++ b/Engine/Plugins/Experimental/Animation/PoseSearch/Source/Runtime/Private/PoseSearchDatabaseSet.cpp
@@ -29,7 +29,7 @@ UE::PoseSearch::FSearchResult UPoseSearchDatabaseSet::Search(UE::PoseSearch::FSe
 	{
 		check(Database);
 		const FPoseSearchIndex& SearchIndex = Database->GetSearchIndex();
-		SearchContext.GetOrBuildQuery(Database, Result.ComposedQuery);
+		SearchContext.GetOrBuildQuery(Database->Schema, Result.ComposedQuery);
 
 		TConstArrayView<float> QueryValues = Result.ComposedQuery.GetValues();
 
diff --git a/Engine/Plugins/Experimental/Animation/PoseSearch/Source/Runtime/Private/PoseSearchLibrary.cpp b/Engine/Plugins/Experimental/Animation/PoseSearch/Source/Runtime/Private/PoseSearchLibrary.cpp
--- a/Engine/Plugins/Experimental/Animation/PoseSearch/Source/Runtime/Private/PoseSearchLibrary.cpp
+++ b/Engine/Plugins/Experimental/Animation/PoseSearch/Source/Runtime/Private/PoseSearchLibrary.cpp
@@ -178,7 +178,7 @@ static void TraceMotionMatchingState(
 			// if throttling is on, the continuing pose can be valid, but no actual search occurred, so the query will not be cached, and we need to build it
 			FPoseSearchFeatureVectorBuilder FeatureVectorBuilder;
 			FeatureVectorBuilder.Init(Database->Schema);
-			SearchContext.GetOrBuildQuery(Database, FeatureVectorBuilder);
+			SearchContext.GetOrBuildQuery(Database->Schema, FeatureVectorBuilder);
 			DatabaseEntries[DbEntryIdx].QueryVector = FeatureVectorBuilder.GetValues();
 		}
 
@@ -372,7 +372,7 @@ void FMotionMatchingState::UpdateWantedPlayRate(const UE::PoseSearch::FSearchCon
 	{
 		if (!FMath::IsNearlyEqual(Settings.PlayRateMin, 1.f, UE_KINDA_SMALL_NUMBER) || !FMath::IsNearlyEqual(Settings.PlayRateMax, 1.f, UE_KINDA_SMALL_NUMBER))
 		{
-			if (const FPoseSearchFeatureVectorBuilder* PoseSearchFeatureVectorBuilder = SearchContext.GetCachedQuery(CurrentSearchResult.Database.Get()))
+			if (const FPoseSearchFeatureVectorBuilder* PoseSearchFeatureVectorBuilder = SearchContext.GetCachedQuery(CurrentSearchResult.Database->Schema))
 			{
 				if (const UPoseSearchFeatureChannel_Trajectory* TrajectoryChannel = CurrentSearchResult.Database->Schema->FindFirstChannelOfType<UPoseSearchFeatureChannel_Trajectory>())
 				{
diff --git a/Engine/Plugins/Experimental/Animation/PoseSearch/Source/Runtime/Public/PoseSearch/PoseSearchContext.h b/Engine/Plugins/Experimental/Animation/PoseSearch/Source/Runtime/Public/PoseSearch/PoseSearchContext.h
--- a/Engine/Plugins/Experimental/Animation/PoseSearch/Source/Runtime/Public/PoseSearch/PoseSearchContext.h
+++ b/Engine/Plugins/Experimental/Animation/PoseSearch/Source/Runtime/Public/PoseSearch/PoseSearchContext.h
@@ -172,8 +172,8 @@ struct POSESEARCH_API FSearchContext
 	void UpdateCurrentBestCost(const FPoseSearchCost& PoseSearchCost);
 	float GetCurrentBestTotalCost() const { return CurrentBestTotalCost; }
 
-	void GetOrBuildQuery(const UPoseSearchDatabase* Database, FPoseSearchFeatureVectorBuilder& FeatureVectorBuilder);
-	const FPoseSearchFeatureVectorBuilder* GetCachedQuery(const UPoseSearchDatabase* Database) const;
+	void GetOrBuildQuery(const UPoseSearchSchema* Schema, FPoseSearchFeatureVectorBuilder& FeatureVectorBuilder);
+	const FPoseSearchFeatureVectorBuilder* GetCachedQuery(const UPoseSearchSchema* Schema) const;
 
 	bool IsCurrentResultFromDatabase(const UPoseSearchDatabase* Database) const;
 
@@ -186,15 +186,7 @@ struct POSESEARCH_API FSearchContext
 private:
 	// transforms cached in component space
 	FCachedTransforms<FTransform> CachedTransforms;
-	
-	struct FCachedQuery
-	{
-		// @todo: why do we need to hold on to a Database? isn't the FeatureVectorBuilder.Schema enough?
-		const UPoseSearchDatabase* Database = nullptr;
-		FPoseSearchFeatureVectorBuilder FeatureVectorBuilder;
-	};
-
-	TArray<FCachedQuery, TInlineAllocator<8>> CachedQueries;
+	TArray<FPoseSearchFeatureVectorBuilder, TInlineAllocator<8>> CachedQueries;
 
 	float CurrentBestTotalCost = MAX_flt;
 
diff --git a/Engine/Plugins/Experimental/Animation/PoseSearch/Source/Runtime/Public/PoseSearch/PoseSearchDatabase.h b/Engine/Plugins/Experimental/Animation/PoseSearch/Source/Runtime/Public/PoseSearch/PoseSearchDatabase.h
--- a/Engine/Plugins/Experimental/Animation/PoseSearch/Source/Runtime/Public/PoseSearch/PoseSearchDatabase.h
+++ b/Engine/Plugins/Experimental/Animation/PoseSearch/Source/Runtime/Public/PoseSearch/PoseSearchDatabase.h
@@ -250,7 +250,6 @@ public:
 	const bool IsSourceAssetLooping(const FPoseSearchIndexAsset& SearchIndexAsset) const;
 	const FString GetSourceAssetName(const FPoseSearchIndexAsset& SearchIndexAsset) const;
 	int32 GetNumberOfPrincipalComponents() const;
-	void BuildQuery(UE::PoseSearch::FSearchContext& SearchContext, FPoseSearchFeatureVectorBuilder& OutQuery) const;
 
 	// Begin UObject
 	virtual void PostLoad() override;
