motion matching - fix for AnalyzeChannelRecursively after CL 23899462: FKeyBuilder would skip the transient channels assigning the same guid to all of them causing a poor normalization or crash in case non homogeneous cardinality (adding a phase channel would trigger a check)

diff --git a/Engine/Plugins/Experimental/Animation/PoseSearch/Source/Runtime/Private/PoseSearchDerivedData.cpp b/Engine/Plugins/Experimental/Animation/PoseSearch/Source/Runtime/Private/PoseSearchDerivedData.cpp
--- a/Engine/Plugins/Experimental/Animation/PoseSearch/Source/Runtime/Private/PoseSearchDerivedData.cpp
+++ b/Engine/Plugins/Experimental/Animation/PoseSearch/Source/Runtime/Private/PoseSearchDerivedData.cpp
@@ -73,13 +73,11 @@ struct FFeatureChannelLayoutSet
 		const TConstArrayView<TObjectPtr<UPoseSearchFeatureChannel>> SubChannels = Channel->GetSubChannels();
 		if (SubChannels.Num() == 0)
 		{
-			// @todo: figure out if Label + SkeletonName is enough to identify this channel, if better performances are needed
-			// FString Label = Channel->GetLabel();
-			// FString SkeletonName = Channel->GetSchema()->Skeleton->GetName();
-			// UE::PoseSearch::FKeyBuilder KeyBuilder;
-			// KeyBuilder << SkeletonName << Label;
-
-			Add(Channel->GetLabel(), UE::PoseSearch::FKeyBuilder(Channel).Finalize(), Channel->GetChannelDataOffset(), Channel->GetChannelCardinality());
+			FString Label = Channel->GetLabel();
+			FString SkeletonName = Channel->GetSchema()->Skeleton->GetName();
+			UE::PoseSearch::FKeyBuilder KeyBuilder;
+			KeyBuilder << SkeletonName << Label;
+			Add(Channel->GetLabel(), KeyBuilder.Finalize(), Channel->GetChannelDataOffset(), Channel->GetChannelCardinality());
 		}
 		else
 		{
diff --git a/Engine/Plugins/Experimental/Animation/PoseSearch/Source/Runtime/Private/PoseSearchDerivedDataKey.cpp b/Engine/Plugins/Experimental/Animation/PoseSearch/Source/Runtime/Private/PoseSearchDerivedDataKey.cpp
--- a/Engine/Plugins/Experimental/Animation/PoseSearch/Source/Runtime/Private/PoseSearchDerivedDataKey.cpp
+++ b/Engine/Plugins/Experimental/Animation/PoseSearch/Source/Runtime/Private/PoseSearchDerivedDataKey.cpp
@@ -41,7 +41,7 @@ FKeyBuilder::FKeyBuilder(const UObject* Object, bool bUseDataVer)
 	if (bUseDataVer)
 	{
 		// used to invalidate the key without having to change POSESEARCHDB_DERIVEDDATA_VER all the times
-		int32 POSESEARCHDB_DERIVEDDATA_VER_SMALL = 29;
+		int32 POSESEARCHDB_DERIVEDDATA_VER_SMALL = 31;
 		FGuid VersionGuid = FDevSystemGuids::GetSystemGuid(FDevSystemGuids::Get().POSESEARCHDB_DERIVEDDATA_VER);
 
 		*this << VersionGuid;
