motion matching - database mean deviation calculation and normalization refactoring (delegate data normalization decision to channels via UPoseSearchFeatureChannel::CanBeNormalizedWith, better error detection in case of channel implementation issues)

diff --git a/Engine/Plugins/Experimental/Animation/PoseSearch/Source/Runtime/Private/PoseSearchDerivedData.cpp b/Engine/Plugins/Experimental/Animation/PoseSearch/Source/Runtime/Private/PoseSearchDerivedData.cpp
--- a/Engine/Plugins/Experimental/Animation/PoseSearch/Source/Runtime/Private/PoseSearchDerivedData.cpp
+++ b/Engine/Plugins/Experimental/Animation/PoseSearch/Source/Runtime/Private/PoseSearchDerivedData.cpp
@@ -40,44 +40,69 @@ static FCookStatsManager::FAutoRegisterCallback RegisterCookStats([](FCookStatsM
 	});
 #endif
 
-// data structure collecting the internal layout representation of UPoseSearchFeatureChannel,
-// so we can aggregate data from different FPoseSearchIndex and calculate mean deviation with a homogeneous data set (ComputeChannelsDeviations
-struct FFeatureChannelLayoutSet
+// helper struct to calculate mean deviations
+struct FMeanDeviationCalculator
 {
-	// data structure holding DataOffset and Cardinality to find the data of the DebugName (channel breakdown / layout) in the related SearchIndexBases[SchemaIndex]
+private:
 	struct FEntry
 	{
-		FString DebugName; // for easier debugging
-		int32 SchemaIndex = -1; // index of the associated Schemas / SearchIndexBases used as input of the algorithm
-		int32 DataOffset = -1; // data offset from the base of SearchIndexBases[SchemaIndex].Values.GetData() from where the data associated to this Item starts
-		int32 Cardinality = -1; // data cardinality
+		int32 SchemaIndex = -1; // index of the Channel associated Schemas / SearchIndexBases index
+		const UPoseSearchFeatureChannel* Channel;
 	};
 
-	// FIoHash is the hash associated to the channel data breakdown (e.g.: it could be a single SampledBones at a specific SampleTimes for a UPoseSearchFeatureChannel_Pose)
-	TMap<FIoHash, TArray<FEntry>> EntriesMap;
-	int32 CurrentSchemaIndex = -1;
-	TWeakObjectPtr<const UPoseSearchSchema> CurrentSchema;
+	typedef TArray<FEntry> FEntries;		// array of FEntry with channels that can be normalized together (for example it contains all the phases of the left foot from different schemas)
+	typedef TArray<FEntries> FEntriesGroup;	// array of FEntries incompatible between each other
 
-	void Add(FString DebugName, FIoHash IoHash, int32 DataOffset, int32 Cardinality)
+	static void Add(const UPoseSearchFeatureChannel* Channel, int32 SchemaIndex, FEntriesGroup& EntriesGroup)
 	{
-		check(DataOffset >= 0 && Cardinality >= 0 && CurrentSchemaIndex >= 0);
-		TArray<FEntry>& Entries = EntriesMap.FindOrAdd(IoHash);
+		bool bEntryFound = false;
+		for (FEntries& Entries : EntriesGroup)
+		{
+			check(Entries[0].Channel != Channel);
+			if (Entries[0].Channel->CanBeNormalizedWith(Channel))
+			{
+				#if DO_CHECK
+				check(Entries[0].Channel->GetChannelCardinality() == Channel->GetChannelCardinality());
+
+				// Channel can be normalized with ALL the channels already collected in Entry. Let's test that it's true
+				for (int EntryIndex = 1; EntryIndex < Entries.Num(); ++EntryIndex)
+				{
+					check(Entries[EntryIndex].Channel != Channel);
+					check(Entries[EntryIndex].Channel->CanBeNormalizedWith(Channel));
+				}
+				#endif // DO_CHECK
 
-		// making sure all the FEntry associated with the same IoHash have the same Cardinality
-		check(Entries.IsEmpty() || Entries[0].Cardinality == Cardinality);
-		Entries.Add({ DebugName, CurrentSchemaIndex, DataOffset, Cardinality });
+				Entries.Add({ SchemaIndex, Channel });
+				bEntryFound = true;
+				#if !DO_CHECK
+				break;
+				#endif // !DO_CHECK
+			}
+			#if DO_CHECK
+			else
+			{
+				for (int EntryIndex = 1; EntryIndex < Entries.Num(); ++EntryIndex)
+				{
+					check(Entries[EntryIndex].Channel != Channel);
+					check(!Entries[EntryIndex].Channel->CanBeNormalizedWith(Channel));
+				}
+			}
+			#endif // DO_CHECK
+		}
+
+		if (!bEntryFound)
+		{
+			FEntries& Entries = EntriesGroup[EntriesGroup.AddDefaulted()];
+			Entries.Add({ SchemaIndex, Channel });
+		}
 	}
 
-	void AnalyzeChannelRecursively(const UPoseSearchFeatureChannel* Channel)
+	static void AnalyzeChannelRecursively(const UPoseSearchFeatureChannel* Channel, int32 SchemaIndex, FEntriesGroup& EntriesGroup)
 	{
 		const TConstArrayView<TObjectPtr<UPoseSearchFeatureChannel>> SubChannels = Channel->GetSubChannels();
 		if (SubChannels.Num() == 0)
 		{
-			FString Label = Channel->GetLabel();
-			FString SkeletonName = Channel->GetSchema()->Skeleton->GetName();
-			UE::PoseSearch::FKeyBuilder KeyBuilder;
-			KeyBuilder << SkeletonName << Label;
-			Add(Channel->GetLabel(), KeyBuilder.Finalize(), Channel->GetChannelDataOffset(), Channel->GetChannelCardinality());
+			Add(Channel, SchemaIndex, EntriesGroup);
 		}
 		else
 		{
@@ -86,130 +111,134 @@ struct FFeatureChannelLayoutSet
 			{
 				if (const UPoseSearchFeatureChannel* SubChannel = SubChannelPtr.Get())
 				{
-					AnalyzeChannelRecursively(SubChannel);
+					AnalyzeChannelRecursively(SubChannel, SchemaIndex, EntriesGroup);
 				}
 			}
 		}
 	}
-};
 
-static float ComputeFeatureMeanDeviation(TConstArrayView<FFeatureChannelLayoutSet::FEntry> Entries, TConstArrayView<FPoseSearchIndexBase> SearchIndexBases, TConstArrayView<const UPoseSearchSchema*> Schemas)
-{
-	check(Schemas.Num() == SearchIndexBases.Num());
-	
-	const int32 EntriesNum = Entries.Num();
-	check(EntriesNum > 0);
-
-	const int32 Cardinality = Entries[0].Cardinality;
-	check(Cardinality > 0);
-
-	int32 TotalNumPoses = 0;
-	for (int32 EntryIdx = 0; EntryIdx < EntriesNum; ++EntryIdx)
+	static void AnalyzeSchemas(TConstArrayView<const UPoseSearchSchema*> Schemas, FEntriesGroup& EntriesGroup)
 	{
-		TotalNumPoses += SearchIndexBases[Entries[EntryIdx].SchemaIndex].NumPoses;
+		for (int32 SchemaIndex = 0; SchemaIndex < Schemas.Num(); ++SchemaIndex)
+		{
+			const UPoseSearchSchema* Schema = Schemas[SchemaIndex];
+			for (const TObjectPtr<UPoseSearchFeatureChannel>& ChannelPtr : Schema->Channels)
+			{
+				if (const UPoseSearchFeatureChannel* Channel = ChannelPtr.Get())
+				{
+					AnalyzeChannelRecursively(Channel, SchemaIndex, EntriesGroup);
+				}
+			}
+		}
 	}
 
-	int32 AccumulatedNumPoses = 0;
-	RowMajorMatrix CenteredSubPoseMatrix(TotalNumPoses, Cardinality);
-	for (int32 EntryIdx = 0; EntryIdx < EntriesNum; ++EntryIdx)
+	// given an array of channels that can be normalized together (Entries), with the same cardinality (Entries[0].Channel->GetChannelCardinality()),
+	// it'll calculate the mean deviation of the associated data (from SearchIndexBases)
+	static float CalculateEntriesMeanDeviation(const FEntries& Entries, TConstArrayView<FPoseSearchIndexBase> SearchIndexBases, TConstArrayView<const UPoseSearchSchema*> Schemas)
 	{
-		const FFeatureChannelLayoutSet::FEntry& Entry = Entries[EntryIdx];
-		check(Cardinality == Entry.Cardinality);
+		check(Schemas.Num() == SearchIndexBases.Num());
+	
+		const int32 EntriesNum = Entries.Num();
+		check(EntriesNum > 0);
 
-		const int32 DataSetIdx = Entry.SchemaIndex;
+		const int32 Cardinality = Entries[0].Channel->GetChannelCardinality();
+		check(Cardinality > 0);
 
-		const UPoseSearchSchema* Schema = Schemas[DataSetIdx];
-		const FPoseSearchIndexBase& SearchIndex = SearchIndexBases[DataSetIdx];
+		int32 TotalNumPoses = 0;
+		for (int32 EntryIdx = 0; EntryIdx < EntriesNum; ++EntryIdx)
+		{
+			TotalNumPoses += SearchIndexBases[Entries[EntryIdx].SchemaIndex].NumPoses;
+		}
 
-		const int32 NumPoses = SearchIndex.NumPoses;
+		int32 AccumulatedNumPoses = 0;
+		RowMajorMatrix CenteredSubPoseMatrix(TotalNumPoses, Cardinality);
+		for (int32 EntryIdx = 0; EntryIdx < EntriesNum; ++EntryIdx)
+		{
+			const FEntry& Entry = Entries[EntryIdx];
+			check(Cardinality == Entry.Channel->GetChannelCardinality());
 
-		// Map input buffer with NumPoses as rows and NumDimensions	as cols
-		RowMajorMatrixMapConst PoseMatrixSourceMap(SearchIndex.Values.GetData(), NumPoses, Schema->SchemaCardinality);
+			const int32 DataSetIdx = Entry.SchemaIndex;
 
-		// Given the sub matrix for the features, find the average distance to the feature's centroid.
-		CenteredSubPoseMatrix.block(AccumulatedNumPoses, 0, NumPoses, Cardinality) = PoseMatrixSourceMap.block(0, Entry.DataOffset, NumPoses, Cardinality);
-		AccumulatedNumPoses += NumPoses;
-	}
+			const UPoseSearchSchema* Schema = Schemas[DataSetIdx];
+			const FPoseSearchIndexBase& SearchIndex = SearchIndexBases[DataSetIdx];
 
-	RowMajorVector SampleMean = CenteredSubPoseMatrix.colwise().mean();
-	CenteredSubPoseMatrix = CenteredSubPoseMatrix.rowwise() - SampleMean;
+			const int32 NumPoses = SearchIndex.NumPoses;
 
-	// after mean centering the data, the average distance to the centroid is simply the average norm.
-	const float FeatureMeanDeviation = CenteredSubPoseMatrix.rowwise().norm().mean();
+			// Map input buffer with NumPoses as rows and NumDimensions	as cols
+			RowMajorMatrixMapConst PoseMatrixSourceMap(SearchIndex.Values.GetData(), NumPoses, Schema->SchemaCardinality);
 
-	return FeatureMeanDeviation;
-}
+			// Given the sub matrix for the features, find the average distance to the feature's centroid.
+			CenteredSubPoseMatrix.block(AccumulatedNumPoses, 0, NumPoses, Cardinality) = PoseMatrixSourceMap.block(0, Entry.Channel->GetChannelDataOffset(), NumPoses, Cardinality);
+			AccumulatedNumPoses += NumPoses;
+		}
 
-// it collects FFeatureChannelLayoutSet from all the Schemas (for example, figuring out the data offsets of SampledBones at a specific 
-// SampleTimes for a UPoseSearchFeatureChannel_Pose for all the SearchIndexBases), and call ComputeFeatureMeanDeviation
-static TArray<float> ComputeChannelsDeviations(TConstArrayView<FPoseSearchIndexBase> SearchIndexBases, TConstArrayView<const UPoseSearchSchema*> Schemas)
-{
-	// This function performs a modified z-score normalization where features are normalized
-	// by mean absolute deviation rather than standard deviation. Both methods are preferable
-	// here to min-max scaling because they preserve outliers.
-	// 
-	// Mean absolute deviation is preferred here over standard deviation because the latter
-	// emphasizes outliers since squaring the distance from the mean increases variance 
-	// exponentially rather than additively and square rooting the sum of squares does not 
-	// remove that bias. [1]
-	//
-	// References:
-	// [1] Gorard, S. (2005), "Revisiting a 90-Year-Old Debate: The Advantages of the Mean Deviation."
-	//     British Journal of Educational Studies, 53: 417-430.
-
-	using namespace Eigen;
-	using namespace UE::PoseSearch;
+		RowMajorVector SampleMean = CenteredSubPoseMatrix.colwise().mean();
+		CenteredSubPoseMatrix = CenteredSubPoseMatrix.rowwise() - SampleMean;
 
-	int32 ThisSchemaIndex = 0;
-	check(SearchIndexBases.Num() == Schemas.Num() && Schemas.Num() > ThisSchemaIndex);
-	const UPoseSearchSchema* ThisSchema = Schemas[ThisSchemaIndex];
-	check(ThisSchema->IsValid());
-	const int32 NumDimensions = ThisSchema->SchemaCardinality;
+		// after mean centering the data, the average distance to the centroid is simply the average norm.
+		const float FeatureMeanDeviation = CenteredSubPoseMatrix.rowwise().norm().mean();
 
-	TArray<float> MeanDeviations;
-	MeanDeviations.Init(1.f, NumDimensions);
-	RowMajorVectorMap MeanDeviationsMap(MeanDeviations.GetData(), 1, NumDimensions);
+		return FeatureMeanDeviation;
+	}
 
-	const EPoseSearchDataPreprocessor DataPreprocessor = ThisSchema->DataPreprocessor;
-	if (SearchIndexBases[ThisSchemaIndex].NumPoses > 0 && (DataPreprocessor == EPoseSearchDataPreprocessor::Normalize || DataPreprocessor == EPoseSearchDataPreprocessor::NormalizeOnlyByDeviation))
+public:
+
+	// it returns an array of dimension Schemas[0]->SchemaCardinality containing the mean deviation calculated from the data passed in with SearchIndexBases following the layout described in the schemas channels:
+	// channels from all the schemas get collected in groups that can be normalized together (FEntriesGroup, populated in AnalyzeSchemas) and then those homogeneous (in cardinality and meaning) groups get processed 
+	// one by one in CalculateEntriesMeanDeviation to extract the group mean deviation against the input data contained in SearchIndexBases
+	static TArray<float> Calculate(TConstArrayView<FPoseSearchIndexBase> SearchIndexBases, TConstArrayView<const UPoseSearchSchema*> Schemas)
 	{
-		FFeatureChannelLayoutSet FeatureChannelLayoutSet;
-		for (int32 SchemaIndex = 0; SchemaIndex < Schemas.Num(); ++SchemaIndex)
+		// This method performs a modified z-score normalization where features are normalized
+		// by mean absolute deviation rather than standard deviation. Both methods are preferable
+		// here to min-max scaling because they preserve outliers.
+		// 
+		// Mean absolute deviation is preferred here over standard deviation because the latter
+		// emphasizes outliers since squaring the distance from the mean increases variance 
+		// exponentially rather than additively and square rooting the sum of squares does not 
+		// remove that bias. [1]
+		//
+		// References:
+		// [1] Gorard, S. (2005), "Revisiting a 90-Year-Old Debate: The Advantages of the Mean Deviation."
+		//     British Journal of Educational Studies, 53: 417-430.
+
+		int32 ThisSchemaIndex = 0;
+		check(SearchIndexBases.Num() == Schemas.Num() && Schemas.Num() > ThisSchemaIndex);
+		const UPoseSearchSchema* ThisSchema = Schemas[ThisSchemaIndex];
+		check(ThisSchema->IsValid());
+		const int32 NumDimensions = ThisSchema->SchemaCardinality;
+
+		TArray<float> MeanDeviations;
+		MeanDeviations.Init(1.f, NumDimensions);
+		RowMajorVectorMap MeanDeviationsMap(MeanDeviations.GetData(), 1, NumDimensions);
+
+		const EPoseSearchDataPreprocessor DataPreprocessor = ThisSchema->DataPreprocessor;
+		if (SearchIndexBases[ThisSchemaIndex].NumPoses > 0 && (DataPreprocessor == EPoseSearchDataPreprocessor::Normalize || DataPreprocessor == EPoseSearchDataPreprocessor::NormalizeOnlyByDeviation))
 		{
-			const UPoseSearchSchema* Schema = Schemas[SchemaIndex];
+			FEntriesGroup EntriesGroup;
 
-			FeatureChannelLayoutSet.CurrentSchemaIndex = SchemaIndex;
-			FeatureChannelLayoutSet.CurrentSchema = Schema;
-			for (const TObjectPtr<UPoseSearchFeatureChannel>& ChannelPtr : Schema->Channels)
-			{
-				if (const UPoseSearchFeatureChannel* Channel = ChannelPtr.Get())
-				{
-					FeatureChannelLayoutSet.AnalyzeChannelRecursively(Channel);
-				}
-			}
-		}
+			AnalyzeSchemas(Schemas, EntriesGroup);
 
-		for (auto Pair : FeatureChannelLayoutSet.EntriesMap)
-		{
-			const TArray<FFeatureChannelLayoutSet::FEntry>& Entries = Pair.Value;
-			for (const FFeatureChannelLayoutSet::FEntry& Entry : Entries)
+			for (const FEntries& Entries : EntriesGroup)
 			{
-				if (Entry.Cardinality > 0 && Entry.SchemaIndex == ThisSchemaIndex)
+				for (const FEntry& Entry : Entries)
 				{
-					const float FeatureMeanDeviation = ComputeFeatureMeanDeviation(Entries, SearchIndexBases, Schemas);
-					// the associated data to all the Entries data is going to be used to calculate the deviation of Deviation[Entry.DataOffset] to Deviation[Entry.DataOffset + Entry.Cardinality]
+					if (Entry.Channel->GetChannelCardinality() > 0 && Entry.SchemaIndex == ThisSchemaIndex)
+					{
+						const float FeatureMeanDeviation = CalculateEntriesMeanDeviation(Entries, SearchIndexBases, Schemas);
+						// the associated data to all the Entries data is going to be used to calculate the deviation of Deviation[Entry.Channel->GetChannelDataOffset()] to Deviation[Entry.Channel->GetChannelDataOffset() + Entry.Channel->GetChannelCardinality()]
 
-					// Fill the feature's corresponding scaling axes with the average distance
-					// Avoid scaling by zero by leaving near-zero deviations as 1.0
-					static const float MinFeatureMeanDeviation = 0.1f;
-					MeanDeviationsMap.segment(Entry.DataOffset, Entry.Cardinality).setConstant(FeatureMeanDeviation > MinFeatureMeanDeviation ? FeatureMeanDeviation : 1.f);
+						// Fill the feature's corresponding scaling axes with the average distance
+						// Avoid scaling by zero by leaving near-zero deviations as 1.0
+						static const float MinFeatureMeanDeviation = 0.1f;
+						MeanDeviationsMap.segment(Entry.Channel->GetChannelDataOffset(), Entry.Channel->GetChannelCardinality()).setConstant(FeatureMeanDeviation > MinFeatureMeanDeviation ? FeatureMeanDeviation : 1.f);
+					}
 				}
 			}
 		}
-	}
 	
-	return MeanDeviations;
-}
+		return MeanDeviations;
+	}
+};
 
 static inline FFloatInterval GetEffectiveSamplingRange(const UAnimSequenceBase* Sequence, FFloatInterval RequestedSamplingRange)
 {
@@ -944,7 +973,7 @@ void FPoseSearchDatabaseAsyncCacheTask::OnGetComplete(UE::DerivedData::FCacheGet
 
 				static_cast<FPoseSearchIndexBase&>(SearchIndex) = SearchIndexBases[0];
 				
-				TArray<float> Deviation = ComputeChannelsDeviations(SearchIndexBases, Schemas);
+				TArray<float> Deviation = FMeanDeviationCalculator::Calculate(SearchIndexBases, Schemas);
 
 				#if WITH_EDITORONLY_DATA
 				SearchIndex.Deviation = Deviation;
diff --git a/Engine/Plugins/Experimental/Animation/PoseSearch/Source/Runtime/Private/PoseSearchFeatureChannel.cpp b/Engine/Plugins/Experimental/Animation/PoseSearch/Source/Runtime/Private/PoseSearchFeatureChannel.cpp
--- a/Engine/Plugins/Experimental/Animation/PoseSearch/Source/Runtime/Private/PoseSearchFeatureChannel.cpp
+++ b/Engine/Plugins/Experimental/Animation/PoseSearch/Source/Runtime/Private/PoseSearchFeatureChannel.cpp
@@ -211,6 +211,31 @@ FString UPoseSearchFeatureChannel::GetLabel() const
 	return Label.ToString();
 }
 
+bool UPoseSearchFeatureChannel::CanBeNormalizedWith(const UPoseSearchFeatureChannel* Other) const
+{
+	if (ChannelCardinality != Other->ChannelCardinality)
+	{
+		return false;
+	}
+
+	if (GetClass() != Other->GetClass())
+	{
+		return false;
+	}
+
+	if (GetSchema()->Skeleton != Other->GetSchema()->Skeleton)
+	{
+		return false;
+	}
+
+	if (GetLabel() != Other->GetLabel())
+	{
+		return false;
+	}
+
+	return true;
+}
+
 const UPoseSearchSchema* UPoseSearchFeatureChannel::GetSchema() const
 {
 	UObject* Outer = GetOuter();
diff --git a/Engine/Plugins/Experimental/Animation/PoseSearch/Source/Runtime/Public/PoseSearch/PoseSearchFeatureChannel.h b/Engine/Plugins/Experimental/Animation/PoseSearch/Source/Runtime/Public/PoseSearch/PoseSearchFeatureChannel.h
--- a/Engine/Plugins/Experimental/Animation/PoseSearch/Source/Runtime/Public/PoseSearch/PoseSearchFeatureChannel.h
+++ b/Engine/Plugins/Experimental/Animation/PoseSearch/Source/Runtime/Public/PoseSearch/PoseSearchFeatureChannel.h
@@ -114,6 +114,7 @@ public:
 #if WITH_EDITOR
 	// returns the FString used editor side to identify this UPoseSearchFeatureChannel (for instance in the pose search debugger)
 	virtual FString GetLabel() const;
+	virtual bool CanBeNormalizedWith(const UPoseSearchFeatureChannel* Other) const;
 	const UPoseSearchSchema* GetSchema() const;
 #endif
 
