motion matching - adding debugdraw to pose history node

diff --git a/Engine/Plugins/Experimental/Animation/PoseSearch/Source/Runtime/Private/AnimNode_PoseSearchHistoryCollector.cpp b/Engine/Plugins/Experimental/Animation/PoseSearch/Source/Runtime/Private/AnimNode_PoseSearchHistoryCollector.cpp
--- a/Engine/Plugins/Experimental/Animation/PoseSearch/Source/Runtime/Private/AnimNode_PoseSearchHistoryCollector.cpp
+++ b/Engine/Plugins/Experimental/Animation/PoseSearch/Source/Runtime/Private/AnimNode_PoseSearchHistoryCollector.cpp
@@ -3,11 +3,16 @@
 #include "PoseSearch/AnimNode_PoseSearchHistoryCollector.h"
 #include "Animation/AnimInstanceProxy.h"
 #include "Animation/AnimNodeMessages.h"
+#include "Engine/SkinnedAsset.h"
 
 #include UE_INLINE_GENERATED_CPP_BY_NAME(AnimNode_PoseSearchHistoryCollector)
 
 #define LOCTEXT_NAMESPACE "AnimNode_PoseSearchHistoryCollector"
 
+#if WITH_EDITORONLY_DATA && ENABLE_ANIM_DEBUG
+TAutoConsoleVariable<bool> CVarAnimPoseHistoryDebugDraw(TEXT("a.AnimNode.PoseHistory.DebugDraw"), false, TEXT("Enable / Disable Pose History DebugDraw"));
+#endif
+
 namespace UE::PoseSearch::Private
 {
 
@@ -64,23 +69,12 @@ void FAnimNode_PoseSearchHistoryCollector::Evaluate_AnyThread(FPoseContext& Outp
 {
 	DECLARE_SCOPE_HIERARCHICAL_COUNTER_ANIMNODE(Evaluate_AnyThread);
 
-	using namespace UE::PoseSearch;
-
 	Source.Evaluate(Output);
+	PoseHistory.Update(Output.AnimInstanceProxy->GetDeltaSeconds(), Output, Output.AnimInstanceProxy->GetComponentTransform());
 
-	FText ErrorText;
-
-	if (!PoseHistory.Update(
-		Output.AnimInstanceProxy->GetDeltaSeconds(),
-		Output,
-		Output.AnimInstanceProxy->GetComponentTransform(),
-		&ErrorText,
-		bUseRootMotion ?
-		FPoseHistory::ERootUpdateMode::RootMotionDelta :
-		FPoseHistory::ERootUpdateMode::ComponentTransformDelta))
-	{
-		Output.LogMessage(EMessageSeverity::Warning, ErrorText);
-	}
+#if WITH_EDITORONLY_DATA && ENABLE_ANIM_DEBUG
+	bWasEvaluated = true;
+#endif
 }
 
 void FAnimNode_PoseSearchHistoryCollector::Update_AnyThread(const FAnimationUpdateContext& Context)
@@ -97,4 +91,32 @@ void FAnimNode_PoseSearchHistoryCollector::GatherDebugData(FNodeDebugData& Debug
 	Source.GatherDebugData(DebugData);
 }
 
+bool FAnimNode_PoseSearchHistoryCollector::HasPreUpdate() const
+{
+	return WITH_EDITORONLY_DATA && ENABLE_ANIM_DEBUG;
+}
+
+void FAnimNode_PoseSearchHistoryCollector::PreUpdate(const UAnimInstance* InAnimInstance)
+{
+#if WITH_EDITORONLY_DATA && ENABLE_ANIM_DEBUG
+	if (bWasEvaluated && CVarAnimPoseHistoryDebugDraw.GetValueOnAnyThread())
+	{
+		if (const USkeletalMeshComponent* SkeletalMeshComponent = InAnimInstance->GetSkelMeshComponent())
+		{
+			const USkinnedAsset* SkinnedAsset = SkeletalMeshComponent->GetSkinnedAsset();
+			const UWorld* World = SkeletalMeshComponent->GetWorld();
+			if (SkinnedAsset && World)
+			{
+				if (const USkeleton* Skeleton = SkinnedAsset->GetSkeleton())
+				{
+					PoseHistory.DebugDraw(World, Skeleton);
+				}
+			}
+		}
+	}
+
+	bWasEvaluated = false;
+#endif // WITH_EDITORONLY_DATA && ENABLE_ANIM_DEBUG
+}
+
 #undef LOCTEXT_NAMESPACE
diff --git a/Engine/Plugins/Experimental/Animation/PoseSearch/Source/Runtime/Private/PoseSearchContext.cpp b/Engine/Plugins/Experimental/Animation/PoseSearch/Source/Runtime/Private/PoseSearchContext.cpp
--- a/Engine/Plugins/Experimental/Animation/PoseSearch/Source/Runtime/Private/PoseSearchContext.cpp
+++ b/Engine/Plugins/Experimental/Animation/PoseSearch/Source/Runtime/Private/PoseSearchContext.cpp
@@ -151,7 +151,7 @@ FTransform FSearchContext::GetTransform(float SampleTime, const UPoseSearchSchem
 {
 	// collecting the RootTransform from the FPoseHistory
 	FTransform RootTransform = FTransform::Identity;
-	History->TrySampleLocalPose(-SampleTime, nullptr, nullptr, &RootTransform);
+	History->GetRootTransformAtTime(SampleTime, RootTransform);
 
 	const FBoneIndexType BoneIndexType = Schema->GetBoneIndexType(SchemaBoneIdx);
 	if (BoneIndexType != RootBoneIndexType)
@@ -177,19 +177,18 @@ FTransform FSearchContext::GetComponentSpaceTransform(float SampleTime, const UP
 	
 		// collecting the local bone transforms from the FPoseHistory
 		TArray<FTransform> SampledLocalPose;
-		if (History->TrySampleLocalPose(-SampleTime, &Schema->BoneIndicesWithParents, &SampledLocalPose, nullptr))
-		{
-			TArray<FTransform> SampledComponentPose;
-			FAnimationRuntime::FillUpComponentSpaceTransforms(Schema->Skeleton->GetReferenceSkeleton(), SampledLocalPose, SampledComponentPose);
-
-			// adding bunch of entries, without caring about adding eventual duplicates
-			for (const FBoneIndexType NewEntryBoneIndexType : Schema->BoneIndicesWithParents)
-			{
-				CachedTransforms.Add(SampleTime, NewEntryBoneIndexType, SampledComponentPose[NewEntryBoneIndexType]);
-			}
+		History->GetLocalPoseAtTime(SampleTime, Schema->BoneIndicesWithParents, SampledLocalPose);
+		
+		TArray<FTransform> SampledComponentPose;
+		FAnimationRuntime::FillUpComponentSpaceTransforms(Schema->Skeleton->GetReferenceSkeleton(), SampledLocalPose, SampledComponentPose);
 
-			return SampledComponentPose[BoneIndexType];
+		// adding bunch of entries, without caring about adding eventual duplicates
+		for (const FBoneIndexType NewEntryBoneIndexType : Schema->BoneIndicesWithParents)
+		{
+			CachedTransforms.Add(SampleTime, NewEntryBoneIndexType, SampledComponentPose[NewEntryBoneIndexType]);
 		}
+
+		return SampledComponentPose[BoneIndexType];
 	}
 
 	return FTransform::Identity;
diff --git a/Engine/Plugins/Experimental/Animation/PoseSearch/Source/Runtime/Private/PoseSearchHistory.cpp b/Engine/Plugins/Experimental/Animation/PoseSearch/Source/Runtime/Private/PoseSearchHistory.cpp
--- a/Engine/Plugins/Experimental/Animation/PoseSearch/Source/Runtime/Private/PoseSearchHistory.cpp
+++ b/Engine/Plugins/Experimental/Animation/PoseSearch/Source/Runtime/Private/PoseSearchHistory.cpp
@@ -5,6 +5,7 @@
 #include "Animation/AnimNodeBase.h"
 #include "Animation/AnimRootMotionProvider.h"
 #include "BonePose.h"
+#include "DrawDebugHelpers.h"
 #include "PoseSearch/PoseSearchResult.h"
 #include "PoseSearch/PoseSearchDatabase.h"
 
@@ -98,57 +99,83 @@ void FPoseHistory::Init(const FPoseHistory& History)
 	TimeHorizon = History.TimeHorizon;
 }
 
-bool FPoseHistory::TrySampleLocalPose(float SecondsAgo, const TArray<FBoneIndexType>* RequiredBones, TArray<FTransform>* LocalPose, FTransform* RootTransform) const
+static FBoneIndexType GetMaxFBoneIndexType(const TArray<FBoneIndexType>& RequiredBones)
 {
-	const int32 NextIdx = LowerBound(Poses.begin(), Poses.end(), SecondsAgo, [](const FPose& Pose, float Value)
-	{
-		return Value < Pose.Time;
-	});
-	if (NextIdx <= 0 || NextIdx >= Poses.Num())
-	{
-		// We may not have accumulated enough poses yet
-		return false;
-	}
-
-	const int32 PrevIdx = NextIdx - 1;
-
-	const FPose& PrevPose = Poses[PrevIdx];
-	const FPose& NextPose = Poses[NextIdx];
-
-#if DO_CHECK
-	check(PrevPose.LocalTransforms.Num() == NextPose.LocalTransforms.Num());
 	FBoneIndexType MaxBoneIndexType = 0;
-	if (RequiredBones)
+	for (FBoneIndexType BoneIndexType : RequiredBones)
 	{
-		for (FBoneIndexType BoneIndexType : *RequiredBones)
+		if (BoneIndexType > MaxBoneIndexType)
 		{
-			if (BoneIndexType > MaxBoneIndexType)
-			{
-				MaxBoneIndexType = BoneIndexType;
-			}
+			MaxBoneIndexType = BoneIndexType;
 		}
-		check(MaxBoneIndexType < PrevPose.LocalTransforms.Num());
 	}
-#endif
-	// Compute alpha between previous and next Poses
-	const float Alpha = FMath::GetMappedRangeValueUnclamped(FVector2f(PrevPose.Time, NextPose.Time), FVector2f(0.0f, 1.0f), SecondsAgo);
+	return MaxBoneIndexType;
+}
 
-	// Lerp between poses by alpha to produce output local pose at requested sample time
-	if (LocalPose)
+void FPoseHistory::GetLocalPoseAtTime(float Time, const TArray<FBoneIndexType>& RequiredBones, TArray<FTransform>& LocalPose) const
+{
+	const int32 Num = Poses.Num();
+	if (Num > 1)
+	{
+		const float SecondsAgo = -Time;
+		const int32 LowerBoundIdx = LowerBound(Poses.begin(), Poses.end(), SecondsAgo, [](const FPose& Pose, float Value) { return Value < Pose.Time; });
+		const int32 NextIdx = FMath::Clamp(LowerBoundIdx, 1, Num - 1);
+		const int32 PrevIdx = NextIdx - 1;
+
+		const FPose& PrevPose = Poses[PrevIdx];
+		const FPose& NextPose = Poses[NextIdx];
+				
+		// Compute alpha between previous and next Poses
+		const float Alpha = FMath::GetMappedRangeValueClamped(FVector2f(PrevPose.Time, NextPose.Time), FVector2f(0.0f, 1.0f), SecondsAgo);
+
+		// Lerp between poses by alpha to produce output local pose at requested sample time
+		check(PrevPose.LocalTransforms.Num() == NextPose.LocalTransforms.Num());
+		check(GetMaxFBoneIndexType(RequiredBones) < PrevPose.LocalTransforms.Num());
+
+		LocalPose = PrevPose.LocalTransforms;
+		FAnimationRuntime::LerpBoneTransforms(LocalPose, NextPose.LocalTransforms, Alpha, RequiredBones);
+	}
+	else if (Num > 0)
+	{
+		check(GetMaxFBoneIndexType(RequiredBones) < Poses[0].LocalTransforms.Num());
+		LocalPose = Poses[0].LocalTransforms;
+	}
+	else
 	{
-		check(RequiredBones);
-		*LocalPose = PrevPose.LocalTransforms;
-		FAnimationRuntime::LerpBoneTransforms(*LocalPose, NextPose.LocalTransforms, Alpha, *RequiredBones);
+		const FBoneIndexType MaxBoneIndexType = GetMaxFBoneIndexType(RequiredBones);
+		LocalPose.Init(FTransform::Identity, MaxBoneIndexType + 1);
 	}
+}
+
+void FPoseHistory::GetRootTransformAtTime(float Time, FTransform& RootTransform) const
+{
+	const int32 Num = Poses.Num();
+	if (Num > 1)
+	{
+		const float SecondsAgo = -Time;
+		const int32 LowerBoundIdx = LowerBound(Poses.begin(), Poses.end(), SecondsAgo, [](const FPose& Pose, float Value) { return Value < Pose.Time; });
+		const int32 NextIdx = FMath::Clamp(LowerBoundIdx, 1, Num - 1);
+		const int32 PrevIdx = NextIdx - 1;
 
-	if (RootTransform)
+		const FPose& PrevPose = Poses[PrevIdx];
+		const FPose& NextPose = Poses[NextIdx];
+
+		// Compute alpha between previous and next Poses
+		const float Alpha = FMath::GetMappedRangeValueClamped(FVector2f(PrevPose.Time, NextPose.Time), FVector2f(0.0f, 1.0f), SecondsAgo);
+
+		RootTransform.Blend(PrevPose.RootTransform, NextPose.RootTransform, Alpha);
+	}
+	else if (Num > 0)
 	{
-		RootTransform->Blend(PrevPose.RootTransform, NextPose.RootTransform, Alpha);
+		RootTransform = Poses[0].RootTransform;
+	}
+	else
+	{
+		RootTransform = FTransform::Identity;
 	}
-	return true;
 }
 
-bool FPoseHistory::Update(float SecondsElapsed, const FPoseContext& PoseContext, FTransform ComponentTransform, FText* OutError, ERootUpdateMode UpdateMode)
+void FPoseHistory::Update(float SecondsElapsed, const FPoseContext& PoseContext, FTransform ComponentTransform)
 {
 	// Age our elapsed times
 	for (FPose& Pose : Poses)
@@ -184,65 +211,60 @@ bool FPoseHistory::Update(float SecondsElapsed, const FPoseContext& PoseContext,
 	// Regardless of the retention policy, we always update the most recent pose
 	FPose& CurrentPose = Poses.Last();
 	CurrentPose.Time = 0.f;
+	CurrentPose.RootTransform = ComponentTransform;
 	CopyCompactToSkeletonPose(PoseContext.Pose, CurrentPose.LocalTransforms);
+}
 
-	// Initialize with Previous Root Transform or Identity
-	CurrentPose.RootTransform = Poses.Num() > 1 ? Poses[Poses.Num() - 2].RootTransform : FTransform::Identity;
-	
-	// Update using either AniumRootMotionProvider or Component Transform
-	if (UpdateMode == ERootUpdateMode::RootMotionDelta)
-	{
-		const UE::Anim::IAnimRootMotionProvider* RootMotionProvider = UE::Anim::IAnimRootMotionProvider::Get();
+float FPoseHistory::GetSampleTimeInterval() const
+{
+	// Reserve one pose for computing derivatives at the time horizon
+	return TimeHorizon / (Poses.Max() - 1);
+}
 
-		if (RootMotionProvider)
-		{
-			if (RootMotionProvider->HasRootMotion(PoseContext.CustomAttributes))
-			{
-				FTransform RootMotionDelta = FTransform::Identity;
-				RootMotionProvider->ExtractRootMotion(PoseContext.CustomAttributes, RootMotionDelta);
+void FPoseHistory::DebugDraw(const UWorld* World, const USkeleton* Skeleton) const
+{
+#if ENABLE_DRAW_DEBUG
 
-				CurrentPose.RootTransform = RootMotionDelta * CurrentPose.RootTransform;
-			}
-#if WITH_EDITORONLY_DATA	
-			else
-			{
-				if (OutError)
-				{
-					*OutError = LOCTEXT("PoseHistoryRootMotionProviderError",
-						"Input to Pose History has no Root Motion Attribute. Try disabling 'Use Root Motion'.");
-				}
-				return false;
-			}
-#endif
+	auto LerpColor = [](FColor A, FColor B, float T) -> FColor
+	{
+		return FColor(
+			FMath::RoundToInt(float(A.R) * (1.f - T) + float(B.R) * T),
+			FMath::RoundToInt(float(A.G) * (1.f - T) + float(B.G) * T),
+			FMath::RoundToInt(float(A.B) * (1.f - T) + float(B.B) * T),
+			FMath::RoundToInt(float(A.A) * (1.f - T) + float(B.A) * T));
+	};
+
+	TArray<FTransform> LocalTransforms;
+	TArray<FTransform> GlobalTransforms;
+	TArray<FTransform> PrevGlobalTransforms;
+	for (int32 PoseIndex = 0; PoseIndex < Poses.Num(); ++PoseIndex)
+	{
+		const FPose& Pose = Poses[PoseIndex];
+		if (Pose.LocalTransforms.IsEmpty())
+		{
+			LocalTransforms.Reset();
+			GlobalTransforms.Reset();
 		}
-#if WITH_EDITORONLY_DATA	
 		else
 		{
-			if (OutError)
+			LocalTransforms = Pose.LocalTransforms;
+			LocalTransforms[0] = LocalTransforms[0] * Pose.RootTransform;
+			FAnimationRuntime::FillUpComponentSpaceTransforms(Skeleton->GetReferenceSkeleton(), LocalTransforms, GlobalTransforms);
+		}
+
+		if (PrevGlobalTransforms.Num() == GlobalTransforms.Num())
+		{
+			const float LerpFactor = float(PoseIndex - 1) / float(Poses.Num() - 1);
+			const FColor Color = LerpColor(FColorList::Red, FColorList::Orange, LerpFactor);
+			for (int32 TransformIndex = 0; TransformIndex < GlobalTransforms.Num(); ++TransformIndex)
 			{
-				*OutError = LOCTEXT("PoseHistoryRootMotionAttributeError",
-					"Could not get Root Motion Provider. Try disabling 'Use Root Motion'.");
+				DrawDebugLine(World, PrevGlobalTransforms[TransformIndex].GetLocation(), GlobalTransforms[TransformIndex].GetLocation(), Color, false, 0.f, ESceneDepthPriorityGroup::SDPG_Foreground + 2);
 			}
-			return false;
 		}
-#endif
-	}
-	else if (UpdateMode == ERootUpdateMode::ComponentTransformDelta)
-	{
-		CurrentPose.RootTransform = ComponentTransform;
-	}
-	else
-	{
-		checkNoEntry();
-	}
 
-	return true;
-}
-
-float FPoseHistory::GetSampleTimeInterval() const
-{
-	// Reserve one pose for computing derivatives at the time horizon
-	return TimeHorizon / (Poses.Max() - 1);
+		PrevGlobalTransforms = GlobalTransforms;
+	}
+#endif // ENABLE_DRAW_DEBUG
 }
 
 //////////////////////////////////////////////////////////////////////////
diff --git a/Engine/Plugins/Experimental/Animation/PoseSearch/Source/Runtime/Public/PoseSearch/AnimNode_PoseSearchHistoryCollector.h b/Engine/Plugins/Experimental/Animation/PoseSearch/Source/Runtime/Public/PoseSearch/AnimNode_PoseSearchHistoryCollector.h
--- a/Engine/Plugins/Experimental/Animation/PoseSearch/Source/Runtime/Public/PoseSearch/AnimNode_PoseSearchHistoryCollector.h
+++ b/Engine/Plugins/Experimental/Animation/PoseSearch/Source/Runtime/Public/PoseSearch/AnimNode_PoseSearchHistoryCollector.h
@@ -17,18 +17,11 @@ public:
 
 	// The maximum amount of poses that can be stored
 	UPROPERTY(EditAnywhere, BlueprintReadWrite, Category=Settings, meta=(ClampMin="0"))
-	int32 PoseCount = 128;
+	int32 PoseCount = 64;
 	
 	// The time horizon for how long a pose will be stored in seconds
 	UPROPERTY(EditAnywhere, BlueprintReadWrite, Category=Settings, meta=(ClampMin="0"))
-	float PoseDuration = 4.f;
-
-	// If this node should compute velocities using the playing animation root motion, or 
-	// the actual movement of the component in the world. While using root motion can 
-	// be more stable, not all input animations support it, and it may produce bad results 
-	// if the current movement in the world does not match closely the animation data.
-	UPROPERTY(EditAnywhere, BlueprintReadWrite, Category = Settings, meta = (PinHiddenByDefault))
-	bool bUseRootMotion = true;
+	float PoseDuration = 1.5f;
 
 public:
 	FAnimNode_PoseSearchHistoryCollector() { }
@@ -39,6 +32,8 @@ public:
 	virtual void Evaluate_AnyThread(FPoseContext& Output) override;
 	virtual void Update_AnyThread(const FAnimationUpdateContext& Context);
 	virtual void GatherDebugData(FNodeDebugData& DebugData);
+	virtual bool HasPreUpdate() const override;
+	virtual void PreUpdate(const UAnimInstance* InAnimInstance) override;
 	// End of FAnimNode_Base interface
 
 	UE::PoseSearch::FPoseHistory& GetPoseHistory() { return PoseHistory; }
@@ -47,4 +42,9 @@ public:
 protected:
 
 	UE::PoseSearch::FPoseHistory PoseHistory;
+
+#if WITH_EDITORONLY_DATA && ENABLE_ANIM_DEBUG
+	// Whether this node was evaluated last frame
+	bool bWasEvaluated = false;
+#endif // WITH_EDITORONLY_DATA
 };
\ No newline at end of file
diff --git a/Engine/Plugins/Experimental/Animation/PoseSearch/Source/Runtime/Public/PoseSearch/PoseSearchHistory.h b/Engine/Plugins/Experimental/Animation/PoseSearch/Source/Runtime/Public/PoseSearch/PoseSearchHistory.h
--- a/Engine/Plugins/Experimental/Animation/PoseSearch/Source/Runtime/Public/PoseSearch/PoseSearchHistory.h
+++ b/Engine/Plugins/Experimental/Animation/PoseSearch/Source/Runtime/Public/PoseSearch/PoseSearchHistory.h
@@ -8,6 +8,8 @@
 #include "UObject/ObjectKey.h"
 
 struct FPoseContext;
+class USkeleton;
+class UWorld;
 
 namespace UE::PoseSearch
 {
@@ -22,18 +24,14 @@ class POSESEARCH_API FPoseHistory
 {
 public:
 
-	enum class ERootUpdateMode
-	{
-		RootMotionDelta,
-		ComponentTransformDelta,
-	};
-
 	void Init(int32 InNumPoses, float InTimeHorizon);
 	void Init(const FPoseHistory& History);
-	bool Update(float SecondsElapsed, const FPoseContext& PoseContext, FTransform ComponentTransform, FText* OutError, ERootUpdateMode UpdateMode = ERootUpdateMode::RootMotionDelta);
+	void Update(float SecondsElapsed, const FPoseContext& PoseContext, FTransform ComponentTransform);
 	float GetSampleTimeInterval() const;
 	float GetTimeHorizon() const { return TimeHorizon; }
-	bool TrySampleLocalPose(float Time, const TArray<FBoneIndexType>* RequiredBones, TArray<FTransform>* LocalPose, FTransform* RootTransform) const;
+	void GetLocalPoseAtTime(float Time, const TArray<FBoneIndexType>& RequiredBones, TArray<FTransform>& LocalPose) const;
+	void GetRootTransformAtTime(float Time, FTransform& RootTransform) const;
+	void DebugDraw(const UWorld* World, const USkeleton* Skeleton) const;
 
 private:
 
