motion matching - adding regex filtering to pose search debugger

diff --git a/Engine/Plugins/Experimental/Animation/PoseSearch/Source/Editor/Private/PoseSearchDebuggerDatabaseView.cpp b/Engine/Plugins/Experimental/Animation/PoseSearch/Source/Editor/Private/PoseSearchDebuggerDatabaseView.cpp
--- a/Engine/Plugins/Experimental/Animation/PoseSearch/Source/Editor/Private/PoseSearchDebuggerDatabaseView.cpp
+++ b/Engine/Plugins/Experimental/Animation/PoseSearch/Source/Editor/Private/PoseSearchDebuggerDatabaseView.cpp
@@ -5,6 +5,7 @@
 #include "Animation/AnimComposite.h"
 #include "Animation/AnimSequence.h"
 #include "InstancedStruct.h"
+#include "Internationalization/Regex.h"
 #include "PoseSearchDebuggerDatabaseRow.h"
 #include "PoseSearchDebuggerView.h"
 #include "PoseSearchDebuggerViewModel.h"
@@ -477,6 +478,19 @@ void SDebuggerDatabaseView::OnHideInvalidPosesCheckboxChanged(ECheckBoxState Sta
 	PopulateViewRows();
 }
 
+void SDebuggerDatabaseView::OnUseRegexCheckboxChanged(ECheckBoxState State)
+{
+	if (State == ECheckBoxState::Checked)
+	{
+		bUseRegex = true;
+	}
+	else
+	{
+		bUseRegex = false;
+	}
+	PopulateViewRows();
+}
+
 void SDebuggerDatabaseView::OnDatabaseRowSelectionChanged(TSharedPtr<FDebuggerDatabaseRowData> Row, ESelectInfo::Type SelectInfo)
 {
 	if (Row.IsValid())
@@ -515,6 +529,7 @@ void SDebuggerDatabaseView::PopulateViewRows()
 	TArray<FString> Tokens;
 	FilterString.ParseIntoArrayWS(Tokens);
 	const bool bHasNameFilter = !Tokens.IsEmpty();
+	FRegexPattern Pattern(FilterString);
 
 	for (const auto& UnfilteredRow : UnfilteredDatabaseRows)
 	{
@@ -533,21 +548,27 @@ void SDebuggerDatabaseView::PopulateViewRows()
 
 		if (bTryAddToFilteredDatabaseViewRows)
 		{
-			if (!bHideInvalidPoses || EnumHasAnyFlags(UnfilteredRow->PoseCandidateFlags, EPoseCandidateFlags::AnyValidMask))
+			bool bPassesNameFilter = true;
+			if (bHideInvalidPoses && !EnumHasAnyFlags(UnfilteredRow->PoseCandidateFlags, EPoseCandidateFlags::AnyValidMask))
 			{
-				bool bPassesNameFilter = true;
-				if (bHasNameFilter)
-				{
-					bPassesNameFilter = Algo::AllOf(Tokens, [&](FString Token)
+				bPassesNameFilter = false;
+			}
+			else if (bUseRegex)
+			{
+				FRegexMatcher Matcher(Pattern, UnfilteredRow->AssetName);
+				bPassesNameFilter = Matcher.FindNext();
+			}
+			else if (bHasNameFilter)
+			{
+				bPassesNameFilter = Algo::AllOf(Tokens, [&](FString Token)
 					{
 						return UnfilteredRow->AssetName.Contains(Token);
 					});
-				}
-
-				if (bPassesNameFilter)
-				{
-					FilteredDatabaseView.Rows.Add(UnfilteredRow);
-				}
+			}
+				
+			if (bPassesNameFilter)
+			{
+				FilteredDatabaseView.Rows.Add(UnfilteredRow);
 			}
 		}
 	}
@@ -838,6 +859,20 @@ void SDebuggerDatabaseView::Construct(const FArguments& InArgs)
 						.Text(LOCTEXT("PoseSearchDebuggerHideInvalidPosesFlag", "Hide Invalid Poses"))
 					]
 				]
+				+ SHorizontalBox::Slot()
+				.AutoWidth()
+				.Padding(10, 5, 10, 5)
+				[
+					SNew(SCheckBox)
+					.OnCheckStateChanged_Lambda([this](ECheckBoxState State)
+						{
+							SDebuggerDatabaseView::OnUseRegexCheckboxChanged(State);
+						})
+				[
+					SNew(STextBlock)
+					.Text(LOCTEXT("PoseSearchDebuggerUseRegexFlag", "Use Regex"))
+				]
+				]
 			]
 		
 			+ SVerticalBox::Slot()
diff --git a/Engine/Plugins/Experimental/Animation/PoseSearch/Source/Editor/Private/PoseSearchDebuggerDatabaseView.h b/Engine/Plugins/Experimental/Animation/PoseSearch/Source/Editor/Private/PoseSearchDebuggerDatabaseView.h
--- a/Engine/Plugins/Experimental/Animation/PoseSearch/Source/Editor/Private/PoseSearchDebuggerDatabaseView.h
+++ b/Engine/Plugins/Experimental/Animation/PoseSearch/Source/Editor/Private/PoseSearchDebuggerDatabaseView.h
@@ -69,6 +69,9 @@ private:
 	/** Called when the HideInvalidPoses Checkbox is Toggled */
 	void OnHideInvalidPosesCheckboxChanged(ECheckBoxState State);
 
+	/** Called when the UseRegex Checkbox is Toggled */
+	void OnUseRegexCheckboxChanged(ECheckBoxState State);
+	
 	/** Row selection to update model view */
 	void OnDatabaseRowSelectionChanged(TSharedPtr<FDebuggerDatabaseRowData> Row, ESelectInfo::Type SelectInfo);
 
@@ -138,6 +141,9 @@ private:
 
 	/* Bool used to hide invalid poses from display */
 	bool bHideInvalidPoses = false;
+
+	/* Bool used to use FilterText as regex filter*/
+	bool bUseRegex = false;
 };
 
 } // namespace UE::PoseSearch
