motion matching - populating UPoseSearchDatabaseStatistics Kinematic Information

diff --git a/Engine/Plugins/Experimental/Animation/PoseSearch/Source/Editor/Private/PoseSearchDatabaseEditorReflection.cpp b/Engine/Plugins/Experimental/Animation/PoseSearch/Source/Editor/Private/PoseSearchDatabaseEditorReflection.cpp
--- a/Engine/Plugins/Experimental/Animation/PoseSearch/Source/Editor/Private/PoseSearchDatabaseEditorReflection.cpp
+++ b/Engine/Plugins/Experimental/Animation/PoseSearch/Source/Editor/Private/PoseSearchDatabaseEditorReflection.cpp
@@ -94,7 +94,7 @@ void UPoseSearchDatabaseStatistics::Initialize(const UPoseSearchDatabase* PoseSe
 	if (PoseSearchDatabase)
 	{
 		const FPoseSearchIndex& SearchIndex = PoseSearchDatabase->GetSearchIndex();
-		// General information
+		// General Information
 	
 		AnimationSequences = PoseSearchDatabase->AnimationAssets.Num();
 			
@@ -113,15 +113,18 @@ void UPoseSearchDatabaseStatistics::Initialize(const UPoseSearchDatabase* PoseSe
 			SearchableTime = FText::Format(TimeFormat, static_cast<double>(NumOfSearchablePoses) / SampleRate);
 		}
 			
-		// Velocity information
+		// Kinematic Information
 	
-		// TODO: Set values once they can be queried from the PoseSearchIndex.
-			
+		AverageSpeed = SearchIndex.Stats.AverageSpeed;
+		MaxSpeed = SearchIndex.Stats.MaxSpeed;
+		AverageAcceleration = SearchIndex.Stats.AverageAcceleration;
+		MaxAcceleration = SearchIndex.Stats.MaxAcceleration;
+
 		// Principal Component Analysis
 			
-		ExplainedVariance = SearchIndex.PCAExplainedVariance;
+		ExplainedVariance = SearchIndex.PCAExplainedVariance * 100.f;
 			
-		// Memory information
+		// Memory Information
 			
 		{
 			const uint32 ValuesBytesSize = SearchIndex.Values.GetAllocatedSize();
diff --git a/Engine/Plugins/Experimental/Animation/PoseSearch/Source/Editor/Public/PoseSearchDatabaseEditorReflection.h b/Engine/Plugins/Experimental/Animation/PoseSearch/Source/Editor/Public/PoseSearchDatabaseEditorReflection.h
--- a/Engine/Plugins/Experimental/Animation/PoseSearch/Source/Editor/Public/PoseSearchDatabaseEditorReflection.h
+++ b/Engine/Plugins/Experimental/Animation/PoseSearch/Source/Editor/Public/PoseSearchDatabaseEditorReflection.h
@@ -105,7 +105,7 @@ class UPoseSearchDatabaseStatistics : public UObject
 
 public:
 	
-	// General information
+	// General Information
 
 	UPROPERTY(VisibleAnywhere, Category = "General Information")
 	uint32 AnimationSequences;
@@ -122,23 +122,23 @@ public:
 	UPROPERTY(VisibleAnywhere, Category = "General Information")
 	FText SearchableTime;
 
-	// Velocity Information
+	// Kinematic Information
 
-	UPROPERTY(VisibleAnywhere, Category = "Velocity Information")
-	double AverageVelocity;
+	UPROPERTY(VisibleAnywhere, Category = "Kinematic Information")
+	float AverageSpeed;
 
-	UPROPERTY(VisibleAnywhere, Category = "Velocity Information")
-	double MaxVelocity;
+	UPROPERTY(VisibleAnywhere, Category = "Kinematic Information")
+	float MaxSpeed;
 
-	UPROPERTY(VisibleAnywhere, Category = "Velocity Information")
-	double AverageAcceleration;
+	UPROPERTY(VisibleAnywhere, Category = "Kinematic Information")
+	float AverageAcceleration;
 
-	UPROPERTY(VisibleAnywhere, Category = "Velocity Information")
-	double MaxAcceleration;
+	UPROPERTY(VisibleAnywhere, Category = "Kinematic Information")
+	float MaxAcceleration;
 
-	// Principal Component Analysis (PCA) Information
+	// Principal Component Analysis Information
 
-	UPROPERTY(VisibleAnywhere, Category = "Principal Component Analysis (PCA) Information")
+	UPROPERTY(VisibleAnywhere, Category = "Principal Component Analysis Information", meta = (Units = "Percent"))
 	float ExplainedVariance;
 
 	// Memory information
diff --git a/Engine/Plugins/Experimental/Animation/PoseSearch/Source/Runtime/Private/PoseSearchAssetIndexer.cpp b/Engine/Plugins/Experimental/Animation/PoseSearch/Source/Runtime/Private/PoseSearchAssetIndexer.cpp
--- a/Engine/Plugins/Experimental/Animation/PoseSearch/Source/Runtime/Private/PoseSearchAssetIndexer.cpp
+++ b/Engine/Plugins/Experimental/Animation/PoseSearch/Source/Runtime/Private/PoseSearchAssetIndexer.cpp
@@ -74,15 +74,29 @@ static FSamplingParam WrapOrClampSamplingParam(bool bCanWrap, float SamplingPara
 
 //////////////////////////////////////////////////////////////////////////
 // FAssetIndexer
-void FAssetIndexer::Reset()
+void FAssetIndexer::FOutput::Reset()
 {
-	Output.FirstIndexedSample = 0;
-	Output.LastIndexedSample = 0;
-	Output.NumIndexedPoses = 0;
+	FirstIndexedSample = 0;
+	LastIndexedSample = 0;
+	NumIndexedPoses = 0;
+
+	FeatureVectorTable.Reset(0);
+	PoseMetadata.Reset(0);
+}
 
-	Output.FeatureVectorTable.Reset(0);
-	Output.PoseMetadata.Reset(0);
-	Output.AllFeaturesNotAdded.Reset();
+void FAssetIndexer::FStats::Reset()
+{
+	NumAccumulatedSamples = 0;
+	AccumulatedSpeed = 0.f;
+	MaxSpeed = 0.f;
+	AccumulatedAcceleration = 0.f;
+	MaxAcceleration = 0.f;
+}
+
+void FAssetIndexer::Reset()
+{
+	Output.Reset();
+	Stats.Reset();
 }
 
 void FAssetIndexer::Init(const FAssetIndexingContext& InIndexingContext, const FBoneContainer& InBoneContainer)
@@ -135,9 +149,56 @@ bool FAssetIndexer::Process()
 		OutputPoseMetadata = GetMetadata(SampleIdx);
 	}
 
+	// Computing stats
+	ComputeStats();
+
 	return true;
 }
 
+void FAssetIndexer::ComputeStats()
+{
+	Stats.Reset();
+
+	const FAssetSamplingContext* SamplingContext = IndexingContext.SamplingContext;
+	check(SamplingContext->FiniteDelta > UE_KINDA_SMALL_NUMBER);
+
+	for (int32 SampleIdx = IndexingContext.BeginSampleIdx; SampleIdx != IndexingContext.EndSampleIdx; ++SampleIdx)
+	{
+		const int32 VectorIdx = SampleIdx - IndexingContext.BeginSampleIdx;
+
+		const float SampleTime = FMath::Min(SampleIdx * IndexingContext.Schema->GetSamplingInterval(), IndexingContext.AssetSampler->GetPlayLength());
+
+		bool AnyClamped = false;
+		const FTransform TrajTransformsPast = GetTransform(SampleTime - SamplingContext->FiniteDelta, AnyClamped);
+		if (!AnyClamped)
+		{
+			const FTransform TrajTransformsPresent = GetTransform(SampleTime, AnyClamped);
+			if (!AnyClamped)
+			{
+				const FTransform TrajTransformsFuture = GetTransform(SampleTime + SamplingContext->FiniteDelta, AnyClamped);
+				if (!AnyClamped)
+				{
+					// if any transform is clamped we just skip the sample entirely
+					const FVector LinearVelocityPresent = (TrajTransformsPresent.GetTranslation() - TrajTransformsPast.GetTranslation()) / SamplingContext->FiniteDelta;
+					const FVector LinearVelocityFuture = (TrajTransformsFuture.GetTranslation() - TrajTransformsPresent.GetTranslation()) / SamplingContext->FiniteDelta;
+					const FVector LinearAcceleration = (LinearVelocityFuture - LinearVelocityPresent) / SamplingContext->FiniteDelta;
+
+					const float Speed = LinearVelocityPresent.Length();
+					const float Acceleration = LinearAcceleration.Length();
+
+					Stats.AccumulatedSpeed += Speed;
+					Stats.MaxSpeed = FMath::Max(Stats.MaxSpeed, Speed);
+
+					Stats.AccumulatedAcceleration += Acceleration;
+					Stats.MaxAcceleration = FMath::Max(Stats.MaxAcceleration, Acceleration);
+
+					++Stats.NumAccumulatedSamples;
+				}
+			}
+		}
+	}
+}
+
 FAssetIndexer::FSampleInfo FAssetIndexer::GetSampleInfo(float SampleTime) const
 {
 	FSampleInfo Sample;
diff --git a/Engine/Plugins/Experimental/Animation/PoseSearch/Source/Runtime/Private/PoseSearchAssetIndexer.h b/Engine/Plugins/Experimental/Animation/PoseSearch/Source/Runtime/Private/PoseSearchAssetIndexer.h
--- a/Engine/Plugins/Experimental/Animation/PoseSearch/Source/Runtime/Private/PoseSearchAssetIndexer.h
+++ b/Engine/Plugins/Experimental/Animation/PoseSearch/Source/Runtime/Private/PoseSearchAssetIndexer.h
@@ -27,15 +27,28 @@ public:
 
 		TArray<float> FeatureVectorTable;
 		TArray<FPoseSearchPoseMetadata> PoseMetadata;
-		TBitArray<> AllFeaturesNotAdded;
-	} Output;
+		
+		void Reset();
+	};
+
+	struct FStats
+	{
+		int32 NumAccumulatedSamples = 0;
+		float AccumulatedSpeed = 0.f;
+		float MaxSpeed = 0.f;
+		float AccumulatedAcceleration = 0.f;
+		float MaxAcceleration = 0.f;
+
+		void Reset();
+	};
 
 	void Reset();
 	void Init(const FAssetIndexingContext& IndexingContext, const FBoneContainer& InBoneContainer);
 	bool Process();
+	const FOutput& GetOutput() const { return Output; }
+	const FStats& GetStats() const { return Stats; }
 
-public: // IAssetIndexer
-
+	// IAssetIndexer
 	const FAssetIndexingContext& GetIndexingContext() const override { return IndexingContext; }
 	FTransform GetTransform(float SampleTime, bool& Clamped, int8 SchemaBoneIdx = RootSchemaBoneIdx) override;
 	FTransform GetComponentSpaceTransform(float SampleTime, bool& Clamped, int8 SchemaBoneIdx = RootSchemaBoneIdx) override;
@@ -56,10 +69,13 @@ private:
 	FTransform MirrorTransform(const FTransform& Transform) const;
 	CachedEntry& GetEntry(float SampleTime);
 	FTransform CalculateComponentSpaceTransform(CachedEntry& Entry, int8 SchemaBoneIdx);
+	void ComputeStats();
 
 	FBoneContainer BoneContainer;
 	FAssetIndexingContext IndexingContext;
 	TMap<float, CachedEntry> CachedEntries;
+	FOutput Output;
+	FStats Stats;
 };
 
 } // namespace UE::PoseSearch
diff --git a/Engine/Plugins/Experimental/Animation/PoseSearch/Source/Runtime/Private/PoseSearchDatabaseIndexingContext.cpp b/Engine/Plugins/Experimental/Animation/PoseSearch/Source/Runtime/Private/PoseSearchDatabaseIndexingContext.cpp
--- a/Engine/Plugins/Experimental/Animation/PoseSearch/Source/Runtime/Private/PoseSearchDatabaseIndexingContext.cpp
+++ b/Engine/Plugins/Experimental/Animation/PoseSearch/Source/Runtime/Private/PoseSearchDatabaseIndexingContext.cpp
@@ -182,10 +182,12 @@ void FDatabaseIndexingContext::JoinIndex()
 	SearchIndexBase->Values.Reset();
 	SearchIndexBase->PoseMetadata.Reset();
 	SearchIndexBase->OverallFlags = EPoseSearchPoseFlags::None;
+	SearchIndexBase->Stats = FPoseSearchStats();
 
+	int32 NumAccumulatedSamples = 0;
 	for (int32 AssetIdx = 0; AssetIdx != SearchIndexBase->Assets.Num(); ++AssetIdx)
 	{
-		const FAssetIndexer::FOutput& Output = Indexers[AssetIdx].Output;
+		const FAssetIndexer::FOutput& Output = Indexers[AssetIdx].GetOutput();
 
 		FPoseSearchIndexAsset& SearchIndexAsset = SearchIndexBase->Assets[AssetIdx];
 		SearchIndexAsset.NumPoses = Output.NumIndexedPoses;
@@ -205,6 +207,21 @@ void FDatabaseIndexingContext::JoinIndex()
 
 		TotalPoses += Output.NumIndexedPoses;
 		TotalFloats += Output.FeatureVectorTable.Num();
+
+		const FAssetIndexer::FStats& Stats = Indexers[AssetIdx].GetStats();
+		SearchIndexBase->Stats.AverageSpeed += Stats.AccumulatedSpeed;
+		SearchIndexBase->Stats.MaxSpeed = FMath::Max(SearchIndexBase->Stats.MaxSpeed, Stats.MaxSpeed);
+		SearchIndexBase->Stats.AverageAcceleration += Stats.AccumulatedAcceleration;
+		SearchIndexBase->Stats.MaxAcceleration = FMath::Max(SearchIndexBase->Stats.MaxAcceleration, Stats.MaxAcceleration);
+
+		NumAccumulatedSamples += Stats.NumAccumulatedSamples;
+	}
+
+	if (NumAccumulatedSamples > 0)
+	{
+		const float Denom = 1.f / float(NumAccumulatedSamples);
+		SearchIndexBase->Stats.AverageSpeed *= Denom;
+		SearchIndexBase->Stats.AverageAcceleration *= Denom;
 	}
 
 	SearchIndexBase->NumPoses = TotalPoses;
diff --git a/Engine/Plugins/Experimental/Animation/PoseSearch/Source/Runtime/Private/PoseSearchDerivedData.cpp b/Engine/Plugins/Experimental/Animation/PoseSearch/Source/Runtime/Private/PoseSearchDerivedData.cpp
--- a/Engine/Plugins/Experimental/Animation/PoseSearch/Source/Runtime/Private/PoseSearchDerivedData.cpp
+++ b/Engine/Plugins/Experimental/Animation/PoseSearch/Source/Runtime/Private/PoseSearchDerivedData.cpp
@@ -443,9 +443,7 @@ static void PreprocessSearchIndexPCAData(FPoseSearchIndex& SearchIndex, int32 Nu
 	SearchIndex.Mean.AddZeroed(NumDimensions);
 	SearchIndex.PCAProjectionMatrix.AddZeroed(NumDimensions * NumberOfPrincipalComponents);
 
-#if WITH_EDITORONLY_DATA
 	SearchIndex.PCAExplainedVariance = 0.f;
-#endif
 
 	if (NumDimensions > 0)
 	{
@@ -507,12 +505,10 @@ static void PreprocessSearchIndexPCAData(FPoseSearchIndex& SearchIndex, int32 Nu
 			AccumulatedVariance += EigenValues[Indexer[PCAComponentIndex]];
 		}
 
-#if WITH_EDITORONLY_DATA
 		// calculating the total variance knowing that eigen values measure variance along the principal components:
 		const float TotalVariance = EigenValues.sum();
 		// and explained variance as ratio between AccumulatedVariance and TotalVariance: https://ro-che.info/articles/2017-12-11-pca-explained-variance
 		SearchIndex.PCAExplainedVariance = TotalVariance > UE_KINDA_SMALL_NUMBER ? AccumulatedVariance / TotalVariance : 0.f;
-#endif // WITH_EDITORONLY_DATA
 
 		MapPCAValues = CenteredValues * PCAProjectionMatrix;
 
@@ -975,10 +971,6 @@ void FPoseSearchDatabaseAsyncCacheTask::OnGetComplete(UE::DerivedData::FCacheGet
 				
 				TArray<float> Deviation = FMeanDeviationCalculator::Calculate(SearchIndexBases, Schemas);
 
-				#if WITH_EDITORONLY_DATA
-				SearchIndex.Deviation = Deviation;
-				#endif // WITH_EDITORONLY_DATA
-
 				// Building FPoseSearchIndex
 				PreprocessSearchIndexWeights(SearchIndex, Database->Schema, Deviation);
 				if (Owner.IsCanceled())
diff --git a/Engine/Plugins/Experimental/Animation/PoseSearch/Source/Runtime/Private/PoseSearchDerivedDataKey.cpp b/Engine/Plugins/Experimental/Animation/PoseSearch/Source/Runtime/Private/PoseSearchDerivedDataKey.cpp
--- a/Engine/Plugins/Experimental/Animation/PoseSearch/Source/Runtime/Private/PoseSearchDerivedDataKey.cpp
+++ b/Engine/Plugins/Experimental/Animation/PoseSearch/Source/Runtime/Private/PoseSearchDerivedDataKey.cpp
@@ -41,7 +41,7 @@ FKeyBuilder::FKeyBuilder(const UObject* Object, bool bUseDataVer)
 	if (bUseDataVer)
 	{
 		// used to invalidate the key without having to change POSESEARCHDB_DERIVEDDATA_VER all the times
-		int32 POSESEARCHDB_DERIVEDDATA_VER_SMALL = 31;
+		int32 POSESEARCHDB_DERIVEDDATA_VER_SMALL = 34;
 		FGuid VersionGuid = FDevSystemGuids::GetSystemGuid(FDevSystemGuids::Get().POSESEARCHDB_DERIVEDDATA_VER);
 
 		*this << VersionGuid;
diff --git a/Engine/Plugins/Experimental/Animation/PoseSearch/Source/Runtime/Private/PoseSearchIndex.cpp b/Engine/Plugins/Experimental/Animation/PoseSearch/Source/Runtime/Private/PoseSearchIndex.cpp
--- a/Engine/Plugins/Experimental/Animation/PoseSearch/Source/Runtime/Private/PoseSearchIndex.cpp
+++ b/Engine/Plugins/Experimental/Animation/PoseSearch/Source/Runtime/Private/PoseSearchIndex.cpp
@@ -31,6 +31,18 @@ void CompareFeatureVectors(TConstArrayView<float> A, TConstArrayView<float> B, T
 
 } // namespace UE::PoseSearch
 
+//////////////////////////////////////////////////////////////////////////
+// FPoseSearchStats
+FArchive& operator<<(FArchive& Ar, FPoseSearchStats& Stats)
+{
+	Ar << Stats.AverageSpeed;
+	Ar << Stats.MaxSpeed;
+	Ar << Stats.AverageAcceleration;
+	Ar << Stats.MaxAcceleration;
+
+	return Ar;
+}
+
 //////////////////////////////////////////////////////////////////////////
 // FPoseSearchBaseIndex
 const FPoseSearchIndexAsset& FPoseSearchIndexBase::GetAssetForPose(int32 PoseIdx) const
@@ -86,8 +98,7 @@ bool FPoseSearchIndexBase::IsEmpty() const
 
 void FPoseSearchIndexBase::Reset()
 {
-	FPoseSearchIndexBase Default;
-	*this = Default;
+	*this = FPoseSearchIndexBase();
 }
 
 FArchive& operator<<(FArchive& Ar, FPoseSearchIndexBase& Index)
@@ -130,6 +141,8 @@ FArchive& operator<<(FArchive& Ar, FPoseSearchIndexBase& Index)
 
 	Ar << Index.MinCostAddend;
 
+	Ar << Index.Stats;
+
 	return Ar;
 }
 
@@ -142,10 +155,7 @@ FPoseSearchIndex::FPoseSearchIndex(const FPoseSearchIndex& Other)
 	, Mean(Other.Mean)
 	, WeightsSqrt(Other.WeightsSqrt)
 	, KDTree(Other.KDTree)
-#if WITH_EDITORONLY_DATA
 	, PCAExplainedVariance(Other.PCAExplainedVariance)
-	, Deviation(Other.Deviation)
-#endif // WITH_EDITORONLY_DATA
 {
 	check(!PCAValues.IsEmpty() || KDTree.DataSource.PointCount == 0);
 	KDTree.DataSource.Data = PCAValues.IsEmpty() ? nullptr : PCAValues.GetData();
@@ -246,7 +256,6 @@ FArchive& operator<<(FArchive& Ar, FPoseSearchIndex& Index)
 	Serialize(Ar, Index.KDTree, Index.PCAValues.GetData());
 
 	Ar << Index.PCAExplainedVariance;
-	Ar << Index.Deviation;
 
 	return Ar;
 }
\ No newline at end of file
diff --git a/Engine/Plugins/Experimental/Animation/PoseSearch/Source/Runtime/Public/PoseSearch/PoseSearchIndex.h b/Engine/Plugins/Experimental/Animation/PoseSearch/Source/Runtime/Public/PoseSearch/PoseSearchIndex.h
--- a/Engine/Plugins/Experimental/Animation/PoseSearch/Source/Runtime/Public/PoseSearch/PoseSearchIndex.h
+++ b/Engine/Plugins/Experimental/Animation/PoseSearch/Source/Runtime/Public/PoseSearch/PoseSearchIndex.h
@@ -82,7 +82,7 @@ USTRUCT()
 struct POSESEARCH_API FPoseSearchIndexAsset
 {
 	GENERATED_BODY()
-public:
+
 	FPoseSearchIndexAsset() {}
 
 	FPoseSearchIndexAsset(
@@ -124,6 +124,25 @@ public:
 	bool IsPoseInRange(int32 PoseIdx) const { return (PoseIdx >= FirstPoseIdx) && (PoseIdx < FirstPoseIdx + NumPoses); }
 };
 
+USTRUCT()
+struct POSESEARCH_API FPoseSearchStats
+{
+	GENERATED_BODY()
+
+	UPROPERTY(meta = (NeverInHash))
+	float AverageSpeed = 0.f;
+
+	UPROPERTY(meta = (NeverInHash))
+	float MaxSpeed = 0.f;
+
+	UPROPERTY(meta = (NeverInHash))
+	float AverageAcceleration = 0.f;
+
+	UPROPERTY(meta = (NeverInHash))
+	float MaxAcceleration = 0.f;
+
+	friend FArchive& operator<<(FArchive& Ar, FPoseSearchStats& Stats);
+};
 
 /**
 * case class for FPoseSearchIndex. building block used to gather data for data mining and calculate weights, pca, kdtree stuff
@@ -133,7 +152,7 @@ struct POSESEARCH_API FPoseSearchIndexBase
 {
 	GENERATED_BODY()
 
-	UPROPERTY(Category = Info, VisibleAnywhere, meta = (NeverInHash))
+	UPROPERTY(meta = (NeverInHash))
 	int32 NumPoses = 0;
 
 	UPROPERTY(meta = (NeverInHash))
@@ -149,9 +168,13 @@ struct POSESEARCH_API FPoseSearchIndexBase
 	TArray<FPoseSearchIndexAsset> Assets;
 
 	// minimum of the database metadata CostAddend: it represents the minimum cost of any search for the associated database (we'll skip the search in case the search result total cost is already less than MinCostAddend)
-	UPROPERTY(Category = Info, VisibleAnywhere, meta = (NeverInHash))
+	UPROPERTY(meta = (NeverInHash))
 	float MinCostAddend = -MAX_FLT;
 
+	// @todo: this property should be editor only
+	UPROPERTY(meta = (NeverInHash))
+	FPoseSearchStats Stats;
+
 	bool IsValidPoseIndex(int32 PoseIdx) const { return PoseIdx < NumPoses; }
 	bool IsEmpty() const;
 
@@ -193,12 +216,10 @@ struct POSESEARCH_API FPoseSearchIndex : public FPoseSearchIndexBase
 
 	UE::PoseSearch::FKDTree KDTree;
 
+	// @todo: this property should be editor only
 	UPROPERTY(Category = Info, VisibleAnywhere, meta = (NeverInHash))
 	float PCAExplainedVariance = 0.f;
 
-	UPROPERTY(Category = Info, VisibleAnywhere, meta = (NeverInHash))
-	TArray<float> Deviation;
-
 	FPoseSearchIndex() = default;
 	~FPoseSearchIndex() = default;
 	FPoseSearchIndex(const FPoseSearchIndex& Other); // custom copy constructor to deal with the KDTree DataSrc
