motion matching - search throttle crash fix and performance improvements (skip generating the query if throttling)

diff --git a/Engine/Plugins/Experimental/Animation/PoseSearch/Source/Editor/Private/PoseSearchDebuggerReflection.h b/Engine/Plugins/Experimental/Animation/PoseSearch/Source/Editor/Private/PoseSearchDebuggerReflection.h
--- a/Engine/Plugins/Experimental/Animation/PoseSearch/Source/Editor/Private/PoseSearchDebuggerReflection.h
+++ b/Engine/Plugins/Experimental/Animation/PoseSearch/Source/Editor/Private/PoseSearchDebuggerReflection.h
@@ -36,9 +36,9 @@ public:
 	UPROPERTY(VisibleAnywhere, Category="Motion Matching State", Meta=(DisplayName="Current Database"))
 	FString CurrentDatabaseName = "";
 
-	/** Time since last PoseSearch jump */
+	/** Time since last PoseSearch */
 	UPROPERTY(VisibleAnywhere, Category="Motion Matching State")
-	float ElapsedPoseJumpTime = 0.0f;
+	float ElapsedPoseSearchTime = 0.0f;
 
 	UPROPERTY(VisibleAnywhere, Category = "Motion Matching State")
 	FString AssetPlayerAssetName = "";
diff --git a/Engine/Plugins/Experimental/Animation/PoseSearch/Source/Editor/Private/PoseSearchDebuggerView.cpp b/Engine/Plugins/Experimental/Animation/PoseSearch/Source/Editor/Private/PoseSearchDebuggerView.cpp
--- a/Engine/Plugins/Experimental/Animation/PoseSearch/Source/Editor/Private/PoseSearchDebuggerView.cpp
+++ b/Engine/Plugins/Experimental/Animation/PoseSearch/Source/Editor/Private/PoseSearchDebuggerView.cpp
@@ -96,7 +96,7 @@ void SDebuggerDetailsView::UpdateReflection(const FTraceMotionMatchingStateMessa
 		int32 CurrentDbPoseIdx = State.GetCurrentDatabasePoseIndex();
 
 		Reflection->CurrentDatabaseName = CurrentDatabase->GetName();
-		Reflection->ElapsedPoseJumpTime = State.ElapsedPoseJumpTime;
+		Reflection->ElapsedPoseSearchTime = State.ElapsedPoseSearchTime;
 
 		Reflection->AssetPlayerAssetName = "None";
 		if (const FPoseSearchIndexAsset* IndexAsset = CurrentSearchIndex.GetAssetForPoseSafe(CurrentDbPoseIdx))
diff --git a/Engine/Plugins/Experimental/Animation/PoseSearch/Source/Runtime/Private/AnimNode_MotionMatching.cpp b/Engine/Plugins/Experimental/Animation/PoseSearch/Source/Runtime/Private/AnimNode_MotionMatching.cpp
--- a/Engine/Plugins/Experimental/Animation/PoseSearch/Source/Runtime/Private/AnimNode_MotionMatching.cpp
+++ b/Engine/Plugins/Experimental/Animation/PoseSearch/Source/Runtime/Private/AnimNode_MotionMatching.cpp
@@ -120,11 +120,11 @@ void FAnimNode_MotionMatching::UpdateAssetPlayer(const FAnimationUpdateContext&
 	);
 
 	// If a new pose is requested, blend into the new asset via BlendStackNode
-	const FPoseSearchIndexAsset* SearchIndexAsset = MotionMatchingState.CurrentSearchResult.GetSearchIndexAsset();
-	if (SearchIndexAsset)
+	if (MotionMatchingState.bJumpedToPose)
 	{
+		const FPoseSearchIndexAsset* SearchIndexAsset = MotionMatchingState.CurrentSearchResult.GetSearchIndexAsset();
 		const UPoseSearchDatabase* Database = MotionMatchingState.CurrentSearchResult.Database.Get();
-		if (Database && Database->Schema && EnumHasAnyFlags(MotionMatchingState.Flags, EMotionMatchingFlags::JumpedToPose))
+		if (SearchIndexAsset && Database && Database->Schema)
 		{
 			if (const FPoseSearchDatabaseAnimationAssetBase* DatabaseAsset = Database->GetAnimationAssetBase(*SearchIndexAsset))
 			{
@@ -132,7 +132,6 @@ void FAnimNode_MotionMatching::UpdateAssetPlayer(const FAnimationUpdateContext&
 					DatabaseAsset->IsLooping(), SearchIndexAsset->bMirrored, Database->Schema->MirrorDataTable.Get(),
 					Settings.MaxActiveBlends, Settings.BlendTime, Settings.BlendProfile, Settings.BlendOption, SearchIndexAsset->BlendParameters, MotionMatchingState.WantedPlayRate);
 			}
-
 		}
 	}
 	BlendStackNode.UpdatePlayRate(MotionMatchingState.WantedPlayRate);
diff --git a/Engine/Plugins/Experimental/Animation/PoseSearch/Source/Runtime/Private/PoseSearchLibrary.cpp b/Engine/Plugins/Experimental/Animation/PoseSearch/Source/Runtime/Private/PoseSearchLibrary.cpp
--- a/Engine/Plugins/Experimental/Animation/PoseSearch/Source/Runtime/Private/PoseSearchLibrary.cpp
+++ b/Engine/Plugins/Experimental/Animation/PoseSearch/Source/Runtime/Private/PoseSearchLibrary.cpp
@@ -26,7 +26,7 @@ void FMotionMatchingState::Reset()
 {
 	CurrentSearchResult.Reset();
 	// Set the elapsed time to INFINITY to trigger a search right away
-	ElapsedPoseJumpTime = INFINITY;
+	ElapsedPoseSearchTime = INFINITY;
 	PoseIndicesHistory.Reset();
 	WantedPlayRate = 1.f;
 }
@@ -135,9 +135,7 @@ void FMotionMatchingState::JumpToPose(const FAnimationUpdateContext& Context, co
 	// Remember which pose and sequence we're playing from the database
 	CurrentSearchResult = Result;
 
-	ElapsedPoseJumpTime = 0.0f;
-
-	Flags |= EMotionMatchingFlags::JumpedToPose;
+	bJumpedToPose = true;
 }
 
 #if UE_POSE_SEARCH_TRACE_ENABLED
@@ -147,7 +145,8 @@ static void TraceMotionMatchingState(
 	UE::PoseSearch::FSearchContext& SearchContext,
 	const FMotionMatchingState& MotionMatchingState,
 	const FTrajectorySampleRange& Trajectory,
-	const UE::PoseSearch::FSearchResult& LastResult)
+	const UE::PoseSearch::FSearchResult& LastResult,
+	bool bSearch)
 {
 	using namespace UE::PoseSearch;
 
@@ -202,7 +201,7 @@ static void TraceMotionMatchingState(
 		DbEntry.PoseEntries.Add(PoseEntry);
 	}
 
-	if (MotionMatchingState.CurrentSearchResult.ContinuingPoseCost.IsValid())
+	if (bSearch && MotionMatchingState.CurrentSearchResult.ContinuingPoseCost.IsValid())
 	{
 		check(LastResult.IsValid());
 
@@ -216,7 +215,7 @@ static void TraceMotionMatchingState(
 		PoseEntry.PoseCandidateFlags = EPoseCandidateFlags::Valid_ContinuingPose;
 	}
 
-	if (MotionMatchingState.CurrentSearchResult.PoseCost.IsValid())
+	if (bSearch && MotionMatchingState.CurrentSearchResult.PoseCost.IsValid())
 	{
 		const int32 DbEntryIdx = AddUniqueDatabase(TraceState.DatabaseEntries, MotionMatchingState.CurrentSearchResult.Database.Get(), SearchContext);
 		FTraceMotionMatchingStateDatabaseEntry& DbEntry = TraceState.DatabaseEntries[DbEntryIdx];
@@ -250,7 +249,7 @@ static void TraceMotionMatchingState(
 	}
 
 	TraceState.SearchableAssetId = FTraceMotionMatchingState::GetIdFromObject(Searchable);
-	TraceState.ElapsedPoseJumpTime = MotionMatchingState.ElapsedPoseJumpTime;
+	TraceState.ElapsedPoseSearchTime = MotionMatchingState.ElapsedPoseSearchTime;
 	TraceState.AssetPlayerTime = MotionMatchingState.CurrentSearchResult.AssetTime;
 	TraceState.DeltaTime = DeltaTime;
 
@@ -282,52 +281,47 @@ void UpdateMotionMatchingState(
 
 	const float DeltaTime = Context.GetDeltaTime();
 
-	// Reset State Flags
-	InOutMotionMatchingState.Flags = EMotionMatchingFlags::None;
+	InOutMotionMatchingState.bJumpedToPose = false;
 
 #if UE_POSE_SEARCH_TRACE_ENABLED
 	// Record Current Pose Index for Debugger
 	const FSearchResult LastResult = InOutMotionMatchingState.CurrentSearchResult;
 #endif
 
-	// Check if we can advance.
-	const bool bCanAdvance = InOutMotionMatchingState.CanAdvance(DeltaTime);
-
-	// If we can't advance or enough time has elapsed since the last pose jump then search
+	// Build the search context
 	FSearchContext SearchContext;
-	if (!bCanAdvance || (InOutMotionMatchingState.ElapsedPoseJumpTime >= Settings.SearchThrottleTime))
+	SearchContext.ActiveTagsContainer = ActiveTagsContainer;
+	SearchContext.Trajectory = &Trajectory;
+	SearchContext.OwningComponent = Context.AnimInstanceProxy->GetSkelMeshComponent();
+	SearchContext.BoneContainer = &Context.AnimInstanceProxy->GetRequiredBones();
+	SearchContext.bIsTracing = IsTracing(Context);
+	SearchContext.bForceInterrupt = bForceInterrupt;
+	SearchContext.bCanAdvance = InOutMotionMatchingState.CanAdvance(DeltaTime);
+	SearchContext.CurrentResult = InOutMotionMatchingState.CurrentSearchResult;
+	SearchContext.PoseJumpThresholdTime = Settings.PoseJumpThresholdTime;
+	SearchContext.PoseIndicesHistory = &InOutMotionMatchingState.PoseIndicesHistory;
+
+	if (IPoseHistoryProvider* PoseHistoryProvider = Context.GetMessage<IPoseHistoryProvider>())
 	{
-		// Build the search context
-		SearchContext.ActiveTagsContainer = ActiveTagsContainer;
-		SearchContext.Trajectory = &Trajectory;
-		SearchContext.OwningComponent = Context.AnimInstanceProxy->GetSkelMeshComponent();
-		SearchContext.BoneContainer = &Context.AnimInstanceProxy->GetRequiredBones();
-		SearchContext.bIsTracing = IsTracing(Context);
-		SearchContext.bForceInterrupt = bForceInterrupt;
-		SearchContext.bCanAdvance = bCanAdvance;
-		SearchContext.CurrentResult = InOutMotionMatchingState.CurrentSearchResult;
-		SearchContext.PoseJumpThresholdTime = Settings.PoseJumpThresholdTime;
-		SearchContext.PoseIndicesHistory = &InOutMotionMatchingState.PoseIndicesHistory;
-
-		IPoseHistoryProvider* PoseHistoryProvider = Context.GetMessage<IPoseHistoryProvider>();
-		if (PoseHistoryProvider)
-		{
-			SearchContext.History = &PoseHistoryProvider->GetPoseHistory();
-		}
+		SearchContext.History = &PoseHistoryProvider->GetPoseHistory();
+	}
 
-		if (const FPoseSearchIndexAsset* CurrentIndexAsset = InOutMotionMatchingState.CurrentSearchResult.GetSearchIndexAsset())
-		{
-			SearchContext.QueryMirrorRequest =
-				CurrentIndexAsset->bMirrored ?
-				EPoseSearchBooleanRequest::TrueValue :
-				EPoseSearchBooleanRequest::FalseValue;
-		}
+	if (const FPoseSearchIndexAsset* CurrentIndexAsset = InOutMotionMatchingState.CurrentSearchResult.GetSearchIndexAsset())
+	{
+		SearchContext.QueryMirrorRequest = CurrentIndexAsset->bMirrored ? EPoseSearchBooleanRequest::TrueValue : EPoseSearchBooleanRequest::FalseValue;
+	}
+
+	// If we can't advance or enough time has elapsed since the last pose jump then search
+	const bool bSearch = !SearchContext.bCanAdvance || (InOutMotionMatchingState.ElapsedPoseSearchTime >= Settings.SearchThrottleTime);
+	if (bSearch)
+	{
+		InOutMotionMatchingState.ElapsedPoseSearchTime = 0.f;
 
 		// Search the database for the nearest match to the updated query vector
 		FSearchResult SearchResult = Searchable->Search(SearchContext);
 
-		// making sure we haven't calculated ContinuingPoseCost if we !bCanAdvance 
-		check(bCanAdvance || !SearchResult.ContinuingPoseCost.IsValid());
+		// making sure we haven't calculated ContinuingPoseCost if we !SearchContext.bCanAdvance 
+		check(SearchContext.bCanAdvance || !SearchResult.ContinuingPoseCost.IsValid());
 		
 		if (SearchResult.PoseCost.GetTotalCost() < SearchResult.ContinuingPoseCost.GetTotalCost())
 		{
@@ -343,15 +337,14 @@ void UpdateMotionMatchingState(
 			InOutMotionMatchingState.CurrentSearchResult.ContinuingPoseCost = SearchResult.ContinuingPoseCost;
 			InOutMotionMatchingState.CurrentSearchResult.ComposedQuery = SearchResult.ComposedQuery;
 		}
+	
+		InOutMotionMatchingState.UpdateWantedPlayRate(SearchContext, Settings);
 	}
-
-	// Tick elapsed pose jump timer
-	if (!(InOutMotionMatchingState.Flags & EMotionMatchingFlags::JumpedToPose))
+	else
 	{
-		InOutMotionMatchingState.ElapsedPoseJumpTime += DeltaTime;
+		InOutMotionMatchingState.ElapsedPoseSearchTime += DeltaTime;
 	}
 
-	InOutMotionMatchingState.UpdateWantedPlayRate(SearchContext, Settings);
 	InOutMotionMatchingState.PoseIndicesHistory.Update(InOutMotionMatchingState.CurrentSearchResult, DeltaTime, Settings.PoseReselectHistory);
 
 #if UE_POSE_SEARCH_TRACE_ENABLED
@@ -362,7 +355,8 @@ void UpdateMotionMatchingState(
 		SearchContext,
 		InOutMotionMatchingState,
 		Trajectory,
-		LastResult);
+		LastResult,
+		bSearch);
 #endif
 }
 
diff --git a/Engine/Plugins/Experimental/Animation/PoseSearch/Source/Runtime/Private/Trace/PoseSearchTraceLogger.cpp b/Engine/Plugins/Experimental/Animation/PoseSearch/Source/Runtime/Private/Trace/PoseSearchTraceLogger.cpp
--- a/Engine/Plugins/Experimental/Animation/PoseSearch/Source/Runtime/Private/Trace/PoseSearchTraceLogger.cpp
+++ b/Engine/Plugins/Experimental/Animation/PoseSearch/Source/Runtime/Private/Trace/PoseSearchTraceLogger.cpp
@@ -71,7 +71,7 @@ FArchive& operator<<(FArchive& Ar, FTraceMotionMatchingStateDatabaseEntry& Entry
 FArchive& operator<<(FArchive& Ar, FTraceMotionMatchingState& State)
 {
 	Ar << State.SearchableAssetId;
-	Ar << State.ElapsedPoseJumpTime;
+	Ar << State.ElapsedPoseSearchTime;
 	Ar << State.AssetPlayerTime;
 	Ar << State.DeltaTime;
 	Ar << State.SimLinearVelocity;
diff --git a/Engine/Plugins/Experimental/Animation/PoseSearch/Source/Runtime/Private/Trace/PoseSearchTraceLogger.h b/Engine/Plugins/Experimental/Animation/PoseSearch/Source/Runtime/Private/Trace/PoseSearchTraceLogger.h
--- a/Engine/Plugins/Experimental/Animation/PoseSearch/Source/Runtime/Private/Trace/PoseSearchTraceLogger.h
+++ b/Engine/Plugins/Experimental/Animation/PoseSearch/Source/Runtime/Private/Trace/PoseSearchTraceLogger.h
@@ -68,7 +68,7 @@ struct POSESEARCH_API FTraceMotionMatchingState
 	uint64 SearchableAssetId = 0;
 	
 	/** Amount of time since the last pose switch */
-	float ElapsedPoseJumpTime = 0.0f;
+	float ElapsedPoseSearchTime = 0.0f;
 
 	float AssetPlayerTime = 0.0f;
 	float DeltaTime = 0.0f;
diff --git a/Engine/Plugins/Experimental/Animation/PoseSearch/Source/Runtime/Public/PoseSearch/PoseSearchLibrary.h b/Engine/Plugins/Experimental/Animation/PoseSearch/Source/Runtime/Public/PoseSearch/PoseSearchLibrary.h
--- a/Engine/Plugins/Experimental/Animation/PoseSearch/Source/Runtime/Public/PoseSearch/PoseSearchLibrary.h
+++ b/Engine/Plugins/Experimental/Animation/PoseSearch/Source/Runtime/Public/PoseSearch/PoseSearchLibrary.h
@@ -23,14 +23,6 @@ struct FGameplayTagContainer;
 struct FTrajectorySampleRange;
 class UPoseSearchSearchableAsset;
 
-UENUM(BlueprintType, Category="Motion Trajectory", meta=(Bitflags, UseEnumValuesAsMaskValuesInEditor="true"))
-enum class EMotionMatchingFlags : uint8
-{
-	None = 0 UMETA(Hidden),
-	JumpedToPose = 1 << 0,		// Signals that motion matching has made a significant deviation in the selected sequence/pose index
-};
-ENUM_CLASS_FLAGS(EMotionMatchingFlags);
-
 USTRUCT(BlueprintType, Category = "Animation|Pose Search")
 struct POSESEARCH_API FMotionMatchingSettings
 {
@@ -100,15 +92,15 @@ struct POSESEARCH_API FMotionMatchingState
 
 	// Time since the last pose jump
 	UPROPERTY(Transient)
-	float ElapsedPoseJumpTime = 0.f;
+	float ElapsedPoseSearchTime = 0.f;
 
 	// wanted PlayRate to have the selected animation playing at the estimated requested speed from the query
 	UPROPERTY(Transient)
 	float WantedPlayRate = 1.f;
 
-	// Evaluation flags relevant to the state of motion matching
+	// true if a new animation has been selected
 	UPROPERTY(VisibleAnywhere, BlueprintReadOnly, Category=State)
-	EMotionMatchingFlags Flags = EMotionMatchingFlags::None;
+	bool bJumpedToPose = false;
 
 	// Root motion delta for currently playing animation. Only required
 	// when UE_POSE_SEARCH_TRACE_ENABLED is active
