motion matching - adding back channel component stripping via EComponentStrippingVector

diff --git a/Engine/Plugins/Experimental/Animation/PoseSearch/Source/Runtime/Private/PoseSearchDerivedDataKey.cpp b/Engine/Plugins/Experimental/Animation/PoseSearch/Source/Runtime/Private/PoseSearchDerivedDataKey.cpp
--- a/Engine/Plugins/Experimental/Animation/PoseSearch/Source/Runtime/Private/PoseSearchDerivedDataKey.cpp
+++ b/Engine/Plugins/Experimental/Animation/PoseSearch/Source/Runtime/Private/PoseSearchDerivedDataKey.cpp
@@ -41,7 +41,7 @@ FKeyBuilder::FKeyBuilder(const UObject* Object, bool bUseDataVer)
 	if (bUseDataVer)
 	{
 		// used to invalidate the key without having to change POSESEARCHDB_DERIVEDDATA_VER all the times
-		int32 POSESEARCHDB_DERIVEDDATA_VER_SMALL = 34;
+		int32 POSESEARCHDB_DERIVEDDATA_VER_SMALL = 41;
 		FGuid VersionGuid = FDevSystemGuids::GetSystemGuid(FDevSystemGuids::Get().POSESEARCHDB_DERIVEDDATA_VER);
 
 		*this << VersionGuid;
diff --git a/Engine/Plugins/Experimental/Animation/PoseSearch/Source/Runtime/Private/PoseSearchFeatureChannel.cpp b/Engine/Plugins/Experimental/Animation/PoseSearch/Source/Runtime/Private/PoseSearchFeatureChannel.cpp
--- a/Engine/Plugins/Experimental/Animation/PoseSearch/Source/Runtime/Private/PoseSearchFeatureChannel.cpp
+++ b/Engine/Plugins/Experimental/Animation/PoseSearch/Source/Runtime/Private/PoseSearchFeatureChannel.cpp
@@ -10,87 +10,58 @@ namespace UE::PoseSearch
 
 //////////////////////////////////////////////////////////////////////////
 // FFeatureVectorHelper
-void FFeatureVectorHelper::EncodeQuat(TArrayView<float> Values, int32& DataOffset, const FQuat& Quat)
+int32 FFeatureVectorHelper::GetVectorCardinality(EComponentStrippingVector ComponentStrippingVector)
 {
-	// unitary quaternions are non Euclidean representation of rotation, so they cannot be used to calculate cost functions within the context of kdtrees, 
-	// so we convert them in a matrix, and pick 2 axis (we choose X,Y), skipping the 3rd since correlated to the cross product of the first two (this saves memory and cpu cycles)
-
-	const FMatrix M = Quat.ToMatrix();
-	const FVector X = M.GetScaledAxis(EAxis::X);
-	const FVector Y = M.GetScaledAxis(EAxis::Y);
-
-	Values[DataOffset + 0] = X.X;
-	Values[DataOffset + 1] = X.Y;
-	Values[DataOffset + 2] = X.Z;
-	Values[DataOffset + 3] = Y.X;
-	Values[DataOffset + 4] = Y.Y;
-	Values[DataOffset + 5] = Y.Z;
-
-	DataOffset += EncodeQuatCardinality;
-}
-
-void FFeatureVectorHelper::EncodeQuat(TArrayView<float> Values, int32& DataOffset, TConstArrayView<float> PrevValues, TConstArrayView<float> CurValues, TConstArrayView<float> NextValues, float LerpValue)
-{
-	FQuat Quat = DecodeQuatAtOffset(CurValues, DataOffset);
-
-	// linear interpolation
-	if (!FMath::IsNearlyZero(LerpValue))
+	switch (ComponentStrippingVector)
 	{
-		if (LerpValue < 0.f)
-		{
-			Quat = FQuat::Slerp(Quat, DecodeQuatAtOffset(PrevValues, DataOffset), -LerpValue);
-		}
-		else
-		{
-			Quat = FQuat::Slerp(Quat, DecodeQuatAtOffset(NextValues, DataOffset), LerpValue);
-		}
+	case EComponentStrippingVector::None:
+		return 3;
+	case EComponentStrippingVector::StripXY:
+		return 1;
+	case EComponentStrippingVector::StripZ:
+		return 2;
+	default:
+		checkNoEntry();
+		return 0;
 	}
-
-	// @todo: do we need to add options for cubic interpolation?
-	EncodeQuat(Values, DataOffset, Quat);
-}
-
-FQuat FFeatureVectorHelper::DecodeQuat(TConstArrayView<float> Values, int32& DataOffset)
-{
-	const FQuat Quat = DecodeQuatAtOffset(Values, DataOffset);
-	DataOffset += EncodeQuatCardinality;
-	return Quat;
-}
-
-FQuat FFeatureVectorHelper::DecodeQuatAtOffset(TConstArrayView<float> Values, int32 DataOffset)
-{
-	// reconstructing the chosen 2 axis (X,Y - from EncodeQuat) from the 6 floats from Values 
-	const FVector X(Values[DataOffset + 0], Values[DataOffset + 1], Values[DataOffset + 2]);
-	const FVector Y(Values[DataOffset + 3], Values[DataOffset + 4], Values[DataOffset + 5]);
-	// calculating the 3rd axis required to reconstruct the rotation matrix
-	const FVector Z = FVector::CrossProduct(X, Y);
-	const FMatrix M(X, Y, Z, FVector::ZeroVector);
-	// converting the matrix to the unitary quaternion
-	return FQuat(M);
 }
 
-void FFeatureVectorHelper::EncodeVector(TArrayView<float> Values, int32& DataOffset, const FVector& Vector)
+void FFeatureVectorHelper::EncodeVector(TArrayView<float> Values, int32 DataOffset, const FVector& Vector, EComponentStrippingVector ComponentStrippingVector)
 {
-	Values[DataOffset + 0] = Vector.X;
-	Values[DataOffset + 1] = Vector.Y;
-	Values[DataOffset + 2] = Vector.Z;
-	DataOffset += EncodeVectorCardinality;
+	switch (ComponentStrippingVector)
+	{
+	case EComponentStrippingVector::None:
+		Values[DataOffset + 0] = Vector.X;
+		Values[DataOffset + 1] = Vector.Y;
+		Values[DataOffset + 2] = Vector.Z;
+		break;
+	case EComponentStrippingVector::StripXY:
+		Values[DataOffset + 0] = Vector.Z;
+		break;
+	case EComponentStrippingVector::StripZ:
+		Values[DataOffset + 0] = Vector.X;
+		Values[DataOffset + 1] = Vector.Y;
+		break;
+	default:
+		checkNoEntry();
+		break;
+	}
 }
 
-void FFeatureVectorHelper::EncodeVector(TArrayView<float> Values, int32& DataOffset, TConstArrayView<float> PrevValues, TConstArrayView<float> CurValues, TConstArrayView<float> NextValues, float LerpValue, bool bNormalize)
+void FFeatureVectorHelper::EncodeVector(TArrayView<float> Values, int32 DataOffset, TConstArrayView<float> PrevValues, TConstArrayView<float> CurValues, TConstArrayView<float> NextValues, float LerpValue, bool bNormalize, EComponentStrippingVector ComponentStrippingVector)
 {
-	FVector Vector = DecodeVectorAtOffset(CurValues, DataOffset);
+	FVector Vector = DecodeVector(CurValues, DataOffset, ComponentStrippingVector);
 
 	// linear interpolation
 	if (!FMath::IsNearlyZero(LerpValue))
 	{
 		if (LerpValue < 0.f)
 		{
-			Vector = FMath::Lerp(Vector, DecodeVectorAtOffset(PrevValues, DataOffset), -LerpValue);
+			Vector = FMath::Lerp(Vector, DecodeVector(PrevValues, DataOffset, ComponentStrippingVector), -LerpValue);
 		}
 		else
 		{
-			Vector = FMath::Lerp(Vector, DecodeVectorAtOffset(NextValues, DataOffset), LerpValue);
+			Vector = FMath::Lerp(Vector, DecodeVector(NextValues, DataOffset, ComponentStrippingVector), LerpValue);
 		}
 	}
 
@@ -100,42 +71,45 @@ void FFeatureVectorHelper::EncodeVector(TArrayView<float> Values, int32& DataOff
 		Vector = Vector.GetSafeNormal(UE_SMALL_NUMBER, FVector::XAxisVector);
 	}
 
-	EncodeVector(Values, DataOffset, Vector);
-}
-
-FVector FFeatureVectorHelper::DecodeVector(TConstArrayView<float> Values, int32& DataOffset)
-{
-	const FVector Vector = DecodeVectorAtOffset(Values, DataOffset);
-	DataOffset += EncodeVectorCardinality;
-	return Vector;
+	EncodeVector(Values, DataOffset, Vector, ComponentStrippingVector);
 }
 
-FVector FFeatureVectorHelper::DecodeVectorAtOffset(TConstArrayView<float> Values, int32 DataOffset)
+FVector FFeatureVectorHelper::DecodeVector(TConstArrayView<float> Values, int32 DataOffset, EComponentStrippingVector ComponentStrippingVector)
 {
-	return FVector(Values[DataOffset + 0], Values[DataOffset + 1], Values[DataOffset + 2]);
+	switch (ComponentStrippingVector)
+	{
+	case EComponentStrippingVector::None:
+		return FVector(Values[DataOffset + 0], Values[DataOffset + 1], Values[DataOffset + 2]);
+	case EComponentStrippingVector::StripXY:
+		return FVector(0, 0, Values[DataOffset + 0]);
+	case EComponentStrippingVector::StripZ:
+		return FVector(Values[DataOffset + 0], Values[DataOffset + 1], 0);
+	default:
+		checkNoEntry();
+		return FVector::Zero();
+	}
 }
 
-void FFeatureVectorHelper::EncodeVector2D(TArrayView<float> Values, int32& DataOffset, const FVector2D& Vector2D)
+void FFeatureVectorHelper::EncodeVector2D(TArrayView<float> Values, int32 DataOffset, const FVector2D& Vector2D)
 {
 	Values[DataOffset + 0] = Vector2D.X;
 	Values[DataOffset + 1] = Vector2D.Y;
-	DataOffset += EncodeVector2DCardinality;
 }
 
-void FFeatureVectorHelper::EncodeVector2D(TArrayView<float> Values, int32& DataOffset, TConstArrayView<float> PrevValues, TConstArrayView<float> CurValues, TConstArrayView<float> NextValues, float LerpValue)
+void FFeatureVectorHelper::EncodeVector2D(TArrayView<float> Values, int32 DataOffset, TConstArrayView<float> PrevValues, TConstArrayView<float> CurValues, TConstArrayView<float> NextValues, float LerpValue)
 {
-	FVector2D Vector2D = DecodeVector2DAtOffset(CurValues, DataOffset);
+	FVector2D Vector2D = DecodeVector2D(CurValues, DataOffset);
 
 	// linear interpolation
 	if (!FMath::IsNearlyZero(LerpValue))
 	{
 		if (LerpValue < 0.f)
 		{
-			Vector2D = FMath::Lerp(Vector2D, DecodeVector2DAtOffset(PrevValues, DataOffset), -LerpValue);
+			Vector2D = FMath::Lerp(Vector2D, DecodeVector2D(PrevValues, DataOffset), -LerpValue);
 		}
 		else
 		{
-			Vector2D = FMath::Lerp(Vector2D, DecodeVector2DAtOffset(NextValues, DataOffset), LerpValue);
+			Vector2D = FMath::Lerp(Vector2D, DecodeVector2D(NextValues, DataOffset), LerpValue);
 		}
 	}
 
@@ -143,38 +117,30 @@ void FFeatureVectorHelper::EncodeVector2D(TArrayView<float> Values, int32& DataO
 	EncodeVector2D(Values, DataOffset, Vector2D);
 }
 
-FVector2D FFeatureVectorHelper::DecodeVector2D(TConstArrayView<float> Values, int32& DataOffset)
-{
-	const FVector2D Vector2D = DecodeVector2DAtOffset(Values, DataOffset);
-	DataOffset += EncodeVector2DCardinality;
-	return Vector2D;
-}
-
-FVector2D FFeatureVectorHelper::DecodeVector2DAtOffset(TConstArrayView<float> Values, int32 DataOffset)
+FVector2D FFeatureVectorHelper::DecodeVector2D(TConstArrayView<float> Values, int32 DataOffset)
 {
 	return FVector2D(Values[DataOffset + 0], Values[DataOffset + 1]);
 }
 
-void FFeatureVectorHelper::EncodeFloat(TArrayView<float> Values, int32& DataOffset, const float Value)
+void FFeatureVectorHelper::EncodeFloat(TArrayView<float> Values, int32 DataOffset, const float Value)
 {
 	Values[DataOffset + 0] = Value;
-	DataOffset += EncodeFloatCardinality;
 }
 
-void FFeatureVectorHelper::EncodeFloat(TArrayView<float> Values, int32& DataOffset, TConstArrayView<float> PrevValues, TConstArrayView<float> CurValues, TConstArrayView<float> NextValues, float LerpValue)
+void FFeatureVectorHelper::EncodeFloat(TArrayView<float> Values, int32 DataOffset, TConstArrayView<float> PrevValues, TConstArrayView<float> CurValues, TConstArrayView<float> NextValues, float LerpValue)
 {
-	float Value = DecodeFloatAtOffset(CurValues, DataOffset);
+	float Value = DecodeFloat(CurValues, DataOffset);
 
 	// linear interpolation
 	if (!FMath::IsNearlyZero(LerpValue))
 	{
 		if (LerpValue < 0.f)
 		{
-			Value = FMath::Lerp(Value, DecodeFloatAtOffset(PrevValues, DataOffset), -LerpValue);
+			Value = FMath::Lerp(Value, DecodeFloat(PrevValues, DataOffset), -LerpValue);
 		}
 		else
 		{
-			Value = FMath::Lerp(Value, DecodeFloatAtOffset(NextValues, DataOffset), LerpValue);
+			Value = FMath::Lerp(Value, DecodeFloat(NextValues, DataOffset), LerpValue);
 		}
 	}
 
@@ -182,14 +148,7 @@ void FFeatureVectorHelper::EncodeFloat(TArrayView<float> Values, int32& DataOffs
 	EncodeFloat(Values, DataOffset, Value);
 }
 
-float FFeatureVectorHelper::DecodeFloat(TConstArrayView<float> Values, int32& DataOffset)
-{
-	const float Value = DecodeFloatAtOffset(Values, DataOffset);
-	DataOffset += EncodeFloatCardinality;
-	return Value;
-}
-
-float FFeatureVectorHelper::DecodeFloatAtOffset(TConstArrayView<float> Values, int32 DataOffset)
+float FFeatureVectorHelper::DecodeFloat(TConstArrayView<float> Values, int32 DataOffset)
 {
 	return Values[DataOffset];
 }
diff --git a/Engine/Plugins/Experimental/Animation/PoseSearch/Source/Runtime/Private/PoseSearchFeatureChannel_FilterCrashingLegs.cpp b/Engine/Plugins/Experimental/Animation/PoseSearch/Source/Runtime/Private/PoseSearchFeatureChannel_FilterCrashingLegs.cpp
--- a/Engine/Plugins/Experimental/Animation/PoseSearch/Source/Runtime/Private/PoseSearchFeatureChannel_FilterCrashingLegs.cpp
+++ b/Engine/Plugins/Experimental/Animation/PoseSearch/Source/Runtime/Private/PoseSearchFeatureChannel_FilterCrashingLegs.cpp
@@ -30,7 +30,7 @@ static float ComputeCrashingLegsValue(const FVector& RightThighPos, const FVecto
 void UPoseSearchFeatureChannel_FilterCrashingLegs::Finalize(UPoseSearchSchema* Schema)
 {
 	ChannelDataOffset = Schema->SchemaCardinality;
-	ChannelCardinality = UE::PoseSearch::FFeatureVectorHelper::EncodeFloatCardinality;
+	ChannelCardinality = 1;
 	Schema->SchemaCardinality += ChannelCardinality;
 
 	LeftThighIdx = Schema->AddBoneReference(LeftThigh);
@@ -41,8 +41,10 @@ void UPoseSearchFeatureChannel_FilterCrashingLegs::Finalize(UPoseSearchSchema* S
 
 void UPoseSearchFeatureChannel_FilterCrashingLegs::FillWeights(TArray<float>& Weights) const
 {
-	check(ChannelCardinality == UE::PoseSearch::FFeatureVectorHelper::EncodeFloatCardinality);
-	Weights[ChannelDataOffset] = Weight;
+	for (int32 i = 0; i < ChannelCardinality; ++i)
+	{
+		Weights[ChannelDataOffset + i] = Weight;
+	}
 }
 
 void UPoseSearchFeatureChannel_FilterCrashingLegs::IndexAsset(UE::PoseSearch::IAssetIndexer& Indexer, TArrayView<float> FeatureVectorTable) const
@@ -63,9 +65,8 @@ void UPoseSearchFeatureChannel_FilterCrashingLegs::IndexAsset(UE::PoseSearch::IA
 
 		const float CrashingLegsValue = ComputeCrashingLegsValue(RightThighTransform.GetTranslation(), LeftThighTransform.GetTranslation(), RightFootTransform.GetTranslation(), LeftFootTransform.GetTranslation());
 
-		int32 DataOffset = ChannelDataOffset;
 		const int32 VectorIdx = SampleIdx - IndexingContext.BeginSampleIdx;
-		FFeatureVectorHelper::EncodeFloat(IndexingContext.GetPoseVector(VectorIdx, FeatureVectorTable), DataOffset, CrashingLegsValue);
+		FFeatureVectorHelper::EncodeFloat(IndexingContext.GetPoseVector(VectorIdx, FeatureVectorTable), ChannelDataOffset, CrashingLegsValue);
 	}
 }
 
@@ -76,13 +77,12 @@ void UPoseSearchFeatureChannel_FilterCrashingLegs::BuildQuery(UE::PoseSearch::FS
 	const bool bIsCurrentResultValid = SearchContext.CurrentResult.IsValid();
 	const bool bSkip = InputQueryPose != EInputQueryPose::UseCharacterPose && bIsCurrentResultValid && SearchContext.CurrentResult.Database->Schema == InOutQuery.GetSchema();
 	
-	int32 DataOffset = ChannelDataOffset;
 	if (bSkip || !SearchContext.History)
 	{
 		if (bIsCurrentResultValid)
 		{
 			const float LerpValue = InputQueryPose == EInputQueryPose::UseInterpolatedContinuingPose ? SearchContext.CurrentResult.LerpValue : 0.f;
-			FFeatureVectorHelper::EncodeFloat(InOutQuery.EditValues(), DataOffset, SearchContext.GetCurrentResultPrevPoseVector(), SearchContext.GetCurrentResultPoseVector(), SearchContext.GetCurrentResultNextPoseVector(), LerpValue);
+			FFeatureVectorHelper::EncodeFloat(InOutQuery.EditValues(), ChannelDataOffset, SearchContext.GetCurrentResultPrevPoseVector(), SearchContext.GetCurrentResultPoseVector(), SearchContext.GetCurrentResultNextPoseVector(), LerpValue);
 		}
 		// else leave the InOutQuery set to zero since the SearchContext.History is invalid and it'll fail if we continue
 	}
@@ -95,7 +95,7 @@ void UPoseSearchFeatureChannel_FilterCrashingLegs::BuildQuery(UE::PoseSearch::FS
 		const FTransform RightFootTransform = SearchContext.GetComponentSpaceTransform(SampleTime, InOutQuery.GetSchema(), RightFootIdx);
 
 		const float CrashingLegsValue = ComputeCrashingLegsValue(RightThighTransform.GetTranslation(), LeftThighTransform.GetTranslation(), RightFootTransform.GetTranslation(), LeftFootTransform.GetTranslation());
-		FFeatureVectorHelper::EncodeFloat(InOutQuery.EditValues(), DataOffset, CrashingLegsValue);
+		FFeatureVectorHelper::EncodeFloat(InOutQuery.EditValues(), ChannelDataOffset, CrashingLegsValue);
 	}
 }
 
@@ -104,7 +104,7 @@ void UPoseSearchFeatureChannel_FilterCrashingLegs::DebugDraw(const UE::PoseSearc
 #if ENABLE_DRAW_DEBUG
 	using namespace UE::PoseSearch;
 
-	const float CrashingLegsValue = FFeatureVectorHelper::DecodeFloatAtOffset(PoseVector, ChannelDataOffset);
+	const float CrashingLegsValue = FFeatureVectorHelper::DecodeFloat(PoseVector, ChannelDataOffset);
 
 	const float LifeTime = DrawParams.DefaultLifeTime;
 	const uint8 DepthPriority = ESceneDepthPriorityGroup::SDPG_Foreground + 2;
@@ -152,8 +152,8 @@ bool UPoseSearchFeatureChannel_FilterCrashingLegs::IsPoseValid(TConstArrayView<f
 {
 	using namespace UE::PoseSearch;
 
-	const float PoseValue = FFeatureVectorHelper::DecodeFloatAtOffset(PoseValues, ChannelDataOffset);
-	const float QueryValue = FFeatureVectorHelper::DecodeFloatAtOffset(QueryValues, ChannelDataOffset);
+	const float PoseValue = FFeatureVectorHelper::DecodeFloat(PoseValues, ChannelDataOffset);
+	const float QueryValue = FFeatureVectorHelper::DecodeFloat(QueryValues, ChannelDataOffset);
 
 	bool bIsPoseValid = true;
 	if (FMath::Abs(QueryValue - PoseValue) > AllowedTolerance)
diff --git a/Engine/Plugins/Experimental/Animation/PoseSearch/Source/Runtime/Private/PoseSearchFeatureChannel_Heading.cpp b/Engine/Plugins/Experimental/Animation/PoseSearch/Source/Runtime/Private/PoseSearchFeatureChannel_Heading.cpp
--- a/Engine/Plugins/Experimental/Animation/PoseSearch/Source/Runtime/Private/PoseSearchFeatureChannel_Heading.cpp
+++ b/Engine/Plugins/Experimental/Animation/PoseSearch/Source/Runtime/Private/PoseSearchFeatureChannel_Heading.cpp
@@ -12,16 +12,14 @@
 void UPoseSearchFeatureChannel_Heading::Finalize(UPoseSearchSchema* Schema)
 {
 	ChannelDataOffset = Schema->SchemaCardinality;
-	ChannelCardinality = UE::PoseSearch::FFeatureVectorHelper::EncodeVectorCardinality;
+	ChannelCardinality = UE::PoseSearch::FFeatureVectorHelper::GetVectorCardinality(ComponentStripping);
 	Schema->SchemaCardinality += ChannelCardinality;
 	SchemaBoneIdx = Schema->AddBoneReference(Bone);
 }
 
 void UPoseSearchFeatureChannel_Heading::FillWeights(TArray<float>& Weights) const
 {
-	using namespace UE::PoseSearch;
-
-	for (int32 i = 0; i != FFeatureVectorHelper::EncodeVectorCardinality; ++i)
+	for (int32 i = 0; i < ChannelCardinality; ++i)
 	{
 		Weights[ChannelDataOffset + i] = Weight;
 	}
@@ -58,9 +56,8 @@ void UPoseSearchFeatureChannel_Heading::IndexAsset(UE::PoseSearch::IAssetIndexer
 
 		bool ClampedPresent;
 		const FTransform BoneTransformsPresent = Indexer.GetComponentSpaceTransform(SubsampleTime, OriginSampleTime, ClampedPresent, SchemaBoneIdx);
-		int32 DataOffset = ChannelDataOffset;
 		const FVector Heading = GetAxis(BoneTransformsPresent.GetRotation());
-		FFeatureVectorHelper::EncodeVector(IndexingContext.GetPoseVector(VectorIdx, FeatureVectorTable), DataOffset, Heading);
+		FFeatureVectorHelper::EncodeVector(IndexingContext.GetPoseVector(VectorIdx, FeatureVectorTable), ChannelDataOffset, Heading, ComponentStripping);
 	}
 }
 
@@ -76,9 +73,8 @@ void UPoseSearchFeatureChannel_Heading::BuildQuery(UE::PoseSearch::FSearchContex
 		if (bIsCurrentResultValid)
 		{
 			const float LerpValue = InputQueryPose == EInputQueryPose::UseInterpolatedContinuingPose ? SearchContext.CurrentResult.LerpValue : 0.f;
-			int32 DataOffset = ChannelDataOffset;
 			// @todo: we should normalize if LerpValue != 0
-			FFeatureVectorHelper::EncodeVector(InOutQuery.EditValues(), DataOffset, SearchContext.GetCurrentResultPrevPoseVector(), SearchContext.GetCurrentResultPoseVector(), SearchContext.GetCurrentResultNextPoseVector(), LerpValue, true);
+			FFeatureVectorHelper::EncodeVector(InOutQuery.EditValues(), ChannelDataOffset, SearchContext.GetCurrentResultPrevPoseVector(), SearchContext.GetCurrentResultPoseVector(), SearchContext.GetCurrentResultNextPoseVector(), LerpValue, true, ComponentStripping);
 		}
 		// else leave the InOutQuery set to zero since the SearchContext.History is invalid and it'll fail if we continue
 	}
@@ -98,8 +94,7 @@ void UPoseSearchFeatureChannel_Heading::BuildQuery(UE::PoseSearch::FSearchContex
 			Transform = TrajectorySample.Transform;
 		}
 
-		int32 DataOffset = ChannelDataOffset;
-		FFeatureVectorHelper::EncodeVector(InOutQuery.EditValues(), DataOffset, GetAxis(Transform.GetRotation()));
+		FFeatureVectorHelper::EncodeVector(InOutQuery.EditValues(), ChannelDataOffset, GetAxis(Transform.GetRotation()), ComponentStripping);
 	}
 }
 
@@ -112,7 +107,7 @@ void UPoseSearchFeatureChannel_Heading::DebugDraw(const UE::PoseSearch::FDebugDr
 	const uint8 DepthPriority = ESceneDepthPriorityGroup::SDPG_Foreground + 2;
 	const bool bPersistent = EnumHasAnyFlags(DrawParams.Flags, EDebugDrawFlags::Persistent);
 	const FColor Color = DrawParams.GetColor(ColorPresetIndex);
-	const FVector BoneHeading = DrawParams.RootTransform.GetRotation().RotateVector(FFeatureVectorHelper::DecodeVectorAtOffset(PoseVector, ChannelDataOffset));
+	const FVector BoneHeading = DrawParams.RootTransform.GetRotation().RotateVector(FFeatureVectorHelper::DecodeVector(PoseVector, ChannelDataOffset, ComponentStripping));
 	const FVector BonePos = DrawParams.GetCachedPosition(SampleTimeOffset, SchemaBoneIdx);
 
 	if (EnumHasAnyFlags(DrawParams.Flags, EDebugDrawFlags::DrawSearchIndex))
@@ -151,6 +146,15 @@ FString UPoseSearchFeatureChannel_Heading::GetLabel() const
 		break;
 	}
 
+	if (ComponentStripping == EComponentStrippingVector::StripXY)
+	{
+		Label.Append(TEXT("_z"));
+	}
+	else if (ComponentStripping == EComponentStrippingVector::StripZ)
+	{
+		Label.Append(TEXT("_xy"));
+	}
+
 	const FBoneReference& BoneReference = GetSchema()->BoneReferences[SchemaBoneIdx];
 	if (BoneReference.HasValidSetup())
 	{
diff --git a/Engine/Plugins/Experimental/Animation/PoseSearch/Source/Runtime/Private/PoseSearchFeatureChannel_Heading.h b/Engine/Plugins/Experimental/Animation/PoseSearch/Source/Runtime/Private/PoseSearchFeatureChannel_Heading.h
--- a/Engine/Plugins/Experimental/Animation/PoseSearch/Source/Runtime/Private/PoseSearchFeatureChannel_Heading.h
+++ b/Engine/Plugins/Experimental/Animation/PoseSearch/Source/Runtime/Private/PoseSearchFeatureChannel_Heading.h
@@ -6,7 +6,7 @@
 #include "BoneContainer.h"
 #include "PoseSearchFeatureChannel_Heading.generated.h"
 
-UENUM(BlueprintType)
+UENUM()
 enum class EHeadingAxis : uint8
 {
 	X,
@@ -17,7 +17,7 @@ enum class EHeadingAxis : uint8
 	Invalid = Num UMETA(Hidden)
 };
 
-UCLASS(BlueprintType, EditInlineNew, meta = (DisplayName = "Heading Channel"), CollapseCategories)
+UCLASS(EditInlineNew, meta = (DisplayName = "Heading Channel"), CollapseCategories)
 class POSESEARCH_API UPoseSearchFeatureChannel_Heading : public UPoseSearchFeatureChannel
 {
 	GENERATED_BODY()
@@ -44,6 +44,9 @@ public:
 	UPROPERTY(EditAnywhere, Category = "Settings")
 	EInputQueryPose InputQueryPose = EInputQueryPose::UseContinuingPose;
 
+	UPROPERTY(EditAnywhere, Category = "Settings")
+	EComponentStrippingVector ComponentStripping = EComponentStrippingVector::None;
+
 	// UPoseSearchFeatureChannel interface
 	virtual void Finalize(UPoseSearchSchema* Schema) override;
 	virtual void FillWeights(TArray<float>& Weights) const override;
diff --git a/Engine/Plugins/Experimental/Animation/PoseSearch/Source/Runtime/Private/PoseSearchFeatureChannel_Phase.cpp b/Engine/Plugins/Experimental/Animation/PoseSearch/Source/Runtime/Private/PoseSearchFeatureChannel_Phase.cpp
--- a/Engine/Plugins/Experimental/Animation/PoseSearch/Source/Runtime/Private/PoseSearchFeatureChannel_Phase.cpp
+++ b/Engine/Plugins/Experimental/Animation/PoseSearch/Source/Runtime/Private/PoseSearchFeatureChannel_Phase.cpp
@@ -286,16 +286,14 @@ namespace UE::PoseSearch
 void UPoseSearchFeatureChannel_Phase::Finalize(UPoseSearchSchema* Schema)
 {
 	ChannelDataOffset = Schema->SchemaCardinality;
-	ChannelCardinality = UE::PoseSearch::FFeatureVectorHelper::EncodeVector2DCardinality;
+	ChannelCardinality = 2;
 	Schema->SchemaCardinality += ChannelCardinality;
 	SchemaBoneIdx = Schema->AddBoneReference(Bone);
 }
 
 void UPoseSearchFeatureChannel_Phase::FillWeights(TArray<float>& Weights) const
 {
-	using namespace UE::PoseSearch;
-
-	for (int32 i = 0; i != FFeatureVectorHelper::EncodeVector2DCardinality; ++i)
+	for (int32 i = 0; i < ChannelCardinality; ++i)
 	{
 		Weights[ChannelDataOffset + i] = Weight;
 	}
@@ -340,8 +338,7 @@ void UPoseSearchFeatureChannel_Phase::IndexAsset(UE::PoseSearch::IAssetIndexer&
 	for (int32 SampleIdx = IndexingContext.BeginSampleIdx; SampleIdx != IndexingContext.EndSampleIdx; ++SampleIdx)
 	{
 		const int32 VectorIdx = SampleIdx - IndexingContext.BeginSampleIdx;
-		int32 DataOffset = ChannelDataOffset;
-		FFeatureVectorHelper::EncodeVector2D(IndexingContext.GetPoseVector(VectorIdx, FeatureVectorTable), DataOffset, Phases[VectorIdx]);
+		FFeatureVectorHelper::EncodeVector2D(IndexingContext.GetPoseVector(VectorIdx, FeatureVectorTable), ChannelDataOffset, Phases[VectorIdx]);
 	}
 }
 
@@ -356,8 +353,7 @@ void UPoseSearchFeatureChannel_Phase::BuildQuery(UE::PoseSearch::FSearchContext&
 		if (bIsCurrentResultValid)
 		{
 			const float LerpValue = InputQueryPose == EInputQueryPose::UseInterpolatedContinuingPose ? SearchContext.CurrentResult.LerpValue : 0.f;
-			int32 DataOffset = ChannelDataOffset;
-			FFeatureVectorHelper::EncodeVector2D(InOutQuery.EditValues(), DataOffset, SearchContext.GetCurrentResultPrevPoseVector(), SearchContext.GetCurrentResultPoseVector(), SearchContext.GetCurrentResultNextPoseVector(), LerpValue);
+			FFeatureVectorHelper::EncodeVector2D(InOutQuery.EditValues(), ChannelDataOffset, SearchContext.GetCurrentResultPrevPoseVector(), SearchContext.GetCurrentResultPoseVector(), SearchContext.GetCurrentResultNextPoseVector(), LerpValue);
 		}
 		// else leave the InOutQuery set to zero since the SearchContext.History is invalid and it'll fail if we continue
 	}
@@ -380,7 +376,7 @@ void UPoseSearchFeatureChannel_Phase::DebugDraw(const UE::PoseSearch::FDebugDraw
 	const bool bPersistent = EnumHasAnyFlags(DrawParams.Flags, EDebugDrawFlags::Persistent);
 	const FColor Color = DrawParams.GetColor(ColorPresetIndex);
 
-	const FVector2D Phase = FFeatureVectorHelper::DecodeVector2DAtOffset(PoseVector, ChannelDataOffset);
+	const FVector2D Phase = FFeatureVectorHelper::DecodeVector2D(PoseVector, ChannelDataOffset);
 	const FVector BonePos = DrawParams.GetCachedPosition(0.f, SchemaBoneIdx);
 
 	const FVector TransformXAxisVector = DrawParams.RootTransform.TransformVector(FVector::XAxisVector);
diff --git a/Engine/Plugins/Experimental/Animation/PoseSearch/Source/Runtime/Private/PoseSearchFeatureChannel_Position.cpp b/Engine/Plugins/Experimental/Animation/PoseSearch/Source/Runtime/Private/PoseSearchFeatureChannel_Position.cpp
--- a/Engine/Plugins/Experimental/Animation/PoseSearch/Source/Runtime/Private/PoseSearchFeatureChannel_Position.cpp
+++ b/Engine/Plugins/Experimental/Animation/PoseSearch/Source/Runtime/Private/PoseSearchFeatureChannel_Position.cpp
@@ -12,7 +12,7 @@
 void UPoseSearchFeatureChannel_Position::Finalize(UPoseSearchSchema* Schema)
 {
 	ChannelDataOffset = Schema->SchemaCardinality;
-	ChannelCardinality = UE::PoseSearch::FFeatureVectorHelper::EncodeVectorCardinality;
+	ChannelCardinality = UE::PoseSearch::FFeatureVectorHelper::GetVectorCardinality(ComponentStripping);
 	Schema->SchemaCardinality += ChannelCardinality;
 
 	SchemaBoneIdx = Schema->AddBoneReference(Bone);
@@ -20,9 +20,7 @@ void UPoseSearchFeatureChannel_Position::Finalize(UPoseSearchSchema* Schema)
 
 void UPoseSearchFeatureChannel_Position::FillWeights(TArray<float>& Weights) const
 {
-	using namespace UE::PoseSearch;
-
-	for (int32 i = 0; i != FFeatureVectorHelper::EncodeVectorCardinality; ++i)
+	for (int32 i = 0; i < ChannelCardinality; ++i)
 	{
 		Weights[ChannelDataOffset + i] = Weight;
 	}
@@ -43,8 +41,7 @@ void UPoseSearchFeatureChannel_Position::IndexAsset(UE::PoseSearch::IAssetIndexe
 
 		bool ClampedPresent;
 		const FTransform BoneTransformsPresent = Indexer.GetComponentSpaceTransform(SubsampleTime, OriginSampleTime, ClampedPresent, SchemaBoneIdx);
-		int32 DataOffset = ChannelDataOffset;
-		FFeatureVectorHelper::EncodeVector(IndexingContext.GetPoseVector(VectorIdx, FeatureVectorTable), DataOffset, BoneTransformsPresent.GetTranslation());
+		FFeatureVectorHelper::EncodeVector(IndexingContext.GetPoseVector(VectorIdx, FeatureVectorTable), ChannelDataOffset, BoneTransformsPresent.GetTranslation(), ComponentStripping);
 	}
 }
 
@@ -61,8 +58,7 @@ void UPoseSearchFeatureChannel_Position::BuildQuery(UE::PoseSearch::FSearchConte
 		if (bIsCurrentResultValid)
 		{
 			const float LerpValue = InputQueryPose == EInputQueryPose::UseInterpolatedContinuingPose ? SearchContext.CurrentResult.LerpValue : 0.f;
-			int32 DataOffset = ChannelDataOffset;
-			FFeatureVectorHelper::EncodeVector(InOutQuery.EditValues(), DataOffset, SearchContext.GetCurrentResultPrevPoseVector(), SearchContext.GetCurrentResultPoseVector(), SearchContext.GetCurrentResultNextPoseVector(), LerpValue);
+			FFeatureVectorHelper::EncodeVector(InOutQuery.EditValues(), ChannelDataOffset, SearchContext.GetCurrentResultPrevPoseVector(), SearchContext.GetCurrentResultPoseVector(), SearchContext.GetCurrentResultNextPoseVector(), LerpValue, false, ComponentStripping);
 		}
 		// else leave the InOutQuery set to zero since the SearchContext.History is invalid and it'll fail if we continue
 	}
@@ -82,8 +78,7 @@ void UPoseSearchFeatureChannel_Position::BuildQuery(UE::PoseSearch::FSearchConte
 			Transform = TrajectorySample.Transform;
 		}
 
-		int32 DataOffset = ChannelDataOffset;
-		FFeatureVectorHelper::EncodeVector(InOutQuery.EditValues(), DataOffset, Transform.GetTranslation());
+		FFeatureVectorHelper::EncodeVector(InOutQuery.EditValues(), ChannelDataOffset, Transform.GetTranslation(), ComponentStripping);
 	}
 }
 
@@ -91,7 +86,7 @@ void UPoseSearchFeatureChannel_Position::PreDebugDraw(UE::PoseSearch::FDebugDraw
 {
 #if ENABLE_DRAW_DEBUG
 	using namespace UE::PoseSearch;
-	const FVector BonePos = DrawParams.RootTransform.TransformPosition(FFeatureVectorHelper::DecodeVectorAtOffset(PoseVector, ChannelDataOffset));
+	const FVector BonePos = DrawParams.RootTransform.TransformPosition(FFeatureVectorHelper::DecodeVector(PoseVector, ChannelDataOffset, ComponentStripping));
 	DrawParams.AddCachedPosition(SampleTimeOffset, SchemaBoneIdx, BonePos);
 #endif // ENABLE_DRAW_DEBUG
 }
@@ -106,7 +101,7 @@ void UPoseSearchFeatureChannel_Position::DebugDraw(const UE::PoseSearch::FDebugD
 	const bool bPersistent = EnumHasAnyFlags(DrawParams.Flags, EDebugDrawFlags::Persistent);
 	const FColor Color = DrawParams.GetColor(ColorPresetIndex);
 
-	const FVector BonePos = DrawParams.RootTransform.TransformPosition(FFeatureVectorHelper::DecodeVectorAtOffset(PoseVector, ChannelDataOffset));
+	const FVector BonePos = DrawParams.RootTransform.TransformPosition(FFeatureVectorHelper::DecodeVector(PoseVector, ChannelDataOffset, ComponentStripping));
 	// validating DrawParams.AddCachedPosition 
 	check(DrawParams.GetCachedPosition(SampleTimeOffset, SchemaBoneIdx) == BonePos);
 
@@ -139,6 +134,16 @@ FString UPoseSearchFeatureChannel_Position::GetLabel() const
 	}
 
 	Label.Append(TEXT("Pos"));
+
+	if (ComponentStripping == EComponentStrippingVector::StripXY)
+	{
+		Label.Append(TEXT("_z"));
+	}
+	else if (ComponentStripping == EComponentStrippingVector::StripZ)
+	{
+		Label.Append(TEXT("_xy"));
+	}
+
 	const FBoneReference& BoneReference = GetSchema()->BoneReferences[SchemaBoneIdx];
 	if (BoneReference.HasValidSetup())
 	{
diff --git a/Engine/Plugins/Experimental/Animation/PoseSearch/Source/Runtime/Private/PoseSearchFeatureChannel_Position.h b/Engine/Plugins/Experimental/Animation/PoseSearch/Source/Runtime/Private/PoseSearchFeatureChannel_Position.h
--- a/Engine/Plugins/Experimental/Animation/PoseSearch/Source/Runtime/Private/PoseSearchFeatureChannel_Position.h
+++ b/Engine/Plugins/Experimental/Animation/PoseSearch/Source/Runtime/Private/PoseSearchFeatureChannel_Position.h
@@ -6,7 +6,7 @@
 #include "BoneContainer.h"
 #include "PoseSearchFeatureChannel_Position.generated.h"
 
-UCLASS(BlueprintType, EditInlineNew, meta = (DisplayName = "Position Channel"), CollapseCategories)
+UCLASS(EditInlineNew, meta = (DisplayName = "Position Channel"), CollapseCategories)
 class POSESEARCH_API UPoseSearchFeatureChannel_Position : public UPoseSearchFeatureChannel
 {
 	GENERATED_BODY()
@@ -30,6 +30,9 @@ public:
 	UPROPERTY(EditAnywhere, Category = "Settings")
 	EInputQueryPose InputQueryPose = EInputQueryPose::UseContinuingPose;
 
+	UPROPERTY(EditAnywhere, Category = "Settings")
+	EComponentStrippingVector ComponentStripping = EComponentStrippingVector::None;
+
 	// UPoseSearchFeatureChannel interface
 	virtual void Finalize(UPoseSearchSchema* Schema) override;
 	virtual void FillWeights(TArray<float>& Weights) const override;
diff --git a/Engine/Plugins/Experimental/Animation/PoseSearch/Source/Runtime/Private/PoseSearchFeatureChannel_Trajectory.cpp b/Engine/Plugins/Experimental/Animation/PoseSearch/Source/Runtime/Private/PoseSearchFeatureChannel_Trajectory.cpp
--- a/Engine/Plugins/Experimental/Animation/PoseSearch/Source/Runtime/Private/PoseSearchFeatureChannel_Trajectory.cpp
+++ b/Engine/Plugins/Experimental/Animation/PoseSearch/Source/Runtime/Private/PoseSearchFeatureChannel_Trajectory.cpp
@@ -27,10 +27,13 @@ void UPoseSearchFeatureChannel_Trajectory::Finalize(UPoseSearchSchema* Schema)
 			Position->SampleTimeOffset = Sample.Offset;
 			Position->ColorPresetIndex = Sample.ColorPresetIndex;
 			Position->InputQueryPose = EInputQueryPose::UseCharacterPose;
+			if (EnumHasAnyFlags(Sample.Flags, EPoseSearchTrajectoryFlags::PositionXY))
+			{
+				Position->ComponentStripping = EComponentStrippingVector::StripZ;
+			}
 			SubChannels.Add(Position);
 		}
 
-		// @todo: implement VelocityXY properly as 2 dimension channel
 		if (EnumHasAnyFlags(Sample.Flags, EPoseSearchTrajectoryFlags::Velocity | EPoseSearchTrajectoryFlags::VelocityXY))
 		{
 			UPoseSearchFeatureChannel_Velocity* Velocity = NewObject<UPoseSearchFeatureChannel_Velocity>(this, NAME_None, RF_Transient);
@@ -39,10 +42,13 @@ void UPoseSearchFeatureChannel_Trajectory::Finalize(UPoseSearchSchema* Schema)
 			Velocity->ColorPresetIndex = Sample.ColorPresetIndex;
 			Velocity->InputQueryPose = EInputQueryPose::UseCharacterPose;
 			Velocity->bUseCharacterSpaceVelocities = false;
+			if (EnumHasAnyFlags(Sample.Flags, EPoseSearchTrajectoryFlags::VelocityXY))
+			{
+				Velocity->ComponentStripping = EComponentStrippingVector::StripZ;
+			}
 			SubChannels.Add(Velocity);
 		}
 
-		// @todo: implement VelocityDirectionXY properly as 2 dimension channel
 		if (EnumHasAnyFlags(Sample.Flags, EPoseSearchTrajectoryFlags::VelocityDirection | EPoseSearchTrajectoryFlags::VelocityDirectionXY))
 		{
 			UPoseSearchFeatureChannel_Velocity* Velocity = NewObject<UPoseSearchFeatureChannel_Velocity>(this, NAME_None, RF_Transient);
@@ -52,10 +58,13 @@ void UPoseSearchFeatureChannel_Trajectory::Finalize(UPoseSearchSchema* Schema)
 			Velocity->InputQueryPose = EInputQueryPose::UseCharacterPose;
 			Velocity->bUseCharacterSpaceVelocities = false;
 			Velocity->bNormalize = true;
+			if (EnumHasAnyFlags(Sample.Flags, EPoseSearchTrajectoryFlags::VelocityDirectionXY))
+			{
+				Velocity->ComponentStripping = EComponentStrippingVector::StripZ;
+			}
 			SubChannels.Add(Velocity);
 		}
 
-		// @todo: implement FacingDirectionXY properly as 2 dimension channel
 		if (EnumHasAnyFlags(Sample.Flags, EPoseSearchTrajectoryFlags::FacingDirection | EPoseSearchTrajectoryFlags::FacingDirectionXY))
 		{
 			UPoseSearchFeatureChannel_Heading* Heading = NewObject<UPoseSearchFeatureChannel_Heading>(this, NAME_None, RF_Transient);
@@ -63,6 +72,10 @@ void UPoseSearchFeatureChannel_Trajectory::Finalize(UPoseSearchSchema* Schema)
 			Heading->SampleTimeOffset = Sample.Offset;
 			Heading->ColorPresetIndex = Sample.ColorPresetIndex;
 			Heading->InputQueryPose = EInputQueryPose::UseCharacterPose;
+			if (EnumHasAnyFlags(Sample.Flags, EPoseSearchTrajectoryFlags::FacingDirectionXY))
+			{
+				Heading->ComponentStripping = EComponentStrippingVector::StripZ;
+			}
 			SubChannels.Add(Heading);
 		}
 	}
@@ -140,6 +153,8 @@ FString UPoseSearchFeatureChannel_Trajectory::GetLabel() const
 
 float UPoseSearchFeatureChannel_Trajectory::GetEstimatedSpeedRatio(TConstArrayView<float> QueryVector, TConstArrayView<float> PoseVector) const
 {
+	using namespace UE::PoseSearch;
+
 	float EstimatedQuerySpeed = 0.f;
 	float EstimatedPoseSpeed = 0.f;
 
@@ -149,8 +164,8 @@ float UPoseSearchFeatureChannel_Trajectory::GetEstimatedSpeedRatio(TConstArrayVi
 		{
 			if (!Velocity->bNormalize)
 			{
-				const FVector QueryVelocity = UE::PoseSearch::FFeatureVectorHelper::DecodeVectorAtOffset(QueryVector, Velocity->GetChannelDataOffset());
-				const FVector PoseVelocity = UE::PoseSearch::FFeatureVectorHelper::DecodeVectorAtOffset(PoseVector, Velocity->GetChannelDataOffset());
+				const FVector QueryVelocity = FFeatureVectorHelper::DecodeVector(QueryVector, Velocity->GetChannelDataOffset(), Velocity->ComponentStripping);
+				const FVector PoseVelocity = FFeatureVectorHelper::DecodeVector(PoseVector, Velocity->GetChannelDataOffset(), Velocity->ComponentStripping);
 				EstimatedQuerySpeed += QueryVelocity.Length();
 				EstimatedPoseSpeed += PoseVelocity.Length();
 			}
diff --git a/Engine/Plugins/Experimental/Animation/PoseSearch/Source/Runtime/Private/PoseSearchFeatureChannel_Velocity.cpp b/Engine/Plugins/Experimental/Animation/PoseSearch/Source/Runtime/Private/PoseSearchFeatureChannel_Velocity.cpp
--- a/Engine/Plugins/Experimental/Animation/PoseSearch/Source/Runtime/Private/PoseSearchFeatureChannel_Velocity.cpp
+++ b/Engine/Plugins/Experimental/Animation/PoseSearch/Source/Runtime/Private/PoseSearchFeatureChannel_Velocity.cpp
@@ -13,7 +13,7 @@
 void UPoseSearchFeatureChannel_Velocity::Finalize(UPoseSearchSchema* Schema)
 {
 	ChannelDataOffset = Schema->SchemaCardinality;
-	ChannelCardinality = UE::PoseSearch::FFeatureVectorHelper::EncodeVectorCardinality;
+	ChannelCardinality = UE::PoseSearch::FFeatureVectorHelper::GetVectorCardinality(ComponentStripping);
 	Schema->SchemaCardinality += ChannelCardinality;
 
 	SchemaBoneIdx = Schema->AddBoneReference(Bone);
@@ -21,9 +21,7 @@ void UPoseSearchFeatureChannel_Velocity::Finalize(UPoseSearchSchema* Schema)
 
 void UPoseSearchFeatureChannel_Velocity::FillWeights(TArray<float>& Weights) const
 {
-	using namespace UE::PoseSearch;
-
-	for (int32 i = 0; i != FFeatureVectorHelper::EncodeVectorCardinality; ++i)
+	for (int32 i = 0; i < ChannelCardinality; ++i)
 	{
 		Weights[ChannelDataOffset + i] = Weight;
 	}
@@ -72,8 +70,7 @@ void UPoseSearchFeatureChannel_Velocity::IndexAsset(UE::PoseSearch::IAssetIndexe
 			LinearVelocity = LinearVelocity.GetClampedToMaxSize(1.f);
 		}
 
-		int32 DataOffset = ChannelDataOffset;
-		FFeatureVectorHelper::EncodeVector(IndexingContext.GetPoseVector(VectorIdx, FeatureVectorTable), DataOffset, LinearVelocity);
+		FFeatureVectorHelper::EncodeVector(IndexingContext.GetPoseVector(VectorIdx, FeatureVectorTable), ChannelDataOffset, LinearVelocity, ComponentStripping);
 	}
 }
 
@@ -89,9 +86,8 @@ void UPoseSearchFeatureChannel_Velocity::BuildQuery(UE::PoseSearch::FSearchConte
 		if (bIsCurrentResultValid)
 		{
 			const float LerpValue = InputQueryPose == EInputQueryPose::UseInterpolatedContinuingPose ? SearchContext.CurrentResult.LerpValue : 0.f;
-			int32 DataOffset = ChannelDataOffset;
-			// @todo: we should normalize if bNormalize && LerpValue != 0
-			FFeatureVectorHelper::EncodeVector(InOutQuery.EditValues(), DataOffset, SearchContext.GetCurrentResultPrevPoseVector(), SearchContext.GetCurrentResultPoseVector(), SearchContext.GetCurrentResultNextPoseVector(), LerpValue);
+			// @todo: we should normalize if EPoseSearchVelocityFlags::Normalize && LerpValue != 0
+			FFeatureVectorHelper::EncodeVector(InOutQuery.EditValues(), ChannelDataOffset, SearchContext.GetCurrentResultPrevPoseVector(), SearchContext.GetCurrentResultPoseVector(), SearchContext.GetCurrentResultNextPoseVector(), LerpValue, false, ComponentStripping);
 		}
 		// else leave the InOutQuery set to zero since the SearchContext.History is invalid and it'll fail if we continue
 	}
@@ -124,8 +120,7 @@ void UPoseSearchFeatureChannel_Velocity::BuildQuery(UE::PoseSearch::FSearchConte
 			LinearVelocity = LinearVelocity.GetClampedToMaxSize(1.f);
 		}
 
-		int32 DataOffset = ChannelDataOffset;
-		FFeatureVectorHelper::EncodeVector(InOutQuery.EditValues(), DataOffset, LinearVelocity);
+		FFeatureVectorHelper::EncodeVector(InOutQuery.EditValues(), ChannelDataOffset, LinearVelocity, ComponentStripping);
 	}
 }
 
@@ -140,7 +135,7 @@ void UPoseSearchFeatureChannel_Velocity::DebugDraw(const UE::PoseSearch::FDebugD
 	const FColor Color = DrawParams.GetColor(ColorPresetIndex);
 	const float LinearVelocityScale = bNormalize ? 15.f : 0.08f;
 
-	const FVector LinearVelocity = DrawParams.RootTransform.TransformVector(FFeatureVectorHelper::DecodeVectorAtOffset(PoseVector, ChannelDataOffset));
+	const FVector LinearVelocity = DrawParams.RootTransform.TransformVector(FFeatureVectorHelper::DecodeVector(PoseVector, ChannelDataOffset, ComponentStripping));
 	const FVector BoneVelDirection = LinearVelocity.GetSafeNormal();
 	const FVector BonePos = DrawParams.GetCachedPosition(SampleTimeOffset, SchemaBoneIdx);
 
@@ -172,6 +167,15 @@ FString UPoseSearchFeatureChannel_Velocity::GetLabel() const
 		Label.Append(TEXT("Dir"));
 	}
 
+	if (ComponentStripping == EComponentStrippingVector::StripXY)
+	{
+		Label.Append(TEXT("_z"));
+	}
+	else if (ComponentStripping == EComponentStrippingVector::StripZ)
+	{
+		Label.Append(TEXT("_xy"));
+	}
+
 	const FBoneReference& BoneReference = GetSchema()->BoneReferences[SchemaBoneIdx];
 	if (BoneReference.HasValidSetup())
 	{
diff --git a/Engine/Plugins/Experimental/Animation/PoseSearch/Source/Runtime/Private/PoseSearchFeatureChannel_Velocity.h b/Engine/Plugins/Experimental/Animation/PoseSearch/Source/Runtime/Private/PoseSearchFeatureChannel_Velocity.h
--- a/Engine/Plugins/Experimental/Animation/PoseSearch/Source/Runtime/Private/PoseSearchFeatureChannel_Velocity.h
+++ b/Engine/Plugins/Experimental/Animation/PoseSearch/Source/Runtime/Private/PoseSearchFeatureChannel_Velocity.h
@@ -6,7 +6,7 @@
 #include "BoneContainer.h"
 #include "PoseSearchFeatureChannel_Velocity.generated.h"
 
-UCLASS(BlueprintType, EditInlineNew, meta = (DisplayName = "Velocity Channel"), CollapseCategories)
+UCLASS(EditInlineNew, meta = (DisplayName = "Velocity Channel"), CollapseCategories)
 class POSESEARCH_API UPoseSearchFeatureChannel_Velocity : public UPoseSearchFeatureChannel
 {
 	GENERATED_BODY()
@@ -38,6 +38,9 @@ public:
 	UPROPERTY(EditAnywhere, Category = "Settings")
 	bool bNormalize = false;
 
+	UPROPERTY(EditAnywhere, Category = "Settings")
+	EComponentStrippingVector ComponentStripping = EComponentStrippingVector::None;
+
 	// UPoseSearchFeatureChannel interface
 	virtual void Finalize(UPoseSearchSchema* Schema) override;
 	virtual void FillWeights(TArray<float>& Weights) const override;
diff --git a/Engine/Plugins/Experimental/Animation/PoseSearch/Source/Runtime/Public/PoseSearch/PoseSearchFeatureChannel.h b/Engine/Plugins/Experimental/Animation/PoseSearch/Source/Runtime/Public/PoseSearch/PoseSearchFeatureChannel.h
--- a/Engine/Plugins/Experimental/Animation/PoseSearch/Source/Runtime/Public/PoseSearch/PoseSearchFeatureChannel.h
+++ b/Engine/Plugins/Experimental/Animation/PoseSearch/Source/Runtime/Public/PoseSearch/PoseSearchFeatureChannel.h
@@ -10,7 +10,20 @@ class UPoseSearchSchema;
 struct FPoseSearchPoseMetadata;
 struct FPoseSearchFeatureVectorBuilder;
 
-UENUM(BlueprintType)
+UENUM()
+enum class EComponentStrippingVector : uint8
+{
+	// no component stripping
+	None,
+
+	// stripping X and Y components (matching only on the horizontal plane) 
+	StripXY,
+
+	// stripping Z (matching only vertically - caring only about the height of the feature) 
+	StripZ,
+};
+
+UENUM()
 enum class EInputQueryPose : uint8
 {
 	// use character pose to compose the query
@@ -21,9 +34,6 @@ enum class EInputQueryPose : uint8
 
 	// if available reuse and interpolate continuing pose from the database to compose the query or else UseCharacterPose
 	UseInterpolatedContinuingPose,
-
-	Num UMETA(Hidden),
-	Invalid = Num UMETA(Hidden)
 };
 
 namespace UE::PoseSearch
@@ -37,29 +47,18 @@ class IAssetIndexer;
 class POSESEARCH_API FFeatureVectorHelper
 {
 public:
-	enum { EncodeQuatCardinality = 6 };
-	static void EncodeQuat(TArrayView<float> Values, int32& DataOffset, const FQuat& Quat);
-	static void EncodeQuat(TArrayView<float> Values, int32& DataOffset, TConstArrayView<float> PrevValues, TConstArrayView<float> CurValues, TConstArrayView<float> NextValues, float LerpValue);
-	static FQuat DecodeQuat(TConstArrayView<float> Values, int32& DataOffset);
-	static FQuat DecodeQuatAtOffset(TConstArrayView<float> Values, int32 DataOffset);
-
-	enum { EncodeVectorCardinality = 3 };
-	static void EncodeVector(TArrayView<float> Values, int32& DataOffset, const FVector& Vector);
-	static void EncodeVector(TArrayView<float> Values, int32& DataOffset, TConstArrayView<float> PrevValues, TConstArrayView<float> CurValues, TConstArrayView<float> NextValues, float LerpValue, bool bNormalize = false);
-	static FVector DecodeVector(TConstArrayView<float> Values, int32& DataOffset);
-	static FVector DecodeVectorAtOffset(TConstArrayView<float> Values, int32 DataOffset);
-
-	enum { EncodeVector2DCardinality = 2 };
-	static void EncodeVector2D(TArrayView<float> Values, int32& DataOffset, const FVector2D& Vector2D);
-	static void EncodeVector2D(TArrayView<float> Values, int32& DataOffset, TConstArrayView<float> PrevValues, TConstArrayView<float> CurValues, TConstArrayView<float> NextValues, float LerpValue);
-	static FVector2D DecodeVector2D(TConstArrayView<float> Values, int32& DataOffset);
-	static FVector2D DecodeVector2DAtOffset(TConstArrayView<float> Values, int32 DataOffset);
-
-	enum { EncodeFloatCardinality = 1 };
-	static void EncodeFloat(TArrayView<float> Values, int32& DataOffset, const float Value);
-	static void EncodeFloat(TArrayView<float> Values, int32& DataOffset, TConstArrayView<float> PrevValues, TConstArrayView<float> CurValues, TConstArrayView<float> NextValues, float LerpValue);
-	static float DecodeFloat(TConstArrayView<float> Values, int32& DataOffset);
-	static float DecodeFloatAtOffset(TConstArrayView<float> Values, int32 DataOffset);
+	static int32 GetVectorCardinality(EComponentStrippingVector ComponentStrippingVector);
+	static void EncodeVector(TArrayView<float> Values, int32 DataOffset, const FVector& Vector, EComponentStrippingVector ComponentStrippingVector);
+	static void EncodeVector(TArrayView<float> Values, int32 DataOffset, TConstArrayView<float> PrevValues, TConstArrayView<float> CurValues, TConstArrayView<float> NextValues, float LerpValue, bool bNormalize, EComponentStrippingVector ComponentStrippingVector);
+	static FVector DecodeVector(TConstArrayView<float> Values, int32 DataOffset, EComponentStrippingVector ComponentStrippingVector);
+
+	static void EncodeVector2D(TArrayView<float> Values, int32 DataOffset, const FVector2D& Vector2D);
+	static void EncodeVector2D(TArrayView<float> Values, int32 DataOffset, TConstArrayView<float> PrevValues, TConstArrayView<float> CurValues, TConstArrayView<float> NextValues, float LerpValue);
+	static FVector2D DecodeVector2D(TConstArrayView<float> Values, int32 DataOffset);
+
+	static void EncodeFloat(TArrayView<float> Values, int32 DataOffset, const float Value);
+	static void EncodeFloat(TArrayView<float> Values, int32 DataOffset, TConstArrayView<float> PrevValues, TConstArrayView<float> CurValues, TConstArrayView<float> NextValues, float LerpValue);
+	static float DecodeFloat(TConstArrayView<float> Values, int32 DataOffset);
 };
 
 } // namespace UE::PoseSearch
