motion matching - avoiding storing kdtree data if EPoseSearchMode is BruteForce

diff --git a/Engine/Plugins/Experimental/Animation/PoseSearch/Source/Runtime/Private/KDTree.cpp b/Engine/Plugins/Experimental/Animation/PoseSearch/Source/Runtime/Private/KDTree.cpp
--- a/Engine/Plugins/Experimental/Animation/PoseSearch/Source/Runtime/Private/KDTree.cpp
+++ b/Engine/Plugins/Experimental/Animation/PoseSearch/Source/Runtime/Private/KDTree.cpp
@@ -28,7 +28,10 @@ FKDTree::FKDTree(int32 Count, int32 Dim, const float* Data, int32 MaxLeafSize)
 , Impl(nullptr)
 {
 #if UE_POSE_SEARCH_USE_NANOFLANN
-	Impl = new FKDTreeImplementation(Dim, DataSource, nanoflann::KDTreeSingleIndexAdaptorParams(MaxLeafSize));
+	if (Count > 0 && Dim > 0 && Data)
+	{
+		Impl = new FKDTreeImplementation(Dim, DataSource, nanoflann::KDTreeSingleIndexAdaptorParams(MaxLeafSize));
+	}
 #endif
 }
 
@@ -36,18 +39,11 @@ FKDTree::FKDTree()
 : DataSource(0, 0, nullptr)
 , Impl(nullptr)
 {
-#if UE_POSE_SEARCH_USE_NANOFLANN
-	Impl = new FKDTreeImplementation(0, DataSource, nanoflann::KDTreeSingleIndexAdaptorParams(0));
-#endif
 }
 
 FKDTree::~FKDTree()
 {
-#if UE_POSE_SEARCH_USE_NANOFLANN
-	delete Impl;
-	Impl = nullptr;
-#endif
-	DataSource = FDataSource();
+	Reset();
 }
 
 #if UE_POSE_SEARCH_USE_NANOFLANN
@@ -83,9 +79,8 @@ void CopySubTree(FKDTree& KDTree, FKDTreeImplementation::NodePtr& ThisNode, cons
 FKDTree::FKDTree(const FKDTree& Other)
 {
 #if UE_POSE_SEARCH_USE_NANOFLANN
-	if (this != &Other)
+	if (this != &Other && Other.Impl)
 	{
-		check(Other.Impl);
 		Impl = new FKDTreeImplementation(0, DataSource, nanoflann::KDTreeSingleIndexAdaptorParams(0));
 
 		DataSource = Other.DataSource;
@@ -129,15 +124,24 @@ FKDTree& FKDTree::operator=(const FKDTree& Other)
 {
 	if (this != &Other)
 	{
-		this->~FKDTree();
+		Reset();
 		new(this)FKDTree(Other);
 	}
 	return *this;
 }
 
+void FKDTree::Reset()
+{
+#if UE_POSE_SEARCH_USE_NANOFLANN
+	delete Impl;
+	Impl = nullptr;
+#endif
+	DataSource = FDataSource();
+}
+
 void FKDTree::Construct(int32 Count, int32 Dim, const float* Data, int32 MaxLeafSize)
 {
-	this->~FKDTree();
+	Reset();
 	new(this)FKDTree(Count, Dim, Data, MaxLeafSize);
 }
 
@@ -215,6 +219,11 @@ FArchive& Serialize(FArchive& Ar, FKDTree& KDTree, const float* KDTreeData)
 {
 #if UE_POSE_SEARCH_USE_NANOFLANN
 
+	if (Ar.IsLoading() && !KDTree.Impl)
+	{
+		KDTree.Impl = new FKDTreeImplementation(0, KDTree.DataSource, nanoflann::KDTreeSingleIndexAdaptorParams(0));
+	}
+
 	check(KDTree.Impl);
 	check(KDTree.Impl->m_size < UINT_MAX);
 	uint32 KDTreeSize = KDTree.Impl->m_size;
diff --git a/Engine/Plugins/Experimental/Animation/PoseSearch/Source/Runtime/Private/PoseSearchDatabase.cpp b/Engine/Plugins/Experimental/Animation/PoseSearch/Source/Runtime/Private/PoseSearchDatabase.cpp
--- a/Engine/Plugins/Experimental/Animation/PoseSearch/Source/Runtime/Private/PoseSearchDatabase.cpp
+++ b/Engine/Plugins/Experimental/Animation/PoseSearch/Source/Runtime/Private/PoseSearchDatabase.cpp
@@ -387,7 +387,7 @@ void UPoseSearchDatabase::SetSearchIndex(const FPoseSearchIndex& SearchIndex)
 const FPoseSearchIndex& UPoseSearchDatabase::GetSearchIndex() const
 {
 	// making sure the search index is consistent. if it fails the calling code hasn't been protected by FAsyncPoseSearchDatabasesManagement::RequestAsyncBuildIndex
-	check(Schema && Schema->IsValid() && !SearchIndexPrivate.IsEmpty() && SearchIndexPrivate.WeightsSqrt.Num() == Schema->SchemaCardinality && SearchIndexPrivate.KDTree.Impl);
+	check(Schema && Schema->IsValid() && !SearchIndexPrivate.IsEmpty() && SearchIndexPrivate.WeightsSqrt.Num() == Schema->SchemaCardinality);
 	return SearchIndexPrivate;
 }
 
diff --git a/Engine/Plugins/Experimental/Animation/PoseSearch/Source/Runtime/Private/PoseSearchDerivedData.cpp b/Engine/Plugins/Experimental/Animation/PoseSearch/Source/Runtime/Private/PoseSearchDerivedData.cpp
--- a/Engine/Plugins/Experimental/Animation/PoseSearch/Source/Runtime/Private/PoseSearchDerivedData.cpp
+++ b/Engine/Plugins/Experimental/Animation/PoseSearch/Source/Runtime/Private/PoseSearchDerivedData.cpp
@@ -435,18 +435,18 @@ static void PreprocessSearchIndexPCAData(FPoseSearchIndex& SearchIndex, int32 Nu
 	// binding SearchIndex.Values and SearchIndex.PCAValues Eigen row major matrix maps
 	const int32 NumPoses = SearchIndex.NumPoses;
 
+	SearchIndex.PCAExplainedVariance = 0.f;
+
 	SearchIndex.PCAValues.Reset();
 	SearchIndex.Mean.Reset();
 	SearchIndex.PCAProjectionMatrix.Reset();
 
-	SearchIndex.PCAValues.AddZeroed(NumPoses * NumberOfPrincipalComponents);
-	SearchIndex.Mean.AddZeroed(NumDimensions);
-	SearchIndex.PCAProjectionMatrix.AddZeroed(NumDimensions * NumberOfPrincipalComponents);
-
-	SearchIndex.PCAExplainedVariance = 0.f;
-
-	if (NumDimensions > 0)
+	if (PoseSearchMode != EPoseSearchMode::BruteForce && NumDimensions > 0)
 	{
+		SearchIndex.PCAValues.AddZeroed(NumPoses * NumberOfPrincipalComponents);
+		SearchIndex.Mean.AddZeroed(NumDimensions);
+		SearchIndex.PCAProjectionMatrix.AddZeroed(NumDimensions * NumberOfPrincipalComponents);
+
 		const RowMajorVectorMapConst MapWeightsSqrt(SearchIndex.WeightsSqrt.GetData(), 1, NumDimensions);
 		const RowMajorMatrixMapConst MapValues(SearchIndex.Values.GetData(), NumPoses, NumDimensions);
 		const RowMajorMatrix WeightedValues = MapValues.array().rowwise() * MapWeightsSqrt.array();
@@ -528,76 +528,80 @@ static void PreprocessSearchIndexPCAData(FPoseSearchIndex& SearchIndex, int32 Nu
 
 static void PreprocessSearchIndexKDTree(FPoseSearchIndex& SearchIndex, int32 NumDimensions, uint32 NumberOfPrincipalComponents, EPoseSearchMode PoseSearchMode, int32 KDTreeMaxLeafSize, int32 KDTreeQueryNumNeighbors)
 {
-	const int32 NumPoses = SearchIndex.NumPoses;
-	SearchIndex.KDTree.Construct(NumPoses, NumberOfPrincipalComponents, SearchIndex.PCAValues.GetData(), KDTreeMaxLeafSize);
-
-	if (PoseSearchMode == EPoseSearchMode::PCAKDTree_Validate)
+	SearchIndex.KDTree.Reset();
+	if (PoseSearchMode != EPoseSearchMode::BruteForce && NumDimensions > 0)
 	{
-		// testing the KDTree is returning the proper searches for all the points in pca space
-		int32 NumberOfFailingPoints = 0;
-		for (size_t PointIndex = 0; PointIndex < NumPoses; ++PointIndex)
+		const int32 NumPoses = SearchIndex.NumPoses;
+		SearchIndex.KDTree.Construct(NumPoses, NumberOfPrincipalComponents, SearchIndex.PCAValues.GetData(), KDTreeMaxLeafSize);
+
+		if (PoseSearchMode == EPoseSearchMode::PCAKDTree_Validate)
 		{
-			TArray<size_t> ResultIndexes;
-			TArray<float> ResultDistanceSqr;
-			ResultIndexes.SetNum(KDTreeQueryNumNeighbors + 1);
-			ResultDistanceSqr.SetNum(KDTreeQueryNumNeighbors + 1);
-			FKDTree::KNNResultSet ResultSet(KDTreeQueryNumNeighbors, ResultIndexes, ResultDistanceSqr);
-			SearchIndex.KDTree.FindNeighbors(ResultSet, &SearchIndex.PCAValues[PointIndex * NumberOfPrincipalComponents]);
-
-			size_t ResultIndex = 0;
-			for (; ResultIndex < ResultSet.Num(); ++ResultIndex)
+			// testing the KDTree is returning the proper searches for all the points in pca space
+			int32 NumberOfFailingPoints = 0;
+			for (size_t PointIndex = 0; PointIndex < NumPoses; ++PointIndex)
 			{
-				if (PointIndex == ResultIndexes[ResultIndex])
+				TArray<size_t> ResultIndexes;
+				TArray<float> ResultDistanceSqr;
+				ResultIndexes.SetNum(KDTreeQueryNumNeighbors + 1);
+				ResultDistanceSqr.SetNum(KDTreeQueryNumNeighbors + 1);
+				FKDTree::KNNResultSet ResultSet(KDTreeQueryNumNeighbors, ResultIndexes, ResultDistanceSqr);
+				SearchIndex.KDTree.FindNeighbors(ResultSet, &SearchIndex.PCAValues[PointIndex * NumberOfPrincipalComponents]);
+
+				size_t ResultIndex = 0;
+				for (; ResultIndex < ResultSet.Num(); ++ResultIndex)
 				{
-					check(ResultDistanceSqr[ResultIndex] < UE_KINDA_SMALL_NUMBER);
-					break;
+					if (PointIndex == ResultIndexes[ResultIndex])
+					{
+						check(ResultDistanceSqr[ResultIndex] < UE_KINDA_SMALL_NUMBER);
+						break;
+					}
+				}
+				if (ResultIndex == ResultSet.Num())
+				{
+					++NumberOfFailingPoints;
 				}
 			}
-			if (ResultIndex == ResultSet.Num())
-			{
-				++NumberOfFailingPoints;
-			}
-		}
 
-		check(NumberOfFailingPoints == 0);
+			check(NumberOfFailingPoints == 0);
 
-		// testing the KDTree is returning the proper searches for all the original points transformed in pca space
-		NumberOfFailingPoints = 0;
-		for (size_t PointIndex = 0; PointIndex < NumPoses; ++PointIndex)
-		{
-			TArray<size_t> ResultIndexes;
-			TArray<float> ResultDistanceSqr;
-			ResultIndexes.SetNum(KDTreeQueryNumNeighbors + 1);
-			ResultDistanceSqr.SetNum(KDTreeQueryNumNeighbors + 1);
-			FKDTree::KNNResultSet ResultSet(KDTreeQueryNumNeighbors, ResultIndexes, ResultDistanceSqr);
+			// testing the KDTree is returning the proper searches for all the original points transformed in pca space
+			NumberOfFailingPoints = 0;
+			for (size_t PointIndex = 0; PointIndex < NumPoses; ++PointIndex)
+			{
+				TArray<size_t> ResultIndexes;
+				TArray<float> ResultDistanceSqr;
+				ResultIndexes.SetNum(KDTreeQueryNumNeighbors + 1);
+				ResultDistanceSqr.SetNum(KDTreeQueryNumNeighbors + 1);
+				FKDTree::KNNResultSet ResultSet(KDTreeQueryNumNeighbors, ResultIndexes, ResultDistanceSqr);
 
-			const RowMajorVectorMapConst MapValues(&SearchIndex.Values[PointIndex * NumDimensions], 1, NumDimensions);
-			const RowMajorVectorMapConst MapWeightsSqrt(SearchIndex.WeightsSqrt.GetData(), 1, NumDimensions);
-			const RowMajorVectorMapConst Mean(SearchIndex.Mean.GetData(), 1, NumDimensions);
-			const ColMajorMatrixMapConst PCAProjectionMatrix(SearchIndex.PCAProjectionMatrix.GetData(), NumDimensions, NumberOfPrincipalComponents);
+				const RowMajorVectorMapConst MapValues(&SearchIndex.Values[PointIndex * NumDimensions], 1, NumDimensions);
+				const RowMajorVectorMapConst MapWeightsSqrt(SearchIndex.WeightsSqrt.GetData(), 1, NumDimensions);
+				const RowMajorVectorMapConst Mean(SearchIndex.Mean.GetData(), 1, NumDimensions);
+				const ColMajorMatrixMapConst PCAProjectionMatrix(SearchIndex.PCAProjectionMatrix.GetData(), NumDimensions, NumberOfPrincipalComponents);
 
-			const RowMajorMatrix WeightedValues = MapValues.array() * MapWeightsSqrt.array();
-			const RowMajorMatrix CenteredValues = WeightedValues - Mean;
-			const RowMajorVector ProjectedValues = CenteredValues * PCAProjectionMatrix;
+				const RowMajorMatrix WeightedValues = MapValues.array() * MapWeightsSqrt.array();
+				const RowMajorMatrix CenteredValues = WeightedValues - Mean;
+				const RowMajorVector ProjectedValues = CenteredValues * PCAProjectionMatrix;
 
-			SearchIndex.KDTree.FindNeighbors(ResultSet, ProjectedValues.data());
+				SearchIndex.KDTree.FindNeighbors(ResultSet, ProjectedValues.data());
 
-			size_t ResultIndex = 0;
-			for (; ResultIndex < ResultSet.Num(); ++ResultIndex)
-			{
-				if (PointIndex == ResultIndexes[ResultIndex])
+				size_t ResultIndex = 0;
+				for (; ResultIndex < ResultSet.Num(); ++ResultIndex)
 				{
-					check(ResultDistanceSqr[ResultIndex] < UE_KINDA_SMALL_NUMBER);
-					break;
+					if (PointIndex == ResultIndexes[ResultIndex])
+					{
+						check(ResultDistanceSqr[ResultIndex] < UE_KINDA_SMALL_NUMBER);
+						break;
+					}
+				}
+				if (ResultIndex == ResultSet.Num())
+				{
+					++NumberOfFailingPoints;
 				}
 			}
-			if (ResultIndex == ResultSet.Num())
-			{
-				++NumberOfFailingPoints;
-			}
-		}
 
-		check(NumberOfFailingPoints == 0);
+			check(NumberOfFailingPoints == 0);
+		}
 	}
 }
 
@@ -786,7 +790,7 @@ void FPoseSearchDatabaseAsyncCacheTask::Wait(FCriticalSection& OuterMutex)
 	{
 		Database->SetSearchIndex(SearchIndex); // @todo: implement FPoseSearchIndex move ctor and assignment operator and use a MoveTemp(SearchIndex) here
 
-		check(Database->Schema && Database->Schema->IsValid() && !SearchIndex.IsEmpty() && SearchIndex.WeightsSqrt.Num() == Database->Schema->SchemaCardinality && SearchIndex.KDTree.Impl);
+		check(Database->Schema && Database->Schema->IsValid() && !SearchIndex.IsEmpty() && SearchIndex.WeightsSqrt.Num() == Database->Schema->SchemaCardinality);
 
 		SetState(EState::Ended);
 		bBroadcastOnDerivedDataRebuild = true;
diff --git a/Engine/Plugins/Experimental/Animation/PoseSearch/Source/Runtime/Private/PoseSearchIndex.cpp b/Engine/Plugins/Experimental/Animation/PoseSearch/Source/Runtime/Private/PoseSearchIndex.cpp
--- a/Engine/Plugins/Experimental/Animation/PoseSearch/Source/Runtime/Private/PoseSearchIndex.cpp
+++ b/Engine/Plugins/Experimental/Animation/PoseSearch/Source/Runtime/Private/PoseSearchIndex.cpp
@@ -150,10 +150,10 @@ FArchive& operator<<(FArchive& Ar, FPoseSearchIndexBase& Index)
 // FPoseSearchIndex
 FPoseSearchIndex::FPoseSearchIndex(const FPoseSearchIndex& Other)
 	: FPoseSearchIndexBase(Other)
+	, WeightsSqrt(Other.WeightsSqrt)
 	, PCAValues(Other.PCAValues)
 	, PCAProjectionMatrix(Other.PCAProjectionMatrix)
 	, Mean(Other.Mean)
-	, WeightsSqrt(Other.WeightsSqrt)
 	, KDTree(Other.KDTree)
 	, PCAExplainedVariance(Other.PCAExplainedVariance)
 {
@@ -230,32 +230,38 @@ FArchive& operator<<(FArchive& Ar, FPoseSearchIndex& Index)
 {
 	Ar << static_cast<FPoseSearchIndexBase&>(Index);
 
-	int32 NumPCAValues = 0;
-
-	if (Ar.IsSaving())
-	{
-		NumPCAValues = Index.PCAValues.Num();
-	}
+	Ar << Index.WeightsSqrt;
 
+	int32 NumPCAValues = Ar.IsSaving() ? Index.PCAValues.Num() : 0;
 	Ar << NumPCAValues;
 
-	if (Ar.IsLoading())
+	if (NumPCAValues > 0)
 	{
-		Index.PCAValues.SetNumUninitialized(NumPCAValues);
-	}
+		if (Ar.IsLoading())
+		{
+			Index.PCAValues.SetNumUninitialized(NumPCAValues);
+		}
 
-	if (Index.PCAValues.Num() > 0)
-	{
-		Ar.Serialize(&Index.PCAValues[0], Index.PCAValues.Num() * Index.PCAValues.GetTypeSize());
-	}
+		if (Index.PCAValues.Num() > 0)
+		{
+			Ar.Serialize(&Index.PCAValues[0], Index.PCAValues.Num() * Index.PCAValues.GetTypeSize());
+		}
 
-	Ar << Index.WeightsSqrt;
-	Ar << Index.Mean;
-	Ar << Index.PCAProjectionMatrix;
+		Ar << Index.Mean;
+		Ar << Index.PCAProjectionMatrix;
 
-	Serialize(Ar, Index.KDTree, Index.PCAValues.GetData());
+		Serialize(Ar, Index.KDTree, Index.PCAValues.GetData());
 
-	Ar << Index.PCAExplainedVariance;
+		Ar << Index.PCAExplainedVariance;
+	}
+	else if (Ar.IsLoading())
+	{
+		Index.PCAValues.Reset();
+		Index.Mean.Reset();
+		Index.PCAProjectionMatrix.Reset();
+		Index.KDTree.Reset();
+		Index.PCAExplainedVariance = 0.f;
+	}
 
 	return Ar;
 }
\ No newline at end of file
diff --git a/Engine/Plugins/Experimental/Animation/PoseSearch/Source/Runtime/Public/PoseSearch/KDTree.h b/Engine/Plugins/Experimental/Animation/PoseSearch/Source/Runtime/Public/PoseSearch/KDTree.h
--- a/Engine/Plugins/Experimental/Animation/PoseSearch/Source/Runtime/Public/PoseSearch/KDTree.h
+++ b/Engine/Plugins/Experimental/Animation/PoseSearch/Source/Runtime/Public/PoseSearch/KDTree.h
@@ -17,7 +17,7 @@ struct POSESEARCH_API FKDTree
 	struct FDataSource
 	{
 		FDataSource(int32 pointCount, int32 pointDim, const float* data)
-		: PointCount(pointDim > 0 ? pointCount : 0)
+		: PointCount((pointDim > 0 && data) ? pointCount : 0)
 		, PointDim(pointDim)
 		, Data(data)
 		{
@@ -133,8 +133,9 @@ struct POSESEARCH_API FKDTree
 	FKDTree& operator=(const FKDTree& Other);
 	FKDTree& operator=(FKDTree&& Other) = delete;
 
-	bool FindNeighbors(KNNResultSet& Result, const float* Query) const;
+	void Reset();
 	void Construct(int32 Count, int32 Dim, const float* Data, int32 MaxLeafSize);
+	bool FindNeighbors(KNNResultSet& Result, const float* Query) const;
 	SIZE_T GetAllocatedSize() const;
 
 	FDataSource DataSource;
