motion matching - defaulting UPoseSearchFeatureChannel_Trajectory to meaningful values for a MM locomotion setup

diff --git a/Engine/Plugins/Experimental/Animation/PoseSearch/Source/Runtime/Private/PoseSearchFeatureChannel_Trajectory.cpp b/Engine/Plugins/Experimental/Animation/PoseSearch/Source/Runtime/Private/PoseSearchFeatureChannel_Trajectory.cpp
--- a/Engine/Plugins/Experimental/Animation/PoseSearch/Source/Runtime/Private/PoseSearchFeatureChannel_Trajectory.cpp
+++ b/Engine/Plugins/Experimental/Animation/PoseSearch/Source/Runtime/Private/PoseSearchFeatureChannel_Trajectory.cpp
@@ -13,6 +13,18 @@
 
 #define LOCTEXT_NAMESPACE "PoseSearchFeatureChannels"
 
+UPoseSearchFeatureChannel_Trajectory::UPoseSearchFeatureChannel_Trajectory()
+{
+#if WITH_EDITOR
+	// defaulting UPoseSearchFeatureChannel_Trajectory to a meaningful locomotion setup
+	Weight = 7.f;
+	Samples.Add(FPoseSearchTrajectorySample({ -0.4f, int32(EPoseSearchTrajectoryFlags::PositionXY), 0.4f, 0 }));
+	Samples.Add(FPoseSearchTrajectorySample({ 0.f, int32(EPoseSearchTrajectoryFlags::VelocityXY | EPoseSearchTrajectoryFlags::FacingDirectionXY), 2.f }));
+	Samples.Add(FPoseSearchTrajectorySample({ 0.35f, int32(EPoseSearchTrajectoryFlags::PositionXY | EPoseSearchTrajectoryFlags::FacingDirectionXY), 0.7f }));
+	Samples.Add(FPoseSearchTrajectorySample({ 0.7f, int32(EPoseSearchTrajectoryFlags::VelocityXY | EPoseSearchTrajectoryFlags::PositionXY | EPoseSearchTrajectoryFlags::FacingDirectionXY), 0.5f }));
+#endif
+}
+
 void UPoseSearchFeatureChannel_Trajectory::Finalize(UPoseSearchSchema* Schema)
 {
 	SubChannels.Reset();
diff --git a/Engine/Plugins/Experimental/Animation/PoseSearch/Source/Runtime/Private/PoseSearchFeatureChannel_Trajectory.h b/Engine/Plugins/Experimental/Animation/PoseSearch/Source/Runtime/Private/PoseSearchFeatureChannel_Trajectory.h
--- a/Engine/Plugins/Experimental/Animation/PoseSearch/Source/Runtime/Private/PoseSearchFeatureChannel_Trajectory.h
+++ b/Engine/Plugins/Experimental/Animation/PoseSearch/Source/Runtime/Private/PoseSearchFeatureChannel_Trajectory.h
@@ -56,6 +56,8 @@ public:
 	UPROPERTY(Transient)
 	TArray<TObjectPtr<UPoseSearchFeatureChannel>> SubChannels;
 
+	UPoseSearchFeatureChannel_Trajectory();
+
 	// UPoseSearchFeatureChannel_GroupBase interface
 	virtual TArrayView<TObjectPtr<UPoseSearchFeatureChannel>> GetSubChannels() override { return SubChannels; }
 	virtual TConstArrayView<TObjectPtr<UPoseSearchFeatureChannel>> GetSubChannels() const override { return SubChannels; }
