motion matching - FMotionMatchingSettings default values tuning

diff --git a/Engine/Plugins/Experimental/Animation/PoseSearch/Source/Runtime/Public/PoseSearch/PoseSearchLibrary.h b/Engine/Plugins/Experimental/Animation/PoseSearch/Source/Runtime/Public/PoseSearch/PoseSearchLibrary.h
--- a/Engine/Plugins/Experimental/Animation/PoseSearch/Source/Runtime/Public/PoseSearch/PoseSearchLibrary.h
+++ b/Engine/Plugins/Experimental/Animation/PoseSearch/Source/Runtime/Public/PoseSearch/PoseSearchLibrary.h
@@ -34,7 +34,7 @@ struct POSESEARCH_API FMotionMatchingSettings
 
 	// Number of max active blendin animation in the blend stack. If MaxActiveBlends is zero then blend stack is disabled
 	UPROPERTY(EditAnywhere, BlueprintReadWrite, Category=Settings, meta=(ClampMin="0"))
-	int32 MaxActiveBlends = 0;
+	int32 MaxActiveBlends = 4;
 
 	UPROPERTY(EditAnywhere, BlueprintReadWrite, Category=Settings, meta=(UseAsBlendProfile = true))
 	TObjectPtr<UBlendProfile> BlendProfile;
@@ -52,7 +52,7 @@ struct POSESEARCH_API FMotionMatchingSettings
 
 	// Don't jump to poses that has been selected previously within this many seconds in the past
 	UPROPERTY(EditAnywhere, BlueprintReadWrite, Category = Settings, meta = (ClampMin = "0"))
-	float PoseReselectHistory = 0.f;
+	float PoseReselectHistory = 0.3f;
 
 	// Minimum amount of time to wait between pose search queries
 	UPROPERTY(EditAnywhere, BlueprintReadWrite, Category=Settings, meta=(ClampMin="0"))
