motion matching - UPoseSearchSchema::LoopingCostBias: cost added to all poses from looping assets of databases that uses this schema

diff --git a/Engine/Plugins/Experimental/Animation/PoseSearch/Source/Runtime/Private/PoseSearchAssetIndexer.cpp b/Engine/Plugins/Experimental/Animation/PoseSearch/Source/Runtime/Private/PoseSearchAssetIndexer.cpp
--- a/Engine/Plugins/Experimental/Animation/PoseSearch/Source/Runtime/Private/PoseSearchAssetIndexer.cpp
+++ b/Engine/Plugins/Experimental/Animation/PoseSearch/Source/Runtime/Private/PoseSearchAssetIndexer.cpp
@@ -304,6 +304,12 @@ FPoseSearchPoseMetadata FAssetIndexer::GetMetadata(int32 SampleIdx) const
 			Metadata.ContinuingPoseCostAddend = ContinuingPoseCostBias->CostAddend;
 		}
 	}
+
+	if (IndexingContext.AssetSampler->IsLoopable())
+	{
+		Metadata.CostAddend += IndexingContext.Schema->LoopingCostBias;
+	}
+
 	return Metadata;
 }
 
diff --git a/Engine/Plugins/Experimental/Animation/PoseSearch/Source/Runtime/Public/PoseSearch/PoseSearchSchema.h b/Engine/Plugins/Experimental/Animation/PoseSearch/Source/Runtime/Public/PoseSearch/PoseSearchSchema.h
--- a/Engine/Plugins/Experimental/Animation/PoseSearch/Source/Runtime/Public/PoseSearch/PoseSearchSchema.h
+++ b/Engine/Plugins/Experimental/Animation/PoseSearch/Source/Runtime/Public/PoseSearch/PoseSearchSchema.h
@@ -89,6 +89,10 @@ public:
 	UPROPERTY(EditAnywhere, Category = "Schema")
 	float MirrorMismatchCostBias = 0.f;
 
+	// cost added to all poses from looping assets of databases that uses this schema
+	UPROPERTY(EditAnywhere, Category = "Schema")
+	float LoopingCostBias = 0.f;
+
 	UPROPERTY(EditAnywhere, Category = "Schema", meta = (ExcludeFromHash))
 	TArray<FPoseSearchSchemaColorPreset> ColorPresets;
 	
