motion matching - FSearchContext GetComponent/GetComponentSpaceTransform improvements, UPoseSearchFeatureChannel_Velocity::BuildQuery fix for !bUseCharacterSpaceVelocities

diff --git a/Engine/Plugins/Experimental/Animation/PoseSearch/Source/Runtime/Private/PoseSearchContext.cpp b/Engine/Plugins/Experimental/Animation/PoseSearch/Source/Runtime/Private/PoseSearchContext.cpp
--- a/Engine/Plugins/Experimental/Animation/PoseSearch/Source/Runtime/Private/PoseSearchContext.cpp
+++ b/Engine/Plugins/Experimental/Animation/PoseSearch/Source/Runtime/Private/PoseSearchContext.cpp
@@ -1,6 +1,7 @@
 // Copyright Epic Games, Inc. All Rights Reserved.
 
 #include "PoseSearch/PoseSearchContext.h"
+#include "Animation/MotionTrajectoryTypes.h"
 #include "AnimationRuntime.h"
 #include "DrawDebugHelpers.h"
 #include "PoseSearch/PoseSearchDatabase.h"
@@ -147,11 +148,21 @@ void DrawFeatureVector(FDebugDrawParams& DrawParams, int32 PoseIdx)
 //////////////////////////////////////////////////////////////////////////
 // FSearchContext
 
-FTransform FSearchContext::GetTransform(float SampleTime, const UPoseSearchSchema* Schema, int8 SchemaBoneIdx)
+FTransform FSearchContext::GetTransform(float SampleTime, const UPoseSearchSchema* Schema, int8 SchemaBoneIdx, bool bUseHistoryRoot)
 {
 	// collecting the RootTransform from the FPoseHistory
 	FTransform RootTransform = FTransform::Identity;
-	History->GetRootTransformAtTime(SampleTime, RootTransform);
+	if (bUseHistoryRoot)
+	{
+		check(History);
+		History->GetRootTransformAtTime(SampleTime, RootTransform);
+	}
+	else
+	{
+		check(Trajectory);
+		const FTrajectorySample TrajectorySample = Trajectory->GetSampleAtTime(SampleTime);
+		RootTransform = TrajectorySample.Transform;
+	}
 
 	const FBoneIndexType BoneIndexType = Schema->GetBoneIndexType(SchemaBoneIdx);
 	if (BoneIndexType != RootBoneIndexType)
@@ -165,7 +176,7 @@ FTransform FSearchContext::GetTransform(float SampleTime, const UPoseSearchSchem
 
 FTransform FSearchContext::GetComponentSpaceTransform(float SampleTime, const UPoseSearchSchema* Schema, int8 SchemaBoneIdx)
 {
-	check(History && Schema);
+	check(Schema);
 
 	const FBoneIndexType BoneIndexType = Schema->GetBoneIndexType(SchemaBoneIdx);
 	if (BoneIndexType != RootBoneIndexType)
@@ -176,6 +187,7 @@ FTransform FSearchContext::GetComponentSpaceTransform(float SampleTime, const UP
 		}
 	
 		// collecting the local bone transforms from the FPoseHistory
+		check(History);
 		TArray<FTransform> SampledLocalPose;
 		History->GetLocalPoseAtTime(SampleTime, Schema->BoneIndicesWithParents, SampledLocalPose);
 		
@@ -194,15 +206,15 @@ FTransform FSearchContext::GetComponentSpaceTransform(float SampleTime, const UP
 	return FTransform::Identity;
 }
 
-FTransform FSearchContext::GetComponentSpaceTransform(float SampleTime, float OriginTime, const UPoseSearchSchema* Schema, int8 SchemaBoneIdx)
+FTransform FSearchContext::GetComponentSpaceTransform(float SampleTime, float OriginTime, const UPoseSearchSchema* Schema, int8 SchemaBoneIdx, bool bUseHistoryRoot)
 {
 	if (SampleTime == OriginTime)
 	{
 		return GetComponentSpaceTransform(SampleTime, Schema, SchemaBoneIdx);
 	}
 
-	const FTransform RootBoneTransform = GetTransform(OriginTime, Schema);
-	FTransform BoneTransform = GetTransform(SampleTime, Schema, SchemaBoneIdx);
+	const FTransform RootBoneTransform = GetTransform(OriginTime, Schema, RootSchemaBoneIdx, bUseHistoryRoot);
+	FTransform BoneTransform = GetTransform(SampleTime, Schema, SchemaBoneIdx, bUseHistoryRoot);
 	BoneTransform.SetToRelativeTransform(RootBoneTransform);
 	return BoneTransform;
 }
diff --git a/Engine/Plugins/Experimental/Animation/PoseSearch/Source/Runtime/Private/PoseSearchFeatureChannel_Heading.cpp b/Engine/Plugins/Experimental/Animation/PoseSearch/Source/Runtime/Private/PoseSearchFeatureChannel_Heading.cpp
--- a/Engine/Plugins/Experimental/Animation/PoseSearch/Source/Runtime/Private/PoseSearchFeatureChannel_Heading.cpp
+++ b/Engine/Plugins/Experimental/Animation/PoseSearch/Source/Runtime/Private/PoseSearchFeatureChannel_Heading.cpp
@@ -80,19 +80,8 @@ void UPoseSearchFeatureChannel_Heading::BuildQuery(UE::PoseSearch::FSearchContex
 	}
 	else
 	{
-		FTransform Transform;
-		if (bBoneValid)
-		{
-			// calculating the Transform in component space for the bone indexed by SchemaBoneIdx
-			Transform = SearchContext.GetComponentSpaceTransform(SampleTimeOffset, 0.f, InOutQuery.GetSchema(), SchemaBoneIdx);
-		}
-		else
-		{
-			check(SearchContext.Trajectory);
-			// @todo: make this call consistent with Transform = SearchContext.TryGetTransformAndCacheResults(SampleTimeOffset, InOutQuery.GetSchema());
-			const FTrajectorySample TrajectorySample = SearchContext.Trajectory->GetSampleAtTime(SampleTimeOffset);
-			Transform = TrajectorySample.Transform;
-		}
+		// calculating the Transform in component space for the bone indexed by SchemaBoneIdx
+		const FTransform Transform = SearchContext.GetComponentSpaceTransform(SampleTimeOffset, 0.f, InOutQuery.GetSchema(), SchemaBoneIdx, bBoneValid);
 
 		FFeatureVectorHelper::EncodeVector(InOutQuery.EditValues(), ChannelDataOffset, GetAxis(Transform.GetRotation()), ComponentStripping);
 	}
diff --git a/Engine/Plugins/Experimental/Animation/PoseSearch/Source/Runtime/Private/PoseSearchFeatureChannel_Position.cpp b/Engine/Plugins/Experimental/Animation/PoseSearch/Source/Runtime/Private/PoseSearchFeatureChannel_Position.cpp
--- a/Engine/Plugins/Experimental/Animation/PoseSearch/Source/Runtime/Private/PoseSearchFeatureChannel_Position.cpp
+++ b/Engine/Plugins/Experimental/Animation/PoseSearch/Source/Runtime/Private/PoseSearchFeatureChannel_Position.cpp
@@ -64,19 +64,8 @@ void UPoseSearchFeatureChannel_Position::BuildQuery(UE::PoseSearch::FSearchConte
 	}
 	else
 	{
-		FTransform Transform;
-		if (bBoneValid)
-		{
-			// calculating the Transform in component space for the bone indexed by SchemaBoneIdx
-			Transform = SearchContext.GetComponentSpaceTransform(SampleTimeOffset, 0.f, InOutQuery.GetSchema(), SchemaBoneIdx);
-		}
-		else
-		{
-			check(SearchContext.Trajectory);
-			// @todo: make this call consistent with Transform = SearchContext.TryGetTransformAndCacheResults(SampleTimeOffset, InOutQuery.GetSchema());
-			const FTrajectorySample TrajectorySample = SearchContext.Trajectory->GetSampleAtTime(SampleTimeOffset);
-			Transform = TrajectorySample.Transform;
-		}
+		// calculating the Transform in component space for the bone indexed by SchemaBoneIdx
+		const FTransform Transform = SearchContext.GetComponentSpaceTransform(SampleTimeOffset, 0.f, InOutQuery.GetSchema(), SchemaBoneIdx, bBoneValid);
 
 		FFeatureVectorHelper::EncodeVector(InOutQuery.EditValues(), ChannelDataOffset, Transform.GetTranslation(), ComponentStripping);
 	}
diff --git a/Engine/Plugins/Experimental/Animation/PoseSearch/Source/Runtime/Private/PoseSearchFeatureChannel_Velocity.cpp b/Engine/Plugins/Experimental/Animation/PoseSearch/Source/Runtime/Private/PoseSearchFeatureChannel_Velocity.cpp
--- a/Engine/Plugins/Experimental/Animation/PoseSearch/Source/Runtime/Private/PoseSearchFeatureChannel_Velocity.cpp
+++ b/Engine/Plugins/Experimental/Animation/PoseSearch/Source/Runtime/Private/PoseSearchFeatureChannel_Velocity.cpp
@@ -93,28 +93,14 @@ void UPoseSearchFeatureChannel_Velocity::BuildQuery(UE::PoseSearch::FSearchConte
 	}
 	else
 	{
-		FVector LinearVelocity;
-		if (bBoneValid)
-		{
-			check(SearchContext.History);
-			const float HistorySameplInterval = SearchContext.History->GetSampleTimeInterval();
-			check(HistorySameplInterval > UE_KINDA_SMALL_NUMBER);
+		const float HistorySampleInterval = SearchContext.History ? SearchContext.History->GetSampleTimeInterval() : 1 / 60.0f;
+		check(HistorySampleInterval > UE_KINDA_SMALL_NUMBER);
 
-			// @todo: optimize this code for SchemaBoneIdx == FSearchContext::SchemaRootBoneIdx (!InOutQuery.GetSchema()->BoneReferences[SchemaBoneIdx].HasValidSetup())
-			// calculating the Transforms in component space for the bone indexed by SchemaBoneIdx
-			const FTransform TransformCurrent = SearchContext.GetComponentSpaceTransform(SampleTimeOffset, 0.f, InOutQuery.GetSchema(), SchemaBoneIdx);
-			const FTransform TransformPrevious = SearchContext.GetComponentSpaceTransform(SampleTimeOffset - HistorySameplInterval, -HistorySameplInterval, InOutQuery.GetSchema(), SchemaBoneIdx);
-
-			LinearVelocity = (TransformCurrent.GetTranslation() - TransformPrevious.GetTranslation()) / HistorySameplInterval;
-		}
-		else
-		{
-			check(SearchContext.Trajectory);
-			// @todo: make this call consistent with Transform = SearchContext.TryGetTransformAndCacheResults(SampleTimeOffset, InOutQuery.GetSchema());
-			const FTrajectorySample TrajectorySample = SearchContext.Trajectory->GetSampleAtTime(SampleTimeOffset);
-			LinearVelocity = TrajectorySample.LinearVelocity;
-		}
+		// calculating the Transforms in component space for the bone indexed by SchemaBoneIdx
+		const FTransform TransformCurrent = SearchContext.GetComponentSpaceTransform(SampleTimeOffset, 0.f, InOutQuery.GetSchema(), SchemaBoneIdx, bBoneValid);
+		const FTransform TransformPrevious = SearchContext.GetComponentSpaceTransform(SampleTimeOffset - HistorySampleInterval, bUseCharacterSpaceVelocities ? -HistorySampleInterval : 0.f, InOutQuery.GetSchema(), SchemaBoneIdx, bBoneValid);
 
+		FVector LinearVelocity = (TransformCurrent.GetTranslation() - TransformPrevious.GetTranslation()) / HistorySampleInterval;
 		if (bNormalize)
 		{
 			LinearVelocity = LinearVelocity.GetClampedToMaxSize(1.f);
diff --git a/Engine/Plugins/Experimental/Animation/PoseSearch/Source/Runtime/Public/PoseSearch/PoseSearchContext.h b/Engine/Plugins/Experimental/Animation/PoseSearch/Source/Runtime/Public/PoseSearch/PoseSearchContext.h
--- a/Engine/Plugins/Experimental/Animation/PoseSearch/Source/Runtime/Public/PoseSearch/PoseSearchContext.h
+++ b/Engine/Plugins/Experimental/Animation/PoseSearch/Source/Runtime/Public/PoseSearch/PoseSearchContext.h
@@ -162,9 +162,9 @@ struct POSESEARCH_API FSearchContext
 	// can the continuing pose advance? (if not we skip evaluating it)
 	bool bCanAdvance = true;
 
-	FTransform GetTransform(float SampleTime, const UPoseSearchSchema* Schema, int8 SchemaBoneIdx = RootSchemaBoneIdx);
+	FTransform GetTransform(float SampleTime, const UPoseSearchSchema* Schema, int8 SchemaBoneIdx = RootSchemaBoneIdx, bool bUseHistoryRoot = false);
 	FTransform GetComponentSpaceTransform(float SampleTime, const UPoseSearchSchema* Schema, int8 SchemaBoneIdx = RootSchemaBoneIdx);
-	FTransform GetComponentSpaceTransform(float SampleTime, float OriginTime, const UPoseSearchSchema* Schema, int8 SchemaBoneIdx = RootSchemaBoneIdx);
+	FTransform GetComponentSpaceTransform(float SampleTime, float OriginTime, const UPoseSearchSchema* Schema, int8 SchemaBoneIdx = RootSchemaBoneIdx, bool bUseHistoryRoot = false);
 
 	void ClearCachedEntries();
 
