motion matching - adding UAnimGraphNode_BlendStack, removed dependency from ESearchIndexAssetType in FAnimNode_BlendStack

diff --git a/Engine/Plugins/Experimental/Animation/PoseSearch/Source/Editor/Private/AnimGraphNode_BlendStack.cpp b/Engine/Plugins/Experimental/Animation/PoseSearch/Source/Editor/Private/AnimGraphNode_BlendStack.cpp
new file mode 100644
--- /dev/null
+++ b/Engine/Plugins/Experimental/Animation/PoseSearch/Source/Editor/Private/AnimGraphNode_BlendStack.cpp
@@ -0,0 +1,63 @@
+// Copyright Epic Games, Inc. All Rights Reserved.
+
+#include "AnimGraphNode_BlendStack.h"
+#include "Animation/AnimRootMotionProvider.h"
+
+#define LOCTEXT_NAMESPACE "AnimGraphNode_BlendStack"
+
+
+FLinearColor UAnimGraphNode_BlendStack::GetNodeTitleColor() const
+{
+	return FColor(86, 182, 194);
+}
+
+FText UAnimGraphNode_BlendStack::GetTooltipText() const
+{
+	return LOCTEXT("NodeToolTip", "Blend Stack");
+}
+
+FText UAnimGraphNode_BlendStack::GetNodeTitle(ENodeTitleType::Type TitleType) const
+{
+	return LOCTEXT("NodeTitle", "Blend Stack");
+}
+
+FText UAnimGraphNode_BlendStack::GetMenuCategory() const
+{
+	return LOCTEXT("NodeCategory", "Pose Search");
+}
+
+void UAnimGraphNode_BlendStack::GetOutputLinkAttributes(FNodeAttributeArray& OutAttributes) const
+{
+	if (UE::Anim::IAnimRootMotionProvider::Get())
+	{
+		OutAttributes.Add(UE::Anim::IAnimRootMotionProvider::AttributeName);
+	}
+}
+
+void UAnimGraphNode_BlendStack::BakeDataDuringCompilation(class FCompilerResultsLog& MessageLog)
+{
+	UAnimBlueprint* AnimBlueprint = GetAnimBlueprint();
+	AnimBlueprint->FindOrAddGroup(Node.GetGroupName());
+}
+
+bool UAnimGraphNode_BlendStack::DoesSupportTimeForTransitionGetter() const
+{
+	return true;
+}
+
+UAnimationAsset* UAnimGraphNode_BlendStack::GetAnimationAsset() const
+{
+	return nullptr;
+}
+
+const TCHAR* UAnimGraphNode_BlendStack::GetTimePropertyName() const
+{
+	return TEXT("InternalTimeAccumulator");
+}
+
+UScriptStruct* UAnimGraphNode_BlendStack::GetTimePropertyStruct() const
+{
+	return FAnimNode_BlendStack::StaticStruct();
+}
+
+#undef LOCTEXT_NAMESPACE
diff --git a/Engine/Plugins/Experimental/Animation/PoseSearch/Source/Editor/Public/AnimGraphNode_BlendStack.h b/Engine/Plugins/Experimental/Animation/PoseSearch/Source/Editor/Public/AnimGraphNode_BlendStack.h
new file mode 100644
--- /dev/null
+++ b/Engine/Plugins/Experimental/Animation/PoseSearch/Source/Editor/Public/AnimGraphNode_BlendStack.h
@@ -0,0 +1,31 @@
+// Copyright Epic Games, Inc. All Rights Reserved.
+
+#pragma once
+
+#include "CoreMinimal.h"
+#include "UObject/ObjectMacros.h"
+#include "AnimGraphNode_AssetPlayerBase.h"
+#include "PoseSearch/AnimNode_BlendStack.h"
+#include "AnimGraphNode_BlendStack.generated.h"
+
+
+UCLASS(MinimalAPI, Experimental)
+class UAnimGraphNode_BlendStack : public UAnimGraphNode_AssetPlayerBase
+{
+	GENERATED_BODY()
+
+	UPROPERTY(EditAnywhere, Category = Settings)
+	FAnimNode_BlendStack Node;
+
+	virtual FLinearColor GetNodeTitleColor() const override;
+	virtual FText GetTooltipText() const override;
+	virtual FText GetNodeTitle(ENodeTitleType::Type TitleType) const override;
+	virtual FText GetMenuCategory() const override;
+	virtual void GetOutputLinkAttributes(FNodeAttributeArray& OutAttributes) const override;
+
+	virtual void BakeDataDuringCompilation(class FCompilerResultsLog& MessageLog) override;
+	virtual bool DoesSupportTimeForTransitionGetter() const override;
+	virtual UAnimationAsset* GetAnimationAsset() const override;
+	virtual const TCHAR* GetTimePropertyName() const override;
+	virtual UScriptStruct* GetTimePropertyStruct() const override;
+};
diff --git a/Engine/Plugins/Experimental/Animation/PoseSearch/Source/Runtime/Private/AnimNode_BlendStack.cpp b/Engine/Plugins/Experimental/Animation/PoseSearch/Source/Runtime/Private/AnimNode_BlendStack.cpp
--- a/Engine/Plugins/Experimental/Animation/PoseSearch/Source/Runtime/Private/AnimNode_BlendStack.cpp
+++ b/Engine/Plugins/Experimental/Animation/PoseSearch/Source/Runtime/Private/AnimNode_BlendStack.cpp
@@ -19,7 +19,7 @@ TAutoConsoleVariable<int32> CVarAnimBlendStackPruningEnable(TEXT("a.AnimNode.Ble
 
 /////////////////////////////////////////////////////
 // FPoseSearchAnimPlayer
-void FPoseSearchAnimPlayer::Initialize(ESearchIndexAssetType InAssetType, UAnimationAsset* AnimationAsset, float AccumulatedTime, bool bLoop, bool bMirrored, UMirrorDataTable* MirrorDataTable, float BlendTime, const UBlendProfile* BlendProfile, EAlphaBlendOption InBlendOption, FVector BlendParameters, float PlayRate)
+void FPoseSearchAnimPlayer::Initialize(UAnimationAsset* AnimationAsset, float AccumulatedTime, bool bLoop, bool bMirrored, UMirrorDataTable* MirrorDataTable, float BlendTime, const UBlendProfile* BlendProfile, EAlphaBlendOption InBlendOption, FVector BlendParameters, float PlayRate)
 {
 	check(AnimationAsset);
 
@@ -46,27 +46,21 @@ void FPoseSearchAnimPlayer::Initialize(ESearchIndexAssetType InAssetType, UAnima
 	}
 	BlendOption = InBlendOption;
 
-	AssetType = InAssetType;
 	TotalBlendInTime = BlendTime;
 	CurrentBlendInTime = 0.f;
 
 	MirrorNode.SetMirrorDataTable(MirrorDataTable);
 	MirrorNode.SetMirror(bMirrored);
 	
-	if (AssetType == ESearchIndexAssetType::Sequence || AssetType == ESearchIndexAssetType::AnimComposite)
+	if (UAnimSequenceBase* Sequence = Cast<UAnimSequenceBase>(AnimationAsset))
 	{
-		UAnimSequenceBase* Sequence = Cast<UAnimSequenceBase>(AnimationAsset);
-		check(Sequence);
-
 		SequencePlayerNode.SetAccumulatedTime(AccumulatedTime);
 		SequencePlayerNode.SetSequence(Sequence);
 		SequencePlayerNode.SetLoopAnimation(bLoop);
 		SequencePlayerNode.SetPlayRate(PlayRate);
 	}
-	else if (AssetType == ESearchIndexAssetType::BlendSpace)
+	else if (UBlendSpace* BlendSpace = Cast<UBlendSpace>(AnimationAsset))
 	{
-		UBlendSpace* BlendSpace = Cast<UBlendSpace>(AnimationAsset);
-		check(BlendSpace);
 		BlendSpacePlayerNode.SetResetPlayTimeWhenBlendSpaceChanges(false /*!bReset*/);
 		BlendSpacePlayerNode.SetAccumulatedTime(AccumulatedTime);
 		BlendSpacePlayerNode.SetBlendSpace(BlendSpace);
@@ -84,11 +78,11 @@ void FPoseSearchAnimPlayer::Initialize(ESearchIndexAssetType InAssetType, UAnima
 
 void FPoseSearchAnimPlayer::UpdatePlayRate(float PlayRate)
 {
-	if (AssetType == ESearchIndexAssetType::Sequence || AssetType == ESearchIndexAssetType::AnimComposite)
+	if (SequencePlayerNode.GetSequence())
 	{
 		SequencePlayerNode.SetPlayRate(PlayRate);
 	}
-	else if (AssetType == ESearchIndexAssetType::BlendSpace)
+	else if (BlendSpacePlayerNode.GetBlendSpace())
 	{
 		BlendSpacePlayerNode.SetPlayRate(PlayRate);
 	}
@@ -100,7 +94,6 @@ void FPoseSearchAnimPlayer::UpdatePlayRate(float PlayRate)
 
 void FPoseSearchAnimPlayer::StorePoseContext(const FPoseContext& PoseContext)
 {
-	AssetType = ESearchIndexAssetType::Invalid;
 	UpdateSourceLinkNode();
 
 	const FBoneContainer& BoneContainer = PoseContext.Pose.GetBoneContainer();
@@ -122,7 +115,7 @@ void FPoseSearchAnimPlayer::StorePoseContext(const FPoseContext& PoseContext)
 
 void FPoseSearchAnimPlayer::RestorePoseContext(FPoseContext& PoseContext) const
 {
-	check(AssetType == ESearchIndexAssetType::Invalid);
+	check(!SequencePlayerNode.GetSequence() && !BlendSpacePlayerNode.GetBlendSpace());
 
 	const FBoneContainer& BoneContainer = PoseContext.Pose.GetBoneContainer();
 	const USkeleton* SkeletonAsset = BoneContainer.GetSkeletonAsset();
@@ -149,11 +142,11 @@ void FPoseSearchAnimPlayer::RestorePoseContext(FPoseContext& PoseContext) const
 // since we're making copies and moving this object in memory, we're using this method to set the MirrorNode SourceLinkNode when necessary
 void FPoseSearchAnimPlayer::UpdateSourceLinkNode()
 {
-	if (AssetType == ESearchIndexAssetType::Sequence || AssetType == ESearchIndexAssetType::AnimComposite)
+	if (SequencePlayerNode.GetSequence())
 	{
 		MirrorNode.SetSourceLinkNode(&SequencePlayerNode);
 	}
-	else if (AssetType == ESearchIndexAssetType::BlendSpace)
+	else if (BlendSpacePlayerNode.GetBlendSpace())
 	{
 		MirrorNode.SetSourceLinkNode(&BlendSpacePlayerNode);
 	}
@@ -165,7 +158,7 @@ void FPoseSearchAnimPlayer::UpdateSourceLinkNode()
 
 void FPoseSearchAnimPlayer::Evaluate_AnyThread(FPoseContext& Output)
 {
-	if (AssetType != ESearchIndexAssetType::Invalid)
+	if (SequencePlayerNode.GetSequence() || BlendSpacePlayerNode.GetBlendSpace())
 	{
 		UpdateSourceLinkNode();
 		MirrorNode.Evaluate_AnyThread(Output);
@@ -187,12 +180,12 @@ void FPoseSearchAnimPlayer::Update_AnyThread(const FAnimationUpdateContext& Cont
 
 float FPoseSearchAnimPlayer::GetAccumulatedTime() const
 {
-	if (AssetType == ESearchIndexAssetType::Sequence || AssetType == ESearchIndexAssetType::AnimComposite)
+	if (SequencePlayerNode.GetSequence())
 	{
 		return SequencePlayerNode.GetAccumulatedTime();
 	}
 	
-	if (AssetType == ESearchIndexAssetType::BlendSpace)
+	if (BlendSpacePlayerNode.GetBlendSpace())
 	{
 		return BlendSpacePlayerNode.GetAccumulatedTime();
 	}
@@ -201,15 +194,25 @@ float FPoseSearchAnimPlayer::GetAccumulatedTime() const
 	return 0.f;
 }
 
-FString FPoseSearchAnimPlayer::GetAnimName() const
+FVector FPoseSearchAnimPlayer::GetBlendParameters() const
+{
+	if (BlendSpacePlayerNode.GetBlendSpace())
+	{
+		return BlendSpacePlayerNode.GetPosition();
+	}
+
+	return FVector::ZeroVector;
+}
+
+FString FPoseSearchAnimPlayer::GetAnimationName() const
 {
-	if (AssetType == ESearchIndexAssetType::Sequence || AssetType == ESearchIndexAssetType::AnimComposite)
+	if (SequencePlayerNode.GetSequence())
 	{
 		check(SequencePlayerNode.GetSequence());
 		return SequencePlayerNode.GetSequence()->GetName();
 	}
 
-	if (AssetType == ESearchIndexAssetType::BlendSpace)
+	if (BlendSpacePlayerNode.GetBlendSpace())
 	{
 		check(BlendSpacePlayerNode.GetBlendSpace());
 		return BlendSpacePlayerNode.GetBlendSpace()->GetName();
@@ -219,6 +222,22 @@ FString FPoseSearchAnimPlayer::GetAnimName() const
 	return FString("StoredPose");
 }
 
+const UAnimationAsset* FPoseSearchAnimPlayer::GetAnimationAsset() const
+{
+	if (SequencePlayerNode.GetSequence())
+	{
+		return SequencePlayerNode.GetSequence();
+	}
+
+	if (BlendSpacePlayerNode.GetBlendSpace())
+	{
+		return BlendSpacePlayerNode.GetBlendSpace();
+	}
+
+	checkNoEntry();
+	return nullptr;
+}
+
 float FPoseSearchAnimPlayer::GetBlendInPercentage() const
 {
 	if (FMath::IsNearlyZero(TotalBlendInTime))
@@ -254,8 +273,8 @@ bool FPoseSearchAnimPlayer::GetBlendInWeights(TArray<float>& Weights) const
 }
 
 /////////////////////////////////////////////////////
-// FAnimNode_BlendStack
-void FAnimNode_BlendStack::Evaluate_AnyThread(FPoseContext& Output)
+// FAnimNode_BlendStack_Standalone
+void FAnimNode_BlendStack_Standalone::Evaluate_AnyThread(FPoseContext& Output)
 {
 	Super::Evaluate_AnyThread(Output);
 
@@ -338,12 +357,12 @@ void FAnimNode_BlendStack::Evaluate_AnyThread(FPoseContext& Output)
 		const int32 ActiveBlends = AnimPlayers.Num() - 1;
 		if (ActiveBlends > RequestedMaxActiveBlends)
 		{
-			UE_LOG(LogPoseSearch, Display, TEXT("FAnimNode_BlendStack NumBlends/MaxNumBlends %d / %d"), ActiveBlends, RequestedMaxActiveBlends);
+			UE_LOG(LogPoseSearch, Display, TEXT("FAnimNode_BlendStack_Standalone NumBlends/MaxNumBlends %d / %d"), ActiveBlends, RequestedMaxActiveBlends);
 		}
 	}
 }
 
-void FAnimNode_BlendStack::UpdateAssetPlayer(const FAnimationUpdateContext& Context)
+void FAnimNode_BlendStack_Standalone::UpdateAssetPlayer(const FAnimationUpdateContext& Context)
 {
 	Super::UpdateAssetPlayer(Context);
 
@@ -360,7 +379,7 @@ void FAnimNode_BlendStack::UpdateAssetPlayer(const FAnimationUpdateContext& Cont
 		const float BlendInPercentage = bIsLastAnimPlayers ? 1.f : AnimPlayer.GetBlendInPercentage();
 		const float AnimPlayerBlendWeight = CurrentWeightMultiplier * BlendInPercentage;
 
-		// don't break for AnimPlayerIndex == 0 since FAnimNode_BlendStack::BlendTo initialize the AnimPlayer with a weight of zero
+		// don't break for AnimPlayerIndex == 0 since FAnimNode_BlendStack_Standalone::BlendTo initialize the AnimPlayer with a weight of zero
 		if (AnimPlayerIndex > 0 && AnimPlayerBlendWeight < UE_KINDA_SMALL_NUMBER)
 		{
 			break;
@@ -379,19 +398,19 @@ void FAnimNode_BlendStack::UpdateAssetPlayer(const FAnimationUpdateContext& Cont
 	}
 }
 
-float FAnimNode_BlendStack::GetAccumulatedTime() const
+float FAnimNode_BlendStack_Standalone::GetAccumulatedTime() const
 {
 	return AnimPlayers.IsEmpty() ? 0.f : AnimPlayers.First().GetAccumulatedTime();
 }
 
-void FAnimNode_BlendStack::BlendTo(ESearchIndexAssetType AssetType, UAnimationAsset* AnimationAsset, float AccumulatedTime, bool bLoop, bool bMirrored, UMirrorDataTable* MirrorDataTable, int32 MaxActiveBlends, float BlendTime, const UBlendProfile* BlendProfile, EAlphaBlendOption BlendOption, FVector BlendParameters, float PlayRate)
+void FAnimNode_BlendStack_Standalone::BlendTo(UAnimationAsset* AnimationAsset, float AccumulatedTime, bool bLoop, bool bMirrored, UMirrorDataTable* MirrorDataTable, int32 MaxActiveBlends, float BlendTime, const UBlendProfile* BlendProfile, EAlphaBlendOption BlendOption, FVector BlendParameters, float PlayRate)
 {
 	RequestedMaxActiveBlends = MaxActiveBlends;
 	AnimPlayers.PushFirst(FPoseSearchAnimPlayer());
-	AnimPlayers.First().Initialize(AssetType, AnimationAsset, AccumulatedTime, bLoop, bMirrored, MirrorDataTable, BlendTime, BlendProfile, BlendOption, BlendParameters, PlayRate);
+	AnimPlayers.First().Initialize(AnimationAsset, AccumulatedTime, bLoop, bMirrored, MirrorDataTable, BlendTime, BlendProfile, BlendOption, BlendParameters, PlayRate);
 }
 
-void FAnimNode_BlendStack::UpdatePlayRate(float PlayRate)
+void FAnimNode_BlendStack_Standalone::UpdatePlayRate(float PlayRate)
 {
 	if (!AnimPlayers.IsEmpty())
 	{
@@ -399,7 +418,7 @@ void FAnimNode_BlendStack::UpdatePlayRate(float PlayRate)
 	}
 }
 
-void FAnimNode_BlendStack::GatherDebugData(FNodeDebugData& DebugData)
+void FAnimNode_BlendStack_Standalone::GatherDebugData(FNodeDebugData& DebugData)
 {
 #if ENABLE_ANIM_DEBUG
 	DebugData.AddDebugItem(FString::Printf(TEXT("%s"), *DebugData.GetNodeName(this)));
@@ -408,7 +427,7 @@ void FAnimNode_BlendStack::GatherDebugData(FNodeDebugData& DebugData)
 		const FPoseSearchAnimPlayer& AnimPlayer = AnimPlayers[i];
 		DebugData.AddDebugItem(FString::Printf(TEXT("%d) t:%.2f/%.2f m:%d %s"),
 				i, AnimPlayer.GetCurrentBlendInTime(), AnimPlayer.GetTotalBlendInTime(),
-				AnimPlayer.GetMirror() ? 1 : 0, *AnimPlayer.GetAnimName()));
+				AnimPlayer.GetMirror() ? 1 : 0, *AnimPlayer.GetAnimationName()));
 	}
 #endif // ENABLE_ANIM_DEBUG
 
@@ -419,4 +438,52 @@ void FAnimNode_BlendStack::GatherDebugData(FNodeDebugData& DebugData)
 	}
 }
 
+/////////////////////////////////////////////////////
+// FAnimNode_BlendStack
+
+void FAnimNode_BlendStack::UpdateAssetPlayer(const FAnimationUpdateContext& Context)
+{
+	GetEvaluateGraphExposedInputs().Execute(Context);
+
+	if (AnimationAsset)
+	{
+		bool bExecuteBlendTo = false;
+		if (AnimPlayers.IsEmpty())
+		{
+			bExecuteBlendTo = true;
+		}
+		else
+		{
+			const FPoseSearchAnimPlayer& MainAnimPlayer = AnimPlayers.First();
+			const UAnimationAsset* PlayingAnimationAsset = MainAnimPlayer.GetAnimationAsset();
+			check(PlayingAnimationAsset);
+
+			if (AnimationAsset != PlayingAnimationAsset)
+			{
+				bExecuteBlendTo = true;
+			}
+			else if (bMirrored != MainAnimPlayer.GetMirror())
+			{
+				bExecuteBlendTo = true;
+			}
+			else if (BlendParameters != MainAnimPlayer.GetBlendParameters())
+			{
+				bExecuteBlendTo = true;
+			}
+			else if (MaxAnimationDeltaTime >= 0.f && FMath::Abs(AnimationTime - MainAnimPlayer.GetAccumulatedTime()) > MaxAnimationDeltaTime)
+			{
+				bExecuteBlendTo = true;
+			}
+		}
+
+		if (bExecuteBlendTo)
+		{
+			BlendTo(AnimationAsset, AnimationTime, bLoop, bMirrored, MirrorDataTable.Get(), MaxActiveBlends, BlendTime, BlendProfile, BlendOption, BlendParameters, WantedPlayRate);
+		}
+	}
+	
+	UpdatePlayRate(WantedPlayRate);
+
+	Super::UpdateAssetPlayer(Context);
+}
 #undef LOCTEXT_NAMESPACE
diff --git a/Engine/Plugins/Experimental/Animation/PoseSearch/Source/Runtime/Private/AnimNode_MotionMatching.cpp b/Engine/Plugins/Experimental/Animation/PoseSearch/Source/Runtime/Private/AnimNode_MotionMatching.cpp
--- a/Engine/Plugins/Experimental/Animation/PoseSearch/Source/Runtime/Private/AnimNode_MotionMatching.cpp
+++ b/Engine/Plugins/Experimental/Animation/PoseSearch/Source/Runtime/Private/AnimNode_MotionMatching.cpp
@@ -128,7 +128,7 @@ void FAnimNode_MotionMatching::UpdateAssetPlayer(const FAnimationUpdateContext&
 		{
 			if (const FPoseSearchDatabaseAnimationAssetBase* DatabaseAsset = Database->GetAnimationAssetBase(*SearchIndexAsset))
 			{
-				BlendStackNode.BlendTo(SearchIndexAsset->Type, DatabaseAsset->GetAnimationAsset(), MotionMatchingState.CurrentSearchResult.AssetTime,
+				BlendStackNode.BlendTo(DatabaseAsset->GetAnimationAsset(), MotionMatchingState.CurrentSearchResult.AssetTime,
 					DatabaseAsset->IsLooping(), SearchIndexAsset->bMirrored, Database->Schema->MirrorDataTable.Get(),
 					Settings.MaxActiveBlends, Settings.BlendTime, Settings.BlendProfile, Settings.BlendOption, SearchIndexAsset->BlendParameters, MotionMatchingState.WantedPlayRate);
 			}
diff --git a/Engine/Plugins/Experimental/Animation/PoseSearch/Source/Runtime/Public/PoseSearch/AnimNode_BlendStack.h b/Engine/Plugins/Experimental/Animation/PoseSearch/Source/Runtime/Public/PoseSearch/AnimNode_BlendStack.h
--- a/Engine/Plugins/Experimental/Animation/PoseSearch/Source/Runtime/Public/PoseSearch/AnimNode_BlendStack.h
+++ b/Engine/Plugins/Experimental/Animation/PoseSearch/Source/Runtime/Public/PoseSearch/AnimNode_BlendStack.h
@@ -15,7 +15,7 @@ struct FPoseSearchAnimPlayer
 {
 	GENERATED_BODY()
 	
-	void Initialize(ESearchIndexAssetType InAssetType, UAnimationAsset* AnimationAsset, float AccumulatedTime, bool bLoop, bool bMirrored, UMirrorDataTable* MirrorDataTable, float BlendTime, const UBlendProfile* BlendProfile, EAlphaBlendOption InBlendOption, FVector BlendParameters, float PlayRate);
+	void Initialize(UAnimationAsset* AnimationAsset, float AccumulatedTime, bool bLoop, bool bMirrored, UMirrorDataTable* MirrorDataTable, float BlendTime, const UBlendProfile* BlendProfile, EAlphaBlendOption InBlendOption, FVector BlendParameters, float PlayRate);
 	void UpdatePlayRate(float PlayRate);
 	void Evaluate_AnyThread(FPoseContext& Output);
 	void Update_AnyThread(const FAnimationUpdateContext& Context, float BlendWeight);
@@ -23,14 +23,16 @@ struct FPoseSearchAnimPlayer
 	
 	float GetBlendInPercentage() const;
 	bool GetBlendInWeights(TArray<float>& Weights) const;
-	ESearchIndexAssetType GetAssetType() const { return AssetType; }
 	EAlphaBlendOption GetBlendOption() const { return BlendOption; }
 	void StorePoseContext(const FPoseContext& PoseContext);
 
 	float GetTotalBlendInTime() const { return TotalBlendInTime; }
 	float GetCurrentBlendInTime() const { return CurrentBlendInTime; }
 	bool GetMirror() const { return MirrorNode.GetMirror(); }
-	FString GetAnimName() const;
+	FVector GetBlendParameters() const;
+	FString GetAnimationName() const;
+	const UAnimationAsset* GetAnimationAsset() const;
+
 	FAnimNode_Mirror_Standalone& GetMirrorNode() { return MirrorNode; }
 
 protected:
@@ -47,10 +49,9 @@ protected:
 	FAnimNode_Mirror_Standalone MirrorNode;
 
 	// @todo: use FCompactHeapPose
-	// if AssetType == ESearchIndexAssetType::Invalid instead of using SequencePlayerNode or BlendSpacePlayerNode (wrapped in MirrorNode), the pose output will be from StoredPose
+	// if SequencePlayerNode.GetSequence() and BlendSpacePlayerNode.GetBlendSpace() are nullptr, instead of using SequencePlayerNode or BlendSpacePlayerNode (wrapped in MirrorNode), the pose output will be from StoredPose
 	TArray<FTransform> StoredPose;
 
-	ESearchIndexAssetType AssetType = ESearchIndexAssetType::Sequence;
 	EAlphaBlendOption BlendOption = EAlphaBlendOption::Linear;
 
 	TCustomBoneIndexArray<float, FSkeletonPoseBoneIndex> TotalBlendInTimePerBone;
@@ -60,7 +61,7 @@ protected:
 };
 
 USTRUCT(BlueprintInternalUseOnly)
-struct POSESEARCH_API FAnimNode_BlendStack : public FAnimNode_AssetPlayerBase
+struct POSESEARCH_API FAnimNode_BlendStack_Standalone : public FAnimNode_AssetPlayerBase
 {
 	GENERATED_BODY()
 
@@ -72,7 +73,7 @@ struct POSESEARCH_API FAnimNode_BlendStack : public FAnimNode_AssetPlayerBase
 	virtual void GatherDebugData(FNodeDebugData& DebugData) override;
 	// End of FAnimNode_Base interface
 
-	void BlendTo(ESearchIndexAssetType AssetType, UAnimationAsset* AnimationAsset, float AccumulatedTime = 0.f, bool bLoop = false, bool bMirrored = false, UMirrorDataTable* MirrorDataTable = nullptr, int32 MaxActiveBlends = 3, float BlendTime = 0.3f, const UBlendProfile* BlendProfile = nullptr, EAlphaBlendOption BlendOption = EAlphaBlendOption::Linear, FVector BlendParameters = FVector::Zero(), float PlayRate = 1.f);
+	void BlendTo(UAnimationAsset* AnimationAsset, float AccumulatedTime = 0.f, bool bLoop = false, bool bMirrored = false, UMirrorDataTable* MirrorDataTable = nullptr, int32 MaxActiveBlends = 3, float BlendTime = 0.3f, const UBlendProfile* BlendProfile = nullptr, EAlphaBlendOption BlendOption = EAlphaBlendOption::Linear, FVector BlendParameters = FVector::Zero(), float PlayRate = 1.f);
 	void UpdatePlayRate(float PlayRate);
 
 	// FAnimNode_AssetPlayerBase interface
@@ -80,3 +81,62 @@ struct POSESEARCH_API FAnimNode_BlendStack : public FAnimNode_AssetPlayerBase
 	virtual void UpdateAssetPlayer(const FAnimationUpdateContext& Context) override;
 	// End of FAnimNode_AssetPlayerBase interface
 };
+
+
+USTRUCT(BlueprintInternalUseOnly)
+struct POSESEARCH_API FAnimNode_BlendStack : public FAnimNode_BlendStack_Standalone
+{
+	GENERATED_BODY()
+
+	// requested animation to play
+	UPROPERTY(EditAnywhere, BlueprintReadWrite, Category = Settings, meta = (PinShownByDefault))
+	TObjectPtr<UAnimationAsset> AnimationAsset;
+
+	// requested animation time. If negative, the animation will play from the beginning uninterrupted
+	UPROPERTY(EditAnywhere, BlueprintReadWrite, Category = Settings, meta = (PinHiddenByDefault))
+	float AnimationTime = -1.f;
+
+	// requested AnimationAsset looping
+	UPROPERTY(EditAnywhere, BlueprintReadWrite, Category = Settings, meta = (PinHiddenByDefault))
+	bool bLoop = true;
+
+	// requested AnimationAsset mirroring
+	UPROPERTY(EditAnywhere, BlueprintReadWrite, Category = Settings, meta = (PinHiddenByDefault))
+	bool bMirrored = false;
+
+	// requested animation play rate
+	UPROPERTY(EditAnywhere, BlueprintReadWrite, Category = Settings, meta = (PinHiddenByDefault))
+	float WantedPlayRate = 1.f;
+	
+	// tunable animation transition blend time 
+	UPROPERTY(EditAnywhere, BlueprintReadWrite, Category = Settings, meta = (PinHiddenByDefault, ClampMin = "0"))
+	float BlendTime = 0.2f;
+
+	// if AnimationTime and MaxAnimationDeltaTime are positive and the currently playing animation total time differs more than MaxAnimationDeltaTime from AnimationTime
+	// (animation desynchronized from the requested time) the blend stack will force a blend into the same animation
+	UPROPERTY(EditAnywhere, BlueprintReadWrite, Category = Settings, meta = (PinHiddenByDefault))
+	float MaxAnimationDeltaTime = -1.f;
+
+	// Number of max active blending animation in the blend stack. If MaxActiveBlends is zero then blend stack is disabled
+	UPROPERTY(EditAnywhere, BlueprintReadWrite, Category = Settings, meta = (PinHiddenByDefault, ClampMin = "0"))
+	int32 MaxActiveBlends = 4;
+
+	UPROPERTY(EditAnywhere, BlueprintReadWrite, Category = Settings, meta = (PinHiddenByDefault, UseAsBlendProfile = true))
+	TObjectPtr<UBlendProfile> BlendProfile;
+
+	UPROPERTY(EditAnywhere, BlueprintReadWrite, Category = Settings, meta = (PinHiddenByDefault))
+	EAlphaBlendOption BlendOption = EAlphaBlendOption::Linear;
+
+	// requested blend space blend parameters (if AnimationAsset is a blend space)
+	UPROPERTY(EditAnywhere, BlueprintReadWrite, Category = Settings, meta = (PinHiddenByDefault))
+	FVector BlendParameters = FVector::Zero();
+
+	// if set and bMirrored MirrorDataTable will be used for mirroring the aniamtion
+	UPROPERTY(EditAnywhere, Category = Settings, meta = (PinHiddenByDefault))
+	TObjectPtr<UMirrorDataTable> MirrorDataTable;
+
+	// FAnimNode_Base interface
+	virtual void UpdateAssetPlayer(const FAnimationUpdateContext& Context) override;
+	// End of FAnimNode_Base interface
+};
+
diff --git a/Engine/Plugins/Experimental/Animation/PoseSearch/Source/Runtime/Public/PoseSearch/AnimNode_MotionMatching.h b/Engine/Plugins/Experimental/Animation/PoseSearch/Source/Runtime/Public/PoseSearch/AnimNode_MotionMatching.h
--- a/Engine/Plugins/Experimental/Animation/PoseSearch/Source/Runtime/Public/PoseSearch/AnimNode_MotionMatching.h
+++ b/Engine/Plugins/Experimental/Animation/PoseSearch/Source/Runtime/Public/PoseSearch/AnimNode_MotionMatching.h
@@ -65,7 +65,7 @@ public:
 
 private:
 
-	FAnimNode_BlendStack BlendStackNode;
+	FAnimNode_BlendStack_Standalone BlendStackNode;
 
 	// Encapsulated motion matching algorithm and internal state
 	FMotionMatchingState MotionMatchingState;
