motion matching - collecting FAnimNode_PoseSearchHistoryCollector_Base::CollectedBones historical bone data in component space

diff --git a/Engine/Plugins/Experimental/Animation/PoseSearch/Source/Editor/Private/AnimGraphNode_PoseSearchHistoryCollector.cpp b/Engine/Plugins/Experimental/Animation/PoseSearch/Source/Editor/Private/AnimGraphNode_PoseSearchHistoryCollector.cpp
--- a/Engine/Plugins/Experimental/Animation/PoseSearch/Source/Editor/Private/AnimGraphNode_PoseSearchHistoryCollector.cpp
+++ b/Engine/Plugins/Experimental/Animation/PoseSearch/Source/Editor/Private/AnimGraphNode_PoseSearchHistoryCollector.cpp
@@ -1,27 +1,47 @@
 // Copyright Epic Games, Inc. All Rights Reserved.
 
 #include "AnimGraphNode_PoseSearchHistoryCollector.h"
+#include "AnimationGraphSchema.h"
 
 #define LOCTEXT_NAMESPACE "AnimGraphNode_PoseSearchHistoryCollector"
 
-FLinearColor UAnimGraphNode_PoseSearchHistoryCollector::GetNodeTitleColor() const
+/////////////////////////////////////////////////////
+// UAnimGraphNode_PoseSearchHistoryCollector_Base
+
+FLinearColor UAnimGraphNode_PoseSearchHistoryCollector_Base::GetNodeTitleColor() const
 {
 	return FColor(86, 182, 194);
 }
 
-FText UAnimGraphNode_PoseSearchHistoryCollector::GetTooltipText() const
+FText UAnimGraphNode_PoseSearchHistoryCollector_Base::GetTooltipText() const
+{
+	return LOCTEXT("NodeToolTip", "Collects bones transforms for motion matching");
+}
+
+FText UAnimGraphNode_PoseSearchHistoryCollector_Base::GetMenuCategory() const
 {
-	return LOCTEXT("NodeToolTip", "Use this pose for matching");
+	return LOCTEXT("NodeCategory", "Pose Search");
 }
 
+/////////////////////////////////////////////////////
+// UAnimGraphNode_PoseSearchHistoryCollector
+
 FText UAnimGraphNode_PoseSearchHistoryCollector::GetNodeTitle(ENodeTitleType::Type TitleType) const
 {
 	return LOCTEXT("NodeTitle", "Pose History");
 }
 
-FText UAnimGraphNode_PoseSearchHistoryCollector::GetMenuCategory() const
+/////////////////////////////////////////////////////
+// UAnimGraphNode_PoseSearchComponentSpaceHistoryCollector
+
+FText UAnimGraphNode_PoseSearchComponentSpaceHistoryCollector::GetNodeTitle(ENodeTitleType::Type TitleType) const
 {
-	return LOCTEXT("NodeCategory", "Pose Search");
+	return LOCTEXT("NodeTitleComponentSpace", "Component Space Pose History");
+}
+
+void UAnimGraphNode_PoseSearchComponentSpaceHistoryCollector::CreateOutputPins()
+{
+	CreatePin(EGPD_Output, UAnimationGraphSchema::PC_Struct, FComponentSpacePoseLink::StaticStruct(), TEXT("Pose"));
 }
 
 #undef LOCTEXT_NAMESPACE
diff --git a/Engine/Plugins/Experimental/Animation/PoseSearch/Source/Editor/Public/AnimGraphNode_PoseSearchHistoryCollector.h b/Engine/Plugins/Experimental/Animation/PoseSearch/Source/Editor/Public/AnimGraphNode_PoseSearchHistoryCollector.h
--- a/Engine/Plugins/Experimental/Animation/PoseSearch/Source/Editor/Public/AnimGraphNode_PoseSearchHistoryCollector.h
+++ b/Engine/Plugins/Experimental/Animation/PoseSearch/Source/Editor/Public/AnimGraphNode_PoseSearchHistoryCollector.h
@@ -2,23 +2,41 @@
 
 #pragma once
 
-#include "CoreMinimal.h"
-#include "UObject/ObjectMacros.h"
 #include "AnimGraphNode_Base.h"
+#include "CoreMinimal.h"
 #include "PoseSearch/AnimNode_PoseSearchHistoryCollector.h"
+#include "UObject/ObjectMacros.h"
 #include "AnimGraphNode_PoseSearchHistoryCollector.generated.h"
 
+UCLASS(Abstract, BlueprintType, Experimental)
+class UAnimGraphNode_PoseSearchHistoryCollector_Base : public UAnimGraphNode_Base
+{
+	GENERATED_BODY()
+	virtual FLinearColor GetNodeTitleColor() const override;
+	virtual FText GetTooltipText() const override;
+	virtual FText GetMenuCategory() const override;
+};
 
 UCLASS(MinimalAPI, Experimental)
-class UAnimGraphNode_PoseSearchHistoryCollector : public UAnimGraphNode_Base
+class UAnimGraphNode_PoseSearchHistoryCollector : public UAnimGraphNode_PoseSearchHistoryCollector_Base
 {
 	GENERATED_BODY()
 
 	UPROPERTY(EditAnywhere, Category = Settings)
 	FAnimNode_PoseSearchHistoryCollector Node;
 
-	virtual FLinearColor GetNodeTitleColor() const override;
-	virtual FText GetTooltipText() const override;
 	virtual FText GetNodeTitle(ENodeTitleType::Type TitleType) const override;
-	virtual FText GetMenuCategory() const override;
 };
+
+UCLASS(MinimalAPI, Experimental)
+class UAnimGraphNode_PoseSearchComponentSpaceHistoryCollector : public UAnimGraphNode_PoseSearchHistoryCollector_Base
+{
+	GENERATED_BODY()
+
+	UPROPERTY(EditAnywhere, Category = Settings)
+	FAnimNode_PoseSearchComponentSpaceHistoryCollector Node;
+	
+	virtual FText GetNodeTitle(ENodeTitleType::Type TitleType) const override;
+	virtual void CreateOutputPins() override;
+
+};
\ No newline at end of file
diff --git a/Engine/Plugins/Experimental/Animation/PoseSearch/Source/Runtime/Private/AnimNode_PoseSearchHistoryCollector.cpp b/Engine/Plugins/Experimental/Animation/PoseSearch/Source/Runtime/Private/AnimNode_PoseSearchHistoryCollector.cpp
--- a/Engine/Plugins/Experimental/Animation/PoseSearch/Source/Runtime/Private/AnimNode_PoseSearchHistoryCollector.cpp
+++ b/Engine/Plugins/Experimental/Animation/PoseSearch/Source/Runtime/Private/AnimNode_PoseSearchHistoryCollector.cpp
@@ -19,26 +19,96 @@ namespace UE::PoseSearch::Private
 class FPoseHistoryProvider : public IPoseHistoryProvider
 {
 public:
-	FPoseHistoryProvider(FAnimNode_PoseSearchHistoryCollector* InNode)
-		: Node(*InNode)
-	{}
+	FPoseHistoryProvider(const FPoseHistory& InPoseHistory)
+		: PoseHistory(InPoseHistory)
+	{
+	}
 
 	// IPoseHistoryProvider interface
 	virtual const FPoseHistory& GetPoseHistory() const override
 	{
-		return Node.GetPoseHistory();
+		return PoseHistory;
 	}
 
-	virtual FPoseHistory& GetPoseHistory() override
+private:
+	const FPoseHistory& PoseHistory;
+};
+
+} // namespace UE::PoseSearch::Private
+
+
+/////////////////////////////////////////////////////
+// FAnimNode_PoseSearchHistoryCollector_Base
+
+void FAnimNode_PoseSearchHistoryCollector_Base::CacheBones_AnyThread(const FAnimationCacheBonesContext& Context)
+{
+	DECLARE_SCOPE_HIERARCHICAL_COUNTER_ANIMNODE(Update_AnyThread);
+	check(Context.AnimInstanceProxy);
+
+	Super::CacheBones_AnyThread(Context);
+
+	TArray<FBoneIndexType> RequiredBones;
+	if (!CollectedBones.IsEmpty())
 	{
-		return Node.GetPoseHistory();
+		if (const USkeletalMeshComponent* SkeletalMeshComponent = Context.AnimInstanceProxy->GetSkelMeshComponent())
+		{
+			if (const USkinnedAsset* SkinnedAsset = SkeletalMeshComponent->GetSkinnedAsset())
+			{
+				if (const USkeleton* Skeleton = SkinnedAsset->GetSkeleton())
+				{
+					RequiredBones.Reserve(CollectedBones.Num());
+					for (FBoneReference& BoneReference : CollectedBones)
+					{
+						if (BoneReference.Initialize(Skeleton))
+						{
+							RequiredBones.AddUnique(BoneReference.BoneIndex);
+						}
+					}
+				}
+			}
+		}
 	}
 
-	// Node we wrap
-	FAnimNode_PoseSearchHistoryCollector& Node;
-};
+	PoseHistory.Init(PoseCount, PoseDuration, RequiredBones);
+}
 
-} // namespace UE::PoseSearch::Private
+void FAnimNode_PoseSearchHistoryCollector_Base::Update_AnyThread(const FAnimationUpdateContext& Context)
+{
+	DECLARE_SCOPE_HIERARCHICAL_COUNTER_ANIMNODE(Update_AnyThread);
+	UE::Anim::TScopedGraphMessage<UE::PoseSearch::Private::FPoseHistoryProvider> ScopedMessage(Context, PoseHistory);
+	Super::Update_AnyThread(Context);
+}
+
+bool FAnimNode_PoseSearchHistoryCollector_Base::HasPreUpdate() const
+{
+	return WITH_EDITORONLY_DATA && ENABLE_ANIM_DEBUG;
+}
+
+void FAnimNode_PoseSearchHistoryCollector_Base::PreUpdate(const UAnimInstance* InAnimInstance)
+{
+	Super::PreUpdate(InAnimInstance);
+
+#if WITH_EDITORONLY_DATA && ENABLE_ANIM_DEBUG
+	check(InAnimInstance);
+	if (bWasEvaluated && CVarAnimPoseHistoryDebugDraw.GetValueOnAnyThread())
+	{
+		if (const USkeletalMeshComponent* SkeletalMeshComponent = InAnimInstance->GetSkelMeshComponent())
+		{
+			const USkinnedAsset* SkinnedAsset = SkeletalMeshComponent->GetSkinnedAsset();
+			const UWorld* World = SkeletalMeshComponent->GetWorld();
+			if (SkinnedAsset && World)
+			{
+				if (const USkeleton* Skeleton = SkinnedAsset->GetSkeleton())
+				{
+					PoseHistory.DebugDraw(World, Skeleton);
+				}
+			}
+		}
+	}
+
+	bWasEvaluated = false;
+#endif // WITH_EDITORONLY_DATA && ENABLE_ANIM_DEBUG
+}
 
 
 /////////////////////////////////////////////////////
@@ -47,30 +117,28 @@ public:
 void FAnimNode_PoseSearchHistoryCollector::Initialize_AnyThread(const FAnimationInitializeContext& Context)
 {
 	DECLARE_SCOPE_HIERARCHICAL_COUNTER_ANIMNODE(Initialize_AnyThread);
-
 	Super::Initialize_AnyThread(Context);
-
-	// @@ need to do this once based on descendant node's (or input param?) search schema, not every node init
-	PoseHistory.Init(PoseCount, PoseDuration);
-
-	UE::Anim::TScopedGraphMessage<UE::PoseSearch::Private::FPoseHistoryProvider> ScopedMessage(Context, this);
-
 	Source.Initialize(Context);
 }
 
 void FAnimNode_PoseSearchHistoryCollector::CacheBones_AnyThread(const FAnimationCacheBonesContext& Context)
 {
 	DECLARE_SCOPE_HIERARCHICAL_COUNTER_ANIMNODE(CacheBones_AnyThread)
-
+	Super::CacheBones_AnyThread(Context);
 	Source.CacheBones(Context);
 }
 
 void FAnimNode_PoseSearchHistoryCollector::Evaluate_AnyThread(FPoseContext& Output)
 {
 	DECLARE_SCOPE_HIERARCHICAL_COUNTER_ANIMNODE(Evaluate_AnyThread);
+	check(Output.AnimInstanceProxy);
 
+	Super::Evaluate_AnyThread(Output);
 	Source.Evaluate(Output);
-	PoseHistory.Update(Output.AnimInstanceProxy->GetDeltaSeconds(), Output, Output.AnimInstanceProxy->GetComponentTransform());
+
+	FCSPose<FCompactPose> ComponentSpacePose;
+	ComponentSpacePose.InitPose(Output.Pose);
+	PoseHistory.Update(Output.AnimInstanceProxy->GetDeltaSeconds(), ComponentSpacePose, Output.AnimInstanceProxy->GetComponentTransform());
 
 #if WITH_EDITORONLY_DATA && ENABLE_ANIM_DEBUG
 	bWasEvaluated = true;
@@ -80,43 +148,60 @@ void FAnimNode_PoseSearchHistoryCollector::Evaluate_AnyThread(FPoseContext& Outp
 void FAnimNode_PoseSearchHistoryCollector::Update_AnyThread(const FAnimationUpdateContext& Context)
 {
 	DECLARE_SCOPE_HIERARCHICAL_COUNTER_ANIMNODE(Update_AnyThread);
-
-	UE::Anim::TScopedGraphMessage<UE::PoseSearch::Private::FPoseHistoryProvider> ScopedMessage(Context, this);
-
+	Super::Update_AnyThread(Context);
 	Source.Update(Context);
 }
 
 void FAnimNode_PoseSearchHistoryCollector::GatherDebugData(FNodeDebugData& DebugData)
 {
+	Super::GatherDebugData(DebugData);
 	Source.GatherDebugData(DebugData);
 }
 
-bool FAnimNode_PoseSearchHistoryCollector::HasPreUpdate() const
+
+/////////////////////////////////////////////////////
+// FAnimNode_PoseSearchComponentSpaceHistoryCollector
+
+void FAnimNode_PoseSearchComponentSpaceHistoryCollector::Initialize_AnyThread(const FAnimationInitializeContext& Context)
 {
-	return WITH_EDITORONLY_DATA && ENABLE_ANIM_DEBUG;
+	DECLARE_SCOPE_HIERARCHICAL_COUNTER_ANIMNODE(Initialize_AnyThread);
+	Super::Initialize_AnyThread(Context);
+	Source.Initialize(Context);
 }
 
-void FAnimNode_PoseSearchHistoryCollector::PreUpdate(const UAnimInstance* InAnimInstance)
+void FAnimNode_PoseSearchComponentSpaceHistoryCollector::CacheBones_AnyThread(const FAnimationCacheBonesContext& Context)
 {
+	DECLARE_SCOPE_HIERARCHICAL_COUNTER_ANIMNODE(CacheBones_AnyThread)
+	Super::CacheBones_AnyThread(Context);
+	Source.CacheBones(Context);
+}
+
+void FAnimNode_PoseSearchComponentSpaceHistoryCollector::EvaluateComponentSpace_AnyThread(FComponentSpacePoseContext& Output)
+{
+	DECLARE_SCOPE_HIERARCHICAL_COUNTER_ANIMNODE(EvaluateComponentSpace_AnyThread);
+	check(Output.AnimInstanceProxy);
+
+	Super::EvaluateComponentSpace_AnyThread(Output);
+	Source.EvaluateComponentSpace(Output);
+
+	PoseHistory.Update(Output.AnimInstanceProxy->GetDeltaSeconds(), Output.Pose, Output.AnimInstanceProxy->GetComponentTransform());
+
 #if WITH_EDITORONLY_DATA && ENABLE_ANIM_DEBUG
-	if (bWasEvaluated && CVarAnimPoseHistoryDebugDraw.GetValueOnAnyThread())
-	{
-		if (const USkeletalMeshComponent* SkeletalMeshComponent = InAnimInstance->GetSkelMeshComponent())
-		{
-			const USkinnedAsset* SkinnedAsset = SkeletalMeshComponent->GetSkinnedAsset();
-			const UWorld* World = SkeletalMeshComponent->GetWorld();
-			if (SkinnedAsset && World)
-			{
-				if (const USkeleton* Skeleton = SkinnedAsset->GetSkeleton())
-				{
-					PoseHistory.DebugDraw(World, Skeleton);
-				}
-			}
-		}
-	}
+	bWasEvaluated = true;
+#endif
+}
 
-	bWasEvaluated = false;
-#endif // WITH_EDITORONLY_DATA && ENABLE_ANIM_DEBUG
+void FAnimNode_PoseSearchComponentSpaceHistoryCollector::Update_AnyThread(const FAnimationUpdateContext& Context)
+{
+	DECLARE_SCOPE_HIERARCHICAL_COUNTER_ANIMNODE(Update_AnyThread);
+	Super::Update_AnyThread(Context);
+	Source.Update(Context);
+}
+
+void FAnimNode_PoseSearchComponentSpaceHistoryCollector::GatherDebugData(FNodeDebugData& DebugData)
+{
+	Super::GatherDebugData(DebugData);
+	Source.GatherDebugData(DebugData);
 }
 
 #undef LOCTEXT_NAMESPACE
diff --git a/Engine/Plugins/Experimental/Animation/PoseSearch/Source/Runtime/Private/PoseSearchContext.cpp b/Engine/Plugins/Experimental/Animation/PoseSearch/Source/Runtime/Private/PoseSearchContext.cpp
--- a/Engine/Plugins/Experimental/Animation/PoseSearch/Source/Runtime/Private/PoseSearchContext.cpp
+++ b/Engine/Plugins/Experimental/Animation/PoseSearch/Source/Runtime/Private/PoseSearchContext.cpp
@@ -188,19 +188,20 @@ FTransform FSearchContext::GetComponentSpaceTransform(float SampleTime, const UP
 	
 		// collecting the local bone transforms from the FPoseHistory
 		check(History);
-		TArray<FTransform> SampledLocalPose;
-		History->GetLocalPoseAtTime(SampleTime, Schema->BoneIndicesWithParents, SampledLocalPose);
-		
-		TArray<FTransform> SampledComponentPose;
-		FAnimationRuntime::FillUpComponentSpaceTransforms(Schema->Skeleton->GetReferenceSkeleton(), SampledLocalPose, SampledComponentPose);
-
-		// adding bunch of entries, without caring about adding eventual duplicates
-		for (const FBoneIndexType NewEntryBoneIndexType : Schema->BoneIndicesWithParents)
+		FTransform BoneComponentSpaceTransform;
+		if (!History->GetComponentSpaceTransformAtTime(SampleTime, BoneIndexType, BoneComponentSpaceTransform))
 		{
-			CachedTransforms.Add(SampleTime, NewEntryBoneIndexType, SampledComponentPose[NewEntryBoneIndexType]);
+			FName BoneName;
+			if (const USkeleton* Skeleton = Schema->Skeleton)
+			{
+				BoneName = Skeleton->GetReferenceSkeleton().GetBoneName(BoneIndexType);
+			}
+
+			UE_LOG(LogPoseSearch, Warning, TEXT("FSearchContext::GetComponentSpaceTransform - Couldn't find BoneIndexType %d (%s) requested by %s"), BoneIndexType, *BoneName.ToString(), *Schema->GetName());
 		}
 
-		return SampledComponentPose[BoneIndexType];
+		CachedTransforms.Add(SampleTime, BoneIndexType, BoneComponentSpaceTransform);
+		return BoneComponentSpaceTransform;
 	}
 
 	return FTransform::Identity;
diff --git a/Engine/Plugins/Experimental/Animation/PoseSearch/Source/Runtime/Private/PoseSearchHistory.cpp b/Engine/Plugins/Experimental/Animation/PoseSearch/Source/Runtime/Private/PoseSearchHistory.cpp
--- a/Engine/Plugins/Experimental/Animation/PoseSearch/Source/Runtime/Private/PoseSearchHistory.cpp
+++ b/Engine/Plugins/Experimental/Animation/PoseSearch/Source/Runtime/Private/PoseSearchHistory.cpp
@@ -8,11 +8,10 @@
 #include "DrawDebugHelpers.h"
 #include "PoseSearch/PoseSearchResult.h"
 #include "PoseSearch/PoseSearchDatabase.h"
+#include "PoseSearch/PoseSearchDefines.h"
 
 IMPLEMENT_ANIMGRAPH_MESSAGE(UE::PoseSearch::IPoseHistoryProvider);
 
-#define LOCTEXT_NAMESPACE "PoseSearch"
-
 namespace UE::PoseSearch
 {
 
@@ -64,39 +63,19 @@ FORCEINLINE auto LowerBound(IteratorType First, IteratorType Last, const ValueTy
 
 //////////////////////////////////////////////////////////////////////////
 // FPoseHistory
-/**
-* Fills skeleton transforms with evaluated compact pose transforms.
-* Bones that weren't evaluated are filled with the bone's reference pose.
-*/
-static void CopyCompactToSkeletonPose(const FCompactPose& Pose, TArray<FTransform>& OutLocalTransforms)
-{
-	const FBoneContainer& BoneContainer = Pose.GetBoneContainer();
-	const USkeleton* SkeletonAsset = BoneContainer.GetSkeletonAsset();
-	check(SkeletonAsset);
-
-	const FReferenceSkeleton& RefSkeleton = SkeletonAsset->GetReferenceSkeleton();
-	TArrayView<const FTransform> RefSkeletonTransforms = MakeArrayView(RefSkeleton.GetRefBonePose());
-	const int32 NumSkeletonBones = RefSkeleton.GetNum();
 
-	OutLocalTransforms.SetNum(NumSkeletonBones);
+void FPoseHistory::Init(int32 InNumPoses, float InTimeHorizon, const TArray<FBoneIndexType>& RequiredBones)
+{
+	TimeHorizon = InTimeHorizon;
 
-	for (auto SkeletonBoneIdx = FSkeletonPoseBoneIndex(0); SkeletonBoneIdx != NumSkeletonBones; ++SkeletonBoneIdx)
+	BoneToTransformMap.Reset();
+	for (int32 i = 0; i < RequiredBones.Num(); ++i)
 	{
-		FCompactPoseBoneIndex CompactBoneIdx = BoneContainer.GetCompactPoseIndexFromSkeletonPoseIndex(SkeletonBoneIdx);
-		OutLocalTransforms[SkeletonBoneIdx.GetInt()] = CompactBoneIdx.IsValid() ? Pose[CompactBoneIdx] : RefSkeletonTransforms[SkeletonBoneIdx.GetInt()];
+		BoneToTransformMap.Add(RequiredBones[i]) = i;
 	}
-}
-
-void FPoseHistory::Init(int32 InNumPoses, float InTimeHorizon)
-{
-	Poses.Reserve(InNumPoses);
-	TimeHorizon = InTimeHorizon;
-}
 
-void FPoseHistory::Init(const FPoseHistory& History)
-{
-	Poses = History.Poses;
-	TimeHorizon = History.TimeHorizon;
+	Entries.Reset();
+	Entries.Reserve(InNumPoses);
 }
 
 static FBoneIndexType GetMaxFBoneIndexType(const TArray<FBoneIndexType>& RequiredBones)
@@ -112,81 +91,104 @@ static FBoneIndexType GetMaxFBoneIndexType(const TArray<FBoneIndexType>& Require
 	return MaxBoneIndexType;
 }
 
-void FPoseHistory::GetLocalPoseAtTime(float Time, const TArray<FBoneIndexType>& RequiredBones, TArray<FTransform>& LocalPose) const
+bool FPoseHistory::GetComponentSpaceTransformAtTime(float Time, FBoneIndexType BoneIndexType, FTransform& OutBoneTransform, bool bExtrapolate) const
 {
-	const int32 Num = Poses.Num();
+	FComponentSpaceTransformIndex TransformIndex = FComponentSpaceTransformIndex(BoneIndexType);
+	if (!BoneToTransformMap.IsEmpty())
+	{
+		if (const FComponentSpaceTransformIndex* FoundTransformIndex = BoneToTransformMap.Find(BoneIndexType))
+		{
+			TransformIndex = *FoundTransformIndex;
+		}
+		else
+		{
+			GetRootTransformAtTime(Time, OutBoneTransform, bExtrapolate);
+			return false;
+		}
+	}
+
+	const int32 Num = Entries.Num();
 	if (Num > 1)
 	{
 		const float SecondsAgo = -Time;
-		const int32 LowerBoundIdx = LowerBound(Poses.begin(), Poses.end(), SecondsAgo, [](const FPose& Pose, float Value) { return Value < Pose.Time; });
+		const int32 LowerBoundIdx = LowerBound(Entries.begin(), Entries.end(), SecondsAgo, [](const FEntry& Entry, float Value) { return Value < Entry.Time; });
 		const int32 NextIdx = FMath::Clamp(LowerBoundIdx, 1, Num - 1);
 		const int32 PrevIdx = NextIdx - 1;
 
-		const FPose& PrevPose = Poses[PrevIdx];
-		const FPose& NextPose = Poses[NextIdx];
-				
-		// Compute alpha between previous and next Poses
-		const float Alpha = FMath::GetMappedRangeValueClamped(FVector2f(PrevPose.Time, NextPose.Time), FVector2f(0.0f, 1.0f), SecondsAgo);
-
-		// Lerp between poses by alpha to produce output local pose at requested sample time
-		check(PrevPose.LocalTransforms.Num() == NextPose.LocalTransforms.Num());
-		check(GetMaxFBoneIndexType(RequiredBones) < PrevPose.LocalTransforms.Num());
+		const FEntry& PrevEntry = Entries[PrevIdx];
+		const FEntry& NextEntry = Entries[NextIdx];
 
-		LocalPose = PrevPose.LocalTransforms;
-		FAnimationRuntime::LerpBoneTransforms(LocalPose, NextPose.LocalTransforms, Alpha, RequiredBones);
+		const float Denominator = NextEntry.Time - PrevEntry.Time;
+		if (!FMath::IsNearlyZero(Denominator))
+		{
+			const float Numerator = SecondsAgo - PrevEntry.Time;
+			const float LerpValue = bExtrapolate ? Numerator / Denominator : FMath::Clamp(Numerator / Denominator, 0.f, 1.f);
+			OutBoneTransform.Blend(PrevEntry.ComponentSpaceTransforms[TransformIndex], NextEntry.ComponentSpaceTransforms[TransformIndex], LerpValue);
+		}
+		else
+		{
+			OutBoneTransform = PrevEntry.ComponentSpaceTransforms[TransformIndex];
+		}
 	}
 	else if (Num > 0)
 	{
-		check(GetMaxFBoneIndexType(RequiredBones) < Poses[0].LocalTransforms.Num());
-		LocalPose = Poses[0].LocalTransforms;
+		OutBoneTransform = Entries[0].ComponentSpaceTransforms[TransformIndex];
 	}
 	else
 	{
-		const FBoneIndexType MaxBoneIndexType = GetMaxFBoneIndexType(RequiredBones);
-		LocalPose.Init(FTransform::Identity, MaxBoneIndexType + 1);
+		OutBoneTransform = FTransform::Identity;
 	}
+
+	return true;
 }
 
-void FPoseHistory::GetRootTransformAtTime(float Time, FTransform& RootTransform) const
+void FPoseHistory::GetRootTransformAtTime(float Time, FTransform& OutRootTransform, bool bExtrapolate) const
 {
-	const int32 Num = Poses.Num();
+	const int32 Num = Entries.Num();
 	if (Num > 1)
 	{
 		const float SecondsAgo = -Time;
-		const int32 LowerBoundIdx = LowerBound(Poses.begin(), Poses.end(), SecondsAgo, [](const FPose& Pose, float Value) { return Value < Pose.Time; });
+		const int32 LowerBoundIdx = LowerBound(Entries.begin(), Entries.end(), SecondsAgo, [](const FEntry& Entry, float Value) { return Value < Entry.Time; });
 		const int32 NextIdx = FMath::Clamp(LowerBoundIdx, 1, Num - 1);
 		const int32 PrevIdx = NextIdx - 1;
 
-		const FPose& PrevPose = Poses[PrevIdx];
-		const FPose& NextPose = Poses[NextIdx];
+		const FEntry& PrevEntry = Entries[PrevIdx];
+		const FEntry& NextEntry = Entries[NextIdx];
 
-		// Compute alpha between previous and next Poses
-		const float Alpha = FMath::GetMappedRangeValueClamped(FVector2f(PrevPose.Time, NextPose.Time), FVector2f(0.0f, 1.0f), SecondsAgo);
-
-		RootTransform.Blend(PrevPose.RootTransform, NextPose.RootTransform, Alpha);
+		const float Denominator = NextEntry.Time - PrevEntry.Time;
+		if (!FMath::IsNearlyZero(Denominator))
+		{
+			const float Numerator = SecondsAgo - PrevEntry.Time;
+			const float LerpValue = bExtrapolate ? Numerator / Denominator : FMath::Clamp(Numerator / Denominator, 0.f, 1.f);
+			OutRootTransform.Blend(PrevEntry.RootTransform, NextEntry.RootTransform, LerpValue);
+		}
+		else
+		{
+			OutRootTransform = PrevEntry.RootTransform;
+		}
 	}
 	else if (Num > 0)
 	{
-		RootTransform = Poses[0].RootTransform;
+		OutRootTransform = Entries[0].RootTransform;
 	}
 	else
 	{
-		RootTransform = FTransform::Identity;
+		OutRootTransform = FTransform::Identity;
 	}
 }
 
-void FPoseHistory::Update(float SecondsElapsed, const FPoseContext& PoseContext, FTransform ComponentTransform)
+void FPoseHistory::Update(float SecondsElapsed, FCSPose<FCompactPose>& ComponentSpacePose, FTransform ComponentTransform)
 {
 	// Age our elapsed times
-	for (FPose& Pose : Poses)
+	for (FEntry& Entry : Entries)
 	{
-		Pose.Time += SecondsElapsed;
+		Entry.Time += SecondsElapsed;
 	}
 
-	if (Poses.Num() != Poses.Max())
+	if (Entries.Num() != Entries.Max())
 	{
 		// Consume every pose until the queue is full
-		Poses.Emplace();
+		Entries.Emplace();
 	}
 	else
 	{
@@ -197,28 +199,55 @@ void FPoseHistory::Update(float SecondsElapsed, const FPoseContext& PoseContext,
 
 		const float SampleInterval = GetSampleTimeInterval();
 
-		bool bCanEvictOldest = Poses[1].Time >= TimeHorizon + SampleInterval;
-		bool bShouldPushNewest = Poses[Poses.Num() - 2].Time >= SampleInterval;
+		bool bCanEvictOldest = Entries[1].Time >= TimeHorizon + SampleInterval;
+		bool bShouldPushNewest = Entries[Entries.Num() - 2].Time >= SampleInterval;
 
 		if (bCanEvictOldest && bShouldPushNewest)
 		{
-			FPose PoseTemp = MoveTemp(Poses.First());
-			Poses.PopFront();
-			Poses.Emplace(MoveTemp(PoseTemp));
+			FEntry EntryTemp = MoveTemp(Entries.First());
+			Entries.PopFront();
+			Entries.Emplace(MoveTemp(EntryTemp));
 		}
 	}
 
-	// Regardless of the retention policy, we always update the most recent pose
-	FPose& CurrentPose = Poses.Last();
-	CurrentPose.Time = 0.f;
-	CurrentPose.RootTransform = ComponentTransform;
-	CopyCompactToSkeletonPose(PoseContext.Pose, CurrentPose.LocalTransforms);
+	// Regardless of the retention policy, we always update the most recent Entry
+	FEntry& CurrentEntry = Entries.Last();
+	CurrentEntry.Time = 0.f;
+	CurrentEntry.RootTransform = ComponentTransform;
+
+	const FBoneContainer& BoneContainer = ComponentSpacePose.GetPose().GetBoneContainer();
+	const USkeleton* SkeletonAsset = BoneContainer.GetSkeletonAsset();
+	check(SkeletonAsset);
+	const FReferenceSkeleton& RefSkeleton = SkeletonAsset->GetReferenceSkeleton();
+	const TArray<FTransform>& RefBonePose = RefSkeleton.GetRefBonePose();
+	const int32 NumSkeletonBones = RefSkeleton.GetNum();
+
+	if (BoneToTransformMap.IsEmpty())
+	{
+		// no mapping: we add all the transforms
+		CurrentEntry.ComponentSpaceTransforms.SetNum(NumSkeletonBones);
+		for (FSkeletonPoseBoneIndex SkeletonBoneIdx(0); SkeletonBoneIdx != NumSkeletonBones; ++SkeletonBoneIdx)
+		{
+			const FCompactPoseBoneIndex CompactBoneIdx = BoneContainer.GetCompactPoseIndexFromSkeletonPoseIndex(SkeletonBoneIdx);
+			CurrentEntry.ComponentSpaceTransforms[SkeletonBoneIdx.GetInt()] = CompactBoneIdx.IsValid() ? ComponentSpacePose.GetComponentSpaceTransform(CompactBoneIdx) : RefBonePose[SkeletonBoneIdx.GetInt()];
+		}
+	}
+	else
+	{
+		CurrentEntry.ComponentSpaceTransforms.SetNum(BoneToTransformMap.Num());
+		for (const TPair<FBoneIndexType, FComponentSpaceTransformIndex>& BoneToTransformPair : BoneToTransformMap)
+		{
+			const FSkeletonPoseBoneIndex SkeletonBoneIdx(BoneToTransformPair.Key);
+			const FCompactPoseBoneIndex CompactBoneIdx = BoneContainer.GetCompactPoseIndexFromSkeletonPoseIndex(SkeletonBoneIdx);
+			CurrentEntry.ComponentSpaceTransforms[BoneToTransformPair.Value] = CompactBoneIdx.IsValid() ? ComponentSpacePose.GetComponentSpaceTransform(CompactBoneIdx) : RefBonePose[SkeletonBoneIdx.GetInt()];
+		}
+	}
 }
 
 float FPoseHistory::GetSampleTimeInterval() const
 {
 	// Reserve one pose for computing derivatives at the time horizon
-	return TimeHorizon / (Poses.Max() - 1);
+	return TimeHorizon / (Entries.Max() - 1);
 }
 
 void FPoseHistory::DebugDraw(const UWorld* World, const USkeleton* Skeleton) const
@@ -234,41 +263,43 @@ void FPoseHistory::DebugDraw(const UWorld* World, const USkeleton* Skeleton) con
 			FMath::RoundToInt(float(A.A) * (1.f - T) + float(B.A) * T));
 	};
 
-	TArray<FTransform> LocalTransforms;
-	TArray<FTransform> GlobalTransforms;
 	TArray<FTransform> PrevGlobalTransforms;
-	for (int32 PoseIndex = 0; PoseIndex < Poses.Num(); ++PoseIndex)
+	for (int32 EntryIndex = 0; EntryIndex < Entries.Num(); ++EntryIndex)
 	{
-		const FPose& Pose = Poses[PoseIndex];
-		if (Pose.LocalTransforms.IsEmpty())
+		const FEntry& Entry = Entries[EntryIndex];
+		if (Entry.ComponentSpaceTransforms.IsEmpty())
 		{
-			LocalTransforms.Reset();
-			GlobalTransforms.Reset();
+			PrevGlobalTransforms.Reset();
 		}
-		else
+		else if (PrevGlobalTransforms.Num() != Entry.ComponentSpaceTransforms.Num())
 		{
-			LocalTransforms = Pose.LocalTransforms;
-			LocalTransforms[0] = LocalTransforms[0] * Pose.RootTransform;
-			FAnimationRuntime::FillUpComponentSpaceTransforms(Skeleton->GetReferenceSkeleton(), LocalTransforms, GlobalTransforms);
+			PrevGlobalTransforms.SetNum(Entry.ComponentSpaceTransforms.Num());
+			for (int32 i = 0; i < Entry.ComponentSpaceTransforms.Num(); ++i)
+			{
+				PrevGlobalTransforms[i] = Entry.ComponentSpaceTransforms[i] * Entry.RootTransform;
+			}
 		}
-
-		if (PrevGlobalTransforms.Num() == GlobalTransforms.Num())
+		else
 		{
-			const float LerpFactor = float(PoseIndex - 1) / float(Poses.Num() - 1);
+			const float LerpFactor = float(EntryIndex - 1) / float(Entries.Num() - 1);
 			const FColor Color = LerpColor(FColorList::Red, FColorList::Orange, LerpFactor);
-			for (int32 TransformIndex = 0; TransformIndex < GlobalTransforms.Num(); ++TransformIndex)
+
+			for (int32 i = 0; i < Entry.ComponentSpaceTransforms.Num(); ++i)
 			{
-				DrawDebugLine(World, PrevGlobalTransforms[TransformIndex].GetLocation(), GlobalTransforms[TransformIndex].GetLocation(), Color, false, 0.f, ESceneDepthPriorityGroup::SDPG_Foreground + 2);
+				const FTransform GlobalTransforms = Entry.ComponentSpaceTransforms[i] * Entry.RootTransform;
+
+				DrawDebugLine(World, PrevGlobalTransforms[i].GetLocation(), GlobalTransforms.GetLocation(), Color, false, 0.f, ESceneDepthPriorityGroup::SDPG_Foreground + 2);
+
+				PrevGlobalTransforms[i] = GlobalTransforms;
 			}
 		}
-
-		PrevGlobalTransforms = GlobalTransforms;
 	}
 #endif // ENABLE_DRAW_DEBUG
 }
 
 //////////////////////////////////////////////////////////////////////////
 // FPoseIndicesHistory
+
 void FPoseIndicesHistory::Update(const FSearchResult& SearchResult, float DeltaTime, float MaxTime)
 {
 	if (MaxTime > 0.f)
@@ -297,5 +328,3 @@ void FPoseIndicesHistory::Update(const FSearchResult& SearchResult, float DeltaT
 }
 
 } // namespace UE::PoseSearch
-
-#undef LOCTEXT_NAMESPACE
\ No newline at end of file
diff --git a/Engine/Plugins/Experimental/Animation/PoseSearch/Source/Runtime/Public/PoseSearch/AnimNode_PoseSearchHistoryCollector.h b/Engine/Plugins/Experimental/Animation/PoseSearch/Source/Runtime/Public/PoseSearch/AnimNode_PoseSearchHistoryCollector.h
--- a/Engine/Plugins/Experimental/Animation/PoseSearch/Source/Runtime/Public/PoseSearch/AnimNode_PoseSearchHistoryCollector.h
+++ b/Engine/Plugins/Experimental/Animation/PoseSearch/Source/Runtime/Public/PoseSearch/AnimNode_PoseSearchHistoryCollector.h
@@ -7,44 +7,73 @@
 #include "AnimNode_PoseSearchHistoryCollector.generated.h"
 
 USTRUCT(BlueprintInternalUseOnly)
-struct POSESEARCH_API FAnimNode_PoseSearchHistoryCollector : public FAnimNode_Base
+struct POSESEARCH_API FAnimNode_PoseSearchHistoryCollector_Base : public FAnimNode_Base
 {
 	GENERATED_BODY()
 
 public:
-	UPROPERTY(EditAnywhere, BlueprintReadWrite, Category=Links)
-	FPoseLink Source;
-
+	
 	// The maximum amount of poses that can be stored
-	UPROPERTY(EditAnywhere, BlueprintReadWrite, Category=Settings, meta=(ClampMin="0"))
+	UPROPERTY(EditAnywhere, BlueprintReadWrite, Category = Settings, meta = (ClampMin="0"))
 	int32 PoseCount = 64;
 	
 	// The time horizon for how long a pose will be stored in seconds
-	UPROPERTY(EditAnywhere, BlueprintReadWrite, Category=Settings, meta=(ClampMin="0"))
+	UPROPERTY(EditAnywhere, BlueprintReadWrite, Category = Settings, meta = (ClampMin="0"))
 	float PoseDuration = 1.5f;
 
-public:
-	FAnimNode_PoseSearchHistoryCollector() { }
+	UPROPERTY(EditAnywhere, Category = Settings)
+	TArray<FBoneReference> CollectedBones;
 
 	// FAnimNode_Base interface
-	virtual void Initialize_AnyThread(const FAnimationInitializeContext& Context) override;
 	virtual void CacheBones_AnyThread(const FAnimationCacheBonesContext& Context) override;
-	virtual void Evaluate_AnyThread(FPoseContext& Output) override;
-	virtual void Update_AnyThread(const FAnimationUpdateContext& Context);
-	virtual void GatherDebugData(FNodeDebugData& DebugData);
+	virtual void Update_AnyThread(const FAnimationUpdateContext& Context) override;
 	virtual bool HasPreUpdate() const override;
 	virtual void PreUpdate(const UAnimInstance* InAnimInstance) override;
 	// End of FAnimNode_Base interface
 
-	UE::PoseSearch::FPoseHistory& GetPoseHistory() { return PoseHistory; }
 	const UE::PoseSearch::FPoseHistory& GetPoseHistory() const { return PoseHistory; }
 
 protected:
-
 	UE::PoseSearch::FPoseHistory PoseHistory;
 
 #if WITH_EDITORONLY_DATA && ENABLE_ANIM_DEBUG
 	// Whether this node was evaluated last frame
 	bool bWasEvaluated = false;
 #endif // WITH_EDITORONLY_DATA
+};
+
+USTRUCT(BlueprintInternalUseOnly)
+struct POSESEARCH_API FAnimNode_PoseSearchHistoryCollector : public FAnimNode_PoseSearchHistoryCollector_Base
+{
+	GENERATED_BODY()
+
+public:
+	UPROPERTY(EditAnywhere, BlueprintReadWrite, Category = Links)
+	FPoseLink Source;
+
+	// FAnimNode_Base interface
+	virtual void Initialize_AnyThread(const FAnimationInitializeContext& Context) override;
+	virtual void CacheBones_AnyThread(const FAnimationCacheBonesContext& Context) override;
+	virtual void Evaluate_AnyThread(FPoseContext& Output) override;
+	virtual void Update_AnyThread(const FAnimationUpdateContext& Context) override;
+	virtual void GatherDebugData(FNodeDebugData& DebugData) override;
+	// End of FAnimNode_Base interface
+};
+
+USTRUCT(BlueprintInternalUseOnly)
+struct POSESEARCH_API FAnimNode_PoseSearchComponentSpaceHistoryCollector : public FAnimNode_PoseSearchHistoryCollector_Base
+{
+	GENERATED_BODY()
+
+public:
+	UPROPERTY(EditAnywhere, BlueprintReadWrite, Category = Links)
+	FComponentSpacePoseLink Source;
+
+	// FAnimNode_Base interface
+	virtual void Initialize_AnyThread(const FAnimationInitializeContext& Context) override;
+	virtual void CacheBones_AnyThread(const FAnimationCacheBonesContext& Context) override;
+	virtual void EvaluateComponentSpace_AnyThread(FComponentSpacePoseContext& Output) override;
+	virtual void Update_AnyThread(const FAnimationUpdateContext& Context) override;
+	virtual void GatherDebugData(FNodeDebugData& DebugData) override;
+	// End of FAnimNode_Base interface
 };
\ No newline at end of file
diff --git a/Engine/Plugins/Experimental/Animation/PoseSearch/Source/Runtime/Public/PoseSearch/PoseSearchContext.h b/Engine/Plugins/Experimental/Animation/PoseSearch/Source/Runtime/Public/PoseSearch/PoseSearchContext.h
--- a/Engine/Plugins/Experimental/Animation/PoseSearch/Source/Runtime/Public/PoseSearch/PoseSearchContext.h
+++ b/Engine/Plugins/Experimental/Animation/PoseSearch/Source/Runtime/Public/PoseSearch/PoseSearchContext.h
@@ -14,7 +14,7 @@ class UPoseSearchSchema;
 namespace UE::PoseSearch
 {
 
-class FPoseHistory;
+struct FPoseHistory;
 struct FPoseIndicesHistory;
 
 enum class EDebugDrawFlags : uint32
@@ -150,7 +150,7 @@ struct POSESEARCH_API FSearchContext
 {
 	EPoseSearchBooleanRequest QueryMirrorRequest = EPoseSearchBooleanRequest::Indifferent;
 	UE::PoseSearch::FDebugDrawParams DebugDrawParams;
-	UE::PoseSearch::FPoseHistory* History = nullptr;
+	const UE::PoseSearch::FPoseHistory* History = nullptr;
 	const FTrajectorySampleRange* Trajectory = nullptr;
 	TObjectPtr<const USkeletalMeshComponent> OwningComponent = nullptr;
 	UE::PoseSearch::FSearchResult CurrentResult;
diff --git a/Engine/Plugins/Experimental/Animation/PoseSearch/Source/Runtime/Public/PoseSearch/PoseSearchHistory.h b/Engine/Plugins/Experimental/Animation/PoseSearch/Source/Runtime/Public/PoseSearch/PoseSearchHistory.h
--- a/Engine/Plugins/Experimental/Animation/PoseSearch/Source/Runtime/Public/PoseSearch/PoseSearchHistory.h
+++ b/Engine/Plugins/Experimental/Animation/PoseSearch/Source/Runtime/Public/PoseSearch/PoseSearchHistory.h
@@ -3,11 +3,10 @@
 #pragma once
 
 #include "Animation/AnimNodeMessages.h"
-#include "BoneIndices.h"
+#include "BonePose.h"
 #include "Containers/RingBuffer.h"
 #include "UObject/ObjectKey.h"
 
-struct FPoseContext;
 class USkeleton;
 class UWorld;
 
@@ -16,33 +15,29 @@ namespace UE::PoseSearch
 
 struct FSearchResult;
 
-/**
-* Records poses over time in a ring buffer.
-* FFeatureVectorBuilder uses this to sample from the present or past poses according to the search schema.
-*/
-class POSESEARCH_API FPoseHistory
+struct POSESEARCH_API FPoseHistory
 {
-public:
-
-	void Init(int32 InNumPoses, float InTimeHorizon);
-	void Init(const FPoseHistory& History);
-	void Update(float SecondsElapsed, const FPoseContext& PoseContext, FTransform ComponentTransform);
+	void Init(int32 InNumPoses, float InTimeHorizon, const TArray<FBoneIndexType>& RequiredBones);
+	void Update(float SecondsElapsed, FCSPose<FCompactPose>& ComponentSpacePose, FTransform ComponentTransform);
 	float GetSampleTimeInterval() const;
 	float GetTimeHorizon() const { return TimeHorizon; }
-	void GetLocalPoseAtTime(float Time, const TArray<FBoneIndexType>& RequiredBones, TArray<FTransform>& LocalPose) const;
-	void GetRootTransformAtTime(float Time, FTransform& RootTransform) const;
+	bool GetComponentSpaceTransformAtTime(float Time, FBoneIndexType BoneIndexType, FTransform& OutBoneTransform, bool bExtrapolate = true) const;
+	void GetRootTransformAtTime(float Time, FTransform& OutRootTransform, bool bExtrapolate = true) const;
 	void DebugDraw(const UWorld* World, const USkeleton* Skeleton) const;
 
 private:
 
-	struct FPose
+	struct FEntry
 	{
-		FTransform RootTransform; // @todo: remove RootTransform: this should be unnecessary, since FTrajectorySampleRange already contains the past as well as the prediction for the RootTransform
-		TArray<FTransform> LocalTransforms;
-		float Time = 0.0f;
+		FTransform RootTransform;
+		TArray<FTransform> ComponentSpaceTransforms;
+		float Time = 0.f;
 	};
-	TRingBuffer<FPose> Poses;
-	float TimeHorizon = 0.0f;
+
+	typedef uint16 FComponentSpaceTransformIndex;
+	TMap<FBoneIndexType, FComponentSpaceTransformIndex> BoneToTransformMap;
+	TRingBuffer<FEntry> Entries;
+	float TimeHorizon = 0.f;
 };
 
 class IPoseHistoryProvider : public UE::Anim::IGraphMessage
@@ -50,7 +45,6 @@ class IPoseHistoryProvider : public UE::Anim::IGraphMessage
 	DECLARE_ANIMGRAPH_MESSAGE(IPoseHistoryProvider);
 public:
 	virtual const FPoseHistory& GetPoseHistory() const = 0;
-	virtual FPoseHistory& GetPoseHistory() = 0;
 };
 
 struct FHistoricalPoseIndex
