motion matching - WIP PoseSearch database montages support (unsupported SlotName and multiple SlotAnimTracks)

diff --git a/Engine/Plugins/Experimental/Animation/PoseSearch/Source/Editor/Private/PoseSearchDatabaseAssetListItem.cpp b/Engine/Plugins/Experimental/Animation/PoseSearch/Source/Editor/Private/PoseSearchDatabaseAssetListItem.cpp
--- a/Engine/Plugins/Experimental/Animation/PoseSearch/Source/Editor/Private/PoseSearchDatabaseAssetListItem.cpp
+++ b/Engine/Plugins/Experimental/Animation/PoseSearch/Source/Editor/Private/PoseSearchDatabaseAssetListItem.cpp
@@ -45,7 +45,7 @@ namespace UE::PoseSearch
 		EditorViewModel = InEditorViewModel;
 		SkeletonView = InHierarchy;
 
-		if (InAssetTreeNode->SourceAssetType == ESearchIndexAssetType::Invalid)
+		if (InAssetTreeNode->SourceAssetIdx == INDEX_NONE)
 		{
 			ConstructGroupItem(OwnerTable);
 		}
@@ -104,6 +104,12 @@ namespace UE::PoseSearch
 		SkeletonView.Pin()->RefreshTreeView(false);
 	}
 
+	void SDatabaseAssetListItem::OnAddAnimMontage()
+	{
+		EditorViewModel.Pin()->AddAnimMontageToDatabase(nullptr);
+		SkeletonView.Pin()->RefreshTreeView(false);
+	}
+
 	FText SDatabaseAssetListItem::GetName() const
 	{
 		TSharedPtr<FDatabaseAssetTreeNode> Node = WeakAssetTreeNode.Pin();
@@ -131,7 +137,7 @@ namespace UE::PoseSearch
 		TSharedPtr<SWidget> ItemWidget;
 		const FDetailColumnSizeData& ColumnSizeData = SkeletonView.Pin()->GetColumnSizeData();
 		
-		if (Node->SourceAssetType == ESearchIndexAssetType::Invalid)
+		if (Node->SourceAssetIdx == INDEX_NONE)
 		{
 			// it's a group
 			SAssignNew(ItemWidget, SBorder)
@@ -334,6 +340,14 @@ namespace UE::PoseSearch
 			NAME_None,
 			EUserInterfaceActionType::Button);
 
+		AddOptions.AddMenuEntry(
+			LOCTEXT("AnimMontageOption", "Add Anim Montage"),
+			LOCTEXT("AddAnimMontageToDefaultGroupTooltip", "Add new anim montage to this group"),
+			FSlateIcon(),
+			FUIAction(FExecuteAction::CreateSP(this, &SDatabaseAssetListItem::OnAddAnimMontage)),
+			NAME_None,
+			EUserInterfaceActionType::Button);
+
 		TSharedPtr<SComboButton> AddButton;
 		SAssignNew(AddButton, SComboButton)
 		.ContentPadding(0)
@@ -372,8 +386,7 @@ namespace UE::PoseSearch
 		TSharedPtr<FDatabaseAssetTreeNode> TreeNodePtr = WeakAssetTreeNode.Pin();
 		if (const FPoseSearchIndexAsset* SelectedIndexAsset = ViewModelPtr->GetSelectedActorIndexAsset())
 		{
-			if (TreeNodePtr->SourceAssetType == ESearchIndexAssetType::Sequence &&
-				TreeNodePtr->SourceAssetIdx == SelectedIndexAsset->SourceAssetIdx)
+			if (TreeNodePtr->SourceAssetIdx == SelectedIndexAsset->SourceAssetIdx)
 			{
 				return EVisibility::Visible;
 			}
diff --git a/Engine/Plugins/Experimental/Animation/PoseSearch/Source/Editor/Private/PoseSearchDatabaseAssetTree.cpp b/Engine/Plugins/Experimental/Animation/PoseSearch/Source/Editor/Private/PoseSearchDatabaseAssetTree.cpp
--- a/Engine/Plugins/Experimental/Animation/PoseSearch/Source/Editor/Private/PoseSearchDatabaseAssetTree.cpp
+++ b/Engine/Plugins/Experimental/Animation/PoseSearch/Source/Editor/Private/PoseSearchDatabaseAssetTree.cpp
@@ -55,7 +55,7 @@ namespace UE::PoseSearch
 					SNew(SPositiveActionButton)
 					.Icon(FAppStyle::Get().GetBrush("Icons.Plus"))
 					.Text(LOCTEXT("AddNew", "Add"))
-					.ToolTipText(LOCTEXT("AddNewToolTip", "Add a new Sequence, Blend Space or Group"))
+					.ToolTipText(LOCTEXT("AddNewToolTip", "Add a new Sequence, Blend Space, Anim Composite, or Anim Montage"))
 					.OnGetMenuContent(this, &SDatabaseAssetTree::CreateAddNewMenuWidget)
 				]
 				+ SHorizontalBox::Slot()
@@ -233,33 +233,13 @@ namespace UE::PoseSearch
 		{
 			// Setup default group node
 			{
-				const TSharedPtr<FDatabaseAssetTreeNode> DefaultGroupNode = MakeShared<FDatabaseAssetTreeNode>(
-					INDEX_NONE,
-					ESearchIndexAssetType::Invalid,
-					ViewModelRef);
+				const TSharedPtr<FDatabaseAssetTreeNode> DefaultGroupNode = MakeShared<FDatabaseAssetTreeNode>(INDEX_NONE, ViewModelRef);
 				AllNodes.Add(DefaultGroupNode);
 				RootNodes.Add(DefaultGroupNode);
 			}
 			
 			const int32 DefaultGroupIdx = RootNodes.Num() - 1;
-			
-			auto CreateAssetNode = [this, ViewModelRef](int32 AssetIdx, ESearchIndexAssetType AssetType, int32 GroupIdx)
-			{
-				// Create sequence node
-				const TSharedPtr<FDatabaseAssetTreeNode> SequenceGroupNode = MakeShared<FDatabaseAssetTreeNode>(
-					AssetIdx,
-					AssetType,
-					ViewModelRef);
-				const TSharedPtr<FDatabaseAssetTreeNode>& ParentGroupNode = RootNodes[GroupIdx];
-
-				// Setup hierarchy
-				SequenceGroupNode->Parent = ParentGroupNode;
-				ParentGroupNode->Children.Add(SequenceGroupNode);
-
-				// Keep track of node
-				AllNodes.Add(SequenceGroupNode);
-			};
-			
+						
 			// Build an index based off of alphabetical order than iterate the index instead
 			TArray<uint32> IndexArray;
 			IndexArray.SetNumUninitialized(Database->AnimationAssets.Num());
@@ -299,7 +279,16 @@ namespace UE::PoseSearch
 
 					if (!bFiltered)
 					{
-						CreateAssetNode(MappedId, DatabaseAnimationAsset->GetSearchIndexType(), DefaultGroupIdx);
+						// Create sequence node
+						const TSharedPtr<FDatabaseAssetTreeNode> SequenceGroupNode = MakeShared<FDatabaseAssetTreeNode>(MappedId, ViewModelRef);
+						const TSharedPtr<FDatabaseAssetTreeNode>& ParentGroupNode = RootNodes[DefaultGroupIdx];
+
+						// Setup hierarchy
+						SequenceGroupNode->Parent = ParentGroupNode;
+						ParentGroupNode->Children.Add(SequenceGroupNode);
+
+						// Keep track of node
+						AllNodes.Add(SequenceGroupNode);
 					}
 				}
 			}
@@ -459,27 +448,35 @@ namespace UE::PoseSearch
 		{
 			AddOptions.AddMenuEntry(
 				LOCTEXT("AddSequenceOption", "Sequence"),
-				LOCTEXT("AddSequenceOptionTooltip", "Add new sequence to the default group"),
+				LOCTEXT("AddSequenceToDatabaseTooltip", "Add new sequence to the database"),
 				FSlateIcon(),
 				FUIAction(FExecuteAction::CreateSP(this, &SDatabaseAssetTree::OnAddSequence, true)),
 				NAME_None,
 				EUserInterfaceActionType::Button);
 
 			AddOptions.AddMenuEntry(
-				LOCTEXT("BlendSpaceOption", "Blend Space"),
-				LOCTEXT("AddBlendSpaceToDefaultGroupTooltip", "Add new blend space to the default group"),
+				LOCTEXT("AddBlendSpaceOption", "Blend Space"),
+				LOCTEXT("AddBlendSpaceToDatabaseTooltip", "Add new blend space to the database"),
 				FSlateIcon(),
 				FUIAction(FExecuteAction::CreateSP(this, &SDatabaseAssetTree::OnAddBlendSpace, true)),
 				NAME_None,
 				EUserInterfaceActionType::Button);
 
 			AddOptions.AddMenuEntry(
-				LOCTEXT("AnimCompositeOption", "Anim Composite"),
-				LOCTEXT("AddAnimCompositeToDefaultGroupTooltip", "Add new anim composite to the default group"),
+				LOCTEXT("AddAnimCompositeOption", "Anim Composite"),
+				LOCTEXT("AddAnimCompositeToDatabaseTooltip", "Add new composite to the database"),
 				FSlateIcon(),
 				FUIAction(FExecuteAction::CreateSP(this, &SDatabaseAssetTree::OnAddAnimComposite, true)),
 				NAME_None,
 				EUserInterfaceActionType::Button);
+
+			AddOptions.AddMenuEntry(
+				LOCTEXT("AddAnimMontageOption", "Anim Montage"),
+				LOCTEXT("AddAnimMontageToDatabaseTooltip", "Add new montage to the database"),
+				FSlateIcon(),
+				FUIAction(FExecuteAction::CreateSP(this, &SDatabaseAssetTree::OnAddAnimMontage, true)),
+				NAME_None,
+				EUserInterfaceActionType::Button);
 		}
 		AddOptions.EndSection();
 
@@ -496,9 +493,7 @@ namespace UE::PoseSearch
 		{
 			MenuBuilder.AddMenuEntry(
 				LOCTEXT("DeleteUngroup", "Delete / Remove"),
-				LOCTEXT(
-					"DeleteUngroupTooltip", 
-					"Deletes groups and ungrouped assets; removes grouped assets from group."),
+				LOCTEXT("DeleteUngroupTooltip", "Deletes groups and ungrouped assets; removes grouped assets from group."),
 				FSlateIcon(),
 				FUIAction(FExecuteAction::CreateSP(this, &SDatabaseAssetTree::OnDeleteNodes)),
 				NAME_None,
@@ -596,6 +591,18 @@ namespace UE::PoseSearch
 		}
 	}
 
+	void SDatabaseAssetTree::OnAddAnimMontage(bool bFinalizeChanges)
+	{
+		FScopedTransaction Transaction(LOCTEXT("AddAnimMontageTransaction", "Add Anim Montage"));
+
+		EditorViewModel.Pin()->AddAnimMontageToDatabase(nullptr);
+
+		if (bFinalizeChanges)
+		{
+			FinalizeTreeChanges();
+		}
+	}
+
 	void SDatabaseAssetTree::OnDeleteAsset(TSharedPtr<FDatabaseAssetTreeNode> Node, bool bFinalizeChanges)
 	{
 		FScopedTransaction Transaction(LOCTEXT("DeleteAsset", "Delete Asset"));
@@ -626,14 +633,7 @@ namespace UE::PoseSearch
 
 		for (const TSharedPtr<FDatabaseAssetTreeNode>& Node : AllNodes)
 		{
-			const bool bFoundNode = PreviouslySelectedNodes.ContainsByPredicate(
-				[Node](const TSharedPtr<FDatabaseAssetTreeNode>& PrevSelectedNode)
-			{
-				return
-					PrevSelectedNode->SourceAssetType == Node->SourceAssetType &&
-					PrevSelectedNode->SourceAssetIdx == Node->SourceAssetIdx;
-			});
-
+			const bool bFoundNode = PreviouslySelectedNodes.ContainsByPredicate([Node](const TSharedPtr<FDatabaseAssetTreeNode>& PrevSelectedNode) { return PrevSelectedNode->SourceAssetIdx == Node->SourceAssetIdx; });
 			if (bFoundNode)
 			{
 				NewSelectedNodes.Add(Node);
@@ -660,8 +660,7 @@ namespace UE::PoseSearch
 		TArray<TSharedPtr<FDatabaseAssetTreeNode>> SelectedNodes = TreeView->GetSelectedItems();
 		for (TSharedPtr<FDatabaseAssetTreeNode> SelectedNode : SelectedNodes)
 		{
-			if (SelectedNode->SourceAssetType != ESearchIndexAssetType::Invalid ||
-				SelectedNode->SourceAssetIdx != INDEX_NONE)
+			if (SelectedNode->SourceAssetIdx != INDEX_NONE)
 			{
 				return true;
 			}
@@ -680,25 +679,14 @@ namespace UE::PoseSearch
 
 			ViewModel->GetPoseSearchDatabase()->Modify();
 			
-			SelectedNodes.Sort(
-				[](const TSharedPtr<FDatabaseAssetTreeNode>& A, const TSharedPtr<FDatabaseAssetTreeNode>& B)
-			{
-				if (A->SourceAssetType != ESearchIndexAssetType::Invalid &&
-					B->SourceAssetType == ESearchIndexAssetType::Invalid)
-				{
-					return true;
-				}
-				if (B->SourceAssetType != ESearchIndexAssetType::Invalid &&
-					A->SourceAssetType == ESearchIndexAssetType::Invalid)
+			SelectedNodes.Sort([](const TSharedPtr<FDatabaseAssetTreeNode>& A, const TSharedPtr<FDatabaseAssetTreeNode>& B)
 				{
-					return false;
-				}
-				return B->SourceAssetIdx < A->SourceAssetIdx;
-			});
+					return B->SourceAssetIdx < A->SourceAssetIdx;
+				});
 
 			for (TSharedPtr<FDatabaseAssetTreeNode> SelectedNode : SelectedNodes)
 			{
-				if (SelectedNode->SourceAssetType != ESearchIndexAssetType::Invalid)
+				if (SelectedNode->SourceAssetIdx != INDEX_NONE)
 				{
 					OnDeleteAsset(SelectedNode, false);
 				}
diff --git a/Engine/Plugins/Experimental/Animation/PoseSearch/Source/Editor/Private/PoseSearchDatabaseAssetTreeNode.cpp b/Engine/Plugins/Experimental/Animation/PoseSearch/Source/Editor/Private/PoseSearchDatabaseAssetTreeNode.cpp
--- a/Engine/Plugins/Experimental/Animation/PoseSearch/Source/Editor/Private/PoseSearchDatabaseAssetTreeNode.cpp
+++ b/Engine/Plugins/Experimental/Animation/PoseSearch/Source/Editor/Private/PoseSearchDatabaseAssetTreeNode.cpp
@@ -26,28 +26,15 @@
 
 namespace UE::PoseSearch
 {
-	FDatabaseAssetTreeNode::FDatabaseAssetTreeNode(
-		int32 InSourceAssetIdx,
-		ESearchIndexAssetType InSourceAssetType,
-		const TSharedRef<FDatabaseViewModel>& InEditorViewModel) 
+	FDatabaseAssetTreeNode::FDatabaseAssetTreeNode(int32 InSourceAssetIdx, const TSharedRef<FDatabaseViewModel>& InEditorViewModel) 
 		: SourceAssetIdx(InSourceAssetIdx)
-		, SourceAssetType(InSourceAssetType)
 		, EditorViewModel(InEditorViewModel)
-	{ }
+	{
+	}
 
-	TSharedRef<ITableRow> FDatabaseAssetTreeNode::MakeTreeRowWidget(
-		const TSharedRef<STableViewBase>& InOwnerTable,
-		TSharedRef<FDatabaseAssetTreeNode> InDatabaseAssetNode,
-		TSharedRef<FUICommandList> InCommandList,
-		TSharedPtr<SDatabaseAssetTree> InHierarchy)
+	TSharedRef<ITableRow> FDatabaseAssetTreeNode::MakeTreeRowWidget(const TSharedRef<STableViewBase>& InOwnerTable, TSharedRef<FDatabaseAssetTreeNode> InDatabaseAssetNode, TSharedRef<FUICommandList> InCommandList, TSharedPtr<SDatabaseAssetTree> InHierarchy)
 	{
-		return SNew(
-			SDatabaseAssetListItem, 
-			EditorViewModel.Pin().ToSharedRef(), 
-			InOwnerTable, 
-			InDatabaseAssetNode, 
-			InCommandList, 
-			InHierarchy);
+		return SNew(SDatabaseAssetListItem, EditorViewModel.Pin().ToSharedRef(), InOwnerTable, InDatabaseAssetNode, InCommandList, InHierarchy);
 	}
 
 	bool FDatabaseAssetTreeNode::IsRootMotionEnabled() const
diff --git a/Engine/Plugins/Experimental/Animation/PoseSearch/Source/Editor/Private/PoseSearchDatabaseEditor.cpp b/Engine/Plugins/Experimental/Animation/PoseSearch/Source/Editor/Private/PoseSearchDatabaseEditor.cpp
--- a/Engine/Plugins/Experimental/Animation/PoseSearch/Source/Editor/Private/PoseSearchDatabaseEditor.cpp
+++ b/Engine/Plugins/Experimental/Animation/PoseSearch/Source/Editor/Private/PoseSearchDatabaseEditor.cpp
@@ -551,6 +551,16 @@ namespace UE::PoseSearch
 					
 					SelectionWidget.SelectedReflections.Add(NewSelectionReflection);
 				}
+				else if (const FPoseSearchDatabaseAnimMontage* DatabaseAnimMontage = DatabaseAsset.GetPtr<FPoseSearchDatabaseAnimMontage>())
+				{
+					UPoseSearchDatabaseAnimMontageReflection* NewSelectionReflection = NewObject<UPoseSearchDatabaseAnimMontageReflection>();
+					NewSelectionReflection->AddToRoot();
+					NewSelectionReflection->AnimMontage = *DatabaseAnimMontage;
+					NewSelectionReflection->SetSourceLink(SelectedItem, AssetTreeWidget);
+					NewSelectionReflection->SetFlags(RF_Transactional);
+
+					SelectionWidget.SelectedReflections.Add(NewSelectionReflection);
+				}
 				else
 				{
 					checkNoEntry();
diff --git a/Engine/Plugins/Experimental/Animation/PoseSearch/Source/Editor/Private/PoseSearchDatabaseEditorReflection.cpp b/Engine/Plugins/Experimental/Animation/PoseSearch/Source/Editor/Private/PoseSearchDatabaseEditorReflection.cpp
--- a/Engine/Plugins/Experimental/Animation/PoseSearch/Source/Editor/Private/PoseSearchDatabaseEditorReflection.cpp
+++ b/Engine/Plugins/Experimental/Animation/PoseSearch/Source/Editor/Private/PoseSearchDatabaseEditorReflection.cpp
@@ -22,7 +22,6 @@ void UPoseSearchDatabaseReflectionBase::SetSourceLink(
 void UPoseSearchDatabaseSequenceReflection::PostEditChangeProperty(struct FPropertyChangedEvent& PropertyChangedEvent)
 {
 	Super::PostEditChangeProperty(PropertyChangedEvent);
-	check(WeakAssetTreeNode.Pin()->SourceAssetType == ESearchIndexAssetType::Sequence);
 
 	if (const TSharedPtr<UE::PoseSearch::FDatabaseViewModel> ViewModel = WeakAssetTreeNode.Pin()->EditorViewModel.Pin())
 	{
@@ -44,7 +43,6 @@ void UPoseSearchDatabaseSequenceReflection::PostEditChangeProperty(struct FPrope
 void UPoseSearchDatabaseBlendSpaceReflection::PostEditChangeProperty(struct FPropertyChangedEvent& PropertyChangedEvent)
 {
 	Super::PostEditChangeProperty(PropertyChangedEvent);
-	check(WeakAssetTreeNode.Pin()->SourceAssetType == ESearchIndexAssetType::BlendSpace);
 
 	if (const TSharedPtr<UE::PoseSearch::FDatabaseViewModel> ViewModel = WeakAssetTreeNode.Pin()->EditorViewModel.Pin())
 	{
@@ -66,7 +64,6 @@ void UPoseSearchDatabaseBlendSpaceReflection::PostEditChangeProperty(struct FPro
 void UPoseSearchDatabaseAnimCompositeReflection::PostEditChangeProperty(struct FPropertyChangedEvent& PropertyChangedEvent)
 {
 	Super::PostEditChangeProperty(PropertyChangedEvent);
-	check(WeakAssetTreeNode.Pin()->SourceAssetType == ESearchIndexAssetType::AnimComposite);
 
 	if (const TSharedPtr<UE::PoseSearch::FDatabaseViewModel> ViewModel = WeakAssetTreeNode.Pin()->EditorViewModel.Pin())
 	{
@@ -85,6 +82,27 @@ void UPoseSearchDatabaseAnimCompositeReflection::PostEditChangeProperty(struct F
 	}
 }
 
+void UPoseSearchDatabaseAnimMontageReflection::PostEditChangeProperty(struct FPropertyChangedEvent& PropertyChangedEvent)
+{
+	Super::PostEditChangeProperty(PropertyChangedEvent);
+
+	if (const TSharedPtr<UE::PoseSearch::FDatabaseViewModel> ViewModel = WeakAssetTreeNode.Pin()->EditorViewModel.Pin())
+	{
+		UPoseSearchDatabase* Database = ViewModel->GetPoseSearchDatabase();
+		if (IsValid(Database))
+		{
+			FInstancedStruct& DatabaseAsset = Database->GetMutableAnimationAssetStruct(WeakAssetTreeNode.Pin()->SourceAssetIdx);
+			if (FPoseSearchDatabaseAnimMontage* DatabaseAnimMontage = DatabaseAsset.GetMutablePtr<FPoseSearchDatabaseAnimMontage>())
+			{
+				*DatabaseAnimMontage = AnimMontage;
+				Database->MarkPackageDirty();
+
+				AssetTreeWidget->FinalizeTreeChanges(true);
+			}
+		}
+	}
+}
+
 #endif // WITH_EDITOR
 
 void UPoseSearchDatabaseStatistics::Initialize(const UPoseSearchDatabase* PoseSearchDatabase)
diff --git a/Engine/Plugins/Experimental/Animation/PoseSearch/Source/Editor/Private/PoseSearchDatabaseViewModel.cpp b/Engine/Plugins/Experimental/Animation/PoseSearch/Source/Editor/Private/PoseSearchDatabaseViewModel.cpp
--- a/Engine/Plugins/Experimental/Animation/PoseSearch/Source/Editor/Private/PoseSearchDatabaseViewModel.cpp
+++ b/Engine/Plugins/Experimental/Animation/PoseSearch/Source/Editor/Private/PoseSearchDatabaseViewModel.cpp
@@ -11,6 +11,7 @@
 #include "Modules/ModuleManager.h"
 #include "AnimPreviewInstance.h"
 #include "Animation/AnimComposite.h"
+#include "Animation/AnimMontage.h"
 #include "Animation/AnimSequence.h"
 #include "Animation/DebugSkelMeshComponent.h"
 #include "Animation/MirrorDataTable.h"
@@ -35,55 +36,21 @@ namespace UE::PoseSearch
 
 	void FDatabasePreviewActor::Process()
 	{
-		switch (Type)
-		{
-		case ESearchIndexAssetType::Sequence:
-		{
-			SequenceSampler.Process();
-			break;
-		}
-		case ESearchIndexAssetType::BlendSpace:
-		{
-			BlendSpaceSampler.Process();
-			break;
-		}
-		case ESearchIndexAssetType::AnimComposite:
-		{
-			SequenceSampler.Process();
-			break;
-		}
-		default:
-		{
-			checkNoEntry();
-			break;
-		}
-		}
+		FAssetSamplerBase* SamplerInterface = Sampler.GetMutablePtr<FAssetSamplerBase>();
+		check(SamplerInterface);
+		SamplerInterface->Process();
 	}
 
-	const IAssetSampler* FDatabasePreviewActor::GetSampler() const
+	const FAssetSamplerBase* FDatabasePreviewActor::GetSampler() const
 	{
-		switch (Type)
-		{
-		case ESearchIndexAssetType::Sequence:
-			return &SequenceSampler;
-		case ESearchIndexAssetType::BlendSpace:
-			return &BlendSpaceSampler;
-		case ESearchIndexAssetType::AnimComposite:
-			return &SequenceSampler;
-		}
-
-		checkNoEntry();
-		return nullptr;
+		const FAssetSamplerBase* SamplerInterface = Sampler.GetPtr<FAssetSamplerBase>();
+		check(SamplerInterface);
+		return SamplerInterface;
 	}
 
 	float FDatabasePreviewActor::GetScaledTime(float Time) const
 	{
-		float ScaledTime = Time;
-		if (Type == ESearchIndexAssetType::BlendSpace)
-		{
-			ScaledTime = BlendSpaceSampler.GetPlayLength() > UE_KINDA_SMALL_NUMBER ? Time / BlendSpaceSampler.GetPlayLength() : 0.f;
-		}
-		return ScaledTime;
+		return GetSampler()->GetScaledTime(Time);
 	}
 
 	UDebugSkelMeshComponent* FDatabasePreviewActor::GetDebugSkelMeshComponent()
@@ -249,7 +216,10 @@ namespace UE::PoseSearch
 			Input.ExtrapolationParameters = PoseSearchDatabase->ExtrapolationParameters;
 			Input.SequenceBase = DatabaseSequence->Sequence;
 
-			PreviewActor.SequenceSampler.Init(Input);
+			PreviewActor.Sampler = FInstancedStruct::Make(FSequenceBaseSampler());
+			FSequenceBaseSampler* Sampler = PreviewActor.Sampler.GetMutablePtr<FSequenceBaseSampler>();
+			check(Sampler);
+			Sampler->Init(Input);
 		}
 		else if (const FPoseSearchDatabaseAnimComposite* DatabaseAnimComposite = DatabaseAsset.GetPtr<FPoseSearchDatabaseAnimComposite>())
 		{
@@ -257,7 +227,10 @@ namespace UE::PoseSearch
 			Input.ExtrapolationParameters = PoseSearchDatabase->ExtrapolationParameters;
 			Input.SequenceBase = DatabaseAnimComposite->AnimComposite;
 
-			PreviewActor.SequenceSampler.Init(Input);
+			PreviewActor.Sampler = FInstancedStruct::Make(FSequenceBaseSampler());
+			FSequenceBaseSampler* Sampler = PreviewActor.Sampler.GetMutablePtr<FSequenceBaseSampler>();
+			check(Sampler);
+			Sampler->Init(Input);
 		}
 		else if (const FPoseSearchDatabaseBlendSpace* DatabaseBlendSpace = DatabaseAsset.GetPtr<FPoseSearchDatabaseBlendSpace>())
 		{
@@ -267,14 +240,27 @@ namespace UE::PoseSearch
 			Input.BlendSpace = DatabaseBlendSpace->BlendSpace;
 			Input.BlendParameters = IndexAsset.BlendParameters;
 
-			PreviewActor.BlendSpaceSampler.Init(Input);
+			PreviewActor.Sampler = FInstancedStruct::Make(FBlendSpaceSampler());
+			FBlendSpaceSampler* Sampler = PreviewActor.Sampler.GetMutablePtr<FBlendSpaceSampler>();
+			check(Sampler);
+			Sampler->Init(Input);
+		}
+		else if (const FPoseSearchDatabaseAnimMontage* DatabaseAnimMontage = DatabaseAsset.GetPtr<FPoseSearchDatabaseAnimMontage>())
+		{
+			FAnimMontageSampler::FInput Input;
+			Input.ExtrapolationParameters = PoseSearchDatabase->ExtrapolationParameters;
+			Input.AnimMontage = DatabaseAnimMontage->AnimMontage;
+
+			PreviewActor.Sampler = FInstancedStruct::Make(FAnimMontageSampler());
+			FAnimMontageSampler* Sampler = PreviewActor.Sampler.GetMutablePtr<FAnimMontageSampler>();
+			check(Sampler);
+			Sampler->Init(Input);
 		}
 		else
 		{
 			checkNoEntry();
 		}
 
-		PreviewActor.Type = IndexAsset.Type;
 		PreviewActor.IndexAssetIndex = IndexAssetIndex;
 		PreviewActor.CurrentPoseIndex = INDEX_NONE;
 
@@ -356,7 +342,7 @@ namespace UE::PoseSearch
 			}
 
 			const UAnimationAsset* PreviewAsset = AnimInstance->GetAnimationAsset();
-			const IAssetSampler* Sampler = PreviewActor.GetSampler();
+			const FAssetSamplerBase* Sampler = PreviewActor.GetSampler();
 			if (!PreviewAsset || !Sampler)
 			{
 				continue;
@@ -462,6 +448,13 @@ namespace UE::PoseSearch
 		PoseSearchDatabase->AnimationAssets.Add(FInstancedStruct::Make(NewAsset));
 	}
 
+	void FDatabaseViewModel::AddAnimMontageToDatabase(UAnimMontage* AnimMontage)
+	{
+		FPoseSearchDatabaseAnimMontage NewAsset;
+		NewAsset.AnimMontage = AnimMontage;
+		PoseSearchDatabase->AnimationAssets.Add(FInstancedStruct::Make(NewAsset));
+	}
+
 	void FDatabaseViewModel::DeleteFromDatabase(int32 AnimationAssetIndex)
 	{
 		PoseSearchDatabase->AnimationAssets.RemoveAt(AnimationAssetIndex);
diff --git a/Engine/Plugins/Experimental/Animation/PoseSearch/Source/Editor/Private/PoseSearchDebuggerDatabaseColumns.h b/Engine/Plugins/Experimental/Animation/PoseSearch/Source/Editor/Private/PoseSearchDebuggerDatabaseColumns.h
--- a/Engine/Plugins/Experimental/Animation/PoseSearch/Source/Editor/Private/PoseSearchDebuggerDatabaseColumns.h
+++ b/Engine/Plugins/Experimental/Animation/PoseSearch/Source/Editor/Private/PoseSearchDebuggerDatabaseColumns.h
@@ -85,7 +85,10 @@ struct FPoseIdx : ITextColumn
 {
 	using ITextColumn::ITextColumn;
 
-	virtual FText GetLabel() const override { return LOCTEXT("ColumnLabelPoseIndex", "Index"); }
+	virtual FText GetLabel() const override
+	{
+		return LOCTEXT("ColumnLabelPoseIndex", "Index");
+	}
 
 	virtual FSortPredicate GetSortPredicate() const override
 	{
@@ -102,7 +105,10 @@ struct FDatabaseName : IColumn
 {
 	using IColumn::IColumn;
 
-	virtual FText GetLabel() const override { return LOCTEXT("ColumnLabelDatabaseName", "Database"); }
+	virtual FText GetLabel() const override
+	{
+		return LOCTEXT("ColumnLabelDatabaseName", "Database");
+	}
 
 	virtual FSortPredicate GetSortPredicate() const override
 	{
@@ -166,7 +172,10 @@ struct FAssetName : IColumn
 {
 	using IColumn::IColumn;
 
-	virtual FText GetLabel() const override { return LOCTEXT("ColumnLabelAssetName", "Asset"); }
+	virtual FText GetLabel() const override
+	{
+		return LOCTEXT("ColumnLabelAssetName", "Asset");
+	}
 		
 	virtual FSortPredicate GetSortPredicate() const override
 	{
@@ -232,7 +241,10 @@ struct FFrame : ITextColumn
 {
 	using ITextColumn::ITextColumn;
 
-	virtual FText GetLabel() const override { return LOCTEXT("ColumnLabelFrame", "Frame"); }
+	virtual FText GetLabel() const override
+	{
+		return LOCTEXT("ColumnLabelFrame", "Frame");
+	}
 		
 	virtual FSortPredicate GetSortPredicate() const override
 	{
@@ -241,42 +253,21 @@ struct FFrame : ITextColumn
 
 	virtual FText GetRowText(const FRowDataRef& Row) const override
 	{
-		FNumberFormattingOptions TimeFormattingOptions = FNumberFormattingOptions()
-			.SetUseGrouping(false)
-			.SetMaximumFractionalDigits(2);
-
-		FNumberFormattingOptions PercentageFormattingOptions = FNumberFormattingOptions()
-			.SetMaximumFractionalDigits(2);
-
-		if (Row->AssetType == ESearchIndexAssetType::Sequence ||
-			Row->AssetType == ESearchIndexAssetType::AnimComposite)
-		{
-			if (GetDefault<UPersonaOptions>()->bTimelineDisplayPercentage)
-			{
-				return FText::Format(
-					FText::FromString("{0} ({1}) ({2})"),
-					FText::AsNumber(Row->AnimFrame, &FNumberFormattingOptions::DefaultNoGrouping()),
-					FText::AsNumber(Row->AssetTime, &TimeFormattingOptions),
-					FText::AsPercent(Row->AnimPercentage, &PercentageFormattingOptions));
-			}
-			else
-			{
-				return FText::Format(
-					FText::FromString("{0} ({1})"),
-					FText::AsNumber(Row->AnimFrame, &FNumberFormattingOptions::DefaultNoGrouping()),
-					FText::AsNumber(Row->AssetTime, &TimeFormattingOptions));
-			}
-		}
-		else if (Row->AssetType == ESearchIndexAssetType::BlendSpace)
+		FNumberFormattingOptions TimeFormattingOptions = FNumberFormattingOptions().SetUseGrouping(false).SetMaximumFractionalDigits(2);
+		FNumberFormattingOptions PercentageFormattingOptions = FNumberFormattingOptions().SetMaximumFractionalDigits(2);
+		
+		if (GetDefault<UPersonaOptions>()->bTimelineDisplayPercentage)
 		{
-			// There is no frame index associated with a blendspace
-			return FText::Format(
-				FText::FromString("({0})"),
-				FText::AsNumber(Row->AssetTime, &TimeFormattingOptions));
+			return FText::Format(FText::FromString("{0} ({1}) ({2})"),
+				FText::AsNumber(Row->AnimFrame, &FNumberFormattingOptions::DefaultNoGrouping()),
+				FText::AsNumber(Row->AssetTime, &TimeFormattingOptions),
+				FText::AsPercent(Row->AnimPercentage, &PercentageFormattingOptions));
 		}
 		else
 		{
-			return FText::FromString(TEXT("-"));
+			return FText::Format(FText::FromString("{0} ({1})"),
+				FText::AsNumber(Row->AnimFrame, &FNumberFormattingOptions::DefaultNoGrouping()),
+				FText::AsNumber(Row->AssetTime, &TimeFormattingOptions));
 		}
 	}
 };
@@ -285,7 +276,10 @@ struct FCost : ITextColumn
 {
 	using ITextColumn::ITextColumn;
 
-	virtual FText GetLabel() const override { return LOCTEXT("ColumnLabelCost", "Cost"); }
+	virtual FText GetLabel() const override
+	{
+		return LOCTEXT("ColumnLabelCost", "Cost");
+	}
 		
 	virtual FSortPredicate GetSortPredicate() const override
 	{
@@ -319,7 +313,10 @@ struct FChannelBreakdownCostColumn : ITextColumn
 	{
 	}
 
-	virtual FText GetLabel() const override { return Label; }
+	virtual FText GetLabel() const override
+	{
+		return Label;
+	}
 
 	virtual FSortPredicate GetSortPredicate() const override
 	{
@@ -353,7 +350,10 @@ struct FCostModifier : ITextColumn
 {
 	using ITextColumn::ITextColumn;
 
-	virtual FText GetLabel() const override { return LOCTEXT("ColumnLabelCostModifier", "Bias"); }
+	virtual FText GetLabel() const override
+	{
+		return LOCTEXT("ColumnLabelCostModifier", "Bias");
+	}
 
 	virtual FSortPredicate GetSortPredicate() const override
 	{
@@ -371,7 +371,10 @@ struct FMirrored : ITextColumn
 {
 	using ITextColumn::ITextColumn;
 
-	virtual FText GetLabel() const override { return LOCTEXT("ColumnLabelMirrored", "Mirror"); }
+	virtual FText GetLabel() const override
+	{
+		return LOCTEXT("ColumnLabelMirrored", "Mirror");
+	}
 
 	virtual FSortPredicate GetSortPredicate() const override
 	{
@@ -388,7 +391,10 @@ struct FLooping : ITextColumn
 {
 	using ITextColumn::ITextColumn;
 
-	virtual FText GetLabel() const override { return LOCTEXT("ColumnLabelLooping", "Loop"); }
+	virtual FText GetLabel() const override
+	{
+		return LOCTEXT("ColumnLabelLooping", "Loop");
+	}
 
 	virtual FSortPredicate GetSortPredicate() const override
 	{
@@ -405,7 +411,10 @@ struct FBlendParameters : ITextColumn
 {
 	using ITextColumn::ITextColumn;
 
-	virtual FText GetLabel() const override { return LOCTEXT("ColumnLabelBlendParams", "Blend Params"); }
+	virtual FText GetLabel() const override
+	{
+		return LOCTEXT("ColumnLabelBlendParams", "Blend Params");
+	}
 
 	virtual FSortPredicate GetSortPredicate() const override
 	{
@@ -419,17 +428,7 @@ struct FBlendParameters : ITextColumn
 
 	virtual FText GetRowText(const FRowDataRef& Row) const override
 	{
-		if (Row->AssetType == ESearchIndexAssetType::BlendSpace)
-		{
-			return FText::Format(
-				LOCTEXT("Blend Parameters", "({0}, {1})"),
-				FText::AsNumber(Row->BlendParameters[0]),
-				FText::AsNumber(Row->BlendParameters[1]));
-		}
-		else
-		{
-			return FText::FromString(TEXT("-"));
-		}
+		return FText::Format(LOCTEXT("Blend Parameters", "{0}, {1}"), FText::AsNumber(Row->BlendParameters[0]), FText::AsNumber(Row->BlendParameters[1]));
 	}
 };
 
@@ -437,7 +436,10 @@ struct FPoseCandidateFlags : ITextColumn
 {
 	using ITextColumn::ITextColumn;
 
-	virtual FText GetLabel() const override { return LOCTEXT("ColumnLabelPoseCandidateFlags", "Flags"); }
+	virtual FText GetLabel() const override
+	{
+		return LOCTEXT("ColumnLabelPoseCandidateFlags", "Flags");
+	}
 
 	virtual FSortPredicate GetSortPredicate() const override
 	{
diff --git a/Engine/Plugins/Experimental/Animation/PoseSearch/Source/Editor/Private/PoseSearchDebuggerDatabaseRowData.h b/Engine/Plugins/Experimental/Animation/PoseSearch/Source/Editor/Private/PoseSearchDebuggerDatabaseRowData.h
--- a/Engine/Plugins/Experimental/Animation/PoseSearch/Source/Editor/Private/PoseSearchDebuggerDatabaseRowData.h
+++ b/Engine/Plugins/Experimental/Animation/PoseSearch/Source/Editor/Private/PoseSearchDebuggerDatabaseRowData.h
@@ -14,7 +14,6 @@ class FDebuggerDatabaseRowData : public TSharedFromThis<FDebuggerDatabaseRowData
 public:
 	FDebuggerDatabaseRowData() = default;
 	
-	ESearchIndexAssetType AssetType = ESearchIndexAssetType::Invalid;
 	int32 PoseIdx = 0;
 	TWeakObjectPtr<const UPoseSearchDatabase> SourceDatabase = nullptr;
 	EPoseCandidateFlags PoseCandidateFlags = EPoseCandidateFlags::None;
diff --git a/Engine/Plugins/Experimental/Animation/PoseSearch/Source/Editor/Private/PoseSearchDebuggerDatabaseView.cpp b/Engine/Plugins/Experimental/Animation/PoseSearch/Source/Editor/Private/PoseSearchDebuggerDatabaseView.cpp
--- a/Engine/Plugins/Experimental/Animation/PoseSearch/Source/Editor/Private/PoseSearchDebuggerDatabaseView.cpp
+++ b/Engine/Plugins/Experimental/Animation/PoseSearch/Source/Editor/Private/PoseSearchDebuggerDatabaseView.cpp
@@ -3,6 +3,7 @@
 #include "PoseSearchDebuggerDatabaseView.h"
 #include "Algo/AllOf.h"
 #include "Animation/AnimComposite.h"
+#include "Animation/AnimMontage.h"
 #include "Animation/AnimSequence.h"
 #include "InstancedStruct.h"
 #include "Internationalization/Regex.h"
@@ -182,7 +183,7 @@ void SDebuggerDatabaseView::Update(const FTraceMotionMatchingStateMessage& State
 				{
 					TSharedRef<FDebuggerDatabaseRowData>& Row = UnfilteredDatabaseRows.Add_GetRef(MakeShared<FDebuggerDatabaseRowData>());
 
-					const float Time = SearchIndex.GetAssetTime(PoseEntry.DbPoseIdx, Database->Schema->GetSamplingInterval());
+					const float Time = Database->GetAssetTime(PoseEntry.DbPoseIdx, Database->Schema->GetSamplingInterval());
 
 					Row->PoseIdx = PoseEntry.DbPoseIdx;
 					Row->SourceDatabase = Database;
@@ -207,7 +208,6 @@ void SDebuggerDatabaseView::Update(const FTraceMotionMatchingStateMessage& State
 					{
 						Row->AssetName = DatabaseAsset->GetName();
 						Row->AssetPath = DatabaseAsset->GetAnimationAsset() ? DatabaseAsset->GetAnimationAsset()->GetPathName() : "";
-						Row->AssetType = DatabaseAsset->GetSearchIndexType();
 						Row->bLooping = DatabaseAsset->IsLooping();
 						Row->BlendParameters = FVector::Zero();
 						Row->AnimFrame = 0;
@@ -227,6 +227,11 @@ void SDebuggerDatabaseView::Update(const FTraceMotionMatchingStateMessage& State
 						{
 							Row->BlendParameters = SearchIndexAsset->BlendParameters;
 						}
+						else if (const FPoseSearchDatabaseAnimMontage* DatabaseAnimMontage = DatabaseAssetStruct.GetPtr<FPoseSearchDatabaseAnimMontage>())
+						{
+							Row->AnimFrame = DatabaseAnimMontage->AnimMontage->GetFrameAtTime(Time);
+							Row->AnimPercentage = Time / DatabaseAnimMontage->AnimMontage->GetPlayLength();
+						}
 						else
 						{
 							checkNoEntry();
diff --git a/Engine/Plugins/Experimental/Animation/PoseSearch/Source/Editor/Private/PoseSearchDebuggerViewModel.cpp b/Engine/Plugins/Experimental/Animation/PoseSearch/Source/Editor/Private/PoseSearchDebuggerViewModel.cpp
--- a/Engine/Plugins/Experimental/Animation/PoseSearch/Source/Editor/Private/PoseSearchDebuggerViewModel.cpp
+++ b/Engine/Plugins/Experimental/Animation/PoseSearch/Source/Editor/Private/PoseSearchDebuggerViewModel.cpp
@@ -74,8 +74,7 @@ void FDebuggerViewModel::ShowSelectedSkeleton(const UPoseSearchDatabase* Databas
 	
 	Component->ResetToStart(); 
 	bSelecting = true;
-	
-	Skeletons[SelectedPose].Type = IndexAsset.Type;
+
 	Skeletons[SelectedPose].Time = Time;
 	Skeletons[SelectedPose].bMirrored = IndexAsset.bMirrored;
 	Skeletons[SelectedPose].SourceDatabase = Database;
@@ -161,7 +160,6 @@ void FDebuggerViewModel::OnUpdateNodeSelection(int32 InNodeId)
 			int32 CurrentPoseIdx = ActiveMotionMatchingState->GetCurrentDatabasePoseIndex();
 			if (const FPoseSearchIndexAsset* IndexAsset = CurrentSearchIndex.GetAssetForPoseSafe(CurrentPoseIdx))
 			{
-				Skeletons[Asset].Type = IndexAsset->Type;
 				Skeletons[Asset].bMirrored = IndexAsset->bMirrored;
 				Skeletons[Asset].SourceDatabase = CurrentDatabase;
 				Skeletons[Asset].AssetIdx = IndexAsset->SourceAssetIdx;
@@ -186,7 +184,6 @@ void FDebuggerViewModel::UpdatePoseSearchContext(UPoseSearchMeshComponent::FUpda
 		const FPoseSearchDatabaseAnimationAssetBase* DatabaseAsset = DatabaseAssetStruct->GetPtr<FPoseSearchDatabaseAnimationAssetBase>();
 		if (DatabaseAsset)
 		{
-			InOutContext.Type = DatabaseAsset->GetSearchIndexType();
 			InOutContext.StartTime = Skeletons[SelectedPose].Time;
 			InOutContext.Time = Skeletons[SelectedPose].Time;
 			InOutContext.bMirrored = Skeletons[SelectedPose].bMirrored;
@@ -255,7 +252,7 @@ void FDebuggerViewModel::OnDraw(FSkeletonDrawParams& DrawParams)
 		{
 			UpdatePoseSearchContext(UpdateContext, Skeletons[SelectedPose]);
 
-			if (UpdateContext.Type != ESearchIndexAssetType::Invalid)
+			if (UpdateContext.SequenceBase || UpdateContext.BlendSpace)
 			{
 				Component->UpdatePose(UpdateContext);
 			}
@@ -272,7 +269,7 @@ void FDebuggerViewModel::OnDraw(FSkeletonDrawParams& DrawParams)
 
 			UpdatePoseSearchContext(UpdateContext, Skeletons[SelectedPose]);
 
-			if (UpdateContext.Type != ESearchIndexAssetType::Invalid)
+			if (UpdateContext.SequenceBase || UpdateContext.BlendSpace)
 			{
 				Component->UpdatePose(UpdateContext);
 			}
@@ -503,7 +500,6 @@ void FDebuggerViewModel::PlaySelection(int32 PoseIdx, float Time)
 	const FPoseSearchIndexAsset& IndexAsset = Database->GetSearchIndex().GetAssetForPose(PoseIdx);
 	Component->ResetToStart();
 	
-	Skeletons[Asset].Type = IndexAsset.Type;
 	Skeletons[Asset].AssetIdx = IndexAsset.SourceAssetIdx;
 	Skeletons[Asset].Time = Time;
 	Skeletons[Asset].bMirrored = IndexAsset.bMirrored;
diff --git a/Engine/Plugins/Experimental/Animation/PoseSearch/Source/Editor/Private/PoseSearchDebuggerViewModel.h b/Engine/Plugins/Experimental/Animation/PoseSearch/Source/Editor/Private/PoseSearchDebuggerViewModel.h
--- a/Engine/Plugins/Experimental/Animation/PoseSearch/Source/Editor/Private/PoseSearchDebuggerViewModel.h
+++ b/Engine/Plugins/Experimental/Animation/PoseSearch/Source/Editor/Private/PoseSearchDebuggerViewModel.h
@@ -102,9 +102,6 @@ private:
 		/** Derived skeletal mesh for setting the skeleton in the scene */
 		TWeakObjectPtr<UPoseSearchMeshComponent> Component = nullptr;
 
-		/** Type of the asset being played */
-		ESearchIndexAssetType Type = ESearchIndexAssetType::Invalid;
-
 		/** Source database for this skeleton  */
 		TWeakObjectPtr<const UPoseSearchDatabase> SourceDatabase = nullptr;
 
diff --git a/Engine/Plugins/Experimental/Animation/PoseSearch/Source/Editor/Private/PoseSearchMeshComponent.cpp b/Engine/Plugins/Experimental/Animation/PoseSearch/Source/Editor/Private/PoseSearchMeshComponent.cpp
--- a/Engine/Plugins/Experimental/Animation/PoseSearch/Source/Editor/Private/PoseSearchMeshComponent.cpp
+++ b/Engine/Plugins/Experimental/Animation/PoseSearch/Source/Editor/Private/PoseSearchMeshComponent.cpp
@@ -59,8 +59,7 @@ void UPoseSearchMeshComponent::UpdatePose(const FUpdateContext& UpdateContext)
 	UE::Anim::FStackAttributeContainer Attributes;
 	FAnimationPoseData PoseData(CompactPose, Curve, Attributes);
 
-	if (UpdateContext.Type == ESearchIndexAssetType::Sequence ||
-		UpdateContext.Type == ESearchIndexAssetType::AnimComposite)
+	if (UpdateContext.SequenceBase)
 	{
 		float AdvancedTime = UpdateContext.StartTime;
 
@@ -75,7 +74,7 @@ void UPoseSearchMeshComponent::UpdatePose(const FUpdateContext& UpdateContext)
 
 		UpdateContext.SequenceBase->GetAnimationPose(PoseData, ExtractionCtx);
 	}
-	else if (UpdateContext.Type == ESearchIndexAssetType::BlendSpace)
+	else if (UpdateContext.BlendSpace)
 	{
 		TArray<FBlendSampleData> BlendSamples;
 		int32 TriangulationIndex = 0;
diff --git a/Engine/Plugins/Experimental/Animation/PoseSearch/Source/Editor/Private/PoseSearchMeshComponent.h b/Engine/Plugins/Experimental/Animation/PoseSearch/Source/Editor/Private/PoseSearchMeshComponent.h
--- a/Engine/Plugins/Experimental/Animation/PoseSearch/Source/Editor/Private/PoseSearchMeshComponent.h
+++ b/Engine/Plugins/Experimental/Animation/PoseSearch/Source/Editor/Private/PoseSearchMeshComponent.h
@@ -17,7 +17,6 @@ public:
 
 	struct FUpdateContext
 	{
-		ESearchIndexAssetType Type = ESearchIndexAssetType::Invalid;
 		const UAnimSequenceBase* SequenceBase = nullptr;
 		const UBlendSpace* BlendSpace = nullptr;
 		float StartTime = 0.0f;
diff --git a/Engine/Plugins/Experimental/Animation/PoseSearch/Source/Editor/Public/PoseSearchDatabaseAssetListItem.h b/Engine/Plugins/Experimental/Animation/PoseSearch/Source/Editor/Public/PoseSearchDatabaseAssetListItem.h
--- a/Engine/Plugins/Experimental/Animation/PoseSearch/Source/Editor/Public/PoseSearchDatabaseAssetListItem.h
+++ b/Engine/Plugins/Experimental/Animation/PoseSearch/Source/Editor/Public/PoseSearchDatabaseAssetListItem.h
@@ -36,6 +36,7 @@ namespace UE::PoseSearch
 		void OnAddSequence();
 		void OnAddBlendSpace();
 		void OnAddAnimComposite();
+		void OnAddAnimMontage();
 
 		EVisibility GetSelectedActorIconVisbility() const;
 
diff --git a/Engine/Plugins/Experimental/Animation/PoseSearch/Source/Editor/Public/PoseSearchDatabaseAssetTree.h b/Engine/Plugins/Experimental/Animation/PoseSearch/Source/Editor/Public/PoseSearchDatabaseAssetTree.h
--- a/Engine/Plugins/Experimental/Animation/PoseSearch/Source/Editor/Public/PoseSearchDatabaseAssetTree.h
+++ b/Engine/Plugins/Experimental/Animation/PoseSearch/Source/Editor/Public/PoseSearchDatabaseAssetTree.h
@@ -88,6 +88,7 @@ namespace UE::PoseSearch
 		void OnAddSequence(bool bFinalizeChanges = true);
 		void OnAddBlendSpace(bool bFinalizeChanges = true);
 		void OnAddAnimComposite(bool bFinalizeChanges = true);
+		void OnAddAnimMontage(bool bFinalizeChanges = true);
 
 		void OnDeleteAsset(TSharedPtr<FDatabaseAssetTreeNode> Node, bool bFinalizeChanges = true);
 		void CreateCommandList();
diff --git a/Engine/Plugins/Experimental/Animation/PoseSearch/Source/Editor/Public/PoseSearchDatabaseAssetTreeNode.h b/Engine/Plugins/Experimental/Animation/PoseSearch/Source/Editor/Public/PoseSearchDatabaseAssetTreeNode.h
--- a/Engine/Plugins/Experimental/Animation/PoseSearch/Source/Editor/Public/PoseSearchDatabaseAssetTreeNode.h
+++ b/Engine/Plugins/Experimental/Animation/PoseSearch/Source/Editor/Public/PoseSearchDatabaseAssetTreeNode.h
@@ -5,7 +5,6 @@
 #include "CoreMinimal.h"
 
 enum class EPoseSearchMirrorOption;
-enum class ESearchIndexAssetType;
 class FUICommandList;
 class ITableRow;
 class STableViewBase;
@@ -19,23 +18,13 @@ namespace UE::PoseSearch
 	{
 
 	public:
-		FDatabaseAssetTreeNode(
-			int32 InSourceAssetIdx,
-			ESearchIndexAssetType InSourceAssetType,
-			const TSharedRef<FDatabaseViewModel>& InEditorViewModel);
-
-		TSharedRef<ITableRow> MakeTreeRowWidget(
-			const TSharedRef<STableViewBase>& InOwnerTable,
-			TSharedRef<FDatabaseAssetTreeNode> InDatabaseAssetNode,
-			TSharedRef<FUICommandList> InCommandList,
-			TSharedPtr<SDatabaseAssetTree> InHierarchy);
-
+		FDatabaseAssetTreeNode(int32 InSourceAssetIdx, const TSharedRef<FDatabaseViewModel>& InEditorViewModel);
+		TSharedRef<ITableRow> MakeTreeRowWidget(const TSharedRef<STableViewBase>& InOwnerTable, TSharedRef<FDatabaseAssetTreeNode> InDatabaseAssetNode, TSharedRef<FUICommandList> InCommandList, TSharedPtr<SDatabaseAssetTree> InHierarchy);
 		bool IsRootMotionEnabled() const;
 		bool IsLooping() const;
 		EPoseSearchMirrorOption GetMirrorOption() const;
 		
-		int32 SourceAssetIdx;
-		ESearchIndexAssetType SourceAssetType;
+		int32 SourceAssetIdx = INDEX_NONE;
 		TSharedPtr<FDatabaseAssetTreeNode> Parent;
 		TArray<TSharedPtr<FDatabaseAssetTreeNode>> Children;
 		TWeakPtr<FDatabaseViewModel> EditorViewModel;
diff --git a/Engine/Plugins/Experimental/Animation/PoseSearch/Source/Editor/Public/PoseSearchDatabaseEditorReflection.h b/Engine/Plugins/Experimental/Animation/PoseSearch/Source/Editor/Public/PoseSearchDatabaseEditorReflection.h
--- a/Engine/Plugins/Experimental/Animation/PoseSearch/Source/Editor/Public/PoseSearchDatabaseEditorReflection.h
+++ b/Engine/Plugins/Experimental/Animation/PoseSearch/Source/Editor/Public/PoseSearchDatabaseEditorReflection.h
@@ -71,6 +71,20 @@ public:
 #endif
 };
 
+UCLASS()
+class UPoseSearchDatabaseAnimMontageReflection : public UPoseSearchDatabaseReflectionBase
+{
+	GENERATED_BODY()
+
+public:
+	UPROPERTY(EditAnywhere, Category = "Selected Anim Montage")
+	FPoseSearchDatabaseAnimMontage AnimMontage;
+
+#if WITH_EDITOR
+	virtual void PostEditChangeProperty(struct FPropertyChangedEvent& PropertyChangedEvent) override;
+#endif
+};
+
 USTRUCT()
 struct FPoseSearchDatabaseMemoryStats
 {
diff --git a/Engine/Plugins/Experimental/Animation/PoseSearch/Source/Editor/Public/PoseSearchDatabaseViewModel.h b/Engine/Plugins/Experimental/Animation/PoseSearch/Source/Editor/Public/PoseSearchDatabaseViewModel.h
--- a/Engine/Plugins/Experimental/Animation/PoseSearch/Source/Editor/Public/PoseSearchDatabaseViewModel.h
+++ b/Engine/Plugins/Experimental/Animation/PoseSearch/Source/Editor/Public/PoseSearchDatabaseViewModel.h
@@ -3,11 +3,12 @@
 #pragma once
 
 #include "CoreMinimal.h"
-#include "UObject/GCObject.h"
+#include "InstancedStruct.h"
 #include "MovieSceneFwd.h"
 #include "PoseSearch/PoseSearchAssetSampler.h"
 #include "PoseSearch/PoseSearchIndex.h"
 #include "PoseSearchDatabasePreviewScene.h"
+#include "UObject/GCObject.h"
 
 class UWorld;
 class UPoseSearchDatabase;
@@ -46,13 +47,11 @@ namespace UE::PoseSearch
 		int32 IndexAssetIndex = INDEX_NONE;
 		int32 CurrentPoseIndex = INDEX_NONE;
 
-		FSequenceBaseSampler SequenceSampler;
-		FBlendSpaceSampler BlendSpaceSampler;
-		ESearchIndexAssetType Type = ESearchIndexAssetType::Invalid;
+		FInstancedStruct Sampler;
 
 		bool IsValid() const;
 		void Process();
-		const IAssetSampler* GetSampler() const;
+		const FAssetSamplerBase* GetSampler() const;
 		float GetScaledTime(float Time) const;
 		UDebugSkelMeshComponent* GetDebugSkelMeshComponent();
 		UAnimPreviewInstance* GetAnimPreviewInstance();
@@ -104,6 +103,7 @@ namespace UE::PoseSearch
 		void AddSequenceToDatabase(UAnimSequence* AnimSequence);
 		void AddBlendSpaceToDatabase(UBlendSpace* BlendSpace);
 		void AddAnimCompositeToDatabase(UAnimComposite* AnimComposite);
+		void AddAnimMontageToDatabase(UAnimMontage* AnimMontage);
 		void DeleteFromDatabase(int32 AnimationAssetIndex);
 
 		void SetIsEnabled(int32 AnimationAssetIndex, bool bEnabled);
diff --git a/Engine/Plugins/Experimental/Animation/PoseSearch/Source/Runtime/Private/AnimNode_MotionMatching.cpp b/Engine/Plugins/Experimental/Animation/PoseSearch/Source/Runtime/Private/AnimNode_MotionMatching.cpp
--- a/Engine/Plugins/Experimental/Animation/PoseSearch/Source/Runtime/Private/AnimNode_MotionMatching.cpp
+++ b/Engine/Plugins/Experimental/Animation/PoseSearch/Source/Runtime/Private/AnimNode_MotionMatching.cpp
@@ -109,7 +109,7 @@ void FAnimNode_MotionMatching::UpdateAssetPlayer(const FAnimationUpdateContext&
 	UpdateCounter.SynchronizeWith(Context.AnimInstanceProxy->GetUpdateCounter());
 
 	// Execute core motion matching algorithm
-	UpdateMotionMatchingState(
+	UPoseSearchLibrary::UpdateMotionMatchingState(
 		Context,
 		Searchable,
 		&ActiveTagsContainer,
diff --git a/Engine/Plugins/Experimental/Animation/PoseSearch/Source/Runtime/Private/PoseSearchAssetSampler.cpp b/Engine/Plugins/Experimental/Animation/PoseSearch/Source/Runtime/Private/PoseSearchAssetSampler.cpp
--- a/Engine/Plugins/Experimental/Animation/PoseSearch/Source/Runtime/Private/PoseSearchAssetSampler.cpp
+++ b/Engine/Plugins/Experimental/Animation/PoseSearch/Source/Runtime/Private/PoseSearchAssetSampler.cpp
@@ -2,11 +2,13 @@
 
 #include "PoseSearch/PoseSearchAssetSampler.h"
 #include "AnimationRuntime.h"
+#include "Animation/AnimMontage.h"
 #include "Animation/AnimSequence.h"
 #include "Animation/AnimRootMotionProvider.h"
 #include "Animation/BlendSpace.h"
 #include "Animation/MirrorDataTable.h"
 #include "PoseSearch/PoseSearchAnimNotifies.h"
+#include "PoseSearch/PoseSearchDefines.h"
 
 namespace UE::PoseSearch
 {
@@ -30,8 +32,7 @@ static FTransform ExtrapolateRootMotion(
 		LinearSpeedToExtrapolate >= ExtrapolationParameters.LinearSpeedThreshold;
 
 	const float AngularSpeedToExtrapolateRad = SampleToExtrapolate.GetRotation().GetAngle() / SampleDelta;
-	const bool bCanExtrapolateRotation =
-		FMath::RadiansToDegrees(AngularSpeedToExtrapolateRad) >= ExtrapolationParameters.AngularSpeedThreshold;
+	const bool bCanExtrapolateRotation = FMath::RadiansToDegrees(AngularSpeedToExtrapolateRad) >= ExtrapolationParameters.AngularSpeedThreshold;
 
 	if (!bCanExtrapolateTranslation && !bCanExtrapolateRotation)
 	{
@@ -52,8 +53,7 @@ static FTransform ExtrapolateRootMotion(
 	// transforms later on.
 	const float AbsExtrapolationTime = FMath::Abs(ExtrapolationTime);
 	const float AbsSampleDelta = FMath::Abs(SampleDelta);
-	const FTransform AbsTimeSampleToExtrapolate =
-		ExtrapolationTime >= 0.0f ? SampleToExtrapolate : SampleToExtrapolate.Inverse();
+	const FTransform AbsTimeSampleToExtrapolate = ExtrapolationTime >= 0.0f ? SampleToExtrapolate : SampleToExtrapolate.Inverse();
 
 	// because we're extrapolating rotation, the extrapolation must be integrated over time
 	const float SampleMultiplier = AbsExtrapolationTime / AbsSampleDelta;
@@ -109,12 +109,13 @@ FTransform FAssetSamplingContext::MirrorTransform(const FTransform& InTransform)
 	return Result;
 }
 
+} // namespace UE::PoseSearch
+
 //////////////////////////////////////////////////////////////////////////
 // FSequenceBaseSampler
 void FSequenceBaseSampler::Init(const FInput& InInput)
 {
 	check(InInput.SequenceBase.Get());
-
 	Input = InInput;
 }
 
@@ -127,6 +128,11 @@ float FSequenceBaseSampler::GetPlayLength() const
 	return Input.SequenceBase->GetPlayLength();
 }
 
+float FSequenceBaseSampler::GetScaledTime(float Time) const
+{
+	return Time;
+}
+
 bool FSequenceBaseSampler::IsLoopable() const
 {
 	return Input.SequenceBase->bLoop;
@@ -168,7 +174,7 @@ FTransform FSequenceBaseSampler::ExtractRootTransform(float Time) const
 		{
 			FTransform SampleToExtrapolate = SequenceBase->ExtractRootMotionFromRange(0.0f, ExtrapolationSampleTime);
 
-			const FTransform ExtrapolatedRootMotion = ExtrapolateRootMotion(
+			const FTransform ExtrapolatedRootMotion = UE::PoseSearch::ExtrapolateRootMotion(
 				SampleToExtrapolate,
 				0.0f, ExtrapolationSampleTime,
 				ExtrapolationTime,
@@ -179,13 +185,13 @@ FTransform FSequenceBaseSampler::ExtractRootTransform(float Time) const
 		{
 			RootTransform = SequenceBase->ExtractRootMotionFromRange(0.0f, ClampedTime);
 
-			// If Time is greater than PlayLength, ExtrapolationTIme will be a positive number. In this case, we extrapolate
+			// If Time is greater than PlayLength, ExtrapolationTime will be a positive number. In this case, we extrapolate
 			// the end of the animation to estimate where the root would be at Time
 			if (ExtrapolationTime > SMALL_NUMBER)
 			{
 				FTransform SampleToExtrapolate = SequenceBase->ExtractRootMotionFromRange(PlayLength - ExtrapolationSampleTime, PlayLength);
 
-				const FTransform ExtrapolatedRootMotion = ExtrapolateRootMotion(
+				const FTransform ExtrapolatedRootMotion = UE::PoseSearch::ExtrapolateRootMotion(
 					SampleToExtrapolate,
 					PlayLength - ExtrapolationSampleTime, PlayLength,
 					ExtrapolationTime,
@@ -241,7 +247,6 @@ const UAnimationAsset* FSequenceBaseSampler::GetAsset() const
 void FBlendSpaceSampler::Init(const FInput& InInput)
 {
 	check(InInput.BlendSpace.Get());
-
 	Input = InInput;
 }
 
@@ -253,6 +258,12 @@ void FBlendSpaceSampler::Process()
 	ProcessRootTransform();
 }
 
+float FBlendSpaceSampler::GetScaledTime(float Time) const
+{
+	const float ScaledTime = GetPlayLength() > UE_KINDA_SMALL_NUMBER ? Time / GetPlayLength() : 0.f;
+	return ScaledTime;
+}
+
 bool FBlendSpaceSampler::IsLoopable() const
 {
 	return Input.BlendSpace->bLoop;
@@ -308,7 +319,7 @@ FTransform FBlendSpaceSampler::ExtractRootTransform(float Time) const
 		{
 			FTransform SampleToExtrapolate = ExtractBlendSpaceRootMotionFromRange(0.0f, ExtrapolationSampleTime);
 
-			const FTransform ExtrapolatedRootMotion = ExtrapolateRootMotion(
+			const FTransform ExtrapolatedRootMotion = UE::PoseSearch::ExtrapolateRootMotion(
 				SampleToExtrapolate,
 				0.0f, ExtrapolationSampleTime,
 				ExtrapolationTime,
@@ -319,13 +330,13 @@ FTransform FBlendSpaceSampler::ExtractRootTransform(float Time) const
 		{
 			RootTransform = ExtractBlendSpaceRootMotionFromRange(0.0f, ClampedTime);
 
-			// If Time is greater than PlayLength, ExtrapolationTIme will be a positive number. In this case, we extrapolate
+			// If Time is greater than PlayLength, ExtrapolationTime will be a positive number. In this case, we extrapolate
 			// the end of the animation to estimate where the root would be at Time
 			if (ExtrapolationTime > SMALL_NUMBER)
 			{
 				FTransform SampleToExtrapolate = ExtractBlendSpaceRootMotionFromRange(PlayLength - ExtrapolationSampleTime, PlayLength);
 
-				const FTransform ExtrapolatedRootMotion = ExtrapolateRootMotion(
+				const FTransform ExtrapolatedRootMotion = UE::PoseSearch::ExtrapolateRootMotion(
 					SampleToExtrapolate,
 					PlayLength - ExtrapolationSampleTime, PlayLength,
 					ExtrapolationTime,
@@ -525,7 +536,7 @@ void FBlendSpaceSampler::ProcessRootTransform()
 
 		FCompactPose Pose;
 		FBlendedCurve BlendedCurve;
-		Anim::FStackAttributeContainer StackAttributeContainer;
+		UE::Anim::FStackAttributeContainer StackAttributeContainer;
 		FAnimationPoseData AnimPoseData(Pose, BlendedCurve, StackAttributeContainer);
 
 		Pose.SetBoneContainer(&Input.BoneContainer);
@@ -533,7 +544,7 @@ void FBlendSpaceSampler::ProcessRootTransform()
 
 		Input.BlendSpace->GetAnimationPose(BlendSamples, ExtractionCtx, AnimPoseData);
 
-		const Anim::IAnimRootMotionProvider* RootMotionProvider = Anim::IAnimRootMotionProvider::Get();
+		const UE::Anim::IAnimRootMotionProvider* RootMotionProvider = UE::Anim::IAnimRootMotionProvider::Get();
 
 		if (ensureMsgf(RootMotionProvider, TEXT("Could not get Root Motion Provider.")))
 		{
@@ -555,4 +566,168 @@ const UAnimationAsset* FBlendSpaceSampler::GetAsset() const
 	return Input.BlendSpace.Get();
 }
 
-} // namespace UE::PoseSearch
+//////////////////////////////////////////////////////////////////////////
+// FAnimMontageSampler
+void FAnimMontageSampler::Init(const FInput& InInput)
+{
+	check(InInput.AnimMontage.Get());
+	Input = InInput;
+}
+
+void FAnimMontageSampler::Process()
+{
+}
+
+float FAnimMontageSampler::GetPlayLength() const
+{
+	return Input.AnimMontage->GetPlayLength();
+}
+
+float FAnimMontageSampler::GetScaledTime(float Time) const
+{
+	return Time;
+}
+
+bool FAnimMontageSampler::IsLoopable() const
+{
+	return Input.AnimMontage->bLoop;
+}
+
+void FAnimMontageSampler::ExtractPose(const FAnimExtractContext& ExtractionCtx, FAnimationPoseData& OutAnimPoseData) const
+{
+	// @todo: add support for SlotName / multiple SlotAnimTracks
+	if (Input.AnimMontage->SlotAnimTracks.Num() != 1)
+	{
+		UE_LOG(LogPoseSearch, Error, TEXT("FAnimMontageSampler::ExtractPose: so far we support only montages with one SlotAnimTracks. %s has %d"), *Input.AnimMontage->GetName(), Input.AnimMontage->SlotAnimTracks.Num());
+		OutAnimPoseData.GetPose().ResetToRefPose();
+		return;
+	}
+
+	Input.AnimMontage->SlotAnimTracks[0].AnimTrack.GetAnimationPose(OutAnimPoseData, ExtractionCtx);
+}
+
+FTransform FAnimMontageSampler::GetTotalRootTransform() const
+{
+	// @todo: add support for SlotName / multiple SlotAnimTracks
+	if (Input.AnimMontage->SlotAnimTracks.Num() != 1)
+	{
+		UE_LOG(LogPoseSearch, Error, TEXT("FAnimMontageSampler::GetTotalRootTransform: so far we support only montages with one SlotAnimTracks. %s has %d"), *Input.AnimMontage->GetName(), Input.AnimMontage->SlotAnimTracks.Num());
+		return FTransform::Identity;
+	}
+
+	// @todo: optimize me
+	const FTransform InitialRootTransform = ExtractRootTransform(0.f);
+	const FTransform LastRootTransform = ExtractRootTransform(GetPlayLength());
+	const FTransform TotalRootTransform = LastRootTransform.GetRelativeTransform(InitialRootTransform);
+	return TotalRootTransform;
+}
+
+FTransform FAnimMontageSampler::ExtractRootTransformInternal(float StartTime, float EndTime) const
+{
+	// @todo: add support for SlotName / multiple SlotAnimTracks
+	if (Input.AnimMontage->SlotAnimTracks.Num() != 1)
+	{
+		UE_LOG(LogPoseSearch, Error, TEXT("FAnimMontageSampler::ExtractRootTransform: so far we support only montages with one SlotAnimTracks. %s has %d"), *Input.AnimMontage->GetName(), Input.AnimMontage->SlotAnimTracks.Num());
+		return FTransform::Identity;
+	}
+
+	const FAnimTrack& RootMotionAnimTrack = Input.AnimMontage->SlotAnimTracks[0].AnimTrack;
+	TArray<FRootMotionExtractionStep> RootMotionExtractionSteps;
+	RootMotionAnimTrack.GetRootMotionExtractionStepsForTrackRange(RootMotionExtractionSteps, StartTime, EndTime);
+	FRootMotionMovementParams AccumulatedRootMotionParams;
+	for (const FRootMotionExtractionStep& CurStep : RootMotionExtractionSteps)
+	{
+		if (CurStep.AnimSequence)
+		{
+			AccumulatedRootMotionParams.Accumulate(CurStep.AnimSequence->ExtractRootMotionFromRange(CurStep.StartPosition, CurStep.EndPosition));
+		}
+	}
+	return AccumulatedRootMotionParams.GetRootMotionTransform();
+}
+
+FTransform FAnimMontageSampler::ExtractRootTransform(float Time) const
+{
+	FTransform RootTransform = FTransform::Identity;
+	if (IsLoopable())
+	{
+		RootTransform = ExtractRootTransformInternal(0.f, Time);
+	}
+	else
+	{
+		const float ExtrapolationSampleTime = Input.ExtrapolationParameters.SampleTime;
+
+		const float PlayLength = GetPlayLength();
+		const float ClampedTime = FMath::Clamp(Time, 0.f, PlayLength);
+		const float ExtrapolationTime = Time - ClampedTime;
+
+		// If Time is less than zero, ExtrapolationTime will be negative. In this case, we extrapolate the beginning of the 
+		// animation to estimate where the root would be at Time
+		if (ExtrapolationTime < -SMALL_NUMBER)
+		{
+			FTransform SampleToExtrapolate = ExtractRootTransformInternal(0.f, ExtrapolationSampleTime);
+
+			const FTransform ExtrapolatedRootMotion = UE::PoseSearch::ExtrapolateRootMotion(
+				SampleToExtrapolate,
+				0.0f, ExtrapolationSampleTime,
+				ExtrapolationTime,
+				Input.ExtrapolationParameters);
+			RootTransform = ExtrapolatedRootMotion;
+		}
+		else
+		{
+			RootTransform = ExtractRootTransformInternal(0.f, ClampedTime);
+
+			// If Time is greater than PlayLength, ExtrapolationTime will be a positive number. In this case, we extrapolate
+			// the end of the animation to estimate where the root would be at Time
+			if (ExtrapolationTime > SMALL_NUMBER)
+			{
+				FTransform SampleToExtrapolate = ExtractRootTransformInternal(PlayLength - ExtrapolationSampleTime, PlayLength);
+
+				const FTransform ExtrapolatedRootMotion = UE::PoseSearch::ExtrapolateRootMotion(
+					SampleToExtrapolate,
+					PlayLength - ExtrapolationSampleTime, PlayLength,
+					ExtrapolationTime,
+					Input.ExtrapolationParameters);
+				RootTransform = ExtrapolatedRootMotion * RootTransform;
+			}
+		}
+	}
+
+	return RootTransform;
+}
+
+void FAnimMontageSampler::ExtractPoseSearchNotifyStates(float Time, TArray<UAnimNotifyState_PoseSearchBase*>& NotifyStates) const
+{
+	// getting pose search notifies in an interval of size ExtractionInterval, centered on Time
+	constexpr float ExtractionInterval = 1.0f / 120.0f;
+	FAnimNotifyContext NotifyContext;
+	Input.AnimMontage->GetAnimNotifies(Time - (ExtractionInterval * 0.5f), ExtractionInterval, NotifyContext);
+
+	// check which notifies actually overlap Time and are of the right base type
+	for (const FAnimNotifyEventReference& EventReference : NotifyContext.ActiveNotifies)
+	{
+		const FAnimNotifyEvent* NotifyEvent = EventReference.GetNotify();
+		if (!NotifyEvent)
+		{
+			continue;
+		}
+
+		if (NotifyEvent->GetTriggerTime() > Time ||
+			NotifyEvent->GetEndTriggerTime() < Time)
+		{
+			continue;
+		}
+
+		UAnimNotifyState_PoseSearchBase* PoseSearchAnimNotify =
+			Cast<UAnimNotifyState_PoseSearchBase>(NotifyEvent->NotifyStateClass);
+		if (PoseSearchAnimNotify)
+		{
+			NotifyStates.Add(PoseSearchAnimNotify);
+		}
+	}
+}
+
+const UAnimationAsset* FAnimMontageSampler::GetAsset() const
+{
+	return Input.AnimMontage.Get();
+}
diff --git a/Engine/Plugins/Experimental/Animation/PoseSearch/Source/Runtime/Private/PoseSearchDatabase.cpp b/Engine/Plugins/Experimental/Animation/PoseSearch/Source/Runtime/Private/PoseSearchDatabase.cpp
--- a/Engine/Plugins/Experimental/Animation/PoseSearch/Source/Runtime/Private/PoseSearchDatabase.cpp
+++ b/Engine/Plugins/Experimental/Animation/PoseSearch/Source/Runtime/Private/PoseSearchDatabase.cpp
@@ -2,6 +2,7 @@
 
 #include "PoseSearch/PoseSearchDatabase.h"
 #include "Animation/AnimComposite.h"
+#include "Animation/AnimMontage.h"
 #include "Animation/AnimSequence.h"
 #include "Animation/BlendSpace.h"
 #include "Animation/BlendSpace1D.h"
@@ -242,6 +243,7 @@ UClass* FPoseSearchDatabaseSequence::GetAnimationAssetStaticClass() const
 
 bool FPoseSearchDatabaseSequence::IsLooping() const
 {
+	// @todo: we should check if against the sample range (look for GetEffectiveSamplingRange()...)
 	return Sequence ? Sequence->bLoop : false;
 }
 
@@ -372,6 +374,33 @@ bool FPoseSearchDatabaseAnimComposite::IsRootMotionEnabled() const
 	return AnimComposite ? AnimComposite->HasRootMotion() : false;
 }
 
+//////////////////////////////////////////////////////////////////////////
+// FPoseSearchDatabaseAnimMontage
+UAnimationAsset* FPoseSearchDatabaseAnimMontage::GetAnimationAsset() const
+{
+	return AnimMontage;
+}
+
+UClass* FPoseSearchDatabaseAnimMontage::GetAnimationAssetStaticClass() const
+{
+	return UAnimMontage::StaticClass();
+}
+
+bool FPoseSearchDatabaseAnimMontage::IsLooping() const
+{
+	return AnimMontage ? AnimMontage->bLoop : false;
+}
+
+const FString FPoseSearchDatabaseAnimMontage::GetName() const
+{
+	return AnimMontage ? AnimMontage->GetName() : FString();
+}
+
+bool FPoseSearchDatabaseAnimMontage::IsRootMotionEnabled() const
+{
+	return AnimMontage ? AnimMontage->HasRootMotion() : false;
+}
+
 //////////////////////////////////////////////////////////////////////////
 // UPoseSearchDatabase
 UPoseSearchDatabase::~UPoseSearchDatabase()
@@ -625,6 +654,25 @@ void UPoseSearchDatabase::Serialize(FArchive& Ar)
 	}
 }
 
+float UPoseSearchDatabase::GetAssetTime(int32 PoseIdx, float SamplingInterval) const
+{
+	const FPoseSearchIndexAsset& Asset = GetSearchIndex().GetAssetForPose(PoseIdx);
+	const bool bIsBlendSpace = AnimationAssets[Asset.SourceAssetIdx].GetPtr<FPoseSearchDatabaseBlendSpace>() != nullptr;
+	const FFloatInterval& SamplingRange = Asset.SamplingInterval;
+
+	if (bIsBlendSpace)
+	{
+		// For BlendSpaces the AssetTime is in the range [0, 1] while the Sampling Range
+		// is in real time (seconds)
+		const float AssetTime = FMath::Min(SamplingRange.Min + SamplingInterval * (PoseIdx - Asset.FirstPoseIdx), SamplingRange.Max) / (Asset.NumPoses * SamplingInterval);
+		return AssetTime;
+	}
+
+	// sequences or anim composites
+	const float AssetTime = FMath::Min(SamplingRange.Min + SamplingInterval * (PoseIdx - Asset.FirstPoseIdx), SamplingRange.Max);
+	return AssetTime;
+}
+
 UE::PoseSearch::FSearchResult UPoseSearchDatabase::Search(UE::PoseSearch::FSearchContext& SearchContext) const
 {
 	using namespace UE::PoseSearch;
@@ -801,7 +849,7 @@ UE::PoseSearch::FSearchResult UPoseSearchDatabase::SearchPCAKDTree(UE::PoseSearc
 	// finalizing Result properties
 	if (Result.PoseIdx != INDEX_NONE)
 	{
-		Result.AssetTime = SearchIndex.GetAssetTime(Result.PoseIdx, Schema->GetSamplingInterval());
+		Result.AssetTime = GetAssetTime(Result.PoseIdx, Schema->GetSamplingInterval());
 		Result.Database = this;
 	}
 
@@ -889,7 +937,7 @@ UE::PoseSearch::FSearchResult UPoseSearchDatabase::SearchBruteForce(UE::PoseSear
 	// finalizing Result properties
 	if (Result.PoseIdx != INDEX_NONE)
 	{
-		Result.AssetTime = SearchIndex.GetAssetTime(Result.PoseIdx, Schema->GetSamplingInterval());
+		Result.AssetTime = GetAssetTime(Result.PoseIdx, Schema->GetSamplingInterval());
 		Result.Database = this;
 	}
 
diff --git a/Engine/Plugins/Experimental/Animation/PoseSearch/Source/Runtime/Private/PoseSearchDatabaseIndexingContext.cpp b/Engine/Plugins/Experimental/Animation/PoseSearch/Source/Runtime/Private/PoseSearchDatabaseIndexingContext.cpp
--- a/Engine/Plugins/Experimental/Animation/PoseSearch/Source/Runtime/Private/PoseSearchDatabaseIndexingContext.cpp
+++ b/Engine/Plugins/Experimental/Animation/PoseSearch/Source/Runtime/Private/PoseSearchDatabaseIndexingContext.cpp
@@ -5,6 +5,7 @@
 #if WITH_EDITOR
 
 #include "Animation/AnimComposite.h"
+#include "Animation/AnimMontage.h"
 #include "Animation/AnimSequence.h"
 #include "Animation/BlendSpace.h"
 #include "InstancedStruct.h"
@@ -23,7 +24,7 @@ void FDatabaseIndexingContext::Prepare(const UPoseSearchDatabase* Database)
 	FBoneContainer BoneContainer;
 	BoneContainer.InitializeTo(Schema->BoneIndicesWithParents, FCurveEvaluationOption(false), *Schema->Skeleton);
 
-	TMap<const UAnimSequenceBase*, int32> SequenceSamplerMap;
+	TMap<const UAnimationAsset*, int32> SamplerMap;
 	TMap<TPair<const UBlendSpace*, FVector>, int32> BlendSpaceSamplerMap;
 
 	SamplingContext.Init(Schema->MirrorDataTable, BoneContainer);
@@ -33,15 +34,16 @@ void FDatabaseIndexingContext::Prepare(const UPoseSearchDatabase* Database)
 	{
 		auto AddSequenceBaseSampler = [&](const UAnimSequenceBase* Sequence)
 		{
-			if (Sequence && !SequenceSamplerMap.Contains(Sequence))
+			if (Sequence && !SamplerMap.Contains(Sequence))
 			{
-				int32 SequenceSamplerIdx = SequenceSamplers.AddDefaulted();
-				SequenceSamplerMap.Add(Sequence, SequenceSamplerIdx);
-
 				FSequenceBaseSampler::FInput Input;
 				Input.ExtrapolationParameters = Database->ExtrapolationParameters;
 				Input.SequenceBase = Sequence;
-				SequenceSamplers[SequenceSamplerIdx].Init(Input);
+
+				SamplerMap.Add(Sequence, Samplers.Num());
+				FSequenceBaseSampler* Sampler = Samplers.Add_GetRef(FInstancedStruct::Make(FSequenceBaseSampler())).GetMutablePtr<FSequenceBaseSampler>();
+				check(Sampler);
+				Sampler->Init(Input);
 			}
 		};
 
@@ -68,52 +70,51 @@ void FDatabaseIndexingContext::Prepare(const UPoseSearchDatabase* Database)
 
 						if (!BlendSpaceSamplerMap.Contains({ DatabaseBlendSpace->BlendSpace, BlendParameters }))
 						{
-							int32 BlendSpaceSamplerIdx = BlendSpaceSamplers.AddDefaulted();
-							BlendSpaceSamplerMap.Add({ DatabaseBlendSpace->BlendSpace, BlendParameters }, BlendSpaceSamplerIdx);
-
 							FBlendSpaceSampler::FInput Input;
 							Input.BoneContainer = BoneContainer;
 							Input.ExtrapolationParameters = Database->ExtrapolationParameters;
 							Input.BlendSpace = DatabaseBlendSpace->BlendSpace;
 							Input.BlendParameters = BlendParameters;
 
-							BlendSpaceSamplers[BlendSpaceSamplerIdx].Init(Input);
+							BlendSpaceSamplerMap.Add({ DatabaseBlendSpace->BlendSpace, BlendParameters }, Samplers.Num());
+							FBlendSpaceSampler* Sampler = Samplers.Add_GetRef(FInstancedStruct::Make(FBlendSpaceSampler())).GetMutablePtr<FBlendSpaceSampler>();
+							check(Sampler);
+							Sampler->Init(Input);
 						}
 					}
 				}
 			}
 		}
-	}
-
-	TArray<IAssetSampler*, TInlineAllocator<512>> AssetSampler;
-	AssetSampler.SetNumUninitialized(SequenceSamplers.Num() + BlendSpaceSamplers.Num());
+		else if (const FPoseSearchDatabaseAnimMontage* DatabaseAnimMontage = DatabaseAssetStruct.GetPtr<FPoseSearchDatabaseAnimMontage>())
+		{
+			if (DatabaseAnimMontage->AnimMontage && !SamplerMap.Contains(DatabaseAnimMontage->AnimMontage))
+			{
+				FAnimMontageSampler::FInput Input;
+				Input.ExtrapolationParameters = Database->ExtrapolationParameters;
+				Input.AnimMontage = DatabaseAnimMontage->AnimMontage;
 
-	for (int i = 0; i < SequenceSamplers.Num(); ++i)
-	{
-		AssetSampler[i] = &SequenceSamplers[i];
-	}
-	for (int i = 0; i < BlendSpaceSamplers.Num(); ++i)
-	{
-		AssetSampler[i + SequenceSamplers.Num()] = &BlendSpaceSamplers[i];
+				SamplerMap.Add(DatabaseAnimMontage->AnimMontage, Samplers.Num());
+				FAnimMontageSampler* Sampler = Samplers.Add_GetRef(FInstancedStruct::Make(FAnimMontageSampler())).GetMutablePtr<FAnimMontageSampler>();
+				check(Sampler);
+				Sampler->Init(Input);
+			}
+		}
+		else
+		{
+			checkNoEntry();
+		}
 	}
 
-	ParallelFor(AssetSampler.Num(), [AssetSampler](int32 SamplerIdx) { AssetSampler[SamplerIdx]->Process(); }, ParallelForFlags);
+	ParallelFor(Samplers.Num(), [this](int32 SamplerIdx)
+		{
+			FAssetSamplerBase* AssetSamplerBase = Samplers[SamplerIdx].GetMutablePtr<FAssetSamplerBase>();
+			check(AssetSamplerBase);
+			AssetSamplerBase->Process();
+		}, ParallelForFlags);
 
 	// prepare indexers
 	Indexers.Reserve(SearchIndexBase->Assets.Num());
 
-	auto GetSequenceBaseSampler = [&](const UAnimSequenceBase* Sequence) -> const FSequenceBaseSampler*
-	{
-		return Sequence ? &SequenceSamplers[SequenceSamplerMap[Sequence]] : nullptr;
-	};
-
-	auto GetBlendSpaceSampler = [&](const UBlendSpace* BlendSpace, const FVector BlendParameters) -> const FBlendSpaceSampler*
-	{
-		return BlendSpace ? &BlendSpaceSamplers[BlendSpaceSamplerMap[{BlendSpace, BlendParameters}]] : nullptr;
-	};
-
-	Indexers.Reserve(SearchIndexBase->Assets.Num());
-
 	for (int32 AssetIdx = 0; AssetIdx != SearchIndexBase->Assets.Num(); ++AssetIdx)
 	{
 		const FPoseSearchIndexAsset& SearchIndexAsset = SearchIndexBase->Assets[AssetIdx];
@@ -127,16 +128,36 @@ void FDatabaseIndexingContext::Prepare(const UPoseSearchDatabase* Database)
 		const FInstancedStruct& DatabaseAsset = Database->GetAnimationAssetStruct(SearchIndexAsset.SourceAssetIdx);
 		if (const FPoseSearchDatabaseSequence* DatabaseSequence = DatabaseAsset.GetPtr<FPoseSearchDatabaseSequence>())
 		{
-			const float SequenceLength = DatabaseSequence->Sequence->GetPlayLength();
-			IndexerContext.AssetSampler = GetSequenceBaseSampler(DatabaseSequence->Sequence);
+			if (DatabaseSequence->Sequence)
+			{
+				const float SequenceLength = DatabaseSequence->Sequence->GetPlayLength();
+				IndexerContext.AssetSampler = Samplers[SamplerMap[DatabaseSequence->Sequence]].GetPtr<FAssetSamplerBase>();
+			}
 		}
 		else if (const FPoseSearchDatabaseAnimComposite* DatabaseAnimComposite = DatabaseAsset.GetPtr<FPoseSearchDatabaseAnimComposite>())
 		{
-			IndexerContext.AssetSampler = GetSequenceBaseSampler(DatabaseAnimComposite->AnimComposite);
+			if (DatabaseAnimComposite->AnimComposite)
+			{
+				IndexerContext.AssetSampler = Samplers[SamplerMap[DatabaseAnimComposite->AnimComposite]].GetPtr<FAssetSamplerBase>();
+			}
 		}
 		else if (const FPoseSearchDatabaseBlendSpace* DatabaseBlendSpace = DatabaseAsset.GetPtr<FPoseSearchDatabaseBlendSpace>())
 		{
-			IndexerContext.AssetSampler = GetBlendSpaceSampler(DatabaseBlendSpace->BlendSpace, SearchIndexAsset.BlendParameters);
+			if (DatabaseBlendSpace->BlendSpace)
+			{
+				IndexerContext.AssetSampler = Samplers[BlendSpaceSamplerMap[{DatabaseBlendSpace->BlendSpace, SearchIndexAsset.BlendParameters}]].GetPtr<FAssetSamplerBase>();
+			}
+		}
+		else if (const FPoseSearchDatabaseAnimMontage* DatabaseAnimMontage = DatabaseAsset.GetPtr<FPoseSearchDatabaseAnimMontage>())
+		{
+			if (DatabaseAnimMontage->AnimMontage)
+			{
+				IndexerContext.AssetSampler = Samplers[SamplerMap[DatabaseAnimMontage->AnimMontage]].GetPtr<FAssetSamplerBase>();
+			}
+		}
+		else
+		{
+			checkNoEntry();
 		}
 
 		FAssetIndexer& Indexer = Indexers.AddDefaulted_GetRef();
diff --git a/Engine/Plugins/Experimental/Animation/PoseSearch/Source/Runtime/Private/PoseSearchDatabaseIndexingContext.h b/Engine/Plugins/Experimental/Animation/PoseSearch/Source/Runtime/Private/PoseSearchDatabaseIndexingContext.h
--- a/Engine/Plugins/Experimental/Animation/PoseSearch/Source/Runtime/Private/PoseSearchDatabaseIndexingContext.h
+++ b/Engine/Plugins/Experimental/Animation/PoseSearch/Source/Runtime/Private/PoseSearchDatabaseIndexingContext.h
@@ -4,6 +4,7 @@
 
 #if WITH_EDITOR
 
+#include "InstancedStruct.h"
 #include "PoseSearch/PoseSearchAssetSampler.h"
 #include "PoseSearchAssetIndexer.h"
 
@@ -16,11 +17,8 @@ namespace UE::PoseSearch
 struct FDatabaseIndexingContext
 {
 	FPoseSearchIndexBase* SearchIndexBase = nullptr;
-
 	FAssetSamplingContext SamplingContext;
-	TArray<FSequenceBaseSampler> SequenceSamplers; // Composite and sequence samplers
-	TArray<FBlendSpaceSampler> BlendSpaceSamplers;
-
+	TArray<FInstancedStruct> Samplers;
 	TArray<FAssetIndexer> Indexers;
 
 	void Prepare(const UPoseSearchDatabase* Database);
diff --git a/Engine/Plugins/Experimental/Animation/PoseSearch/Source/Runtime/Private/PoseSearchDatabaseSet.cpp b/Engine/Plugins/Experimental/Animation/PoseSearch/Source/Runtime/Private/PoseSearchDatabaseSet.cpp
--- a/Engine/Plugins/Experimental/Animation/PoseSearch/Source/Runtime/Private/PoseSearchDatabaseSet.cpp
+++ b/Engine/Plugins/Experimental/Animation/PoseSearch/Source/Runtime/Private/PoseSearchDatabaseSet.cpp
@@ -38,7 +38,7 @@ UE::PoseSearch::FSearchResult UPoseSearchDatabaseSet::Search(UE::PoseSearch::FSe
 		Result.ContinuingPoseCost = Result.PoseCost;
 		ContinuingCost = Result.PoseCost;
 
-		Result.AssetTime = SearchIndex.GetAssetTime(Result.PoseIdx, Database->Schema->GetSamplingInterval());
+		Result.AssetTime = Database->GetAssetTime(Result.PoseIdx, Database->Schema->GetSamplingInterval());
 		Result.Database = Database;
 
 		if (Database->GetSkipSearchIfPossible())
diff --git a/Engine/Plugins/Experimental/Animation/PoseSearch/Source/Runtime/Private/PoseSearchDerivedData.cpp b/Engine/Plugins/Experimental/Animation/PoseSearch/Source/Runtime/Private/PoseSearchDerivedData.cpp
--- a/Engine/Plugins/Experimental/Animation/PoseSearch/Source/Runtime/Private/PoseSearchDerivedData.cpp
+++ b/Engine/Plugins/Experimental/Animation/PoseSearch/Source/Runtime/Private/PoseSearchDerivedData.cpp
@@ -5,6 +5,7 @@
 #if WITH_EDITOR
 
 #include "Animation/AnimComposite.h"
+#include "Animation/AnimMontage.h"
 #include "Animation/AnimSequence.h"
 #include "Animation/BlendSpace.h"
 #include "AssetRegistry/ARFilter.h"
@@ -324,12 +325,12 @@ static void InitSearchIndexAssets(FPoseSearchIndexBase& SearchIndex, UPoseSearch
 				{
 					if (bAddUnmirrored)
 					{
-						SearchIndex.Assets.Add(FPoseSearchIndexAsset(ESearchIndexAssetType::Sequence, AnimationAssetIndex, false, FFloatInterval(Range.GetLowerBoundValue(), Range.GetUpperBoundValue())));
+						SearchIndex.Assets.Add(FPoseSearchIndexAsset(AnimationAssetIndex, false, FFloatInterval(Range.GetLowerBoundValue(), Range.GetUpperBoundValue())));
 					}
 
 					if (bAddMirrored)
 					{
-						SearchIndex.Assets.Add(FPoseSearchIndexAsset(ESearchIndexAssetType::Sequence, AnimationAssetIndex, true, FFloatInterval(Range.GetLowerBoundValue(), Range.GetUpperBoundValue())));
+						SearchIndex.Assets.Add(FPoseSearchIndexAsset(AnimationAssetIndex, true, FFloatInterval(Range.GetLowerBoundValue(), Range.GetUpperBoundValue())));
 					}
 				}
 			}
@@ -341,12 +342,12 @@ static void InitSearchIndexAssets(FPoseSearchIndexBase& SearchIndex, UPoseSearch
 				{
 					if (bAddUnmirrored)
 					{
-						SearchIndex.Assets.Add(FPoseSearchIndexAsset(ESearchIndexAssetType::AnimComposite, AnimationAssetIndex, false, FFloatInterval(Range.GetLowerBoundValue(), Range.GetUpperBoundValue())));
+						SearchIndex.Assets.Add(FPoseSearchIndexAsset(AnimationAssetIndex, false, FFloatInterval(Range.GetLowerBoundValue(), Range.GetUpperBoundValue())));
 					}
 
 					if (bAddMirrored)
 					{
-						SearchIndex.Assets.Add(FPoseSearchIndexAsset(ESearchIndexAssetType::AnimComposite, AnimationAssetIndex, true, FFloatInterval(Range.GetLowerBoundValue(), Range.GetUpperBoundValue())));
+						SearchIndex.Assets.Add(FPoseSearchIndexAsset(AnimationAssetIndex, true, FFloatInterval(Range.GetLowerBoundValue(), Range.GetUpperBoundValue())));
 					}
 				}
 			}
@@ -370,16 +371,33 @@ static void InitSearchIndexAssets(FPoseSearchIndexBase& SearchIndex, UPoseSearch
 
 						if (bAddUnmirrored)
 						{
-							SearchIndex.Assets.Add(FPoseSearchIndexAsset(ESearchIndexAssetType::BlendSpace, AnimationAssetIndex, false, FFloatInterval(0.0f, PlayLength), BlendParameters));
+							SearchIndex.Assets.Add(FPoseSearchIndexAsset(AnimationAssetIndex, false, FFloatInterval(0.0f, PlayLength), BlendParameters));
 						}
 
 						if (bAddMirrored)
 						{
-							SearchIndex.Assets.Add(FPoseSearchIndexAsset(ESearchIndexAssetType::BlendSpace, AnimationAssetIndex, true, FFloatInterval(0.0f, PlayLength), BlendParameters));
+							SearchIndex.Assets.Add(FPoseSearchIndexAsset(AnimationAssetIndex, true, FFloatInterval(0.0f, PlayLength), BlendParameters));
 						}
 					}
 				}
 			}
+			else if (const FPoseSearchDatabaseAnimMontage* DatabaseAnimMontage = DatabaseAssetStruct.GetPtr<FPoseSearchDatabaseAnimMontage>())
+			{
+				ValidRanges.Reset();
+				FindValidSequenceIntervals(DatabaseAnimMontage->AnimMontage, DatabaseAnimMontage->SamplingRange, DatabaseAnimMontage->IsLooping(), Database->ExcludeFromDatabaseParameters, ValidRanges);
+				for (const FFloatRange& Range : ValidRanges)
+				{
+					if (bAddUnmirrored)
+					{
+						SearchIndex.Assets.Add(FPoseSearchIndexAsset(AnimationAssetIndex, false, FFloatInterval(Range.GetLowerBoundValue(), Range.GetUpperBoundValue())));
+					}
+
+					if (bAddMirrored)
+					{
+						SearchIndex.Assets.Add(FPoseSearchIndexAsset(AnimationAssetIndex, true, FFloatInterval(Range.GetLowerBoundValue(), Range.GetUpperBoundValue())));
+					}
+				}
+			}
 			else
 			{
 				checkNoEntry();
@@ -644,11 +662,10 @@ private:
 	void SetState(EState State) { ThreadSafeState.Set(int32(State)); }
 
 	TWeakObjectPtr<UPoseSearchDatabase> Database;
-	// @todo: this is not relevant when the async task is completed, so to save memory we should move it as pointer perhaps
 	FPoseSearchIndex SearchIndex;
 	UE::DerivedData::FRequestOwner Owner;
 	FIoHash DerivedDataKey = FIoHash::Zero;
-	TSet<TWeakObjectPtr<const UObject>> DatabaseDependencies; // @todo: make this const
+	TSet<TWeakObjectPtr<const UObject>> DatabaseDependencies;
 		
 	FThreadSafeCounter ThreadSafeState = int32(EState::Notstarted);
 	bool bBroadcastOnDerivedDataRebuild = false;
diff --git a/Engine/Plugins/Experimental/Animation/PoseSearch/Source/Runtime/Private/PoseSearchDerivedDataKey.cpp b/Engine/Plugins/Experimental/Animation/PoseSearch/Source/Runtime/Private/PoseSearchDerivedDataKey.cpp
--- a/Engine/Plugins/Experimental/Animation/PoseSearch/Source/Runtime/Private/PoseSearchDerivedDataKey.cpp
+++ b/Engine/Plugins/Experimental/Animation/PoseSearch/Source/Runtime/Private/PoseSearchDerivedDataKey.cpp
@@ -41,7 +41,7 @@ FKeyBuilder::FKeyBuilder(const UObject* Object, bool bUseDataVer)
 	if (bUseDataVer)
 	{
 		// used to invalidate the key without having to change POSESEARCHDB_DERIVEDDATA_VER all the times
-		int32 POSESEARCHDB_DERIVEDDATA_VER_SMALL = 41;
+		int32 POSESEARCHDB_DERIVEDDATA_VER_SMALL = 51;
 		FGuid VersionGuid = FDevSystemGuids::GetSystemGuid(FDevSystemGuids::Get().POSESEARCHDB_DERIVEDDATA_VER);
 
 		*this << VersionGuid;
diff --git a/Engine/Plugins/Experimental/Animation/PoseSearch/Source/Runtime/Private/PoseSearchIndex.cpp b/Engine/Plugins/Experimental/Animation/PoseSearch/Source/Runtime/Private/PoseSearchIndex.cpp
--- a/Engine/Plugins/Experimental/Animation/PoseSearch/Source/Runtime/Private/PoseSearchIndex.cpp
+++ b/Engine/Plugins/Experimental/Animation/PoseSearch/Source/Runtime/Private/PoseSearchIndex.cpp
@@ -64,32 +64,6 @@ const FPoseSearchIndexAsset* FPoseSearchIndexBase::GetAssetForPoseSafe(int32 Pos
 	return nullptr;
 }
 
-float FPoseSearchIndexBase::GetAssetTime(int32 PoseIdx, float SamplingInterval) const
-{
-	const FPoseSearchIndexAsset& Asset = GetAssetForPose(PoseIdx);
-
-	if (Asset.Type == ESearchIndexAssetType::Sequence || Asset.Type == ESearchIndexAssetType::AnimComposite)
-	{
-		const FFloatInterval SamplingRange = Asset.SamplingInterval;
-
-		float AssetTime = FMath::Min(SamplingRange.Min + SamplingInterval * (PoseIdx - Asset.FirstPoseIdx), SamplingRange.Max);
-		return AssetTime;
-	}
-
-	if (Asset.Type == ESearchIndexAssetType::BlendSpace)
-	{
-		const FFloatInterval SamplingRange = Asset.SamplingInterval;
-
-		// For BlendSpaces the AssetTime is in the range [0, 1] while the Sampling Range
-		// is in real time (seconds)
-		float AssetTime = FMath::Min(SamplingRange.Min + SamplingInterval * (PoseIdx - Asset.FirstPoseIdx), SamplingRange.Max) / (Asset.NumPoses * SamplingInterval);
-		return AssetTime;
-	}
-	
-	checkNoEntry();
-	return -1.0f;
-}
-
 bool FPoseSearchIndexBase::IsEmpty() const
 {
 	const bool bEmpty = Assets.Num() == 0 || NumPoses == 0;
diff --git a/Engine/Plugins/Experimental/Animation/PoseSearch/Source/Runtime/Private/PoseSearchLibrary.cpp b/Engine/Plugins/Experimental/Animation/PoseSearch/Source/Runtime/Private/PoseSearchLibrary.cpp
--- a/Engine/Plugins/Experimental/Animation/PoseSearch/Source/Runtime/Private/PoseSearchLibrary.cpp
+++ b/Engine/Plugins/Experimental/Animation/PoseSearch/Source/Runtime/Private/PoseSearchLibrary.cpp
@@ -1,15 +1,16 @@
 // Copyright Epic Games, Inc. All Rights Reserved.
 
 #include "PoseSearch/PoseSearchLibrary.h"
-
 #include "Animation/AnimInstanceProxy.h"
-#include "Animation/AnimSequence.h"
 #include "Animation/AnimNode_Inertialization.h"
 #include "Animation/AnimNode_SequencePlayer.h"
 #include "Animation/AnimRootMotionProvider.h"
+#include "Animation/AnimSequence.h"
+#include "Animation/AnimSubsystem_Tag.h"
 #include "Animation/BlendSpace.h"
 #include "InstancedStruct.h"
 #include "PoseSearch/AnimNode_MotionMatching.h"
+#include "PoseSearch/AnimNode_PoseSearchHistoryCollector.h"
 #include "PoseSearch/PoseSearchDatabase.h"
 #include "PoseSearch/PoseSearchSchema.h"
 #include "PoseSearchFeatureChannel_Trajectory.h"
@@ -138,25 +139,65 @@
 	bJumpedToPose = true;
 }
 
-#if UE_POSE_SEARCH_TRACE_ENABLED
-static void TraceMotionMatchingState(
-	const FAnimationUpdateContext& UpdateContext,
+void FMotionMatchingState::UpdateWantedPlayRate(const UE::PoseSearch::FSearchContext& SearchContext, const FMotionMatchingSettings& Settings)
+{
+	if (CurrentSearchResult.IsValid())
+	{
+		if (!FMath::IsNearlyEqual(Settings.PlayRateMin, 1.f, UE_KINDA_SMALL_NUMBER) || !FMath::IsNearlyEqual(Settings.PlayRateMax, 1.f, UE_KINDA_SMALL_NUMBER))
+		{
+			if (const FPoseSearchFeatureVectorBuilder* PoseSearchFeatureVectorBuilder = SearchContext.GetCachedQuery(CurrentSearchResult.Database->Schema))
+			{
+				if (const UPoseSearchFeatureChannel_Trajectory* TrajectoryChannel = CurrentSearchResult.Database->Schema->FindFirstChannelOfType<UPoseSearchFeatureChannel_Trajectory>())
+				{
+					TConstArrayView<float> QueryData = PoseSearchFeatureVectorBuilder->GetValues();
+					TConstArrayView<float> ResultData = CurrentSearchResult.Database->GetSearchIndex().GetPoseValues(CurrentSearchResult.PoseIdx);
+					const float EstimatedSpeedRatio = TrajectoryChannel->GetEstimatedSpeedRatio(QueryData, ResultData);
+					check(Settings.PlayRateMin <= Settings.PlayRateMax);
+					WantedPlayRate = FMath::Clamp(EstimatedSpeedRatio, Settings.PlayRateMin, Settings.PlayRateMax);
+				}
+				else
+				{
+					UE_LOG(LogPoseSearch, Warning,
+						TEXT("Couldn't update the WantedPlayRate in FMotionMatchingState::UpdateWantedPlayRate, because Schema '%s' couldn't find a UPoseSearchFeatureChannel_Trajectory channel"),
+						*GetNameSafe(CurrentSearchResult.Database->Schema));
+				}
+			}
+		}
+	}
+}
+
+float FMotionMatchingState::ComputeJumpBlendTime(const UE::PoseSearch::FSearchResult& Result, const FMotionMatchingSettings& Settings) const
+{
+	const FPoseSearchIndexAsset* SearchIndexAsset = CurrentSearchResult.GetSearchIndexAsset();
+
+	// Use alternate blend time when changing between mirrored and unmirrored
+	float JumpBlendTime = Settings.BlendTime;
+	if (SearchIndexAsset && Settings.MirrorChangeBlendTime > 0.0f)
+	{
+		if (Result.GetSearchIndexAsset(true)->bMirrored != SearchIndexAsset->bMirrored)
+		{
+			JumpBlendTime = Settings.MirrorChangeBlendTime;
+		}
+	}
+
+	return JumpBlendTime;
+}
+
+void UPoseSearchLibrary::TraceMotionMatchingState(
 	const UPoseSearchSearchableAsset* Searchable,
 	UE::PoseSearch::FSearchContext& SearchContext,
-	const FMotionMatchingState& MotionMatchingState,
-	const FTrajectorySampleRange& Trajectory,
+	const UE::PoseSearch::FSearchResult& CurrentResult,
 	const UE::PoseSearch::FSearchResult& LastResult,
+	float ElapsedPoseSearchTime,
+	const FTransform& RootMotionTransformDelta,
+	const UObject* AnimInstance,
+	int32 NodeId,
+	float DeltaTime,
 	bool bSearch)
 {
+#if UE_POSE_SEARCH_TRACE_ENABLED
 	using namespace UE::PoseSearch;
-
-	if (!IsTracing(UpdateContext))
-	{
-		return;
-	}
-
-	const float DeltaTime = UpdateContext.GetDeltaTime();
-
+	
 	auto AddUniqueDatabase = [](TArray<FTraceMotionMatchingStateDatabaseEntry>& DatabaseEntries, const UPoseSearchDatabase* Database, UE::PoseSearch::FSearchContext& SearchContext) -> int32
 	{
 		const uint64 DatabaseId = FTraceMotionMatchingState::GetIdFromObject(Database);
@@ -201,40 +242,40 @@
 		DbEntry.PoseEntries.Add(PoseEntry);
 	}
 
-	if (bSearch && MotionMatchingState.CurrentSearchResult.ContinuingPoseCost.IsValid())
+	if (bSearch && CurrentResult.ContinuingPoseCost.IsValid())
 	{
 		check(LastResult.IsValid());
 
 		const int32 DbEntryIdx = AddUniqueDatabase(TraceState.DatabaseEntries, LastResult.Database.Get(), SearchContext);
 		FTraceMotionMatchingStateDatabaseEntry& DbEntry = TraceState.DatabaseEntries[DbEntryIdx];
 
-		const int32 LastResultPoseEntryIdx = DbEntry.PoseEntries.Add({LastResult.PoseIdx});
+		const int32 LastResultPoseEntryIdx = DbEntry.PoseEntries.Add({ LastResult.PoseIdx });
 		FTraceMotionMatchingStatePoseEntry& PoseEntry = DbEntry.PoseEntries[LastResultPoseEntryIdx];
 
-		PoseEntry.Cost = MotionMatchingState.CurrentSearchResult.ContinuingPoseCost;
+		PoseEntry.Cost = CurrentResult.ContinuingPoseCost;
 		PoseEntry.PoseCandidateFlags = EPoseCandidateFlags::Valid_ContinuingPose;
 	}
 
-	if (bSearch && MotionMatchingState.CurrentSearchResult.PoseCost.IsValid())
+	if (bSearch && CurrentResult.PoseCost.IsValid())
 	{
-		const int32 DbEntryIdx = AddUniqueDatabase(TraceState.DatabaseEntries, MotionMatchingState.CurrentSearchResult.Database.Get(), SearchContext);
+		const int32 DbEntryIdx = AddUniqueDatabase(TraceState.DatabaseEntries, CurrentResult.Database.Get(), SearchContext);
 		FTraceMotionMatchingStateDatabaseEntry& DbEntry = TraceState.DatabaseEntries[DbEntryIdx];
 
-		const int32 PoseEntryIdx = DbEntry.PoseEntries.Add({MotionMatchingState.CurrentSearchResult.PoseIdx});
+		const int32 PoseEntryIdx = DbEntry.PoseEntries.Add({ CurrentResult.PoseIdx });
 		FTraceMotionMatchingStatePoseEntry& PoseEntry = DbEntry.PoseEntries[PoseEntryIdx];
 
-		PoseEntry.Cost = MotionMatchingState.CurrentSearchResult.PoseCost;
+		PoseEntry.Cost = CurrentResult.PoseCost;
 		PoseEntry.PoseCandidateFlags = EPoseCandidateFlags::Valid_CurrentPose;
 
 		TraceState.CurrentDbEntryIdx = DbEntryIdx;
 		TraceState.CurrentPoseEntryIdx = PoseEntryIdx;
 	}
 
-	if (DeltaTime > SMALL_NUMBER)
+	if (DeltaTime > SMALL_NUMBER && SearchContext.Trajectory)
 	{
 		// simulation
-		const FTrajectorySample PrevSample = Trajectory.GetSampleAtTime(-DeltaTime);
-		const FTrajectorySample CurrSample = Trajectory.GetSampleAtTime(0.f);
+		const FTrajectorySample PrevSample = SearchContext.Trajectory->GetSampleAtTime(-DeltaTime);
+		const FTrajectorySample CurrSample = SearchContext.Trajectory->GetSampleAtTime(0.f);
 
 		const FTransform SimDelta = CurrSample.Transform.GetRelativeTransform(PrevSample.Transform);
 
@@ -242,30 +283,27 @@
 		TraceState.SimAngularVelocity = FMath::RadiansToDegrees(SimDelta.GetRotation().GetAngle()) / DeltaTime;
 
 		// animation
-		const FTransform AnimDelta = MotionMatchingState.RootMotionTransformDelta;
-
-		TraceState.AnimLinearVelocity = AnimDelta.GetTranslation().Size() / DeltaTime;
-		TraceState.AnimAngularVelocity = FMath::RadiansToDegrees(AnimDelta.GetRotation().GetAngle()) / DeltaTime;
+		TraceState.AnimLinearVelocity = RootMotionTransformDelta.GetTranslation().Size() / DeltaTime;
+		TraceState.AnimAngularVelocity = FMath::RadiansToDegrees(RootMotionTransformDelta.GetRotation().GetAngle()) / DeltaTime;
 	}
 
 	TraceState.SearchableAssetId = FTraceMotionMatchingState::GetIdFromObject(Searchable);
-	TraceState.ElapsedPoseSearchTime = MotionMatchingState.ElapsedPoseSearchTime;
-	TraceState.AssetPlayerTime = MotionMatchingState.CurrentSearchResult.AssetTime;
+	TraceState.ElapsedPoseSearchTime = ElapsedPoseSearchTime;
+	TraceState.AssetPlayerTime = CurrentResult.AssetTime;
 	TraceState.DeltaTime = DeltaTime;
 
-	TraceState.Output(UpdateContext);
-}
+	TraceState.Output(AnimInstance, NodeId);
 #endif
+}
 
-void UpdateMotionMatchingState(
+void UPoseSearchLibrary::UpdateMotionMatchingState(
 	const FAnimationUpdateContext& Context,
 	const UPoseSearchSearchableAsset* Searchable,
 	const FGameplayTagContainer* ActiveTagsContainer,
 	const FTrajectorySampleRange& Trajectory,
 	const FMotionMatchingSettings& Settings,
 	FMotionMatchingState& InOutMotionMatchingState,
-	bool bForceInterrupt
-)
+	bool bForceInterrupt)
 {
 	QUICK_SCOPE_CYCLE_COUNTER(STAT_PoseSearch_Update);
 
@@ -274,7 +312,7 @@
 	if (!Searchable)
 	{
 		Context.LogMessage(
-			EMessageSeverity::Error, 
+			EMessageSeverity::Error,
 			LOCTEXT("NoSearchable", "No searchable asset provided for motion matching."));
 		return;
 	}
@@ -292,9 +330,6 @@
 	FSearchContext SearchContext;
 	SearchContext.ActiveTagsContainer = ActiveTagsContainer;
 	SearchContext.Trajectory = &Trajectory;
-	SearchContext.OwningComponent = Context.AnimInstanceProxy->GetSkelMeshComponent();
-	SearchContext.BoneContainer = &Context.AnimInstanceProxy->GetRequiredBones();
-	SearchContext.bIsTracing = IsTracing(Context);
 	SearchContext.bForceInterrupt = bForceInterrupt;
 	SearchContext.bCanAdvance = InOutMotionMatchingState.CanAdvance(DeltaTime);
 	SearchContext.CurrentResult = InOutMotionMatchingState.CurrentSearchResult;
@@ -322,7 +357,7 @@
 
 		// making sure we haven't calculated ContinuingPoseCost if we !SearchContext.bCanAdvance 
 		check(SearchContext.bCanAdvance || !SearchResult.ContinuingPoseCost.IsValid());
-		
+
 		if (SearchResult.PoseCost.GetTotalCost() < SearchResult.ContinuingPoseCost.GetTotalCost())
 		{
 			InOutMotionMatchingState.JumpToPose(Context, Settings, SearchResult);
@@ -337,7 +372,7 @@
 			InOutMotionMatchingState.CurrentSearchResult.ContinuingPoseCost = SearchResult.ContinuingPoseCost;
 			InOutMotionMatchingState.CurrentSearchResult.ComposedQuery = SearchResult.ComposedQuery;
 		}
-	
+
 		InOutMotionMatchingState.UpdateWantedPlayRate(SearchContext, Settings);
 	}
 	else
@@ -349,59 +384,84 @@
 
 #if UE_POSE_SEARCH_TRACE_ENABLED
 	// Record debugger details
-	TraceMotionMatchingState(
-		Context,
-		Searchable,
-		SearchContext,
-		InOutMotionMatchingState,
-		Trajectory,
-		LastResult,
-		bSearch);
+	if (IsTracing(Context))
+	{
+		TraceMotionMatchingState(Searchable, SearchContext, InOutMotionMatchingState.CurrentSearchResult, LastResult, InOutMotionMatchingState.ElapsedPoseSearchTime,
+			InOutMotionMatchingState.RootMotionTransformDelta, Context.AnimInstanceProxy->GetAnimInstanceObject(), Context.GetCurrentNodeId(), DeltaTime, bSearch);
+	}
 #endif
 }
 
-void FMotionMatchingState::UpdateWantedPlayRate(const UE::PoseSearch::FSearchContext& SearchContext, const FMotionMatchingSettings& Settings)
+void UPoseSearchLibrary::MotionMatch(
+	const UAnimInstance* AnimInstance,
+	const UPoseSearchSearchableAsset* Searchable,
+	const FGameplayTagContainer ActiveTagsContainer,
+	const FTrajectorySampleRange Trajectory,
+	const FName PoseHistoryName,
+	UAnimationAsset*& SelectedAnimation,
+	float& SelectedTime,
+	bool& bLoop,
+	bool& bIsMirrored,
+	FVector& BlendParameters,
+	const int DebugSessionUniqueIdentifier)
 {
-	if (CurrentSearchResult.IsValid())
+	using namespace UE::PoseSearch;
+
+	SelectedAnimation = nullptr;
+	SelectedTime = 0.f;
+	bLoop = false;
+	bIsMirrored = false;
+	BlendParameters = FVector::ZeroVector;
+
+	if (Searchable)
 	{
-		if (!FMath::IsNearlyEqual(Settings.PlayRateMin, 1.f, UE_KINDA_SMALL_NUMBER) || !FMath::IsNearlyEqual(Settings.PlayRateMax, 1.f, UE_KINDA_SMALL_NUMBER))
+		// Build the search context
+		FSearchContext SearchContext;
+		SearchContext.ActiveTagsContainer = &ActiveTagsContainer;
+		SearchContext.Trajectory = &Trajectory;
+
+		if (AnimInstance)
 		{
-			if (const FPoseSearchFeatureVectorBuilder* PoseSearchFeatureVectorBuilder = SearchContext.GetCachedQuery(CurrentSearchResult.Database->Schema))
+			if (IAnimClassInterface* AnimBlueprintClass = IAnimClassInterface::GetFromClass(AnimInstance->GetClass()))
 			{
-				if (const UPoseSearchFeatureChannel_Trajectory* TrajectoryChannel = CurrentSearchResult.Database->Schema->FindFirstChannelOfType<UPoseSearchFeatureChannel_Trajectory>())
+				if (const FAnimSubsystem_Tag* TagSubsystem = AnimBlueprintClass->FindSubsystem<FAnimSubsystem_Tag>())
 				{
-					TConstArrayView<float> QueryData = PoseSearchFeatureVectorBuilder->GetValues();
-					TConstArrayView<float> ResultData = CurrentSearchResult.Database->GetSearchIndex().GetPoseValues(CurrentSearchResult.PoseIdx);
-					const float EstimatedSpeedRatio = TrajectoryChannel->GetEstimatedSpeedRatio(QueryData, ResultData);
-					check(Settings.PlayRateMin <= Settings.PlayRateMax);
-					WantedPlayRate = FMath::Clamp(EstimatedSpeedRatio, Settings.PlayRateMin, Settings.PlayRateMax);
-				}
-				else
-				{
-					UE_LOG(LogPoseSearch, Warning,
-						TEXT("Couldn't update the WantedPlayRate in FMotionMatchingState::UpdateWantedPlayRate, because Schema '%s' couldn't find a UPoseSearchFeatureChannel_Trajectory channel"),
-						*GetNameSafe(CurrentSearchResult.Database->Schema));
+					if (const FAnimNode_PoseSearchHistoryCollector_Base* PoseHistoryNode = TagSubsystem->FindNodeByTag<FAnimNode_PoseSearchHistoryCollector_Base>(PoseHistoryName, AnimInstance))
+					{
+						SearchContext.History = &PoseHistoryNode->GetPoseHistory();
+					}
 				}
 			}
+
+			if (!SearchContext.History)
+			{
+				UE_LOG(LogPoseSearch, Warning, TEXT("UPoseSearchLibrary::MotionMatch - Couldn't find pose history with name '%s'"), *PoseHistoryName.ToString());
+			}
+		}
+		
+		// @todo: finish set up SearchContext by exposing or calculating additional members
+
+		FSearchResult SearchResult = Searchable->Search(SearchContext);
+		if (SearchResult.IsValid())
+		{
+			const FPoseSearchIndexAsset* SearchIndexAsset = SearchResult.GetSearchIndexAsset();
+			if (const FPoseSearchDatabaseAnimationAssetBase* DatabaseAsset = SearchResult.Database->GetAnimationAssetBase(*SearchIndexAsset))
+			{
+				SelectedAnimation = DatabaseAsset->GetAnimationAsset();
+				SelectedTime = SearchResult.AssetTime;
+				bLoop = DatabaseAsset->IsLooping();
+				bIsMirrored = SearchIndexAsset->bMirrored;
+				BlendParameters = SearchIndexAsset->BlendParameters;
+			}
+		}
+
+		if (AnimInstance)
+		{
+#if UE_POSE_SEARCH_TRACE_ENABLED
+			TraceMotionMatchingState(Searchable, SearchContext, SearchResult, FSearchResult(), 0.f, FTransform::Identity, AnimInstance, DebugSessionUniqueIdentifier, AnimInstance->GetDeltaSeconds(), true);
+#endif // UE_POSE_SEARCH_TRACE_ENABLED
 		}
 	}
 }
 
-float FMotionMatchingState::ComputeJumpBlendTime(const UE::PoseSearch::FSearchResult& Result, const FMotionMatchingSettings& Settings) const
-{
-	const FPoseSearchIndexAsset* SearchIndexAsset = CurrentSearchResult.GetSearchIndexAsset();
-
-	// Use alternate blend time when changing between mirrored and unmirrored
-	float JumpBlendTime = Settings.BlendTime;
-	if (SearchIndexAsset && Settings.MirrorChangeBlendTime > 0.0f)
-	{
-		if (Result.GetSearchIndexAsset(true)->bMirrored != SearchIndexAsset->bMirrored)
-		{
-			JumpBlendTime = Settings.MirrorChangeBlendTime;
-		}
-	}
-
-	return JumpBlendTime;
-}
-
 #undef LOCTEXT_NAMESPACE
diff --git a/Engine/Plugins/Experimental/Animation/PoseSearch/Source/Runtime/Private/PoseSearchResult.cpp b/Engine/Plugins/Experimental/Animation/PoseSearch/Source/Runtime/Private/PoseSearchResult.cpp
--- a/Engine/Plugins/Experimental/Animation/PoseSearch/Source/Runtime/Private/PoseSearchResult.cpp
+++ b/Engine/Plugins/Experimental/Animation/PoseSearch/Source/Runtime/Private/PoseSearchResult.cpp
@@ -70,6 +70,7 @@ void FSearchResult::Update(float NewAssetTime)
 		else
 		{
 			checkNoEntry();
+			Reset();
 		}
 	}
 }
diff --git a/Engine/Plugins/Experimental/Animation/PoseSearch/Source/Runtime/Private/Trace/PoseSearchTraceLogger.cpp b/Engine/Plugins/Experimental/Animation/PoseSearch/Source/Runtime/Private/Trace/PoseSearchTraceLogger.cpp
--- a/Engine/Plugins/Experimental/Animation/PoseSearch/Source/Runtime/Private/Trace/PoseSearchTraceLogger.cpp
+++ b/Engine/Plugins/Experimental/Animation/PoseSearch/Source/Runtime/Private/Trace/PoseSearchTraceLogger.cpp
@@ -84,18 +84,12 @@ FArchive& operator<<(FArchive& Ar, FTraceMotionMatchingState& State)
 	return Ar;
 }
 
-void FTraceMotionMatchingState::Output(const FAnimationBaseContext& InContext)
+void FTraceMotionMatchingState::Output(const UObject* AnimInstance, int32 NodeId)
 {
 #if OBJECT_TRACE_ENABLED
-	if (!IsTracing(InContext))
-	{
-		return;
-	}
-
 	TArray<uint8> ArchiveData;
 	FMemoryWriter Archive(ArchiveData);
 
-	UObject* AnimInstance = InContext.AnimInstanceProxy->GetAnimInstanceObject();
 	TRACE_OBJECT(AnimInstance);
 	UObject* SkeletalMeshComponent = AnimInstance->GetOuter();
 
@@ -103,14 +97,13 @@ void FTraceMotionMatchingState::Output(const FAnimationBaseContext& InContext)
 	TraceMessage.Cycle = FPlatformTime::Cycles64();
 	TraceMessage.AnimInstanceId = FObjectTrace::GetObjectId(AnimInstance);
 	TraceMessage.SkeletalMeshComponentId = FObjectTrace::GetObjectId(SkeletalMeshComponent);
-	TraceMessage.NodeId = InContext.GetCurrentNodeId();
+	TraceMessage.NodeId = NodeId;
 	TraceMessage.FrameCounter = FObjectTrace::GetObjectWorldTickCounter(AnimInstance);
 
 	Archive << TraceMessage;
 	Archive << *this;
 
-	UE_TRACE_LOG(PoseSearch, MotionMatchingState, PoseSearchChannel)
-		<< MotionMatchingState.Data(ArchiveData.GetData(), ArchiveData.Num());
+	UE_TRACE_LOG(PoseSearch, MotionMatchingState, PoseSearchChannel) << MotionMatchingState.Data(ArchiveData.GetData(), ArchiveData.Num());
 #endif
 }
 
diff --git a/Engine/Plugins/Experimental/Animation/PoseSearch/Source/Runtime/Private/Trace/PoseSearchTraceLogger.h b/Engine/Plugins/Experimental/Animation/PoseSearch/Source/Runtime/Private/Trace/PoseSearchTraceLogger.h
--- a/Engine/Plugins/Experimental/Animation/PoseSearch/Source/Runtime/Private/Trace/PoseSearchTraceLogger.h
+++ b/Engine/Plugins/Experimental/Animation/PoseSearch/Source/Runtime/Private/Trace/PoseSearchTraceLogger.h
@@ -86,7 +86,7 @@ struct POSESEARCH_API FTraceMotionMatchingState
 	int32 CurrentPoseEntryIdx = INDEX_NONE;
 
 	/** Output the current state info to the logger */
-	void Output(const FAnimationBaseContext& InContext);
+	void Output(const UObject* AnimInstance, int32 NodeId);
 
 	const UPoseSearchDatabase* GetCurrentDatabase() const;
 	int32 GetCurrentDatabasePoseIndex() const;
diff --git a/Engine/Plugins/Experimental/Animation/PoseSearch/Source/Runtime/Public/PoseSearch/PoseSearchAssetIndexer.h b/Engine/Plugins/Experimental/Animation/PoseSearch/Source/Runtime/Public/PoseSearch/PoseSearchAssetIndexer.h
--- a/Engine/Plugins/Experimental/Animation/PoseSearch/Source/Runtime/Public/PoseSearch/PoseSearchAssetIndexer.h
+++ b/Engine/Plugins/Experimental/Animation/PoseSearch/Source/Runtime/Public/PoseSearch/PoseSearchAssetIndexer.h
@@ -6,11 +6,11 @@
 #include "PoseSearch/PoseSearchDefines.h"
 #include "PoseSearch/PoseSearchSchema.h"
 
+struct FAssetSamplerBase;
+
 namespace UE::PoseSearch
 {
-	
 struct FAssetSamplingContext;
-class IAssetSampler;
 
 /**
  * Inputs for asset indexing
@@ -19,7 +19,7 @@ struct FAssetIndexingContext
 {
 	const FAssetSamplingContext* SamplingContext = nullptr;
 	const UPoseSearchSchema* Schema = nullptr;
-	const IAssetSampler* AssetSampler = nullptr;
+	const FAssetSamplerBase* AssetSampler = nullptr;
 	bool bMirrored = false;
 	FFloatInterval RequestedSamplingRange = FFloatInterval(0.0f, 0.0f);
 	
@@ -39,7 +39,7 @@ class POSESEARCH_API IAssetIndexer
 public:
 	struct FSampleInfo
 	{
-		const IAssetSampler* Clip = nullptr;
+		const FAssetSamplerBase* Clip = nullptr;
 		FTransform RootTransform;
 		float ClipTime = 0.0f;
 		bool bClamped = false;
diff --git a/Engine/Plugins/Experimental/Animation/PoseSearch/Source/Runtime/Public/PoseSearch/PoseSearchAssetSampler.h b/Engine/Plugins/Experimental/Animation/PoseSearch/Source/Runtime/Public/PoseSearch/PoseSearchAssetSampler.h
--- a/Engine/Plugins/Experimental/Animation/PoseSearch/Source/Runtime/Public/PoseSearch/PoseSearchAssetSampler.h
+++ b/Engine/Plugins/Experimental/Animation/PoseSearch/Source/Runtime/Public/PoseSearch/PoseSearchAssetSampler.h
@@ -8,10 +8,34 @@
 struct FAnimationPoseData;
 struct FAnimExtractContext;
 class UAnimationAsset;
+class UAnimInstance;
 class UAnimNotifyState_PoseSearchBase;
 class UBlendSpace;
 class UMirrorDataTable;
 
+namespace UE::PoseSearch
+{
+	struct POSESEARCH_API FAssetSamplingContext
+	{
+		// Time delta used for computing pose derivatives
+		static constexpr float FiniteDelta = 1 / 60.0f;
+
+		// Mirror data table pointer copied from Schema for convenience
+		TObjectPtr<const UMirrorDataTable> MirrorDataTable;
+
+		// Compact pose format of Mirror Bone Map
+		TCustomBoneIndexArray<FCompactPoseBoneIndex, FCompactPoseBoneIndex> CompactPoseMirrorBones;
+
+		// Pre-calculated component space rotations of reference pose, which allows mirror to work with any joint orientation
+		// Only initialized and used when a mirroring table is specified
+		TCustomBoneIndexArray<FQuat, FCompactPoseBoneIndex> ComponentSpaceRefRotations;
+
+		void Init(const UMirrorDataTable* InMirrorDataTable, const FBoneContainer& BoneContainer);
+		FTransform MirrorTransform(const FTransform& Transform) const;
+	};
+
+} // namespace UE::PoseSearch
+
 USTRUCT()
 struct POSESEARCH_API FPoseSearchExtrapolationParameters
 {
@@ -30,61 +54,46 @@ struct POSESEARCH_API FPoseSearchExtrapolationParameters
 	float SampleTime = 0.05f;
 };
 
-namespace UE::PoseSearch
-{
-
-struct POSESEARCH_API FAssetSamplingContext
-{
-	// Time delta used for computing pose derivatives
-	static constexpr float FiniteDelta = 1 / 60.0f;
-
-	// Mirror data table pointer copied from Schema for convenience
-	TObjectPtr<const UMirrorDataTable> MirrorDataTable;
-
-	// Compact pose format of Mirror Bone Map
-	TCustomBoneIndexArray<FCompactPoseBoneIndex, FCompactPoseBoneIndex> CompactPoseMirrorBones;
-
-	// Pre-calculated component space rotations of reference pose, which allows mirror to work with any joint orientation
-	// Only initialized and used when a mirroring table is specified
-	TCustomBoneIndexArray<FQuat, FCompactPoseBoneIndex> ComponentSpaceRefRotations;
-
-	void Init(const UMirrorDataTable* InMirrorDataTable, const FBoneContainer& BoneContainer);
-	FTransform MirrorTransform(const FTransform& Transform) const;
-};
-
 /**
  * Helper interface for sampling data from animation assets
  */
-class POSESEARCH_API IAssetSampler
+USTRUCT()
+struct POSESEARCH_API FAssetSamplerBase
 {
+	GENERATED_BODY()
+
 public:
-	virtual ~IAssetSampler() {};
+	virtual ~FAssetSamplerBase() {}
 
-	virtual float GetPlayLength() const = 0;
-	virtual bool IsLoopable() const = 0;
+	virtual float GetPlayLength() const { return 0.f; }
+	virtual float GetScaledTime(float Time) const { return Time; }
+	virtual bool IsLoopable() const { return false; }
 
 	// Gets the final root transformation at the end of the asset's playback time
-	virtual FTransform GetTotalRootTransform() const = 0;
+	virtual FTransform GetTotalRootTransform() const { return FTransform::Identity; }
 
 	// Extracts pose for this asset for a given context
-	virtual void ExtractPose(const FAnimExtractContext& ExtractionCtx, FAnimationPoseData& OutAnimPoseData) const = 0;
+	virtual void ExtractPose(const FAnimExtractContext& ExtractionCtx, FAnimationPoseData& OutAnimPoseData) const { }
 
 	// Extracts root transform at the given time, using the extremities of the sequence to extrapolate beyond the 
 	// sequence limits when Time is less than zero or greater than the sequence length.
-	virtual FTransform ExtractRootTransform(float Time) const = 0;
+	virtual FTransform ExtractRootTransform(float Time) const { return FTransform::Identity; }
 
 	// Extracts notify states inheriting from UAnimNotifyState_PoseSearchBase present in the sequence at Time.
 	// The function does not empty NotifyStates before adding new notifies!
-	virtual void ExtractPoseSearchNotifyStates(float Time, TArray<UAnimNotifyState_PoseSearchBase*>& NotifyStates) const = 0;
+	virtual void ExtractPoseSearchNotifyStates(float Time, TArray<UAnimNotifyState_PoseSearchBase*>& NotifyStates) const { }
 
-	virtual const UAnimationAsset* GetAsset() const = 0;
+	virtual const UAnimationAsset* GetAsset() const { return nullptr; }
 
-	virtual void Process() = 0;
+	virtual void Process() { }
 };
 
 // Sampler working with UAnimSequenceBase so it can be used for UAnimSequence as well as UAnimComposite.
-struct POSESEARCH_API FSequenceBaseSampler : public IAssetSampler
+USTRUCT()
+struct POSESEARCH_API FSequenceBaseSampler : public FAssetSamplerBase
 {
+	GENERATED_BODY()
+
 public:
 	struct FInput
 	{
@@ -96,6 +105,7 @@ public:
 	virtual void Process() override;
 
 	virtual float GetPlayLength() const override;
+	virtual float GetScaledTime(float Time) const override;
 	virtual bool IsLoopable() const override;
 
 	virtual FTransform GetTotalRootTransform() const override;
@@ -106,8 +116,11 @@ public:
 	virtual const UAnimationAsset* GetAsset() const override;
 };
 
-struct POSESEARCH_API FBlendSpaceSampler : public IAssetSampler
+USTRUCT()
+struct POSESEARCH_API FBlendSpaceSampler : public FAssetSamplerBase
 {
+	GENERATED_BODY()
+
 public:
 	struct FInput
 	{
@@ -122,6 +135,7 @@ public:
 	virtual void Process() override;
 
 	virtual float GetPlayLength() const override { return PlayLength; }
+	virtual float GetScaledTime(float Time) const override;
 	virtual bool IsLoopable() const override;
 
 	virtual FTransform GetTotalRootTransform() const override;
@@ -145,4 +159,34 @@ private:
 	FTransform ExtractBlendSpaceRootMotionFromRange(float StartTrackPosition, float EndTrackPosition) const;
 };
 
-} // namespace UE::PoseSearch
+USTRUCT()
+struct POSESEARCH_API FAnimMontageSampler : public FAssetSamplerBase
+{
+	GENERATED_BODY()
+
+public:
+	struct FInput
+	{
+		// @todo: add support for SlotName / multiple SlotAnimTracks
+		TWeakObjectPtr<UAnimMontage> AnimMontage;
+		FPoseSearchExtrapolationParameters ExtrapolationParameters;
+	} Input;
+
+	void Init(const FInput& Input);
+	virtual void Process() override;
+
+	virtual float GetPlayLength() const override;
+	virtual float GetScaledTime(float Time) const override;
+	virtual bool IsLoopable() const override;
+
+	virtual FTransform GetTotalRootTransform() const override;
+
+	virtual void ExtractPose(const FAnimExtractContext& ExtractionCtx, FAnimationPoseData& OutAnimPoseData) const override;
+	virtual FTransform ExtractRootTransform(float Time) const override;
+	virtual void ExtractPoseSearchNotifyStates(float Time, TArray<class UAnimNotifyState_PoseSearchBase*>& NotifyStates) const override;
+	virtual const UAnimationAsset* GetAsset() const override;
+
+private:
+	FTransform ExtractRootTransformInternal(float StartTime, float EndTime) const;
+};
+
diff --git a/Engine/Plugins/Experimental/Animation/PoseSearch/Source/Runtime/Public/PoseSearch/PoseSearchContext.h b/Engine/Plugins/Experimental/Animation/PoseSearch/Source/Runtime/Public/PoseSearch/PoseSearchContext.h
--- a/Engine/Plugins/Experimental/Animation/PoseSearch/Source/Runtime/Public/PoseSearch/PoseSearchContext.h
+++ b/Engine/Plugins/Experimental/Animation/PoseSearch/Source/Runtime/Public/PoseSearch/PoseSearchContext.h
@@ -149,15 +149,12 @@ POSESEARCH_API void DrawFeatureVector(FDebugDrawParams& DrawParams, int32 PoseId
 struct POSESEARCH_API FSearchContext
 {
 	EPoseSearchBooleanRequest QueryMirrorRequest = EPoseSearchBooleanRequest::Indifferent;
-	UE::PoseSearch::FDebugDrawParams DebugDrawParams;
-	const UE::PoseSearch::FPoseHistory* History = nullptr;
+	FDebugDrawParams DebugDrawParams;
+	const FPoseHistory* History = nullptr;
 	const FTrajectorySampleRange* Trajectory = nullptr;
-	TObjectPtr<const USkeletalMeshComponent> OwningComponent = nullptr;
-	UE::PoseSearch::FSearchResult CurrentResult;
-	const FBoneContainer* BoneContainer = nullptr;
+	FSearchResult CurrentResult;
 	const FGameplayTagContainer* ActiveTagsContainer = nullptr;
 	float PoseJumpThresholdTime = 0.f;
-	bool bIsTracing = false;
 	bool bForceInterrupt = false;
 	// can the continuing pose advance? (if not we skip evaluating it)
 	bool bCanAdvance = true;
diff --git a/Engine/Plugins/Experimental/Animation/PoseSearch/Source/Runtime/Public/PoseSearch/PoseSearchDatabase.h b/Engine/Plugins/Experimental/Animation/PoseSearch/Source/Runtime/Public/PoseSearch/PoseSearchDatabase.h
--- a/Engine/Plugins/Experimental/Animation/PoseSearch/Source/Runtime/Public/PoseSearch/PoseSearchDatabase.h
+++ b/Engine/Plugins/Experimental/Animation/PoseSearch/Source/Runtime/Public/PoseSearch/PoseSearchDatabase.h
@@ -62,10 +62,9 @@ struct POSESEARCH_API FPoseSearchDatabaseAnimationAssetBase
 	virtual void SetIsEnabled(bool bInIsEnabled) {}
 	virtual bool IsRootMotionEnabled() const { return false; }
 	virtual EPoseSearchMirrorOption GetMirrorOption() const { return EPoseSearchMirrorOption::Invalid; }
-	virtual ESearchIndexAssetType GetSearchIndexType() const { return ESearchIndexAssetType::Invalid; }
 };
 
-/** An entry in a UPoseSearchDatabase. */
+/** A sequence entry in a UPoseSearchDatabase. */
 USTRUCT(BlueprintType, Category = "Animation|Pose Search")
 struct POSESEARCH_API FPoseSearchDatabaseSequence : public FPoseSearchDatabaseAnimationAssetBase
 {
@@ -92,7 +91,6 @@ struct POSESEARCH_API FPoseSearchDatabaseSequence : public FPoseSearchDatabaseAn
 	void SetIsEnabled(bool bInIsEnabled) override { bEnabled = bInIsEnabled; }
 	bool IsRootMotionEnabled() const override;
 	EPoseSearchMirrorOption GetMirrorOption() const override { return MirrorOption; }
-	ESearchIndexAssetType GetSearchIndexType() const override { return ESearchIndexAssetType::Sequence; }
 };
 
 /** An blend space entry in a UPoseSearchDatabase. */
@@ -130,7 +128,6 @@ struct POSESEARCH_API FPoseSearchDatabaseBlendSpace : public FPoseSearchDatabase
 	void SetIsEnabled(bool bInIsEnabled) override { bEnabled = bInIsEnabled; }
 	bool IsRootMotionEnabled() const override;
 	EPoseSearchMirrorOption GetMirrorOption() const override { return MirrorOption; }
-	ESearchIndexAssetType GetSearchIndexType() const override { return ESearchIndexAssetType::BlendSpace; }
 
 	void GetBlendSpaceParameterSampleRanges(int32& HorizontalBlendNum, int32& VerticalBlendNum) const;
 	FVector BlendParameterForSampleRanges(int32 HorizontalBlendIndex, int32 VerticalBlendIndex) const;
@@ -163,9 +160,38 @@ struct POSESEARCH_API FPoseSearchDatabaseAnimComposite : public FPoseSearchDatab
 	void SetIsEnabled(bool bInIsEnabled) override { bEnabled = bInIsEnabled; }
 	bool IsRootMotionEnabled() const override;
 	EPoseSearchMirrorOption GetMirrorOption() const override { return MirrorOption; }
-	ESearchIndexAssetType GetSearchIndexType() const override { return ESearchIndexAssetType::AnimComposite; }
 };
 
+/** An anim montage entry in a UPoseSearchDatabase. */
+USTRUCT(BlueprintType, Category = "Animation|Pose Search")
+struct POSESEARCH_API FPoseSearchDatabaseAnimMontage : public FPoseSearchDatabaseAnimationAssetBase
+{
+	GENERATED_BODY()
+	virtual ~FPoseSearchDatabaseAnimMontage() = default;
+
+	UPROPERTY(EditAnywhere, Category="AnimMontage")
+	TObjectPtr<UAnimMontage> AnimMontage;
+
+	UPROPERTY(EditAnywhere, Category = "AnimMontage")
+	bool bEnabled = true;
+
+	UPROPERTY(EditAnywhere, Category="AnimMontage")
+	FFloatInterval SamplingRange = FFloatInterval(0.0f, 0.0f);
+
+	UPROPERTY(EditAnywhere, Category = "AnimMontage")
+	EPoseSearchMirrorOption MirrorOption = EPoseSearchMirrorOption::UnmirroredOnly;
+
+	UAnimationAsset* GetAnimationAsset() const override;
+	UClass* GetAnimationAssetStaticClass() const override;
+	bool IsLooping() const override;
+	const FString GetName() const override;
+	bool IsEnabled() const override { return bEnabled; }
+	void SetIsEnabled(bool bInIsEnabled) override { bEnabled = bInIsEnabled; }
+	bool IsRootMotionEnabled() const override;
+	EPoseSearchMirrorOption GetMirrorOption() const override { return MirrorOption; }
+};
+
+
 /** A data asset for indexing a collection of animation sequences. */
 UCLASS(BlueprintType, Category = "Animation|Pose Search", Experimental, meta = (DisplayName = "Normalization Set"))
 class POSESEARCH_API UNormalizationSetAsset : public UDataAsset
@@ -250,6 +276,7 @@ public:
 	const bool IsSourceAssetLooping(const FPoseSearchIndexAsset& SearchIndexAsset) const;
 	const FString GetSourceAssetName(const FPoseSearchIndexAsset& SearchIndexAsset) const;
 	int32 GetNumberOfPrincipalComponents() const;
+	float GetAssetTime(int32 PoseIdx, float SamplingInterval) const;
 
 	// Begin UObject
 	virtual void PostLoad() override;
diff --git a/Engine/Plugins/Experimental/Animation/PoseSearch/Source/Runtime/Public/PoseSearch/PoseSearchIndex.h b/Engine/Plugins/Experimental/Animation/PoseSearch/Source/Runtime/Public/PoseSearch/PoseSearchIndex.h
--- a/Engine/Plugins/Experimental/Animation/PoseSearch/Source/Runtime/Public/PoseSearch/PoseSearchIndex.h
+++ b/Engine/Plugins/Experimental/Animation/PoseSearch/Source/Runtime/Public/PoseSearch/PoseSearchIndex.h
@@ -20,15 +20,6 @@ ENUM_CLASS_FLAGS(EPoseComparisonFlags);
 
 } // namespace UE::PoseSearch
 
-UENUM()
-enum class ESearchIndexAssetType : int32
-{
-	Invalid,
-	Sequence,
-	BlendSpace,
-	AnimComposite,
-};
-
 UENUM()
 enum class EPoseSearchBooleanRequest : uint8
 {
@@ -86,22 +77,15 @@ struct POSESEARCH_API FPoseSearchIndexAsset
 	FPoseSearchIndexAsset() {}
 
 	FPoseSearchIndexAsset(
-		ESearchIndexAssetType InType,
 		int32 InSourceAssetIdx, 
 		bool bInMirrored, 
 		const FFloatInterval& InSamplingInterval,
 		FVector InBlendParameters = FVector::Zero())
-		: Type(InType)
-		, SourceAssetIdx(InSourceAssetIdx)
+		: SourceAssetIdx(InSourceAssetIdx)
 		, bMirrored(bInMirrored)
 		, BlendParameters(InBlendParameters)
 		, SamplingInterval(InSamplingInterval) {}
 
-	// Default to Sequence for now for backward compatibility but
-	// at some point we might want to change this to Invalid.
-	UPROPERTY(meta = (NeverInHash))
-	ESearchIndexAssetType Type = ESearchIndexAssetType::Sequence;
-
 	// Index of the source asset in search index's container (i.e. UPoseSearchDatabase)
 	UPROPERTY(meta = (NeverInHash))
 	int32 SourceAssetIdx = INDEX_NONE;
@@ -181,8 +165,6 @@ struct POSESEARCH_API FPoseSearchIndexBase
 	const FPoseSearchIndexAsset& GetAssetForPose(int32 PoseIdx) const;
 	const FPoseSearchIndexAsset* GetAssetForPoseSafe(int32 PoseIdx) const;
 
-	float GetAssetTime(int32 PoseIdx, float SamplingInterval) const;
-
 	void Reset();
 	
 	friend FArchive& operator<<(FArchive& Ar, FPoseSearchIndexBase& Index);
diff --git a/Engine/Plugins/Experimental/Animation/PoseSearch/Source/Runtime/Public/PoseSearch/PoseSearchLibrary.h b/Engine/Plugins/Experimental/Animation/PoseSearch/Source/Runtime/Public/PoseSearch/PoseSearchLibrary.h
--- a/Engine/Plugins/Experimental/Animation/PoseSearch/Source/Runtime/Public/PoseSearch/PoseSearchLibrary.h
+++ b/Engine/Plugins/Experimental/Animation/PoseSearch/Source/Runtime/Public/PoseSearch/PoseSearchLibrary.h
@@ -19,6 +19,7 @@ namespace UE::PoseSearch
 	struct FSearchContext;
 } // namespace UE::PoseSearch
 
+struct FAnimationUpdateContext;
 struct FGameplayTagContainer;
 struct FTrajectorySampleRange;
 class UPoseSearchSearchableAsset;
@@ -117,22 +118,70 @@ struct POSESEARCH_API FMotionMatchingState
 	UE::PoseSearch::FPoseIndicesHistory PoseIndicesHistory;
 };
 
-/**
-* Implementation of the core motion matching algorithm
-*
-* @param UpdateContext				Input animation update context providing access to the proxy and delta time
-* @param Database					Input collection of animations for motion matching
-* @param Trajectory					Input motion trajectory samples for pose search queries
-* @param Settings					Input motion matching algorithm configuration settings
-* @param InOutMotionMatchingState	Input/Output encapsulated motion matching algorithm and state
-*/
-POSESEARCH_API void UpdateMotionMatchingState(
-	const FAnimationUpdateContext& Context,
-	const UPoseSearchSearchableAsset* Searchable,
-	const FGameplayTagContainer* ActiveTagsContainer,
-	const FTrajectorySampleRange& Trajectory,
-	const FMotionMatchingSettings& Settings,
-	FMotionMatchingState& InOutMotionMatchingState,
-	bool bForceInterrupt
-);
+UCLASS()
+class POSESEARCH_API UPoseSearchLibrary : public UBlueprintFunctionLibrary
+{
+	GENERATED_BODY()
+	static void TraceMotionMatchingState(
+		const UPoseSearchSearchableAsset* Searchable,
+		UE::PoseSearch::FSearchContext& SearchContext,
+		const UE::PoseSearch::FSearchResult& CurrentResult,
+		const UE::PoseSearch::FSearchResult& LastResult,
+		float ElapsedPoseSearchTime,
+		const FTransform& RootMotionTransformDelta,
+		const UObject* AnimInstance,
+		int32 NodeId,
+		float DeltaTime,
+		bool bSearch);
+
+public:
+	/**
+	* Implementation of the core motion matching algorithm
+	*
+	* @param Context						Input animation update context providing access to the proxy and delta time
+	* @param Searchable						Input collection of animations for motion matching
+	* @param ActiveTagsContainer			Input gameplay tag container
+	* @param Trajectory						Input motion trajectory samples for pose search queries
+	* @param Settings						Input motion matching algorithm configuration settings
+	* @param InOutMotionMatchingState		Input/Output encapsulated motion matching algorithm and state
+	* @param bForceInterrupt				Input force interrupt request (if true the continuing pose will be invalidated)
+	*/
+	static void UpdateMotionMatchingState(
+		const FAnimationUpdateContext& Context,
+		const UPoseSearchSearchableAsset* Searchable,
+		const FGameplayTagContainer* ActiveTagsContainer,
+		const FTrajectorySampleRange& Trajectory,
+		const FMotionMatchingSettings& Settings,
+		FMotionMatchingState& InOutMotionMatchingState,
+		bool bForceInterrupt);
+
+	/**
+	* Implementation of the core motion matching algorithm
+	*
+	* @param AnimInstance					Input animation instance
+	* @param Searchable						Input searchable asset
+	* @param ActiveTagsContainer			Input gameplay tag container
+	* @param Trajectory						Input motion trajectory samples for pose search queries
+	* @param PoseHistoryName				Input tag of the associated PoseSearchHistoryCollector node in the anim graph
+	* @param SelectedAnimation				Output selected animation from the searchable asset
+	* @param SelectedTime					Output selected animation time
+	* @param bLoop							Output selected animation looping state
+	* @param bIsMirrored					Output selected animation mirror state
+	* @param BlendParameters				Output selected animation blend space parameters (if SelectedAnimation is a blend space)
+	* @param DebugSessionUniqueIdentifier	Input unique identifier used to identify TraceMotionMatchingState (rewind debugger / pose search debugger) session. Similarly the MM node uses Context.GetCurrentNodeId()
+	*/
+	UFUNCTION(BlueprintPure, Category = "Animation|Pose Search", meta = (BlueprintThreadSafe, Keywords = "PoseMatch"))
+	static void MotionMatch(
+		const UAnimInstance* AnimInstance,
+		const UPoseSearchSearchableAsset* Searchable,
+		const FGameplayTagContainer ActiveTagsContainer,
+		const FTrajectorySampleRange Trajectory,
+		const FName PoseHistoryName,
+		UAnimationAsset*& SelectedAnimation,
+		float& SelectedTime,
+		bool& bLoop,
+		bool& bIsMirrored,
+		FVector& BlendParameters,
+		const int DebugSessionUniqueIdentifier = 6174);
+};
 
