motion matching - FAnimNode_BlendStack_Standalone support for anim curve and attributes. QOL improvements when using pose search databases with montages

diff --git a/Engine/Plugins/Experimental/Animation/PoseSearch/Source/Runtime/Private/AnimNode_BlendStack.cpp b/Engine/Plugins/Experimental/Animation/PoseSearch/Source/Runtime/Private/AnimNode_BlendStack.cpp
--- a/Engine/Plugins/Experimental/Animation/PoseSearch/Source/Runtime/Private/AnimNode_BlendStack.cpp
+++ b/Engine/Plugins/Experimental/Animation/PoseSearch/Source/Runtime/Private/AnimNode_BlendStack.cpp
@@ -6,12 +6,11 @@
 #include "Animation/AnimComposite.h"
 #include "Animation/AnimSequence.h"
 #include "Animation/BlendSpace.h"
+#include "Animation/AnimMontage.h"
 #include "PoseSearch/PoseSearchDefines.h"
 
 #include UE_INLINE_GENERATED_CPP_BY_NAME(AnimNode_BlendStack)
 
-#define LOCTEXT_NAMESPACE "AnimNode_BlendStack"
-
 #if ENABLE_ANIM_DEBUG
 TAutoConsoleVariable<int32> CVarAnimBlendStackEnable(TEXT("a.AnimNode.BlendStack.Enable"), 1, TEXT("Enable / Disable Blend Stack"));
 TAutoConsoleVariable<int32> CVarAnimBlendStackPruningEnable(TEXT("a.AnimNode.BlendStack.Pruning.Enable"), 1, TEXT("Enable / Disable Blend Stack Pruning"));
@@ -25,11 +24,7 @@ void FPoseSearchAnimPlayer::Initialize(UAnimationAsset* AnimationAsset, float Ac
 
 	if (bMirrored && !MirrorDataTable)
 	{
-		UE_LOG(
-			LogPoseSearch,
-			Error,
-			TEXT("FPoseSearchAnimPlayer failed to Initialize for %s. Mirroring will not work becasue MirrorDataTable is missing"),
-			*GetNameSafe(AnimationAsset));
+		UE_LOG(LogPoseSearch, Error, TEXT("FPoseSearchAnimPlayer failed to Initialize for %s. Mirroring will not work becasue MirrorDataTable is missing"), *GetNameSafe(AnimationAsset));
 	}
 
 	if (BlendProfile != nullptr)
@@ -52,10 +47,14 @@ void FPoseSearchAnimPlayer::Initialize(UAnimationAsset* AnimationAsset, float Ac
 	MirrorNode.SetMirrorDataTable(MirrorDataTable);
 	MirrorNode.SetMirror(bMirrored);
 	
-	if (UAnimSequenceBase* Sequence = Cast<UAnimSequenceBase>(AnimationAsset))
+	if (Cast<UAnimMontage>(AnimationAsset))
+	{
+		UE_LOG(LogPoseSearch, Error, TEXT("FPoseSearchAnimPlayer unsupported AnimationAsset %s"), *GetNameSafe(AnimationAsset));
+	}
+	else if (UAnimSequenceBase* SequenceBase = Cast<UAnimSequenceBase>(AnimationAsset))
 	{
 		SequencePlayerNode.SetAccumulatedTime(AccumulatedTime);
-		SequencePlayerNode.SetSequence(Sequence);
+		SequencePlayerNode.SetSequence(SequenceBase);
 		SequencePlayerNode.SetLoopAnimation(bLoop);
 		SequencePlayerNode.SetPlayRate(PlayRate);
 	}
@@ -68,7 +67,7 @@ void FPoseSearchAnimPlayer::Initialize(UAnimationAsset* AnimationAsset, float Ac
 		BlendSpacePlayerNode.SetPlayRate(PlayRate);
 		BlendSpacePlayerNode.SetPosition(BlendParameters);
 	}
-	else 
+	else
 	{
 		checkNoEntry();
 	}
@@ -86,59 +85,50 @@ void FPoseSearchAnimPlayer::UpdatePlayRate(float PlayRate)
 	{
 		BlendSpacePlayerNode.SetPlayRate(PlayRate);
 	}
-	else
-	{
-		checkNoEntry();
-	}
 }
 
 void FPoseSearchAnimPlayer::StorePoseContext(const FPoseContext& PoseContext)
 {
-	UpdateSourceLinkNode();
-
-	const FBoneContainer& BoneContainer = PoseContext.Pose.GetBoneContainer();
-	const USkeleton* SkeletonAsset = BoneContainer.GetSkeletonAsset();
-	check(SkeletonAsset);
+	SequencePlayerNode.SetSequence(nullptr);
+	BlendSpacePlayerNode.SetBlendSpace(nullptr);
+	MirrorNode.SetSourceLinkNode(nullptr);
 
-	const FReferenceSkeleton& RefSkeleton = SkeletonAsset->GetReferenceSkeleton();
-	const int32 NumSkeletonBones = RefSkeleton.GetNum();
+	if (PoseContext.Pose.IsValid())
+	{
+		StoredPose.CopyBonesFrom(PoseContext.Pose);
+	}
 
-	StoredPose.SetNum(NumSkeletonBones);
-	for (FSkeletonPoseBoneIndex SkeletonBoneIdx(0); SkeletonBoneIdx != NumSkeletonBones; ++SkeletonBoneIdx)
+	if (PoseContext.Curve.IsValid())
 	{
-		FCompactPoseBoneIndex CompactBoneIdx = BoneContainer.GetCompactPoseIndexFromSkeletonPoseIndex(SkeletonBoneIdx);
-		StoredPose[SkeletonBoneIdx.GetInt()] = CompactBoneIdx.IsValid() ? PoseContext.Pose[CompactBoneIdx] : RefSkeleton.GetRefBonePose()[SkeletonBoneIdx.GetInt()];
+		StoredCurve.CopyFrom(PoseContext.Curve);
 	}
 
-	// @todo: perhaps copy PoseContext.Curve and PoseContext.CustomAttributes?
+	StoredAttributes.CopyFrom(PoseContext.CustomAttributes);
 }
 
 void FPoseSearchAnimPlayer::RestorePoseContext(FPoseContext& PoseContext) const
 {
 	check(!SequencePlayerNode.GetSequence() && !BlendSpacePlayerNode.GetBlendSpace());
 
-	const FBoneContainer& BoneContainer = PoseContext.Pose.GetBoneContainer();
-	const USkeleton* SkeletonAsset = BoneContainer.GetSkeletonAsset();
-	check(SkeletonAsset);
-
-	const FReferenceSkeleton& RefSkeleton = SkeletonAsset->GetReferenceSkeleton();
-	const int32 NumSkeletonBones = RefSkeleton.GetNum();
+	if (StoredPose.IsValid() && PoseContext.Pose.GetNumBones() == StoredPose.GetNumBones())
+	{
+		PoseContext.Pose.CopyBonesFrom(StoredPose);
+	}
+	else
+	{
+		PoseContext.Pose.ResetToRefPose();
+	}
 
-	for (const FBoneIndexType BoneIdx : BoneContainer.GetBoneIndicesArray())
+	if (StoredCurve.IsValid())
 	{
-		const FCompactPoseBoneIndex CompactBoneIdx(BoneIdx);
-		const FSkeletonPoseBoneIndex SkeletonBoneIdx = BoneContainer.GetSkeletonPoseIndexFromCompactPoseIndex(CompactBoneIdx);
-		if (SkeletonBoneIdx.IsValid() && BoneIdx < StoredPose.Num())
-		{
-			PoseContext.Pose[CompactBoneIdx] = StoredPose[SkeletonBoneIdx.GetInt()];
-		}
+		PoseContext.Curve.CopyFrom(StoredCurve);
 	}
 
-	// @todo: perhaps copy PoseContext.Curve and PoseContext.CustomAttributes?
+	PoseContext.CustomAttributes.CopyFrom(StoredAttributes);
 }
 
 
-// @todo: maybe implement copy/move constructors and assignement operator do so (or use a list instead of an array)
+// @todo: maybe implement copy/move constructors and assignment operator do so (or use a list instead of an array)
 // since we're making copies and moving this object in memory, we're using this method to set the MirrorNode SourceLinkNode when necessary
 void FPoseSearchAnimPlayer::UpdateSourceLinkNode()
 {
@@ -190,7 +180,6 @@ float FPoseSearchAnimPlayer::GetAccumulatedTime() const
 		return BlendSpacePlayerNode.GetAccumulatedTime();
 	}
 
-	checkNoEntry();
 	return 0.f;
 }
 
@@ -218,7 +207,6 @@ FString FPoseSearchAnimPlayer::GetAnimationName() const
 		return BlendSpacePlayerNode.GetBlendSpace()->GetName();
 	}
 
-	checkNoEntry();
 	return FString("StoredPose");
 }
 
@@ -234,7 +222,6 @@ const UAnimationAsset* FPoseSearchAnimPlayer::GetAnimationAsset() const
 		return BlendSpacePlayerNode.GetBlendSpace();
 	}
 
-	checkNoEntry();
 	return nullptr;
 }
 
@@ -341,7 +328,7 @@ void FAnimNode_BlendStack_Standalone::Evaluate_AnyThread(FPoseContext& Output)
 #endif // ENABLE_ANIM_DEBUG
 				)
 			{
-				// too many AnimPlayers! we don't have enought available blends to hold them all, so we accumulate the blended poses into Output / BlendedPoseContext, until...
+				// too many AnimPlayers! we don't have enough available blends to hold them all, so we accumulate the blended poses into Output / BlendedPoseContext, until...
 				AnimPlayers.PopLast();
 				
 				if (i == RequestedMaxActiveBlends)
@@ -486,4 +473,3 @@ void FAnimNode_BlendStack::UpdateAssetPlayer(const FAnimationUpdateContext& Cont
 
 	Super::UpdateAssetPlayer(Context);
 }
-#undef LOCTEXT_NAMESPACE
diff --git a/Engine/Plugins/Experimental/Animation/PoseSearch/Source/Runtime/Private/PoseSearchResult.cpp b/Engine/Plugins/Experimental/Animation/PoseSearch/Source/Runtime/Private/PoseSearchResult.cpp
--- a/Engine/Plugins/Experimental/Animation/PoseSearch/Source/Runtime/Private/PoseSearchResult.cpp
+++ b/Engine/Plugins/Experimental/Animation/PoseSearch/Source/Runtime/Private/PoseSearchResult.cpp
@@ -69,7 +69,6 @@ void FSearchResult::Update(float NewAssetTime)
 		}
 		else
 		{
-			checkNoEntry();
 			Reset();
 		}
 	}
diff --git a/Engine/Plugins/Experimental/Animation/PoseSearch/Source/Runtime/Public/PoseSearch/AnimNode_BlendStack.h b/Engine/Plugins/Experimental/Animation/PoseSearch/Source/Runtime/Public/PoseSearch/AnimNode_BlendStack.h
--- a/Engine/Plugins/Experimental/Animation/PoseSearch/Source/Runtime/Public/PoseSearch/AnimNode_BlendStack.h
+++ b/Engine/Plugins/Experimental/Animation/PoseSearch/Source/Runtime/Public/PoseSearch/AnimNode_BlendStack.h
@@ -39,18 +39,21 @@ protected:
 	void RestorePoseContext(FPoseContext& PoseContext) const;
 	void UpdateSourceLinkNode();
 
-	// Embedded sequence player node for playing animations from the motion matching database
+	// Embedded standalone player to play sequence
 	FAnimNode_SequencePlayer_Standalone SequencePlayerNode;
 
-	// Embedded blendspace player node for playing blendspaces from the motion matching database
+	// Embedded standalone player to play blend spaces
 	FAnimNode_BlendSpacePlayer_Standalone BlendSpacePlayerNode;
 
 	// Embedded mirror node to handle mirroring if the pose search results in a mirrored sequence
 	FAnimNode_Mirror_Standalone MirrorNode;
 
-	// @todo: use FCompactHeapPose
-	// if SequencePlayerNode.GetSequence() and BlendSpacePlayerNode.GetBlendSpace() are nullptr, instead of using SequencePlayerNode or BlendSpacePlayerNode (wrapped in MirrorNode), the pose output will be from StoredPose
-	TArray<FTransform> StoredPose;
+	// if SequencePlayerNode.GetSequence() and BlendSpacePlayerNode.GetBlendSpace() are nullptr, 
+	// instead of using SequencePlayerNode or BlendSpacePlayerNode (wrapped in MirrorNode),
+	// the output FPoseContext will be from StoredPose, StoredCurve, StoredAttributes
+	FCompactHeapPose StoredPose;
+	FBlendedHeapCurve StoredCurve;
+	UE::Anim::FHeapAttributeContainer StoredAttributes;
 
 	EAlphaBlendOption BlendOption = EAlphaBlendOption::Linear;
 
