motion matching - UPoseSearchFeatureChannel_Velocity indexing speed up by using only 2 position samples rather than 3

diff --git a/Engine/Plugins/Experimental/Animation/PoseSearch/Source/Runtime/Private/PoseSearchDerivedDataKey.cpp b/Engine/Plugins/Experimental/Animation/PoseSearch/Source/Runtime/Private/PoseSearchDerivedDataKey.cpp
--- a/Engine/Plugins/Experimental/Animation/PoseSearch/Source/Runtime/Private/PoseSearchDerivedDataKey.cpp
+++ b/Engine/Plugins/Experimental/Animation/PoseSearch/Source/Runtime/Private/PoseSearchDerivedDataKey.cpp
@@ -41,7 +41,7 @@ FKeyBuilder::FKeyBuilder(const UObject* Object, bool bUseDataVer)
 	if (bUseDataVer)
 	{
 		// used to invalidate the key without having to change POSESEARCHDB_DERIVEDDATA_VER all the times
-		int32 POSESEARCHDB_DERIVEDDATA_VER_SMALL = 51;
+		int32 POSESEARCHDB_DERIVEDDATA_VER_SMALL = 53;
 		FGuid VersionGuid = FDevSystemGuids::GetSystemGuid(FDevSystemGuids::Get().POSESEARCHDB_DERIVEDDATA_VER);
 
 		*this << VersionGuid;
diff --git a/Engine/Plugins/Experimental/Animation/PoseSearch/Source/Runtime/Private/PoseSearchFeatureChannel_Velocity.cpp b/Engine/Plugins/Experimental/Animation/PoseSearch/Source/Runtime/Private/PoseSearchFeatureChannel_Velocity.cpp
--- a/Engine/Plugins/Experimental/Animation/PoseSearch/Source/Runtime/Private/PoseSearchFeatureChannel_Velocity.cpp
+++ b/Engine/Plugins/Experimental/Animation/PoseSearch/Source/Runtime/Private/PoseSearchFeatureChannel_Velocity.cpp
@@ -45,24 +45,16 @@ void UPoseSearchFeatureChannel_Velocity::IndexAsset(UE::PoseSearch::IAssetIndexe
 		bool ClampedPast, ClampedPresent, ClampedFuture;
 		const FTransform BoneTransformsPast = Indexer.GetComponentSpaceTransform(SubsampleTime - SamplingContext->FiniteDelta, bUseCharacterSpaceVelocities ? OriginSampleTime - SamplingContext->FiniteDelta : OriginSampleTime, ClampedPast, SchemaBoneIdx);
 		const FTransform BoneTransformsPresent = Indexer.GetComponentSpaceTransform(SubsampleTime, OriginSampleTime, ClampedPresent, SchemaBoneIdx);
-		const FTransform BoneTransformsFuture = Indexer.GetComponentSpaceTransform(SubsampleTime + SamplingContext->FiniteDelta, bUseCharacterSpaceVelocities ? OriginSampleTime + SamplingContext->FiniteDelta : OriginSampleTime, ClampedFuture, SchemaBoneIdx);
 
-		// We can get a better finite difference if we ignore samples that have
-		// been clamped at either side of the clip. However, if the central sample 
-		// itself is clamped, or there are no samples that are clamped, we can just 
-		// use the central difference as normal.
 		FVector LinearVelocity;
-		if (ClampedPast && !ClampedPresent && !ClampedFuture)
-		{
-			LinearVelocity = (BoneTransformsFuture.GetTranslation() - BoneTransformsPresent.GetTranslation()) / SamplingContext->FiniteDelta;
-		}
-		else if (ClampedFuture && !ClampedPresent && !ClampedPast)
+		if (!ClampedPast)
 		{
 			LinearVelocity = (BoneTransformsPresent.GetTranslation() - BoneTransformsPast.GetTranslation()) / SamplingContext->FiniteDelta;
 		}
 		else
 		{
-			LinearVelocity = (BoneTransformsFuture.GetTranslation() - BoneTransformsPast.GetTranslation()) / (SamplingContext->FiniteDelta * 2.f);
+			const FTransform BoneTransformsFuture = Indexer.GetComponentSpaceTransform(SubsampleTime + SamplingContext->FiniteDelta, bUseCharacterSpaceVelocities ? OriginSampleTime + SamplingContext->FiniteDelta : OriginSampleTime, ClampedFuture, SchemaBoneIdx);
+			LinearVelocity = (BoneTransformsFuture.GetTranslation() - BoneTransformsPresent.GetTranslation()) / SamplingContext->FiniteDelta;
 		}
 
 		if (bNormalize)
