motion matching - guarding code with #if WITH_EDITOR and #if ENABLE_DRAW_DEBUG

diff --git a/Engine/Plugins/Experimental/Animation/PoseSearch/Source/Editor/Private/PoseSearchDatabaseViewModel.cpp b/Engine/Plugins/Experimental/Animation/PoseSearch/Source/Editor/Private/PoseSearchDatabaseViewModel.cpp
--- a/Engine/Plugins/Experimental/Animation/PoseSearch/Source/Editor/Private/PoseSearchDatabaseViewModel.cpp
+++ b/Engine/Plugins/Experimental/Animation/PoseSearch/Source/Editor/Private/PoseSearchDatabaseViewModel.cpp
@@ -36,21 +36,20 @@ namespace UE::PoseSearch
 
 	void FDatabasePreviewActor::Process()
 	{
-		FAssetSamplerBase* SamplerInterface = Sampler.GetMutablePtr<FAssetSamplerBase>();
-		check(SamplerInterface);
-		SamplerInterface->Process();
+		check(Sampler);
+		Sampler->Process();
 	}
 
 	const FAssetSamplerBase* FDatabasePreviewActor::GetSampler() const
 	{
-		const FAssetSamplerBase* SamplerInterface = Sampler.GetPtr<FAssetSamplerBase>();
-		check(SamplerInterface);
-		return SamplerInterface;
+		check(Sampler);
+		return Sampler.Get();
 	}
 
 	float FDatabasePreviewActor::GetScaledTime(float Time) const
 	{
-		return GetSampler()->GetScaledTime(Time);
+		check(Sampler);
+		return Sampler->GetScaledTime(Time);
 	}
 
 	UDebugSkelMeshComponent* FDatabasePreviewActor::GetDebugSkelMeshComponent()
@@ -216,10 +215,9 @@ namespace UE::PoseSearch
 			Input.ExtrapolationParameters = PoseSearchDatabase->ExtrapolationParameters;
 			Input.SequenceBase = DatabaseSequence->Sequence;
 
-			PreviewActor.Sampler = FInstancedStruct::Make(FSequenceBaseSampler());
-			FSequenceBaseSampler* Sampler = PreviewActor.Sampler.GetMutablePtr<FSequenceBaseSampler>();
-			check(Sampler);
+			TSharedPtr<FSequenceBaseSampler> Sampler = MakeShared<FSequenceBaseSampler>();
 			Sampler->Init(Input);
+			PreviewActor.Sampler = Sampler;
 		}
 		else if (const FPoseSearchDatabaseAnimComposite* DatabaseAnimComposite = DatabaseAsset.GetPtr<FPoseSearchDatabaseAnimComposite>())
 		{
@@ -227,10 +225,9 @@ namespace UE::PoseSearch
 			Input.ExtrapolationParameters = PoseSearchDatabase->ExtrapolationParameters;
 			Input.SequenceBase = DatabaseAnimComposite->AnimComposite;
 
-			PreviewActor.Sampler = FInstancedStruct::Make(FSequenceBaseSampler());
-			FSequenceBaseSampler* Sampler = PreviewActor.Sampler.GetMutablePtr<FSequenceBaseSampler>();
-			check(Sampler);
+			TSharedPtr<FSequenceBaseSampler> Sampler = MakeShared<FSequenceBaseSampler>();
 			Sampler->Init(Input);
+			PreviewActor.Sampler = Sampler;
 		}
 		else if (const FPoseSearchDatabaseBlendSpace* DatabaseBlendSpace = DatabaseAsset.GetPtr<FPoseSearchDatabaseBlendSpace>())
 		{
@@ -240,10 +237,9 @@ namespace UE::PoseSearch
 			Input.BlendSpace = DatabaseBlendSpace->BlendSpace;
 			Input.BlendParameters = IndexAsset.BlendParameters;
 
-			PreviewActor.Sampler = FInstancedStruct::Make(FBlendSpaceSampler());
-			FBlendSpaceSampler* Sampler = PreviewActor.Sampler.GetMutablePtr<FBlendSpaceSampler>();
-			check(Sampler);
+			TSharedPtr<FBlendSpaceSampler> Sampler = MakeShared<FBlendSpaceSampler>();
 			Sampler->Init(Input);
+			PreviewActor.Sampler = Sampler;
 		}
 		else if (const FPoseSearchDatabaseAnimMontage* DatabaseAnimMontage = DatabaseAsset.GetPtr<FPoseSearchDatabaseAnimMontage>())
 		{
@@ -251,10 +247,9 @@ namespace UE::PoseSearch
 			Input.ExtrapolationParameters = PoseSearchDatabase->ExtrapolationParameters;
 			Input.AnimMontage = DatabaseAnimMontage->AnimMontage;
 
-			PreviewActor.Sampler = FInstancedStruct::Make(FAnimMontageSampler());
-			FAnimMontageSampler* Sampler = PreviewActor.Sampler.GetMutablePtr<FAnimMontageSampler>();
-			check(Sampler);
+			TSharedPtr<FAnimMontageSampler> Sampler = MakeShared<FAnimMontageSampler>();
 			Sampler->Init(Input);
+			PreviewActor.Sampler = Sampler;
 		}
 		else
 		{
diff --git a/Engine/Plugins/Experimental/Animation/PoseSearch/Source/Editor/Public/PoseSearchDatabaseViewModel.h b/Engine/Plugins/Experimental/Animation/PoseSearch/Source/Editor/Public/PoseSearchDatabaseViewModel.h
--- a/Engine/Plugins/Experimental/Animation/PoseSearch/Source/Editor/Public/PoseSearchDatabaseViewModel.h
+++ b/Engine/Plugins/Experimental/Animation/PoseSearch/Source/Editor/Public/PoseSearchDatabaseViewModel.h
@@ -47,7 +47,7 @@ namespace UE::PoseSearch
 		int32 IndexAssetIndex = INDEX_NONE;
 		int32 CurrentPoseIndex = INDEX_NONE;
 
-		FInstancedStruct Sampler;
+		TSharedPtr<FAssetSamplerBase> Sampler;
 
 		bool IsValid() const;
 		void Process();
diff --git a/Engine/Plugins/Experimental/Animation/PoseSearch/Source/Runtime/Private/AnimNode_MotionMatching.cpp b/Engine/Plugins/Experimental/Animation/PoseSearch/Source/Runtime/Private/AnimNode_MotionMatching.cpp
--- a/Engine/Plugins/Experimental/Animation/PoseSearch/Source/Runtime/Private/AnimNode_MotionMatching.cpp
+++ b/Engine/Plugins/Experimental/Animation/PoseSearch/Source/Runtime/Private/AnimNode_MotionMatching.cpp
@@ -9,7 +9,6 @@
 #include "Animation/BlendSpace.h"
 #include "Animation/MotionTrajectoryTypes.h"
 #include "Components/SkeletalMeshComponent.h"
-#include "DrawDebugHelpers.h"
 #include "PoseSearch/AnimNode_PoseSearchHistoryCollector.h"
 #include "PoseSearch/PoseSearchDatabase.h"
 #include "PoseSearch/PoseSearchDerivedData.h"
@@ -68,7 +67,7 @@ void FAnimNode_MotionMatching::UpdateAssetPlayer(const FAnimationUpdateContext&
 
 	GetEvaluateGraphExposedInputs().Execute(Context);
 
-	bool bNeedsReset =
+	const bool bNeedsReset =
 		bResetOnBecomingRelevant &&
 		UpdateCounter.HasEverBeenUpdated() &&
 		!UpdateCounter.WasSynchronizedCounter(Context.AnimInstanceProxy->GetUpdateCounter());
diff --git a/Engine/Plugins/Experimental/Animation/PoseSearch/Source/Runtime/Private/PoseSearchAssetIndexer.cpp b/Engine/Plugins/Experimental/Animation/PoseSearch/Source/Runtime/Private/PoseSearchAssetIndexer.cpp
--- a/Engine/Plugins/Experimental/Animation/PoseSearch/Source/Runtime/Private/PoseSearchAssetIndexer.cpp
+++ b/Engine/Plugins/Experimental/Animation/PoseSearch/Source/Runtime/Private/PoseSearchAssetIndexer.cpp
@@ -1,15 +1,15 @@
 // Copyright Epic Games, Inc. All Rights Reserved.
 
-#include "PoseSearchAssetIndexer.h"
-
 #if WITH_EDITOR
 
+#include "PoseSearch/PoseSearchAssetIndexer.h"
 #include "AnimationRuntime.h"
 #include "Animation/MirrorDataTable.h"
 #include "PoseSearch/PoseSearchAnimNotifies.h"
 #include "PoseSearch/PoseSearchAssetSampler.h"
 #include "PoseSearch/PoseSearchDefines.h"
 #include "PoseSearch/PoseSearchFeatureChannel.h"
+#include "PoseSearch/PoseSearchIndex.h"
 #include "PoseSearch/PoseSearchSchema.h"
 
 namespace UE::PoseSearch
@@ -74,29 +74,10 @@ static FSamplingParam WrapOrClampSamplingParam(bool bCanWrap, float SamplingPara
 
 //////////////////////////////////////////////////////////////////////////
 // FAssetIndexer
-void FAssetIndexer::FOutput::Reset()
-{
-	FirstIndexedSample = 0;
-	LastIndexedSample = 0;
-	NumIndexedPoses = 0;
-
-	FeatureVectorTable.Reset(0);
-	PoseMetadata.Reset(0);
-}
-
-void FAssetIndexer::FStats::Reset()
-{
-	NumAccumulatedSamples = 0;
-	AccumulatedSpeed = 0.f;
-	MaxSpeed = 0.f;
-	AccumulatedAcceleration = 0.f;
-	MaxAcceleration = 0.f;
-}
-
 void FAssetIndexer::Reset()
 {
-	Output.Reset();
-	Stats.Reset();
+	Output = FOutput();
+	Stats = FStats();
 }
 
 void FAssetIndexer::Init(const FAssetIndexingContext& InIndexingContext, const FBoneContainer& InBoneContainer)
@@ -157,7 +138,7 @@ bool FAssetIndexer::Process()
 
 void FAssetIndexer::ComputeStats()
 {
-	Stats.Reset();
+	Stats = FStats();
 
 	const FAssetSamplingContext* SamplingContext = IndexingContext.SamplingContext;
 	check(SamplingContext->FiniteDelta > UE_KINDA_SMALL_NUMBER);
@@ -336,7 +317,7 @@ FAssetIndexer::CachedEntry& FAssetIndexer::GetEntry(float SampleTime)
 				BoneContainer.GetCompactPoseNumBones());
 		}
 
-		const IAssetIndexer::FSampleInfo Sample = GetSampleInfo(SampleTime);
+		const FAssetIndexer::FSampleInfo Sample = GetSampleInfo(SampleTime);
 		float CurrentTime = Sample.ClipTime;
 		float PreviousTime = CurrentTime - SamplingContext->FiniteDelta;
 
@@ -385,17 +366,17 @@ FAssetIndexer::CachedEntry& FAssetIndexer::GetEntry(float SampleTime)
 
 		Entry->ComponentSpacePose.InitPose(MoveTemp(Pose));
 		Entry->RootTransform = Sample.RootTransform;
-		Entry->Clamped = Sample.bClamped;
+		Entry->bClamped = Sample.bClamped;
 	}
 
 	return *Entry;
 }
 
 // returns the transform in component space for the bone indexed by Schema->BoneReferences[SchemaBoneIdx] at SampleTime seconds
-FTransform FAssetIndexer::GetComponentSpaceTransform(float SampleTime, bool& Clamped, int8 SchemaBoneIdx)
+FTransform FAssetIndexer::GetComponentSpaceTransform(float SampleTime, bool& bClamped, int8 SchemaBoneIdx)
 {
 	CachedEntry& Entry = GetEntry(SampleTime);
-	Clamped = Entry.Clamped;
+	bClamped = Entry.bClamped;
 
 	if (SchemaBoneIdx >= 0 && IndexingContext.Schema->BoneReferences[SchemaBoneIdx].HasValidSetup())
 	{
@@ -406,10 +387,10 @@ FTransform FAssetIndexer::GetComponentSpaceTransform(float SampleTime, bool& Cla
 }
 
 // returns the transform in animation space for the bone indexed by Schema->BoneReferences[SchemaBoneIdx] at SampleTime seconds
-FTransform FAssetIndexer::GetTransform(float SampleTime, bool& Clamped, int8 SchemaBoneIdx)
+FTransform FAssetIndexer::GetTransform(float SampleTime, bool& bClamped, int8 SchemaBoneIdx)
 {
 	CachedEntry& Entry = GetEntry(SampleTime);
-	Clamped = Entry.Clamped;
+	bClamped = Entry.bClamped;
 
 	const FTransform MirroredRootTransform = MirrorTransform(Entry.RootTransform);
 	if (SchemaBoneIdx >= 0 && IndexingContext.Schema->BoneReferences[SchemaBoneIdx].HasValidSetup())
@@ -429,16 +410,16 @@ FTransform FAssetIndexer::CalculateComponentSpaceTransform(FAssetIndexer::Cached
 
 // returns the transform in component space for the bone indexed by Schema->BoneReferences[SchemaBoneIdx] 
 // at SampleTime-OriginTime seconds ahead (or behind) the pose at OriginTime
-FTransform FAssetIndexer::GetComponentSpaceTransform(float SampleTime, float OriginTime, bool& Clamped, int8 SchemaBoneIdx)
+FTransform FAssetIndexer::GetComponentSpaceTransform(float SampleTime, float OriginTime, bool& bClamped, int8 SchemaBoneIdx)
 {
 	if (SampleTime == OriginTime)
 	{
-		return GetComponentSpaceTransform(SampleTime, Clamped, SchemaBoneIdx);
+		return GetComponentSpaceTransform(SampleTime, bClamped, SchemaBoneIdx);
 	}
 
-	bool Unused = false;
-	const FTransform RootBoneTransform = GetTransform(OriginTime, Unused);
-	FTransform BoneTransform = GetTransform(SampleTime, Clamped, SchemaBoneIdx);
+	bool bUnused = false;
+	const FTransform RootBoneTransform = GetTransform(OriginTime, bUnused);
+	FTransform BoneTransform = GetTransform(SampleTime, bClamped, SchemaBoneIdx);
 	BoneTransform.SetToRelativeTransform(RootBoneTransform);
 	return BoneTransform;
 }
diff --git a/Engine/Plugins/Experimental/Animation/PoseSearch/Source/Runtime/Private/PoseSearchAssetIndexer.h b/Engine/Plugins/Experimental/Animation/PoseSearch/Source/Runtime/Private/PoseSearchAssetIndexer.h
deleted file mode 100644
diff --git a/Engine/Plugins/Experimental/Animation/PoseSearch/Source/Runtime/Private/PoseSearchAssetSampler.cpp b/Engine/Plugins/Experimental/Animation/PoseSearch/Source/Runtime/Private/PoseSearchAssetSampler.cpp
--- a/Engine/Plugins/Experimental/Animation/PoseSearch/Source/Runtime/Private/PoseSearchAssetSampler.cpp
+++ b/Engine/Plugins/Experimental/Animation/PoseSearch/Source/Runtime/Private/PoseSearchAssetSampler.cpp
@@ -1,5 +1,7 @@
 // Copyright Epic Games, Inc. All Rights Reserved.
 
+#if WITH_EDITOR
+
 #include "PoseSearch/PoseSearchAssetSampler.h"
 #include "AnimationRuntime.h"
 #include "Animation/AnimMontage.h"
@@ -109,8 +111,6 @@ FTransform FAssetSamplingContext::MirrorTransform(const FTransform& InTransform)
 	return Result;
 }
 
-} // namespace UE::PoseSearch
-
 //////////////////////////////////////////////////////////////////////////
 // FSequenceBaseSampler
 void FSequenceBaseSampler::Init(const FInput& InInput)
@@ -731,3 +731,7 @@ const UAnimationAsset* FAnimMontageSampler::GetAsset() const
 {
 	return Input.AnimMontage.Get();
 }
+
+} // namespace UE::PoseSearch
+
+#endif // WITH_EDITOR
\ No newline at end of file
diff --git a/Engine/Plugins/Experimental/Animation/PoseSearch/Source/Runtime/Private/PoseSearchContext.cpp b/Engine/Plugins/Experimental/Animation/PoseSearch/Source/Runtime/Private/PoseSearchContext.cpp
--- a/Engine/Plugins/Experimental/Animation/PoseSearch/Source/Runtime/Private/PoseSearchContext.cpp
+++ b/Engine/Plugins/Experimental/Animation/PoseSearch/Source/Runtime/Private/PoseSearchContext.cpp
@@ -3,7 +3,6 @@
 #include "PoseSearch/PoseSearchContext.h"
 #include "Animation/MotionTrajectoryTypes.h"
 #include "AnimationRuntime.h"
-#include "DrawDebugHelpers.h"
 #include "PoseSearch/PoseSearchDatabase.h"
 #include "PoseSearch/PoseSearchFeatureChannel.h"
 #include "PoseSearch/PoseSearchHistory.h"
@@ -12,20 +11,16 @@
 namespace UE::PoseSearch
 {
 	
+#if ENABLE_DRAW_DEBUG
 //////////////////////////////////////////////////////////////////////////
 // FDebugDrawParams
 bool FDebugDrawParams::CanDraw() const
 {
-#if ENABLE_DRAW_DEBUG
 	return World && Database && Database->Schema && Database->Schema->IsValid();
-#else // ENABLE_DRAW_DEBUG
-	return false;
-#endif // ENABLE_DRAW_DEBUG
 }
 
 FColor FDebugDrawParams::GetColor(int32 ColorPreset) const
 {
-#if ENABLE_DRAW_DEBUG
 	FLinearColor Color = FLinearColor::Red;
 
 	const UPoseSearchSchema* Schema = GetSchema();
@@ -57,9 +52,6 @@ FColor FDebugDrawParams::GetColor(int32 ColorPreset) const
 	}
 
 	return Color.ToFColor(true);
-#else // ENABLE_DRAW_DEBUG
-	return FColor::Black;
-#endif // ENABLE_DRAW_DEBUG
 }
 
 const FPoseSearchIndex* FDebugDrawParams::GetSearchIndex() const
@@ -96,7 +88,7 @@ FVector FDebugDrawParams::GetCachedPosition(float TimeOffset, int8 SchemaBoneIdx
 
 		if (Mesh.IsValid() && SchemaBoneIdx >= 0)
 		{
-			return Mesh->GetSocketTransform(Schema->BoneReferences[SchemaBoneIdx].BoneName).GetLocation();
+			return Mesh->GetSocketTransform(Schema->BoneReferences[SchemaBoneIdx].BoneName).GetTranslation();
 		}
 	}
 	return RootTransform.GetTranslation();
@@ -104,7 +96,6 @@ FVector FDebugDrawParams::GetCachedPosition(float TimeOffset, int8 SchemaBoneIdx
 
 void DrawFeatureVector(FDebugDrawParams& DrawParams, TConstArrayView<float> PoseVector)
 {
-#if ENABLE_DRAW_DEBUG
 	DrawParams.ClearCachedPositions();
 
 	if (DrawParams.CanDraw())
@@ -131,19 +122,17 @@ void DrawFeatureVector(FDebugDrawParams& DrawParams, TConstArrayView<float> Pose
 			}
 		}
 	}
-#endif // ENABLE_DRAW_DEBUG
 }
 
 void DrawFeatureVector(FDebugDrawParams& DrawParams, int32 PoseIdx)
 {
-#if ENABLE_DRAW_DEBUG
 	// if we're editing the schema while in PIE with Rewind Debugger active, PoseIdx could be out of bound / stale
 	if (DrawParams.CanDraw() && PoseIdx >= 0 && PoseIdx < DrawParams.GetSearchIndex()->NumPoses)
 	{
 		DrawFeatureVector(DrawParams, DrawParams.GetSearchIndex()->GetPoseValues(PoseIdx));
 	}
-#endif // ENABLE_DRAW_DEBUG
 }
+#endif // ENABLE_DRAW_DEBUG
 
 //////////////////////////////////////////////////////////////////////////
 // FSearchContext
diff --git a/Engine/Plugins/Experimental/Animation/PoseSearch/Source/Runtime/Private/PoseSearchDatabaseIndexingContext.cpp b/Engine/Plugins/Experimental/Animation/PoseSearch/Source/Runtime/Private/PoseSearchDatabaseIndexingContext.cpp
--- a/Engine/Plugins/Experimental/Animation/PoseSearch/Source/Runtime/Private/PoseSearchDatabaseIndexingContext.cpp
+++ b/Engine/Plugins/Experimental/Animation/PoseSearch/Source/Runtime/Private/PoseSearchDatabaseIndexingContext.cpp
@@ -1,9 +1,8 @@
 // Copyright Epic Games, Inc. All Rights Reserved.
 
-#include "PoseSearchDatabaseIndexingContext.h"
-
 #if WITH_EDITOR
 
+#include "PoseSearchDatabaseIndexingContext.h"
 #include "Animation/AnimComposite.h"
 #include "Animation/AnimMontage.h"
 #include "Animation/AnimSequence.h"
@@ -41,9 +40,10 @@ void FDatabaseIndexingContext::Prepare(const UPoseSearchDatabase* Database)
 				Input.SequenceBase = Sequence;
 
 				SamplerMap.Add(Sequence, Samplers.Num());
-				FSequenceBaseSampler* Sampler = Samplers.Add_GetRef(FInstancedStruct::Make(FSequenceBaseSampler())).GetMutablePtr<FSequenceBaseSampler>();
-				check(Sampler);
+
+				TSharedPtr<FSequenceBaseSampler> Sampler = MakeShared<FSequenceBaseSampler>();
 				Sampler->Init(Input);
+				Samplers.Add(Sampler);
 			}
 		};
 
@@ -77,9 +77,10 @@ void FDatabaseIndexingContext::Prepare(const UPoseSearchDatabase* Database)
 							Input.BlendParameters = BlendParameters;
 
 							BlendSpaceSamplerMap.Add({ DatabaseBlendSpace->BlendSpace, BlendParameters }, Samplers.Num());
-							FBlendSpaceSampler* Sampler = Samplers.Add_GetRef(FInstancedStruct::Make(FBlendSpaceSampler())).GetMutablePtr<FBlendSpaceSampler>();
-							check(Sampler);
+
+							TSharedPtr<FBlendSpaceSampler> Sampler = MakeShared<FBlendSpaceSampler>();
 							Sampler->Init(Input);
+							Samplers.Add(Sampler);
 						}
 					}
 				}
@@ -94,9 +95,10 @@ void FDatabaseIndexingContext::Prepare(const UPoseSearchDatabase* Database)
 				Input.AnimMontage = DatabaseAnimMontage->AnimMontage;
 
 				SamplerMap.Add(DatabaseAnimMontage->AnimMontage, Samplers.Num());
-				FAnimMontageSampler* Sampler = Samplers.Add_GetRef(FInstancedStruct::Make(FAnimMontageSampler())).GetMutablePtr<FAnimMontageSampler>();
-				check(Sampler);
+
+				TSharedPtr<FAnimMontageSampler> Sampler = MakeShared<FAnimMontageSampler>();
 				Sampler->Init(Input);
+				Samplers.Add(Sampler);
 			}
 		}
 		else
@@ -105,12 +107,7 @@ void FDatabaseIndexingContext::Prepare(const UPoseSearchDatabase* Database)
 		}
 	}
 
-	ParallelFor(Samplers.Num(), [this](int32 SamplerIdx)
-		{
-			FAssetSamplerBase* AssetSamplerBase = Samplers[SamplerIdx].GetMutablePtr<FAssetSamplerBase>();
-			check(AssetSamplerBase);
-			AssetSamplerBase->Process();
-		}, ParallelForFlags);
+	ParallelFor(Samplers.Num(), [this](int32 SamplerIdx) { Samplers[SamplerIdx]->Process(); }, ParallelForFlags);
 
 	// prepare indexers
 	Indexers.Reserve(SearchIndexBase->Assets.Num());
@@ -131,28 +128,28 @@ void FDatabaseIndexingContext::Prepare(const UPoseSearchDatabase* Database)
 			if (DatabaseSequence->Sequence)
 			{
 				const float SequenceLength = DatabaseSequence->Sequence->GetPlayLength();
-				IndexerContext.AssetSampler = Samplers[SamplerMap[DatabaseSequence->Sequence]].GetPtr<FAssetSamplerBase>();
+				IndexerContext.AssetSampler = Samplers[SamplerMap[DatabaseSequence->Sequence]];
 			}
 		}
 		else if (const FPoseSearchDatabaseAnimComposite* DatabaseAnimComposite = DatabaseAsset.GetPtr<FPoseSearchDatabaseAnimComposite>())
 		{
 			if (DatabaseAnimComposite->AnimComposite)
 			{
-				IndexerContext.AssetSampler = Samplers[SamplerMap[DatabaseAnimComposite->AnimComposite]].GetPtr<FAssetSamplerBase>();
+				IndexerContext.AssetSampler = Samplers[SamplerMap[DatabaseAnimComposite->AnimComposite]];
 			}
 		}
 		else if (const FPoseSearchDatabaseBlendSpace* DatabaseBlendSpace = DatabaseAsset.GetPtr<FPoseSearchDatabaseBlendSpace>())
 		{
 			if (DatabaseBlendSpace->BlendSpace)
 			{
-				IndexerContext.AssetSampler = Samplers[BlendSpaceSamplerMap[{DatabaseBlendSpace->BlendSpace, SearchIndexAsset.BlendParameters}]].GetPtr<FAssetSamplerBase>();
+				IndexerContext.AssetSampler = Samplers[BlendSpaceSamplerMap[{DatabaseBlendSpace->BlendSpace, SearchIndexAsset.BlendParameters}]];
 			}
 		}
 		else if (const FPoseSearchDatabaseAnimMontage* DatabaseAnimMontage = DatabaseAsset.GetPtr<FPoseSearchDatabaseAnimMontage>())
 		{
 			if (DatabaseAnimMontage->AnimMontage)
 			{
-				IndexerContext.AssetSampler = Samplers[SamplerMap[DatabaseAnimMontage->AnimMontage]].GetPtr<FAssetSamplerBase>();
+				IndexerContext.AssetSampler = Samplers[SamplerMap[DatabaseAnimMontage->AnimMontage]];
 			}
 		}
 		else
diff --git a/Engine/Plugins/Experimental/Animation/PoseSearch/Source/Runtime/Private/PoseSearchDatabaseIndexingContext.h b/Engine/Plugins/Experimental/Animation/PoseSearch/Source/Runtime/Private/PoseSearchDatabaseIndexingContext.h
--- a/Engine/Plugins/Experimental/Animation/PoseSearch/Source/Runtime/Private/PoseSearchDatabaseIndexingContext.h
+++ b/Engine/Plugins/Experimental/Animation/PoseSearch/Source/Runtime/Private/PoseSearchDatabaseIndexingContext.h
@@ -4,9 +4,8 @@
 
 #if WITH_EDITOR
 
-#include "InstancedStruct.h"
+#include "PoseSearch/PoseSearchAssetIndexer.h"
 #include "PoseSearch/PoseSearchAssetSampler.h"
-#include "PoseSearchAssetIndexer.h"
 
 struct FPoseSearchIndexBase;
 class UPoseSearchDatabase;
@@ -18,7 +17,7 @@ struct FDatabaseIndexingContext
 {
 	FPoseSearchIndexBase* SearchIndexBase = nullptr;
 	FAssetSamplingContext SamplingContext;
-	TArray<FInstancedStruct> Samplers;
+	TArray<TSharedPtr<FAssetSamplerBase>> Samplers;
 	TArray<FAssetIndexer> Indexers;
 
 	void Prepare(const UPoseSearchDatabase* Database);
diff --git a/Engine/Plugins/Experimental/Animation/PoseSearch/Source/Runtime/Private/PoseSearchDerivedDataKey.cpp b/Engine/Plugins/Experimental/Animation/PoseSearch/Source/Runtime/Private/PoseSearchDerivedDataKey.cpp
--- a/Engine/Plugins/Experimental/Animation/PoseSearch/Source/Runtime/Private/PoseSearchDerivedDataKey.cpp
+++ b/Engine/Plugins/Experimental/Animation/PoseSearch/Source/Runtime/Private/PoseSearchDerivedDataKey.cpp
@@ -41,7 +41,7 @@ FKeyBuilder::FKeyBuilder(const UObject* Object, bool bUseDataVer)
 	if (bUseDataVer)
 	{
 		// used to invalidate the key without having to change POSESEARCHDB_DERIVEDDATA_VER all the times
-		int32 POSESEARCHDB_DERIVEDDATA_VER_SMALL = 53;
+		int32 POSESEARCHDB_DERIVEDDATA_VER_SMALL = 61;
 		FGuid VersionGuid = FDevSystemGuids::GetSystemGuid(FDevSystemGuids::Get().POSESEARCHDB_DERIVEDDATA_VER);
 
 		*this << VersionGuid;
diff --git a/Engine/Plugins/Experimental/Animation/PoseSearch/Source/Runtime/Private/PoseSearchFeatureChannel_FilterCrashingLegs.cpp b/Engine/Plugins/Experimental/Animation/PoseSearch/Source/Runtime/Private/PoseSearchFeatureChannel_FilterCrashingLegs.cpp
--- a/Engine/Plugins/Experimental/Animation/PoseSearch/Source/Runtime/Private/PoseSearchFeatureChannel_FilterCrashingLegs.cpp
+++ b/Engine/Plugins/Experimental/Animation/PoseSearch/Source/Runtime/Private/PoseSearchFeatureChannel_FilterCrashingLegs.cpp
@@ -1,7 +1,6 @@
 // Copyright Epic Games, Inc. All Rights Reserved.
 
 #include "PoseSearchFeatureChannel_FilterCrashingLegs.h"
-#include "DrawDebugHelpers.h"
 #include "PoseSearch/PoseSearchAssetIndexer.h"
 #include "PoseSearch/PoseSearchAssetSampler.h"
 #include "PoseSearch/PoseSearchContext.h"
@@ -39,6 +38,7 @@ void UPoseSearchFeatureChannel_FilterCrashingLegs::Finalize(UPoseSearchSchema* S
 	RightFootIdx = Schema->AddBoneReference(RightFoot);
 }
 
+#if WITH_EDITOR
 void UPoseSearchFeatureChannel_FilterCrashingLegs::FillWeights(TArray<float>& Weights) const
 {
 	for (int32 i = 0; i < ChannelCardinality; ++i)
@@ -47,7 +47,7 @@ void UPoseSearchFeatureChannel_FilterCrashingLegs::FillWeights(TArray<float>& We
 	}
 }
 
-void UPoseSearchFeatureChannel_FilterCrashingLegs::IndexAsset(UE::PoseSearch::IAssetIndexer& Indexer, TArrayView<float> FeatureVectorTable) const
+void UPoseSearchFeatureChannel_FilterCrashingLegs::IndexAsset(UE::PoseSearch::FAssetIndexer& Indexer, TArrayView<float> FeatureVectorTable) const
 {
 	using namespace UE::PoseSearch;
 
@@ -69,6 +69,7 @@ void UPoseSearchFeatureChannel_FilterCrashingLegs::IndexAsset(UE::PoseSearch::IA
 		FFeatureVectorHelper::EncodeFloat(IndexingContext.GetPoseVector(VectorIdx, FeatureVectorTable), ChannelDataOffset, CrashingLegsValue);
 	}
 }
+#endif // WITH_EDITOR
 
 void UPoseSearchFeatureChannel_FilterCrashingLegs::BuildQuery(UE::PoseSearch::FSearchContext& SearchContext, FPoseSearchFeatureVectorBuilder& InOutQuery) const
 {
@@ -99,9 +100,9 @@ void UPoseSearchFeatureChannel_FilterCrashingLegs::BuildQuery(UE::PoseSearch::FS
 	}
 }
 
+#if ENABLE_DRAW_DEBUG
 void UPoseSearchFeatureChannel_FilterCrashingLegs::DebugDraw(const UE::PoseSearch::FDebugDrawParams& DrawParams, TConstArrayView<float> PoseVector) const
 {
-#if ENABLE_DRAW_DEBUG
 	using namespace UE::PoseSearch;
 
 	const float CrashingLegsValue = FFeatureVectorHelper::DecodeFloat(PoseVector, ChannelDataOffset);
@@ -139,8 +140,8 @@ void UPoseSearchFeatureChannel_FilterCrashingLegs::DebugDraw(const UE::PoseSearc
 
 	DrawDebugLine(DrawParams.World, LeftFootPosition, LeftFootPosition + CrossingLegsVector, Color, bPersistent, LifeTime, DepthPriority);
 	DrawDebugLine(DrawParams.World, RightFootPosition, RightFootPosition - CrossingLegsVector, Color, bPersistent, LifeTime, DepthPriority);
-#endif // ENABLE_DRAW_DEBUG
 }
+#endif // ENABLE_DRAW_DEBUG
 
 // IPoseFilter interface
 bool UPoseSearchFeatureChannel_FilterCrashingLegs::IsPoseFilterActive() const
diff --git a/Engine/Plugins/Experimental/Animation/PoseSearch/Source/Runtime/Private/PoseSearchFeatureChannel_FilterCrashingLegs.h b/Engine/Plugins/Experimental/Animation/PoseSearch/Source/Runtime/Private/PoseSearchFeatureChannel_FilterCrashingLegs.h
--- a/Engine/Plugins/Experimental/Animation/PoseSearch/Source/Runtime/Private/PoseSearchFeatureChannel_FilterCrashingLegs.h
+++ b/Engine/Plugins/Experimental/Animation/PoseSearch/Source/Runtime/Private/PoseSearchFeatureChannel_FilterCrashingLegs.h
@@ -54,11 +54,15 @@ public:
 
 	// UPoseSearchFeatureChannel interface
 	virtual void Finalize(UPoseSearchSchema* Schema) override;
-	virtual void FillWeights(TArray<float>& Weights) const override;
-	virtual void IndexAsset(UE::PoseSearch::IAssetIndexer& Indexer, TArrayView<float> FeatureVectorTable) const override;
 	virtual void BuildQuery(UE::PoseSearch::FSearchContext& SearchContext, FPoseSearchFeatureVectorBuilder& InOutQuery) const override;
+
+#if ENABLE_DRAW_DEBUG
 	virtual void DebugDraw(const UE::PoseSearch::FDebugDrawParams& DrawParams, TConstArrayView<float> PoseVector) const override;
+#endif // ENABLE_DRAW_DEBUG
+
 #if WITH_EDITOR
+	virtual void FillWeights(TArray<float>& Weights) const override;
+	virtual void IndexAsset(UE::PoseSearch::FAssetIndexer& Indexer, TArrayView<float> FeatureVectorTable) const override;
 	virtual FString GetLabel() const;
 #endif
 
diff --git a/Engine/Plugins/Experimental/Animation/PoseSearch/Source/Runtime/Private/PoseSearchFeatureChannel_Group.cpp b/Engine/Plugins/Experimental/Animation/PoseSearch/Source/Runtime/Private/PoseSearchFeatureChannel_Group.cpp
--- a/Engine/Plugins/Experimental/Animation/PoseSearch/Source/Runtime/Private/PoseSearchFeatureChannel_Group.cpp
+++ b/Engine/Plugins/Experimental/Animation/PoseSearch/Source/Runtime/Private/PoseSearchFeatureChannel_Group.cpp
@@ -1,7 +1,6 @@
 // Copyright Epic Games, Inc. All Rights Reserved.
 
 #include "PoseSearchFeatureChannel_Group.h"
-#include "DrawDebugHelpers.h"
 #include "PoseSearch/PoseSearchSchema.h"
 
 void UPoseSearchFeatureChannel_GroupBase::Finalize(UPoseSearchSchema* Schema)
@@ -17,6 +16,7 @@ void UPoseSearchFeatureChannel_GroupBase::Finalize(UPoseSearchSchema* Schema)
 	ChannelCardinality = Schema->SchemaCardinality - ChannelDataOffset;
 }
 
+#if WITH_EDITOR
 void UPoseSearchFeatureChannel_GroupBase::FillWeights(TArray<float>& Weights) const
 {
 	for (const TObjectPtr<UPoseSearchFeatureChannel>& SubChannelPtr : GetSubChannels())
@@ -28,7 +28,7 @@ void UPoseSearchFeatureChannel_GroupBase::FillWeights(TArray<float>& Weights) co
 	}
 }
 
-void UPoseSearchFeatureChannel_GroupBase::IndexAsset(UE::PoseSearch::IAssetIndexer& Indexer, TArrayView<float> FeatureVectorTable) const
+void UPoseSearchFeatureChannel_GroupBase::IndexAsset(UE::PoseSearch::FAssetIndexer& Indexer, TArrayView<float> FeatureVectorTable) const
 {
 	for (const TObjectPtr<UPoseSearchFeatureChannel>& SubChannelPtr : GetSubChannels())
 	{
@@ -38,6 +38,7 @@ void UPoseSearchFeatureChannel_GroupBase::IndexAsset(UE::PoseSearch::IAssetIndex
 		}
 	}
 }
+#endif // WITH_EDITOR
 
 void UPoseSearchFeatureChannel_GroupBase::BuildQuery(UE::PoseSearch::FSearchContext& SearchContext, FPoseSearchFeatureVectorBuilder& InOutQuery) const
 {
@@ -50,9 +51,9 @@ void UPoseSearchFeatureChannel_GroupBase::BuildQuery(UE::PoseSearch::FSearchCont
 	}
 }
 
+#if ENABLE_DRAW_DEBUG
 void UPoseSearchFeatureChannel_GroupBase::PreDebugDraw(UE::PoseSearch::FDebugDrawParams& DrawParams, TConstArrayView<float> PoseVector) const
 {
-#if ENABLE_DRAW_DEBUG
 	for (const TObjectPtr<UPoseSearchFeatureChannel>& SubChannelPtr : GetSubChannels())
 	{
 		if (const UPoseSearchFeatureChannel* SubChannel = SubChannelPtr.Get())
@@ -60,12 +61,10 @@ void UPoseSearchFeatureChannel_GroupBase::PreDebugDraw(UE::PoseSearch::FDebugDra
 			SubChannel->PreDebugDraw(DrawParams, PoseVector);
 		}
 	}
-#endif // ENABLE_DRAW_DEBUG
 }
 
 void UPoseSearchFeatureChannel_GroupBase::DebugDraw(const UE::PoseSearch::FDebugDrawParams& DrawParams, TConstArrayView<float> PoseVector) const
 {
-#if ENABLE_DRAW_DEBUG
 	for (const TObjectPtr<UPoseSearchFeatureChannel>& SubChannelPtr : GetSubChannels())
 	{
 		if (const UPoseSearchFeatureChannel* SubChannel = SubChannelPtr.Get())
@@ -73,8 +72,8 @@ void UPoseSearchFeatureChannel_GroupBase::DebugDraw(const UE::PoseSearch::FDebug
 			SubChannel->DebugDraw(DrawParams, PoseVector);
 		}
 	}
-#endif // ENABLE_DRAW_DEBUG
 }
+#endif // ENABLE_DRAW_DEBUG
 
 // IPoseFilter interface
 bool UPoseSearchFeatureChannel_GroupBase::IsPoseFilterActive() const
diff --git a/Engine/Plugins/Experimental/Animation/PoseSearch/Source/Runtime/Private/PoseSearchFeatureChannel_Group.h b/Engine/Plugins/Experimental/Animation/PoseSearch/Source/Runtime/Private/PoseSearchFeatureChannel_Group.h
--- a/Engine/Plugins/Experimental/Animation/PoseSearch/Source/Runtime/Private/PoseSearchFeatureChannel_Group.h
+++ b/Engine/Plugins/Experimental/Animation/PoseSearch/Source/Runtime/Private/PoseSearchFeatureChannel_Group.h
@@ -15,11 +15,17 @@ class POSESEARCH_API UPoseSearchFeatureChannel_GroupBase : public UPoseSearchFea
 public:
 	// UPoseSearchFeatureChannel interface
 	virtual void Finalize(UPoseSearchSchema* Schema) override;
-	virtual void FillWeights(TArray<float>& Weights) const override;
-	virtual void IndexAsset(UE::PoseSearch::IAssetIndexer& Indexer, TArrayView<float> FeatureVectorTable) const override;
 	virtual void BuildQuery(UE::PoseSearch::FSearchContext& SearchContext, FPoseSearchFeatureVectorBuilder& InOutQuery) const override;
+
+#if ENABLE_DRAW_DEBUG
 	virtual void PreDebugDraw(UE::PoseSearch::FDebugDrawParams& DrawParams, TConstArrayView<float> PoseVector) const override;
 	virtual void DebugDraw(const UE::PoseSearch::FDebugDrawParams& DrawParams, TConstArrayView<float> PoseVector) const override;
+#endif // ENABLE_DRAW_DEBUG
+
+#if WITH_EDITOR
+	virtual void FillWeights(TArray<float>& Weights) const override;
+	virtual void IndexAsset(UE::PoseSearch::FAssetIndexer& Indexer, TArrayView<float> FeatureVectorTable) const override;
+#endif //WITH_EDITOR
 
 	// IPoseFilter interface
 	virtual bool IsPoseFilterActive() const override;
diff --git a/Engine/Plugins/Experimental/Animation/PoseSearch/Source/Runtime/Private/PoseSearchFeatureChannel_Heading.cpp b/Engine/Plugins/Experimental/Animation/PoseSearch/Source/Runtime/Private/PoseSearchFeatureChannel_Heading.cpp
--- a/Engine/Plugins/Experimental/Animation/PoseSearch/Source/Runtime/Private/PoseSearchFeatureChannel_Heading.cpp
+++ b/Engine/Plugins/Experimental/Animation/PoseSearch/Source/Runtime/Private/PoseSearchFeatureChannel_Heading.cpp
@@ -2,7 +2,6 @@
 
 #include "PoseSearchFeatureChannel_Heading.h"
 #include "Animation/MotionTrajectoryTypes.h"
-#include "DrawDebugHelpers.h"
 #include "PoseSearch/PoseSearchAssetIndexer.h"
 #include "PoseSearch/PoseSearchAssetSampler.h"
 #include "PoseSearch/PoseSearchContext.h"
@@ -17,14 +16,6 @@ void UPoseSearchFeatureChannel_Heading::Finalize(UPoseSearchSchema* Schema)
 	SchemaBoneIdx = Schema->AddBoneReference(Bone);
 }
 
-void UPoseSearchFeatureChannel_Heading::FillWeights(TArray<float>& Weights) const
-{
-	for (int32 i = 0; i < ChannelCardinality; ++i)
-	{
-		Weights[ChannelDataOffset + i] = Weight;
-	}
-}
-
 FVector UPoseSearchFeatureChannel_Heading::GetAxis(const FQuat& Rotation) const
 {
 	switch (HeadingAxis)
@@ -41,7 +32,16 @@ FVector UPoseSearchFeatureChannel_Heading::GetAxis(const FQuat& Rotation) const
 	return FVector(1.f, 0.f, 0.f);
 }
 
-void UPoseSearchFeatureChannel_Heading::IndexAsset(UE::PoseSearch::IAssetIndexer& Indexer, TArrayView<float> FeatureVectorTable) const
+#if WITH_EDITOR
+void UPoseSearchFeatureChannel_Heading::FillWeights(TArray<float>& Weights) const
+{
+	for (int32 i = 0; i < ChannelCardinality; ++i)
+	{
+		Weights[ChannelDataOffset + i] = Weight;
+	}
+}
+
+void UPoseSearchFeatureChannel_Heading::IndexAsset(UE::PoseSearch::FAssetIndexer& Indexer, TArrayView<float> FeatureVectorTable) const
 {
 	using namespace UE::PoseSearch;
 
@@ -60,6 +60,7 @@ void UPoseSearchFeatureChannel_Heading::IndexAsset(UE::PoseSearch::IAssetIndexer
 		FFeatureVectorHelper::EncodeVector(IndexingContext.GetPoseVector(VectorIdx, FeatureVectorTable), ChannelDataOffset, Heading, ComponentStripping);
 	}
 }
+#endif // WITH_EDITOR
 
 void UPoseSearchFeatureChannel_Heading::BuildQuery(UE::PoseSearch::FSearchContext& SearchContext, FPoseSearchFeatureVectorBuilder& InOutQuery) const
 {
@@ -87,9 +88,9 @@ void UPoseSearchFeatureChannel_Heading::BuildQuery(UE::PoseSearch::FSearchContex
 	}
 }
 
+#if ENABLE_DRAW_DEBUG
 void UPoseSearchFeatureChannel_Heading::DebugDraw(const UE::PoseSearch::FDebugDrawParams& DrawParams, TConstArrayView<float> PoseVector) const
 {
-#if ENABLE_DRAW_DEBUG
 	using namespace UE::PoseSearch;
 
 	const float LifeTime = DrawParams.DefaultLifeTime;
@@ -108,8 +109,8 @@ void UPoseSearchFeatureChannel_Heading::DebugDraw(const UE::PoseSearch::FDebugDr
 		const float AdjustedThickness = EnumHasAnyFlags(DrawParams.Flags, EDebugDrawFlags::DrawFast) ? 0.0f : 1.f;
 		DrawDebugLine(DrawParams.World, BonePos, BonePos + BoneHeading * 15.f, Color, bPersistent, LifeTime, DepthPriority, AdjustedThickness);
 	}
-#endif // ENABLE_DRAW_DEBUG
 }
+#endif // ENABLE_DRAW_DEBUG
 
 #if WITH_EDITOR
 FString UPoseSearchFeatureChannel_Heading::GetLabel() const
diff --git a/Engine/Plugins/Experimental/Animation/PoseSearch/Source/Runtime/Private/PoseSearchFeatureChannel_Heading.h b/Engine/Plugins/Experimental/Animation/PoseSearch/Source/Runtime/Private/PoseSearchFeatureChannel_Heading.h
--- a/Engine/Plugins/Experimental/Animation/PoseSearch/Source/Runtime/Private/PoseSearchFeatureChannel_Heading.h
+++ b/Engine/Plugins/Experimental/Animation/PoseSearch/Source/Runtime/Private/PoseSearchFeatureChannel_Heading.h
@@ -49,11 +49,15 @@ public:
 
 	// UPoseSearchFeatureChannel interface
 	virtual void Finalize(UPoseSearchSchema* Schema) override;
-	virtual void FillWeights(TArray<float>& Weights) const override;
-	virtual void IndexAsset(UE::PoseSearch::IAssetIndexer& Indexer, TArrayView<float> FeatureVectorTable) const override;
 	virtual void BuildQuery(UE::PoseSearch::FSearchContext& SearchContext, FPoseSearchFeatureVectorBuilder& InOutQuery) const override;
+
+#if ENABLE_DRAW_DEBUG
 	virtual void DebugDraw(const UE::PoseSearch::FDebugDrawParams& DrawParams, TConstArrayView<float> PoseVector) const override;
+#endif // ENABLE_DRAW_DEBUG
+
 #if WITH_EDITOR
+	virtual void FillWeights(TArray<float>& Weights) const override;
+	virtual void IndexAsset(UE::PoseSearch::FAssetIndexer& Indexer, TArrayView<float> FeatureVectorTable) const override;
 	virtual FString GetLabel() const;
 #endif
 
diff --git a/Engine/Plugins/Experimental/Animation/PoseSearch/Source/Runtime/Private/PoseSearchFeatureChannel_Phase.cpp b/Engine/Plugins/Experimental/Animation/PoseSearch/Source/Runtime/Private/PoseSearchFeatureChannel_Phase.cpp
--- a/Engine/Plugins/Experimental/Animation/PoseSearch/Source/Runtime/Private/PoseSearchFeatureChannel_Phase.cpp
+++ b/Engine/Plugins/Experimental/Animation/PoseSearch/Source/Runtime/Private/PoseSearchFeatureChannel_Phase.cpp
@@ -1,13 +1,13 @@
 // Copyright Epic Games, Inc. All Rights Reserved.
 
 #include "PoseSearchFeatureChannel_Phase.h"
-#include "DrawDebugHelpers.h"
 #include "PoseSearch/PoseSearchAssetIndexer.h"
 #include "PoseSearch/PoseSearchAssetSampler.h"
 #include "PoseSearch/PoseSearchContext.h"
 #include "PoseSearch/PoseSearchDatabase.h"
 #include "PoseSearch/PoseSearchSchema.h"
 
+#if WITH_EDITOR
 namespace UE::PoseSearch
 {
 	struct LocalMinMax
@@ -40,7 +40,7 @@ namespace UE::PoseSearch
 		return (Values[Num - 1] - Values[Num - 2]) * (Sample - (Num - 1)) + Values[Num - 1];
 	}
 
-	static void CollectBonePositions(TArray<FVector>& BonePositions, IAssetIndexer& Indexer, int8 SchemaBoneIdx)
+	static void CollectBonePositions(TArray<FVector>& BonePositions, FAssetIndexer& Indexer, int8 SchemaBoneIdx)
 	{
 		const FAssetIndexingContext& IndexingContext = Indexer.GetIndexingContext();
 		const float FiniteDelta = IndexingContext.Schema->GetSamplingInterval();
@@ -282,6 +282,7 @@ namespace UE::PoseSearch
 	}
 
 } // namespace UE::PoseSearch
+#endif // WITH_EDITOR
 
 void UPoseSearchFeatureChannel_Phase::Finalize(UPoseSearchSchema* Schema)
 {
@@ -291,6 +292,7 @@ void UPoseSearchFeatureChannel_Phase::Finalize(UPoseSearchSchema* Schema)
 	SchemaBoneIdx = Schema->AddBoneReference(Bone);
 }
 
+#if WITH_EDITOR
 void UPoseSearchFeatureChannel_Phase::FillWeights(TArray<float>& Weights) const
 {
 	for (int32 i = 0; i < ChannelCardinality; ++i)
@@ -299,7 +301,7 @@ void UPoseSearchFeatureChannel_Phase::FillWeights(TArray<float>& Weights) const
 	}
 }
 
-void UPoseSearchFeatureChannel_Phase::IndexAsset(UE::PoseSearch::IAssetIndexer& Indexer, TArrayView<float> FeatureVectorTable) const
+void UPoseSearchFeatureChannel_Phase::IndexAsset(UE::PoseSearch::FAssetIndexer& Indexer, TArrayView<float> FeatureVectorTable) const
 {
 	using namespace UE::PoseSearch;
 
@@ -341,6 +343,7 @@ void UPoseSearchFeatureChannel_Phase::IndexAsset(UE::PoseSearch::IAssetIndexer&
 		FFeatureVectorHelper::EncodeVector2D(IndexingContext.GetPoseVector(VectorIdx, FeatureVectorTable), ChannelDataOffset, Phases[VectorIdx]);
 	}
 }
+#endif // WITH_EDITOR
 
 void UPoseSearchFeatureChannel_Phase::BuildQuery(UE::PoseSearch::FSearchContext& SearchContext, FPoseSearchFeatureVectorBuilder& InOutQuery) const
 {
@@ -364,9 +367,9 @@ void UPoseSearchFeatureChannel_Phase::BuildQuery(UE::PoseSearch::FSearchContext&
 	}
 }
 
+#if ENABLE_DRAW_DEBUG
 void UPoseSearchFeatureChannel_Phase::DebugDraw(const UE::PoseSearch::FDebugDrawParams& DrawParams, TConstArrayView<float> PoseVector) const
 {
-#if ENABLE_DRAW_DEBUG
 	using namespace UE::PoseSearch;
 
 	static float ScaleFactor = 1.f;
@@ -390,8 +393,8 @@ void UPoseSearchFeatureChannel_Phase::DebugDraw(const UE::PoseSearch::FDebugDraw
 	FMatrix CircleTransform;
 	CircleTransform.SetAxes(&TransformXAxisVector, &TransformYAxisVector, &TransformZAxisVector, &BonePos);
 	DrawDebugCircle(DrawParams.World, CircleTransform, PhaseVector.Length(), Segments, Color, bPersistent, LifeTime, DepthPriority, 0.f, false);
-#endif // ENABLE_DRAW_DEBUG
 }
+#endif // ENABLE_DRAW_DEBUG
 
 #if WITH_EDITOR
 FString UPoseSearchFeatureChannel_Phase::GetLabel() const
diff --git a/Engine/Plugins/Experimental/Animation/PoseSearch/Source/Runtime/Private/PoseSearchFeatureChannel_Phase.h b/Engine/Plugins/Experimental/Animation/PoseSearch/Source/Runtime/Private/PoseSearchFeatureChannel_Phase.h
--- a/Engine/Plugins/Experimental/Animation/PoseSearch/Source/Runtime/Private/PoseSearchFeatureChannel_Phase.h
+++ b/Engine/Plugins/Experimental/Animation/PoseSearch/Source/Runtime/Private/PoseSearchFeatureChannel_Phase.h
@@ -29,11 +29,15 @@ public:
 
 	// UPoseSearchFeatureChannel interface
 	virtual void Finalize(UPoseSearchSchema* Schema) override;
-	virtual void FillWeights(TArray<float>& Weights) const override;
-	virtual void IndexAsset(UE::PoseSearch::IAssetIndexer& Indexer, TArrayView<float> FeatureVectorTable) const override;
 	virtual void BuildQuery(UE::PoseSearch::FSearchContext& SearchContext, FPoseSearchFeatureVectorBuilder& InOutQuery) const override;
+
+#if ENABLE_DRAW_DEBUG
 	virtual void DebugDraw(const UE::PoseSearch::FDebugDrawParams& DrawParams, TConstArrayView<float> PoseVector) const override;
+#endif // ENABLE_DRAW_DEBUG
+
 #if WITH_EDITOR
+	virtual void FillWeights(TArray<float>& Weights) const override;
+	virtual void IndexAsset(UE::PoseSearch::FAssetIndexer& Indexer, TArrayView<float> FeatureVectorTable) const override;
 	virtual FString GetLabel() const;
 #endif
 };
\ No newline at end of file
diff --git a/Engine/Plugins/Experimental/Animation/PoseSearch/Source/Runtime/Private/PoseSearchFeatureChannel_Position.cpp b/Engine/Plugins/Experimental/Animation/PoseSearch/Source/Runtime/Private/PoseSearchFeatureChannel_Position.cpp
--- a/Engine/Plugins/Experimental/Animation/PoseSearch/Source/Runtime/Private/PoseSearchFeatureChannel_Position.cpp
+++ b/Engine/Plugins/Experimental/Animation/PoseSearch/Source/Runtime/Private/PoseSearchFeatureChannel_Position.cpp
@@ -2,7 +2,6 @@
 
 #include "PoseSearchFeatureChannel_Position.h"
 #include "Animation/MotionTrajectoryTypes.h"
-#include "DrawDebugHelpers.h"
 #include "PoseSearch/PoseSearchAssetIndexer.h"
 #include "PoseSearch/PoseSearchAssetSampler.h"
 #include "PoseSearch/PoseSearchContext.h"
@@ -18,6 +17,7 @@ void UPoseSearchFeatureChannel_Position::Finalize(UPoseSearchSchema* Schema)
 	SchemaBoneIdx = Schema->AddBoneReference(Bone);
 }
 
+#if WITH_EDITOR
 void UPoseSearchFeatureChannel_Position::FillWeights(TArray<float>& Weights) const
 {
 	for (int32 i = 0; i < ChannelCardinality; ++i)
@@ -26,7 +26,7 @@ void UPoseSearchFeatureChannel_Position::FillWeights(TArray<float>& Weights) con
 	}
 }
 
-void UPoseSearchFeatureChannel_Position::IndexAsset(UE::PoseSearch::IAssetIndexer& Indexer, TArrayView<float> FeatureVectorTable) const
+void UPoseSearchFeatureChannel_Position::IndexAsset(UE::PoseSearch::FAssetIndexer& Indexer, TArrayView<float> FeatureVectorTable) const
 {
 	using namespace UE::PoseSearch;
 
@@ -44,6 +44,7 @@ void UPoseSearchFeatureChannel_Position::IndexAsset(UE::PoseSearch::IAssetIndexe
 		FFeatureVectorHelper::EncodeVector(IndexingContext.GetPoseVector(VectorIdx, FeatureVectorTable), ChannelDataOffset, BoneTransformsPresent.GetTranslation(), ComponentStripping);
 	}
 }
+#endif // WITH_EDITOR
 
 void UPoseSearchFeatureChannel_Position::BuildQuery(UE::PoseSearch::FSearchContext& SearchContext, FPoseSearchFeatureVectorBuilder& InOutQuery) const
 {
@@ -71,18 +72,16 @@ void UPoseSearchFeatureChannel_Position::BuildQuery(UE::PoseSearch::FSearchConte
 	}
 }
 
+#if ENABLE_DRAW_DEBUG
 void UPoseSearchFeatureChannel_Position::PreDebugDraw(UE::PoseSearch::FDebugDrawParams& DrawParams, TConstArrayView<float> PoseVector) const
 {
-#if ENABLE_DRAW_DEBUG
 	using namespace UE::PoseSearch;
 	const FVector BonePos = DrawParams.RootTransform.TransformPosition(FFeatureVectorHelper::DecodeVector(PoseVector, ChannelDataOffset, ComponentStripping));
 	DrawParams.AddCachedPosition(SampleTimeOffset, SchemaBoneIdx, BonePos);
-#endif // ENABLE_DRAW_DEBUG
 }
 
 void UPoseSearchFeatureChannel_Position::DebugDraw(const UE::PoseSearch::FDebugDrawParams& DrawParams, TConstArrayView<float> PoseVector) const
 {
-#if ENABLE_DRAW_DEBUG
 	using namespace UE::PoseSearch;
 
 	const float LifeTime = DrawParams.DefaultLifeTime;
@@ -109,8 +108,8 @@ void UPoseSearchFeatureChannel_Position::DebugDraw(const UE::PoseSearch::FDebugD
 		check(Schema && Schema->IsValid());
 		DrawDebugString(DrawParams.World, BonePos + FVector(0.0, 0.0, 10.0), Schema->BoneReferences[SchemaBoneIdx].BoneName.ToString(), nullptr, Color, LifeTime, false, 1.0f);
 	}
-#endif // ENABLE_DRAW_DEBUG
 }
+#endif // ENABLE_DRAW_DEBUG
 
 #if WITH_EDITOR
 FString UPoseSearchFeatureChannel_Position::GetLabel() const
diff --git a/Engine/Plugins/Experimental/Animation/PoseSearch/Source/Runtime/Private/PoseSearchFeatureChannel_Position.h b/Engine/Plugins/Experimental/Animation/PoseSearch/Source/Runtime/Private/PoseSearchFeatureChannel_Position.h
--- a/Engine/Plugins/Experimental/Animation/PoseSearch/Source/Runtime/Private/PoseSearchFeatureChannel_Position.h
+++ b/Engine/Plugins/Experimental/Animation/PoseSearch/Source/Runtime/Private/PoseSearchFeatureChannel_Position.h
@@ -35,12 +35,16 @@ public:
 
 	// UPoseSearchFeatureChannel interface
 	virtual void Finalize(UPoseSearchSchema* Schema) override;
-	virtual void FillWeights(TArray<float>& Weights) const override;
-	virtual void IndexAsset(UE::PoseSearch::IAssetIndexer& Indexer, TArrayView<float> FeatureVectorTable) const override;
 	virtual void BuildQuery(UE::PoseSearch::FSearchContext& SearchContext, FPoseSearchFeatureVectorBuilder& InOutQuery) const override;
+
+#if ENABLE_DRAW_DEBUG
 	virtual void PreDebugDraw(UE::PoseSearch::FDebugDrawParams& DrawParams, TConstArrayView<float> PoseVector) const override;
 	virtual void DebugDraw(const UE::PoseSearch::FDebugDrawParams& DrawParams, TConstArrayView<float> PoseVector) const override;
+#endif // ENABLE_DRAW_DEBUG
+
 #if WITH_EDITOR
+	virtual void FillWeights(TArray<float>& Weights) const override;
+	virtual void IndexAsset(UE::PoseSearch::FAssetIndexer& Indexer, TArrayView<float> FeatureVectorTable) const override;
 	virtual FString GetLabel() const;
 #endif
 };
diff --git a/Engine/Plugins/Experimental/Animation/PoseSearch/Source/Runtime/Private/PoseSearchFeatureChannel_Trajectory.cpp b/Engine/Plugins/Experimental/Animation/PoseSearch/Source/Runtime/Private/PoseSearchFeatureChannel_Trajectory.cpp
--- a/Engine/Plugins/Experimental/Animation/PoseSearch/Source/Runtime/Private/PoseSearchFeatureChannel_Trajectory.cpp
+++ b/Engine/Plugins/Experimental/Animation/PoseSearch/Source/Runtime/Private/PoseSearchFeatureChannel_Trajectory.cpp
@@ -2,7 +2,6 @@
 
 #include "PoseSearchFeatureChannel_Trajectory.h"
 #include "Animation/Skeleton.h"
-#include "DrawDebugHelpers.h"
 #include "PoseSearch/PoseSearchContext.h"
 #include "PoseSearch/PoseSearchDerivedDataKey.h"
 #include "PoseSearch/PoseSearchSchema.h"
@@ -31,7 +30,6 @@ void UPoseSearchFeatureChannel_Trajectory::Finalize(UPoseSearchSchema* Schema)
 
 	for (const FPoseSearchTrajectorySample& Sample : Samples)
 	{
-		// @todo: implement PositionXY properly as 2 dimension channel
 		if (EnumHasAnyFlags(Sample.Flags, EPoseSearchTrajectoryFlags::Position | EPoseSearchTrajectoryFlags::PositionXY))
 		{
 			UPoseSearchFeatureChannel_Position* Position = NewObject<UPoseSearchFeatureChannel_Position>(this, NAME_None, RF_Transient);
@@ -95,9 +93,9 @@ void UPoseSearchFeatureChannel_Trajectory::Finalize(UPoseSearchSchema* Schema)
 	Super::Finalize(Schema);
 }
 
+#if ENABLE_DRAW_DEBUG
 void UPoseSearchFeatureChannel_Trajectory::DebugDraw(const UE::PoseSearch::FDebugDrawParams& DrawParams, TConstArrayView<float> PoseVector) const
 {
-#if ENABLE_DRAW_DEBUG
 	using namespace UE::PoseSearch;
 
 	TArray<const UPoseSearchFeatureChannel_Position*, TInlineAllocator<32>> Positions;
@@ -146,8 +144,8 @@ void UPoseSearchFeatureChannel_Trajectory::DebugDraw(const UE::PoseSearch::FDebu
 	}
 
 	Super::DebugDraw(DrawParams, PoseVector);
-#endif // ENABLE_DRAW_DEBUG
 }
+#endif // ENABLE_DRAW_DEBUG
 
 #if WITH_EDITOR
 FString UPoseSearchFeatureChannel_Trajectory::GetLabel() const
diff --git a/Engine/Plugins/Experimental/Animation/PoseSearch/Source/Runtime/Private/PoseSearchFeatureChannel_Trajectory.h b/Engine/Plugins/Experimental/Animation/PoseSearch/Source/Runtime/Private/PoseSearchFeatureChannel_Trajectory.h
--- a/Engine/Plugins/Experimental/Animation/PoseSearch/Source/Runtime/Private/PoseSearchFeatureChannel_Trajectory.h
+++ b/Engine/Plugins/Experimental/Animation/PoseSearch/Source/Runtime/Private/PoseSearchFeatureChannel_Trajectory.h
@@ -64,7 +64,11 @@ public:
 
 	// UPoseSearchFeatureChannel interface
 	virtual void Finalize(UPoseSearchSchema* Schema) override;
+
+#if ENABLE_DRAW_DEBUG
 	virtual void DebugDraw(const UE::PoseSearch::FDebugDrawParams& DrawParams, TConstArrayView<float> PoseVector) const override;
+#endif // ENABLE_DRAW_DEBUG
+
 #if WITH_EDITOR
 	virtual FString GetLabel() const;
 #endif
diff --git a/Engine/Plugins/Experimental/Animation/PoseSearch/Source/Runtime/Private/PoseSearchFeatureChannel_Velocity.cpp b/Engine/Plugins/Experimental/Animation/PoseSearch/Source/Runtime/Private/PoseSearchFeatureChannel_Velocity.cpp
--- a/Engine/Plugins/Experimental/Animation/PoseSearch/Source/Runtime/Private/PoseSearchFeatureChannel_Velocity.cpp
+++ b/Engine/Plugins/Experimental/Animation/PoseSearch/Source/Runtime/Private/PoseSearchFeatureChannel_Velocity.cpp
@@ -2,7 +2,6 @@
 
 #include "PoseSearchFeatureChannel_Velocity.h"
 #include "Animation/MotionTrajectoryTypes.h"
-#include "DrawDebugHelpers.h"
 #include "PoseSearch/PoseSearchAssetIndexer.h"
 #include "PoseSearch/PoseSearchAssetSampler.h"
 #include "PoseSearch/PoseSearchContext.h"
@@ -19,6 +18,7 @@ void UPoseSearchFeatureChannel_Velocity::Finalize(UPoseSearchSchema* Schema)
 	SchemaBoneIdx = Schema->AddBoneReference(Bone);
 }
 
+#if WITH_EDITOR
 void UPoseSearchFeatureChannel_Velocity::FillWeights(TArray<float>& Weights) const
 {
 	for (int32 i = 0; i < ChannelCardinality; ++i)
@@ -27,7 +27,7 @@ void UPoseSearchFeatureChannel_Velocity::FillWeights(TArray<float>& Weights) con
 	}
 }
 
-void UPoseSearchFeatureChannel_Velocity::IndexAsset(UE::PoseSearch::IAssetIndexer& Indexer, TArrayView<float> FeatureVectorTable) const
+void UPoseSearchFeatureChannel_Velocity::IndexAsset(UE::PoseSearch::FAssetIndexer& Indexer, TArrayView<float> FeatureVectorTable) const
 {
 	using namespace UE::PoseSearch;
 
@@ -65,6 +65,7 @@ void UPoseSearchFeatureChannel_Velocity::IndexAsset(UE::PoseSearch::IAssetIndexe
 		FFeatureVectorHelper::EncodeVector(IndexingContext.GetPoseVector(VectorIdx, FeatureVectorTable), ChannelDataOffset, LinearVelocity, ComponentStripping);
 	}
 }
+#endif // WITH_EDITOR
 
 void UPoseSearchFeatureChannel_Velocity::BuildQuery(UE::PoseSearch::FSearchContext& SearchContext, FPoseSearchFeatureVectorBuilder& InOutQuery) const
 {
@@ -102,9 +103,9 @@ void UPoseSearchFeatureChannel_Velocity::BuildQuery(UE::PoseSearch::FSearchConte
 	}
 }
 
+#if ENABLE_DRAW_DEBUG
 void UPoseSearchFeatureChannel_Velocity::DebugDraw(const UE::PoseSearch::FDebugDrawParams& DrawParams, TConstArrayView<float> PoseVector) const
 {
-#if ENABLE_DRAW_DEBUG
 	using namespace UE::PoseSearch;
 
 	const float LifeTime = DrawParams.DefaultLifeTime;
@@ -126,8 +127,8 @@ void UPoseSearchFeatureChannel_Velocity::DebugDraw(const UE::PoseSearch::FDebugD
 		const float AdjustedThickness = EnumHasAnyFlags(DrawParams.Flags, EDebugDrawFlags::DrawFast) ? 0.f : 1.f;
 		DrawDebugLine(DrawParams.World, BonePos + BoneVelDirection * 2.f, BonePos + LinearVelocity * LinearVelocityScale, Color, bPersistent, LifeTime, DepthPriority, AdjustedThickness);
 	}
-#endif // ENABLE_DRAW_DEBUG
 }
+#endif // ENABLE_DRAW_DEBUG
 
 #if WITH_EDITOR
 FString UPoseSearchFeatureChannel_Velocity::GetLabel() const
diff --git a/Engine/Plugins/Experimental/Animation/PoseSearch/Source/Runtime/Private/PoseSearchFeatureChannel_Velocity.h b/Engine/Plugins/Experimental/Animation/PoseSearch/Source/Runtime/Private/PoseSearchFeatureChannel_Velocity.h
--- a/Engine/Plugins/Experimental/Animation/PoseSearch/Source/Runtime/Private/PoseSearchFeatureChannel_Velocity.h
+++ b/Engine/Plugins/Experimental/Animation/PoseSearch/Source/Runtime/Private/PoseSearchFeatureChannel_Velocity.h
@@ -43,11 +43,15 @@ public:
 
 	// UPoseSearchFeatureChannel interface
 	virtual void Finalize(UPoseSearchSchema* Schema) override;
-	virtual void FillWeights(TArray<float>& Weights) const override;
-	virtual void IndexAsset(UE::PoseSearch::IAssetIndexer& Indexer, TArrayView<float> FeatureVectorTable) const override;
 	virtual void BuildQuery(UE::PoseSearch::FSearchContext& SearchContext, FPoseSearchFeatureVectorBuilder& InOutQuery) const override;
+
+#if ENABLE_DRAW_DEBUG
 	virtual void DebugDraw(const UE::PoseSearch::FDebugDrawParams& DrawParams, TConstArrayView<float> PoseVector) const override;
+#endif // ENABLE_DRAW_DEBUG
+
 #if WITH_EDITOR
+	virtual void FillWeights(TArray<float>& Weights) const override;
+	virtual void IndexAsset(UE::PoseSearch::FAssetIndexer& Indexer, TArrayView<float> FeatureVectorTable) const override;
 	virtual FString GetLabel() const;
 #endif
 };
diff --git a/Engine/Plugins/Experimental/Animation/PoseSearch/Source/Runtime/Private/PoseSearchHistory.cpp b/Engine/Plugins/Experimental/Animation/PoseSearch/Source/Runtime/Private/PoseSearchHistory.cpp
--- a/Engine/Plugins/Experimental/Animation/PoseSearch/Source/Runtime/Private/PoseSearchHistory.cpp
+++ b/Engine/Plugins/Experimental/Animation/PoseSearch/Source/Runtime/Private/PoseSearchHistory.cpp
@@ -288,7 +288,7 @@ void FPoseHistory::DebugDraw(const UWorld* World, const USkeleton* Skeleton) con
 			{
 				const FTransform GlobalTransforms = Entry.ComponentSpaceTransforms[i] * Entry.RootTransform;
 
-				DrawDebugLine(World, PrevGlobalTransforms[i].GetLocation(), GlobalTransforms.GetLocation(), Color, false, 0.f, ESceneDepthPriorityGroup::SDPG_Foreground + 2);
+				DrawDebugLine(World, PrevGlobalTransforms[i].GetTranslation(), GlobalTransforms.GetTranslation(), Color, false, 0.f, ESceneDepthPriorityGroup::SDPG_Foreground + 2);
 
 				PrevGlobalTransforms[i] = GlobalTransforms;
 			}
diff --git a/Engine/Plugins/Experimental/Animation/PoseSearch/Source/Runtime/Public/PoseSearch/AnimNode_MotionMatching.h b/Engine/Plugins/Experimental/Animation/PoseSearch/Source/Runtime/Public/PoseSearch/AnimNode_MotionMatching.h
--- a/Engine/Plugins/Experimental/Animation/PoseSearch/Source/Runtime/Public/PoseSearch/AnimNode_MotionMatching.h
+++ b/Engine/Plugins/Experimental/Animation/PoseSearch/Source/Runtime/Public/PoseSearch/AnimNode_MotionMatching.h
@@ -56,6 +56,7 @@ public:
 #endif
 
 	// FAnimNode_Base interface
+	// @todo: implement CacheBones_AnyThread to rebind the schema bones
 	virtual void Initialize_AnyThread(const FAnimationInitializeContext& Context) override;
 	virtual void Evaluate_AnyThread(FPoseContext& Output) override;
 	virtual void GatherDebugData(FNodeDebugData& DebugData) override;
diff --git a/Engine/Plugins/Experimental/Animation/PoseSearch/Source/Runtime/Public/PoseSearch/PoseSearchAssetIndexer.h b/Engine/Plugins/Experimental/Animation/PoseSearch/Source/Runtime/Public/PoseSearch/PoseSearchAssetIndexer.h
--- a/Engine/Plugins/Experimental/Animation/PoseSearch/Source/Runtime/Public/PoseSearch/PoseSearchAssetIndexer.h
+++ b/Engine/Plugins/Experimental/Animation/PoseSearch/Source/Runtime/Public/PoseSearch/PoseSearchAssetIndexer.h
@@ -2,14 +2,19 @@
 
 #pragma once
 
+#if WITH_EDITOR
+
+#include "BonePose.h"
 #include "CoreMinimal.h"
 #include "PoseSearch/PoseSearchDefines.h"
 #include "PoseSearch/PoseSearchSchema.h"
 
-struct FAssetSamplerBase;
+struct FPoseSearchPoseMetadata;
 
 namespace UE::PoseSearch
 {
+
+struct FAssetSamplerBase;
 struct FAssetSamplingContext;
 
 /**
@@ -19,7 +24,7 @@ struct FAssetIndexingContext
 {
 	const FAssetSamplingContext* SamplingContext = nullptr;
 	const UPoseSearchSchema* Schema = nullptr;
-	const FAssetSamplerBase* AssetSampler = nullptr;
+	TSharedPtr<FAssetSamplerBase> AssetSampler;
 	bool bMirrored = false;
 	FFloatInterval RequestedSamplingRange = FFloatInterval(0.0f, 0.0f);
 	
@@ -34,25 +39,71 @@ struct FAssetIndexingContext
 	}
 };
 
-class POSESEARCH_API IAssetIndexer
+class POSESEARCH_API FAssetIndexer
 {
 public:
+	struct FOutput
+	{
+		int32 FirstIndexedSample = 0;
+		int32 LastIndexedSample = 0;
+		int32 NumIndexedPoses = 0;
+
+		TArray<float> FeatureVectorTable;
+		TArray<FPoseSearchPoseMetadata> PoseMetadata;
+	};
+
+	struct FStats
+	{
+		int32 NumAccumulatedSamples = 0;
+		float AccumulatedSpeed = 0.f;
+		float MaxSpeed = 0.f;
+		float AccumulatedAcceleration = 0.f;
+		float MaxAcceleration = 0.f;
+	};
+
+	void Reset();
+	void Init(const FAssetIndexingContext& IndexingContext, const FBoneContainer& InBoneContainer);
+	bool Process();
+	const FOutput& GetOutput() const { return Output; }
+	const FStats& GetStats() const { return Stats; }
+
+	const FAssetIndexingContext& GetIndexingContext() const { return IndexingContext; }
+	FTransform GetTransform(float SampleTime, bool& bClamped, int8 SchemaBoneIdx = RootSchemaBoneIdx);
+	FTransform GetComponentSpaceTransform(float SampleTime, bool& bClamped, int8 SchemaBoneIdx = RootSchemaBoneIdx);
+	FTransform GetComponentSpaceTransform(float SampleTime, float OriginTime, bool& bClamped, int8 SchemaBoneIdx = RootSchemaBoneIdx);
+	
+private:
 	struct FSampleInfo
 	{
-		const FAssetSamplerBase* Clip = nullptr;
+		TSharedPtr<FAssetSamplerBase> Clip;
 		FTransform RootTransform;
-		float ClipTime = 0.0f;
+		float ClipTime = 0.f;
 		bool bClamped = false;
+	};
 
-		bool IsValid() const { return Clip != nullptr; }
+	struct CachedEntry
+	{
+		float SampleTime;
+		bool bClamped;
+
+		FTransform RootTransform;
+		FCSPose<FCompactPose> ComponentSpacePose;
 	};
 
-	virtual ~IAssetIndexer() {}
+	FSampleInfo GetSampleInfo(float SampleTime) const;
+	FPoseSearchPoseMetadata GetMetadata(int32 SampleIdx) const;
+	FTransform MirrorTransform(const FTransform& Transform) const;
+	CachedEntry& GetEntry(float SampleTime);
+	FTransform CalculateComponentSpaceTransform(CachedEntry& Entry, int8 SchemaBoneIdx);
+	void ComputeStats();
 
-	virtual const FAssetIndexingContext& GetIndexingContext() const = 0;
-	virtual FTransform GetTransform(float SampleTime, bool& Clamped, int8 SchemaBoneIdx = RootSchemaBoneIdx) = 0;
-	virtual FTransform GetComponentSpaceTransform(float SampleTime, bool& Clamped, int8 SchemaBoneIdx = RootSchemaBoneIdx) = 0;
-	virtual FTransform GetComponentSpaceTransform(float SampleTime, float OriginTime, bool& Clamped, int8 SchemaBoneIdx = RootSchemaBoneIdx) = 0;
+	FBoneContainer BoneContainer;
+	FAssetIndexingContext IndexingContext;
+	TMap<float, CachedEntry> CachedEntries;
+	FOutput Output;
+	FStats Stats;
 };
 
 } // namespace UE::PoseSearch
+
+#endif // WITH_EDITOR
diff --git a/Engine/Plugins/Experimental/Animation/PoseSearch/Source/Runtime/Public/PoseSearch/PoseSearchAssetSampler.h b/Engine/Plugins/Experimental/Animation/PoseSearch/Source/Runtime/Public/PoseSearch/PoseSearchAssetSampler.h
--- a/Engine/Plugins/Experimental/Animation/PoseSearch/Source/Runtime/Public/PoseSearch/PoseSearchAssetSampler.h
+++ b/Engine/Plugins/Experimental/Animation/PoseSearch/Source/Runtime/Public/PoseSearch/PoseSearchAssetSampler.h
@@ -5,37 +5,6 @@
 #include "BonePose.h"
 #include "PoseSearchAssetSampler.generated.h"
 
-struct FAnimationPoseData;
-struct FAnimExtractContext;
-class UAnimationAsset;
-class UAnimInstance;
-class UAnimNotifyState_PoseSearchBase;
-class UBlendSpace;
-class UMirrorDataTable;
-
-namespace UE::PoseSearch
-{
-	struct POSESEARCH_API FAssetSamplingContext
-	{
-		// Time delta used for computing pose derivatives
-		static constexpr float FiniteDelta = 1 / 60.0f;
-
-		// Mirror data table pointer copied from Schema for convenience
-		TObjectPtr<const UMirrorDataTable> MirrorDataTable;
-
-		// Compact pose format of Mirror Bone Map
-		TCustomBoneIndexArray<FCompactPoseBoneIndex, FCompactPoseBoneIndex> CompactPoseMirrorBones;
-
-		// Pre-calculated component space rotations of reference pose, which allows mirror to work with any joint orientation
-		// Only initialized and used when a mirroring table is specified
-		TCustomBoneIndexArray<FQuat, FCompactPoseBoneIndex> ComponentSpaceRefRotations;
-
-		void Init(const UMirrorDataTable* InMirrorDataTable, const FBoneContainer& BoneContainer);
-		FTransform MirrorTransform(const FTransform& Transform) const;
-	};
-
-} // namespace UE::PoseSearch
-
 USTRUCT()
 struct POSESEARCH_API FPoseSearchExtrapolationParameters
 {
@@ -54,15 +23,43 @@ struct POSESEARCH_API FPoseSearchExtrapolationParameters
 	float SampleTime = 0.05f;
 };
 
+#if WITH_EDITOR
+
+struct FAnimationPoseData;
+struct FAnimExtractContext;
+class UAnimationAsset;
+class UAnimInstance;
+class UAnimNotifyState_PoseSearchBase;
+class UBlendSpace;
+class UMirrorDataTable;
+
+namespace UE::PoseSearch
+{
+
+struct POSESEARCH_API FAssetSamplingContext
+{
+	// Time delta used for computing pose derivatives
+	static constexpr float FiniteDelta = 1 / 60.0f;
+
+	// Mirror data table pointer copied from Schema for convenience
+	TObjectPtr<const UMirrorDataTable> MirrorDataTable;
+
+	// Compact pose format of Mirror Bone Map
+	TCustomBoneIndexArray<FCompactPoseBoneIndex, FCompactPoseBoneIndex> CompactPoseMirrorBones;
+
+	// Pre-calculated component space rotations of reference pose, which allows mirror to work with any joint orientation
+	// Only initialized and used when a mirroring table is specified
+	TCustomBoneIndexArray<FQuat, FCompactPoseBoneIndex> ComponentSpaceRefRotations;
+
+	void Init(const UMirrorDataTable* InMirrorDataTable, const FBoneContainer& BoneContainer);
+	FTransform MirrorTransform(const FTransform& Transform) const;
+};
+
 /**
  * Helper interface for sampling data from animation assets
  */
-USTRUCT()
-struct POSESEARCH_API FAssetSamplerBase
+struct POSESEARCH_API FAssetSamplerBase : public TSharedFromThis<FAssetSamplerBase>
 {
-	GENERATED_BODY()
-
-public:
 	virtual ~FAssetSamplerBase() {}
 
 	virtual float GetPlayLength() const { return 0.f; }
@@ -89,12 +86,8 @@ public:
 };
 
 // Sampler working with UAnimSequenceBase so it can be used for UAnimSequence as well as UAnimComposite.
-USTRUCT()
 struct POSESEARCH_API FSequenceBaseSampler : public FAssetSamplerBase
 {
-	GENERATED_BODY()
-
-public:
 	struct FInput
 	{
 		TWeakObjectPtr<const UAnimSequenceBase> SequenceBase;
@@ -116,12 +109,8 @@ public:
 	virtual const UAnimationAsset* GetAsset() const override;
 };
 
-USTRUCT()
 struct POSESEARCH_API FBlendSpaceSampler : public FAssetSamplerBase
 {
-	GENERATED_BODY()
-
-public:
 	struct FInput
 	{
 		FBoneContainer BoneContainer;
@@ -159,12 +148,8 @@ private:
 	FTransform ExtractBlendSpaceRootMotionFromRange(float StartTrackPosition, float EndTrackPosition) const;
 };
 
-USTRUCT()
 struct POSESEARCH_API FAnimMontageSampler : public FAssetSamplerBase
 {
-	GENERATED_BODY()
-
-public:
 	struct FInput
 	{
 		// @todo: add support for SlotName / multiple SlotAnimTracks
@@ -190,3 +175,6 @@ private:
 	FTransform ExtractRootTransformInternal(float StartTime, float EndTime) const;
 };
 
+} // namespace UE::PoseSearch
+
+#endif // WITH_EDITOR
diff --git a/Engine/Plugins/Experimental/Animation/PoseSearch/Source/Runtime/Public/PoseSearch/PoseSearchContext.h b/Engine/Plugins/Experimental/Animation/PoseSearch/Source/Runtime/Public/PoseSearch/PoseSearchContext.h
--- a/Engine/Plugins/Experimental/Animation/PoseSearch/Source/Runtime/Public/PoseSearch/PoseSearchContext.h
+++ b/Engine/Plugins/Experimental/Animation/PoseSearch/Source/Runtime/Public/PoseSearch/PoseSearchContext.h
@@ -2,6 +2,7 @@
 
 #pragma once
 
+#include "DrawDebugHelpers.h"
 #include "PoseSearch/PoseSearchDefines.h"
 #include "PoseSearch/PoseSearchIndex.h"
 #include "PoseSearch/PoseSearchResult.h"
@@ -115,7 +116,7 @@ private:
 	TArray<FCachedTransform<FTransformType>, TInlineAllocator<64>> CachedTransforms;
 };
 
-// @todo: FDebugDrawParams should be enclosed with #if ENABLE_DRAW_DEBUG
+#if ENABLE_DRAW_DEBUG
 struct POSESEARCH_API FDebugDrawParams
 {
 	const UWorld* World = nullptr;
@@ -146,10 +147,11 @@ private:
 POSESEARCH_API void DrawFeatureVector(FDebugDrawParams& DrawParams, TConstArrayView<float> PoseVector);
 POSESEARCH_API void DrawFeatureVector(FDebugDrawParams& DrawParams, int32 PoseIdx);
 
+#endif // ENABLE_DRAW_DEBUG
+
 struct POSESEARCH_API FSearchContext
 {
 	EPoseSearchBooleanRequest QueryMirrorRequest = EPoseSearchBooleanRequest::Indifferent;
-	FDebugDrawParams DebugDrawParams;
 	const FPoseHistory* History = nullptr;
 	const FTrajectorySampleRange* Trajectory = nullptr;
 	FSearchResult CurrentResult;
diff --git a/Engine/Plugins/Experimental/Animation/PoseSearch/Source/Runtime/Public/PoseSearch/PoseSearchFeatureChannel.h b/Engine/Plugins/Experimental/Animation/PoseSearch/Source/Runtime/Public/PoseSearch/PoseSearchFeatureChannel.h
--- a/Engine/Plugins/Experimental/Animation/PoseSearch/Source/Runtime/Public/PoseSearch/PoseSearchFeatureChannel.h
+++ b/Engine/Plugins/Experimental/Animation/PoseSearch/Source/Runtime/Public/PoseSearch/PoseSearchFeatureChannel.h
@@ -4,6 +4,7 @@
 
 #include "IO/IoHash.h"
 #include "Interfaces/Interface_BoneReferenceSkeletonProvider.h"
+#include "DrawDebugHelpers.h"
 #include "PoseSearchFeatureChannel.generated.h"
 
 class UPoseSearchSchema;
@@ -41,7 +42,10 @@ namespace UE::PoseSearch
 
 struct FDebugDrawParams;
 struct FSearchContext;
-class IAssetIndexer;
+
+#if WITH_EDITOR
+class FAssetIndexer;
+#endif // WITH_EDITOR
 
 /** Helper class for extracting and encoding features into a float buffer */
 class POSESEARCH_API FFeatureVectorHelper
@@ -89,28 +93,30 @@ public:
 	// Called during UPoseSearchSchema::Finalize to prepare the schema for this channel
 	virtual void Finalize(UPoseSearchSchema* Schema) PURE_VIRTUAL(UPoseSearchFeatureChannel::Finalize, );
 	
-	// Called at database build time to collect feature weights.
-	// Weights is sized to the cardinality of the schema and the feature channel should write
-	// its weights at the channel's data offset. Channels should provide a weight for each dimension.
-	virtual void FillWeights(TArray<float>& Weights) const PURE_VIRTUAL(UPoseSearchFeatureChannel::FillWeights, );
-
-	// Called at database build time to populate pose vectors with this channel's data
-	virtual void IndexAsset(UE::PoseSearch::IAssetIndexer& Indexer, TArrayView<float> FeatureVectorTable) const PURE_VIRTUAL(UPoseSearchFeatureChannel::IndexAsset, );
-
 	// Called at runtime to add this channel's data to the query pose vector
 	virtual void BuildQuery(UE::PoseSearch::FSearchContext& SearchContext, FPoseSearchFeatureVectorBuilder& InOutQuery) const PURE_VIRTUAL(UPoseSearchFeatureChannel::BuildQuery, );
 
+	// UPoseSearchFeatureChannels can hold sub channels
+	virtual TArrayView<TObjectPtr<UPoseSearchFeatureChannel>> GetSubChannels() { return TArrayView<TObjectPtr<UPoseSearchFeatureChannel>>(); }
+	virtual TConstArrayView<TObjectPtr<UPoseSearchFeatureChannel>> GetSubChannels() const { return TConstArrayView<TObjectPtr<UPoseSearchFeatureChannel>>(); }
+
+#if ENABLE_DRAW_DEBUG
 	// API called before DebugDraw to collect shared channel informations such as decoded positions form the PoseVector
 	virtual void PreDebugDraw(UE::PoseSearch::FDebugDrawParams& DrawParams, TConstArrayView<float> PoseVector) const {}
 
 	// Draw this channel's data for the given pose vector
 	virtual void DebugDraw(const UE::PoseSearch::FDebugDrawParams& DrawParams, TConstArrayView<float> PoseVector) const PURE_VIRTUAL(UPoseSearchFeatureChannel::DebugDraw, );
-
-	// UPoseSearchFeatureChannels can hold sub channels
-	virtual TArrayView<TObjectPtr<UPoseSearchFeatureChannel>> GetSubChannels() { return TArrayView<TObjectPtr<UPoseSearchFeatureChannel>>(); }
-	virtual TConstArrayView<TObjectPtr<UPoseSearchFeatureChannel>> GetSubChannels() const { return TConstArrayView<TObjectPtr<UPoseSearchFeatureChannel>>(); }
+#endif //ENABLE_DRAW_DEBUG
 
 #if WITH_EDITOR
+	// Called at database build time to collect feature weights.
+	// Weights is sized to the cardinality of the schema and the feature channel should write
+	// its weights at the channel's data offset. Channels should provide a weight for each dimension.
+	virtual void FillWeights(TArray<float>& Weights) const PURE_VIRTUAL(UPoseSearchFeatureChannel::FillWeights, );
+
+	// Called at database build time to populate pose vectors with this channel's data
+	virtual void IndexAsset(UE::PoseSearch::FAssetIndexer& Indexer, TArrayView<float> FeatureVectorTable) const PURE_VIRTUAL(UPoseSearchFeatureChannel::IndexAsset, );
+
 	// returns the FString used editor side to identify this UPoseSearchFeatureChannel (for instance in the pose search debugger)
 	virtual FString GetLabel() const;
 	virtual bool CanBeNormalizedWith(const UPoseSearchFeatureChannel* Other) const;
