motion matching - adding origin bone to UPoseSearchFeatureChannel_Position

diff --git a/Engine/Plugins/Experimental/Animation/PoseSearch/Source/Runtime/Private/PoseSearchAssetIndexer.cpp b/Engine/Plugins/Experimental/Animation/PoseSearch/Source/Runtime/Private/PoseSearchAssetIndexer.cpp
--- a/Engine/Plugins/Experimental/Animation/PoseSearch/Source/Runtime/Private/PoseSearchAssetIndexer.cpp
+++ b/Engine/Plugins/Experimental/Animation/PoseSearch/Source/Runtime/Private/PoseSearchAssetIndexer.cpp
@@ -378,7 +378,7 @@ FTransform FAssetIndexer::GetComponentSpaceTransform(float SampleTime, bool& bCl
 	CachedEntry& Entry = GetEntry(SampleTime);
 	bClamped = Entry.bClamped;
 
-	if (SchemaBoneIdx >= 0 && IndexingContext.Schema->BoneReferences[SchemaBoneIdx].HasValidSetup())
+	if (!IndexingContext.Schema->IsRootBone(SchemaBoneIdx))
 	{
 		return CalculateComponentSpaceTransform(Entry, SchemaBoneIdx);
 	}
@@ -393,7 +393,7 @@ FTransform FAssetIndexer::GetTransform(float SampleTime, bool& bClamped, int8 Sc
 	bClamped = Entry.bClamped;
 
 	const FTransform MirroredRootTransform = MirrorTransform(Entry.RootTransform);
-	if (SchemaBoneIdx >= 0 && IndexingContext.Schema->BoneReferences[SchemaBoneIdx].HasValidSetup())
+	if (!IndexingContext.Schema->IsRootBone(SchemaBoneIdx))
 	{
 		return CalculateComponentSpaceTransform(Entry, SchemaBoneIdx) * MirroredRootTransform;
 	}
@@ -408,21 +408,97 @@ FTransform FAssetIndexer::CalculateComponentSpaceTransform(FAssetIndexer::Cached
 	return Entry.ComponentSpacePose.GetComponentSpaceTransform(CompactBoneIndex);
 }
 
-// returns the transform in component space for the bone indexed by Schema->BoneReferences[SchemaBoneIdx] 
-// at SampleTime-OriginTime seconds ahead (or behind) the pose at OriginTime
-FTransform FAssetIndexer::GetComponentSpaceTransform(float SampleTime, float OriginTime, bool& bClamped, int8 SchemaBoneIdx)
+FQuat FAssetIndexer::GetSampleRotation(float SampleTimeOffset, int32 SampleIdx, int8 SchemaSampleBoneIdx, int8 SchemaOriginBoneIdx)
 {
+	const float OriginTime = FMath::Min(SampleIdx * IndexingContext.Schema->GetSamplingInterval(), IndexingContext.AssetSampler->GetPlayLength());
+	const float SampleTime = OriginTime + SampleTimeOffset;
+
+	// @todo: add support for SchemaSampleBoneIdx
+	if (!IndexingContext.Schema->IsRootBone(SchemaOriginBoneIdx))
+	{
+		UE_LOG(LogPoseSearch,
+			Error,
+			TEXT("FAssetIndexer::GetSampleRotation: support for non root origin bones not implemented (bone: '%s', schema: '%s'"), 
+			*IndexingContext.Schema->BoneReferences[SchemaOriginBoneIdx].BoneName.ToString(),
+			*GetNameSafe(IndexingContext.Schema));
+	}
+
+	bool bUnused;
 	if (SampleTime == OriginTime)
 	{
-		return GetComponentSpaceTransform(SampleTime, bClamped, SchemaBoneIdx);
+		return GetComponentSpaceTransform(SampleTime, bUnused, SchemaSampleBoneIdx).GetRotation();
 	}
 
-	bool bUnused = false;
 	const FTransform RootBoneTransform = GetTransform(OriginTime, bUnused);
-	FTransform BoneTransform = GetTransform(SampleTime, bClamped, SchemaBoneIdx);
+	FTransform BoneTransform = GetTransform(SampleTime, bUnused, SchemaSampleBoneIdx);
 	BoneTransform.SetToRelativeTransform(RootBoneTransform);
-	return BoneTransform;
+	return BoneTransform.GetRotation();
+}
+
+FVector FAssetIndexer::GetSamplePosition(float SampleTimeOffset, int32 SampleIdx, int8 SchemaSampleBoneIdx, int8 SchemaOriginBoneIdx)
+{
+	const float OriginTime = FMath::Min(SampleIdx * IndexingContext.Schema->GetSamplingInterval(), IndexingContext.AssetSampler->GetPlayLength());
+	const float SampleTime = OriginTime + SampleTimeOffset;
+
+	bool bUnused;
+	return GetSamplePositionInternal(SampleTime, OriginTime, bUnused, SchemaSampleBoneIdx, SchemaOriginBoneIdx);
+}
+
+FVector FAssetIndexer::GetSamplePositionInternal(float SampleTime, float OriginTime, bool& bClamped, int8 SchemaSampleBoneIdx, int8 SchemaOriginBoneIdx)
+{
+	if (SampleTime == OriginTime)
+	{
+		if (IndexingContext.Schema->IsRootBone(SchemaOriginBoneIdx))
+		{
+			return GetComponentSpaceTransform(SampleTime, bClamped, SchemaSampleBoneIdx).GetTranslation();
+		}
+
+		bool bOriginClamped;
+		const FVector SampleBonePosition = GetComponentSpaceTransform(SampleTime, bClamped, SchemaSampleBoneIdx).GetTranslation();
+		const FVector OriginBonePosition = GetComponentSpaceTransform(OriginTime, bOriginClamped, SchemaOriginBoneIdx).GetTranslation();
+		bClamped |= bOriginClamped;
+		return SampleBonePosition - OriginBonePosition;
+	}
+
+	bool Unused;
+	const FTransform RootBoneTransform = GetTransform(OriginTime, Unused);
+	const FTransform SampleBoneTransform = GetTransform(SampleTime, bClamped, SchemaSampleBoneIdx);
+	if (IndexingContext.Schema->IsRootBone(SchemaOriginBoneIdx))
+	{
+		return RootBoneTransform.InverseTransformPosition(SampleBoneTransform.GetTranslation());
+	}
+
+	bool bOriginClamped;
+	const FTransform OriginBoneTransform = GetTransform(OriginTime, bOriginClamped, SchemaOriginBoneIdx);
+	bClamped |= bOriginClamped;
+	const FVector DeltaBoneTranslation = SampleBoneTransform.GetTranslation() - OriginBoneTransform.GetTranslation();
+	return RootBoneTransform.InverseTransformVector(DeltaBoneTranslation);
+}
+
+FVector FAssetIndexer::GetSampleVelocity(float SampleTimeOffset, int32 SampleIdx, int8 SchemaSampleBoneIdx, int8 SchemaOriginBoneIdx, bool bUseCharacterSpaceVelocities)
+{
+	const float OriginSampleTime = FMath::Min(SampleIdx * IndexingContext.Schema->GetSamplingInterval(), IndexingContext.AssetSampler->GetPlayLength());
+	const float SubsampleTime = OriginSampleTime + SampleTimeOffset;
+	const float FiniteDelta = IndexingContext.SamplingContext->FiniteDelta;
+
+	bool bClampedPast, bClampedPresent;
+	const FVector BonePositionPast = GetSamplePositionInternal(SubsampleTime - FiniteDelta, bUseCharacterSpaceVelocities ? OriginSampleTime - FiniteDelta : OriginSampleTime, bClampedPast, SchemaSampleBoneIdx, SchemaOriginBoneIdx);
+	const FVector BonePositionPresent = GetSamplePositionInternal(SubsampleTime, OriginSampleTime, bClampedPresent, SchemaSampleBoneIdx, SchemaOriginBoneIdx);
+
+	FVector LinearVelocity;
+	if (!bClampedPast)
+	{
+		LinearVelocity = (BonePositionPresent - BonePositionPast) / FiniteDelta;
+	}
+	else
+	{
+		bool Unused;
+		const FVector BonePositionFuture = GetSamplePositionInternal(SubsampleTime + FiniteDelta, bUseCharacterSpaceVelocities ? OriginSampleTime + FiniteDelta : OriginSampleTime, Unused, SchemaSampleBoneIdx, SchemaOriginBoneIdx);
+		LinearVelocity = (BonePositionFuture - BonePositionPresent) / FiniteDelta;
+	}
+	return LinearVelocity;
 }
 
+
 } // namespace UE::PoseSearch
 #endif // WITH_EDITOR
diff --git a/Engine/Plugins/Experimental/Animation/PoseSearch/Source/Runtime/Private/PoseSearchContext.cpp b/Engine/Plugins/Experimental/Animation/PoseSearch/Source/Runtime/Private/PoseSearchContext.cpp
--- a/Engine/Plugins/Experimental/Animation/PoseSearch/Source/Runtime/Private/PoseSearchContext.cpp
+++ b/Engine/Plugins/Experimental/Animation/PoseSearch/Source/Runtime/Private/PoseSearchContext.cpp
@@ -136,6 +136,38 @@ void DrawFeatureVector(FDebugDrawParams& DrawParams, int32 PoseIdx)
 
 //////////////////////////////////////////////////////////////////////////
 // FSearchContext
+FQuat FSearchContext::GetSampleRotation(float SampleTimeOffset, const UPoseSearchSchema* Schema, int8 SchemaSampleBoneIdx, int8 SchemaOriginBoneIdx, bool bUseHistoryRoot)
+{
+	// @todo: add support for SchemaSampleBoneIdx
+	if (SchemaOriginBoneIdx != RootSchemaBoneIdx)
+	{
+		UE_LOG(LogPoseSearch,
+			Error,
+			TEXT("FSearchContext::GetSampleRotation: support for non root origin bones not implemented (bone: '%s', schema: '%s'"),
+			*Schema->BoneReferences[SchemaOriginBoneIdx].BoneName.ToString(),
+			*GetNameSafe(Schema));
+	}
+
+	return GetComponentSpaceTransform(SampleTimeOffset, Schema, SchemaSampleBoneIdx).GetRotation();
+}
+
+FVector FSearchContext::GetSamplePosition(float SampleTimeOffset, const UPoseSearchSchema* Schema, int8 SchemaSampleBoneIdx, int8 SchemaOriginBoneIdx, bool bUseHistoryRoot)
+{
+	return GetSamplePositionInternal(SampleTimeOffset, 0.f, Schema, SchemaSampleBoneIdx, SchemaOriginBoneIdx, bUseHistoryRoot);
+}
+
+FVector FSearchContext::GetSampleVelocity(float SampleTimeOffset, const UPoseSearchSchema* Schema, int8 SchemaSampleBoneIdx, int8 SchemaOriginBoneIdx, bool bUseCharacterSpaceVelocities, bool bUseHistoryRoot)
+{
+	const float HistorySampleInterval = History ? History->GetSampleTimeInterval() : 1 / 60.0f;
+	check(HistorySampleInterval > UE_KINDA_SMALL_NUMBER);
+
+	// calculating the Position in component space for the bone indexed by SchemaSampleBoneIdx
+	const FVector CurrentTranslation = GetSamplePositionInternal(SampleTimeOffset, 0.f, Schema, SchemaSampleBoneIdx, SchemaOriginBoneIdx, bUseHistoryRoot);
+	const FVector PreviousTranslation = GetSamplePositionInternal(SampleTimeOffset - HistorySampleInterval, bUseCharacterSpaceVelocities ? -HistorySampleInterval : 0.f, Schema, SchemaSampleBoneIdx, SchemaOriginBoneIdx, bUseHistoryRoot);
+
+	const FVector LinearVelocity = (CurrentTranslation - PreviousTranslation) / HistorySampleInterval;
+	return LinearVelocity;
+}
 
 FTransform FSearchContext::GetTransform(float SampleTime, const UPoseSearchSchema* Schema, int8 SchemaBoneIdx, bool bUseHistoryRoot)
 {
@@ -196,17 +228,30 @@ FTransform FSearchContext::GetComponentSpaceTransform(float SampleTime, const UP
 	return FTransform::Identity;
 }
 
-FTransform FSearchContext::GetComponentSpaceTransform(float SampleTime, float OriginTime, const UPoseSearchSchema* Schema, int8 SchemaBoneIdx, bool bUseHistoryRoot)
+FVector FSearchContext::GetSamplePositionInternal(float SampleTime, float OriginTime, const UPoseSearchSchema* Schema, int8 SchemaSampleBoneIdx, int8 SchemaOriginBoneIdx, bool bUseHistoryRoot)
 {
 	if (SampleTime == OriginTime)
 	{
-		return GetComponentSpaceTransform(SampleTime, Schema, SchemaBoneIdx);
+		if (Schema->IsRootBone(SchemaOriginBoneIdx))
+		{
+			return GetComponentSpaceTransform(SampleTime, Schema, SchemaSampleBoneIdx).GetTranslation();
+		}
+
+		const FVector SampleBonePosition = GetComponentSpaceTransform(SampleTime, Schema, SchemaSampleBoneIdx).GetTranslation();
+		const FVector OriginBonePosition = GetComponentSpaceTransform(OriginTime, Schema, SchemaOriginBoneIdx).GetTranslation();
+		return SampleBonePosition - OriginBonePosition;
 	}
 
 	const FTransform RootBoneTransform = GetTransform(OriginTime, Schema, RootSchemaBoneIdx, bUseHistoryRoot);
-	FTransform BoneTransform = GetTransform(SampleTime, Schema, SchemaBoneIdx, bUseHistoryRoot);
-	BoneTransform.SetToRelativeTransform(RootBoneTransform);
-	return BoneTransform;
+	const FTransform SampleBoneTransform = GetTransform(SampleTime, Schema, SchemaSampleBoneIdx, bUseHistoryRoot);
+	if (Schema->IsRootBone(SchemaOriginBoneIdx))
+	{
+		return RootBoneTransform.InverseTransformPosition(SampleBoneTransform.GetTranslation());
+	}
+
+	const FTransform OriginBoneTransform = GetTransform(OriginTime, Schema, SchemaOriginBoneIdx, bUseHistoryRoot);
+	const FVector DeltaBoneTranslation = SampleBoneTransform.GetTranslation() - OriginBoneTransform.GetTranslation();
+	return RootBoneTransform.InverseTransformVector(DeltaBoneTranslation);
 }
 
 void FSearchContext::ClearCachedEntries()
diff --git a/Engine/Plugins/Experimental/Animation/PoseSearch/Source/Runtime/Private/PoseSearchDerivedDataKey.cpp b/Engine/Plugins/Experimental/Animation/PoseSearch/Source/Runtime/Private/PoseSearchDerivedDataKey.cpp
--- a/Engine/Plugins/Experimental/Animation/PoseSearch/Source/Runtime/Private/PoseSearchDerivedDataKey.cpp
+++ b/Engine/Plugins/Experimental/Animation/PoseSearch/Source/Runtime/Private/PoseSearchDerivedDataKey.cpp
@@ -41,7 +41,7 @@ FKeyBuilder::FKeyBuilder(const UObject* Object, bool bUseDataVer)
 	if (bUseDataVer)
 	{
 		// used to invalidate the key without having to change POSESEARCHDB_DERIVEDDATA_VER all the times
-		int32 POSESEARCHDB_DERIVEDDATA_VER_SMALL = 61;
+		int32 POSESEARCHDB_DERIVEDDATA_VER_SMALL = 67;
 		FGuid VersionGuid = FDevSystemGuids::GetSystemGuid(FDevSystemGuids::Get().POSESEARCHDB_DERIVEDDATA_VER);
 
 		*this << VersionGuid;
diff --git a/Engine/Plugins/Experimental/Animation/PoseSearch/Source/Runtime/Private/PoseSearchFeatureChannel_FilterCrashingLegs.cpp b/Engine/Plugins/Experimental/Animation/PoseSearch/Source/Runtime/Private/PoseSearchFeatureChannel_FilterCrashingLegs.cpp
--- a/Engine/Plugins/Experimental/Animation/PoseSearch/Source/Runtime/Private/PoseSearchFeatureChannel_FilterCrashingLegs.cpp
+++ b/Engine/Plugins/Experimental/Animation/PoseSearch/Source/Runtime/Private/PoseSearchFeatureChannel_FilterCrashingLegs.cpp
@@ -55,15 +55,12 @@ void UPoseSearchFeatureChannel_FilterCrashingLegs::IndexAsset(UE::PoseSearch::FA
 
 	for (int32 SampleIdx = IndexingContext.BeginSampleIdx; SampleIdx != IndexingContext.EndSampleIdx; ++SampleIdx)
 	{
-		const float SubsampleTime = SampleIdx * IndexingContext.Schema->GetSamplingInterval();
+		const FVector RightThighPosition = Indexer.GetSamplePosition(0.f, SampleIdx, RightThighIdx);
+		const FVector LeftThighPosition = Indexer.GetSamplePosition(0.f, SampleIdx, LeftThighIdx);
+		const FVector RightFootPosition = Indexer.GetSamplePosition(0.f, SampleIdx, RightFootIdx);
+		const FVector LeftFootPosition = Indexer.GetSamplePosition(0.f, SampleIdx, LeftFootIdx);
 
-		bool bUnused;
-		const FTransform RightThighTransform = Indexer.GetComponentSpaceTransform(SubsampleTime, bUnused, RightThighIdx);
-		const FTransform LeftThighTransform = Indexer.GetComponentSpaceTransform(SubsampleTime, bUnused, LeftThighIdx);
-		const FTransform RightFootTransform = Indexer.GetComponentSpaceTransform(SubsampleTime, bUnused, RightFootIdx);
-		const FTransform LeftFootTransform = Indexer.GetComponentSpaceTransform(SubsampleTime, bUnused, LeftFootIdx);
-
-		const float CrashingLegsValue = ComputeCrashingLegsValue(RightThighTransform.GetTranslation(), LeftThighTransform.GetTranslation(), RightFootTransform.GetTranslation(), LeftFootTransform.GetTranslation());
+		const float CrashingLegsValue = ComputeCrashingLegsValue(RightThighPosition, LeftThighPosition, RightFootPosition, LeftFootPosition);
 
 		const int32 VectorIdx = SampleIdx - IndexingContext.BeginSampleIdx;
 		FFeatureVectorHelper::EncodeFloat(IndexingContext.GetPoseVector(VectorIdx, FeatureVectorTable), ChannelDataOffset, CrashingLegsValue);
@@ -89,13 +86,13 @@ void UPoseSearchFeatureChannel_FilterCrashingLegs::BuildQuery(UE::PoseSearch::FS
 	}
 	else
 	{
-		const float SampleTime = 0.f;
-		const FTransform LeftThighTransform = SearchContext.GetComponentSpaceTransform(SampleTime, InOutQuery.GetSchema(), LeftThighIdx);
-		const FTransform RightThighTransform = SearchContext.GetComponentSpaceTransform(SampleTime, InOutQuery.GetSchema(), RightThighIdx);
-		const FTransform LeftFootTransform = SearchContext.GetComponentSpaceTransform(SampleTime, InOutQuery.GetSchema(), LeftFootIdx);
-		const FTransform RightFootTransform = SearchContext.GetComponentSpaceTransform(SampleTime, InOutQuery.GetSchema(), RightFootIdx);
+		const FVector RightThighPosition = SearchContext.GetSamplePosition(0.f, InOutQuery.GetSchema(), RightThighIdx);
+		const FVector LeftThighPosition = SearchContext.GetSamplePosition(0.f, InOutQuery.GetSchema(), LeftThighIdx);
+		const FVector RightFootPosition = SearchContext.GetSamplePosition(0.f, InOutQuery.GetSchema(), RightFootIdx);
+		const FVector LeftFootPosition = SearchContext.GetSamplePosition(0.f, InOutQuery.GetSchema(), LeftFootIdx);
+
+		const float CrashingLegsValue = ComputeCrashingLegsValue(RightThighPosition, LeftThighPosition, RightFootPosition, LeftFootPosition);
 
-		const float CrashingLegsValue = ComputeCrashingLegsValue(RightThighTransform.GetTranslation(), LeftThighTransform.GetTranslation(), RightFootTransform.GetTranslation(), LeftFootTransform.GetTranslation());
 		FFeatureVectorHelper::EncodeFloat(InOutQuery.EditValues(), ChannelDataOffset, CrashingLegsValue);
 	}
 }
diff --git a/Engine/Plugins/Experimental/Animation/PoseSearch/Source/Runtime/Private/PoseSearchFeatureChannel_Heading.cpp b/Engine/Plugins/Experimental/Animation/PoseSearch/Source/Runtime/Private/PoseSearchFeatureChannel_Heading.cpp
--- a/Engine/Plugins/Experimental/Animation/PoseSearch/Source/Runtime/Private/PoseSearchFeatureChannel_Heading.cpp
+++ b/Engine/Plugins/Experimental/Animation/PoseSearch/Source/Runtime/Private/PoseSearchFeatureChannel_Heading.cpp
@@ -49,14 +49,8 @@ void UPoseSearchFeatureChannel_Heading::IndexAsset(UE::PoseSearch::FAssetIndexer
 
 	for (int32 SampleIdx = IndexingContext.BeginSampleIdx; SampleIdx != IndexingContext.EndSampleIdx; ++SampleIdx)
 	{
+		const FVector Heading = GetAxis(Indexer.GetSampleRotation(SampleTimeOffset, SampleIdx, SchemaBoneIdx));
 		const int32 VectorIdx = SampleIdx - IndexingContext.BeginSampleIdx;
-
-		const float OriginSampleTime = FMath::Min(SampleIdx * IndexingContext.Schema->GetSamplingInterval(), IndexingContext.AssetSampler->GetPlayLength());
-		const float SubsampleTime = OriginSampleTime + SampleTimeOffset;
-
-		bool ClampedPresent;
-		const FTransform BoneTransformsPresent = Indexer.GetComponentSpaceTransform(SubsampleTime, OriginSampleTime, ClampedPresent, SchemaBoneIdx);
-		const FVector Heading = GetAxis(BoneTransformsPresent.GetRotation());
 		FFeatureVectorHelper::EncodeVector(IndexingContext.GetPoseVector(VectorIdx, FeatureVectorTable), ChannelDataOffset, Heading, ComponentStripping);
 	}
 }
@@ -81,10 +75,9 @@ void UPoseSearchFeatureChannel_Heading::BuildQuery(UE::PoseSearch::FSearchContex
 	}
 	else
 	{
-		// calculating the Transform in component space for the bone indexed by SchemaBoneIdx
-		const FTransform Transform = SearchContext.GetComponentSpaceTransform(SampleTimeOffset, 0.f, InOutQuery.GetSchema(), SchemaBoneIdx, bBoneValid);
-
-		FFeatureVectorHelper::EncodeVector(InOutQuery.EditValues(), ChannelDataOffset, GetAxis(Transform.GetRotation()), ComponentStripping);
+		// calculating the BoneRotation in component space for the bone indexed by SchemaBoneIdx
+		const FQuat BoneRotation = SearchContext.GetSampleRotation(SampleTimeOffset, InOutQuery.GetSchema(), SchemaBoneIdx, RootSchemaBoneIdx, bBoneValid);
+		FFeatureVectorHelper::EncodeVector(InOutQuery.EditValues(), ChannelDataOffset, GetAxis(BoneRotation), ComponentStripping);
 	}
 }
 
diff --git a/Engine/Plugins/Experimental/Animation/PoseSearch/Source/Runtime/Private/PoseSearchFeatureChannel_Phase.cpp b/Engine/Plugins/Experimental/Animation/PoseSearch/Source/Runtime/Private/PoseSearchFeatureChannel_Phase.cpp
--- a/Engine/Plugins/Experimental/Animation/PoseSearch/Source/Runtime/Private/PoseSearchFeatureChannel_Phase.cpp
+++ b/Engine/Plugins/Experimental/Animation/PoseSearch/Source/Runtime/Private/PoseSearchFeatureChannel_Phase.cpp
@@ -43,8 +43,6 @@ namespace UE::PoseSearch
 	static void CollectBonePositions(TArray<FVector>& BonePositions, FAssetIndexer& Indexer, int8 SchemaBoneIdx)
 	{
 		const FAssetIndexingContext& IndexingContext = Indexer.GetIndexingContext();
-		const float FiniteDelta = IndexingContext.Schema->GetSamplingInterval();
-		const float SampleTimeStart = FMath::Min(IndexingContext.BeginSampleIdx * FiniteDelta, IndexingContext.AssetSampler->GetPlayLength());
 		const int32 NumSamples = IndexingContext.EndSampleIdx - IndexingContext.BeginSampleIdx;
 
 		// collecting all the bone transforms
@@ -52,9 +50,7 @@ namespace UE::PoseSearch
 		BonePositions.AddDefaulted(NumSamples);
 		for (int32 SampleIdx = 0; SampleIdx != NumSamples; ++SampleIdx)
 		{
-			const float SampleTime = SampleTimeStart + SampleIdx * FiniteDelta;
-			bool bUnused;
-			BonePositions[SampleIdx] = Indexer.GetComponentSpaceTransform(SampleTime, SampleTimeStart, bUnused, SchemaBoneIdx).GetTranslation();
+			BonePositions[SampleIdx] = Indexer.GetSamplePosition(0.f, SampleIdx, SchemaBoneIdx);
 		}
 	}
 
diff --git a/Engine/Plugins/Experimental/Animation/PoseSearch/Source/Runtime/Private/PoseSearchFeatureChannel_Position.cpp b/Engine/Plugins/Experimental/Animation/PoseSearch/Source/Runtime/Private/PoseSearchFeatureChannel_Position.cpp
--- a/Engine/Plugins/Experimental/Animation/PoseSearch/Source/Runtime/Private/PoseSearchFeatureChannel_Position.cpp
+++ b/Engine/Plugins/Experimental/Animation/PoseSearch/Source/Runtime/Private/PoseSearchFeatureChannel_Position.cpp
@@ -15,6 +15,7 @@ void UPoseSearchFeatureChannel_Position::Finalize(UPoseSearchSchema* Schema)
 	Schema->SchemaCardinality += ChannelCardinality;
 
 	SchemaBoneIdx = Schema->AddBoneReference(Bone);
+	SchemaOriginBoneIdx = Schema->AddBoneReference(OriginBone);
 }
 
 #if WITH_EDITOR
@@ -34,14 +35,9 @@ void UPoseSearchFeatureChannel_Position::IndexAsset(UE::PoseSearch::FAssetIndexe
 
 	for (int32 SampleIdx = IndexingContext.BeginSampleIdx; SampleIdx != IndexingContext.EndSampleIdx; ++SampleIdx)
 	{
+		const FVector BonePosition = Indexer.GetSamplePosition(SampleTimeOffset, SampleIdx, SchemaBoneIdx, SchemaOriginBoneIdx);
 		const int32 VectorIdx = SampleIdx - IndexingContext.BeginSampleIdx;
-
-		const float OriginSampleTime = FMath::Min(SampleIdx * IndexingContext.Schema->GetSamplingInterval(), IndexingContext.AssetSampler->GetPlayLength());
-		const float SubsampleTime = OriginSampleTime + SampleTimeOffset;
-
-		bool ClampedPresent;
-		const FTransform BoneTransformsPresent = Indexer.GetComponentSpaceTransform(SubsampleTime, OriginSampleTime, ClampedPresent, SchemaBoneIdx);
-		FFeatureVectorHelper::EncodeVector(IndexingContext.GetPoseVector(VectorIdx, FeatureVectorTable), ChannelDataOffset, BoneTransformsPresent.GetTranslation(), ComponentStripping);
+		FFeatureVectorHelper::EncodeVector(IndexingContext.GetPoseVector(VectorIdx, FeatureVectorTable), ChannelDataOffset, BonePosition, ComponentStripping);
 	}
 }
 #endif // WITH_EDITOR
@@ -65,10 +61,9 @@ void UPoseSearchFeatureChannel_Position::BuildQuery(UE::PoseSearch::FSearchConte
 	}
 	else
 	{
-		// calculating the Transform in component space for the bone indexed by SchemaBoneIdx
-		const FTransform Transform = SearchContext.GetComponentSpaceTransform(SampleTimeOffset, 0.f, InOutQuery.GetSchema(), SchemaBoneIdx, bBoneValid);
-
-		FFeatureVectorHelper::EncodeVector(InOutQuery.EditValues(), ChannelDataOffset, Transform.GetTranslation(), ComponentStripping);
+		// calculating the BonePosition in component space for the bone indexed by SchemaBoneIdx
+		const FVector BonePosition = SearchContext.GetSamplePosition(SampleTimeOffset, InOutQuery.GetSchema(), SchemaBoneIdx, SchemaOriginBoneIdx, bBoneValid);
+		FFeatureVectorHelper::EncodeVector(InOutQuery.EditValues(), ChannelDataOffset, BonePosition, ComponentStripping);
 	}
 }
 
@@ -76,8 +71,13 @@ void UPoseSearchFeatureChannel_Position::BuildQuery(UE::PoseSearch::FSearchConte
 void UPoseSearchFeatureChannel_Position::PreDebugDraw(UE::PoseSearch::FDebugDrawParams& DrawParams, TConstArrayView<float> PoseVector) const
 {
 	using namespace UE::PoseSearch;
-	const FVector BonePos = DrawParams.RootTransform.TransformPosition(FFeatureVectorHelper::DecodeVector(PoseVector, ChannelDataOffset, ComponentStripping));
-	DrawParams.AddCachedPosition(SampleTimeOffset, SchemaBoneIdx, BonePos);
+	
+	// if the SchemaOriginBoneIdx is not the root bone, the feature vector position doesn't represent a component space delta position for the SchemaBoneIdx, so we shouldn't collect it
+	if (!DrawParams.Database->Schema->BoneReferences[SchemaOriginBoneIdx].HasValidSetup())
+	{
+		const FVector BonePos = DrawParams.RootTransform.TransformPosition(FFeatureVectorHelper::DecodeVector(PoseVector, ChannelDataOffset, ComponentStripping));
+		DrawParams.AddCachedPosition(SampleTimeOffset, SchemaBoneIdx, BonePos);
+	}
 }
 
 void UPoseSearchFeatureChannel_Position::DebugDraw(const UE::PoseSearch::FDebugDrawParams& DrawParams, TConstArrayView<float> PoseVector) const
@@ -89,25 +89,40 @@ void UPoseSearchFeatureChannel_Position::DebugDraw(const UE::PoseSearch::FDebugD
 	const bool bPersistent = EnumHasAnyFlags(DrawParams.Flags, EDebugDrawFlags::Persistent);
 	const FColor Color = DrawParams.GetColor(ColorPresetIndex);
 
-	const FVector BonePos = DrawParams.RootTransform.TransformPosition(FFeatureVectorHelper::DecodeVector(PoseVector, ChannelDataOffset, ComponentStripping));
-	// validating DrawParams.AddCachedPosition 
-	check(DrawParams.GetCachedPosition(SampleTimeOffset, SchemaBoneIdx) == BonePos);
-
-	if (EnumHasAnyFlags(DrawParams.Flags, EDebugDrawFlags::DrawFast | EDebugDrawFlags::DrawSearchIndex))
+	if (!DrawParams.Database->Schema->BoneReferences[SchemaOriginBoneIdx].HasValidSetup())
 	{
-		DrawDebugPoint(DrawParams.World, BonePos, DrawParams.PointSize, Color, bPersistent, LifeTime, DepthPriority);
+		const FVector BonePos = DrawParams.RootTransform.TransformPosition(FFeatureVectorHelper::DecodeVector(PoseVector, ChannelDataOffset, ComponentStripping));
+
+		if (EnumHasAnyFlags(DrawParams.Flags, EDebugDrawFlags::DrawFast | EDebugDrawFlags::DrawSearchIndex))
+		{
+			DrawDebugPoint(DrawParams.World, BonePos, DrawParams.PointSize, Color, bPersistent, LifeTime, DepthPriority);
+		}
+		else
+		{
+			DrawDebugSphere(DrawParams.World, BonePos, 2.f, 8, Color, bPersistent, LifeTime, DepthPriority);
+		}
+
+		if (EnumHasAnyFlags(DrawParams.Flags, EDebugDrawFlags::DrawBoneNames))
+		{
+			const UPoseSearchSchema* Schema = DrawParams.GetSchema();
+			check(Schema && Schema->IsValid());
+			DrawDebugString(DrawParams.World, BonePos + FVector(0.0, 0.0, 10.0), Schema->BoneReferences[SchemaBoneIdx].BoneName.ToString(), nullptr, Color, LifeTime, false, 1.0f);
+		}
 	}
 	else
 	{
-		DrawDebugSphere(DrawParams.World, BonePos, 2.f, 8, Color, bPersistent, LifeTime, DepthPriority);
-	}
+		const FVector OriginBonePos = DrawParams.GetCachedPosition(SampleTimeOffset, SchemaOriginBoneIdx);
+		const FVector DeltaPos = DrawParams.RootTransform.TransformVector(FFeatureVectorHelper::DecodeVector(PoseVector, ChannelDataOffset, ComponentStripping));
 
-	if (EnumHasAnyFlags(DrawParams.Flags, EDebugDrawFlags::DrawBoneNames))
-	{
-		const UPoseSearchSchema* Schema = DrawParams.GetSchema();
-		check(Schema && Schema->IsValid());
-		DrawDebugString(DrawParams.World, BonePos + FVector(0.0, 0.0, 10.0), Schema->BoneReferences[SchemaBoneIdx].BoneName.ToString(), nullptr, Color, LifeTime, false, 1.0f);
-	}
+		DrawDebugLine(DrawParams.World, OriginBonePos, OriginBonePos + DeltaPos, Color, bPersistent, LifeTime, DepthPriority);
+
+		if (EnumHasAnyFlags(DrawParams.Flags, EDebugDrawFlags::DrawBoneNames))
+		{
+			const UPoseSearchSchema* Schema = DrawParams.GetSchema();
+			check(Schema && Schema->IsValid());
+			DrawDebugString(DrawParams.World, OriginBonePos + DeltaPos + FVector(0.0, 0.0, 10.0), Schema->BoneReferences[SchemaBoneIdx].BoneName.ToString(), nullptr, Color, LifeTime, false, 1.0f);
+		}
+	}	
 }
 #endif // ENABLE_DRAW_DEBUG
 
@@ -139,6 +154,13 @@ FString UPoseSearchFeatureChannel_Position::GetLabel() const
 		Label.Append(BoneReference.BoneName.ToString());
 	}
 
+	const FBoneReference& OriginBoneReference = GetSchema()->BoneReferences[SchemaOriginBoneIdx];
+	if (OriginBoneReference.HasValidSetup())
+	{
+		Label.Append(TEXT("_"));
+		Label.Append(OriginBoneReference.BoneName.ToString());
+	}
+
 	Label.Appendf(TEXT(" %.1f"), SampleTimeOffset);
 	return Label.ToString();
 }
diff --git a/Engine/Plugins/Experimental/Animation/PoseSearch/Source/Runtime/Private/PoseSearchFeatureChannel_Position.h b/Engine/Plugins/Experimental/Animation/PoseSearch/Source/Runtime/Private/PoseSearchFeatureChannel_Position.h
--- a/Engine/Plugins/Experimental/Animation/PoseSearch/Source/Runtime/Private/PoseSearchFeatureChannel_Position.h
+++ b/Engine/Plugins/Experimental/Animation/PoseSearch/Source/Runtime/Private/PoseSearchFeatureChannel_Position.h
@@ -15,6 +15,9 @@ public:
 	UPROPERTY(EditAnywhere, Category = "Settings")
 	FBoneReference Bone;
 
+	UPROPERTY(EditAnywhere, Category = "Settings")
+	FBoneReference OriginBone;
+
 	UPROPERTY(EditAnywhere, Category = "Settings")
 	float Weight = 1.f;
 
@@ -24,6 +27,9 @@ public:
 	UPROPERTY(Transient)
 	int8 SchemaBoneIdx = 0;
 
+	UPROPERTY(Transient)
+	int8 SchemaOriginBoneIdx = 0;
+
 	UPROPERTY(EditAnywhere, Category = "Settings", meta = (ExcludeFromHash))
 	int32 ColorPresetIndex = 0;
 
diff --git a/Engine/Plugins/Experimental/Animation/PoseSearch/Source/Runtime/Private/PoseSearchFeatureChannel_Velocity.cpp b/Engine/Plugins/Experimental/Animation/PoseSearch/Source/Runtime/Private/PoseSearchFeatureChannel_Velocity.cpp
--- a/Engine/Plugins/Experimental/Animation/PoseSearch/Source/Runtime/Private/PoseSearchFeatureChannel_Velocity.cpp
+++ b/Engine/Plugins/Experimental/Animation/PoseSearch/Source/Runtime/Private/PoseSearchFeatureChannel_Velocity.cpp
@@ -37,31 +37,13 @@ void UPoseSearchFeatureChannel_Velocity::IndexAsset(UE::PoseSearch::FAssetIndexe
 
 	for (int32 SampleIdx = IndexingContext.BeginSampleIdx; SampleIdx != IndexingContext.EndSampleIdx; ++SampleIdx)
 	{
-		const int32 VectorIdx = SampleIdx - IndexingContext.BeginSampleIdx;
-
-		const float OriginSampleTime = FMath::Min(SampleIdx * IndexingContext.Schema->GetSamplingInterval(), IndexingContext.AssetSampler->GetPlayLength());
-		const float SubsampleTime = OriginSampleTime + SampleTimeOffset;
-
-		bool ClampedPast, ClampedPresent, ClampedFuture;
-		const FTransform BoneTransformsPast = Indexer.GetComponentSpaceTransform(SubsampleTime - SamplingContext->FiniteDelta, bUseCharacterSpaceVelocities ? OriginSampleTime - SamplingContext->FiniteDelta : OriginSampleTime, ClampedPast, SchemaBoneIdx);
-		const FTransform BoneTransformsPresent = Indexer.GetComponentSpaceTransform(SubsampleTime, OriginSampleTime, ClampedPresent, SchemaBoneIdx);
-
-		FVector LinearVelocity;
-		if (!ClampedPast)
-		{
-			LinearVelocity = (BoneTransformsPresent.GetTranslation() - BoneTransformsPast.GetTranslation()) / SamplingContext->FiniteDelta;
-		}
-		else
-		{
-			const FTransform BoneTransformsFuture = Indexer.GetComponentSpaceTransform(SubsampleTime + SamplingContext->FiniteDelta, bUseCharacterSpaceVelocities ? OriginSampleTime + SamplingContext->FiniteDelta : OriginSampleTime, ClampedFuture, SchemaBoneIdx);
-			LinearVelocity = (BoneTransformsFuture.GetTranslation() - BoneTransformsPresent.GetTranslation()) / SamplingContext->FiniteDelta;
-		}
-
+		FVector LinearVelocity = Indexer.GetSampleVelocity(SampleTimeOffset, SampleIdx, SchemaBoneIdx, RootSchemaBoneIdx, bUseCharacterSpaceVelocities);
 		if (bNormalize)
 		{
 			LinearVelocity = LinearVelocity.GetClampedToMaxSize(1.f);
 		}
 
+		const int32 VectorIdx = SampleIdx - IndexingContext.BeginSampleIdx;
 		FFeatureVectorHelper::EncodeVector(IndexingContext.GetPoseVector(VectorIdx, FeatureVectorTable), ChannelDataOffset, LinearVelocity, ComponentStripping);
 	}
 }
@@ -86,14 +68,8 @@ void UPoseSearchFeatureChannel_Velocity::BuildQuery(UE::PoseSearch::FSearchConte
 	}
 	else
 	{
-		const float HistorySampleInterval = SearchContext.History ? SearchContext.History->GetSampleTimeInterval() : 1 / 60.0f;
-		check(HistorySampleInterval > UE_KINDA_SMALL_NUMBER);
-
-		// calculating the Transforms in component space for the bone indexed by SchemaBoneIdx
-		const FTransform TransformCurrent = SearchContext.GetComponentSpaceTransform(SampleTimeOffset, 0.f, InOutQuery.GetSchema(), SchemaBoneIdx, bBoneValid);
-		const FTransform TransformPrevious = SearchContext.GetComponentSpaceTransform(SampleTimeOffset - HistorySampleInterval, bUseCharacterSpaceVelocities ? -HistorySampleInterval : 0.f, InOutQuery.GetSchema(), SchemaBoneIdx, bBoneValid);
-
-		FVector LinearVelocity = (TransformCurrent.GetTranslation() - TransformPrevious.GetTranslation()) / HistorySampleInterval;
+		// calculating the LinearVelocity for the bone indexed by SchemaBoneIdx
+		FVector LinearVelocity = SearchContext.GetSampleVelocity(SampleTimeOffset, InOutQuery.GetSchema(), SchemaBoneIdx, RootSchemaBoneIdx, bUseCharacterSpaceVelocities, bBoneValid);
 		if (bNormalize)
 		{
 			LinearVelocity = LinearVelocity.GetClampedToMaxSize(1.f);
diff --git a/Engine/Plugins/Experimental/Animation/PoseSearch/Source/Runtime/Private/PoseSearchSchema.cpp b/Engine/Plugins/Experimental/Animation/PoseSearch/Source/Runtime/Private/PoseSearchSchema.cpp
--- a/Engine/Plugins/Experimental/Animation/PoseSearch/Source/Runtime/Private/PoseSearchSchema.cpp
+++ b/Engine/Plugins/Experimental/Animation/PoseSearch/Source/Runtime/Private/PoseSearchSchema.cpp
@@ -59,6 +59,11 @@ FBoneIndexType UPoseSearchSchema::GetBoneIndexType(int8 SchemaBoneIdx) const
 	return SchemaBoneIdx >= 0 && BoneReferences[SchemaBoneIdx].HasValidSetup() ? BoneReferences[SchemaBoneIdx].BoneIndex : RootBoneIndexType;
 }
 
+bool UPoseSearchSchema::IsRootBone(int8 SchemaBoneIdx) const
+{
+	return SchemaBoneIdx < 0 || !BoneReferences[SchemaBoneIdx].HasValidSetup();
+}
+
 void UPoseSearchSchema::Finalize()
 {
 	BoneReferences.Reset();
diff --git a/Engine/Plugins/Experimental/Animation/PoseSearch/Source/Runtime/Public/PoseSearch/PoseSearchAssetIndexer.h b/Engine/Plugins/Experimental/Animation/PoseSearch/Source/Runtime/Public/PoseSearch/PoseSearchAssetIndexer.h
--- a/Engine/Plugins/Experimental/Animation/PoseSearch/Source/Runtime/Public/PoseSearch/PoseSearchAssetIndexer.h
+++ b/Engine/Plugins/Experimental/Animation/PoseSearch/Source/Runtime/Public/PoseSearch/PoseSearchAssetIndexer.h
@@ -68,11 +68,16 @@ public:
 	const FStats& GetStats() const { return Stats; }
 
 	const FAssetIndexingContext& GetIndexingContext() const { return IndexingContext; }
+
+	FQuat GetSampleRotation(float SampleTimeOffset, int32 SampleIdx, int8 SchemaSampleBoneIdx = RootSchemaBoneIdx, int8 SchemaOriginBoneIdx = RootSchemaBoneIdx);
+	FVector GetSamplePosition(float SampleTimeOffset, int32 SampleIdx, int8 SchemaSampleBoneIdx = RootSchemaBoneIdx, int8 SchemaOriginBoneIdx = RootSchemaBoneIdx);
+	FVector GetSampleVelocity(float SampleTimeOffset, int32 SampleIdx, int8 SchemaSampleBoneIdx = RootSchemaBoneIdx, int8 SchemaOriginBoneIdx = RootSchemaBoneIdx, bool bUseCharacterSpaceVelocities = true);
+
+private:
 	FTransform GetTransform(float SampleTime, bool& bClamped, int8 SchemaBoneIdx = RootSchemaBoneIdx);
 	FTransform GetComponentSpaceTransform(float SampleTime, bool& bClamped, int8 SchemaBoneIdx = RootSchemaBoneIdx);
-	FTransform GetComponentSpaceTransform(float SampleTime, float OriginTime, bool& bClamped, int8 SchemaBoneIdx = RootSchemaBoneIdx);
-	
-private:
+	FVector GetSamplePositionInternal(float SampleTime, float OriginTime, bool& bClamped, int8 SchemaSampleBoneIdx, int8 SchemaOriginBoneIdx);
+
 	struct FSampleInfo
 	{
 		TSharedPtr<FAssetSamplerBase> Clip;
diff --git a/Engine/Plugins/Experimental/Animation/PoseSearch/Source/Runtime/Public/PoseSearch/PoseSearchContext.h b/Engine/Plugins/Experimental/Animation/PoseSearch/Source/Runtime/Public/PoseSearch/PoseSearchContext.h
--- a/Engine/Plugins/Experimental/Animation/PoseSearch/Source/Runtime/Public/PoseSearch/PoseSearchContext.h
+++ b/Engine/Plugins/Experimental/Animation/PoseSearch/Source/Runtime/Public/PoseSearch/PoseSearchContext.h
@@ -161,9 +161,9 @@ struct POSESEARCH_API FSearchContext
 	// can the continuing pose advance? (if not we skip evaluating it)
 	bool bCanAdvance = true;
 
-	FTransform GetTransform(float SampleTime, const UPoseSearchSchema* Schema, int8 SchemaBoneIdx = RootSchemaBoneIdx, bool bUseHistoryRoot = false);
-	FTransform GetComponentSpaceTransform(float SampleTime, const UPoseSearchSchema* Schema, int8 SchemaBoneIdx = RootSchemaBoneIdx);
-	FTransform GetComponentSpaceTransform(float SampleTime, float OriginTime, const UPoseSearchSchema* Schema, int8 SchemaBoneIdx = RootSchemaBoneIdx, bool bUseHistoryRoot = false);
+	FQuat GetSampleRotation(float SampleTimeOffset, const UPoseSearchSchema* Schema, int8 SchemaSampleBoneIdx = RootSchemaBoneIdx, int8 SchemaOriginBoneIdx = RootSchemaBoneIdx, bool bUseHistoryRoot = false);
+	FVector GetSamplePosition(float SampleTimeOffset, const UPoseSearchSchema* Schema, int8 SchemaSampleBoneIdx = RootSchemaBoneIdx, int8 SchemaOriginBoneIdx = RootSchemaBoneIdx, bool bUseHistoryRoot = false);
+	FVector GetSampleVelocity(float SampleTimeOffset, const UPoseSearchSchema* Schema, int8 SchemaSampleBoneIdx = RootSchemaBoneIdx, int8 SchemaOriginBoneIdx = RootSchemaBoneIdx, bool bUseCharacterSpaceVelocities = true, bool bUseHistoryRoot = false);
 
 	void ClearCachedEntries();
 
@@ -183,6 +183,10 @@ struct POSESEARCH_API FSearchContext
 	const FPoseIndicesHistory* PoseIndicesHistory = nullptr;
 
 private:
+	FTransform GetTransform(float SampleTime, const UPoseSearchSchema* Schema, int8 SchemaBoneIdx = RootSchemaBoneIdx, bool bUseHistoryRoot = false);
+	FTransform GetComponentSpaceTransform(float SampleTime, const UPoseSearchSchema* Schema, int8 SchemaSampleBoneIdx = RootSchemaBoneIdx);
+	FVector GetSamplePositionInternal(float SampleTime, float OriginTime, const UPoseSearchSchema* Schema, int8 SchemaSampleBoneIdx = RootSchemaBoneIdx, int8 SchemaOriginBoneIdx = RootSchemaBoneIdx, bool bUseHistoryRoot = false);
+
 	// transforms cached in component space
 	FCachedTransforms<FTransform> CachedTransforms;
 	TArray<FPoseSearchFeatureVectorBuilder, TInlineAllocator<8>> CachedQueries;
diff --git a/Engine/Plugins/Experimental/Animation/PoseSearch/Source/Runtime/Public/PoseSearch/PoseSearchSchema.h b/Engine/Plugins/Experimental/Animation/PoseSearch/Source/Runtime/Public/PoseSearch/PoseSearchSchema.h
--- a/Engine/Plugins/Experimental/Animation/PoseSearch/Source/Runtime/Public/PoseSearch/PoseSearchSchema.h
+++ b/Engine/Plugins/Experimental/Animation/PoseSearch/Source/Runtime/Public/PoseSearch/PoseSearchSchema.h
@@ -130,6 +130,8 @@ public:
 
 	FBoneIndexType GetBoneIndexType(int8 SchemaBoneIdx) const;
 
+	bool IsRootBone(int8 SchemaBoneIdx) const;
+	
 private:
 	void Finalize();
 	void ResolveBoneReferences();
