motion matching - fix for NonUnity build

diff --git a/Engine/Plugins/Experimental/Animation/PoseSearch/Source/Runtime/Public/PoseSearch/PoseSearchDatabase.h b/Engine/Plugins/Experimental/Animation/PoseSearch/Source/Runtime/Public/PoseSearch/PoseSearchDatabase.h
--- a/Engine/Plugins/Experimental/Animation/PoseSearch/Source/Runtime/Public/PoseSearch/PoseSearchDatabase.h
+++ b/Engine/Plugins/Experimental/Animation/PoseSearch/Source/Runtime/Public/PoseSearch/PoseSearchDatabase.h
@@ -8,7 +8,10 @@
 #include "PoseSearchDatabase.generated.h"
 
 struct FInstancedStruct;
+class UAnimationAsset;
 class UAnimComposite;
+class UAnimMontage;
+class UBlendSpace;
 
 UENUM()
 enum class EPoseSearchMode : int32
