Fix PoseSearchHistoryCollector's scoped message being popped before child nodes are updated.

diff --git a/Engine/Plugins/Experimental/Animation/PoseSearch/Source/Runtime/Private/AnimNode_PoseSearchHistoryCollector.cpp b/Engine/Plugins/Experimental/Animation/PoseSearch/Source/Runtime/Private/AnimNode_PoseSearchHistoryCollector.cpp
--- a/Engine/Plugins/Experimental/Animation/PoseSearch/Source/Runtime/Private/AnimNode_PoseSearchHistoryCollector.cpp
+++ b/Engine/Plugins/Experimental/Animation/PoseSearch/Source/Runtime/Private/AnimNode_PoseSearchHistoryCollector.cpp
@@ -75,7 +75,6 @@ void FAnimNode_PoseSearchHistoryCollector_Base::CacheBones_AnyThread(const FAnim
 void FAnimNode_PoseSearchHistoryCollector_Base::Update_AnyThread(const FAnimationUpdateContext& Context)
 {
 	DECLARE_SCOPE_HIERARCHICAL_COUNTER_ANIMNODE(Update_AnyThread);
-	UE::Anim::TScopedGraphMessage<UE::PoseSearch::Private::FPoseHistoryProvider> ScopedMessage(Context, PoseHistory);
 	Super::Update_AnyThread(Context);
 }
 
@@ -149,6 +148,7 @@ void FAnimNode_PoseSearchHistoryCollector::Update_AnyThread(const FAnimationUpda
 {
 	DECLARE_SCOPE_HIERARCHICAL_COUNTER_ANIMNODE(Update_AnyThread);
 	Super::Update_AnyThread(Context);
+	UE::Anim::TScopedGraphMessage<UE::PoseSearch::Private::FPoseHistoryProvider> ScopedMessage(Context, PoseHistory);
 	Source.Update(Context);
 }
 
@@ -195,6 +195,7 @@ void FAnimNode_PoseSearchComponentSpaceHistoryCollector::Update_AnyThread(const
 {
 	DECLARE_SCOPE_HIERARCHICAL_COUNTER_ANIMNODE(Update_AnyThread);
 	Super::Update_AnyThread(Context);
+	UE::Anim::TScopedGraphMessage<UE::PoseSearch::Private::FPoseHistoryProvider> ScopedMessage(Context, PoseHistory);
 	Source.Update(Context);
 }
 
