motion matching - code cleanup

diff --git a/Engine/Plugins/Experimental/Animation/PoseSearch/Source/Runtime/Private/PoseSearchAssetIndexer.cpp b/Engine/Plugins/Experimental/Animation/PoseSearch/Source/Runtime/Private/PoseSearchAssetIndexer.cpp
--- a/Engine/Plugins/Experimental/Animation/PoseSearch/Source/Runtime/Private/PoseSearchAssetIndexer.cpp
+++ b/Engine/Plugins/Experimental/Animation/PoseSearch/Source/Runtime/Private/PoseSearchAssetIndexer.cpp
@@ -99,17 +99,13 @@ void FAssetIndexer::Init(const FAssetIndexingContext& InIndexingContext, const F
 	Output.PoseMetadata.SetNum(Output.NumIndexedPoses);
 }
 
-bool FAssetIndexer::Process()
+void FAssetIndexer::Process()
 {
-	check(IndexingContext.Schema);
-	check(IndexingContext.Schema->IsValid());
+	check(GetSchema()->IsValid());
 	check(IndexingContext.AssetSampler);
 
 	FMemMark Mark(FMemStack::Get());
 
-	IndexingContext.BeginSampleIdx = Output.FirstIndexedSample;
-	IndexingContext.EndSampleIdx = Output.LastIndexedSample + 1;
-
 	if (IndexingContext.Schema->SchemaCardinality > 0)
 	{
 		// Index each channel
@@ -123,17 +119,13 @@ bool FAssetIndexer::Process()
 	}
 
 	// Generate pose metadata
-	for (int32 SampleIdx = IndexingContext.BeginSampleIdx; SampleIdx != IndexingContext.EndSampleIdx; ++SampleIdx)
+	for (int32 SampleIdx = GetBeginSampleIdx(); SampleIdx != GetEndSampleIdx(); ++SampleIdx)
 	{
-		const int32 PoseIdx = SampleIdx - Output.FirstIndexedSample;
-		FPoseSearchPoseMetadata& OutputPoseMetadata = Output.PoseMetadata[PoseIdx];
-		OutputPoseMetadata = GetMetadata(SampleIdx);
+		Output.PoseMetadata[GetVectorIdx(SampleIdx)] = GetMetadata(SampleIdx);
 	}
 
 	// Computing stats
 	ComputeStats();
-
-	return true;
 }
 
 void FAssetIndexer::ComputeStats()
@@ -143,10 +135,8 @@ void FAssetIndexer::ComputeStats()
 	const FAssetSamplingContext* SamplingContext = IndexingContext.SamplingContext;
 	check(SamplingContext->FiniteDelta > UE_KINDA_SMALL_NUMBER);
 
-	for (int32 SampleIdx = IndexingContext.BeginSampleIdx; SampleIdx != IndexingContext.EndSampleIdx; ++SampleIdx)
+	for (int32 SampleIdx = GetBeginSampleIdx(); SampleIdx != GetEndSampleIdx(); ++SampleIdx)
 	{
-		const int32 VectorIdx = SampleIdx - IndexingContext.BeginSampleIdx;
-
 		const float SampleTime = FMath::Min(SampleIdx * IndexingContext.Schema->GetSamplingInterval(), IndexingContext.AssetSampler->GetPlayLength());
 
 		bool AnyClamped = false;
@@ -499,6 +489,22 @@ FVector FAssetIndexer::GetSampleVelocity(float SampleTimeOffset, int32 SampleIdx
 	return LinearVelocity;
 }
 
+int32 FAssetIndexer::GetVectorIdx(int32 SampleIdx) const
+{
+	return SampleIdx - GetBeginSampleIdx();
+}
+
+TArrayView<float> FAssetIndexer::GetPoseVector(int32 SampleIdx, TArrayView<float> FeatureVectorTable) const
+{
+	check(IndexingContext.Schema);
+	return MakeArrayView(&FeatureVectorTable[GetVectorIdx(SampleIdx) * IndexingContext.Schema->SchemaCardinality], IndexingContext.Schema->SchemaCardinality);
+}
+
+const UPoseSearchSchema* FAssetIndexer::GetSchema() const
+{
+	check(IndexingContext.Schema);
+	return IndexingContext.Schema;
+}
 
 } // namespace UE::PoseSearch
 #endif // WITH_EDITOR
diff --git a/Engine/Plugins/Experimental/Animation/PoseSearch/Source/Runtime/Private/PoseSearchFeatureChannel_FilterCrashingLegs.cpp b/Engine/Plugins/Experimental/Animation/PoseSearch/Source/Runtime/Private/PoseSearchFeatureChannel_FilterCrashingLegs.cpp
--- a/Engine/Plugins/Experimental/Animation/PoseSearch/Source/Runtime/Private/PoseSearchFeatureChannel_FilterCrashingLegs.cpp
+++ b/Engine/Plugins/Experimental/Animation/PoseSearch/Source/Runtime/Private/PoseSearchFeatureChannel_FilterCrashingLegs.cpp
@@ -51,9 +51,7 @@ void UPoseSearchFeatureChannel_FilterCrashingLegs::IndexAsset(UE::PoseSearch::FA
 {
 	using namespace UE::PoseSearch;
 
-	const FAssetIndexingContext& IndexingContext = Indexer.GetIndexingContext();
-
-	for (int32 SampleIdx = IndexingContext.BeginSampleIdx; SampleIdx != IndexingContext.EndSampleIdx; ++SampleIdx)
+	for (int32 SampleIdx = Indexer.GetBeginSampleIdx(); SampleIdx != Indexer.GetEndSampleIdx(); ++SampleIdx)
 	{
 		const FVector RightThighPosition = Indexer.GetSamplePosition(0.f, SampleIdx, RightThighIdx);
 		const FVector LeftThighPosition = Indexer.GetSamplePosition(0.f, SampleIdx, LeftThighIdx);
@@ -62,8 +60,7 @@ void UPoseSearchFeatureChannel_FilterCrashingLegs::IndexAsset(UE::PoseSearch::FA
 
 		const float CrashingLegsValue = ComputeCrashingLegsValue(RightThighPosition, LeftThighPosition, RightFootPosition, LeftFootPosition);
 
-		const int32 VectorIdx = SampleIdx - IndexingContext.BeginSampleIdx;
-		FFeatureVectorHelper::EncodeFloat(IndexingContext.GetPoseVector(VectorIdx, FeatureVectorTable), ChannelDataOffset, CrashingLegsValue);
+		FFeatureVectorHelper::EncodeFloat(Indexer.GetPoseVector(SampleIdx, FeatureVectorTable), ChannelDataOffset, CrashingLegsValue);
 	}
 }
 #endif // WITH_EDITOR
diff --git a/Engine/Plugins/Experimental/Animation/PoseSearch/Source/Runtime/Private/PoseSearchFeatureChannel_Heading.cpp b/Engine/Plugins/Experimental/Animation/PoseSearch/Source/Runtime/Private/PoseSearchFeatureChannel_Heading.cpp
--- a/Engine/Plugins/Experimental/Animation/PoseSearch/Source/Runtime/Private/PoseSearchFeatureChannel_Heading.cpp
+++ b/Engine/Plugins/Experimental/Animation/PoseSearch/Source/Runtime/Private/PoseSearchFeatureChannel_Heading.cpp
@@ -45,13 +45,10 @@ void UPoseSearchFeatureChannel_Heading::IndexAsset(UE::PoseSearch::FAssetIndexer
 {
 	using namespace UE::PoseSearch;
 
-	const FAssetIndexingContext& IndexingContext = Indexer.GetIndexingContext();
-
-	for (int32 SampleIdx = IndexingContext.BeginSampleIdx; SampleIdx != IndexingContext.EndSampleIdx; ++SampleIdx)
+	for (int32 SampleIdx = Indexer.GetBeginSampleIdx(); SampleIdx != Indexer.GetEndSampleIdx(); ++SampleIdx)
 	{
 		const FVector Heading = GetAxis(Indexer.GetSampleRotation(SampleTimeOffset, SampleIdx, SchemaBoneIdx));
-		const int32 VectorIdx = SampleIdx - IndexingContext.BeginSampleIdx;
-		FFeatureVectorHelper::EncodeVector(IndexingContext.GetPoseVector(VectorIdx, FeatureVectorTable), ChannelDataOffset, Heading, ComponentStripping);
+		FFeatureVectorHelper::EncodeVector(Indexer.GetPoseVector(SampleIdx, FeatureVectorTable), ChannelDataOffset, Heading, ComponentStripping);
 	}
 }
 #endif // WITH_EDITOR
diff --git a/Engine/Plugins/Experimental/Animation/PoseSearch/Source/Runtime/Private/PoseSearchFeatureChannel_Phase.cpp b/Engine/Plugins/Experimental/Animation/PoseSearch/Source/Runtime/Private/PoseSearchFeatureChannel_Phase.cpp
--- a/Engine/Plugins/Experimental/Animation/PoseSearch/Source/Runtime/Private/PoseSearchFeatureChannel_Phase.cpp
+++ b/Engine/Plugins/Experimental/Animation/PoseSearch/Source/Runtime/Private/PoseSearchFeatureChannel_Phase.cpp
@@ -42,8 +42,7 @@ namespace UE::PoseSearch
 
 	static void CollectBonePositions(TArray<FVector>& BonePositions, FAssetIndexer& Indexer, int8 SchemaBoneIdx)
 	{
-		const FAssetIndexingContext& IndexingContext = Indexer.GetIndexingContext();
-		const int32 NumSamples = IndexingContext.EndSampleIdx - IndexingContext.BeginSampleIdx;
+		const int32 NumSamples = Indexer.GetEndSampleIdx() - Indexer.GetBeginSampleIdx();
 
 		// collecting all the bone transforms
 		BonePositions.Reset();
@@ -307,8 +306,7 @@ void UPoseSearchFeatureChannel_Phase::IndexAsset(UE::PoseSearch::FAssetIndexer&
 	static float SmoothingWindowTime = 0.3f; // seconds
 
 
-	const FAssetIndexingContext& IndexingContext = Indexer.GetIndexingContext();
-	const UPoseSearchSchema* Schema = IndexingContext.Schema;
+	const UPoseSearchSchema* Schema = Indexer.GetSchema();
 	const float FiniteDelta = Schema->GetSamplingInterval();
 	
 	TArray<FVector2D> Phases;
@@ -333,10 +331,9 @@ void UPoseSearchFeatureChannel_Phase::IndexAsset(UE::PoseSearch::FAssetIndexer&
 	ValidateLocalMinMax(LocalMinMax);
 	CalculatePhasesFromLocalMinMax(LocalMinMax, Phases, SmoothedSignal.Num());
 
-	for (int32 SampleIdx = IndexingContext.BeginSampleIdx; SampleIdx != IndexingContext.EndSampleIdx; ++SampleIdx)
+	for (int32 SampleIdx = Indexer.GetBeginSampleIdx(); SampleIdx != Indexer.GetEndSampleIdx(); ++SampleIdx)
 	{
-		const int32 VectorIdx = SampleIdx - IndexingContext.BeginSampleIdx;
-		FFeatureVectorHelper::EncodeVector2D(IndexingContext.GetPoseVector(VectorIdx, FeatureVectorTable), ChannelDataOffset, Phases[VectorIdx]);
+		FFeatureVectorHelper::EncodeVector2D(Indexer.GetPoseVector(SampleIdx, FeatureVectorTable), ChannelDataOffset, Phases[SampleIdx - Indexer.GetBeginSampleIdx()]);
 	}
 }
 #endif // WITH_EDITOR
diff --git a/Engine/Plugins/Experimental/Animation/PoseSearch/Source/Runtime/Private/PoseSearchFeatureChannel_Position.cpp b/Engine/Plugins/Experimental/Animation/PoseSearch/Source/Runtime/Private/PoseSearchFeatureChannel_Position.cpp
--- a/Engine/Plugins/Experimental/Animation/PoseSearch/Source/Runtime/Private/PoseSearchFeatureChannel_Position.cpp
+++ b/Engine/Plugins/Experimental/Animation/PoseSearch/Source/Runtime/Private/PoseSearchFeatureChannel_Position.cpp
@@ -31,13 +31,10 @@ void UPoseSearchFeatureChannel_Position::IndexAsset(UE::PoseSearch::FAssetIndexe
 {
 	using namespace UE::PoseSearch;
 
-	const FAssetIndexingContext& IndexingContext = Indexer.GetIndexingContext();
-
-	for (int32 SampleIdx = IndexingContext.BeginSampleIdx; SampleIdx != IndexingContext.EndSampleIdx; ++SampleIdx)
+	for (int32 SampleIdx = Indexer.GetBeginSampleIdx(); SampleIdx != Indexer.GetEndSampleIdx(); ++SampleIdx)
 	{
 		const FVector BonePosition = Indexer.GetSamplePosition(SampleTimeOffset, SampleIdx, SchemaBoneIdx, SchemaOriginBoneIdx);
-		const int32 VectorIdx = SampleIdx - IndexingContext.BeginSampleIdx;
-		FFeatureVectorHelper::EncodeVector(IndexingContext.GetPoseVector(VectorIdx, FeatureVectorTable), ChannelDataOffset, BonePosition, ComponentStripping);
+		FFeatureVectorHelper::EncodeVector(Indexer.GetPoseVector(SampleIdx, FeatureVectorTable), ChannelDataOffset, BonePosition, ComponentStripping);
 	}
 }
 #endif // WITH_EDITOR
diff --git a/Engine/Plugins/Experimental/Animation/PoseSearch/Source/Runtime/Private/PoseSearchFeatureChannel_Velocity.cpp b/Engine/Plugins/Experimental/Animation/PoseSearch/Source/Runtime/Private/PoseSearchFeatureChannel_Velocity.cpp
--- a/Engine/Plugins/Experimental/Animation/PoseSearch/Source/Runtime/Private/PoseSearchFeatureChannel_Velocity.cpp
+++ b/Engine/Plugins/Experimental/Animation/PoseSearch/Source/Runtime/Private/PoseSearchFeatureChannel_Velocity.cpp
@@ -31,20 +31,14 @@ void UPoseSearchFeatureChannel_Velocity::IndexAsset(UE::PoseSearch::FAssetIndexe
 {
 	using namespace UE::PoseSearch;
 
-	const FAssetIndexingContext& IndexingContext = Indexer.GetIndexingContext();
-	const FAssetSamplingContext* SamplingContext = IndexingContext.SamplingContext;
-	check(SamplingContext->FiniteDelta > UE_KINDA_SMALL_NUMBER);
-
-	for (int32 SampleIdx = IndexingContext.BeginSampleIdx; SampleIdx != IndexingContext.EndSampleIdx; ++SampleIdx)
+	for (int32 SampleIdx = Indexer.GetBeginSampleIdx(); SampleIdx != Indexer.GetEndSampleIdx(); ++SampleIdx)
 	{
 		FVector LinearVelocity = Indexer.GetSampleVelocity(SampleTimeOffset, SampleIdx, SchemaBoneIdx, RootSchemaBoneIdx, bUseCharacterSpaceVelocities);
 		if (bNormalize)
 		{
 			LinearVelocity = LinearVelocity.GetClampedToMaxSize(1.f);
 		}
-
-		const int32 VectorIdx = SampleIdx - IndexingContext.BeginSampleIdx;
-		FFeatureVectorHelper::EncodeVector(IndexingContext.GetPoseVector(VectorIdx, FeatureVectorTable), ChannelDataOffset, LinearVelocity, ComponentStripping);
+		FFeatureVectorHelper::EncodeVector(Indexer.GetPoseVector(SampleIdx, FeatureVectorTable), ChannelDataOffset, LinearVelocity, ComponentStripping);
 	}
 }
 #endif // WITH_EDITOR
diff --git a/Engine/Plugins/Experimental/Animation/PoseSearch/Source/Runtime/Public/PoseSearch/PoseSearchAssetIndexer.h b/Engine/Plugins/Experimental/Animation/PoseSearch/Source/Runtime/Public/PoseSearch/PoseSearchAssetIndexer.h
--- a/Engine/Plugins/Experimental/Animation/PoseSearch/Source/Runtime/Public/PoseSearch/PoseSearchAssetIndexer.h
+++ b/Engine/Plugins/Experimental/Animation/PoseSearch/Source/Runtime/Public/PoseSearch/PoseSearchAssetIndexer.h
@@ -27,16 +27,6 @@ struct FAssetIndexingContext
 	TSharedPtr<FAssetSamplerBase> AssetSampler;
 	bool bMirrored = false;
 	FFloatInterval RequestedSamplingRange = FFloatInterval(0.0f, 0.0f);
-	
-	// Index this asset's data from BeginPoseIdx up to but not including EndPoseIdx
-	int32 BeginSampleIdx = 0;
-	int32 EndSampleIdx = 0;
-
-	TArrayView<float> GetPoseVector(int32 VectorIdx, TArrayView<float> FeatureVectorTable) const
-	{
-		check(Schema);
-		return MakeArrayView(&FeatureVectorTable[VectorIdx * Schema->SchemaCardinality], Schema->SchemaCardinality);
-	}
 };
 
 class POSESEARCH_API FAssetIndexer
@@ -63,17 +53,22 @@ public:
 
 	void Reset();
 	void Init(const FAssetIndexingContext& IndexingContext, const FBoneContainer& InBoneContainer);
-	bool Process();
+	void Process();
 	const FOutput& GetOutput() const { return Output; }
 	const FStats& GetStats() const { return Stats; }
 
-	const FAssetIndexingContext& GetIndexingContext() const { return IndexingContext; }
-
 	FQuat GetSampleRotation(float SampleTimeOffset, int32 SampleIdx, int8 SchemaSampleBoneIdx = RootSchemaBoneIdx, int8 SchemaOriginBoneIdx = RootSchemaBoneIdx);
 	FVector GetSamplePosition(float SampleTimeOffset, int32 SampleIdx, int8 SchemaSampleBoneIdx = RootSchemaBoneIdx, int8 SchemaOriginBoneIdx = RootSchemaBoneIdx);
 	FVector GetSampleVelocity(float SampleTimeOffset, int32 SampleIdx, int8 SchemaSampleBoneIdx = RootSchemaBoneIdx, int8 SchemaOriginBoneIdx = RootSchemaBoneIdx, bool bUseCharacterSpaceVelocities = true);
 
+	int32 GetBeginSampleIdx() const { return Output.FirstIndexedSample; }
+	int32 GetEndSampleIdx() const { return Output.LastIndexedSample + 1; }
+	
+	TArrayView<float> GetPoseVector(int32 SampleIdx, TArrayView<float> FeatureVectorTable) const;
+	const UPoseSearchSchema* GetSchema() const;
+
 private:
+	int32 GetVectorIdx(int32 SampleIdx) const;
 	FTransform GetTransform(float SampleTime, bool& bClamped, int8 SchemaBoneIdx = RootSchemaBoneIdx);
 	FTransform GetComponentSpaceTransform(float SampleTime, bool& bClamped, int8 SchemaBoneIdx = RootSchemaBoneIdx);
 	FVector GetSamplePositionInternal(float SampleTime, float OriginTime, bool& bClamped, int8 SchemaSampleBoneIdx, int8 SchemaOriginBoneIdx);
diff --git a/Engine/Plugins/Experimental/Animation/PoseSearch/Source/Runtime/Public/PoseSearch/PoseSearchSchema.h b/Engine/Plugins/Experimental/Animation/PoseSearch/Source/Runtime/Public/PoseSearch/PoseSearchSchema.h
--- a/Engine/Plugins/Experimental/Animation/PoseSearch/Source/Runtime/Public/PoseSearch/PoseSearchSchema.h
+++ b/Engine/Plugins/Experimental/Animation/PoseSearch/Source/Runtime/Public/PoseSearch/PoseSearchSchema.h
@@ -77,23 +77,23 @@ public:
 	TArray<uint16> BoneIndicesWithParents;
 
 	// cost added to the continuing pose from databases that uses this schema
-	UPROPERTY(EditAnywhere, Category = "Schema")
+	UPROPERTY(EditAnywhere, Category = "Bias")
 	float ContinuingPoseCostBias = 0.f;
 
 	// base cost added to all poses from databases that uses this schema. it can be overridden by UAnimNotifyState_PoseSearchModifyCost
-	UPROPERTY(EditAnywhere, Category = "Schema")
+	UPROPERTY(EditAnywhere, Category = "Bias")
 	float BaseCostBias = 0.f;
 
 	// If there's a mirroring mismatch between the currently playing asset and a search candidate, this cost will be 
 	// added to the candidate, making it less likely to be selected
-	UPROPERTY(EditAnywhere, Category = "Schema")
+	UPROPERTY(EditAnywhere, Category = "Bias")
 	float MirrorMismatchCostBias = 0.f;
 
 	// cost added to all poses from looping assets of databases that uses this schema
-	UPROPERTY(EditAnywhere, Category = "Schema")
+	UPROPERTY(EditAnywhere, Category = "Bias")
 	float LoopingCostBias = 0.f;
 
-	UPROPERTY(EditAnywhere, Category = "Schema", meta = (ExcludeFromHash))
+	UPROPERTY(EditAnywhere, Category = "Debug", meta = (ExcludeFromHash))
 	TArray<FPoseSearchSchemaColorPreset> ColorPresets;
 	
 	bool IsValid () const;
