motion matching - fix for channels crashing while changing databses with different configs and unavailable pose history

diff --git a/Engine/Plugins/Experimental/Animation/PoseSearch/Source/Runtime/Private/PoseSearchFeatureChannel_FilterCrashingLegs.cpp b/Engine/Plugins/Experimental/Animation/PoseSearch/Source/Runtime/Private/PoseSearchFeatureChannel_FilterCrashingLegs.cpp
--- a/Engine/Plugins/Experimental/Animation/PoseSearch/Source/Runtime/Private/PoseSearchFeatureChannel_FilterCrashingLegs.cpp
+++ b/Engine/Plugins/Experimental/Animation/PoseSearch/Source/Runtime/Private/PoseSearchFeatureChannel_FilterCrashingLegs.cpp
@@ -69,8 +69,8 @@ void UPoseSearchFeatureChannel_FilterCrashingLegs::BuildQuery(UE::PoseSearch::FS
 {
 	using namespace UE::PoseSearch;
 
-	const bool bIsCurrentResultValid = SearchContext.CurrentResult.IsValid();
-	const bool bSkip = InputQueryPose != EInputQueryPose::UseCharacterPose && bIsCurrentResultValid && SearchContext.CurrentResult.Database->Schema == InOutQuery.GetSchema();
+	const bool bIsCurrentResultValid = SearchContext.CurrentResult.IsValid() && SearchContext.CurrentResult.Database->Schema == InOutQuery.GetSchema();
+	const bool bSkip = InputQueryPose != EInputQueryPose::UseCharacterPose && bIsCurrentResultValid;
 	
 	if (bSkip || !SearchContext.History)
 	{
diff --git a/Engine/Plugins/Experimental/Animation/PoseSearch/Source/Runtime/Private/PoseSearchFeatureChannel_Heading.cpp b/Engine/Plugins/Experimental/Animation/PoseSearch/Source/Runtime/Private/PoseSearchFeatureChannel_Heading.cpp
--- a/Engine/Plugins/Experimental/Animation/PoseSearch/Source/Runtime/Private/PoseSearchFeatureChannel_Heading.cpp
+++ b/Engine/Plugins/Experimental/Animation/PoseSearch/Source/Runtime/Private/PoseSearchFeatureChannel_Heading.cpp
@@ -57,8 +57,8 @@ void UPoseSearchFeatureChannel_Heading::BuildQuery(UE::PoseSearch::FSearchContex
 {
 	using namespace UE::PoseSearch;
 
-	const bool bIsCurrentResultValid = SearchContext.CurrentResult.IsValid();
-	const bool bSkip = InputQueryPose != EInputQueryPose::UseCharacterPose && bIsCurrentResultValid && SearchContext.CurrentResult.Database->Schema == InOutQuery.GetSchema();
+	const bool bIsCurrentResultValid = SearchContext.CurrentResult.IsValid() && SearchContext.CurrentResult.Database->Schema == InOutQuery.GetSchema();
+	const bool bSkip = InputQueryPose != EInputQueryPose::UseCharacterPose && bIsCurrentResultValid;
 	const bool bBoneValid = InOutQuery.GetSchema()->BoneReferences[SchemaBoneIdx].HasValidSetup();
 	if (bSkip || (!SearchContext.History && bBoneValid))
 	{
diff --git a/Engine/Plugins/Experimental/Animation/PoseSearch/Source/Runtime/Private/PoseSearchFeatureChannel_Phase.cpp b/Engine/Plugins/Experimental/Animation/PoseSearch/Source/Runtime/Private/PoseSearchFeatureChannel_Phase.cpp
--- a/Engine/Plugins/Experimental/Animation/PoseSearch/Source/Runtime/Private/PoseSearchFeatureChannel_Phase.cpp
+++ b/Engine/Plugins/Experimental/Animation/PoseSearch/Source/Runtime/Private/PoseSearchFeatureChannel_Phase.cpp
@@ -342,8 +342,8 @@ void UPoseSearchFeatureChannel_Phase::BuildQuery(UE::PoseSearch::FSearchContext&
 {
 	using namespace UE::PoseSearch;
 
-	const bool bIsCurrentResultValid = SearchContext.CurrentResult.IsValid();
-	const bool bSkip = InputQueryPose != EInputQueryPose::UseCharacterPose && bIsCurrentResultValid && SearchContext.CurrentResult.Database->Schema == InOutQuery.GetSchema();
+	const bool bIsCurrentResultValid = SearchContext.CurrentResult.IsValid() && SearchContext.CurrentResult.Database->Schema == InOutQuery.GetSchema();
+	const bool bSkip = InputQueryPose != EInputQueryPose::UseCharacterPose && bIsCurrentResultValid;
 	if (bSkip || !SearchContext.History)
 	{
 		if (bIsCurrentResultValid)
diff --git a/Engine/Plugins/Experimental/Animation/PoseSearch/Source/Runtime/Private/PoseSearchFeatureChannel_Position.cpp b/Engine/Plugins/Experimental/Animation/PoseSearch/Source/Runtime/Private/PoseSearchFeatureChannel_Position.cpp
--- a/Engine/Plugins/Experimental/Animation/PoseSearch/Source/Runtime/Private/PoseSearchFeatureChannel_Position.cpp
+++ b/Engine/Plugins/Experimental/Animation/PoseSearch/Source/Runtime/Private/PoseSearchFeatureChannel_Position.cpp
@@ -44,8 +44,8 @@ void UPoseSearchFeatureChannel_Position::BuildQuery(UE::PoseSearch::FSearchConte
 	using namespace UE::PoseSearch;
 
 	check(InOutQuery.GetSchema());
-	const bool bIsCurrentResultValid = SearchContext.CurrentResult.IsValid();
-	const bool bSkip = InputQueryPose != EInputQueryPose::UseCharacterPose && bIsCurrentResultValid && SearchContext.CurrentResult.Database->Schema == InOutQuery.GetSchema();
+	const bool bIsCurrentResultValid = SearchContext.CurrentResult.IsValid() && SearchContext.CurrentResult.Database->Schema == InOutQuery.GetSchema();
+	const bool bSkip = InputQueryPose != EInputQueryPose::UseCharacterPose && bIsCurrentResultValid;
 	const bool bBoneValid = InOutQuery.GetSchema()->BoneReferences[SchemaBoneIdx].HasValidSetup();
 	if (bSkip || (!SearchContext.History && bBoneValid))
 	{
diff --git a/Engine/Plugins/Experimental/Animation/PoseSearch/Source/Runtime/Private/PoseSearchFeatureChannel_Velocity.cpp b/Engine/Plugins/Experimental/Animation/PoseSearch/Source/Runtime/Private/PoseSearchFeatureChannel_Velocity.cpp
--- a/Engine/Plugins/Experimental/Animation/PoseSearch/Source/Runtime/Private/PoseSearchFeatureChannel_Velocity.cpp
+++ b/Engine/Plugins/Experimental/Animation/PoseSearch/Source/Runtime/Private/PoseSearchFeatureChannel_Velocity.cpp
@@ -47,8 +47,8 @@ void UPoseSearchFeatureChannel_Velocity::BuildQuery(UE::PoseSearch::FSearchConte
 {
 	using namespace UE::PoseSearch;
 
-	const bool bIsCurrentResultValid = SearchContext.CurrentResult.IsValid();
-	const bool bSkip = InputQueryPose != EInputQueryPose::UseCharacterPose && bIsCurrentResultValid && SearchContext.CurrentResult.Database->Schema == InOutQuery.GetSchema();
+	const bool bIsCurrentResultValid = SearchContext.CurrentResult.IsValid() && SearchContext.CurrentResult.Database->Schema == InOutQuery.GetSchema();
+	const bool bSkip = InputQueryPose != EInputQueryPose::UseCharacterPose && bIsCurrentResultValid;
 	const bool bBoneValid = InOutQuery.GetSchema()->BoneReferences[SchemaBoneIdx].HasValidSetup();
 	if (bSkip || (!SearchContext.History && bBoneValid))
 	{
