motion matching - indexing memory optimization

diff --git a/Engine/Plugins/Experimental/Animation/PoseSearch/Source/Runtime/Private/PoseSearchAssetIndexer.cpp b/Engine/Plugins/Experimental/Animation/PoseSearch/Source/Runtime/Private/PoseSearchAssetIndexer.cpp
--- a/Engine/Plugins/Experimental/Animation/PoseSearch/Source/Runtime/Private/PoseSearchAssetIndexer.cpp
+++ b/Engine/Plugins/Experimental/Animation/PoseSearch/Source/Runtime/Private/PoseSearchAssetIndexer.cpp
@@ -95,11 +95,11 @@ void FAssetIndexer::Init(const FAssetIndexingContext& InIndexingContext, const F
 	Output.LastIndexedSample = FMath::Max(0, FMath::CeilToInt(IndexingContext.RequestedSamplingRange.Max * IndexingContext.Schema->SampleRate));
 	Output.NumIndexedPoses = Output.LastIndexedSample - Output.FirstIndexedSample + 1;
 
-	Output.FeatureVectorTable.SetNumZeroed(IndexingContext.Schema->SchemaCardinality * Output.NumIndexedPoses);
-	Output.PoseMetadata.SetNum(Output.NumIndexedPoses);
+	Output.FeatureVectorTable = TArrayView<float>();
+	Output.PoseMetadata = TArrayView<FPoseSearchPoseMetadata>();
 }
 
-void FAssetIndexer::Process()
+void FAssetIndexer::Process(int32 AssetIdx)
 {
 	check(GetSchema()->IsValid());
 	check(IndexingContext.AssetSampler);
@@ -121,7 +121,7 @@ void FAssetIndexer::Process()
 	// Generate pose metadata
 	for (int32 SampleIdx = GetBeginSampleIdx(); SampleIdx != GetEndSampleIdx(); ++SampleIdx)
 	{
-		Output.PoseMetadata[GetVectorIdx(SampleIdx)] = GetMetadata(SampleIdx);
+		Output.PoseMetadata[GetVectorIdx(SampleIdx)] = GetMetadata(SampleIdx, AssetIdx);
 	}
 
 	// Computing stats
@@ -245,7 +245,7 @@ FTransform FAssetIndexer::MirrorTransform(const FTransform& Transform) const
 	return IndexingContext.bMirrored ? IndexingContext.SamplingContext->MirrorTransform(Transform) : Transform;
 }
 
-FPoseSearchPoseMetadata FAssetIndexer::GetMetadata(int32 SampleIdx) const
+FPoseSearchPoseMetadata FAssetIndexer::GetMetadata(int32 SampleIdx, int32 AssetIdx) const
 {
 	const float SequenceLength = IndexingContext.AssetSampler->GetPlayLength();
 	const float SampleTime = FMath::Min(SampleIdx * IndexingContext.Schema->GetSamplingInterval(), SequenceLength);
@@ -281,6 +281,7 @@ FPoseSearchPoseMetadata FAssetIndexer::GetMetadata(int32 SampleIdx) const
 		Metadata.CostAddend += IndexingContext.Schema->LoopingCostBias;
 	}
 
+	Metadata.AssetIndex = AssetIdx;
 	return Metadata;
 }
 
diff --git a/Engine/Plugins/Experimental/Animation/PoseSearch/Source/Runtime/Private/PoseSearchDatabaseIndexingContext.cpp b/Engine/Plugins/Experimental/Animation/PoseSearch/Source/Runtime/Private/PoseSearchDatabaseIndexingContext.cpp
--- a/Engine/Plugins/Experimental/Animation/PoseSearch/Source/Runtime/Private/PoseSearchDatabaseIndexingContext.cpp
+++ b/Engine/Plugins/Experimental/Animation/PoseSearch/Source/Runtime/Private/PoseSearchDatabaseIndexingContext.cpp
@@ -7,6 +7,7 @@
 #include "Animation/AnimMontage.h"
 #include "Animation/AnimSequence.h"
 #include "Animation/BlendSpace.h"
+#include "DerivedDataRequestOwner.h"
 #include "InstancedStruct.h"
 #include "PoseSearch/PoseSearchDatabase.h"
 #include "PoseSearch/PoseSearchDefines.h"
@@ -15,9 +16,9 @@
 namespace UE::PoseSearch
 {
 
-void FDatabaseIndexingContext::Prepare(const UPoseSearchDatabase* Database)
+bool FDatabaseIndexingContext::IndexDatabase(FPoseSearchIndexBase& SearchIndexBase, const UPoseSearchDatabase& Database, UE::DerivedData::FRequestOwner& Owner)
 {
-	const UPoseSearchSchema* Schema = Database->Schema;
+	const UPoseSearchSchema* Schema = Database.Schema;
 	check(Schema);
 
 	FBoneContainer BoneContainer;
@@ -28,15 +29,21 @@ void FDatabaseIndexingContext::Prepare(const UPoseSearchDatabase* Database)
 
 	SamplingContext.Init(Schema->MirrorDataTable, BoneContainer);
 
+	if (Owner.IsCanceled())
+	{
+		return false;
+	}
+
 	// Prepare samplers for all animation assets.
-	for (const FInstancedStruct& DatabaseAssetStruct : Database->AnimationAssets)
+	////////////////////////////////////////////////////////////////////////////////////////////////////////
+	for (const FInstancedStruct& DatabaseAssetStruct : Database.AnimationAssets)
 	{
 		auto AddSequenceBaseSampler = [&](const UAnimSequenceBase* Sequence)
 		{
 			if (Sequence && !SamplerMap.Contains(Sequence))
 			{
 				FSequenceBaseSampler::FInput Input;
-				Input.ExtrapolationParameters = Database->ExtrapolationParameters;
+				Input.ExtrapolationParameters = Database.ExtrapolationParameters;
 				Input.SequenceBase = Sequence;
 
 				SamplerMap.Add(Sequence, Samplers.Num());
@@ -72,7 +79,7 @@ void FDatabaseIndexingContext::Prepare(const UPoseSearchDatabase* Database)
 						{
 							FBlendSpaceSampler::FInput Input;
 							Input.BoneContainer = BoneContainer;
-							Input.ExtrapolationParameters = Database->ExtrapolationParameters;
+							Input.ExtrapolationParameters = Database.ExtrapolationParameters;
 							Input.BlendSpace = DatabaseBlendSpace->BlendSpace;
 							Input.BlendParameters = BlendParameters;
 
@@ -91,7 +98,7 @@ void FDatabaseIndexingContext::Prepare(const UPoseSearchDatabase* Database)
 			if (DatabaseAnimMontage->AnimMontage && !SamplerMap.Contains(DatabaseAnimMontage->AnimMontage))
 			{
 				FAnimMontageSampler::FInput Input;
-				Input.ExtrapolationParameters = Database->ExtrapolationParameters;
+				Input.ExtrapolationParameters = Database.ExtrapolationParameters;
 				Input.AnimMontage = DatabaseAnimMontage->AnimMontage;
 
 				SamplerMap.Add(DatabaseAnimMontage->AnimMontage, Samplers.Num());
@@ -109,12 +116,20 @@ void FDatabaseIndexingContext::Prepare(const UPoseSearchDatabase* Database)
 
 	ParallelFor(Samplers.Num(), [this](int32 SamplerIdx) { Samplers[SamplerIdx]->Process(); }, ParallelForFlags);
 
+	if (Owner.IsCanceled())
+	{
+		return false;
+	}
+
 	// prepare indexers
-	Indexers.Reserve(SearchIndexBase->Assets.Num());
+	////////////////////////////////////////////////////////////////////////////////////////////////////////
+	Indexers.Reserve(SearchIndexBase.Assets.Num());
 
-	for (int32 AssetIdx = 0; AssetIdx != SearchIndexBase->Assets.Num(); ++AssetIdx)
+	int32 TotalPoses = 0;
+	for (int32 AssetIdx = 0; AssetIdx != SearchIndexBase.Assets.Num(); ++AssetIdx)
 	{
-		const FPoseSearchIndexAsset& SearchIndexAsset = SearchIndexBase->Assets[AssetIdx];
+		FPoseSearchIndexAsset& SearchIndexAsset = SearchIndexBase.Assets[AssetIdx];
+		SearchIndexAsset.FirstPoseIdx = TotalPoses;
 
 		FAssetIndexingContext IndexerContext;
 		IndexerContext.SamplingContext = &SamplingContext;
@@ -122,7 +137,7 @@ void FDatabaseIndexingContext::Prepare(const UPoseSearchDatabase* Database)
 		IndexerContext.RequestedSamplingRange = SearchIndexAsset.SamplingInterval;
 		IndexerContext.bMirrored = SearchIndexAsset.bMirrored;
 
-		const FInstancedStruct& DatabaseAsset = Database->GetAnimationAssetStruct(SearchIndexAsset.SourceAssetIdx);
+		const FInstancedStruct& DatabaseAsset = Database.GetAnimationAssetStruct(SearchIndexAsset.SourceAssetIdx);
 		if (const FPoseSearchDatabaseSequence* DatabaseSequence = DatabaseAsset.GetPtr<FPoseSearchDatabaseSequence>())
 		{
 			if (DatabaseSequence->Sequence)
@@ -159,91 +174,97 @@ void FDatabaseIndexingContext::Prepare(const UPoseSearchDatabase* Database)
 
 		FAssetIndexer& Indexer = Indexers.AddDefaulted_GetRef();
 		Indexer.Init(IndexerContext, BoneContainer);
+
+		const FAssetIndexer::FOutput& Output = Indexer.GetOutput();
+		SearchIndexAsset.NumPoses = Output.NumIndexedPoses;
+		TotalPoses += Output.NumIndexedPoses;
 	}
-}
 
-bool FDatabaseIndexingContext::IndexAssets()
-{
-	// Index asset data
-	ParallelFor(Indexers.Num(), [this](int32 AssetIdx) { Indexers[AssetIdx].Process(); }, ParallelForFlags);
-	return true;
-}
+	// allocating Values and PoseMetadata
+	SearchIndexBase.Values.Reset();
+	SearchIndexBase.PoseMetadata.Reset();
 
-float FDatabaseIndexingContext::CalculateMinCostAddend() const
-{
-	float MinCostAddend = 0.f;
+	SearchIndexBase.Values.SetNumZeroed(Schema->SchemaCardinality * TotalPoses);
+	SearchIndexBase.PoseMetadata.SetNumZeroed(TotalPoses);
 
-	check(SearchIndexBase);
-	if (!SearchIndexBase->PoseMetadata.IsEmpty())
+	// assigning local data to each Indexer
+	TotalPoses = 0;
+	for (int32 AssetIdx = 0; AssetIdx != SearchIndexBase.Assets.Num(); ++AssetIdx)
 	{
-		MinCostAddend = MAX_FLT;
-		for (const FPoseSearchPoseMetadata& PoseMetadata : SearchIndexBase->PoseMetadata)
-		{
-			if (PoseMetadata.CostAddend < MinCostAddend)
-			{
-				MinCostAddend = PoseMetadata.CostAddend;
-			}
-		}
+		FAssetIndexer::FOutput& Output = Indexers[AssetIdx].EditOutput();
+		Output.FeatureVectorTable = MakeArrayView(SearchIndexBase.Values.GetData() + Schema->SchemaCardinality * TotalPoses, Schema->SchemaCardinality * Output.NumIndexedPoses);
+		Output.PoseMetadata = MakeArrayView(SearchIndexBase.PoseMetadata.GetData() + TotalPoses, Output.NumIndexedPoses);
+		TotalPoses += Output.NumIndexedPoses;
 	}
-	return MinCostAddend;
-}
 
-void FDatabaseIndexingContext::JoinIndex()
-{
-	// Write index info to asset and count up total poses and storage required
-	int32 TotalPoses = 0;
-	int32 TotalFloats = 0;
+	if (Owner.IsCanceled())
+	{
+		return false;
+	}
 
-	check(SearchIndexBase);
+	// Index asset data
+	////////////////////////////////////////////////////////////////////////////////////////////////////////
+	ParallelFor(Indexers.Num(), [this](int32 AssetIdx) { Indexers[AssetIdx].Process(AssetIdx); }, ParallelForFlags);
 
-	// Join animation data into a single search index
-	SearchIndexBase->Values.Reset();
-	SearchIndexBase->PoseMetadata.Reset();
-	SearchIndexBase->OverallFlags = EPoseSearchPoseFlags::None;
-	SearchIndexBase->Stats = FPoseSearchStats();
+	if (Owner.IsCanceled())
+	{
+		return false;
+	}
 
+	// Joining Metadata.Flags into OverallFlags
+	////////////////////////////////////////////////////////////////////////////////////////////////////////
+	SearchIndexBase.OverallFlags = EPoseSearchPoseFlags::None;
+	for (const FPoseSearchPoseMetadata& Metadata : SearchIndexBase.PoseMetadata)
+	{
+		SearchIndexBase.OverallFlags |= Metadata.Flags;
+	}
+
+	// Joining Stats
+	////////////////////////////////////////////////////////////////////////////////////////////////////////
 	int32 NumAccumulatedSamples = 0;
-	for (int32 AssetIdx = 0; AssetIdx != SearchIndexBase->Assets.Num(); ++AssetIdx)
+	SearchIndexBase.Stats = FPoseSearchStats();
+	for (int32 AssetIdx = 0; AssetIdx != SearchIndexBase.Assets.Num(); ++AssetIdx)
 	{
-		const FAssetIndexer::FOutput& Output = Indexers[AssetIdx].GetOutput();
+		const FAssetIndexer::FStats& Stats = Indexers[AssetIdx].GetStats();
+		SearchIndexBase.Stats.AverageSpeed += Stats.AccumulatedSpeed;
+		SearchIndexBase.Stats.MaxSpeed = FMath::Max(SearchIndexBase.Stats.MaxSpeed, Stats.MaxSpeed);
+		SearchIndexBase.Stats.AverageAcceleration += Stats.AccumulatedAcceleration;
+		SearchIndexBase.Stats.MaxAcceleration = FMath::Max(SearchIndexBase.Stats.MaxAcceleration, Stats.MaxAcceleration);
 
-		FPoseSearchIndexAsset& SearchIndexAsset = SearchIndexBase->Assets[AssetIdx];
-		SearchIndexAsset.NumPoses = Output.NumIndexedPoses;
-		SearchIndexAsset.FirstPoseIdx = TotalPoses;
+		NumAccumulatedSamples += Stats.NumAccumulatedSamples;
+	}
 
-		const int32 PoseMetadataStartIdx = SearchIndexBase->PoseMetadata.Num();
-		const int32 PoseMetadataEndIdx = PoseMetadataStartIdx + Output.PoseMetadata.Num();
+	if (NumAccumulatedSamples > 0)
+	{
+		const float Denom = 1.f / float(NumAccumulatedSamples);
+		SearchIndexBase.Stats.AverageSpeed *= Denom;
+		SearchIndexBase.Stats.AverageAcceleration *= Denom;
+	}
 
-		SearchIndexBase->Values.Append(Output.FeatureVectorTable.GetData(), Output.FeatureVectorTable.Num());
-		SearchIndexBase->PoseMetadata.Append(Output.PoseMetadata);
+	SearchIndexBase.NumPoses = TotalPoses;
 
-		for (int32 i = PoseMetadataStartIdx; i < PoseMetadataEndIdx; ++i)
+	// Calculate Min Cost Addend
+	////////////////////////////////////////////////////////////////////////////////////////////////////////
+	SearchIndexBase.MinCostAddend = 0.f;
+
+	if (!SearchIndexBase.PoseMetadata.IsEmpty())
+	{
+		SearchIndexBase.MinCostAddend = MAX_FLT;
+		for (const FPoseSearchPoseMetadata& PoseMetadata : SearchIndexBase.PoseMetadata)
 		{
-			SearchIndexBase->PoseMetadata[i].AssetIndex = AssetIdx;
-			SearchIndexBase->OverallFlags |= SearchIndexBase->PoseMetadata[i].Flags;
+			if (PoseMetadata.CostAddend < SearchIndexBase.MinCostAddend)
+			{
+				SearchIndexBase.MinCostAddend = PoseMetadata.CostAddend;
+			}
 		}
-
-		TotalPoses += Output.NumIndexedPoses;
-		TotalFloats += Output.FeatureVectorTable.Num();
-
-		const FAssetIndexer::FStats& Stats = Indexers[AssetIdx].GetStats();
-		SearchIndexBase->Stats.AverageSpeed += Stats.AccumulatedSpeed;
-		SearchIndexBase->Stats.MaxSpeed = FMath::Max(SearchIndexBase->Stats.MaxSpeed, Stats.MaxSpeed);
-		SearchIndexBase->Stats.AverageAcceleration += Stats.AccumulatedAcceleration;
-		SearchIndexBase->Stats.MaxAcceleration = FMath::Max(SearchIndexBase->Stats.MaxAcceleration, Stats.MaxAcceleration);
-
-		NumAccumulatedSamples += Stats.NumAccumulatedSamples;
 	}
 
-	if (NumAccumulatedSamples > 0)
+	if (Owner.IsCanceled())
 	{
-		const float Denom = 1.f / float(NumAccumulatedSamples);
-		SearchIndexBase->Stats.AverageSpeed *= Denom;
-		SearchIndexBase->Stats.AverageAcceleration *= Denom;
+		return false;
 	}
 
-	SearchIndexBase->NumPoses = TotalPoses;
-	SearchIndexBase->MinCostAddend = CalculateMinCostAddend();
+	return true;
 }
 
 } // namespace UE::PoseSearch
diff --git a/Engine/Plugins/Experimental/Animation/PoseSearch/Source/Runtime/Private/PoseSearchDatabaseIndexingContext.h b/Engine/Plugins/Experimental/Animation/PoseSearch/Source/Runtime/Private/PoseSearchDatabaseIndexingContext.h
--- a/Engine/Plugins/Experimental/Animation/PoseSearch/Source/Runtime/Private/PoseSearchDatabaseIndexingContext.h
+++ b/Engine/Plugins/Experimental/Animation/PoseSearch/Source/Runtime/Private/PoseSearchDatabaseIndexingContext.h
@@ -9,21 +9,19 @@
 
 struct FPoseSearchIndexBase;
 class UPoseSearchDatabase;
+namespace UE::DerivedData { class FRequestOwner; }
 
 namespace UE::PoseSearch
 {
 
 struct FDatabaseIndexingContext
 {
-	FPoseSearchIndexBase* SearchIndexBase = nullptr;
+	bool IndexDatabase(FPoseSearchIndexBase& SearchIndexBase, const UPoseSearchDatabase& Database, UE::DerivedData::FRequestOwner& Owner);
+
+private:
 	FAssetSamplingContext SamplingContext;
 	TArray<TSharedPtr<FAssetSamplerBase>> Samplers;
 	TArray<FAssetIndexer> Indexers;
-
-	void Prepare(const UPoseSearchDatabase* Database);
-	bool IndexAssets();
-	void JoinIndex();
-	float CalculateMinCostAddend() const;
 };
 
 } // namespace UE::PoseSearch
diff --git a/Engine/Plugins/Experimental/Animation/PoseSearch/Source/Runtime/Private/PoseSearchDerivedData.cpp b/Engine/Plugins/Experimental/Animation/PoseSearch/Source/Runtime/Private/PoseSearchDerivedData.cpp
--- a/Engine/Plugins/Experimental/Animation/PoseSearch/Source/Runtime/Private/PoseSearchDerivedData.cpp
+++ b/Engine/Plugins/Experimental/Animation/PoseSearch/Source/Runtime/Private/PoseSearchDerivedData.cpp
@@ -946,20 +946,13 @@ void FPoseSearchDatabaseAsyncCacheTask::OnGetComplete(UE::DerivedData::FCacheGet
 					}
 
 					FDatabaseIndexingContext DbIndexingContext;
-					DbIndexingContext.SearchIndexBase = &SearchIndexBase;
-					DbIndexingContext.Prepare(IndexBaseDatabase.Get());
-
-					if (Owner.IsCanceled())
-					{
-						UE_LOG(LogPoseSearch, Log, TEXT("%s - %s BuildIndex Cancelled"), *LexToString(FullIndexKey.Hash), *Database->GetName());
-						SearchIndex.Reset();
-						return;
-					}
-
-					const bool bSuccess = DbIndexingContext.IndexAssets();
-					if (!bSuccess)
+					if (!DbIndexingContext.IndexDatabase(SearchIndexBase, *IndexBaseDatabase, Owner))
 					{
-						if (IndexBaseDatabase == Database)
+						if (Owner.IsCanceled())
+						{
+							UE_LOG(LogPoseSearch, Log, TEXT("%s - %s BuildIndex Cancelled"), *LexToString(FullIndexKey.Hash), *Database->GetName());
+						}
+						else if (IndexBaseDatabase == Database)
 						{
 							UE_LOG(LogPoseSearch, Log, TEXT("%s - %s BuildIndex Failed"), *LexToString(FullIndexKey.Hash), *Database->GetName());
 						}
@@ -967,21 +960,7 @@ void FPoseSearchDatabaseAsyncCacheTask::OnGetComplete(UE::DerivedData::FCacheGet
 						{
 							UE_LOG(LogPoseSearch, Log, TEXT("%s - %s BuildIndex Failed because of dependent database fail '%s'"), *LexToString(FullIndexKey.Hash), *Database->GetName(), *IndexBaseDatabase->GetName());
 						}
-						SearchIndex.Reset();
-						return;
-					}
-
-					if (Owner.IsCanceled())
-					{
-						UE_LOG(LogPoseSearch, Log, TEXT("%s - %s BuildIndex Cancelled"), *LexToString(FullIndexKey.Hash), *Database->GetName());
-						SearchIndex.Reset();
-						return;
-					}
 
-					DbIndexingContext.JoinIndex();
-					if (Owner.IsCanceled())
-					{
-						UE_LOG(LogPoseSearch, Log, TEXT("%s - %s BuildIndex Cancelled"), *LexToString(FullIndexKey.Hash), *Database->GetName());
 						SearchIndex.Reset();
 						return;
 					}
diff --git a/Engine/Plugins/Experimental/Animation/PoseSearch/Source/Runtime/Public/PoseSearch/PoseSearchAssetIndexer.h b/Engine/Plugins/Experimental/Animation/PoseSearch/Source/Runtime/Public/PoseSearch/PoseSearchAssetIndexer.h
--- a/Engine/Plugins/Experimental/Animation/PoseSearch/Source/Runtime/Public/PoseSearch/PoseSearchAssetIndexer.h
+++ b/Engine/Plugins/Experimental/Animation/PoseSearch/Source/Runtime/Public/PoseSearch/PoseSearchAssetIndexer.h
@@ -38,8 +38,8 @@ public:
 		int32 LastIndexedSample = 0;
 		int32 NumIndexedPoses = 0;
 
-		TArray<float> FeatureVectorTable;
-		TArray<FPoseSearchPoseMetadata> PoseMetadata;
+		TArrayView<float> FeatureVectorTable;
+		TArrayView<FPoseSearchPoseMetadata> PoseMetadata;
 	};
 
 	struct FStats
@@ -53,8 +53,9 @@ public:
 
 	void Reset();
 	void Init(const FAssetIndexingContext& IndexingContext, const FBoneContainer& InBoneContainer);
-	void Process();
+	void Process(int32 AssetIdx);
 	const FOutput& GetOutput() const { return Output; }
+	FOutput& EditOutput() { return Output; }
 	const FStats& GetStats() const { return Stats; }
 
 	FQuat GetSampleRotation(float SampleTimeOffset, int32 SampleIdx, int8 SchemaSampleBoneIdx = RootSchemaBoneIdx, int8 SchemaOriginBoneIdx = RootSchemaBoneIdx);
@@ -91,7 +92,7 @@ private:
 	};
 
 	FSampleInfo GetSampleInfo(float SampleTime) const;
-	FPoseSearchPoseMetadata GetMetadata(int32 SampleIdx) const;
+	FPoseSearchPoseMetadata GetMetadata(int32 SampleIdx, int32 AssetIdx) const;
 	FTransform MirrorTransform(const FTransform& Transform) const;
 	CachedEntry& GetEntry(float SampleTime);
 	FTransform CalculateComponentSpaceTransform(CachedEntry& Entry, int8 SchemaBoneIdx);
