motion matching - database index serialization cleanup

diff --git a/Engine/Plugins/Experimental/Animation/PoseSearch/Source/Runtime/Private/KDTree.cpp b/Engine/Plugins/Experimental/Animation/PoseSearch/Source/Runtime/Private/KDTree.cpp
--- a/Engine/Plugins/Experimental/Animation/PoseSearch/Source/Runtime/Private/KDTree.cpp
+++ b/Engine/Plugins/Experimental/Animation/PoseSearch/Source/Runtime/Private/KDTree.cpp
@@ -218,20 +218,21 @@ FArchive& SerializeSubTree(FArchive& Ar, FKDTree& KDTree, FKDTreeImplementation:
 FArchive& Serialize(FArchive& Ar, FKDTree& KDTree, const float* KDTreeData)
 {
 #if UE_POSE_SEARCH_USE_NANOFLANN
+	uint32 KDTreeSize = KDTree.Impl ? KDTree.Impl->m_size : 0;
 
-	if (Ar.IsLoading() && !KDTree.Impl)
-	{
-		KDTree.Impl = new FKDTreeImplementation(0, KDTree.DataSource, nanoflann::KDTreeSingleIndexAdaptorParams(0));
-	}
-
-	check(KDTree.Impl);
-	check(KDTree.Impl->m_size < UINT_MAX);
-	uint32 KDTreeSize = KDTree.Impl->m_size;
 	Ar << KDTreeSize;
-	KDTree.Impl->m_size = KDTreeSize;
 
-	if (KDTree.Impl->m_size  > 0)
+	check(KDTreeSize < UINT_MAX);
+
+	if (KDTreeSize > 0)
 	{
+		if (Ar.IsLoading() && !KDTree.Impl)
+		{
+			KDTree.Impl = new FKDTreeImplementation(0, KDTree.DataSource, nanoflann::KDTreeSingleIndexAdaptorParams(0));
+		}
+
+		KDTree.Impl->m_size = KDTreeSize;
+
 		Ar << KDTree.Impl->dim;
 
 		uint32 root_bbox_size = KDTree.Impl->root_bbox.size();
@@ -270,6 +271,10 @@ FArchive& Serialize(FArchive& Ar, FKDTree& KDTree, const float* KDTreeData)
 		}
 		SerializeSubTree(Ar, KDTree, KDTree.Impl->root_node);
 	}
+	else if (Ar.IsLoading())
+	{
+		KDTree.Reset();
+	}
 #endif
 
 	return Ar;
diff --git a/Engine/Plugins/Experimental/Animation/PoseSearch/Source/Runtime/Private/PoseSearchDerivedData.cpp b/Engine/Plugins/Experimental/Animation/PoseSearch/Source/Runtime/Private/PoseSearchDerivedData.cpp
--- a/Engine/Plugins/Experimental/Animation/PoseSearch/Source/Runtime/Private/PoseSearchDerivedData.cpp
+++ b/Engine/Plugins/Experimental/Animation/PoseSearch/Source/Runtime/Private/PoseSearchDerivedData.cpp
@@ -458,7 +458,7 @@ static void PreprocessSearchIndexPCAData(FPoseSearchIndex& SearchIndex, int32 Nu
 	SearchIndex.Mean.Reset();
 	SearchIndex.PCAProjectionMatrix.Reset();
 
-	if (PoseSearchMode != EPoseSearchMode::BruteForce && NumDimensions > 0)
+	if (PoseSearchMode != EPoseSearchMode::BruteForce && NumDimensions > 0 && NumPoses > 0 && NumberOfPrincipalComponents > 0)
 	{
 		SearchIndex.PCAValues.AddZeroed(NumPoses * NumberOfPrincipalComponents);
 		SearchIndex.Mean.AddZeroed(NumDimensions);
diff --git a/Engine/Plugins/Experimental/Animation/PoseSearch/Source/Runtime/Private/PoseSearchDerivedDataKey.cpp b/Engine/Plugins/Experimental/Animation/PoseSearch/Source/Runtime/Private/PoseSearchDerivedDataKey.cpp
--- a/Engine/Plugins/Experimental/Animation/PoseSearch/Source/Runtime/Private/PoseSearchDerivedDataKey.cpp
+++ b/Engine/Plugins/Experimental/Animation/PoseSearch/Source/Runtime/Private/PoseSearchDerivedDataKey.cpp
@@ -41,7 +41,7 @@ FKeyBuilder::FKeyBuilder(const UObject* Object, bool bUseDataVer)
 	if (bUseDataVer)
 	{
 		// used to invalidate the key without having to change POSESEARCHDB_DERIVEDDATA_VER all the times
-		int32 POSESEARCHDB_DERIVEDDATA_VER_SMALL = 67;
+		int32 POSESEARCHDB_DERIVEDDATA_VER_SMALL = 79;
 		FGuid VersionGuid = FDevSystemGuids::GetSystemGuid(FDevSystemGuids::Get().POSESEARCHDB_DERIVEDDATA_VER);
 
 		*this << VersionGuid;
diff --git a/Engine/Plugins/Experimental/Animation/PoseSearch/Source/Runtime/Private/PoseSearchIndex.cpp b/Engine/Plugins/Experimental/Animation/PoseSearch/Source/Runtime/Private/PoseSearchIndex.cpp
--- a/Engine/Plugins/Experimental/Animation/PoseSearch/Source/Runtime/Private/PoseSearchIndex.cpp
+++ b/Engine/Plugins/Experimental/Animation/PoseSearch/Source/Runtime/Private/PoseSearchIndex.cpp
@@ -31,6 +31,30 @@ void CompareFeatureVectors(TConstArrayView<float> A, TConstArrayView<float> B, T
 
 } // namespace UE::PoseSearch
 
+//////////////////////////////////////////////////////////////////////////
+// FPoseSearchPoseMetadata
+FArchive& operator<<(FArchive& Ar, FPoseSearchPoseMetadata& Metadata)
+{
+	Ar << Metadata.Flags;
+	Ar << Metadata.CostAddend;
+	Ar << Metadata.ContinuingPoseCostAddend;
+	Ar << Metadata.AssetIndex;
+	return Ar;
+}
+
+//////////////////////////////////////////////////////////////////////////
+// FPoseSearchIndexAsset
+FArchive& operator<<(FArchive& Ar, FPoseSearchIndexAsset& IndexAsset)
+{
+	Ar << IndexAsset.SourceAssetIdx;
+	Ar << IndexAsset.bMirrored;
+	Ar << IndexAsset.BlendParameters;
+	Ar << IndexAsset.SamplingInterval;
+	Ar << IndexAsset.FirstPoseIdx;
+	Ar << IndexAsset.NumPoses;
+	return Ar;
+}
+
 //////////////////////////////////////////////////////////////////////////
 // FPoseSearchStats
 FArchive& operator<<(FArchive& Ar, FPoseSearchStats& Stats)
@@ -39,7 +63,6 @@ FArchive& operator<<(FArchive& Ar, FPoseSearchStats& Stats)
 	Ar << Stats.MaxSpeed;
 	Ar << Stats.AverageAcceleration;
 	Ar << Stats.MaxAcceleration;
-
 	return Ar;
 }
 
@@ -77,46 +100,13 @@ void FPoseSearchIndexBase::Reset()
 
 FArchive& operator<<(FArchive& Ar, FPoseSearchIndexBase& Index)
 {
-	int32 NumValues = 0;
-	int32 NumAssets = 0;
-
-	if (Ar.IsSaving())
-	{
-		NumValues = Index.Values.Num();
-		NumAssets = Index.Assets.Num();
-	}
-
 	Ar << Index.NumPoses;
-	Ar << NumValues;
-	Ar << NumAssets;
+	Ar << Index.Values;
+	Ar << Index.PoseMetadata;
 	Ar << Index.OverallFlags;
-
-	if (Ar.IsLoading())
-	{
-		Index.Values.SetNumUninitialized(NumValues);
-		Index.PoseMetadata.SetNumUninitialized(Index.NumPoses);
-		Index.Assets.SetNumUninitialized(NumAssets);
-	}
-
-	if (Index.Values.Num() > 0)
-	{
-		Ar.Serialize(&Index.Values[0], Index.Values.Num() * Index.Values.GetTypeSize());
-	}
-
-	if (Index.PoseMetadata.Num() > 0)
-	{
-		Ar.Serialize(&Index.PoseMetadata[0], Index.PoseMetadata.Num() * Index.PoseMetadata.GetTypeSize());
-	}
-
-	if (Index.Assets.Num() > 0)
-	{
-		Ar.Serialize(&Index.Assets[0], Index.Assets.Num() * Index.Assets.GetTypeSize());
-	}
-
+	Ar << Index.Assets;
 	Ar << Index.MinCostAddend;
-
 	Ar << Index.Stats;
-
 	return Ar;
 }
 
@@ -205,37 +195,11 @@ FArchive& operator<<(FArchive& Ar, FPoseSearchIndex& Index)
 	Ar << static_cast<FPoseSearchIndexBase&>(Index);
 
 	Ar << Index.WeightsSqrt;
+	Ar << Index.PCAValues;
+	Ar << Index.PCAProjectionMatrix;
+	Ar << Index.Mean;
+	Ar << Index.PCAExplainedVariance;
 
-	int32 NumPCAValues = Ar.IsSaving() ? Index.PCAValues.Num() : 0;
-	Ar << NumPCAValues;
-
-	if (NumPCAValues > 0)
-	{
-		if (Ar.IsLoading())
-		{
-			Index.PCAValues.SetNumUninitialized(NumPCAValues);
-		}
-
-		if (Index.PCAValues.Num() > 0)
-		{
-			Ar.Serialize(&Index.PCAValues[0], Index.PCAValues.Num() * Index.PCAValues.GetTypeSize());
-		}
-
-		Ar << Index.Mean;
-		Ar << Index.PCAProjectionMatrix;
-
-		Serialize(Ar, Index.KDTree, Index.PCAValues.GetData());
-
-		Ar << Index.PCAExplainedVariance;
-	}
-	else if (Ar.IsLoading())
-	{
-		Index.PCAValues.Reset();
-		Index.Mean.Reset();
-		Index.PCAProjectionMatrix.Reset();
-		Index.KDTree.Reset();
-		Index.PCAExplainedVariance = 0.f;
-	}
-
+	Serialize(Ar, Index.KDTree, Index.PCAValues.GetData());
 	return Ar;
 }
\ No newline at end of file
diff --git a/Engine/Plugins/Experimental/Animation/PoseSearch/Source/Runtime/Public/PoseSearch/PoseSearchIndex.h b/Engine/Plugins/Experimental/Animation/PoseSearch/Source/Runtime/Public/PoseSearch/PoseSearchIndex.h
--- a/Engine/Plugins/Experimental/Animation/PoseSearch/Source/Runtime/Public/PoseSearch/PoseSearchIndex.h
+++ b/Engine/Plugins/Experimental/Animation/PoseSearch/Source/Runtime/Public/PoseSearch/PoseSearchIndex.h
@@ -63,6 +63,8 @@ struct POSESEARCH_API FPoseSearchPoseMetadata
 	// @todo: consider int16
 	UPROPERTY(meta = (NeverInHash))
 	int32 AssetIndex = INDEX_NONE;
+
+	friend FArchive& operator<<(FArchive& Ar, FPoseSearchPoseMetadata& Metadata);
 };
 
 /**
@@ -106,6 +108,7 @@ struct POSESEARCH_API FPoseSearchIndexAsset
 	int32 NumPoses = 0;
 
 	bool IsPoseInRange(int32 PoseIdx) const { return (PoseIdx >= FirstPoseIdx) && (PoseIdx < FirstPoseIdx + NumPoses); }
+	friend FArchive& operator<<(FArchive& Ar, FPoseSearchIndexAsset& IndexAsset);
 };
 
 USTRUCT()
