motion matching - channel api cleanup

diff --git a/Engine/Plugins/Experimental/Animation/PoseSearch/Source/Editor/Private/PoseSearchDatabaseEditorReflection.cpp b/Engine/Plugins/Experimental/Animation/PoseSearch/Source/Editor/Private/PoseSearchDatabaseEditorReflection.cpp
--- a/Engine/Plugins/Experimental/Animation/PoseSearch/Source/Editor/Private/PoseSearchDatabaseEditorReflection.cpp
+++ b/Engine/Plugins/Experimental/Animation/PoseSearch/Source/Editor/Private/PoseSearchDatabaseEditorReflection.cpp
@@ -117,8 +117,8 @@ void UPoseSearchDatabaseStatistics::Initialize(const UPoseSearchDatabase* PoseSe
 		AnimationSequences = PoseSearchDatabase->AnimationAssets.Num();
 			
 		const int32 SampleRate = FMath::Max(1, PoseSearchDatabase->Schema->SampleRate);
-		TotalAnimationPosesInFrames = SearchIndex.NumPoses;
-		TotalAnimationPosesInTime = FText::Format(TimeFormat, static_cast<double>(SearchIndex.NumPoses) / SampleRate);
+		TotalAnimationPosesInFrames = SearchIndex.GetNumPoses();
+		TotalAnimationPosesInTime = FText::Format(TimeFormat, static_cast<double>(TotalAnimationPosesInFrames) / SampleRate);
 			
 		{
 			uint32 NumOfSearchablePoses = 0;
diff --git a/Engine/Plugins/Experimental/Animation/PoseSearch/Source/Editor/Private/PoseSearchDebuggerDatabaseView.cpp b/Engine/Plugins/Experimental/Animation/PoseSearch/Source/Editor/Private/PoseSearchDebuggerDatabaseView.cpp
--- a/Engine/Plugins/Experimental/Animation/PoseSearch/Source/Editor/Private/PoseSearchDebuggerDatabaseView.cpp
+++ b/Engine/Plugins/Experimental/Animation/PoseSearch/Source/Editor/Private/PoseSearchDebuggerDatabaseView.cpp
@@ -183,7 +183,7 @@ void SDebuggerDatabaseView::Update(const FTraceMotionMatchingStateMessage& State
 				{
 					TSharedRef<FDebuggerDatabaseRowData>& Row = UnfilteredDatabaseRows.Add_GetRef(MakeShared<FDebuggerDatabaseRowData>());
 
-					const float Time = Database->GetAssetTime(PoseEntry.DbPoseIdx, Database->Schema->GetSamplingInterval());
+					const float Time = Database->GetAssetTime(PoseEntry.DbPoseIdx);
 
 					Row->PoseIdx = PoseEntry.DbPoseIdx;
 					Row->SourceDatabase = Database;
diff --git a/Engine/Plugins/Experimental/Animation/PoseSearch/Source/Runtime/Private/PoseSearchAssetIndexer.cpp b/Engine/Plugins/Experimental/Animation/PoseSearch/Source/Runtime/Private/PoseSearchAssetIndexer.cpp
--- a/Engine/Plugins/Experimental/Animation/PoseSearch/Source/Runtime/Private/PoseSearchAssetIndexer.cpp
+++ b/Engine/Plugins/Experimental/Animation/PoseSearch/Source/Runtime/Private/PoseSearchAssetIndexer.cpp
@@ -74,13 +74,8 @@ static FSamplingParam WrapOrClampSamplingParam(bool bCanWrap, float SamplingPara
 
 //////////////////////////////////////////////////////////////////////////
 // FAssetIndexer
-void FAssetIndexer::Reset()
-{
-	Output = FOutput();
-	Stats = FStats();
-}
-
-void FAssetIndexer::Init(const FAssetIndexingContext& InIndexingContext, const FBoneContainer& InBoneContainer)
+FAssetIndexer::FAssetIndexer(const FAssetIndexingContext& InIndexingContext, const FBoneContainer& InBoneContainer, const FPoseSearchIndexAsset& InSearchIndexAsset)
+: SearchIndexAsset(InSearchIndexAsset)
 {
 	check(InIndexingContext.Schema);
 	check(InIndexingContext.Schema->IsValid());
@@ -89,14 +84,15 @@ void FAssetIndexer::Init(const FAssetIndexingContext& InIndexingContext, const F
 	BoneContainer = InBoneContainer;
 	IndexingContext = InIndexingContext;
 
-	Reset();
-
-	Output.FirstIndexedSample = FMath::FloorToInt(IndexingContext.RequestedSamplingRange.Min * IndexingContext.Schema->SampleRate);
-	Output.LastIndexedSample = FMath::Max(0, FMath::CeilToInt(IndexingContext.RequestedSamplingRange.Max * IndexingContext.Schema->SampleRate));
-	Output.NumIndexedPoses = Output.LastIndexedSample - Output.FirstIndexedSample + 1;
+	FirstIndexedSample = FMath::FloorToInt(IndexingContext.RequestedSamplingRange.Min * IndexingContext.Schema->SampleRate);
+	// @todo: parhaps we should use FMath::CeilToInt to avoid sampling over teh length of the aniamtion
+	LastIndexedSample = FMath::Max(0, FMath::CeilToInt(IndexingContext.RequestedSamplingRange.Max * IndexingContext.Schema->SampleRate));
+}
 
-	Output.FeatureVectorTable = TArrayView<float>();
-	Output.PoseMetadata = TArrayView<FPoseSearchPoseMetadata>();
+void FAssetIndexer::AssignWorkingData(TArrayView<float> InOutFeatureVectorTable, TArrayView<FPoseSearchPoseMetadata> InOutPoseMetadata)
+{
+	FeatureVectorTable = InOutFeatureVectorTable;
+	PoseMetadata = InOutPoseMetadata;
 }
 
 void FAssetIndexer::Process(int32 AssetIdx)
@@ -106,22 +102,56 @@ void FAssetIndexer::Process(int32 AssetIdx)
 
 	FMemMark Mark(FMemStack::Get());
 
-	if (IndexingContext.Schema->SchemaCardinality > 0)
+	// Generate pose metadata
+	const float SequenceLength = IndexingContext.AssetSampler->GetPlayLength();
+	for (int32 SampleIdx = GetBeginSampleIdx(); SampleIdx != GetEndSampleIdx(); ++SampleIdx)
 	{
-		// Index each channel
-		for (const TObjectPtr<UPoseSearchFeatureChannel>& ChannelPtr : IndexingContext.Schema->Channels)
+		const float SampleTime = FMath::Min(SampleIdx * IndexingContext.Schema->GetSamplingInterval(), SequenceLength);
+		float CostAddend = IndexingContext.Schema->BaseCostBias;
+		float ContinuingPoseCostAddend = IndexingContext.Schema->ContinuingPoseCostBias;
+		EPoseSearchPoseFlags Flags = EPoseSearchPoseFlags::None;
+
+		TArray<UAnimNotifyState_PoseSearchBase*> NotifyStates;
+		IndexingContext.AssetSampler->ExtractPoseSearchNotifyStates(SampleTime, NotifyStates);
+		for (const UAnimNotifyState_PoseSearchBase* PoseSearchNotify : NotifyStates)
 		{
-			if (ChannelPtr)
+			if (PoseSearchNotify->GetClass()->IsChildOf<UAnimNotifyState_PoseSearchBlockTransition>())
 			{
-				ChannelPtr->IndexAsset(*this, Output.FeatureVectorTable);
+				EnumAddFlags(Flags, EPoseSearchPoseFlags::BlockTransition);
 			}
+			else if (const UAnimNotifyState_PoseSearchModifyCost* ModifyCostNotify = Cast<const UAnimNotifyState_PoseSearchModifyCost>(PoseSearchNotify))
+			{
+				CostAddend = ModifyCostNotify->CostAddend;
+			}
+			else if (const UAnimNotifyState_PoseSearchOverrideContinuingPoseCostBias* ContinuingPoseCostBias = Cast<const UAnimNotifyState_PoseSearchOverrideContinuingPoseCostBias>(PoseSearchNotify))
+			{
+				ContinuingPoseCostAddend = ContinuingPoseCostBias->CostAddend;
+			}
+		}
+
+		if (IndexingContext.AssetSampler->IsLoopable())
+		{
+			CostAddend += IndexingContext.Schema->LoopingCostBias;
 		}
+
+		const int32 VectorIdx = GetVectorIdx(SampleIdx);
+		FPoseSearchPoseMetadata& Metadata = PoseMetadata[VectorIdx];
+		Metadata.Flags = Flags;
+		Metadata.CostAddend = CostAddend;
+		Metadata.ContinuingPoseCostAddend = ContinuingPoseCostAddend;
+		Metadata.AssetIndex = AssetIdx;
 	}
 
-	// Generate pose metadata
-	for (int32 SampleIdx = GetBeginSampleIdx(); SampleIdx != GetEndSampleIdx(); ++SampleIdx)
+	// Generate pose features data
+	if (IndexingContext.Schema->SchemaCardinality > 0)
 	{
-		Output.PoseMetadata[GetVectorIdx(SampleIdx)] = GetMetadata(SampleIdx, AssetIdx);
+		for (const TObjectPtr<UPoseSearchFeatureChannel>& ChannelPtr : IndexingContext.Schema->Channels)
+		{
+			if (ChannelPtr)
+			{
+				ChannelPtr->IndexAsset(*this);
+			}
+		}
 	}
 
 	// Computing stats
@@ -245,46 +275,6 @@ FTransform FAssetIndexer::MirrorTransform(const FTransform& Transform) const
 	return IndexingContext.bMirrored ? IndexingContext.SamplingContext->MirrorTransform(Transform) : Transform;
 }
 
-FPoseSearchPoseMetadata FAssetIndexer::GetMetadata(int32 SampleIdx, int32 AssetIdx) const
-{
-	const float SequenceLength = IndexingContext.AssetSampler->GetPlayLength();
-	const float SampleTime = FMath::Min(SampleIdx * IndexingContext.Schema->GetSamplingInterval(), SequenceLength);
-
-	FPoseSearchPoseMetadata Metadata;
-	Metadata.CostAddend = IndexingContext.Schema->BaseCostBias;
-	Metadata.ContinuingPoseCostAddend = IndexingContext.Schema->ContinuingPoseCostBias;
-
-	TArray<UAnimNotifyState_PoseSearchBase*> NotifyStates;
-	IndexingContext.AssetSampler->ExtractPoseSearchNotifyStates(SampleTime, NotifyStates);
-	for (const UAnimNotifyState_PoseSearchBase* PoseSearchNotify : NotifyStates)
-	{
-		if (PoseSearchNotify->GetClass()->IsChildOf<UAnimNotifyState_PoseSearchBlockTransition>())
-		{
-			EnumAddFlags(Metadata.Flags, EPoseSearchPoseFlags::BlockTransition);
-		}
-		else if (PoseSearchNotify->GetClass()->IsChildOf<UAnimNotifyState_PoseSearchModifyCost>())
-		{
-			const UAnimNotifyState_PoseSearchModifyCost* ModifyCostNotify =
-				Cast<const UAnimNotifyState_PoseSearchModifyCost>(PoseSearchNotify);
-			Metadata.CostAddend = ModifyCostNotify->CostAddend;
-		}
-		else if (PoseSearchNotify->GetClass()->IsChildOf<UAnimNotifyState_PoseSearchOverrideContinuingPoseCostBias>())
-		{
-			const UAnimNotifyState_PoseSearchOverrideContinuingPoseCostBias* ContinuingPoseCostBias =
-				Cast<const UAnimNotifyState_PoseSearchOverrideContinuingPoseCostBias>(PoseSearchNotify);
-			Metadata.ContinuingPoseCostAddend = ContinuingPoseCostBias->CostAddend;
-		}
-	}
-
-	if (IndexingContext.AssetSampler->IsLoopable())
-	{
-		Metadata.CostAddend += IndexingContext.Schema->LoopingCostBias;
-	}
-
-	Metadata.AssetIndex = AssetIdx;
-	return Metadata;
-}
-
 FAssetIndexer::CachedEntry& FAssetIndexer::GetEntry(float SampleTime)
 {
 	using namespace UE::Anim;
@@ -420,7 +410,7 @@ FQuat FAssetIndexer::GetSampleRotation(float SampleTimeOffset, int32 SampleIdx,
 		return GetComponentSpaceTransform(SampleTime, bUnused, SchemaSampleBoneIdx).GetRotation();
 	}
 
-	const FTransform RootBoneTransform = GetTransform(OriginTime, bUnused);
+	const FTransform RootBoneTransform = GetTransform(OriginTime, bUnused, RootSchemaBoneIdx);
 	FTransform BoneTransform = GetTransform(SampleTime, bUnused, SchemaSampleBoneIdx);
 	BoneTransform.SetToRelativeTransform(RootBoneTransform);
 	return BoneTransform.GetRotation();
@@ -452,7 +442,7 @@ FVector FAssetIndexer::GetSamplePositionInternal(float SampleTime, float OriginT
 	}
 
 	bool Unused;
-	const FTransform RootBoneTransform = GetTransform(OriginTime, Unused);
+	const FTransform RootBoneTransform = GetTransform(OriginTime, Unused, RootSchemaBoneIdx);
 	const FTransform SampleBoneTransform = GetTransform(SampleTime, bClamped, SchemaSampleBoneIdx);
 	if (IndexingContext.Schema->IsRootBone(SchemaOriginBoneIdx))
 	{
@@ -468,13 +458,13 @@ FVector FAssetIndexer::GetSamplePositionInternal(float SampleTime, float OriginT
 
 FVector FAssetIndexer::GetSampleVelocity(float SampleTimeOffset, int32 SampleIdx, int8 SchemaSampleBoneIdx, int8 SchemaOriginBoneIdx, bool bUseCharacterSpaceVelocities)
 {
-	const float OriginSampleTime = FMath::Min(SampleIdx * IndexingContext.Schema->GetSamplingInterval(), IndexingContext.AssetSampler->GetPlayLength());
-	const float SubsampleTime = OriginSampleTime + SampleTimeOffset;
+	const float OriginTime = FMath::Min(SampleIdx * IndexingContext.Schema->GetSamplingInterval(), IndexingContext.AssetSampler->GetPlayLength());
+	const float SampleTime = OriginTime + SampleTimeOffset;
 	const float FiniteDelta = IndexingContext.SamplingContext->FiniteDelta;
 
 	bool bClampedPast, bClampedPresent;
-	const FVector BonePositionPast = GetSamplePositionInternal(SubsampleTime - FiniteDelta, bUseCharacterSpaceVelocities ? OriginSampleTime - FiniteDelta : OriginSampleTime, bClampedPast, SchemaSampleBoneIdx, SchemaOriginBoneIdx);
-	const FVector BonePositionPresent = GetSamplePositionInternal(SubsampleTime, OriginSampleTime, bClampedPresent, SchemaSampleBoneIdx, SchemaOriginBoneIdx);
+	const FVector BonePositionPast = GetSamplePositionInternal(SampleTime - FiniteDelta, bUseCharacterSpaceVelocities ? OriginTime - FiniteDelta : OriginTime, bClampedPast, SchemaSampleBoneIdx, SchemaOriginBoneIdx);
+	const FVector BonePositionPresent = GetSamplePositionInternal(SampleTime, OriginTime, bClampedPresent, SchemaSampleBoneIdx, SchemaOriginBoneIdx);
 
 	FVector LinearVelocity;
 	if (!bClampedPast)
@@ -484,7 +474,7 @@ FVector FAssetIndexer::GetSampleVelocity(float SampleTimeOffset, int32 SampleIdx
 	else
 	{
 		bool Unused;
-		const FVector BonePositionFuture = GetSamplePositionInternal(SubsampleTime + FiniteDelta, bUseCharacterSpaceVelocities ? OriginSampleTime + FiniteDelta : OriginSampleTime, Unused, SchemaSampleBoneIdx, SchemaOriginBoneIdx);
+		const FVector BonePositionFuture = GetSamplePositionInternal(SampleTime + FiniteDelta, bUseCharacterSpaceVelocities ? OriginTime + FiniteDelta : OriginTime, Unused, SchemaSampleBoneIdx, SchemaOriginBoneIdx);
 		LinearVelocity = (BonePositionFuture - BonePositionPresent) / FiniteDelta;
 	}
 	return LinearVelocity;
@@ -495,7 +485,7 @@ int32 FAssetIndexer::GetVectorIdx(int32 SampleIdx) const
 	return SampleIdx - GetBeginSampleIdx();
 }
 
-TArrayView<float> FAssetIndexer::GetPoseVector(int32 SampleIdx, TArrayView<float> FeatureVectorTable) const
+TArrayView<float> FAssetIndexer::GetPoseVector(int32 SampleIdx) const
 {
 	check(IndexingContext.Schema);
 	return MakeArrayView(&FeatureVectorTable[GetVectorIdx(SampleIdx) * IndexingContext.Schema->SchemaCardinality], IndexingContext.Schema->SchemaCardinality);
diff --git a/Engine/Plugins/Experimental/Animation/PoseSearch/Source/Runtime/Private/PoseSearchContext.cpp b/Engine/Plugins/Experimental/Animation/PoseSearch/Source/Runtime/Private/PoseSearchContext.cpp
--- a/Engine/Plugins/Experimental/Animation/PoseSearch/Source/Runtime/Private/PoseSearchContext.cpp
+++ b/Engine/Plugins/Experimental/Animation/PoseSearch/Source/Runtime/Private/PoseSearchContext.cpp
@@ -127,7 +127,7 @@ void DrawFeatureVector(FDebugDrawParams& DrawParams, TConstArrayView<float> Pose
 void DrawFeatureVector(FDebugDrawParams& DrawParams, int32 PoseIdx)
 {
 	// if we're editing the schema while in PIE with Rewind Debugger active, PoseIdx could be out of bound / stale
-	if (DrawParams.CanDraw() && PoseIdx >= 0 && PoseIdx < DrawParams.GetSearchIndex()->NumPoses)
+	if (DrawParams.CanDraw() && PoseIdx >= 0 && PoseIdx < DrawParams.GetSearchIndex()->GetNumPoses())
 	{
 		DrawFeatureVector(DrawParams, DrawParams.GetSearchIndex()->GetPoseValues(PoseIdx));
 	}
diff --git a/Engine/Plugins/Experimental/Animation/PoseSearch/Source/Runtime/Private/PoseSearchDatabase.cpp b/Engine/Plugins/Experimental/Animation/PoseSearch/Source/Runtime/Private/PoseSearchDatabase.cpp
--- a/Engine/Plugins/Experimental/Animation/PoseSearch/Source/Runtime/Private/PoseSearchDatabase.cpp
+++ b/Engine/Plugins/Experimental/Animation/PoseSearch/Source/Runtime/Private/PoseSearchDatabase.cpp
@@ -121,7 +121,7 @@ static void PopulateNonSelectableIdx(FNonSelectableIdx& NonSelectableIdx, FSearc
 				check(HistoricalPoseIndex.PoseIndex >= 0);
 
 				// if we're editing the database and removing assets it's possible that the PoseIndicesHistory contains invalid pose indexes
-				if (HistoricalPoseIndex.PoseIndex < SearchIndex.NumPoses)
+				if (HistoricalPoseIndex.PoseIndex < SearchIndex.GetNumPoses())
 				{
 					const FPoseSearchCost PoseCost = SearchIndex.ComparePoses(HistoricalPoseIndex.PoseIndex, SearchContext.QueryMirrorRequest, EPoseComparisonFlags::None, Database->Schema->MirrorMismatchCostBias, QueryValues);
 					SearchContext.BestCandidates.Add(PoseCost, HistoricalPoseIndex.PoseIndex, Database, EPoseCandidateFlags::DiscardedBy_PoseReselectHistory);
@@ -654,8 +654,10 @@ void UPoseSearchDatabase::Serialize(FArchive& Ar)
 	}
 }
 
-float UPoseSearchDatabase::GetAssetTime(int32 PoseIdx, float SamplingInterval) const
+float UPoseSearchDatabase::GetAssetTime(int32 PoseIdx) const
 {
+	check(Schema);
+	const float SamplingInterval = Schema->GetSamplingInterval();
 	const FPoseSearchIndexAsset& Asset = GetSearchIndex().GetAssetForPose(PoseIdx);
 	const bool bIsBlendSpace = AnimationAssets[Asset.SourceAssetIdx].GetPtr<FPoseSearchDatabaseBlendSpace>() != nullptr;
 	const FFloatInterval& SamplingRange = Asset.SamplingInterval;
@@ -724,7 +726,7 @@ UE::PoseSearch::FSearchResult UPoseSearchDatabase::SearchPCAKDTree(UE::PoseSearc
 	const FPoseSearchIndex& SearchIndex = GetSearchIndex();
 
 	const uint32 ClampedNumberOfPrincipalComponents = GetNumberOfPrincipalComponents();
-	const uint32 ClampedKDTreeQueryNumNeighbors = FMath::Clamp<uint32>(KDTreeQueryNumNeighbors, 1, SearchIndex.NumPoses);
+	const uint32 ClampedKDTreeQueryNumNeighbors = FMath::Clamp<uint32>(KDTreeQueryNumNeighbors, 1, SearchIndex.GetNumPoses());
 
 	//stack allocated temporaries
 	TArrayView<size_t> ResultIndexes((size_t*)FMemory_Alloca((ClampedKDTreeQueryNumNeighbors + 1) * sizeof(size_t)), ClampedKDTreeQueryNumNeighbors + 1);
@@ -739,7 +741,7 @@ UE::PoseSearch::FSearchResult UPoseSearchDatabase::SearchPCAKDTree(UE::PoseSearc
 		const RowMajorVectorMapConst MapWeightsSqrt(SearchIndex.WeightsSqrt.GetData(), 1, NumDimensions);
 
 		// testing the KDTree is returning the proper searches for all the original points transformed in pca space
-		for (int32 PoseIdx = 0; PoseIdx < SearchIndex.NumPoses; ++PoseIdx)
+		for (int32 PoseIdx = 0; PoseIdx < SearchIndex.GetNumPoses(); ++PoseIdx)
 		{
 			FKDTree::KNNResultSet ResultSet(ClampedKDTreeQueryNumNeighbors, ResultIndexes, ResultDistanceSqr);
 			TConstArrayView<float> PoseValues = SearchIndex.GetPoseValues(PoseIdx);
@@ -849,7 +851,7 @@ UE::PoseSearch::FSearchResult UPoseSearchDatabase::SearchPCAKDTree(UE::PoseSearc
 	// finalizing Result properties
 	if (Result.PoseIdx != INDEX_NONE)
 	{
-		Result.AssetTime = GetAssetTime(Result.PoseIdx, Schema->GetSamplingInterval());
+		Result.AssetTime = GetAssetTime(Result.PoseIdx);
 		Result.Database = this;
 	}
 
@@ -896,7 +898,7 @@ UE::PoseSearch::FSearchResult UPoseSearchDatabase::SearchBruteForce(UE::PoseSear
 		check(Algo::IsSorted(NonSelectableIdx));
 
 		const FPoseFilters PoseFilters(Schema, NonSelectableIdx, SearchIndex.OverallFlags);
-		for (int32 PoseIdx = 0; PoseIdx < SearchIndex.NumPoses; ++PoseIdx)
+		for (int32 PoseIdx = 0; PoseIdx < SearchIndex.GetNumPoses(); ++PoseIdx)
 		{
 			if (PoseFilters.AreFiltersValid(SearchIndex, QueryValues, PoseIdx, SearchIndex.PoseMetadata[PoseIdx]
 #if UE_POSE_SEARCH_TRACE_ENABLED
@@ -937,7 +939,7 @@ UE::PoseSearch::FSearchResult UPoseSearchDatabase::SearchBruteForce(UE::PoseSear
 	// finalizing Result properties
 	if (Result.PoseIdx != INDEX_NONE)
 	{
-		Result.AssetTime = GetAssetTime(Result.PoseIdx, Schema->GetSamplingInterval());
+		Result.AssetTime = GetAssetTime(Result.PoseIdx);
 		Result.Database = this;
 	}
 
diff --git a/Engine/Plugins/Experimental/Animation/PoseSearch/Source/Runtime/Private/PoseSearchDatabaseIndexingContext.cpp b/Engine/Plugins/Experimental/Animation/PoseSearch/Source/Runtime/Private/PoseSearchDatabaseIndexingContext.cpp
--- a/Engine/Plugins/Experimental/Animation/PoseSearch/Source/Runtime/Private/PoseSearchDatabaseIndexingContext.cpp
+++ b/Engine/Plugins/Experimental/Animation/PoseSearch/Source/Runtime/Private/PoseSearchDatabaseIndexingContext.cpp
@@ -35,7 +35,6 @@ bool FDatabaseIndexingContext::IndexDatabase(FPoseSearchIndexBase& SearchIndexBa
 	}
 
 	// Prepare samplers for all animation assets.
-	////////////////////////////////////////////////////////////////////////////////////////////////////////
 	for (const FInstancedStruct& DatabaseAssetStruct : Database.AnimationAssets)
 	{
 		auto AddSequenceBaseSampler = [&](const UAnimSequenceBase* Sequence)
@@ -122,7 +121,6 @@ bool FDatabaseIndexingContext::IndexDatabase(FPoseSearchIndexBase& SearchIndexBa
 	}
 
 	// prepare indexers
-	////////////////////////////////////////////////////////////////////////////////////////////////////////
 	Indexers.Reserve(SearchIndexBase.Assets.Num());
 
 	int32 TotalPoses = 0;
@@ -172,12 +170,9 @@ bool FDatabaseIndexingContext::IndexDatabase(FPoseSearchIndexBase& SearchIndexBa
 			checkNoEntry();
 		}
 
-		FAssetIndexer& Indexer = Indexers.AddDefaulted_GetRef();
-		Indexer.Init(IndexerContext, BoneContainer);
-
-		const FAssetIndexer::FOutput& Output = Indexer.GetOutput();
-		SearchIndexAsset.NumPoses = Output.NumIndexedPoses;
-		TotalPoses += Output.NumIndexedPoses;
+		const int32 NewIndexerIdx = Indexers.Emplace(IndexerContext, BoneContainer, SearchIndexAsset);
+		SearchIndexAsset.NumPoses = Indexers[NewIndexerIdx].GetNumIndexedPoses();
+		TotalPoses += SearchIndexAsset.NumPoses;
 	}
 
 	// allocating Values and PoseMetadata
@@ -191,10 +186,11 @@ bool FDatabaseIndexingContext::IndexDatabase(FPoseSearchIndexBase& SearchIndexBa
 	TotalPoses = 0;
 	for (int32 AssetIdx = 0; AssetIdx != SearchIndexBase.Assets.Num(); ++AssetIdx)
 	{
-		FAssetIndexer::FOutput& Output = Indexers[AssetIdx].EditOutput();
-		Output.FeatureVectorTable = MakeArrayView(SearchIndexBase.Values.GetData() + Schema->SchemaCardinality * TotalPoses, Schema->SchemaCardinality * Output.NumIndexedPoses);
-		Output.PoseMetadata = MakeArrayView(SearchIndexBase.PoseMetadata.GetData() + TotalPoses, Output.NumIndexedPoses);
-		TotalPoses += Output.NumIndexedPoses;
+		const int32 NumIndexedPoses = Indexers[AssetIdx].GetNumIndexedPoses();
+		Indexers[AssetIdx].AssignWorkingData(
+			MakeArrayView(SearchIndexBase.Values.GetData() + Schema->SchemaCardinality * TotalPoses, Schema->SchemaCardinality * NumIndexedPoses),
+			MakeArrayView(SearchIndexBase.PoseMetadata.GetData() + TotalPoses, NumIndexedPoses));
+		TotalPoses += NumIndexedPoses;
 	}
 
 	if (Owner.IsCanceled())
@@ -203,7 +199,6 @@ bool FDatabaseIndexingContext::IndexDatabase(FPoseSearchIndexBase& SearchIndexBa
 	}
 
 	// Index asset data
-	////////////////////////////////////////////////////////////////////////////////////////////////////////
 	ParallelFor(Indexers.Num(), [this](int32 AssetIdx) { Indexers[AssetIdx].Process(AssetIdx); }, ParallelForFlags);
 
 	if (Owner.IsCanceled())
@@ -212,7 +207,6 @@ bool FDatabaseIndexingContext::IndexDatabase(FPoseSearchIndexBase& SearchIndexBa
 	}
 
 	// Joining Metadata.Flags into OverallFlags
-	////////////////////////////////////////////////////////////////////////////////////////////////////////
 	SearchIndexBase.OverallFlags = EPoseSearchPoseFlags::None;
 	for (const FPoseSearchPoseMetadata& Metadata : SearchIndexBase.PoseMetadata)
 	{
@@ -220,7 +214,6 @@ bool FDatabaseIndexingContext::IndexDatabase(FPoseSearchIndexBase& SearchIndexBa
 	}
 
 	// Joining Stats
-	////////////////////////////////////////////////////////////////////////////////////////////////////////
 	int32 NumAccumulatedSamples = 0;
 	SearchIndexBase.Stats = FPoseSearchStats();
 	for (int32 AssetIdx = 0; AssetIdx != SearchIndexBase.Assets.Num(); ++AssetIdx)
@@ -241,12 +234,8 @@ bool FDatabaseIndexingContext::IndexDatabase(FPoseSearchIndexBase& SearchIndexBa
 		SearchIndexBase.Stats.AverageAcceleration *= Denom;
 	}
 
-	SearchIndexBase.NumPoses = TotalPoses;
-
 	// Calculate Min Cost Addend
-	////////////////////////////////////////////////////////////////////////////////////////////////////////
 	SearchIndexBase.MinCostAddend = 0.f;
-
 	if (!SearchIndexBase.PoseMetadata.IsEmpty())
 	{
 		SearchIndexBase.MinCostAddend = MAX_FLT;
diff --git a/Engine/Plugins/Experimental/Animation/PoseSearch/Source/Runtime/Private/PoseSearchDatabaseSet.cpp b/Engine/Plugins/Experimental/Animation/PoseSearch/Source/Runtime/Private/PoseSearchDatabaseSet.cpp
--- a/Engine/Plugins/Experimental/Animation/PoseSearch/Source/Runtime/Private/PoseSearchDatabaseSet.cpp
+++ b/Engine/Plugins/Experimental/Animation/PoseSearch/Source/Runtime/Private/PoseSearchDatabaseSet.cpp
@@ -38,7 +38,7 @@ UE::PoseSearch::FSearchResult UPoseSearchDatabaseSet::Search(UE::PoseSearch::FSe
 		Result.ContinuingPoseCost = Result.PoseCost;
 		ContinuingCost = Result.PoseCost;
 
-		Result.AssetTime = Database->GetAssetTime(Result.PoseIdx, Database->Schema->GetSamplingInterval());
+		Result.AssetTime = Database->GetAssetTime(Result.PoseIdx);
 		Result.Database = Database;
 
 		if (Database->GetSkipSearchIfPossible())
diff --git a/Engine/Plugins/Experimental/Animation/PoseSearch/Source/Runtime/Private/PoseSearchDerivedData.cpp b/Engine/Plugins/Experimental/Animation/PoseSearch/Source/Runtime/Private/PoseSearchDerivedData.cpp
--- a/Engine/Plugins/Experimental/Animation/PoseSearch/Source/Runtime/Private/PoseSearchDerivedData.cpp
+++ b/Engine/Plugins/Experimental/Animation/PoseSearch/Source/Runtime/Private/PoseSearchDerivedData.cpp
@@ -147,7 +147,7 @@ private:
 		int32 TotalNumPoses = 0;
 		for (int32 EntryIdx = 0; EntryIdx < EntriesNum; ++EntryIdx)
 		{
-			TotalNumPoses += SearchIndexBases[Entries[EntryIdx].SchemaIndex].NumPoses;
+			TotalNumPoses += SearchIndexBases[Entries[EntryIdx].SchemaIndex].GetNumPoses();
 		}
 
 		int32 AccumulatedNumPoses = 0;
@@ -162,7 +162,7 @@ private:
 			const UPoseSearchSchema* Schema = Schemas[DataSetIdx];
 			const FPoseSearchIndexBase& SearchIndex = SearchIndexBases[DataSetIdx];
 
-			const int32 NumPoses = SearchIndex.NumPoses;
+			const int32 NumPoses = SearchIndex.GetNumPoses();
 
 			// Map input buffer with NumPoses as rows and NumDimensions	as cols
 			RowMajorMatrixMapConst PoseMatrixSourceMap(SearchIndex.Values.GetData(), NumPoses, Schema->SchemaCardinality);
@@ -212,7 +212,7 @@ public:
 		RowMajorVectorMap MeanDeviationsMap(MeanDeviations.GetData(), 1, NumDimensions);
 
 		const EPoseSearchDataPreprocessor DataPreprocessor = ThisSchema->DataPreprocessor;
-		if (SearchIndexBases[ThisSchemaIndex].NumPoses > 0 && (DataPreprocessor == EPoseSearchDataPreprocessor::Normalize || DataPreprocessor == EPoseSearchDataPreprocessor::NormalizeOnlyByDeviation))
+		if (SearchIndexBases[ThisSchemaIndex].GetNumPoses() > 0 && (DataPreprocessor == EPoseSearchDataPreprocessor::Normalize || DataPreprocessor == EPoseSearchDataPreprocessor::NormalizeOnlyByDeviation))
 		{
 			FEntriesGroup EntriesGroup;
 
@@ -450,7 +450,7 @@ static void PreprocessSearchIndexWeights(FPoseSearchIndex& SearchIndex, const UP
 static void PreprocessSearchIndexPCAData(FPoseSearchIndex& SearchIndex, int32 NumDimensions, uint32 NumberOfPrincipalComponents, EPoseSearchMode PoseSearchMode)
 {
 	// binding SearchIndex.Values and SearchIndex.PCAValues Eigen row major matrix maps
-	const int32 NumPoses = SearchIndex.NumPoses;
+	const int32 NumPoses = SearchIndex.GetNumPoses();
 
 	SearchIndex.PCAExplainedVariance = 0.f;
 
@@ -548,7 +548,7 @@ static void PreprocessSearchIndexKDTree(FPoseSearchIndex& SearchIndex, int32 Num
 	SearchIndex.KDTree.Reset();
 	if (PoseSearchMode != EPoseSearchMode::BruteForce && NumDimensions > 0)
 	{
-		const int32 NumPoses = SearchIndex.NumPoses;
+		const int32 NumPoses = SearchIndex.GetNumPoses();
 		SearchIndex.KDTree.Construct(NumPoses, NumberOfPrincipalComponents, SearchIndex.PCAValues.GetData(), KDTreeMaxLeafSize);
 
 		if (PoseSearchMode == EPoseSearchMode::PCAKDTree_Validate)
diff --git a/Engine/Plugins/Experimental/Animation/PoseSearch/Source/Runtime/Private/PoseSearchDerivedDataKey.cpp b/Engine/Plugins/Experimental/Animation/PoseSearch/Source/Runtime/Private/PoseSearchDerivedDataKey.cpp
--- a/Engine/Plugins/Experimental/Animation/PoseSearch/Source/Runtime/Private/PoseSearchDerivedDataKey.cpp
+++ b/Engine/Plugins/Experimental/Animation/PoseSearch/Source/Runtime/Private/PoseSearchDerivedDataKey.cpp
@@ -41,7 +41,7 @@ FKeyBuilder::FKeyBuilder(const UObject* Object, bool bUseDataVer)
 	if (bUseDataVer)
 	{
 		// used to invalidate the key without having to change POSESEARCHDB_DERIVEDDATA_VER all the times
-		int32 POSESEARCHDB_DERIVEDDATA_VER_SMALL = 79;
+		int32 POSESEARCHDB_DERIVEDDATA_VER_SMALL = 107;
 		FGuid VersionGuid = FDevSystemGuids::GetSystemGuid(FDevSystemGuids::Get().POSESEARCHDB_DERIVEDDATA_VER);
 
 		*this << VersionGuid;
diff --git a/Engine/Plugins/Experimental/Animation/PoseSearch/Source/Runtime/Private/PoseSearchFeatureChannel_FilterCrashingLegs.cpp b/Engine/Plugins/Experimental/Animation/PoseSearch/Source/Runtime/Private/PoseSearchFeatureChannel_FilterCrashingLegs.cpp
--- a/Engine/Plugins/Experimental/Animation/PoseSearch/Source/Runtime/Private/PoseSearchFeatureChannel_FilterCrashingLegs.cpp
+++ b/Engine/Plugins/Experimental/Animation/PoseSearch/Source/Runtime/Private/PoseSearchFeatureChannel_FilterCrashingLegs.cpp
@@ -47,7 +47,7 @@ void UPoseSearchFeatureChannel_FilterCrashingLegs::FillWeights(TArray<float>& We
 	}
 }
 
-void UPoseSearchFeatureChannel_FilterCrashingLegs::IndexAsset(UE::PoseSearch::FAssetIndexer& Indexer, TArrayView<float> FeatureVectorTable) const
+void UPoseSearchFeatureChannel_FilterCrashingLegs::IndexAsset(UE::PoseSearch::FAssetIndexer& Indexer) const
 {
 	using namespace UE::PoseSearch;
 
@@ -60,7 +60,7 @@ void UPoseSearchFeatureChannel_FilterCrashingLegs::IndexAsset(UE::PoseSearch::FA
 
 		const float CrashingLegsValue = ComputeCrashingLegsValue(RightThighPosition, LeftThighPosition, RightFootPosition, LeftFootPosition);
 
-		FFeatureVectorHelper::EncodeFloat(Indexer.GetPoseVector(SampleIdx, FeatureVectorTable), ChannelDataOffset, CrashingLegsValue);
+		FFeatureVectorHelper::EncodeFloat(Indexer.GetPoseVector(SampleIdx), ChannelDataOffset, CrashingLegsValue);
 	}
 }
 #endif // WITH_EDITOR
diff --git a/Engine/Plugins/Experimental/Animation/PoseSearch/Source/Runtime/Private/PoseSearchFeatureChannel_FilterCrashingLegs.h b/Engine/Plugins/Experimental/Animation/PoseSearch/Source/Runtime/Private/PoseSearchFeatureChannel_FilterCrashingLegs.h
--- a/Engine/Plugins/Experimental/Animation/PoseSearch/Source/Runtime/Private/PoseSearchFeatureChannel_FilterCrashingLegs.h
+++ b/Engine/Plugins/Experimental/Animation/PoseSearch/Source/Runtime/Private/PoseSearchFeatureChannel_FilterCrashingLegs.h
@@ -62,7 +62,7 @@ public:
 
 #if WITH_EDITOR
 	virtual void FillWeights(TArray<float>& Weights) const override;
-	virtual void IndexAsset(UE::PoseSearch::FAssetIndexer& Indexer, TArrayView<float> FeatureVectorTable) const override;
+	virtual void IndexAsset(UE::PoseSearch::FAssetIndexer& Indexer) const override;
 	virtual FString GetLabel() const;
 #endif
 
diff --git a/Engine/Plugins/Experimental/Animation/PoseSearch/Source/Runtime/Private/PoseSearchFeatureChannel_Group.cpp b/Engine/Plugins/Experimental/Animation/PoseSearch/Source/Runtime/Private/PoseSearchFeatureChannel_Group.cpp
--- a/Engine/Plugins/Experimental/Animation/PoseSearch/Source/Runtime/Private/PoseSearchFeatureChannel_Group.cpp
+++ b/Engine/Plugins/Experimental/Animation/PoseSearch/Source/Runtime/Private/PoseSearchFeatureChannel_Group.cpp
@@ -28,13 +28,13 @@ void UPoseSearchFeatureChannel_GroupBase::FillWeights(TArray<float>& Weights) co
 	}
 }
 
-void UPoseSearchFeatureChannel_GroupBase::IndexAsset(UE::PoseSearch::FAssetIndexer& Indexer, TArrayView<float> FeatureVectorTable) const
+void UPoseSearchFeatureChannel_GroupBase::IndexAsset(UE::PoseSearch::FAssetIndexer& Indexer) const
 {
 	for (const TObjectPtr<UPoseSearchFeatureChannel>& SubChannelPtr : GetSubChannels())
 	{
 		if (const UPoseSearchFeatureChannel* SubChannel = SubChannelPtr.Get())
 		{
-			SubChannel->IndexAsset(Indexer, FeatureVectorTable);
+			SubChannel->IndexAsset(Indexer);
 		}
 	}
 }
diff --git a/Engine/Plugins/Experimental/Animation/PoseSearch/Source/Runtime/Private/PoseSearchFeatureChannel_Group.h b/Engine/Plugins/Experimental/Animation/PoseSearch/Source/Runtime/Private/PoseSearchFeatureChannel_Group.h
--- a/Engine/Plugins/Experimental/Animation/PoseSearch/Source/Runtime/Private/PoseSearchFeatureChannel_Group.h
+++ b/Engine/Plugins/Experimental/Animation/PoseSearch/Source/Runtime/Private/PoseSearchFeatureChannel_Group.h
@@ -24,7 +24,7 @@ public:
 
 #if WITH_EDITOR
 	virtual void FillWeights(TArray<float>& Weights) const override;
-	virtual void IndexAsset(UE::PoseSearch::FAssetIndexer& Indexer, TArrayView<float> FeatureVectorTable) const override;
+	virtual void IndexAsset(UE::PoseSearch::FAssetIndexer& Indexer) const override;
 #endif //WITH_EDITOR
 
 	// IPoseFilter interface
diff --git a/Engine/Plugins/Experimental/Animation/PoseSearch/Source/Runtime/Private/PoseSearchFeatureChannel_Heading.cpp b/Engine/Plugins/Experimental/Animation/PoseSearch/Source/Runtime/Private/PoseSearchFeatureChannel_Heading.cpp
--- a/Engine/Plugins/Experimental/Animation/PoseSearch/Source/Runtime/Private/PoseSearchFeatureChannel_Heading.cpp
+++ b/Engine/Plugins/Experimental/Animation/PoseSearch/Source/Runtime/Private/PoseSearchFeatureChannel_Heading.cpp
@@ -41,14 +41,14 @@ void UPoseSearchFeatureChannel_Heading::FillWeights(TArray<float>& Weights) cons
 	}
 }
 
-void UPoseSearchFeatureChannel_Heading::IndexAsset(UE::PoseSearch::FAssetIndexer& Indexer, TArrayView<float> FeatureVectorTable) const
+void UPoseSearchFeatureChannel_Heading::IndexAsset(UE::PoseSearch::FAssetIndexer& Indexer) const
 {
 	using namespace UE::PoseSearch;
 
 	for (int32 SampleIdx = Indexer.GetBeginSampleIdx(); SampleIdx != Indexer.GetEndSampleIdx(); ++SampleIdx)
 	{
 		const FVector Heading = GetAxis(Indexer.GetSampleRotation(SampleTimeOffset, SampleIdx, SchemaBoneIdx));
-		FFeatureVectorHelper::EncodeVector(Indexer.GetPoseVector(SampleIdx, FeatureVectorTable), ChannelDataOffset, Heading, ComponentStripping);
+		FFeatureVectorHelper::EncodeVector(Indexer.GetPoseVector(SampleIdx), ChannelDataOffset, Heading, ComponentStripping);
 	}
 }
 #endif // WITH_EDITOR
diff --git a/Engine/Plugins/Experimental/Animation/PoseSearch/Source/Runtime/Private/PoseSearchFeatureChannel_Heading.h b/Engine/Plugins/Experimental/Animation/PoseSearch/Source/Runtime/Private/PoseSearchFeatureChannel_Heading.h
--- a/Engine/Plugins/Experimental/Animation/PoseSearch/Source/Runtime/Private/PoseSearchFeatureChannel_Heading.h
+++ b/Engine/Plugins/Experimental/Animation/PoseSearch/Source/Runtime/Private/PoseSearchFeatureChannel_Heading.h
@@ -57,7 +57,7 @@ public:
 
 #if WITH_EDITOR
 	virtual void FillWeights(TArray<float>& Weights) const override;
-	virtual void IndexAsset(UE::PoseSearch::FAssetIndexer& Indexer, TArrayView<float> FeatureVectorTable) const override;
+	virtual void IndexAsset(UE::PoseSearch::FAssetIndexer& Indexer) const override;
 	virtual FString GetLabel() const;
 #endif
 
diff --git a/Engine/Plugins/Experimental/Animation/PoseSearch/Source/Runtime/Private/PoseSearchFeatureChannel_Phase.cpp b/Engine/Plugins/Experimental/Animation/PoseSearch/Source/Runtime/Private/PoseSearchFeatureChannel_Phase.cpp
--- a/Engine/Plugins/Experimental/Animation/PoseSearch/Source/Runtime/Private/PoseSearchFeatureChannel_Phase.cpp
+++ b/Engine/Plugins/Experimental/Animation/PoseSearch/Source/Runtime/Private/PoseSearchFeatureChannel_Phase.cpp
@@ -296,7 +296,7 @@ void UPoseSearchFeatureChannel_Phase::FillWeights(TArray<float>& Weights) const
 	}
 }
 
-void UPoseSearchFeatureChannel_Phase::IndexAsset(UE::PoseSearch::FAssetIndexer& Indexer, TArrayView<float> FeatureVectorTable) const
+void UPoseSearchFeatureChannel_Phase::IndexAsset(UE::PoseSearch::FAssetIndexer& Indexer) const
 {
 	using namespace UE::PoseSearch;
 
@@ -333,7 +333,7 @@ void UPoseSearchFeatureChannel_Phase::IndexAsset(UE::PoseSearch::FAssetIndexer&
 
 	for (int32 SampleIdx = Indexer.GetBeginSampleIdx(); SampleIdx != Indexer.GetEndSampleIdx(); ++SampleIdx)
 	{
-		FFeatureVectorHelper::EncodeVector2D(Indexer.GetPoseVector(SampleIdx, FeatureVectorTable), ChannelDataOffset, Phases[SampleIdx - Indexer.GetBeginSampleIdx()]);
+		FFeatureVectorHelper::EncodeVector2D(Indexer.GetPoseVector(SampleIdx), ChannelDataOffset, Phases[SampleIdx - Indexer.GetBeginSampleIdx()]);
 	}
 }
 #endif // WITH_EDITOR
diff --git a/Engine/Plugins/Experimental/Animation/PoseSearch/Source/Runtime/Private/PoseSearchFeatureChannel_Phase.h b/Engine/Plugins/Experimental/Animation/PoseSearch/Source/Runtime/Private/PoseSearchFeatureChannel_Phase.h
--- a/Engine/Plugins/Experimental/Animation/PoseSearch/Source/Runtime/Private/PoseSearchFeatureChannel_Phase.h
+++ b/Engine/Plugins/Experimental/Animation/PoseSearch/Source/Runtime/Private/PoseSearchFeatureChannel_Phase.h
@@ -37,7 +37,7 @@ public:
 
 #if WITH_EDITOR
 	virtual void FillWeights(TArray<float>& Weights) const override;
-	virtual void IndexAsset(UE::PoseSearch::FAssetIndexer& Indexer, TArrayView<float> FeatureVectorTable) const override;
+	virtual void IndexAsset(UE::PoseSearch::FAssetIndexer& Indexer) const override;
 	virtual FString GetLabel() const;
 #endif
 };
\ No newline at end of file
diff --git a/Engine/Plugins/Experimental/Animation/PoseSearch/Source/Runtime/Private/PoseSearchFeatureChannel_Position.cpp b/Engine/Plugins/Experimental/Animation/PoseSearch/Source/Runtime/Private/PoseSearchFeatureChannel_Position.cpp
--- a/Engine/Plugins/Experimental/Animation/PoseSearch/Source/Runtime/Private/PoseSearchFeatureChannel_Position.cpp
+++ b/Engine/Plugins/Experimental/Animation/PoseSearch/Source/Runtime/Private/PoseSearchFeatureChannel_Position.cpp
@@ -27,14 +27,14 @@ void UPoseSearchFeatureChannel_Position::FillWeights(TArray<float>& Weights) con
 	}
 }
 
-void UPoseSearchFeatureChannel_Position::IndexAsset(UE::PoseSearch::FAssetIndexer& Indexer, TArrayView<float> FeatureVectorTable) const
+void UPoseSearchFeatureChannel_Position::IndexAsset(UE::PoseSearch::FAssetIndexer& Indexer) const
 {
 	using namespace UE::PoseSearch;
 
 	for (int32 SampleIdx = Indexer.GetBeginSampleIdx(); SampleIdx != Indexer.GetEndSampleIdx(); ++SampleIdx)
 	{
 		const FVector BonePosition = Indexer.GetSamplePosition(SampleTimeOffset, SampleIdx, SchemaBoneIdx, SchemaOriginBoneIdx);
-		FFeatureVectorHelper::EncodeVector(Indexer.GetPoseVector(SampleIdx, FeatureVectorTable), ChannelDataOffset, BonePosition, ComponentStripping);
+		FFeatureVectorHelper::EncodeVector(Indexer.GetPoseVector(SampleIdx), ChannelDataOffset, BonePosition, ComponentStripping);
 	}
 }
 #endif // WITH_EDITOR
diff --git a/Engine/Plugins/Experimental/Animation/PoseSearch/Source/Runtime/Private/PoseSearchFeatureChannel_Position.h b/Engine/Plugins/Experimental/Animation/PoseSearch/Source/Runtime/Private/PoseSearchFeatureChannel_Position.h
--- a/Engine/Plugins/Experimental/Animation/PoseSearch/Source/Runtime/Private/PoseSearchFeatureChannel_Position.h
+++ b/Engine/Plugins/Experimental/Animation/PoseSearch/Source/Runtime/Private/PoseSearchFeatureChannel_Position.h
@@ -50,7 +50,7 @@ public:
 
 #if WITH_EDITOR
 	virtual void FillWeights(TArray<float>& Weights) const override;
-	virtual void IndexAsset(UE::PoseSearch::FAssetIndexer& Indexer, TArrayView<float> FeatureVectorTable) const override;
+	virtual void IndexAsset(UE::PoseSearch::FAssetIndexer& Indexer) const override;
 	virtual FString GetLabel() const;
 #endif
 };
diff --git a/Engine/Plugins/Experimental/Animation/PoseSearch/Source/Runtime/Private/PoseSearchFeatureChannel_Velocity.cpp b/Engine/Plugins/Experimental/Animation/PoseSearch/Source/Runtime/Private/PoseSearchFeatureChannel_Velocity.cpp
--- a/Engine/Plugins/Experimental/Animation/PoseSearch/Source/Runtime/Private/PoseSearchFeatureChannel_Velocity.cpp
+++ b/Engine/Plugins/Experimental/Animation/PoseSearch/Source/Runtime/Private/PoseSearchFeatureChannel_Velocity.cpp
@@ -27,7 +27,7 @@ void UPoseSearchFeatureChannel_Velocity::FillWeights(TArray<float>& Weights) con
 	}
 }
 
-void UPoseSearchFeatureChannel_Velocity::IndexAsset(UE::PoseSearch::FAssetIndexer& Indexer, TArrayView<float> FeatureVectorTable) const
+void UPoseSearchFeatureChannel_Velocity::IndexAsset(UE::PoseSearch::FAssetIndexer& Indexer) const
 {
 	using namespace UE::PoseSearch;
 
@@ -38,7 +38,7 @@ void UPoseSearchFeatureChannel_Velocity::IndexAsset(UE::PoseSearch::FAssetIndexe
 		{
 			LinearVelocity = LinearVelocity.GetClampedToMaxSize(1.f);
 		}
-		FFeatureVectorHelper::EncodeVector(Indexer.GetPoseVector(SampleIdx, FeatureVectorTable), ChannelDataOffset, LinearVelocity, ComponentStripping);
+		FFeatureVectorHelper::EncodeVector(Indexer.GetPoseVector(SampleIdx), ChannelDataOffset, LinearVelocity, ComponentStripping);
 	}
 }
 #endif // WITH_EDITOR
diff --git a/Engine/Plugins/Experimental/Animation/PoseSearch/Source/Runtime/Private/PoseSearchFeatureChannel_Velocity.h b/Engine/Plugins/Experimental/Animation/PoseSearch/Source/Runtime/Private/PoseSearchFeatureChannel_Velocity.h
--- a/Engine/Plugins/Experimental/Animation/PoseSearch/Source/Runtime/Private/PoseSearchFeatureChannel_Velocity.h
+++ b/Engine/Plugins/Experimental/Animation/PoseSearch/Source/Runtime/Private/PoseSearchFeatureChannel_Velocity.h
@@ -51,7 +51,7 @@ public:
 
 #if WITH_EDITOR
 	virtual void FillWeights(TArray<float>& Weights) const override;
-	virtual void IndexAsset(UE::PoseSearch::FAssetIndexer& Indexer, TArrayView<float> FeatureVectorTable) const override;
+	virtual void IndexAsset(UE::PoseSearch::FAssetIndexer& Indexer) const override;
 	virtual FString GetLabel() const;
 #endif
 };
diff --git a/Engine/Plugins/Experimental/Animation/PoseSearch/Source/Runtime/Private/PoseSearchIndex.cpp b/Engine/Plugins/Experimental/Animation/PoseSearch/Source/Runtime/Private/PoseSearchIndex.cpp
--- a/Engine/Plugins/Experimental/Animation/PoseSearch/Source/Runtime/Private/PoseSearchIndex.cpp
+++ b/Engine/Plugins/Experimental/Animation/PoseSearch/Source/Runtime/Private/PoseSearchIndex.cpp
@@ -6,7 +6,7 @@
 namespace UE::PoseSearch
 {
 
-static inline float CompareFeatureVectors(TConstArrayView<float> A, TConstArrayView<float> B, TConstArrayView<float> WeightsSqrt)
+float CompareFeatureVectors(TConstArrayView<float> A, TConstArrayView<float> B, TConstArrayView<float> WeightsSqrt)
 {
 	check(A.Num() == B.Num() && A.Num() == WeightsSqrt.Num());
 
@@ -89,8 +89,7 @@ const FPoseSearchIndexAsset* FPoseSearchIndexBase::GetAssetForPoseSafe(int32 Pos
 
 bool FPoseSearchIndexBase::IsEmpty() const
 {
-	const bool bEmpty = Assets.Num() == 0 || NumPoses == 0;
-	return bEmpty;
+	return Assets.IsEmpty() || PoseMetadata.IsEmpty();
 }
 
 void FPoseSearchIndexBase::Reset()
@@ -100,7 +99,6 @@ void FPoseSearchIndexBase::Reset()
 
 FArchive& operator<<(FArchive& Ar, FPoseSearchIndexBase& Index)
 {
-	Ar << Index.NumPoses;
 	Ar << Index.Values;
 	Ar << Index.PoseMetadata;
 	Ar << Index.OverallFlags;
@@ -144,18 +142,16 @@ void FPoseSearchIndex::Reset()
 TConstArrayView<float> FPoseSearchIndex::GetPoseValues(int32 PoseIdx) const
 {
 	const int32 SchemaCardinality = WeightsSqrt.Num();
-	check(PoseIdx >= 0 && PoseIdx < NumPoses&& SchemaCardinality > 0);
+	check(PoseIdx >= 0 && PoseIdx < GetNumPoses() && SchemaCardinality > 0);
 	const int32 ValueOffset = PoseIdx * SchemaCardinality;
 	return MakeArrayView(&Values[ValueOffset], SchemaCardinality);
 }
 
 TConstArrayView<float> FPoseSearchIndex::GetPoseValuesSafe(int32 PoseIdx) const
 {
-	if (PoseIdx >= 0 && PoseIdx < NumPoses)
+	if (PoseIdx >= 0 && PoseIdx < GetNumPoses())
 	{
-		const int32 SchemaCardinality = WeightsSqrt.Num();
-		const int32 ValueOffset = PoseIdx * SchemaCardinality;
-		return MakeArrayView(&Values[ValueOffset], SchemaCardinality);
+		return GetPoseValues(PoseIdx);
 	}
 	return TConstArrayView<float>();
 }
diff --git a/Engine/Plugins/Experimental/Animation/PoseSearch/Source/Runtime/Public/PoseSearch/PoseSearchAssetIndexer.h b/Engine/Plugins/Experimental/Animation/PoseSearch/Source/Runtime/Public/PoseSearch/PoseSearchAssetIndexer.h
--- a/Engine/Plugins/Experimental/Animation/PoseSearch/Source/Runtime/Public/PoseSearch/PoseSearchAssetIndexer.h
+++ b/Engine/Plugins/Experimental/Animation/PoseSearch/Source/Runtime/Public/PoseSearch/PoseSearchAssetIndexer.h
@@ -9,6 +9,7 @@
 #include "PoseSearch/PoseSearchDefines.h"
 #include "PoseSearch/PoseSearchSchema.h"
 
+struct FPoseSearchIndexAsset;
 struct FPoseSearchPoseMetadata;
 
 namespace UE::PoseSearch
@@ -32,16 +33,6 @@ struct FAssetIndexingContext
 class POSESEARCH_API FAssetIndexer
 {
 public:
-	struct FOutput
-	{
-		int32 FirstIndexedSample = 0;
-		int32 LastIndexedSample = 0;
-		int32 NumIndexedPoses = 0;
-
-		TArrayView<float> FeatureVectorTable;
-		TArrayView<FPoseSearchPoseMetadata> PoseMetadata;
-	};
-
 	struct FStats
 	{
 		int32 NumAccumulatedSamples = 0;
@@ -51,21 +42,20 @@ public:
 		float MaxAcceleration = 0.f;
 	};
 
-	void Reset();
-	void Init(const FAssetIndexingContext& IndexingContext, const FBoneContainer& InBoneContainer);
+	FAssetIndexer(const FAssetIndexingContext& IndexingContext, const FBoneContainer& InBoneContainer, const FPoseSearchIndexAsset& InSearchIndexAsset);
+	void AssignWorkingData(TArrayView<float> InOutFeatureVectorTable, TArrayView<FPoseSearchPoseMetadata> InOutPoseMetadata);
 	void Process(int32 AssetIdx);
-	const FOutput& GetOutput() const { return Output; }
-	FOutput& EditOutput() { return Output; }
 	const FStats& GetStats() const { return Stats; }
 
 	FQuat GetSampleRotation(float SampleTimeOffset, int32 SampleIdx, int8 SchemaSampleBoneIdx = RootSchemaBoneIdx, int8 SchemaOriginBoneIdx = RootSchemaBoneIdx);
 	FVector GetSamplePosition(float SampleTimeOffset, int32 SampleIdx, int8 SchemaSampleBoneIdx = RootSchemaBoneIdx, int8 SchemaOriginBoneIdx = RootSchemaBoneIdx);
 	FVector GetSampleVelocity(float SampleTimeOffset, int32 SampleIdx, int8 SchemaSampleBoneIdx = RootSchemaBoneIdx, int8 SchemaOriginBoneIdx = RootSchemaBoneIdx, bool bUseCharacterSpaceVelocities = true);
 
-	int32 GetBeginSampleIdx() const { return Output.FirstIndexedSample; }
-	int32 GetEndSampleIdx() const { return Output.LastIndexedSample + 1; }
+	int32 GetBeginSampleIdx() const { return FirstIndexedSample; }
+	int32 GetEndSampleIdx() const { return LastIndexedSample + 1; }
+	int32 GetNumIndexedPoses() const { return GetEndSampleIdx() - GetBeginSampleIdx(); }
 	
-	TArrayView<float> GetPoseVector(int32 SampleIdx, TArrayView<float> FeatureVectorTable) const;
+	TArrayView<float> GetPoseVector(int32 SampleIdx) const;
 	const UPoseSearchSchema* GetSchema() const;
 
 private:
@@ -92,7 +82,6 @@ private:
 	};
 
 	FSampleInfo GetSampleInfo(float SampleTime) const;
-	FPoseSearchPoseMetadata GetMetadata(int32 SampleIdx, int32 AssetIdx) const;
 	FTransform MirrorTransform(const FTransform& Transform) const;
 	CachedEntry& GetEntry(float SampleTime);
 	FTransform CalculateComponentSpaceTransform(CachedEntry& Entry, int8 SchemaBoneIdx);
@@ -101,7 +90,14 @@ private:
 	FBoneContainer BoneContainer;
 	FAssetIndexingContext IndexingContext;
 	TMap<float, CachedEntry> CachedEntries;
-	FOutput Output;
+	const FPoseSearchIndexAsset& SearchIndexAsset;
+	
+	int32 FirstIndexedSample = 0;
+	int32 LastIndexedSample = 0;
+
+	TArrayView<float> FeatureVectorTable;
+	TArrayView<FPoseSearchPoseMetadata> PoseMetadata;
+
 	FStats Stats;
 };
 
diff --git a/Engine/Plugins/Experimental/Animation/PoseSearch/Source/Runtime/Public/PoseSearch/PoseSearchDatabase.h b/Engine/Plugins/Experimental/Animation/PoseSearch/Source/Runtime/Public/PoseSearch/PoseSearchDatabase.h
--- a/Engine/Plugins/Experimental/Animation/PoseSearch/Source/Runtime/Public/PoseSearch/PoseSearchDatabase.h
+++ b/Engine/Plugins/Experimental/Animation/PoseSearch/Source/Runtime/Public/PoseSearch/PoseSearchDatabase.h
@@ -279,7 +279,7 @@ public:
 	const bool IsSourceAssetLooping(const FPoseSearchIndexAsset& SearchIndexAsset) const;
 	const FString GetSourceAssetName(const FPoseSearchIndexAsset& SearchIndexAsset) const;
 	int32 GetNumberOfPrincipalComponents() const;
-	float GetAssetTime(int32 PoseIdx, float SamplingInterval) const;
+	float GetAssetTime(int32 PoseIdx) const;
 
 	// Begin UObject
 	virtual void PostLoad() override;
diff --git a/Engine/Plugins/Experimental/Animation/PoseSearch/Source/Runtime/Public/PoseSearch/PoseSearchFeatureChannel.h b/Engine/Plugins/Experimental/Animation/PoseSearch/Source/Runtime/Public/PoseSearch/PoseSearchFeatureChannel.h
--- a/Engine/Plugins/Experimental/Animation/PoseSearch/Source/Runtime/Public/PoseSearch/PoseSearchFeatureChannel.h
+++ b/Engine/Plugins/Experimental/Animation/PoseSearch/Source/Runtime/Public/PoseSearch/PoseSearchFeatureChannel.h
@@ -115,7 +115,7 @@ public:
 	virtual void FillWeights(TArray<float>& Weights) const PURE_VIRTUAL(UPoseSearchFeatureChannel::FillWeights, );
 
 	// Called at database build time to populate pose vectors with this channel's data
-	virtual void IndexAsset(UE::PoseSearch::FAssetIndexer& Indexer, TArrayView<float> FeatureVectorTable) const PURE_VIRTUAL(UPoseSearchFeatureChannel::IndexAsset, );
+	virtual void IndexAsset(UE::PoseSearch::FAssetIndexer& Indexer) const PURE_VIRTUAL(UPoseSearchFeatureChannel::IndexAsset, );
 
 	// returns the FString used editor side to identify this UPoseSearchFeatureChannel (for instance in the pose search debugger)
 	virtual FString GetLabel() const;
diff --git a/Engine/Plugins/Experimental/Animation/PoseSearch/Source/Runtime/Public/PoseSearch/PoseSearchIndex.h b/Engine/Plugins/Experimental/Animation/PoseSearch/Source/Runtime/Public/PoseSearch/PoseSearchIndex.h
--- a/Engine/Plugins/Experimental/Animation/PoseSearch/Source/Runtime/Public/PoseSearch/PoseSearchIndex.h
+++ b/Engine/Plugins/Experimental/Animation/PoseSearch/Source/Runtime/Public/PoseSearch/PoseSearchIndex.h
@@ -9,6 +9,7 @@
 namespace UE::PoseSearch
 {
 
+float CompareFeatureVectors(TConstArrayView<float> A, TConstArrayView<float> B, TConstArrayView<float> WeightsSqrt);
 POSESEARCH_API void CompareFeatureVectors(TConstArrayView<float> A, TConstArrayView<float> B, TConstArrayView<float> WeightsSqrt, TArrayView<float> Result);
 
 enum class EPoseComparisonFlags : int32
@@ -32,7 +33,7 @@ enum class EPoseSearchBooleanRequest : uint8
 };
 
 UENUM()
-enum class EPoseSearchPoseFlags : uint32
+enum class EPoseSearchPoseFlags : uint8
 {
 	None = 0,
 
@@ -60,10 +61,12 @@ struct POSESEARCH_API FPoseSearchPoseMetadata
 	UPROPERTY(meta = (NeverInHash))
 	float ContinuingPoseCostAddend = 0.0f;
 
-	// @todo: consider int16
+	// @todo: consider uint16
 	UPROPERTY(meta = (NeverInHash))
 	int32 AssetIndex = INDEX_NONE;
 
+	bool IsInitialized() const { return AssetIndex != INDEX_NONE; }
+
 	friend FArchive& operator<<(FArchive& Ar, FPoseSearchPoseMetadata& Metadata);
 };
 
@@ -79,14 +82,18 @@ struct POSESEARCH_API FPoseSearchIndexAsset
 	FPoseSearchIndexAsset() {}
 
 	FPoseSearchIndexAsset(
-		int32 InSourceAssetIdx, 
+		int32 InSourceAssetIdx,
 		bool bInMirrored, 
 		const FFloatInterval& InSamplingInterval,
 		FVector InBlendParameters = FVector::Zero())
 		: SourceAssetIdx(InSourceAssetIdx)
 		, bMirrored(bInMirrored)
 		, BlendParameters(InBlendParameters)
-		, SamplingInterval(InSamplingInterval) {}
+		, SamplingInterval(InSamplingInterval)
+		, FirstPoseIdx(INDEX_NONE)
+		, NumPoses(INDEX_NONE)
+	{
+	}
 
 	// Index of the source asset in search index's container (i.e. UPoseSearchDatabase)
 	UPROPERTY(meta = (NeverInHash))
@@ -105,7 +112,7 @@ struct POSESEARCH_API FPoseSearchIndexAsset
 	int32 FirstPoseIdx = INDEX_NONE;
 
 	UPROPERTY(meta = (NeverInHash))
-	int32 NumPoses = 0;
+	int32 NumPoses = INDEX_NONE;
 
 	bool IsPoseInRange(int32 PoseIdx) const { return (PoseIdx >= FirstPoseIdx) && (PoseIdx < FirstPoseIdx + NumPoses); }
 	friend FArchive& operator<<(FArchive& Ar, FPoseSearchIndexAsset& IndexAsset);
@@ -139,9 +146,6 @@ struct POSESEARCH_API FPoseSearchIndexBase
 {
 	GENERATED_BODY()
 
-	UPROPERTY(meta = (NeverInHash))
-	int32 NumPoses = 0;
-
 	UPROPERTY(meta = (NeverInHash))
 	TArray<float> Values;
 	
@@ -162,7 +166,8 @@ struct POSESEARCH_API FPoseSearchIndexBase
 	UPROPERTY(meta = (NeverInHash))
 	FPoseSearchStats Stats;
 
-	bool IsValidPoseIndex(int32 PoseIdx) const { return PoseIdx < NumPoses; }
+	int32 GetNumPoses() const { return PoseMetadata.Num(); }
+	bool IsValidPoseIndex(int32 PoseIdx) const { return PoseIdx < GetNumPoses(); }
 	bool IsEmpty() const;
 
 	const FPoseSearchIndexAsset& GetAssetForPose(int32 PoseIdx) const;
