motion matching - data indexing permutations WIP

diff --git a/Engine/Plugins/Experimental/Animation/PoseSearch/Source/Runtime/Private/PoseSearchAssetIndexer.cpp b/Engine/Plugins/Experimental/Animation/PoseSearch/Source/Runtime/Private/PoseSearchAssetIndexer.cpp
--- a/Engine/Plugins/Experimental/Animation/PoseSearch/Source/Runtime/Private/PoseSearchAssetIndexer.cpp
+++ b/Engine/Plugins/Experimental/Animation/PoseSearch/Source/Runtime/Private/PoseSearchAssetIndexer.cpp
@@ -85,7 +85,7 @@ FAssetIndexer::FAssetIndexer(const FAssetIndexingContext& InIndexingContext, con
 	IndexingContext = InIndexingContext;
 
 	FirstIndexedSample = FMath::FloorToInt(IndexingContext.RequestedSamplingRange.Min * IndexingContext.Schema->SampleRate);
-	// @todo: parhaps we should use FMath::CeilToInt to avoid sampling over teh length of the aniamtion
+	// @todo: parhaps we should use FMath::CeilToInt to avoid sampling over teh length of the aniamation
 	LastIndexedSample = FMath::Max(0, FMath::CeilToInt(IndexingContext.RequestedSamplingRange.Max * IndexingContext.Schema->SampleRate));
 }
 
@@ -391,8 +391,9 @@ FTransform FAssetIndexer::CalculateComponentSpaceTransform(FAssetIndexer::Cached
 
 FQuat FAssetIndexer::GetSampleRotation(float SampleTimeOffset, int32 SampleIdx, int8 SchemaSampleBoneIdx, int8 SchemaOriginBoneIdx)
 {
-	const float OriginTime = FMath::Min(SampleIdx * IndexingContext.Schema->GetSamplingInterval(), IndexingContext.AssetSampler->GetPlayLength());
-	const float SampleTime = OriginTime + SampleTimeOffset;
+	const float Time = SampleIdx * IndexingContext.Schema->GetSamplingInterval();
+	const float SampleTime = Time + SampleTimeOffset + PermutationSampleTimeOffset;
+	const float OriginTime = Time + PermutationOriginTimeOffset;
 
 	// @todo: add support for SchemaSampleBoneIdx
 	if (!IndexingContext.Schema->IsRootBone(SchemaOriginBoneIdx))
@@ -418,9 +419,10 @@ FQuat FAssetIndexer::GetSampleRotation(float SampleTimeOffset, int32 SampleIdx,
 
 FVector FAssetIndexer::GetSamplePosition(float SampleTimeOffset, int32 SampleIdx, int8 SchemaSampleBoneIdx, int8 SchemaOriginBoneIdx)
 {
-	const float OriginTime = FMath::Min(SampleIdx * IndexingContext.Schema->GetSamplingInterval(), IndexingContext.AssetSampler->GetPlayLength());
-	const float SampleTime = OriginTime + SampleTimeOffset;
-
+	const float Time = SampleIdx * IndexingContext.Schema->GetSamplingInterval();
+	const float SampleTime = Time + SampleTimeOffset + PermutationSampleTimeOffset;
+	const float OriginTime = Time + PermutationOriginTimeOffset;
+	
 	bool bUnused;
 	return GetSamplePositionInternal(SampleTime, OriginTime, bUnused, SchemaSampleBoneIdx, SchemaOriginBoneIdx);
 }
@@ -458,8 +460,9 @@ FVector FAssetIndexer::GetSamplePositionInternal(float SampleTime, float OriginT
 
 FVector FAssetIndexer::GetSampleVelocity(float SampleTimeOffset, int32 SampleIdx, int8 SchemaSampleBoneIdx, int8 SchemaOriginBoneIdx, bool bUseCharacterSpaceVelocities)
 {
-	const float OriginTime = FMath::Min(SampleIdx * IndexingContext.Schema->GetSamplingInterval(), IndexingContext.AssetSampler->GetPlayLength());
-	const float SampleTime = OriginTime + SampleTimeOffset;
+	const float Time = SampleIdx * IndexingContext.Schema->GetSamplingInterval();
+	const float SampleTime = Time + SampleTimeOffset + PermutationSampleTimeOffset;
+	const float OriginTime = Time + PermutationOriginTimeOffset;
 	const float FiniteDelta = IndexingContext.SamplingContext->FiniteDelta;
 
 	bool bClampedPast, bClampedPresent;
@@ -497,5 +500,27 @@ const UPoseSearchSchema* FAssetIndexer::GetSchema() const
 	return IndexingContext.Schema;
 }
 
+void FAssetIndexer::SetPermutationTimeOffsets(float InPermutationSampleTimeOffset, float InPermutationOriginTimeOffset)
+{
+	// right now we disallow having nested channel controlling time offsets
+	check(PermutationSampleTimeOffset == 0.f && PermutationOriginTimeOffset == 0.f);
+	PermutationSampleTimeOffset = InPermutationSampleTimeOffset;
+	PermutationOriginTimeOffset = InPermutationOriginTimeOffset;
+}
+
+void FAssetIndexer::ResetPermutationTimeOffsets()
+{
+	PermutationSampleTimeOffset = 0.f;
+	PermutationOriginTimeOffset = 0.f;
+}
+
+float FAssetIndexer::CalculatePermutationTimeOffset() const
+{
+	const UPoseSearchSchema* Schema = GetSchema();
+	check(Schema->PermutationsSampleRate > 0 && SearchIndexAsset.IsInitialized());
+	const float PermutationTimeOffset = Schema->PermutationsTimeOffset + SearchIndexAsset.PermutationIdx / float(Schema->PermutationsSampleRate);
+	return PermutationTimeOffset;
+}
+
 } // namespace UE::PoseSearch
 #endif // WITH_EDITOR
diff --git a/Engine/Plugins/Experimental/Animation/PoseSearch/Source/Runtime/Private/PoseSearchDerivedData.cpp b/Engine/Plugins/Experimental/Animation/PoseSearch/Source/Runtime/Private/PoseSearchDerivedData.cpp
--- a/Engine/Plugins/Experimental/Animation/PoseSearch/Source/Runtime/Private/PoseSearchDerivedData.cpp
+++ b/Engine/Plugins/Experimental/Animation/PoseSearch/Source/Runtime/Private/PoseSearchDerivedData.cpp
@@ -322,14 +322,17 @@ static void InitSearchIndexAssets(FPoseSearchIndexBase& SearchIndex, UPoseSearch
 				FindValidSequenceIntervals(DatabaseSequence->Sequence, DatabaseSequence->SamplingRange, DatabaseSequence->IsLooping(), Database->ExcludeFromDatabaseParameters, ValidRanges);
 				for (const FFloatRange& Range : ValidRanges)
 				{
-					if (bAddUnmirrored)
+					for (int32 PermutationIdx = 0; PermutationIdx < Database->Schema->NumberOfPermutations; ++PermutationIdx)
 					{
-						SearchIndex.Assets.Add(FPoseSearchIndexAsset(AnimationAssetIndex, false, FFloatInterval(Range.GetLowerBoundValue(), Range.GetUpperBoundValue())));
-					}
+						if (bAddUnmirrored)
+						{
+							SearchIndex.Assets.Emplace(AnimationAssetIndex, false, FFloatInterval(Range.GetLowerBoundValue(), Range.GetUpperBoundValue()), PermutationIdx);
+						}
 
-					if (bAddMirrored)
-					{
-						SearchIndex.Assets.Add(FPoseSearchIndexAsset(AnimationAssetIndex, true, FFloatInterval(Range.GetLowerBoundValue(), Range.GetUpperBoundValue())));
+						if (bAddMirrored)
+						{
+							SearchIndex.Assets.Emplace(AnimationAssetIndex, true, FFloatInterval(Range.GetLowerBoundValue(), Range.GetUpperBoundValue()), PermutationIdx);
+						}
 					}
 				}
 			}
@@ -339,14 +342,17 @@ static void InitSearchIndexAssets(FPoseSearchIndexBase& SearchIndex, UPoseSearch
 				FindValidSequenceIntervals(DatabaseAnimComposite->AnimComposite, DatabaseAnimComposite->SamplingRange, DatabaseAnimComposite->IsLooping(), Database->ExcludeFromDatabaseParameters, ValidRanges);
 				for (const FFloatRange& Range : ValidRanges)
 				{
-					if (bAddUnmirrored)
+					for (int32 PermutationIdx = 0; PermutationIdx < Database->Schema->NumberOfPermutations; ++PermutationIdx)
 					{
-						SearchIndex.Assets.Add(FPoseSearchIndexAsset(AnimationAssetIndex, false, FFloatInterval(Range.GetLowerBoundValue(), Range.GetUpperBoundValue())));
-					}
+						if (bAddUnmirrored)
+						{
+							SearchIndex.Assets.Emplace(AnimationAssetIndex, false, FFloatInterval(Range.GetLowerBoundValue(), Range.GetUpperBoundValue()), PermutationIdx);
+						}
 
-					if (bAddMirrored)
-					{
-						SearchIndex.Assets.Add(FPoseSearchIndexAsset(AnimationAssetIndex, true, FFloatInterval(Range.GetLowerBoundValue(), Range.GetUpperBoundValue())));
+						if (bAddMirrored)
+						{
+							SearchIndex.Assets.Emplace(AnimationAssetIndex, true, FFloatInterval(Range.GetLowerBoundValue(), Range.GetUpperBoundValue()), PermutationIdx);
+						}
 					}
 				}
 			}
@@ -366,16 +372,19 @@ static void InitSearchIndexAssets(FPoseSearchIndexBase& SearchIndex, UPoseSearch
 						int32 TriangulationIndex = 0;
 						DatabaseBlendSpace->BlendSpace->GetSamplesFromBlendInput(BlendParameters, BlendSamples, TriangulationIndex, true);
 
-						float PlayLength = DatabaseBlendSpace->BlendSpace->GetAnimationLengthFromSampleData(BlendSamples);
+						const float PlayLength = DatabaseBlendSpace->BlendSpace->GetAnimationLengthFromSampleData(BlendSamples);
 
-						if (bAddUnmirrored)
+						for (int32 PermutationIdx = 0; PermutationIdx < Database->Schema->NumberOfPermutations; ++PermutationIdx)
 						{
-							SearchIndex.Assets.Add(FPoseSearchIndexAsset(AnimationAssetIndex, false, FFloatInterval(0.0f, PlayLength), BlendParameters));
-						}
-
-						if (bAddMirrored)
-						{
-							SearchIndex.Assets.Add(FPoseSearchIndexAsset(AnimationAssetIndex, true, FFloatInterval(0.0f, PlayLength), BlendParameters));
+							if (bAddUnmirrored)
+							{
+								SearchIndex.Assets.Emplace(AnimationAssetIndex, false, FFloatInterval(0.0f, PlayLength), PermutationIdx, BlendParameters);
+							}
+
+							if (bAddMirrored)
+							{
+								SearchIndex.Assets.Emplace(AnimationAssetIndex, true, FFloatInterval(0.0f, PlayLength), PermutationIdx, BlendParameters);
+							}
 						}
 					}
 				}
@@ -386,14 +395,17 @@ static void InitSearchIndexAssets(FPoseSearchIndexBase& SearchIndex, UPoseSearch
 				FindValidSequenceIntervals(DatabaseAnimMontage->AnimMontage, DatabaseAnimMontage->SamplingRange, DatabaseAnimMontage->IsLooping(), Database->ExcludeFromDatabaseParameters, ValidRanges);
 				for (const FFloatRange& Range : ValidRanges)
 				{
-					if (bAddUnmirrored)
+					for (int32 PermutationIdx = 0; PermutationIdx < Database->Schema->NumberOfPermutations; ++PermutationIdx)
 					{
-						SearchIndex.Assets.Add(FPoseSearchIndexAsset(AnimationAssetIndex, false, FFloatInterval(Range.GetLowerBoundValue(), Range.GetUpperBoundValue())));
-					}
+						if (bAddUnmirrored)
+						{
+							SearchIndex.Assets.Emplace(AnimationAssetIndex, false, FFloatInterval(Range.GetLowerBoundValue(), Range.GetUpperBoundValue()), PermutationIdx);
+						}
 
-					if (bAddMirrored)
-					{
-						SearchIndex.Assets.Add(FPoseSearchIndexAsset(AnimationAssetIndex, true, FFloatInterval(Range.GetLowerBoundValue(), Range.GetUpperBoundValue())));
+						if (bAddMirrored)
+						{
+							SearchIndex.Assets.Emplace(AnimationAssetIndex, true, FFloatInterval(Range.GetLowerBoundValue(), Range.GetUpperBoundValue()), PermutationIdx);
+						}
 					}
 				}
 			}
diff --git a/Engine/Plugins/Experimental/Animation/PoseSearch/Source/Runtime/Private/PoseSearchDerivedDataKey.cpp b/Engine/Plugins/Experimental/Animation/PoseSearch/Source/Runtime/Private/PoseSearchDerivedDataKey.cpp
--- a/Engine/Plugins/Experimental/Animation/PoseSearch/Source/Runtime/Private/PoseSearchDerivedDataKey.cpp
+++ b/Engine/Plugins/Experimental/Animation/PoseSearch/Source/Runtime/Private/PoseSearchDerivedDataKey.cpp
@@ -41,7 +41,7 @@ FKeyBuilder::FKeyBuilder(const UObject* Object, bool bUseDataVer)
 	if (bUseDataVer)
 	{
 		// used to invalidate the key without having to change POSESEARCHDB_DERIVEDDATA_VER all the times
-		int32 POSESEARCHDB_DERIVEDDATA_VER_SMALL = 107;
+		int32 POSESEARCHDB_DERIVEDDATA_VER_SMALL = 109;
 		FGuid VersionGuid = FDevSystemGuids::GetSystemGuid(FDevSystemGuids::Get().POSESEARCHDB_DERIVEDDATA_VER);
 
 		*this << VersionGuid;
diff --git a/Engine/Plugins/Experimental/Animation/PoseSearch/Source/Runtime/Private/PoseSearchFeatureChannel_FilterCrashingLegs.cpp b/Engine/Plugins/Experimental/Animation/PoseSearch/Source/Runtime/Private/PoseSearchFeatureChannel_FilterCrashingLegs.cpp
--- a/Engine/Plugins/Experimental/Animation/PoseSearch/Source/Runtime/Private/PoseSearchFeatureChannel_FilterCrashingLegs.cpp
+++ b/Engine/Plugins/Experimental/Animation/PoseSearch/Source/Runtime/Private/PoseSearchFeatureChannel_FilterCrashingLegs.cpp
@@ -38,33 +38,6 @@ void UPoseSearchFeatureChannel_FilterCrashingLegs::Finalize(UPoseSearchSchema* S
 	RightFootIdx = Schema->AddBoneReference(RightFoot);
 }
 
-#if WITH_EDITOR
-void UPoseSearchFeatureChannel_FilterCrashingLegs::FillWeights(TArray<float>& Weights) const
-{
-	for (int32 i = 0; i < ChannelCardinality; ++i)
-	{
-		Weights[ChannelDataOffset + i] = Weight;
-	}
-}
-
-void UPoseSearchFeatureChannel_FilterCrashingLegs::IndexAsset(UE::PoseSearch::FAssetIndexer& Indexer) const
-{
-	using namespace UE::PoseSearch;
-
-	for (int32 SampleIdx = Indexer.GetBeginSampleIdx(); SampleIdx != Indexer.GetEndSampleIdx(); ++SampleIdx)
-	{
-		const FVector RightThighPosition = Indexer.GetSamplePosition(0.f, SampleIdx, RightThighIdx);
-		const FVector LeftThighPosition = Indexer.GetSamplePosition(0.f, SampleIdx, LeftThighIdx);
-		const FVector RightFootPosition = Indexer.GetSamplePosition(0.f, SampleIdx, RightFootIdx);
-		const FVector LeftFootPosition = Indexer.GetSamplePosition(0.f, SampleIdx, LeftFootIdx);
-
-		const float CrashingLegsValue = ComputeCrashingLegsValue(RightThighPosition, LeftThighPosition, RightFootPosition, LeftFootPosition);
-
-		FFeatureVectorHelper::EncodeFloat(Indexer.GetPoseVector(SampleIdx), ChannelDataOffset, CrashingLegsValue);
-	}
-}
-#endif // WITH_EDITOR
-
 void UPoseSearchFeatureChannel_FilterCrashingLegs::BuildQuery(UE::PoseSearch::FSearchContext& SearchContext, FPoseSearchFeatureVectorBuilder& InOutQuery) const
 {
 	using namespace UE::PoseSearch;
@@ -160,6 +133,31 @@ bool UPoseSearchFeatureChannel_FilterCrashingLegs::IsPoseValid(TConstArrayView<f
 }
 
 #if WITH_EDITOR
+void UPoseSearchFeatureChannel_FilterCrashingLegs::FillWeights(TArray<float>& Weights) const
+{
+	for (int32 i = 0; i < ChannelCardinality; ++i)
+	{
+		Weights[ChannelDataOffset + i] = Weight;
+	}
+}
+
+void UPoseSearchFeatureChannel_FilterCrashingLegs::IndexAsset(UE::PoseSearch::FAssetIndexer& Indexer) const
+{
+	using namespace UE::PoseSearch;
+
+	for (int32 SampleIdx = Indexer.GetBeginSampleIdx(); SampleIdx != Indexer.GetEndSampleIdx(); ++SampleIdx)
+	{
+		const FVector RightThighPosition = Indexer.GetSamplePosition(0.f, SampleIdx, RightThighIdx);
+		const FVector LeftThighPosition = Indexer.GetSamplePosition(0.f, SampleIdx, LeftThighIdx);
+		const FVector RightFootPosition = Indexer.GetSamplePosition(0.f, SampleIdx, RightFootIdx);
+		const FVector LeftFootPosition = Indexer.GetSamplePosition(0.f, SampleIdx, LeftFootIdx);
+
+		const float CrashingLegsValue = ComputeCrashingLegsValue(RightThighPosition, LeftThighPosition, RightFootPosition, LeftFootPosition);
+
+		FFeatureVectorHelper::EncodeFloat(Indexer.GetPoseVector(SampleIdx), ChannelDataOffset, CrashingLegsValue);
+	}
+}
+
 FString UPoseSearchFeatureChannel_FilterCrashingLegs::GetLabel() const
 {
 	return Super::GetLabel();
diff --git a/Engine/Plugins/Experimental/Animation/PoseSearch/Source/Runtime/Private/PoseSearchFeatureChannel_FilterCrashingLegs.h b/Engine/Plugins/Experimental/Animation/PoseSearch/Source/Runtime/Private/PoseSearchFeatureChannel_FilterCrashingLegs.h
--- a/Engine/Plugins/Experimental/Animation/PoseSearch/Source/Runtime/Private/PoseSearchFeatureChannel_FilterCrashingLegs.h
+++ b/Engine/Plugins/Experimental/Animation/PoseSearch/Source/Runtime/Private/PoseSearchFeatureChannel_FilterCrashingLegs.h
@@ -63,7 +63,7 @@ public:
 #if WITH_EDITOR
 	virtual void FillWeights(TArray<float>& Weights) const override;
 	virtual void IndexAsset(UE::PoseSearch::FAssetIndexer& Indexer) const override;
-	virtual FString GetLabel() const;
+	virtual FString GetLabel() const override;
 #endif
 
 	// IPoseFilter interface
diff --git a/Engine/Plugins/Experimental/Animation/PoseSearch/Source/Runtime/Private/PoseSearchFeatureChannel_Group.h b/Engine/Plugins/Experimental/Animation/PoseSearch/Source/Runtime/Private/PoseSearchFeatureChannel_Group.h
--- a/Engine/Plugins/Experimental/Animation/PoseSearch/Source/Runtime/Private/PoseSearchFeatureChannel_Group.h
+++ b/Engine/Plugins/Experimental/Animation/PoseSearch/Source/Runtime/Private/PoseSearchFeatureChannel_Group.h
@@ -41,7 +41,7 @@ class POSESEARCH_API UPoseSearchFeatureChannel_Group : public UPoseSearchFeature
 	virtual TConstArrayView<TObjectPtr<UPoseSearchFeatureChannel>> GetSubChannels() const override { return SubChannels; }
 
 #if WITH_EDITOR
-	virtual FString GetLabel() const;
+	virtual FString GetLabel() const override;
 #endif
 
 public:
diff --git a/Engine/Plugins/Experimental/Animation/PoseSearch/Source/Runtime/Private/PoseSearchFeatureChannel_Heading.cpp b/Engine/Plugins/Experimental/Animation/PoseSearch/Source/Runtime/Private/PoseSearchFeatureChannel_Heading.cpp
--- a/Engine/Plugins/Experimental/Animation/PoseSearch/Source/Runtime/Private/PoseSearchFeatureChannel_Heading.cpp
+++ b/Engine/Plugins/Experimental/Animation/PoseSearch/Source/Runtime/Private/PoseSearchFeatureChannel_Heading.cpp
@@ -32,27 +32,6 @@ FVector UPoseSearchFeatureChannel_Heading::GetAxis(const FQuat& Rotation) const
 	return FVector(1.f, 0.f, 0.f);
 }
 
-#if WITH_EDITOR
-void UPoseSearchFeatureChannel_Heading::FillWeights(TArray<float>& Weights) const
-{
-	for (int32 i = 0; i < ChannelCardinality; ++i)
-	{
-		Weights[ChannelDataOffset + i] = Weight;
-	}
-}
-
-void UPoseSearchFeatureChannel_Heading::IndexAsset(UE::PoseSearch::FAssetIndexer& Indexer) const
-{
-	using namespace UE::PoseSearch;
-
-	for (int32 SampleIdx = Indexer.GetBeginSampleIdx(); SampleIdx != Indexer.GetEndSampleIdx(); ++SampleIdx)
-	{
-		const FVector Heading = GetAxis(Indexer.GetSampleRotation(SampleTimeOffset, SampleIdx, SchemaBoneIdx));
-		FFeatureVectorHelper::EncodeVector(Indexer.GetPoseVector(SampleIdx), ChannelDataOffset, Heading, ComponentStripping);
-	}
-}
-#endif // WITH_EDITOR
-
 void UPoseSearchFeatureChannel_Heading::BuildQuery(UE::PoseSearch::FSearchContext& SearchContext, FPoseSearchFeatureVectorBuilder& InOutQuery) const
 {
 	using namespace UE::PoseSearch;
@@ -103,6 +82,25 @@ void UPoseSearchFeatureChannel_Heading::DebugDraw(const UE::PoseSearch::FDebugDr
 #endif // ENABLE_DRAW_DEBUG
 
 #if WITH_EDITOR
+void UPoseSearchFeatureChannel_Heading::FillWeights(TArray<float>& Weights) const
+{
+	for (int32 i = 0; i < ChannelCardinality; ++i)
+	{
+		Weights[ChannelDataOffset + i] = Weight;
+	}
+}
+
+void UPoseSearchFeatureChannel_Heading::IndexAsset(UE::PoseSearch::FAssetIndexer& Indexer) const
+{
+	using namespace UE::PoseSearch;
+
+	for (int32 SampleIdx = Indexer.GetBeginSampleIdx(); SampleIdx != Indexer.GetEndSampleIdx(); ++SampleIdx)
+	{
+		const FVector Heading = GetAxis(Indexer.GetSampleRotation(SampleTimeOffset, SampleIdx, SchemaBoneIdx));
+		FFeatureVectorHelper::EncodeVector(Indexer.GetPoseVector(SampleIdx), ChannelDataOffset, Heading, ComponentStripping);
+	}
+}
+
 FString UPoseSearchFeatureChannel_Heading::GetLabel() const
 {
 	TStringBuilder<256> Label;
diff --git a/Engine/Plugins/Experimental/Animation/PoseSearch/Source/Runtime/Private/PoseSearchFeatureChannel_Heading.h b/Engine/Plugins/Experimental/Animation/PoseSearch/Source/Runtime/Private/PoseSearchFeatureChannel_Heading.h
--- a/Engine/Plugins/Experimental/Animation/PoseSearch/Source/Runtime/Private/PoseSearchFeatureChannel_Heading.h
+++ b/Engine/Plugins/Experimental/Animation/PoseSearch/Source/Runtime/Private/PoseSearchFeatureChannel_Heading.h
@@ -58,7 +58,7 @@ public:
 #if WITH_EDITOR
 	virtual void FillWeights(TArray<float>& Weights) const override;
 	virtual void IndexAsset(UE::PoseSearch::FAssetIndexer& Indexer) const override;
-	virtual FString GetLabel() const;
+	virtual FString GetLabel() const override;
 #endif
 
 	FVector GetAxis(const FQuat& Rotation) const;
diff --git a/Engine/Plugins/Experimental/Animation/PoseSearch/Source/Runtime/Private/PoseSearchFeatureChannel_PermutationTime.cpp b/Engine/Plugins/Experimental/Animation/PoseSearch/Source/Runtime/Private/PoseSearchFeatureChannel_PermutationTime.cpp
new file mode 100644
--- /dev/null
+++ b/Engine/Plugins/Experimental/Animation/PoseSearch/Source/Runtime/Private/PoseSearchFeatureChannel_PermutationTime.cpp
@@ -0,0 +1,53 @@
+// Copyright Epic Games, Inc. All Rights Reserved.
+
+#include "PoseSearchFeatureChannel_PermutationTime.h"
+#include "PoseSearch/PoseSearchAssetIndexer.h"
+#include "PoseSearch/PoseSearchContext.h"
+#include "PoseSearch/PoseSearchSchema.h"
+
+void UPoseSearchFeatureChannel_PermutationTime::Finalize(UPoseSearchSchema* Schema)
+{
+	ChannelDataOffset = Schema->SchemaCardinality;
+	ChannelCardinality = 1;
+	Schema->SchemaCardinality += ChannelCardinality;
+}
+
+void UPoseSearchFeatureChannel_PermutationTime::BuildQuery(UE::PoseSearch::FSearchContext& SearchContext, FPoseSearchFeatureVectorBuilder& InOutQuery) const
+{
+	using namespace UE::PoseSearch;
+	FFeatureVectorHelper::EncodeFloat(InOutQuery.EditValues(), ChannelDataOffset, SearchContext.DesiredPermutationTimeOffset);
+}
+
+#if WITH_EDITOR
+void UPoseSearchFeatureChannel_PermutationTime::FillWeights(TArray<float>& Weights) const
+{
+	for (int32 i = 0; i < ChannelCardinality; ++i)
+	{
+		Weights[ChannelDataOffset + i] = Weight;
+	}
+}
+
+void UPoseSearchFeatureChannel_PermutationTime::IndexAsset(UE::PoseSearch::FAssetIndexer& Indexer) const
+{
+	using namespace UE::PoseSearch;
+
+	const float PermutationTimeOffset = Indexer.CalculatePermutationTimeOffset();
+	for (int32 SampleIdx = Indexer.GetBeginSampleIdx(); SampleIdx != Indexer.GetEndSampleIdx(); ++SampleIdx)
+	{
+		FFeatureVectorHelper::EncodeFloat(Indexer.GetPoseVector(SampleIdx), ChannelDataOffset, PermutationTimeOffset);
+	}
+}
+
+FString UPoseSearchFeatureChannel_PermutationTime::GetLabel() const
+{
+	TStringBuilder<256> Label;
+	if (const UPoseSearchFeatureChannel* OuterChannel = Cast<UPoseSearchFeatureChannel>(GetOuter()))
+	{
+		Label.Append(OuterChannel->GetLabel());
+		Label.Append(TEXT("_"));
+	}
+
+	Label.Append(TEXT("PermTime"));
+	return Label.ToString();
+}
+#endif
\ No newline at end of file
diff --git a/Engine/Plugins/Experimental/Animation/PoseSearch/Source/Runtime/Private/PoseSearchFeatureChannel_PermutationTime.h b/Engine/Plugins/Experimental/Animation/PoseSearch/Source/Runtime/Private/PoseSearchFeatureChannel_PermutationTime.h
new file mode 100644
--- /dev/null
+++ b/Engine/Plugins/Experimental/Animation/PoseSearch/Source/Runtime/Private/PoseSearchFeatureChannel_PermutationTime.h
@@ -0,0 +1,26 @@
+// Copyright Epic Games, Inc. All Rights Reserved.
+
+#pragma once
+
+#include "PoseSearch/PoseSearchFeatureChannel.h"
+#include "PoseSearchFeatureChannel_PermutationTime.generated.h"
+
+UCLASS(EditInlineNew, meta = (DisplayName = "Permutation Time Channel"), CollapseCategories)
+class POSESEARCH_API UPoseSearchFeatureChannel_PermutationTime : public UPoseSearchFeatureChannel
+{
+	GENERATED_BODY()
+
+public:
+	UPROPERTY(EditAnywhere, Category = "Settings")
+	float Weight = 1.f;
+
+	// UPoseSearchFeatureChannel interface
+	virtual void Finalize(UPoseSearchSchema* Schema) override;
+	virtual void BuildQuery(UE::PoseSearch::FSearchContext& SearchContext, FPoseSearchFeatureVectorBuilder& InOutQuery) const override;
+
+#if WITH_EDITOR
+	virtual void FillWeights(TArray<float>& Weights) const override;
+	virtual void IndexAsset(UE::PoseSearch::FAssetIndexer& Indexer) const override;
+	virtual FString GetLabel() const override;
+#endif
+};
diff --git a/Engine/Plugins/Experimental/Animation/PoseSearch/Source/Runtime/Private/PoseSearchFeatureChannel_PermutationType.cpp b/Engine/Plugins/Experimental/Animation/PoseSearch/Source/Runtime/Private/PoseSearchFeatureChannel_PermutationType.cpp
new file mode 100644
--- /dev/null
+++ b/Engine/Plugins/Experimental/Animation/PoseSearch/Source/Runtime/Private/PoseSearchFeatureChannel_PermutationType.cpp
@@ -0,0 +1,48 @@
+// Copyright Epic Games, Inc. All Rights Reserved.
+
+#include "PoseSearchFeatureChannel_PermutationType.h"
+#include "PoseSearch/PoseSearchAssetIndexer.h"
+
+void UPoseSearchFeatureChannel_PermutationType::BuildQuery(UE::PoseSearch::FSearchContext& SearchContext, FPoseSearchFeatureVectorBuilder& InOutQuery) const
+{
+	Super::BuildQuery(SearchContext, InOutQuery);
+}
+
+#if WITH_EDITOR
+void UPoseSearchFeatureChannel_PermutationType::IndexAsset(UE::PoseSearch::FAssetIndexer& Indexer) const
+{
+	float PermutationSampleTimeOffset = 0.f;
+	float PermutationOriginTimeOffset = 0.f;
+	switch (PermutationType)
+	{
+	case EPermutationType::UseOriginTime:
+		PermutationSampleTimeOffset = 0.f;
+		PermutationOriginTimeOffset = 0.f;
+		break;
+	case EPermutationType::UsePermutationTime:
+		PermutationSampleTimeOffset = Indexer.CalculatePermutationTimeOffset();
+		PermutationOriginTimeOffset = PermutationSampleTimeOffset;
+		break;
+	case EPermutationType::UseOriginToPermutationTime:
+		PermutationSampleTimeOffset = Indexer.CalculatePermutationTimeOffset();
+		PermutationOriginTimeOffset = 0.f;
+		break;
+	}
+
+	Indexer.SetPermutationTimeOffsets(PermutationSampleTimeOffset, PermutationOriginTimeOffset);
+	Super::IndexAsset(Indexer);
+	Indexer.ResetPermutationTimeOffsets();
+}
+
+FString UPoseSearchFeatureChannel_PermutationType::GetLabel() const
+{
+	TStringBuilder<256> Label;
+	if (const UPoseSearchFeatureChannel* OuterChannel = Cast<UPoseSearchFeatureChannel>(GetOuter()))
+	{
+		Label.Append(OuterChannel->GetLabel());
+		Label.Append(TEXT("_"));
+	}
+	Label.Append(TEXT("PermType"));
+	return Label.ToString();
+}
+#endif // WITH_EDITOR
\ No newline at end of file
diff --git a/Engine/Plugins/Experimental/Animation/PoseSearch/Source/Runtime/Private/PoseSearchFeatureChannel_PermutationType.h b/Engine/Plugins/Experimental/Animation/PoseSearch/Source/Runtime/Private/PoseSearchFeatureChannel_PermutationType.h
new file mode 100644
--- /dev/null
+++ b/Engine/Plugins/Experimental/Animation/PoseSearch/Source/Runtime/Private/PoseSearchFeatureChannel_PermutationType.h
@@ -0,0 +1,41 @@
+// Copyright Epic Games, Inc. All Rights Reserved.
+
+#pragma once
+
+#include "PoseSearchFeatureChannel_Group.h"
+#include "PoseSearchFeatureChannel_PermutationType.generated.h"
+
+UENUM(meta = (Bitflags, UseEnumValuesAsMaskValuesInEditor = "true"))
+enum class EPermutationType : uint8
+{
+	UseOriginTime,
+	UsePermutationTime,
+	UseOriginToPermutationTime,
+};
+
+UCLASS(BlueprintType, EditInlineNew, meta = (DisplayName = "Permutation Type Channel"), CollapseCategories)
+class POSESEARCH_API UPoseSearchFeatureChannel_PermutationType : public UPoseSearchFeatureChannel_GroupBase
+{
+	GENERATED_BODY()
+
+public:
+	UPROPERTY(EditAnywhere, Category = "Settings")
+	EPermutationType PermutationType = EPermutationType::UseOriginToPermutationTime;
+
+	UPROPERTY(EditAnywhere, Instanced, BlueprintReadWrite, Category = "Settings")
+	TArray<TObjectPtr<UPoseSearchFeatureChannel>> SubChannels;
+
+	// UPoseSearchFeatureChannel_GroupBase interface
+	virtual TArrayView<TObjectPtr<UPoseSearchFeatureChannel>> GetSubChannels() override { return SubChannels; }
+	virtual TConstArrayView<TObjectPtr<UPoseSearchFeatureChannel>> GetSubChannels() const override { return SubChannels; }
+
+	// UPoseSearchFeatureChannel interface
+	virtual void BuildQuery(UE::PoseSearch::FSearchContext& SearchContext, FPoseSearchFeatureVectorBuilder& InOutQuery) const override;
+
+#if WITH_EDITOR
+	virtual void IndexAsset(UE::PoseSearch::FAssetIndexer& Indexer) const override;
+	virtual FString GetLabel() const override;
+#endif
+};
+
+
diff --git a/Engine/Plugins/Experimental/Animation/PoseSearch/Source/Runtime/Private/PoseSearchFeatureChannel_Phase.cpp b/Engine/Plugins/Experimental/Animation/PoseSearch/Source/Runtime/Private/PoseSearchFeatureChannel_Phase.cpp
--- a/Engine/Plugins/Experimental/Animation/PoseSearch/Source/Runtime/Private/PoseSearchFeatureChannel_Phase.cpp
+++ b/Engine/Plugins/Experimental/Animation/PoseSearch/Source/Runtime/Private/PoseSearchFeatureChannel_Phase.cpp
@@ -287,57 +287,6 @@
 	SchemaBoneIdx = Schema->AddBoneReference(Bone);
 }
 
-#if WITH_EDITOR
-void UPoseSearchFeatureChannel_Phase::FillWeights(TArray<float>& Weights) const
-{
-	for (int32 i = 0; i < ChannelCardinality; ++i)
-	{
-		Weights[ChannelDataOffset + i] = Weight;
-	}
-}
-
-void UPoseSearchFeatureChannel_Phase::IndexAsset(UE::PoseSearch::FAssetIndexer& Indexer) const
-{
-	using namespace UE::PoseSearch;
-
-	// @todo: do we really need to use double(s) in all this math?
-	// @todo: expose them via UI
-	static float BoneSamplingCentralDifferencesTime = 0.2f; // seconds
-	static float SmoothingWindowTime = 0.3f; // seconds
-
-
-	const UPoseSearchSchema* Schema = Indexer.GetSchema();
-	const float FiniteDelta = Schema->GetSamplingInterval();
-	
-	TArray<FVector2D> Phases;
-	TArray<float> Signal;
-	TArray<float> SmoothedSignal;
-	TArray<LocalMinMax> LocalMinMax;
-	TArray<FVector> BonePositions;
-	
-	CollectBonePositions(BonePositions, Indexer, SchemaBoneIdx);
-
-	// @todo: have different way of calculating signals, for example: height of the bone transform, acceleration, etc?
-	const int32 BoneSamplingCentralDifferencesOffset = FMath::Max(FMath::CeilToInt(BoneSamplingCentralDifferencesTime / FiniteDelta), 1);
-	CalculateSignal(BonePositions, Signal, BoneSamplingCentralDifferencesOffset);
-
-	const int32 SmoothingWindowOffset = FMath::Max(FMath::CeilToInt(SmoothingWindowTime / FiniteDelta), 1);
-	SmoothSignal(Signal, SmoothedSignal, SmoothingWindowOffset);
-
-	FindLocalMinMax(SmoothedSignal, LocalMinMax);
-	ValidateLocalMinMax(LocalMinMax);
-
-	ExtrapolateLocalMinMaxBoundaries(LocalMinMax, SmoothedSignal);
-	ValidateLocalMinMax(LocalMinMax);
-	CalculatePhasesFromLocalMinMax(LocalMinMax, Phases, SmoothedSignal.Num());
-
-	for (int32 SampleIdx = Indexer.GetBeginSampleIdx(); SampleIdx != Indexer.GetEndSampleIdx(); ++SampleIdx)
-	{
-		FFeatureVectorHelper::EncodeVector2D(Indexer.GetPoseVector(SampleIdx), ChannelDataOffset, Phases[SampleIdx - Indexer.GetBeginSampleIdx()]);
-	}
-}
-#endif // WITH_EDITOR
-
 void UPoseSearchFeatureChannel_Phase::BuildQuery(UE::PoseSearch::FSearchContext& SearchContext, FPoseSearchFeatureVectorBuilder& InOutQuery) const
 {
 	using namespace UE::PoseSearch;
@@ -390,6 +339,55 @@
 #endif // ENABLE_DRAW_DEBUG
 
 #if WITH_EDITOR
+void UPoseSearchFeatureChannel_Phase::FillWeights(TArray<float>& Weights) const
+{
+	for (int32 i = 0; i < ChannelCardinality; ++i)
+	{
+		Weights[ChannelDataOffset + i] = Weight;
+	}
+}
+
+void UPoseSearchFeatureChannel_Phase::IndexAsset(UE::PoseSearch::FAssetIndexer& Indexer) const
+{
+	using namespace UE::PoseSearch;
+
+	// @todo: do we really need to use double(s) in all this math?
+	// @todo: expose them via UI
+	static float BoneSamplingCentralDifferencesTime = 0.2f; // seconds
+	static float SmoothingWindowTime = 0.3f; // seconds
+
+
+	const UPoseSearchSchema* Schema = Indexer.GetSchema();
+	const float FiniteDelta = Schema->GetSamplingInterval();
+
+	TArray<FVector2D> Phases;
+	TArray<float> Signal;
+	TArray<float> SmoothedSignal;
+	TArray<LocalMinMax> LocalMinMax;
+	TArray<FVector> BonePositions;
+
+	CollectBonePositions(BonePositions, Indexer, SchemaBoneIdx);
+
+	// @todo: have different way of calculating signals, for example: height of the bone transform, acceleration, etc?
+	const int32 BoneSamplingCentralDifferencesOffset = FMath::Max(FMath::CeilToInt(BoneSamplingCentralDifferencesTime / FiniteDelta), 1);
+	CalculateSignal(BonePositions, Signal, BoneSamplingCentralDifferencesOffset);
+
+	const int32 SmoothingWindowOffset = FMath::Max(FMath::CeilToInt(SmoothingWindowTime / FiniteDelta), 1);
+	SmoothSignal(Signal, SmoothedSignal, SmoothingWindowOffset);
+
+	FindLocalMinMax(SmoothedSignal, LocalMinMax);
+	ValidateLocalMinMax(LocalMinMax);
+
+	ExtrapolateLocalMinMaxBoundaries(LocalMinMax, SmoothedSignal);
+	ValidateLocalMinMax(LocalMinMax);
+	CalculatePhasesFromLocalMinMax(LocalMinMax, Phases, SmoothedSignal.Num());
+
+	for (int32 SampleIdx = Indexer.GetBeginSampleIdx(); SampleIdx != Indexer.GetEndSampleIdx(); ++SampleIdx)
+	{
+		FFeatureVectorHelper::EncodeVector2D(Indexer.GetPoseVector(SampleIdx), ChannelDataOffset, Phases[SampleIdx - Indexer.GetBeginSampleIdx()]);
+	}
+}
+
 FString UPoseSearchFeatureChannel_Phase::GetLabel() const
 {
 	TStringBuilder<256> Label;
diff --git a/Engine/Plugins/Experimental/Animation/PoseSearch/Source/Runtime/Private/PoseSearchFeatureChannel_Phase.h b/Engine/Plugins/Experimental/Animation/PoseSearch/Source/Runtime/Private/PoseSearchFeatureChannel_Phase.h
--- a/Engine/Plugins/Experimental/Animation/PoseSearch/Source/Runtime/Private/PoseSearchFeatureChannel_Phase.h
+++ b/Engine/Plugins/Experimental/Animation/PoseSearch/Source/Runtime/Private/PoseSearchFeatureChannel_Phase.h
@@ -38,6 +38,6 @@ public:
 #if WITH_EDITOR
 	virtual void FillWeights(TArray<float>& Weights) const override;
 	virtual void IndexAsset(UE::PoseSearch::FAssetIndexer& Indexer) const override;
-	virtual FString GetLabel() const;
+	virtual FString GetLabel() const override;
 #endif
 };
\ No newline at end of file
diff --git a/Engine/Plugins/Experimental/Animation/PoseSearch/Source/Runtime/Private/PoseSearchFeatureChannel_Pose.h b/Engine/Plugins/Experimental/Animation/PoseSearch/Source/Runtime/Private/PoseSearchFeatureChannel_Pose.h
--- a/Engine/Plugins/Experimental/Animation/PoseSearch/Source/Runtime/Private/PoseSearchFeatureChannel_Pose.h
+++ b/Engine/Plugins/Experimental/Animation/PoseSearch/Source/Runtime/Private/PoseSearchFeatureChannel_Pose.h
@@ -69,7 +69,7 @@ public:
 	virtual void Finalize(UPoseSearchSchema* Schema) override;
 
 #if WITH_EDITOR
-	virtual FString GetLabel() const;
+	virtual FString GetLabel() const override;
 #endif
 };
 
diff --git a/Engine/Plugins/Experimental/Animation/PoseSearch/Source/Runtime/Private/PoseSearchFeatureChannel_Position.cpp b/Engine/Plugins/Experimental/Animation/PoseSearch/Source/Runtime/Private/PoseSearchFeatureChannel_Position.cpp
--- a/Engine/Plugins/Experimental/Animation/PoseSearch/Source/Runtime/Private/PoseSearchFeatureChannel_Position.cpp
+++ b/Engine/Plugins/Experimental/Animation/PoseSearch/Source/Runtime/Private/PoseSearchFeatureChannel_Position.cpp
@@ -18,27 +18,6 @@ void UPoseSearchFeatureChannel_Position::Finalize(UPoseSearchSchema* Schema)
 	SchemaOriginBoneIdx = Schema->AddBoneReference(OriginBone);
 }
 
-#if WITH_EDITOR
-void UPoseSearchFeatureChannel_Position::FillWeights(TArray<float>& Weights) const
-{
-	for (int32 i = 0; i < ChannelCardinality; ++i)
-	{
-		Weights[ChannelDataOffset + i] = Weight;
-	}
-}
-
-void UPoseSearchFeatureChannel_Position::IndexAsset(UE::PoseSearch::FAssetIndexer& Indexer) const
-{
-	using namespace UE::PoseSearch;
-
-	for (int32 SampleIdx = Indexer.GetBeginSampleIdx(); SampleIdx != Indexer.GetEndSampleIdx(); ++SampleIdx)
-	{
-		const FVector BonePosition = Indexer.GetSamplePosition(SampleTimeOffset, SampleIdx, SchemaBoneIdx, SchemaOriginBoneIdx);
-		FFeatureVectorHelper::EncodeVector(Indexer.GetPoseVector(SampleIdx), ChannelDataOffset, BonePosition, ComponentStripping);
-	}
-}
-#endif // WITH_EDITOR
-
 void UPoseSearchFeatureChannel_Position::BuildQuery(UE::PoseSearch::FSearchContext& SearchContext, FPoseSearchFeatureVectorBuilder& InOutQuery) const
 {
 	using namespace UE::PoseSearch;
@@ -124,6 +103,25 @@ void UPoseSearchFeatureChannel_Position::DebugDraw(const UE::PoseSearch::FDebugD
 #endif // ENABLE_DRAW_DEBUG
 
 #if WITH_EDITOR
+void UPoseSearchFeatureChannel_Position::FillWeights(TArray<float>& Weights) const
+{
+	for (int32 i = 0; i < ChannelCardinality; ++i)
+	{
+		Weights[ChannelDataOffset + i] = Weight;
+	}
+}
+
+void UPoseSearchFeatureChannel_Position::IndexAsset(UE::PoseSearch::FAssetIndexer& Indexer) const
+{
+	using namespace UE::PoseSearch;
+
+	for (int32 SampleIdx = Indexer.GetBeginSampleIdx(); SampleIdx != Indexer.GetEndSampleIdx(); ++SampleIdx)
+	{
+		const FVector BonePosition = Indexer.GetSamplePosition(SampleTimeOffset, SampleIdx, SchemaBoneIdx, SchemaOriginBoneIdx);
+		FFeatureVectorHelper::EncodeVector(Indexer.GetPoseVector(SampleIdx), ChannelDataOffset, BonePosition, ComponentStripping);
+	}
+}
+
 FString UPoseSearchFeatureChannel_Position::GetLabel() const
 {
 	TStringBuilder<256> Label;
diff --git a/Engine/Plugins/Experimental/Animation/PoseSearch/Source/Runtime/Private/PoseSearchFeatureChannel_Position.h b/Engine/Plugins/Experimental/Animation/PoseSearch/Source/Runtime/Private/PoseSearchFeatureChannel_Position.h
--- a/Engine/Plugins/Experimental/Animation/PoseSearch/Source/Runtime/Private/PoseSearchFeatureChannel_Position.h
+++ b/Engine/Plugins/Experimental/Animation/PoseSearch/Source/Runtime/Private/PoseSearchFeatureChannel_Position.h
@@ -51,6 +51,6 @@ public:
 #if WITH_EDITOR
 	virtual void FillWeights(TArray<float>& Weights) const override;
 	virtual void IndexAsset(UE::PoseSearch::FAssetIndexer& Indexer) const override;
-	virtual FString GetLabel() const;
+	virtual FString GetLabel() const override;
 #endif
 };
diff --git a/Engine/Plugins/Experimental/Animation/PoseSearch/Source/Runtime/Private/PoseSearchFeatureChannel_Trajectory.h b/Engine/Plugins/Experimental/Animation/PoseSearch/Source/Runtime/Private/PoseSearchFeatureChannel_Trajectory.h
--- a/Engine/Plugins/Experimental/Animation/PoseSearch/Source/Runtime/Private/PoseSearchFeatureChannel_Trajectory.h
+++ b/Engine/Plugins/Experimental/Animation/PoseSearch/Source/Runtime/Private/PoseSearchFeatureChannel_Trajectory.h
@@ -70,7 +70,7 @@ public:
 #endif // ENABLE_DRAW_DEBUG
 
 #if WITH_EDITOR
-	virtual FString GetLabel() const;
+	virtual FString GetLabel() const override;
 #endif
 
 	float GetEstimatedSpeedRatio(TConstArrayView<float> QueryVector, TConstArrayView<float> PoseVector) const;
diff --git a/Engine/Plugins/Experimental/Animation/PoseSearch/Source/Runtime/Private/PoseSearchFeatureChannel_Velocity.cpp b/Engine/Plugins/Experimental/Animation/PoseSearch/Source/Runtime/Private/PoseSearchFeatureChannel_Velocity.cpp
--- a/Engine/Plugins/Experimental/Animation/PoseSearch/Source/Runtime/Private/PoseSearchFeatureChannel_Velocity.cpp
+++ b/Engine/Plugins/Experimental/Animation/PoseSearch/Source/Runtime/Private/PoseSearchFeatureChannel_Velocity.cpp
@@ -18,31 +18,6 @@ void UPoseSearchFeatureChannel_Velocity::Finalize(UPoseSearchSchema* Schema)
 	SchemaBoneIdx = Schema->AddBoneReference(Bone);
 }
 
-#if WITH_EDITOR
-void UPoseSearchFeatureChannel_Velocity::FillWeights(TArray<float>& Weights) const
-{
-	for (int32 i = 0; i < ChannelCardinality; ++i)
-	{
-		Weights[ChannelDataOffset + i] = Weight;
-	}
-}
-
-void UPoseSearchFeatureChannel_Velocity::IndexAsset(UE::PoseSearch::FAssetIndexer& Indexer) const
-{
-	using namespace UE::PoseSearch;
-
-	for (int32 SampleIdx = Indexer.GetBeginSampleIdx(); SampleIdx != Indexer.GetEndSampleIdx(); ++SampleIdx)
-	{
-		FVector LinearVelocity = Indexer.GetSampleVelocity(SampleTimeOffset, SampleIdx, SchemaBoneIdx, RootSchemaBoneIdx, bUseCharacterSpaceVelocities);
-		if (bNormalize)
-		{
-			LinearVelocity = LinearVelocity.GetClampedToMaxSize(1.f);
-		}
-		FFeatureVectorHelper::EncodeVector(Indexer.GetPoseVector(SampleIdx), ChannelDataOffset, LinearVelocity, ComponentStripping);
-	}
-}
-#endif // WITH_EDITOR
-
 void UPoseSearchFeatureChannel_Velocity::BuildQuery(UE::PoseSearch::FSearchContext& SearchContext, FPoseSearchFeatureVectorBuilder& InOutQuery) const
 {
 	using namespace UE::PoseSearch;
@@ -101,6 +76,29 @@ void UPoseSearchFeatureChannel_Velocity::DebugDraw(const UE::PoseSearch::FDebugD
 #endif // ENABLE_DRAW_DEBUG
 
 #if WITH_EDITOR
+void UPoseSearchFeatureChannel_Velocity::FillWeights(TArray<float>& Weights) const
+{
+	for (int32 i = 0; i < ChannelCardinality; ++i)
+	{
+		Weights[ChannelDataOffset + i] = Weight;
+	}
+}
+
+void UPoseSearchFeatureChannel_Velocity::IndexAsset(UE::PoseSearch::FAssetIndexer& Indexer) const
+{
+	using namespace UE::PoseSearch;
+
+	for (int32 SampleIdx = Indexer.GetBeginSampleIdx(); SampleIdx != Indexer.GetEndSampleIdx(); ++SampleIdx)
+	{
+		FVector LinearVelocity = Indexer.GetSampleVelocity(SampleTimeOffset, SampleIdx, SchemaBoneIdx, RootSchemaBoneIdx, bUseCharacterSpaceVelocities);
+		if (bNormalize)
+		{
+			LinearVelocity = LinearVelocity.GetClampedToMaxSize(1.f);
+		}
+		FFeatureVectorHelper::EncodeVector(Indexer.GetPoseVector(SampleIdx), ChannelDataOffset, LinearVelocity, ComponentStripping);
+	}
+}
+
 FString UPoseSearchFeatureChannel_Velocity::GetLabel() const
 {
 	TStringBuilder<256> Label;
diff --git a/Engine/Plugins/Experimental/Animation/PoseSearch/Source/Runtime/Private/PoseSearchFeatureChannel_Velocity.h b/Engine/Plugins/Experimental/Animation/PoseSearch/Source/Runtime/Private/PoseSearchFeatureChannel_Velocity.h
--- a/Engine/Plugins/Experimental/Animation/PoseSearch/Source/Runtime/Private/PoseSearchFeatureChannel_Velocity.h
+++ b/Engine/Plugins/Experimental/Animation/PoseSearch/Source/Runtime/Private/PoseSearchFeatureChannel_Velocity.h
@@ -52,6 +52,6 @@ public:
 #if WITH_EDITOR
 	virtual void FillWeights(TArray<float>& Weights) const override;
 	virtual void IndexAsset(UE::PoseSearch::FAssetIndexer& Indexer) const override;
-	virtual FString GetLabel() const;
+	virtual FString GetLabel() const override;
 #endif
 };
diff --git a/Engine/Plugins/Experimental/Animation/PoseSearch/Source/Runtime/Private/PoseSearchIndex.cpp b/Engine/Plugins/Experimental/Animation/PoseSearch/Source/Runtime/Private/PoseSearchIndex.cpp
--- a/Engine/Plugins/Experimental/Animation/PoseSearch/Source/Runtime/Private/PoseSearchIndex.cpp
+++ b/Engine/Plugins/Experimental/Animation/PoseSearch/Source/Runtime/Private/PoseSearchIndex.cpp
@@ -48,8 +48,9 @@ FArchive& operator<<(FArchive& Ar, FPoseSearchIndexAsset& IndexAsset)
 {
 	Ar << IndexAsset.SourceAssetIdx;
 	Ar << IndexAsset.bMirrored;
-	Ar << IndexAsset.BlendParameters;
 	Ar << IndexAsset.SamplingInterval;
+	Ar << IndexAsset.PermutationIdx;
+	Ar << IndexAsset.BlendParameters;
 	Ar << IndexAsset.FirstPoseIdx;
 	Ar << IndexAsset.NumPoses;
 	return Ar;
diff --git a/Engine/Plugins/Experimental/Animation/PoseSearch/Source/Runtime/Public/PoseSearch/PoseSearchAssetIndexer.h b/Engine/Plugins/Experimental/Animation/PoseSearch/Source/Runtime/Public/PoseSearch/PoseSearchAssetIndexer.h
--- a/Engine/Plugins/Experimental/Animation/PoseSearch/Source/Runtime/Public/PoseSearch/PoseSearchAssetIndexer.h
+++ b/Engine/Plugins/Experimental/Animation/PoseSearch/Source/Runtime/Public/PoseSearch/PoseSearchAssetIndexer.h
@@ -58,6 +58,10 @@ public:
 	TArrayView<float> GetPoseVector(int32 SampleIdx) const;
 	const UPoseSearchSchema* GetSchema() const;
 
+	void SetPermutationTimeOffsets(float InPermutationSampleTimeOffset, float InPermutationOriginTimeOffset);
+	void ResetPermutationTimeOffsets();
+	float CalculatePermutationTimeOffset() const;
+	
 private:
 	int32 GetVectorIdx(int32 SampleIdx) const;
 	FTransform GetTransform(float SampleTime, bool& bClamped, int8 SchemaBoneIdx = RootSchemaBoneIdx);
@@ -95,6 +99,10 @@ private:
 	int32 FirstIndexedSample = 0;
 	int32 LastIndexedSample = 0;
 
+	// time offsets controlled by sampling data permutations
+	float PermutationSampleTimeOffset = 0.f;
+	float PermutationOriginTimeOffset = 0.f;
+
 	TArrayView<float> FeatureVectorTable;
 	TArrayView<FPoseSearchPoseMetadata> PoseMetadata;
 
diff --git a/Engine/Plugins/Experimental/Animation/PoseSearch/Source/Runtime/Public/PoseSearch/PoseSearchContext.h b/Engine/Plugins/Experimental/Animation/PoseSearch/Source/Runtime/Public/PoseSearch/PoseSearchContext.h
--- a/Engine/Plugins/Experimental/Animation/PoseSearch/Source/Runtime/Public/PoseSearch/PoseSearchContext.h
+++ b/Engine/Plugins/Experimental/Animation/PoseSearch/Source/Runtime/Public/PoseSearch/PoseSearchContext.h
@@ -160,6 +160,7 @@ struct POSESEARCH_API FSearchContext
 	bool bForceInterrupt = false;
 	// can the continuing pose advance? (if not we skip evaluating it)
 	bool bCanAdvance = true;
+	float DesiredPermutationTimeOffset = 0.f;
 
 	FQuat GetSampleRotation(float SampleTimeOffset, const UPoseSearchSchema* Schema, int8 SchemaSampleBoneIdx = RootSchemaBoneIdx, int8 SchemaOriginBoneIdx = RootSchemaBoneIdx, bool bUseHistoryRoot = false);
 	FVector GetSamplePosition(float SampleTimeOffset, const UPoseSearchSchema* Schema, int8 SchemaSampleBoneIdx = RootSchemaBoneIdx, int8 SchemaOriginBoneIdx = RootSchemaBoneIdx, bool bUseHistoryRoot = false);
diff --git a/Engine/Plugins/Experimental/Animation/PoseSearch/Source/Runtime/Public/PoseSearch/PoseSearchFeatureChannel.h b/Engine/Plugins/Experimental/Animation/PoseSearch/Source/Runtime/Public/PoseSearch/PoseSearchFeatureChannel.h
--- a/Engine/Plugins/Experimental/Animation/PoseSearch/Source/Runtime/Public/PoseSearch/PoseSearchFeatureChannel.h
+++ b/Engine/Plugins/Experimental/Animation/PoseSearch/Source/Runtime/Public/PoseSearch/PoseSearchFeatureChannel.h
@@ -105,7 +105,7 @@ public:
 	virtual void PreDebugDraw(UE::PoseSearch::FDebugDrawParams& DrawParams, TConstArrayView<float> PoseVector) const {}
 
 	// Draw this channel's data for the given pose vector
-	virtual void DebugDraw(const UE::PoseSearch::FDebugDrawParams& DrawParams, TConstArrayView<float> PoseVector) const PURE_VIRTUAL(UPoseSearchFeatureChannel::DebugDraw, );
+	virtual void DebugDraw(const UE::PoseSearch::FDebugDrawParams& DrawParams, TConstArrayView<float> PoseVector) const {}
 #endif //ENABLE_DRAW_DEBUG
 
 #if WITH_EDITOR
diff --git a/Engine/Plugins/Experimental/Animation/PoseSearch/Source/Runtime/Public/PoseSearch/PoseSearchIndex.h b/Engine/Plugins/Experimental/Animation/PoseSearch/Source/Runtime/Public/PoseSearch/PoseSearchIndex.h
--- a/Engine/Plugins/Experimental/Animation/PoseSearch/Source/Runtime/Public/PoseSearch/PoseSearchIndex.h
+++ b/Engine/Plugins/Experimental/Animation/PoseSearch/Source/Runtime/Public/PoseSearch/PoseSearchIndex.h
@@ -82,14 +82,16 @@ struct POSESEARCH_API FPoseSearchIndexAsset
 	FPoseSearchIndexAsset() {}
 
 	FPoseSearchIndexAsset(
-		int32 InSourceAssetIdx,
+		int32 InSourceAssetIdx, 
 		bool bInMirrored, 
 		const FFloatInterval& InSamplingInterval,
+		int32 InPermutationIdx,
 		FVector InBlendParameters = FVector::Zero())
 		: SourceAssetIdx(InSourceAssetIdx)
 		, bMirrored(bInMirrored)
-		, BlendParameters(InBlendParameters)
 		, SamplingInterval(InSamplingInterval)
+		, PermutationIdx(InPermutationIdx)
+		, BlendParameters(InBlendParameters)
 		, FirstPoseIdx(INDEX_NONE)
 		, NumPoses(INDEX_NONE)
 	{
@@ -103,10 +105,13 @@ struct POSESEARCH_API FPoseSearchIndexAsset
 	bool bMirrored = false;
 
 	UPROPERTY(meta = (NeverInHash))
-	FVector BlendParameters = FVector::Zero();
+	FFloatInterval SamplingInterval;
 
 	UPROPERTY(meta = (NeverInHash))
-	FFloatInterval SamplingInterval;
+	int32 PermutationIdx = INDEX_NONE;
+
+	UPROPERTY(meta = (NeverInHash))
+	FVector BlendParameters = FVector::Zero();
 
 	UPROPERTY(meta = (NeverInHash))
 	int32 FirstPoseIdx = INDEX_NONE;
@@ -115,6 +120,8 @@ struct POSESEARCH_API FPoseSearchIndexAsset
 	int32 NumPoses = INDEX_NONE;
 
 	bool IsPoseInRange(int32 PoseIdx) const { return (PoseIdx >= FirstPoseIdx) && (PoseIdx < FirstPoseIdx + NumPoses); }
+	bool IsInitialized() const { return SourceAssetIdx != INDEX_NONE && PermutationIdx != INDEX_NONE && FirstPoseIdx != INDEX_NONE && NumPoses != INDEX_NONE; }
+
 	friend FArchive& operator<<(FArchive& Ar, FPoseSearchIndexAsset& IndexAsset);
 };
 
diff --git a/Engine/Plugins/Experimental/Animation/PoseSearch/Source/Runtime/Public/PoseSearch/PoseSearchSchema.h b/Engine/Plugins/Experimental/Animation/PoseSearch/Source/Runtime/Public/PoseSearch/PoseSearchSchema.h
--- a/Engine/Plugins/Experimental/Animation/PoseSearch/Source/Runtime/Public/PoseSearch/PoseSearchSchema.h
+++ b/Engine/Plugins/Experimental/Animation/PoseSearch/Source/Runtime/Public/PoseSearch/PoseSearchSchema.h
@@ -93,6 +93,19 @@ public:
 	UPROPERTY(EditAnywhere, Category = "Bias")
 	float LoopingCostBias = 0.f;
 
+	// how many times the animation assets of the database using this schema will be indexed
+	UPROPERTY(EditAnywhere, meta = (ClampMin = "1"), Category = "Permutations")
+	int32 NumberOfPermutations = 1;
+
+	// delta time between every permutation indexing
+	UPROPERTY(EditAnywhere, meta = (ClampMin = "1", ClampMax = "240"), Category = "Permutations")
+	int32 PermutationsSampleRate = 30;
+
+	// starting offset of the "PermutationTime" from the "SamplingTime" of the first permutation.
+	// subsequent permutations will have PermutationTime = SamplingTime + PermutationsTimeOffset + PermutationIndex / PermutationsSampleRate
+	UPROPERTY(EditAnywhere, Category = "Permutations")
+	float PermutationsTimeOffset = 0.f;
+
 	UPROPERTY(EditAnywhere, Category = "Debug", meta = (ExcludeFromHash))
 	TArray<FPoseSearchSchemaColorPreset> ColorPresets;
 	
