Remove PoseSearchDatabaseSet. Equivalent functionality will be replaced soon by anim node functions and (optionally) choosers.

diff --git a/Engine/Plugins/Experimental/Animation/PoseSearch/Source/Editor/PoseSearchEditor.Build.cs b/Engine/Plugins/Experimental/Animation/PoseSearch/Source/Editor/PoseSearchEditor.Build.cs
--- a/Engine/Plugins/Experimental/Animation/PoseSearch/Source/Editor/PoseSearchEditor.Build.cs
+++ b/Engine/Plugins/Experimental/Animation/PoseSearch/Source/Editor/PoseSearchEditor.Build.cs
@@ -19,7 +19,6 @@ public class PoseSearchEditor : ModuleRules
 				"CoreUObject",
 				"Engine",
 				"PoseSearch",
-				"GameplayTags",
 				"StructUtils",
 				
 				// Trace-related dependencies
diff --git a/Engine/Plugins/Experimental/Animation/PoseSearch/Source/Runtime/PoseSearch.Build.cs b/Engine/Plugins/Experimental/Animation/PoseSearch/Source/Runtime/PoseSearch.Build.cs
--- a/Engine/Plugins/Experimental/Animation/PoseSearch/Source/Runtime/PoseSearch.Build.cs
+++ b/Engine/Plugins/Experimental/Animation/PoseSearch/Source/Runtime/PoseSearch.Build.cs
@@ -22,7 +22,6 @@ public class PoseSearch : ModuleRules
 				"AnimationCore",
 				"TraceLog",
 				"AnimGraphRuntime",
-				"GameplayTags",
 				"StructUtils"
 			}
 		);
diff --git a/Engine/Plugins/Experimental/Animation/PoseSearch/Source/Runtime/Private/AnimNode_MotionMatching.cpp b/Engine/Plugins/Experimental/Animation/PoseSearch/Source/Runtime/Private/AnimNode_MotionMatching.cpp
--- a/Engine/Plugins/Experimental/Animation/PoseSearch/Source/Runtime/Private/AnimNode_MotionMatching.cpp
+++ b/Engine/Plugins/Experimental/Animation/PoseSearch/Source/Runtime/Private/AnimNode_MotionMatching.cpp
@@ -111,7 +111,6 @@ void FAnimNode_MotionMatching::UpdateAssetPlayer(const FAnimationUpdateContext&
 	UPoseSearchLibrary::UpdateMotionMatchingState(
 		Context,
 		Searchable,
-		&ActiveTagsContainer,
 		Trajectory,
 		Settings,
 		MotionMatchingState,
diff --git a/Engine/Plugins/Experimental/Animation/PoseSearch/Source/Runtime/Private/PoseSearchDatabaseSet.cpp b/Engine/Plugins/Experimental/Animation/PoseSearch/Source/Runtime/Private/PoseSearchDatabaseSet.cpp
deleted file mode 100644
diff --git a/Engine/Plugins/Experimental/Animation/PoseSearch/Source/Runtime/Private/PoseSearchLibrary.cpp b/Engine/Plugins/Experimental/Animation/PoseSearch/Source/Runtime/Private/PoseSearchLibrary.cpp
--- a/Engine/Plugins/Experimental/Animation/PoseSearch/Source/Runtime/Private/PoseSearchLibrary.cpp
+++ b/Engine/Plugins/Experimental/Animation/PoseSearch/Source/Runtime/Private/PoseSearchLibrary.cpp
@@ -299,7 +299,6 @@ void UPoseSearchLibrary::TraceMotionMatchingState(
 void UPoseSearchLibrary::UpdateMotionMatchingState(
 	const FAnimationUpdateContext& Context,
 	const UPoseSearchSearchableAsset* Searchable,
-	const FGameplayTagContainer* ActiveTagsContainer,
 	const FTrajectorySampleRange& Trajectory,
 	const FMotionMatchingSettings& Settings,
 	FMotionMatchingState& InOutMotionMatchingState,
@@ -328,7 +327,6 @@ void UPoseSearchLibrary::UpdateMotionMatchingState(
 
 	// Build the search context
 	FSearchContext SearchContext;
-	SearchContext.ActiveTagsContainer = ActiveTagsContainer;
 	SearchContext.Trajectory = &Trajectory;
 	SearchContext.bForceInterrupt = bForceInterrupt;
 	SearchContext.bCanAdvance = InOutMotionMatchingState.CanAdvance(DeltaTime);
@@ -395,7 +393,6 @@ void UPoseSearchLibrary::UpdateMotionMatchingState(
 void UPoseSearchLibrary::MotionMatch(
 	const UAnimInstance* AnimInstance,
 	const UPoseSearchSearchableAsset* Searchable,
-	const FGameplayTagContainer ActiveTagsContainer,
 	const FTrajectorySampleRange Trajectory,
 	const FName PoseHistoryName,
 	UAnimationAsset*& SelectedAnimation,
@@ -417,7 +414,6 @@ void UPoseSearchLibrary::MotionMatch(
 	{
 		// Build the search context
 		FSearchContext SearchContext;
-		SearchContext.ActiveTagsContainer = &ActiveTagsContainer;
 		SearchContext.Trajectory = &Trajectory;
 
 		if (AnimInstance)
diff --git a/Engine/Plugins/Experimental/Animation/PoseSearch/Source/Runtime/Public/PoseSearch/AnimNode_MotionMatching.h b/Engine/Plugins/Experimental/Animation/PoseSearch/Source/Runtime/Public/PoseSearch/AnimNode_MotionMatching.h
--- a/Engine/Plugins/Experimental/Animation/PoseSearch/Source/Runtime/Public/PoseSearch/AnimNode_MotionMatching.h
+++ b/Engine/Plugins/Experimental/Animation/PoseSearch/Source/Runtime/Public/PoseSearch/AnimNode_MotionMatching.h
@@ -24,9 +24,6 @@ public:
 	UPROPERTY(EditAnywhere, BlueprintReadWrite, Category=Settings, meta=(PinShownByDefault))
 	TObjectPtr<const UPoseSearchSearchableAsset> Searchable = nullptr;
 
-	UPROPERTY(EditAnywhere, BlueprintReadWrite, Category = Settings, meta = (PinShownByDefault))
-	FGameplayTagContainer ActiveTagsContainer;
-
 	// Motion trajectory samples for pose search queries. Expected to be in the space of the SkeletalMeshComponent.
 	UPROPERTY(EditAnywhere, BlueprintReadWrite, Category=Settings, meta=(PinShownByDefault))
 	FTrajectorySampleRange Trajectory;
diff --git a/Engine/Plugins/Experimental/Animation/PoseSearch/Source/Runtime/Public/PoseSearch/PoseSearchContext.h b/Engine/Plugins/Experimental/Animation/PoseSearch/Source/Runtime/Public/PoseSearch/PoseSearchContext.h
--- a/Engine/Plugins/Experimental/Animation/PoseSearch/Source/Runtime/Public/PoseSearch/PoseSearchContext.h
+++ b/Engine/Plugins/Experimental/Animation/PoseSearch/Source/Runtime/Public/PoseSearch/PoseSearchContext.h
@@ -7,7 +7,6 @@
 #include "PoseSearch/PoseSearchIndex.h"
 #include "PoseSearch/PoseSearchResult.h"
 
-struct FGameplayTagContainer;
 struct FTrajectorySampleRange;
 class UPoseSearchDatabase;
 class UPoseSearchSchema;
@@ -155,7 +154,6 @@ struct POSESEARCH_API FSearchContext
 	const FPoseHistory* History = nullptr;
 	const FTrajectorySampleRange* Trajectory = nullptr;
 	FSearchResult CurrentResult;
-	const FGameplayTagContainer* ActiveTagsContainer = nullptr;
 	float PoseJumpThresholdTime = 0.f;
 	bool bForceInterrupt = false;
 	// can the continuing pose advance? (if not we skip evaluating it)
diff --git a/Engine/Plugins/Experimental/Animation/PoseSearch/Source/Runtime/Public/PoseSearch/PoseSearchDatabaseSet.h b/Engine/Plugins/Experimental/Animation/PoseSearch/Source/Runtime/Public/PoseSearch/PoseSearchDatabaseSet.h
deleted file mode 100644
diff --git a/Engine/Plugins/Experimental/Animation/PoseSearch/Source/Runtime/Public/PoseSearch/PoseSearchLibrary.h b/Engine/Plugins/Experimental/Animation/PoseSearch/Source/Runtime/Public/PoseSearch/PoseSearchLibrary.h
--- a/Engine/Plugins/Experimental/Animation/PoseSearch/Source/Runtime/Public/PoseSearch/PoseSearchLibrary.h
+++ b/Engine/Plugins/Experimental/Animation/PoseSearch/Source/Runtime/Public/PoseSearch/PoseSearchLibrary.h
@@ -20,7 +20,6 @@ namespace UE::PoseSearch
 } // namespace UE::PoseSearch
 
 struct FAnimationUpdateContext;
-struct FGameplayTagContainer;
 struct FTrajectorySampleRange;
 class UPoseSearchSearchableAsset;
 
@@ -140,7 +139,6 @@ public:
 	*
 	* @param Context						Input animation update context providing access to the proxy and delta time
 	* @param Searchable						Input collection of animations for motion matching
-	* @param ActiveTagsContainer			Input gameplay tag container
 	* @param Trajectory						Input motion trajectory samples for pose search queries
 	* @param Settings						Input motion matching algorithm configuration settings
 	* @param InOutMotionMatchingState		Input/Output encapsulated motion matching algorithm and state
@@ -149,7 +147,6 @@ public:
 	static void UpdateMotionMatchingState(
 		const FAnimationUpdateContext& Context,
 		const UPoseSearchSearchableAsset* Searchable,
-		const FGameplayTagContainer* ActiveTagsContainer,
 		const FTrajectorySampleRange& Trajectory,
 		const FMotionMatchingSettings& Settings,
 		FMotionMatchingState& InOutMotionMatchingState,
@@ -160,7 +157,6 @@ public:
 	*
 	* @param AnimInstance					Input animation instance
 	* @param Searchable						Input searchable asset
-	* @param ActiveTagsContainer			Input gameplay tag container
 	* @param Trajectory						Input motion trajectory samples for pose search queries
 	* @param PoseHistoryName				Input tag of the associated PoseSearchHistoryCollector node in the anim graph
 	* @param SelectedAnimation				Output selected animation from the searchable asset
@@ -174,7 +170,6 @@ public:
 	static void MotionMatch(
 		const UAnimInstance* AnimInstance,
 		const UPoseSearchSearchableAsset* Searchable,
-		const FGameplayTagContainer ActiveTagsContainer,
 		const FTrajectorySampleRange Trajectory,
 		const FName PoseHistoryName,
 		UAnimationAsset*& SelectedAnimation,
