Move continuing pose search into UpdateMotionMatchingState

diff --git a/Engine/Plugins/Experimental/Animation/PoseSearch/Source/Runtime/Private/PoseSearchDatabase.cpp b/Engine/Plugins/Experimental/Animation/PoseSearch/Source/Runtime/Private/PoseSearchDatabase.cpp
--- a/Engine/Plugins/Experimental/Animation/PoseSearch/Source/Runtime/Private/PoseSearchDatabase.cpp
+++ b/Engine/Plugins/Experimental/Animation/PoseSearch/Source/Runtime/Private/PoseSearchDatabase.cpp
@@ -769,29 +769,13 @@ UE::PoseSearch::FSearchResult UPoseSearchDatabase::SearchPCAKDTree(UE::PoseSearc
 		}
 	}
 
-	SearchContext.GetOrBuildQuery(Schema, Result.ComposedQuery);
-
-	TConstArrayView<float> QueryValues = Result.ComposedQuery.GetValues();
-
-	const bool IsCurrentResultFromThisDatabase = SearchContext.IsCurrentResultFromDatabase(this);
-
-	// evaluating the continuing pose only if it hasn't already being evaluated and the related animation can advance
-	if (!SearchContext.bForceInterrupt && IsCurrentResultFromThisDatabase && SearchContext.bCanAdvance && !Result.ContinuingPoseCost.IsValid())
-	{
-		Result.PoseIdx = SearchContext.CurrentResult.PoseIdx;
-		Result.PoseCost = SearchIndex.ComparePoses(Result.PoseIdx, SearchContext.QueryMirrorRequest, EPoseComparisonFlags::ContinuingPose, Schema->MirrorMismatchCostBias, QueryValues);
-		Result.ContinuingPoseCost = Result.PoseCost;
-
-		if (GetSkipSearchIfPossible())
-		{
-			SearchContext.UpdateCurrentBestCost(Result.PoseCost);
-		}
-	}
-
 	// since any PoseCost calculated here is at least SearchIndex.MinCostAddend,
 	// there's no point in performing the search if CurrentBestTotalCost is already better than that
-	if (SearchContext.GetCurrentBestTotalCost() > SearchIndex.MinCostAddend)
+	if (!GetSkipSearchIfPossible() || SearchContext.GetCurrentBestTotalCost() > SearchIndex.MinCostAddend)
 	{
+		SearchContext.GetOrBuildQuery(Schema, Result.ComposedQuery);
+		TConstArrayView<float> QueryValues = Result.ComposedQuery.GetValues();
+
 		FNonSelectableIdx NonSelectableIdx;
 		PopulateNonSelectableIdx(NonSelectableIdx, SearchContext, this, QueryValues);
 		FKDTree::KNNResultSet ResultSet(ClampedKDTreeQueryNumNeighbors, ResultIndexes, ResultDistanceSqr, NonSelectableIdx);
@@ -833,16 +817,13 @@ UE::PoseSearch::FSearchResult UPoseSearchDatabase::SearchPCAKDTree(UE::PoseSearc
 #endif
 			}
 		}
-
-		if (GetSkipSearchIfPossible() && Result.PoseCost.IsValid())
-		{
-			SearchContext.UpdateCurrentBestCost(Result.PoseCost);
-		}
 	}
 	else
 	{
 #if UE_POSE_SEARCH_TRACE_ENABLED
 		// calling just for reporting non selectable poses
+		SearchContext.GetOrBuildQuery(Schema, Result.ComposedQuery);
+		TConstArrayView<float> QueryValues = Result.ComposedQuery.GetValues();
 		FNonSelectableIdx NonSelectableIdx;
 		PopulateNonSelectableIdx(NonSelectableIdx, SearchContext, this, QueryValues);
 #endif
@@ -869,30 +850,13 @@ UE::PoseSearch::FSearchResult UPoseSearchDatabase::SearchBruteForce(UE::PoseSear
 
 	const FPoseSearchIndex& SearchIndex = GetSearchIndex();
 
-	SearchContext.GetOrBuildQuery(Schema, Result.ComposedQuery);
-	TConstArrayView<float> QueryValues = Result.ComposedQuery.GetValues();
-
-	const bool IsCurrentResultFromThisDatabase = SearchContext.IsCurrentResultFromDatabase(this);
-	if (!SearchContext.bForceInterrupt && IsCurrentResultFromThisDatabase)
-	{
-		// evaluating the continuing pose only if it hasn't already being evaluated and the related animation can advance
-		if (SearchContext.bCanAdvance && !Result.ContinuingPoseCost.IsValid())
-		{
-			Result.PoseIdx = SearchContext.CurrentResult.PoseIdx;
-			Result.PoseCost = SearchIndex.ComparePoses(Result.PoseIdx, SearchContext.QueryMirrorRequest, EPoseComparisonFlags::ContinuingPose, Schema->MirrorMismatchCostBias, QueryValues);
-			Result.ContinuingPoseCost = Result.PoseCost;
-
-			if (GetSkipSearchIfPossible())
-			{
-				SearchContext.UpdateCurrentBestCost(Result.PoseCost);
-			}
-		}
-	}
-
 	// since any PoseCost calculated here is at least SearchIndex.MinCostAddend,
 	// there's no point in performing the search if CurrentBestTotalCost is already better than that
-	if (SearchContext.GetCurrentBestTotalCost() > SearchIndex.MinCostAddend)
+	if (!GetSkipSearchIfPossible() || SearchContext.GetCurrentBestTotalCost() > SearchIndex.MinCostAddend)
 	{
+		SearchContext.GetOrBuildQuery(Schema, Result.ComposedQuery);
+		TConstArrayView<float> QueryValues = Result.ComposedQuery.GetValues();
+
 		FNonSelectableIdx NonSelectableIdx;
 		PopulateNonSelectableIdx(NonSelectableIdx, SearchContext, this, QueryValues);
 		check(Algo::IsSorted(NonSelectableIdx));
@@ -921,16 +885,13 @@ UE::PoseSearch::FSearchResult UPoseSearchDatabase::SearchBruteForce(UE::PoseSear
 #endif
 			}
 		}
-
-		if (GetSkipSearchIfPossible() && Result.PoseCost.IsValid())
-		{
-			SearchContext.UpdateCurrentBestCost(Result.PoseCost);
-		}
 	}
 	else
 	{
 #if UE_POSE_SEARCH_TRACE_ENABLED
 		// calling just for reporting non selectable poses
+		SearchContext.GetOrBuildQuery(Schema, Result.ComposedQuery);
+		TConstArrayView<float> QueryValues = Result.ComposedQuery.GetValues();
 		FNonSelectableIdx NonSelectableIdx;
 		PopulateNonSelectableIdx(NonSelectableIdx, SearchContext, this, QueryValues);
 #endif
diff --git a/Engine/Plugins/Experimental/Animation/PoseSearch/Source/Runtime/Private/PoseSearchLibrary.cpp b/Engine/Plugins/Experimental/Animation/PoseSearch/Source/Runtime/Private/PoseSearchLibrary.cpp
--- a/Engine/Plugins/Experimental/Animation/PoseSearch/Source/Runtime/Private/PoseSearchLibrary.cpp
+++ b/Engine/Plugins/Experimental/Animation/PoseSearch/Source/Runtime/Private/PoseSearchLibrary.cpp
@@ -12,6 +12,7 @@
 #include "PoseSearch/AnimNode_MotionMatching.h"
 #include "PoseSearch/AnimNode_PoseSearchHistoryCollector.h"
 #include "PoseSearch/PoseSearchDatabase.h"
+#include "PoseSearch/PoseSearchDerivedData.h"
 #include "PoseSearch/PoseSearchSchema.h"
 #include "PoseSearchFeatureChannel_Trajectory.h"
 #include "Trace/PoseSearchTraceLogger.h"
@@ -350,25 +351,46 @@ void UPoseSearchLibrary::UpdateMotionMatchingState(
 	{
 		InOutMotionMatchingState.ElapsedPoseSearchTime = 0.f;
 
+		// Evaluate continuing pose
+		FPoseSearchCost ContinuingPoseCost;
+		FPoseSearchFeatureVectorBuilder ContinuingPoseComposedQuery;
+		if (!SearchContext.bForceInterrupt && SearchContext.bCanAdvance && SearchContext.CurrentResult.Database.IsValid()
+#if WITH_EDITOR
+			&& FAsyncPoseSearchDatabasesManagement::RequestAsyncBuildIndex(SearchContext.CurrentResult.Database.Get(), ERequestAsyncBuildFlag::ContinueRequest)
+#endif
+		)
+		{
+			const UPoseSearchDatabase* ContinuingPoseDatabase = SearchContext.CurrentResult.Database.Get();
+			if (ensure(ContinuingPoseDatabase->Schema))
+			{	
+				SearchContext.GetOrBuildQuery(ContinuingPoseDatabase->Schema, ContinuingPoseComposedQuery);
+				ContinuingPoseCost = ContinuingPoseDatabase->GetSearchIndex().ComparePoses(SearchContext.CurrentResult.PoseIdx, SearchContext.QueryMirrorRequest,
+					EPoseComparisonFlags::ContinuingPose, ContinuingPoseDatabase->Schema->MirrorMismatchCostBias, ContinuingPoseComposedQuery.GetValues());
+				SearchContext.UpdateCurrentBestCost(ContinuingPoseCost);
+			}
+		}
+
 		// Search the database for the nearest match to the updated query vector
 		FSearchResult SearchResult = Searchable->Search(SearchContext);
+		if (SearchResult.PoseCost.IsValid())
+		{
+			SearchContext.UpdateCurrentBestCost(SearchResult.PoseCost);
+		}
 
-		// making sure we haven't calculated ContinuingPoseCost if we !SearchContext.bCanAdvance 
-		check(SearchContext.bCanAdvance || !SearchResult.ContinuingPoseCost.IsValid());
-
-		if (SearchResult.PoseCost.GetTotalCost() < SearchResult.ContinuingPoseCost.GetTotalCost())
+		if (SearchResult.PoseCost.GetTotalCost() < ContinuingPoseCost.GetTotalCost())
 		{
+			SearchResult.ContinuingPoseCost = ContinuingPoseCost;
 			InOutMotionMatchingState.JumpToPose(Context, Settings, SearchResult);
 		}
 		else
 		{
 			// copying few properties of SearchResult into CurrentSearchResult to facilitate debug drawing
 #if WITH_EDITOR
-			InOutMotionMatchingState.CurrentSearchResult.BruteForcePoseCost = SearchResult.BruteForcePoseCost;
+			InOutMotionMatchingState.CurrentSearchResult.BruteForcePoseCost = ContinuingPoseCost;
 #endif
-			InOutMotionMatchingState.CurrentSearchResult.PoseCost = SearchResult.PoseCost;
-			InOutMotionMatchingState.CurrentSearchResult.ContinuingPoseCost = SearchResult.ContinuingPoseCost;
-			InOutMotionMatchingState.CurrentSearchResult.ComposedQuery = SearchResult.ComposedQuery;
+			InOutMotionMatchingState.CurrentSearchResult.PoseCost = ContinuingPoseCost;
+			InOutMotionMatchingState.CurrentSearchResult.ContinuingPoseCost = ContinuingPoseCost;
+			InOutMotionMatchingState.CurrentSearchResult.ComposedQuery = ContinuingPoseComposedQuery;
 		}
 
 		InOutMotionMatchingState.UpdateWantedPlayRate(SearchContext, Settings);
