motion matching - data indexing permutations WIP

diff --git a/Engine/Plugins/Experimental/Animation/PoseSearch/Source/Runtime/Private/AnimNode_PoseSearchHistoryCollector.cpp b/Engine/Plugins/Experimental/Animation/PoseSearch/Source/Runtime/Private/AnimNode_PoseSearchHistoryCollector.cpp
--- a/Engine/Plugins/Experimental/Animation/PoseSearch/Source/Runtime/Private/AnimNode_PoseSearchHistoryCollector.cpp
+++ b/Engine/Plugins/Experimental/Animation/PoseSearch/Source/Runtime/Private/AnimNode_PoseSearchHistoryCollector.cpp
@@ -19,19 +19,19 @@ namespace UE::PoseSearch::Private
 class FPoseHistoryProvider : public IPoseHistoryProvider
 {
 public:
-	FPoseHistoryProvider(const FPoseHistory& InPoseHistory)
+	FPoseHistoryProvider(const IPoseHistory& InPoseHistory)
 		: PoseHistory(InPoseHistory)
 	{
 	}
 
 	// IPoseHistoryProvider interface
-	virtual const FPoseHistory& GetPoseHistory() const override
+	virtual const IPoseHistory& GetPoseHistory() const override
 	{
 		return PoseHistory;
 	}
 
 private:
-	const FPoseHistory& PoseHistory;
+	const IPoseHistory& PoseHistory;
 };
 
 } // namespace UE::PoseSearch::Private
@@ -72,12 +72,6 @@ void FAnimNode_PoseSearchHistoryCollector_Base::CacheBones_AnyThread(const FAnim
 	PoseHistory.Init(PoseCount, PoseDuration, RequiredBones);
 }
 
-void FAnimNode_PoseSearchHistoryCollector_Base::Update_AnyThread(const FAnimationUpdateContext& Context)
-{
-	DECLARE_SCOPE_HIERARCHICAL_COUNTER_ANIMNODE(Update_AnyThread);
-	Super::Update_AnyThread(Context);
-}
-
 bool FAnimNode_PoseSearchHistoryCollector_Base::HasPreUpdate() const
 {
 	return WITH_EDITORONLY_DATA && ENABLE_ANIM_DEBUG;
diff --git a/Engine/Plugins/Experimental/Animation/PoseSearch/Source/Runtime/Private/PoseSearchAssetIndexer.cpp b/Engine/Plugins/Experimental/Animation/PoseSearch/Source/Runtime/Private/PoseSearchAssetIndexer.cpp
--- a/Engine/Plugins/Experimental/Animation/PoseSearch/Source/Runtime/Private/PoseSearchAssetIndexer.cpp
+++ b/Engine/Plugins/Experimental/Animation/PoseSearch/Source/Runtime/Private/PoseSearchAssetIndexer.cpp
@@ -11,6 +11,7 @@
 #include "PoseSearch/PoseSearchFeatureChannel.h"
 #include "PoseSearch/PoseSearchIndex.h"
 #include "PoseSearch/PoseSearchSchema.h"
+#include "PoseSearchDatabaseIndexingContext.h"
 
 namespace UE::PoseSearch
 {
diff --git a/Engine/Plugins/Experimental/Animation/PoseSearch/Source/Runtime/Private/PoseSearchAssetSampler.cpp b/Engine/Plugins/Experimental/Animation/PoseSearch/Source/Runtime/Private/PoseSearchAssetSampler.cpp
--- a/Engine/Plugins/Experimental/Animation/PoseSearch/Source/Runtime/Private/PoseSearchAssetSampler.cpp
+++ b/Engine/Plugins/Experimental/Animation/PoseSearch/Source/Runtime/Private/PoseSearchAssetSampler.cpp
@@ -1,7 +1,5 @@
 // Copyright Epic Games, Inc. All Rights Reserved.
 
-#if WITH_EDITOR
-
 #include "PoseSearch/PoseSearchAssetSampler.h"
 #include "AnimationRuntime.h"
 #include "Animation/AnimMontage.h"
@@ -81,36 +79,6 @@ static FTransform ExtrapolateRootMotion(
 	return ExtrapolatedRootMotion;
 }
 
-//////////////////////////////////////////////////////////////////////////
-// FAssetSamplingContext
-void FAssetSamplingContext::Init(const UMirrorDataTable* InMirrorDataTable, const FBoneContainer& BoneContainer)
-{
-	MirrorDataTable = InMirrorDataTable;
-
-	if (InMirrorDataTable)
-	{
-		InMirrorDataTable->FillCompactPoseAndComponentRefRotations(BoneContainer, CompactPoseMirrorBones, ComponentSpaceRefRotations);
-	}
-	else
-	{
-		CompactPoseMirrorBones.Reset();
-		ComponentSpaceRefRotations.Reset();
-	}
-}
-
-FTransform FAssetSamplingContext::MirrorTransform(const FTransform& InTransform) const
-{
-	const EAxis::Type MirrorAxis = MirrorDataTable->MirrorAxis;
-	FVector T = InTransform.GetTranslation();
-	T = FAnimationRuntime::MirrorVector(T, MirrorAxis);
-	const FQuat ReferenceRotation = ComponentSpaceRefRotations[FCompactPoseBoneIndex(0)];
-	FQuat Q = InTransform.GetRotation();
-	Q = FAnimationRuntime::MirrorQuat(Q, MirrorAxis);
-	Q *= FAnimationRuntime::MirrorQuat(ReferenceRotation, MirrorAxis).Inverse() * ReferenceRotation;
-	FTransform Result = FTransform(Q, T, InTransform.GetScale3D());
-	return Result;
-}
-
 //////////////////////////////////////////////////////////////////////////
 // FSequenceBaseSampler
 void FSequenceBaseSampler::Init(const FInput& InInput)
@@ -733,5 +701,3 @@ const UAnimationAsset* FAnimMontageSampler::GetAsset() const
 }
 
 } // namespace UE::PoseSearch
-
-#endif // WITH_EDITOR
\ No newline at end of file
diff --git a/Engine/Plugins/Experimental/Animation/PoseSearch/Source/Runtime/Private/PoseSearchContext.cpp b/Engine/Plugins/Experimental/Animation/PoseSearch/Source/Runtime/Private/PoseSearchContext.cpp
--- a/Engine/Plugins/Experimental/Animation/PoseSearch/Source/Runtime/Private/PoseSearchContext.cpp
+++ b/Engine/Plugins/Experimental/Animation/PoseSearch/Source/Runtime/Private/PoseSearchContext.cpp
@@ -141,37 +141,48 @@ FQuat FSearchContext::GetSampleRotation(float SampleTimeOffset, const UPoseSearc
 	// @todo: add support for SchemaSampleBoneIdx
 	if (SchemaOriginBoneIdx != RootSchemaBoneIdx)
 	{
-		UE_LOG(LogPoseSearch,
-			Error,
-			TEXT("FSearchContext::GetSampleRotation: support for non root origin bones not implemented (bone: '%s', schema: '%s'"),
-			*Schema->BoneReferences[SchemaOriginBoneIdx].BoneName.ToString(),
-			*GetNameSafe(Schema));
+		UE_LOG(LogPoseSearch, Error, TEXT("FSearchContext::GetSampleRotation: support for non root origin bones not implemented (bone: '%s', schema: '%s'"),
+			*Schema->BoneReferences[SchemaOriginBoneIdx].BoneName.ToString(), *GetNameSafe(Schema));
 	}
 
-	return GetComponentSpaceTransform(SampleTimeOffset, Schema, SchemaSampleBoneIdx).GetRotation();
+	const float SampleTime = SampleTimeOffset + PermutationSampleTimeOffset;
+	const float OriginTime = PermutationOriginTimeOffset;
+
+	// @todo: add support for OriginTime != 0 (like in GetSamplePosition and GetSampleVelocity)
+	if (OriginTime != 0.f)
+	{
+		UE_LOG(LogPoseSearch, Error, TEXT("FSearchContext::GetSampleRotation: support for OriginTime != 0 not implemented (bone: '%s', schema: '%s'"),
+			SchemaOriginBoneIdx >= 0 ? *Schema->BoneReferences[SchemaOriginBoneIdx].BoneName.ToString() : TEXT("RootBone"), *GetNameSafe(Schema));
+	}
+
+	return GetComponentSpaceTransform(SampleTime, Schema, SchemaSampleBoneIdx).GetRotation();
 }
 
 FVector FSearchContext::GetSamplePosition(float SampleTimeOffset, const UPoseSearchSchema* Schema, int8 SchemaSampleBoneIdx, int8 SchemaOriginBoneIdx, bool bUseHistoryRoot)
 {
-	return GetSamplePositionInternal(SampleTimeOffset, 0.f, Schema, SchemaSampleBoneIdx, SchemaOriginBoneIdx, bUseHistoryRoot);
+	const float SampleTime = SampleTimeOffset + PermutationSampleTimeOffset;
+	const float OriginTime = PermutationOriginTimeOffset;
+	return GetSamplePositionInternal(SampleTime, OriginTime, Schema, SchemaSampleBoneIdx, SchemaOriginBoneIdx, bUseHistoryRoot);
 }
 
 FVector FSearchContext::GetSampleVelocity(float SampleTimeOffset, const UPoseSearchSchema* Schema, int8 SchemaSampleBoneIdx, int8 SchemaOriginBoneIdx, bool bUseCharacterSpaceVelocities, bool bUseHistoryRoot)
 {
-	const float HistorySampleInterval = History ? History->GetSampleTimeInterval() : 1 / 60.0f;
-	check(HistorySampleInterval > UE_KINDA_SMALL_NUMBER);
+	const float SampleTime = SampleTimeOffset + PermutationSampleTimeOffset;
+	const float OriginTime = PermutationOriginTimeOffset;
+	const float FiniteDelta = History ? History->GetSampleTimeInterval() : 1 / 60.0f;
+	check(FiniteDelta > UE_KINDA_SMALL_NUMBER);
 
 	// calculating the Position in component space for the bone indexed by SchemaSampleBoneIdx
-	const FVector CurrentTranslation = GetSamplePositionInternal(SampleTimeOffset, 0.f, Schema, SchemaSampleBoneIdx, SchemaOriginBoneIdx, bUseHistoryRoot);
-	const FVector PreviousTranslation = GetSamplePositionInternal(SampleTimeOffset - HistorySampleInterval, bUseCharacterSpaceVelocities ? -HistorySampleInterval : 0.f, Schema, SchemaSampleBoneIdx, SchemaOriginBoneIdx, bUseHistoryRoot);
+	const FVector PreviousTranslation = GetSamplePositionInternal(SampleTime - FiniteDelta, bUseCharacterSpaceVelocities ? OriginTime - FiniteDelta : OriginTime, Schema, SchemaSampleBoneIdx, SchemaOriginBoneIdx, bUseHistoryRoot);
+	const FVector CurrentTranslation = GetSamplePositionInternal(SampleTime, OriginTime, Schema, SchemaSampleBoneIdx, SchemaOriginBoneIdx, bUseHistoryRoot);
 
-	const FVector LinearVelocity = (CurrentTranslation - PreviousTranslation) / HistorySampleInterval;
+	const FVector LinearVelocity = (CurrentTranslation - PreviousTranslation) / FiniteDelta;
 	return LinearVelocity;
 }
 
 FTransform FSearchContext::GetTransform(float SampleTime, const UPoseSearchSchema* Schema, int8 SchemaBoneIdx, bool bUseHistoryRoot)
 {
-	// collecting the RootTransform from the FPoseHistory
+	// collecting the RootTransform from the IPoseHistory
 	FTransform RootTransform = FTransform::Identity;
 	if (bUseHistoryRoot)
 	{
@@ -207,7 +218,7 @@ FTransform FSearchContext::GetComponentSpaceTransform(float SampleTime, const UP
 			return CachedTransform->Transform;
 		}
 	
-		// collecting the local bone transforms from the FPoseHistory
+		// collecting the local bone transforms from the IPoseHistory
 		check(History);
 		FTransform BoneComponentSpaceTransform;
 		if (!History->GetComponentSpaceTransformAtTime(SampleTime, BoneIndexType, BoneComponentSpaceTransform))
@@ -254,6 +265,20 @@ FVector FSearchContext::GetSamplePositionInternal(float SampleTime, float Origin
 	return RootBoneTransform.InverseTransformVector(DeltaBoneTranslation);
 }
 
+void FSearchContext::SetPermutationTimeOffsets(float InPermutationSampleTimeOffset, float InPermutationOriginTimeOffset)
+{
+	// right now we disallow having nested channel controlling time offsets
+	check(PermutationSampleTimeOffset == 0.f && PermutationOriginTimeOffset == 0.f);
+	PermutationSampleTimeOffset = InPermutationSampleTimeOffset;
+	PermutationOriginTimeOffset = InPermutationOriginTimeOffset;
+}
+
+void FSearchContext::ResetPermutationTimeOffsets()
+{
+	PermutationSampleTimeOffset = 0.f;
+	PermutationOriginTimeOffset = 0.f;
+}
+
 void FSearchContext::ClearCachedEntries()
 {
 	CachedTransforms.Reset();
diff --git a/Engine/Plugins/Experimental/Animation/PoseSearch/Source/Runtime/Private/PoseSearchDatabaseIndexingContext.cpp b/Engine/Plugins/Experimental/Animation/PoseSearch/Source/Runtime/Private/PoseSearchDatabaseIndexingContext.cpp
--- a/Engine/Plugins/Experimental/Animation/PoseSearch/Source/Runtime/Private/PoseSearchDatabaseIndexingContext.cpp
+++ b/Engine/Plugins/Experimental/Animation/PoseSearch/Source/Runtime/Private/PoseSearchDatabaseIndexingContext.cpp
@@ -7,6 +7,7 @@
 #include "Animation/AnimMontage.h"
 #include "Animation/AnimSequence.h"
 #include "Animation/BlendSpace.h"
+#include "Animation/MirrorDataTable.h"
 #include "DerivedDataRequestOwner.h"
 #include "InstancedStruct.h"
 #include "PoseSearch/PoseSearchDatabase.h"
@@ -16,6 +17,38 @@
 namespace UE::PoseSearch
 {
 
+//////////////////////////////////////////////////////////////////////////
+// FAssetSamplingContext
+void FAssetSamplingContext::Init(const UMirrorDataTable* InMirrorDataTable, const FBoneContainer& BoneContainer)
+{
+	MirrorDataTable = InMirrorDataTable;
+
+	if (InMirrorDataTable)
+	{
+		InMirrorDataTable->FillCompactPoseAndComponentRefRotations(BoneContainer, CompactPoseMirrorBones, ComponentSpaceRefRotations);
+	}
+	else
+	{
+		CompactPoseMirrorBones.Reset();
+		ComponentSpaceRefRotations.Reset();
+	}
+}
+
+FTransform FAssetSamplingContext::MirrorTransform(const FTransform& InTransform) const
+{
+	const EAxis::Type MirrorAxis = MirrorDataTable->MirrorAxis;
+	FVector T = InTransform.GetTranslation();
+	T = FAnimationRuntime::MirrorVector(T, MirrorAxis);
+	const FQuat ReferenceRotation = ComponentSpaceRefRotations[FCompactPoseBoneIndex(0)];
+	FQuat Q = InTransform.GetRotation();
+	Q = FAnimationRuntime::MirrorQuat(Q, MirrorAxis);
+	Q *= FAnimationRuntime::MirrorQuat(ReferenceRotation, MirrorAxis).Inverse() * ReferenceRotation;
+	FTransform Result = FTransform(Q, T, InTransform.GetScale3D());
+	return Result;
+}
+
+//////////////////////////////////////////////////////////////////////////
+// FDatabaseIndexingContext
 bool FDatabaseIndexingContext::IndexDatabase(FPoseSearchIndexBase& SearchIndexBase, const UPoseSearchDatabase& Database, UE::DerivedData::FRequestOwner& Owner)
 {
 	const UPoseSearchSchema* Schema = Database.Schema;
diff --git a/Engine/Plugins/Experimental/Animation/PoseSearch/Source/Runtime/Private/PoseSearchDatabaseIndexingContext.h b/Engine/Plugins/Experimental/Animation/PoseSearch/Source/Runtime/Private/PoseSearchDatabaseIndexingContext.h
--- a/Engine/Plugins/Experimental/Animation/PoseSearch/Source/Runtime/Private/PoseSearchDatabaseIndexingContext.h
+++ b/Engine/Plugins/Experimental/Animation/PoseSearch/Source/Runtime/Private/PoseSearchDatabaseIndexingContext.h
@@ -14,6 +14,25 @@ namespace UE::DerivedData { class FRequestOwner; }
 namespace UE::PoseSearch
 {
 
+struct FAssetSamplingContext
+{
+	// Time delta used for computing pose derivatives
+	static constexpr float FiniteDelta = 1 / 60.0f;
+
+	// Mirror data table pointer copied from Schema for convenience
+	TObjectPtr<const UMirrorDataTable> MirrorDataTable;
+
+	// Compact pose format of Mirror Bone Map
+	TCustomBoneIndexArray<FCompactPoseBoneIndex, FCompactPoseBoneIndex> CompactPoseMirrorBones;
+
+	// Pre-calculated component space rotations of reference pose, which allows mirror to work with any joint orientation
+	// Only initialized and used when a mirroring table is specified
+	TCustomBoneIndexArray<FQuat, FCompactPoseBoneIndex> ComponentSpaceRefRotations;
+
+	void Init(const UMirrorDataTable* InMirrorDataTable, const FBoneContainer& BoneContainer);
+	FTransform MirrorTransform(const FTransform& Transform) const;
+};
+
 struct FDatabaseIndexingContext
 {
 	bool IndexDatabase(FPoseSearchIndexBase& SearchIndexBase, const UPoseSearchDatabase& Database, UE::DerivedData::FRequestOwner& Owner);
diff --git a/Engine/Plugins/Experimental/Animation/PoseSearch/Source/Runtime/Private/PoseSearchFeatureChannel_PermutationType.cpp b/Engine/Plugins/Experimental/Animation/PoseSearch/Source/Runtime/Private/PoseSearchFeatureChannel_PermutationType.cpp
--- a/Engine/Plugins/Experimental/Animation/PoseSearch/Source/Runtime/Private/PoseSearchFeatureChannel_PermutationType.cpp
+++ b/Engine/Plugins/Experimental/Animation/PoseSearch/Source/Runtime/Private/PoseSearchFeatureChannel_PermutationType.cpp
@@ -2,32 +2,48 @@
 
 #include "PoseSearchFeatureChannel_PermutationType.h"
 #include "PoseSearch/PoseSearchAssetIndexer.h"
+#include "PoseSearch/PoseSearchContext.h"
 
-void UPoseSearchFeatureChannel_PermutationType::BuildQuery(UE::PoseSearch::FSearchContext& SearchContext, FPoseSearchFeatureVectorBuilder& InOutQuery) const
-{
-	Super::BuildQuery(SearchContext, InOutQuery);
-}
-
-#if WITH_EDITOR
-void UPoseSearchFeatureChannel_PermutationType::IndexAsset(UE::PoseSearch::FAssetIndexer& Indexer) const
+void UPoseSearchFeatureChannel_PermutationType::GetPermutationTimeOffsets(float DesiredPermutationTimeOffset, float &OutPermutationSampleTimeOffset, float& OutPermutationOriginTimeOffset) const
 {
-	float PermutationSampleTimeOffset = 0.f;
-	float PermutationOriginTimeOffset = 0.f;
 	switch (PermutationType)
 	{
 	case EPermutationType::UseOriginTime:
-		PermutationSampleTimeOffset = 0.f;
-		PermutationOriginTimeOffset = 0.f;
+		OutPermutationSampleTimeOffset = 0.f;
+		OutPermutationOriginTimeOffset = 0.f;
 		break;
 	case EPermutationType::UsePermutationTime:
-		PermutationSampleTimeOffset = Indexer.CalculatePermutationTimeOffset();
-		PermutationOriginTimeOffset = PermutationSampleTimeOffset;
+		OutPermutationSampleTimeOffset = DesiredPermutationTimeOffset;
+		OutPermutationOriginTimeOffset = DesiredPermutationTimeOffset;
 		break;
 	case EPermutationType::UseOriginToPermutationTime:
-		PermutationSampleTimeOffset = Indexer.CalculatePermutationTimeOffset();
-		PermutationOriginTimeOffset = 0.f;
+		OutPermutationSampleTimeOffset = DesiredPermutationTimeOffset;
+		OutPermutationOriginTimeOffset = 0.f;
+		break;
+	default:
+		OutPermutationSampleTimeOffset = 0.f;
+		OutPermutationOriginTimeOffset = 0.f;
 		break;
 	}
+}
+
+void UPoseSearchFeatureChannel_PermutationType::BuildQuery(UE::PoseSearch::FSearchContext& SearchContext, FPoseSearchFeatureVectorBuilder& InOutQuery) const
+{
+	float PermutationSampleTimeOffset = 0.f;
+	float PermutationOriginTimeOffset = 0.f;
+	GetPermutationTimeOffsets(SearchContext.DesiredPermutationTimeOffset, PermutationSampleTimeOffset, PermutationOriginTimeOffset);
+
+	SearchContext.SetPermutationTimeOffsets(PermutationSampleTimeOffset, PermutationOriginTimeOffset);
+	Super::BuildQuery(SearchContext, InOutQuery);
+	SearchContext.ResetPermutationTimeOffsets();
+}
+
+#if WITH_EDITOR
+void UPoseSearchFeatureChannel_PermutationType::IndexAsset(UE::PoseSearch::FAssetIndexer& Indexer) const
+{
+	float PermutationSampleTimeOffset = 0.f;
+	float PermutationOriginTimeOffset = 0.f;
+	GetPermutationTimeOffsets(Indexer.CalculatePermutationTimeOffset(), PermutationSampleTimeOffset, PermutationOriginTimeOffset);
 
 	Indexer.SetPermutationTimeOffsets(PermutationSampleTimeOffset, PermutationOriginTimeOffset);
 	Super::IndexAsset(Indexer);
diff --git a/Engine/Plugins/Experimental/Animation/PoseSearch/Source/Runtime/Private/PoseSearchFeatureChannel_PermutationType.h b/Engine/Plugins/Experimental/Animation/PoseSearch/Source/Runtime/Private/PoseSearchFeatureChannel_PermutationType.h
--- a/Engine/Plugins/Experimental/Animation/PoseSearch/Source/Runtime/Private/PoseSearchFeatureChannel_PermutationType.h
+++ b/Engine/Plugins/Experimental/Animation/PoseSearch/Source/Runtime/Private/PoseSearchFeatureChannel_PermutationType.h
@@ -36,6 +36,9 @@ public:
 	virtual void IndexAsset(UE::PoseSearch::FAssetIndexer& Indexer) const override;
 	virtual FString GetLabel() const override;
 #endif
+
+private:
+	void GetPermutationTimeOffsets(float DesiredPermutationTimeOffset, float& OutPermutationSampleTimeOffset, float& OutPermutationOriginTimeOffset) const;
 };
 
 
diff --git a/Engine/Plugins/Experimental/Animation/PoseSearch/Source/Runtime/Private/PoseSearchHistory.cpp b/Engine/Plugins/Experimental/Animation/PoseSearch/Source/Runtime/Private/PoseSearchHistory.cpp
--- a/Engine/Plugins/Experimental/Animation/PoseSearch/Source/Runtime/Private/PoseSearchHistory.cpp
+++ b/Engine/Plugins/Experimental/Animation/PoseSearch/Source/Runtime/Private/PoseSearchHistory.cpp
@@ -62,8 +62,43 @@
 }
 
 //////////////////////////////////////////////////////////////////////////
+// FPoseHistoryEntry
+void FPoseHistoryEntry::Update(float InTime, FCSPose<FCompactPose>& ComponentSpacePose, const FTransform& ComponentTransform, const FBoneToTransformMap& BoneToTransformMap)
+{
+	Time = InTime;
+	RootTransform = ComponentTransform;
+
+	const FBoneContainer& BoneContainer = ComponentSpacePose.GetPose().GetBoneContainer();
+	const USkeleton* SkeletonAsset = BoneContainer.GetSkeletonAsset();
+	check(SkeletonAsset);
+	const FReferenceSkeleton& RefSkeleton = SkeletonAsset->GetReferenceSkeleton();
+	const TArray<FTransform>& RefBonePose = RefSkeleton.GetRefBonePose();
+	const int32 NumSkeletonBones = RefSkeleton.GetNum();
+
+	if (BoneToTransformMap.IsEmpty())
+	{
+		// no mapping: we add all the transforms
+		ComponentSpaceTransforms.SetNum(NumSkeletonBones);
+		for (FSkeletonPoseBoneIndex SkeletonBoneIdx(0); SkeletonBoneIdx != NumSkeletonBones; ++SkeletonBoneIdx)
+		{
+			const FCompactPoseBoneIndex CompactBoneIdx = BoneContainer.GetCompactPoseIndexFromSkeletonPoseIndex(SkeletonBoneIdx);
+			ComponentSpaceTransforms[SkeletonBoneIdx.GetInt()] = CompactBoneIdx.IsValid() ? ComponentSpacePose.GetComponentSpaceTransform(CompactBoneIdx) : RefBonePose[SkeletonBoneIdx.GetInt()];
+		}
+	}
+	else
+	{
+		ComponentSpaceTransforms.SetNum(BoneToTransformMap.Num());
+		for (const FBoneToTransformPair& BoneToTransformPair : BoneToTransformMap)
+		{
+			const FSkeletonPoseBoneIndex SkeletonBoneIdx(BoneToTransformPair.Key);
+			const FCompactPoseBoneIndex CompactBoneIdx = BoneContainer.GetCompactPoseIndexFromSkeletonPoseIndex(SkeletonBoneIdx);
+			ComponentSpaceTransforms[BoneToTransformPair.Value] = CompactBoneIdx.IsValid() ? ComponentSpacePose.GetComponentSpaceTransform(CompactBoneIdx) : RefBonePose[SkeletonBoneIdx.GetInt()];
+		}
+	}
+}
+
+//////////////////////////////////////////////////////////////////////////
 // FPoseHistory
-
 void FPoseHistory::Init(int32 InNumPoses, float InTimeHorizon, const TArray<FBoneIndexType>& RequiredBones)
 {
 	TimeHorizon = InTimeHorizon;
@@ -78,19 +113,6 @@
 	Entries.Reserve(InNumPoses);
 }
 
-static FBoneIndexType GetMaxFBoneIndexType(const TArray<FBoneIndexType>& RequiredBones)
-{
-	FBoneIndexType MaxBoneIndexType = 0;
-	for (FBoneIndexType BoneIndexType : RequiredBones)
-	{
-		if (BoneIndexType > MaxBoneIndexType)
-		{
-			MaxBoneIndexType = BoneIndexType;
-		}
-	}
-	return MaxBoneIndexType;
-}
-
 bool FPoseHistory::GetComponentSpaceTransformAtTime(float Time, FBoneIndexType BoneIndexType, FTransform& OutBoneTransform, bool bExtrapolate) const
 {
 	FComponentSpaceTransformIndex TransformIndex = FComponentSpaceTransformIndex(BoneIndexType);
@@ -111,12 +133,12 @@
 	if (Num > 1)
 	{
 		const float SecondsAgo = -Time;
-		const int32 LowerBoundIdx = LowerBound(Entries.begin(), Entries.end(), SecondsAgo, [](const FEntry& Entry, float Value) { return Value < Entry.Time; });
+		const int32 LowerBoundIdx = LowerBound(Entries.begin(), Entries.end(), SecondsAgo, [](const FPoseHistoryEntry& Entry, float Value) { return Value < Entry.Time; });
 		const int32 NextIdx = FMath::Clamp(LowerBoundIdx, 1, Num - 1);
 		const int32 PrevIdx = NextIdx - 1;
 
-		const FEntry& PrevEntry = Entries[PrevIdx];
-		const FEntry& NextEntry = Entries[NextIdx];
+		const FPoseHistoryEntry& PrevEntry = Entries[PrevIdx];
+		const FPoseHistoryEntry& NextEntry = Entries[NextIdx];
 
 		const float Denominator = NextEntry.Time - PrevEntry.Time;
 		if (!FMath::IsNearlyZero(Denominator))
@@ -136,7 +158,8 @@
 	}
 	else
 	{
-		OutBoneTransform = FTransform::Identity;
+		GetRootTransformAtTime(Time, OutBoneTransform, bExtrapolate);
+		return false;
 	}
 
 	return true;
@@ -148,12 +171,12 @@
 	if (Num > 1)
 	{
 		const float SecondsAgo = -Time;
-		const int32 LowerBoundIdx = LowerBound(Entries.begin(), Entries.end(), SecondsAgo, [](const FEntry& Entry, float Value) { return Value < Entry.Time; });
+		const int32 LowerBoundIdx = LowerBound(Entries.begin(), Entries.end(), SecondsAgo, [](const FPoseHistoryEntry& Entry, float Value) { return Value < Entry.Time; });
 		const int32 NextIdx = FMath::Clamp(LowerBoundIdx, 1, Num - 1);
 		const int32 PrevIdx = NextIdx - 1;
 
-		const FEntry& PrevEntry = Entries[PrevIdx];
-		const FEntry& NextEntry = Entries[NextIdx];
+		const FPoseHistoryEntry& PrevEntry = Entries[PrevIdx];
+		const FPoseHistoryEntry& NextEntry = Entries[NextIdx];
 
 		const float Denominator = NextEntry.Time - PrevEntry.Time;
 		if (!FMath::IsNearlyZero(Denominator))
@@ -177,10 +200,10 @@
 	}
 }
 
-void FPoseHistory::Update(float SecondsElapsed, FCSPose<FCompactPose>& ComponentSpacePose, FTransform ComponentTransform)
+void FPoseHistory::Update(float SecondsElapsed, FCSPose<FCompactPose>& ComponentSpacePose, const FTransform& ComponentTransform)
 {
 	// Age our elapsed times
-	for (FEntry& Entry : Entries)
+	for (FPoseHistoryEntry& Entry : Entries)
 	{
 		Entry.Time += SecondsElapsed;
 	}
@@ -204,44 +227,14 @@
 
 		if (bCanEvictOldest && bShouldPushNewest)
 		{
-			FEntry EntryTemp = MoveTemp(Entries.First());
+			FPoseHistoryEntry EntryTemp = MoveTemp(Entries.First());
 			Entries.PopFront();
 			Entries.Emplace(MoveTemp(EntryTemp));
 		}
 	}
 
 	// Regardless of the retention policy, we always update the most recent Entry
-	FEntry& CurrentEntry = Entries.Last();
-	CurrentEntry.Time = 0.f;
-	CurrentEntry.RootTransform = ComponentTransform;
-
-	const FBoneContainer& BoneContainer = ComponentSpacePose.GetPose().GetBoneContainer();
-	const USkeleton* SkeletonAsset = BoneContainer.GetSkeletonAsset();
-	check(SkeletonAsset);
-	const FReferenceSkeleton& RefSkeleton = SkeletonAsset->GetReferenceSkeleton();
-	const TArray<FTransform>& RefBonePose = RefSkeleton.GetRefBonePose();
-	const int32 NumSkeletonBones = RefSkeleton.GetNum();
-
-	if (BoneToTransformMap.IsEmpty())
-	{
-		// no mapping: we add all the transforms
-		CurrentEntry.ComponentSpaceTransforms.SetNum(NumSkeletonBones);
-		for (FSkeletonPoseBoneIndex SkeletonBoneIdx(0); SkeletonBoneIdx != NumSkeletonBones; ++SkeletonBoneIdx)
-		{
-			const FCompactPoseBoneIndex CompactBoneIdx = BoneContainer.GetCompactPoseIndexFromSkeletonPoseIndex(SkeletonBoneIdx);
-			CurrentEntry.ComponentSpaceTransforms[SkeletonBoneIdx.GetInt()] = CompactBoneIdx.IsValid() ? ComponentSpacePose.GetComponentSpaceTransform(CompactBoneIdx) : RefBonePose[SkeletonBoneIdx.GetInt()];
-		}
-	}
-	else
-	{
-		CurrentEntry.ComponentSpaceTransforms.SetNum(BoneToTransformMap.Num());
-		for (const TPair<FBoneIndexType, FComponentSpaceTransformIndex>& BoneToTransformPair : BoneToTransformMap)
-		{
-			const FSkeletonPoseBoneIndex SkeletonBoneIdx(BoneToTransformPair.Key);
-			const FCompactPoseBoneIndex CompactBoneIdx = BoneContainer.GetCompactPoseIndexFromSkeletonPoseIndex(SkeletonBoneIdx);
-			CurrentEntry.ComponentSpaceTransforms[BoneToTransformPair.Value] = CompactBoneIdx.IsValid() ? ComponentSpacePose.GetComponentSpaceTransform(CompactBoneIdx) : RefBonePose[SkeletonBoneIdx.GetInt()];
-		}
-	}
+	Entries.Last().Update(0.f, ComponentSpacePose, ComponentTransform, BoneToTransformMap);
 }
 
 float FPoseHistory::GetSampleTimeInterval() const
@@ -256,3 +249,170 @@
 
 	auto LerpColor = [](FColor A, FColor B, float T) -> FColor
 	{
+		return FColor(
+			FMath::RoundToInt(float(A.R) * (1.f - T) + float(B.R) * T),
+			FMath::RoundToInt(float(A.G) * (1.f - T) + float(B.G) * T),
+			FMath::RoundToInt(float(A.B) * (1.f - T) + float(B.B) * T),
+			FMath::RoundToInt(float(A.A) * (1.f - T) + float(B.A) * T));
+	};
+
+	TArray<FTransform> PrevGlobalTransforms;
+	for (int32 EntryIndex = 0; EntryIndex < Entries.Num(); ++EntryIndex)
+	{
+		const FPoseHistoryEntry& Entry = Entries[EntryIndex];
+		if (Entry.ComponentSpaceTransforms.IsEmpty())
+		{
+			PrevGlobalTransforms.Reset();
+		}
+		else if (PrevGlobalTransforms.Num() != Entry.ComponentSpaceTransforms.Num())
+		{
+			PrevGlobalTransforms.SetNum(Entry.ComponentSpaceTransforms.Num());
+			for (int32 i = 0; i < Entry.ComponentSpaceTransforms.Num(); ++i)
+			{
+				PrevGlobalTransforms[i] = Entry.ComponentSpaceTransforms[i] * Entry.RootTransform;
+			}
+		}
+		else
+		{
+			const float LerpFactor = float(EntryIndex - 1) / float(Entries.Num() - 1);
+			const FColor Color = LerpColor(FColorList::Red, FColorList::Orange, LerpFactor);
+
+			for (int32 i = 0; i < Entry.ComponentSpaceTransforms.Num(); ++i)
+			{
+				const FTransform GlobalTransforms = Entry.ComponentSpaceTransforms[i] * Entry.RootTransform;
+
+				DrawDebugLine(World, PrevGlobalTransforms[i].GetTranslation(), GlobalTransforms.GetTranslation(), Color, false, 0.f, ESceneDepthPriorityGroup::SDPG_Foreground + 2);
+
+				PrevGlobalTransforms[i] = GlobalTransforms;
+			}
+		}
+	}
+#endif // ENABLE_DRAW_DEBUG
+}
+
+//////////////////////////////////////////////////////////////////////////
+// FExtendedPoseHistory
+void FExtendedPoseHistory::Init(const FPoseHistory* InPoseHistory)
+{
+	check(InPoseHistory);
+	PoseHistory = InPoseHistory;
+}
+
+float FExtendedPoseHistory::GetSampleTimeInterval() const
+{
+	check(PoseHistory);
+	return PoseHistory->GetSampleTimeInterval();
+}
+
+bool FExtendedPoseHistory::GetComponentSpaceTransformAtTime(float Time, FBoneIndexType BoneIndexType, FTransform& OutBoneTransform, bool bExtrapolate) const
+{
+	check(PoseHistory);
+
+	if (Time > 0.f)
+	{
+ 		FComponentSpaceTransformIndex TransformIndex = FComponentSpaceTransformIndex(BoneIndexType);
+		const FBoneToTransformMap& BoneToTransformMap = PoseHistory->GetBoneToTransformMap();
+		if (!BoneToTransformMap.IsEmpty())
+		{
+			if (const FComponentSpaceTransformIndex* FoundTransformIndex = BoneToTransformMap.Find(BoneIndexType))
+			{
+				TransformIndex = *FoundTransformIndex;
+			}
+			else
+			{
+				GetRootTransformAtTime(Time, OutBoneTransform, bExtrapolate);
+				return false;
+			}
+		}
+
+		const int32 Num = FutureEntries.Num();
+		if (Num > 0)
+		{
+			const float SecondsAgo = -Time;
+			const int32 LowerBoundIdx = Algo::LowerBound(FutureEntries, SecondsAgo, [](const FPoseHistoryEntry& Entry, float Value) { return Value < Entry.Time; });
+			const int32 NextIdx = FMath::Min(LowerBoundIdx, Num - 1);
+			const FPoseHistoryEntries& PastEntries = PoseHistory->GetEntries();
+			const FPoseHistoryEntry& NextEntry = FutureEntries[NextIdx];
+			const FPoseHistoryEntry& PrevEntry = NextIdx > 0 ? FutureEntries[NextIdx - 1] : !PastEntries.IsEmpty() ? PastEntries.First() : NextEntry;
+
+			const float Denominator = NextEntry.Time - PrevEntry.Time;
+			if (!FMath::IsNearlyZero(Denominator))
+			{
+				const float Numerator = SecondsAgo - PrevEntry.Time;
+				const float LerpValue = bExtrapolate ? Numerator / Denominator : FMath::Clamp(Numerator / Denominator, 0.f, 1.f);
+				OutBoneTransform.Blend(PrevEntry.ComponentSpaceTransforms[TransformIndex], NextEntry.ComponentSpaceTransforms[TransformIndex], LerpValue);
+			}
+			else
+			{
+				OutBoneTransform = NextEntry.ComponentSpaceTransforms[TransformIndex];
+			}
+			return true;
+		}
+
+		// we've no FutureEntries, so let's clamp the time to zero and query PoseHistory
+		Time = 0.f;
+	}
+
+	return PoseHistory->GetComponentSpaceTransformAtTime(Time, BoneIndexType, OutBoneTransform, bExtrapolate);
+}
+
+void FExtendedPoseHistory::GetRootTransformAtTime(float Time, FTransform& OutRootTransform, bool bExtrapolate) const
+{
+	check(PoseHistory);
+
+	if (Time > 0.f)
+	{
+		const int32 Num = FutureEntries.Num();
+		if (Num > 0)
+		{
+			const float SecondsAgo = -Time;
+			const int32 LowerBoundIdx = Algo::LowerBound(FutureEntries, SecondsAgo, [](const FPoseHistoryEntry& Entry, float Value) { return Value < Entry.Time; });
+			const int32 NextIdx = FMath::Min(LowerBoundIdx, Num - 1);
+			const FPoseHistoryEntries& PastEntries = PoseHistory->GetEntries();
+			const FPoseHistoryEntry& NextEntry = FutureEntries[NextIdx];
+			const FPoseHistoryEntry& PrevEntry = NextIdx > 0 ? FutureEntries[NextIdx - 1] : !PastEntries.IsEmpty() ? PastEntries.First() : NextEntry;
+
+			const float Denominator = NextEntry.Time - PrevEntry.Time;
+			if (!FMath::IsNearlyZero(Denominator))
+			{
+				const float Numerator = SecondsAgo - PrevEntry.Time;
+				const float LerpValue = bExtrapolate ? Numerator / Denominator : FMath::Clamp(Numerator / Denominator, 0.f, 1.f);
+				OutRootTransform.Blend(PrevEntry.RootTransform, NextEntry.RootTransform, LerpValue);
+			}
+			else
+			{
+				OutRootTransform = NextEntry.RootTransform;
+			}
+
+			return;
+		}
+
+		// we've no FutureEntries, so let's clamp the time to zero and query PoseHistory
+		Time = 0.f;
+	}
+
+	PoseHistory->GetRootTransformAtTime(Time, OutRootTransform, bExtrapolate);
+}
+
+void FExtendedPoseHistory::ResetFuturePoses()
+{
+	FutureEntries.Reset();
+}
+
+void FExtendedPoseHistory::AddFuturePose(float SecondsInTheFuture, FCSPose<FCompactPose>& ComponentSpacePose, const FTransform& ComponentTransform)
+{
+	// we don't allow to add "past" or "present" poses to FutureEntries
+	check(SecondsInTheFuture > 0.f);
+	check(PoseHistory);	
+	const float SecondsAgo = -SecondsInTheFuture;
+	const int32 LowerBoundIdx = Algo::LowerBound(FutureEntries, SecondsAgo, [](const FPoseHistoryEntry& Entry, float Value) { return Value < Entry.Time; });
+	FutureEntries.InsertDefaulted_GetRef(LowerBoundIdx).Update(SecondsAgo, ComponentSpacePose, ComponentTransform, PoseHistory->GetBoneToTransformMap());
+}
+
+void FExtendedPoseHistory::DebugDraw(const UWorld* World, const USkeleton* Skeleton) const
+{
+#if ENABLE_DRAW_DEBUG
+	check(PoseHistory);
+
+	auto LerpColor = [](FColor A, FColor B, float T) -> FColor
+	{
@@ -264,9 +424,9 @@
 	};
 
 	TArray<FTransform> PrevGlobalTransforms;
-	for (int32 EntryIndex = 0; EntryIndex < Entries.Num(); ++EntryIndex)
+	for (int32 EntryIndex = 0; EntryIndex < FutureEntries.Num(); ++EntryIndex)
 	{
-		const FEntry& Entry = Entries[EntryIndex];
+		const FPoseHistoryEntry& Entry = FutureEntries[EntryIndex];
 		if (Entry.ComponentSpaceTransforms.IsEmpty())
 		{
 			PrevGlobalTransforms.Reset();
@@ -281,8 +441,8 @@
 		}
 		else
 		{
-			const float LerpFactor = float(EntryIndex - 1) / float(Entries.Num() - 1);
-			const FColor Color = LerpColor(FColorList::Red, FColorList::Orange, LerpFactor);
+			const float LerpFactor = float(EntryIndex - 1) / float(FutureEntries.Num() - 1);
+			const FColor Color = LerpColor(FColorList::Green, FColorList::Violet, LerpFactor);
 
 			for (int32 i = 0; i < Entry.ComponentSpaceTransforms.Num(); ++i)
 			{
@@ -294,12 +454,13 @@
 			}
 		}
 	}
+
+	PoseHistory->DebugDraw(World, Skeleton);
 #endif // ENABLE_DRAW_DEBUG
 }
 
 //////////////////////////////////////////////////////////////////////////
 // FPoseIndicesHistory
-
 void FPoseIndicesHistory::Update(const FSearchResult& SearchResult, float DeltaTime, float MaxTime)
 {
 	if (MaxTime > 0.f)
diff --git a/Engine/Plugins/Experimental/Animation/PoseSearch/Source/Runtime/Private/PoseSearchLibrary.cpp b/Engine/Plugins/Experimental/Animation/PoseSearch/Source/Runtime/Private/PoseSearchLibrary.cpp
--- a/Engine/Plugins/Experimental/Animation/PoseSearch/Source/Runtime/Private/PoseSearchLibrary.cpp
+++ b/Engine/Plugins/Experimental/Animation/PoseSearch/Source/Runtime/Private/PoseSearchLibrary.cpp
@@ -2,6 +2,7 @@
 
 #include "PoseSearch/PoseSearchLibrary.h"
 #include "Animation/AnimInstanceProxy.h"
+#include "Animation/AnimMontage.h"
 #include "Animation/AnimNode_Inertialization.h"
 #include "Animation/AnimNode_SequencePlayer.h"
 #include "Animation/AnimRootMotionProvider.h"
@@ -21,6 +22,12 @@
 
 #define LOCTEXT_NAMESPACE "PoseSearchLibrary"
 
+#if ENABLE_ANIM_DEBUG
+TAutoConsoleVariable<int32> CVarAnimMotionMatchDrawQueryEnable(TEXT("a.MotionMatch.DrawQuery.Enable"), 0, TEXT("Enable / Disable MotionMatch Draw Query"));
+TAutoConsoleVariable<int32> CVarAnimMotionMatchDrawMatchEnable(TEXT("a.MotionMatch.DrawMatch.Enable"), 0, TEXT("Enable / Disable MotionMatch Draw Match"));
+TAutoConsoleVariable<int32> CVarAnimMotionMatchDrawHistoryEnable(TEXT("a.MotionMatch.DrawHistory.Enable"), 0, TEXT("Enable / Disable MotionMatch Draw History"));
+#endif
+
 //////////////////////////////////////////////////////////////////////////
 // FMotionMatchingState
 
@@ -422,22 +429,43 @@ void UPoseSearchLibrary::MotionMatch(
 	bool& bLoop,
 	bool& bIsMirrored,
 	FVector& BlendParameters,
+	const UAnimationAsset* FutureAnimation,
+	float FutureAnimationStartTime,
+	float TimeToFutureAnimationStart,
 	const int DebugSessionUniqueIdentifier)
 {
+	using namespace UE::Anim;
 	using namespace UE::PoseSearch;
 
+	static constexpr float FiniteDelta = 1 / 60.0f;
+
 	SelectedAnimation = nullptr;
 	SelectedTime = 0.f;
 	bLoop = false;
 	bIsMirrored = false;
 	BlendParameters = FVector::ZeroVector;
 
+	if (FutureAnimationStartTime < FiniteDelta)
+	{
+		UE_LOG(LogPoseSearch, Warning, TEXT("UPoseSearchLibrary::MotionMatch - provided FutureAnimationStartTime (%f) is too small to be able to calculate velocities. Clamping it to minimum value of %f"), TimeToFutureAnimationStart, FiniteDelta);
+		FutureAnimationStartTime = FiniteDelta;
+	}
+
+	const float MinTimeToFutureAnimationStart = FiniteDelta + UE_KINDA_SMALL_NUMBER;
+	if (TimeToFutureAnimationStart < MinTimeToFutureAnimationStart)
+	{
+		UE_LOG(LogPoseSearch, Warning, TEXT("UPoseSearchLibrary::MotionMatch - provided TimeToFutureAnimationStart (%f) is too small. Clamping it to minimum value of %f"), TimeToFutureAnimationStart, MinTimeToFutureAnimationStart);
+		TimeToFutureAnimationStart = MinTimeToFutureAnimationStart;
+	}
+
 	if (Searchable)
 	{
 		// Build the search context
 		FSearchContext SearchContext;
+		SearchContext.DesiredPermutationTimeOffset = TimeToFutureAnimationStart;
 		SearchContext.Trajectory = &Trajectory;
 
+		FExtendedPoseHistory ExtendedPoseHistory;
 		if (AnimInstance)
 		{
 			if (IAnimClassInterface* AnimBlueprintClass = IAnimClassInterface::GetFromClass(AnimInstance->GetClass()))
@@ -446,19 +474,107 @@ void UPoseSearchLibrary::MotionMatch(
 				{
 					if (const FAnimNode_PoseSearchHistoryCollector_Base* PoseHistoryNode = TagSubsystem->FindNodeByTag<FAnimNode_PoseSearchHistoryCollector_Base>(PoseHistoryName, AnimInstance))
 					{
-						SearchContext.History = &PoseHistoryNode->GetPoseHistory();
+						const FPoseHistory* PoseHistory = &PoseHistoryNode->GetPoseHistory();
+						ExtendedPoseHistory.Init(PoseHistory);
+						SearchContext.History = PoseHistory;
 					}
 				}
 			}
 
 			if (!SearchContext.History)
 			{
-				UE_LOG(LogPoseSearch, Warning, TEXT("UPoseSearchLibrary::MotionMatch - Couldn't find pose history with name '%s'"), *PoseHistoryName.ToString());
+				if (FutureAnimation)
+				{
+					UE_LOG(LogPoseSearch, Error, TEXT("UPoseSearchLibrary::MotionMatch - Couldn't find pose history with name '%s'. FutureAnimation search will not be performed"), *PoseHistoryName.ToString());
+				}
+				else
+				{
+					UE_LOG(LogPoseSearch, Warning, TEXT("UPoseSearchLibrary::MotionMatch - Couldn't find pose history with name '%s'"), *PoseHistoryName.ToString());
+				}
+			}
+			else if (FutureAnimation)
+			{
+				const FBoneContainer& BoneContainer = AnimInstance->GetRequiredBonesOnAnyThread();
+
+				FBlendSpaceSampler BlendSpaceSampler;
+				FSequenceBaseSampler SequenceBaseSampler;
+				FAnimMontageSampler AnimMontageSampler;
+				const FAssetSamplerBase* AssetSamplerBase = nullptr;
+				if (const UBlendSpace* BlendSpace = Cast<UBlendSpace>(FutureAnimation))
+				{
+					FBlendSpaceSampler::FInput Input;
+					Input.BoneContainer = BoneContainer;
+					Input.BlendSpace = BlendSpace;
+					// @todo... add BlendParameters
+					//Input.BlendParameters = BlendParameters;
+					BlendSpaceSampler.Init(Input);
+					AssetSamplerBase = &BlendSpaceSampler;
+				}
+				else if (const UAnimMontage* AnimMontage = Cast<UAnimMontage>(FutureAnimation))
+				{
+					FAnimMontageSampler::FInput Input;
+					Input.AnimMontage = AnimMontage;
+					AnimMontageSampler.Init(Input);
+					AssetSamplerBase = &AnimMontageSampler;
+				}
+				else if (const UAnimSequenceBase* SequenceBase = Cast<UAnimSequenceBase>(FutureAnimation))
+				{
+					FSequenceBaseSampler::FInput Input;
+					Input.SequenceBase = SequenceBase;
+					SequenceBaseSampler.Init(Input);
+					AssetSamplerBase = &SequenceBaseSampler;
+				}
+				else
+				{
+					UE_LOG(LogPoseSearch, Error, TEXT("UPoseSearchLibrary::MotionMatch - unsupported FutureAnimation '%s'"), *GetNameSafe(FutureAnimation));
+				}
+
+				if (AssetSamplerBase)
+				{
+					// ExtendedPoseHistory will hold future poses to match AssetSamplerBase (at FutureAnimationStartTime) TimeToFutureAnimationStart seconds in the future
+					SearchContext.History = &ExtendedPoseHistory;
+
+					FCompactPose Pose;
+					FBlendedCurve UnusedCurve;
+					FStackAttributeContainer UnusedAtrribute;
+					FAnimationPoseData AnimPoseData = { Pose, UnusedCurve, UnusedAtrribute };
+
+					UnusedCurve.InitFrom(BoneContainer);
+					Pose.SetBoneContainer(&BoneContainer);
+
+					// extracting 2 poses to be able to calculate velocities
+					for (int i = 0; i < 2; ++i)
+					{
+						const float ExtractionTime = FutureAnimationStartTime + (i - 1) * FiniteDelta;
+						const float FutureAnimationTime = TimeToFutureAnimationStart + (i - 1) * FiniteDelta;
+
+						FDeltaTimeRecord DeltaTimeRecord;
+						DeltaTimeRecord.Set(ExtractionTime - FiniteDelta, FiniteDelta);
+						FAnimExtractContext ExtractionCtx(double(ExtractionTime), false, DeltaTimeRecord, false);
+
+						AssetSamplerBase->ExtractPose(ExtractionCtx, AnimPoseData);
+
+						FCSPose<FCompactPose> ComponentSpacePose;
+						ComponentSpacePose.InitPose(Pose);
+
+						const FTrajectorySample TrajectorySample = Trajectory.GetSampleAtTime(ExtractionTime);
+						const FTransform& ComponentTransform = AnimInstance->GetOwningComponent()->GetComponentTransform();
+						const FTransform FutureComponentTransform = TrajectorySample.Transform * ComponentTransform;
+
+						ExtendedPoseHistory.AddFuturePose(FutureAnimationTime, ComponentSpacePose, FutureComponentTransform);
+					}
+
+#if ENABLE_DRAW_DEBUG && ENABLE_ANIM_DEBUG
+					if (CVarAnimMotionMatchDrawHistoryEnable.GetValueOnAnyThread())
+					{
+						ExtendedPoseHistory.DebugDraw(AnimInstance->GetOwningComponent()->GetWorld(), BoneContainer.GetSkeletonAsset());
+					}
+#endif // ENABLE_DRAW_DEBUG && ENABLE_ANIM_DEBUG
+				}
 			}
 		}
-		
-		// @todo: finish set up SearchContext by exposing or calculating additional members
 
+		// @todo: finish set up SearchContext by exposing or calculating additional members
 		FSearchResult SearchResult = Searchable->Search(SearchContext);
 		if (SearchResult.IsValid())
 		{
@@ -475,6 +591,31 @@ void UPoseSearchLibrary::MotionMatch(
 
 		if (AnimInstance)
 		{
+#if ENABLE_DRAW_DEBUG && ENABLE_ANIM_DEBUG
+			if (SearchResult.IsValid())
+			{
+				UE::PoseSearch::FDebugDrawParams DrawParams;
+				DrawParams.RootTransform = AnimInstance->GetOwningComponent()->GetComponentTransform();
+				DrawParams.Database = SearchResult.Database.Get();
+				DrawParams.World = AnimInstance->GetOwningComponent()->GetWorld();
+				DrawParams.DefaultLifeTime = 0.0f;
+
+				if (DrawParams.CanDraw())
+				{
+					if (CVarAnimMotionMatchDrawMatchEnable.GetValueOnAnyThread())
+					{
+						DrawFeatureVector(DrawParams, SearchResult.PoseIdx);
+					}
+
+					if (CVarAnimMotionMatchDrawQueryEnable.GetValueOnAnyThread())
+					{
+						EnumAddFlags(DrawParams.Flags, EDebugDrawFlags::DrawQuery);
+						DrawFeatureVector(DrawParams, SearchResult.ComposedQuery.GetValues());
+					}
+				}
+			}
+#endif // ENABLE_DRAW_DEBUG && ENABLE_ANIM_DEBUG
+
 #if UE_POSE_SEARCH_TRACE_ENABLED
 			TraceMotionMatchingState(Searchable, SearchContext, SearchResult, FSearchResult(), 0.f, FTransform::Identity, AnimInstance, DebugSessionUniqueIdentifier, AnimInstance->GetDeltaSeconds(), true);
 #endif // UE_POSE_SEARCH_TRACE_ENABLED
diff --git a/Engine/Plugins/Experimental/Animation/PoseSearch/Source/Runtime/Public/PoseSearch/AnimNode_PoseSearchHistoryCollector.h b/Engine/Plugins/Experimental/Animation/PoseSearch/Source/Runtime/Public/PoseSearch/AnimNode_PoseSearchHistoryCollector.h
--- a/Engine/Plugins/Experimental/Animation/PoseSearch/Source/Runtime/Public/PoseSearch/AnimNode_PoseSearchHistoryCollector.h
+++ b/Engine/Plugins/Experimental/Animation/PoseSearch/Source/Runtime/Public/PoseSearch/AnimNode_PoseSearchHistoryCollector.h
@@ -26,7 +26,6 @@ public:
 
 	// FAnimNode_Base interface
 	virtual void CacheBones_AnyThread(const FAnimationCacheBonesContext& Context) override;
-	virtual void Update_AnyThread(const FAnimationUpdateContext& Context) override;
 	virtual bool HasPreUpdate() const override;
 	virtual void PreUpdate(const UAnimInstance* InAnimInstance) override;
 	// End of FAnimNode_Base interface
@@ -76,4 +75,4 @@ public:
 	virtual void Update_AnyThread(const FAnimationUpdateContext& Context) override;
 	virtual void GatherDebugData(FNodeDebugData& DebugData) override;
 	// End of FAnimNode_Base interface
-};
\ No newline at end of file
+};
diff --git a/Engine/Plugins/Experimental/Animation/PoseSearch/Source/Runtime/Public/PoseSearch/PoseSearchAssetIndexer.h b/Engine/Plugins/Experimental/Animation/PoseSearch/Source/Runtime/Public/PoseSearch/PoseSearchAssetIndexer.h
--- a/Engine/Plugins/Experimental/Animation/PoseSearch/Source/Runtime/Public/PoseSearch/PoseSearchAssetIndexer.h
+++ b/Engine/Plugins/Experimental/Animation/PoseSearch/Source/Runtime/Public/PoseSearch/PoseSearchAssetIndexer.h
@@ -30,7 +30,7 @@ struct FAssetIndexingContext
 	FFloatInterval RequestedSamplingRange = FFloatInterval(0.0f, 0.0f);
 };
 
-class POSESEARCH_API FAssetIndexer
+class FAssetIndexer
 {
 public:
 	struct FStats
diff --git a/Engine/Plugins/Experimental/Animation/PoseSearch/Source/Runtime/Public/PoseSearch/PoseSearchAssetSampler.h b/Engine/Plugins/Experimental/Animation/PoseSearch/Source/Runtime/Public/PoseSearch/PoseSearchAssetSampler.h
--- a/Engine/Plugins/Experimental/Animation/PoseSearch/Source/Runtime/Public/PoseSearch/PoseSearchAssetSampler.h
+++ b/Engine/Plugins/Experimental/Animation/PoseSearch/Source/Runtime/Public/PoseSearch/PoseSearchAssetSampler.h
@@ -23,8 +23,6 @@ struct POSESEARCH_API FPoseSearchExtrapolationParameters
 	float SampleTime = 0.05f;
 };
 
-#if WITH_EDITOR
-
 struct FAnimationPoseData;
 struct FAnimExtractContext;
 class UAnimationAsset;
@@ -36,25 +34,6 @@ class UMirrorDataTable;
 namespace UE::PoseSearch
 {
 
-struct POSESEARCH_API FAssetSamplingContext
-{
-	// Time delta used for computing pose derivatives
-	static constexpr float FiniteDelta = 1 / 60.0f;
-
-	// Mirror data table pointer copied from Schema for convenience
-	TObjectPtr<const UMirrorDataTable> MirrorDataTable;
-
-	// Compact pose format of Mirror Bone Map
-	TCustomBoneIndexArray<FCompactPoseBoneIndex, FCompactPoseBoneIndex> CompactPoseMirrorBones;
-
-	// Pre-calculated component space rotations of reference pose, which allows mirror to work with any joint orientation
-	// Only initialized and used when a mirroring table is specified
-	TCustomBoneIndexArray<FQuat, FCompactPoseBoneIndex> ComponentSpaceRefRotations;
-
-	void Init(const UMirrorDataTable* InMirrorDataTable, const FBoneContainer& BoneContainer);
-	FTransform MirrorTransform(const FTransform& Transform) const;
-};
-
 /**
  * Helper interface for sampling data from animation assets
  */
@@ -153,7 +132,7 @@ struct POSESEARCH_API FAnimMontageSampler : public FAssetSamplerBase
 	struct FInput
 	{
 		// @todo: add support for SlotName / multiple SlotAnimTracks
-		TWeakObjectPtr<UAnimMontage> AnimMontage;
+		TWeakObjectPtr<const UAnimMontage> AnimMontage;
 		FPoseSearchExtrapolationParameters ExtrapolationParameters;
 	} Input;
 
@@ -177,4 +156,3 @@ private:
 
 } // namespace UE::PoseSearch
 
-#endif // WITH_EDITOR
diff --git a/Engine/Plugins/Experimental/Animation/PoseSearch/Source/Runtime/Public/PoseSearch/PoseSearchContext.h b/Engine/Plugins/Experimental/Animation/PoseSearch/Source/Runtime/Public/PoseSearch/PoseSearchContext.h
--- a/Engine/Plugins/Experimental/Animation/PoseSearch/Source/Runtime/Public/PoseSearch/PoseSearchContext.h
+++ b/Engine/Plugins/Experimental/Animation/PoseSearch/Source/Runtime/Public/PoseSearch/PoseSearchContext.h
@@ -14,8 +14,8 @@ class UPoseSearchSchema;
 namespace UE::PoseSearch
 {
 
-struct FPoseHistory;
 struct FPoseIndicesHistory;
+struct IPoseHistory;
 
 enum class EDebugDrawFlags : uint32
 {
@@ -151,14 +151,13 @@ POSESEARCH_API void DrawFeatureVector(FDebugDrawParams& DrawParams, int32 PoseId
 struct POSESEARCH_API FSearchContext
 {
 	EPoseSearchBooleanRequest QueryMirrorRequest = EPoseSearchBooleanRequest::Indifferent;
-	const FPoseHistory* History = nullptr;
+	const IPoseHistory* History = nullptr;
 	const FTrajectorySampleRange* Trajectory = nullptr;
 	FSearchResult CurrentResult;
 	float PoseJumpThresholdTime = 0.f;
 	bool bForceInterrupt = false;
 	// can the continuing pose advance? (if not we skip evaluating it)
 	bool bCanAdvance = true;
-	float DesiredPermutationTimeOffset = 0.f;
 
 	FQuat GetSampleRotation(float SampleTimeOffset, const UPoseSearchSchema* Schema, int8 SchemaSampleBoneIdx = RootSchemaBoneIdx, int8 SchemaOriginBoneIdx = RootSchemaBoneIdx, bool bUseHistoryRoot = false);
 	FVector GetSamplePosition(float SampleTimeOffset, const UPoseSearchSchema* Schema, int8 SchemaSampleBoneIdx = RootSchemaBoneIdx, int8 SchemaOriginBoneIdx = RootSchemaBoneIdx, bool bUseHistoryRoot = false);
@@ -179,6 +178,11 @@ struct POSESEARCH_API FSearchContext
 	TConstArrayView<float> GetCurrentResultPoseVector() const;
 	TConstArrayView<float> GetCurrentResultNextPoseVector() const;
 
+	void SetPermutationTimeOffsets(float InPermutationSampleTimeOffset, float InPermutationOriginTimeOffset);
+	void ResetPermutationTimeOffsets();
+	
+	float DesiredPermutationTimeOffset = 0.f;
+
 	const FPoseIndicesHistory* PoseIndicesHistory = nullptr;
 
 private:
@@ -191,6 +195,10 @@ private:
 	TArray<FPoseSearchFeatureVectorBuilder, TInlineAllocator<8>> CachedQueries;
 
 	float CurrentBestTotalCost = MAX_flt;
+	
+	// time offsets controlled by sampling data permutations
+	float PermutationSampleTimeOffset = 0.f;
+	float PermutationOriginTimeOffset = 0.f;
 
 #if UE_POSE_SEARCH_TRACE_ENABLED
 
diff --git a/Engine/Plugins/Experimental/Animation/PoseSearch/Source/Runtime/Public/PoseSearch/PoseSearchHistory.h b/Engine/Plugins/Experimental/Animation/PoseSearch/Source/Runtime/Public/PoseSearch/PoseSearchHistory.h
--- a/Engine/Plugins/Experimental/Animation/PoseSearch/Source/Runtime/Public/PoseSearch/PoseSearchHistory.h
+++ b/Engine/Plugins/Experimental/Animation/PoseSearch/Source/Runtime/Public/PoseSearch/PoseSearchHistory.h
@@ -14,37 +14,75 @@ namespace UE::PoseSearch
 {
 
 struct FSearchResult;
+typedef uint16 FComponentSpaceTransformIndex;
+typedef TPair<FBoneIndexType, FComponentSpaceTransformIndex> FBoneToTransformPair;
+typedef TMap<FBoneIndexType, FComponentSpaceTransformIndex> FBoneToTransformMap;
 
-struct POSESEARCH_API FPoseHistory
+struct POSESEARCH_API IPoseHistory
+{
+	virtual ~IPoseHistory() {}
+	virtual float GetSampleTimeInterval() const = 0;
+	virtual bool GetComponentSpaceTransformAtTime(float Time, FBoneIndexType BoneIndexType, FTransform& OutBoneTransform, bool bExtrapolate = true) const = 0;
+	virtual void GetRootTransformAtTime(float Time, FTransform& OutRootTransform, bool bExtrapolate = true) const = 0;
+};
+
+struct FPoseHistoryEntry
+{
+	FTransform RootTransform;
+	TArray<FTransform> ComponentSpaceTransforms;
+	float Time = 0.f;
+
+	void Update(float InTime, FCSPose<FCompactPose>& ComponentSpacePose, const FTransform& ComponentTransform, const FBoneToTransformMap& BoneToTransformMap);
+};
+
+typedef TRingBuffer<FPoseHistoryEntry> FPoseHistoryEntries;
+typedef TArray<FPoseHistoryEntry> FPoseHistoryFutureEntries;
+
+struct FPoseHistory : public IPoseHistory
 {
 	void Init(int32 InNumPoses, float InTimeHorizon, const TArray<FBoneIndexType>& RequiredBones);
-	void Update(float SecondsElapsed, FCSPose<FCompactPose>& ComponentSpacePose, FTransform ComponentTransform);
-	float GetSampleTimeInterval() const;
+	void Update(float SecondsElapsed, FCSPose<FCompactPose>& ComponentSpacePose, const FTransform& ComponentTransform);
 	float GetTimeHorizon() const { return TimeHorizon; }
-	bool GetComponentSpaceTransformAtTime(float Time, FBoneIndexType BoneIndexType, FTransform& OutBoneTransform, bool bExtrapolate = true) const;
-	void GetRootTransformAtTime(float Time, FTransform& OutRootTransform, bool bExtrapolate = true) const;
 	void DebugDraw(const UWorld* World, const USkeleton* Skeleton) const;
+	const FBoneToTransformMap& GetBoneToTransformMap() const { return BoneToTransformMap; }
+	const FPoseHistoryEntries& GetEntries() const { return Entries; }
 
-private:
+	// IPoseHistory interface
+	virtual float GetSampleTimeInterval() const override;
+	virtual bool GetComponentSpaceTransformAtTime(float Time, FBoneIndexType BoneIndexType, FTransform& OutBoneTransform, bool bExtrapolate = true) const override;
+	virtual void GetRootTransformAtTime(float Time, FTransform& OutRootTransform, bool bExtrapolate = true) const override;
+	// End of IPoseHistory interface
 
-	struct FEntry
-	{
-		FTransform RootTransform;
-		TArray<FTransform> ComponentSpaceTransforms;
-		float Time = 0.f;
-	};
-
-	typedef uint16 FComponentSpaceTransformIndex;
-	TMap<FBoneIndexType, FComponentSpaceTransformIndex> BoneToTransformMap;
-	TRingBuffer<FEntry> Entries;
+private:
+	FBoneToTransformMap BoneToTransformMap;
+	FPoseHistoryEntries Entries;
 	float TimeHorizon = 0.f;
 };
 
+struct FExtendedPoseHistory : public IPoseHistory
+{
+	void Init(const FPoseHistory* InPoseHistory);
+
+	// IPoseHistory interface
+	virtual float GetSampleTimeInterval() const override;
+	virtual bool GetComponentSpaceTransformAtTime(float Time, FBoneIndexType BoneIndexType, FTransform& OutBoneTransform, bool bExtrapolate = true) const override;
+	virtual void GetRootTransformAtTime(float Time, FTransform& OutRootTransform, bool bExtrapolate = true) const override;
+	// End of IPoseHistory interface
+
+	void ResetFuturePoses();
+	void AddFuturePose(float SecondsInTheFuture, FCSPose<FCompactPose>& ComponentSpacePose, const FTransform& ComponentTransform);
+	void DebugDraw(const UWorld* World, const USkeleton* Skeleton) const;
+
+private:
+	const FPoseHistory* PoseHistory = nullptr;
+	FPoseHistoryFutureEntries FutureEntries;
+};
+
 class IPoseHistoryProvider : public UE::Anim::IGraphMessage
 {
 	DECLARE_ANIMGRAPH_MESSAGE(IPoseHistoryProvider);
 public:
-	virtual const FPoseHistory& GetPoseHistory() const = 0;
+	virtual const IPoseHistory& GetPoseHistory() const = 0;
 };
 
 struct FHistoricalPoseIndex
diff --git a/Engine/Plugins/Experimental/Animation/PoseSearch/Source/Runtime/Public/PoseSearch/PoseSearchLibrary.h b/Engine/Plugins/Experimental/Animation/PoseSearch/Source/Runtime/Public/PoseSearch/PoseSearchLibrary.h
--- a/Engine/Plugins/Experimental/Animation/PoseSearch/Source/Runtime/Public/PoseSearch/PoseSearchLibrary.h
+++ b/Engine/Plugins/Experimental/Animation/PoseSearch/Source/Runtime/Public/PoseSearch/PoseSearchLibrary.h
@@ -164,6 +164,9 @@ public:
 	* @param bLoop							Output selected animation looping state
 	* @param bIsMirrored					Output selected animation mirror state
 	* @param BlendParameters				Output selected animation blend space parameters (if SelectedAnimation is a blend space)
+	* @param FutureAnimation				Input animation we want to match after TimeToFutureAnimationStart seconds
+	* @param FutureAnimationStartTime		Input start time for the first pose of FutureAnimation
+	* @param TimeToFutureAnimationStart		Input time in seconds before start playing FutureAnimation (from FutureAnimationStartTime seconds)
 	* @param DebugSessionUniqueIdentifier	Input unique identifier used to identify TraceMotionMatchingState (rewind debugger / pose search debugger) session. Similarly the MM node uses Context.GetCurrentNodeId()
 	*/
 	UFUNCTION(BlueprintPure, Category = "Animation|Pose Search", meta = (BlueprintThreadSafe, Keywords = "PoseMatch"))
@@ -177,6 +180,9 @@ public:
 		bool& bLoop,
 		bool& bIsMirrored,
 		FVector& BlendParameters,
+		const UAnimationAsset* FutureAnimation = nullptr,
+		float FutureAnimationStartTime = 0.f,
+		float TimeToFutureAnimationStart = 0.f,
 		const int DebugSessionUniqueIdentifier = 6174);
 };
 
