motion matching - debug lines cleanup amd UPoseSearchLibrary::MotionMatch threading issue fix

diff --git a/Engine/Plugins/Experimental/Animation/PoseSearch/Source/Editor/Private/PoseSearchDatabasePreviewScene.cpp b/Engine/Plugins/Experimental/Animation/PoseSearch/Source/Editor/Private/PoseSearchDatabasePreviewScene.cpp
--- a/Engine/Plugins/Experimental/Animation/PoseSearch/Source/Editor/Private/PoseSearchDatabasePreviewScene.cpp
+++ b/Engine/Plugins/Experimental/Animation/PoseSearch/Source/Editor/Private/PoseSearchDatabasePreviewScene.cpp
@@ -67,21 +67,8 @@ namespace UE::PoseSearch
 				{
 					if (UDebugSkelMeshComponent* Mesh = PreviewActor.GetDebugSkelMeshComponent())
 					{
-						UE::PoseSearch::FDebugDrawParams DrawParams;
-						DrawParams.RootTransform = Mesh->GetComponentTransform();
-						DrawParams.Database = Database;
-						DrawParams.World = GetWorld();
-						DrawParams.DefaultLifeTime = 0.f;
-						DrawParams.PointSize = 5.f;
-						DrawParams.Mesh = Mesh;
-
-						if (ViewModel->IsPoseFeaturesDrawMode(EFeaturesDrawMode::Detailed))
-						{
-							EnumAddFlags(DrawParams.Flags, UE::PoseSearch::EDebugDrawFlags::DrawBoneNames);
-						}
-
-						EnumAddFlags(DrawParams.Flags, UE::PoseSearch::EDebugDrawFlags::DrawFast);
-						DrawFeatureVector(DrawParams, PreviewActor.CurrentPoseIndex);
+						UE::PoseSearch::FDebugDrawParams DrawParams(GetWorld(), Mesh, Database);
+						DrawParams.DrawFeatureVector(PreviewActor.CurrentPoseIndex);
 					}
 				}
 			}
diff --git a/Engine/Plugins/Experimental/Animation/PoseSearch/Source/Editor/Private/PoseSearchDebuggerReflection.h b/Engine/Plugins/Experimental/Animation/PoseSearch/Source/Editor/Private/PoseSearchDebuggerReflection.h
--- a/Engine/Plugins/Experimental/Animation/PoseSearch/Source/Editor/Private/PoseSearchDebuggerReflection.h
+++ b/Engine/Plugins/Experimental/Animation/PoseSearch/Source/Editor/Private/PoseSearchDebuggerReflection.h
@@ -15,12 +15,6 @@ struct FPoseSearchDebuggerFeatureDrawOptions
 
 	UPROPERTY(EditAnywhere, Category="Draw Options")
     bool bDisable = false;
-
-	UPROPERTY(EditAnywhere, Category = "Draw Options", Meta = (EditCondition = "!bDisable"))
-	bool bDrawBoneNames = false;
-
-	UPROPERTY(EditAnywhere, Category = "Draw Options", Meta = (EditCondition = "!bDisable"))
-	bool bDrawSampleLabels = false;
 };
 
 /**
diff --git a/Engine/Plugins/Experimental/Animation/PoseSearch/Source/Editor/Private/PoseSearchDebuggerView.cpp b/Engine/Plugins/Experimental/Animation/PoseSearch/Source/Editor/Private/PoseSearchDebuggerView.cpp
--- a/Engine/Plugins/Experimental/Animation/PoseSearch/Source/Editor/Private/PoseSearchDebuggerView.cpp
+++ b/Engine/Plugins/Experimental/Animation/PoseSearch/Source/Editor/Private/PoseSearchDebuggerView.cpp
@@ -354,48 +354,22 @@ void SDebuggerView::DrawFeatures(
 	const USkinnedMeshComponent* Mesh
 ) const
 {
-	auto SetDrawFlags = [](FDebugDrawParams& InDrawParams, const FPoseSearchDebuggerFeatureDrawOptions& Options)
-	{
-		InDrawParams.Flags = EDebugDrawFlags::None;
-		if (!Options.bDisable)
-		{
-			if (Options.bDrawBoneNames)
-			{
-				EnumAddFlags(InDrawParams.Flags, EDebugDrawFlags::DrawBoneNames);
-			}
-			
-			if (Options.bDrawSampleLabels)
-			{
-				EnumAddFlags(InDrawParams.Flags, EDebugDrawFlags::DrawSampleLabels);
-			}
-		}
-	};
-
 	const TObjectPtr<UPoseSearchDebuggerReflection> Reflection = DetailsView->GetReflection();
 
 	// Draw query vector
+	if (!Reflection || !Reflection->QueryDrawOptions.bDisable)
 	{
 		const UPoseSearchDatabase* CurrentDatabase = ViewModel.Get()->GetCurrentDatabase();
 		if (CurrentDatabase)
 		{
-			// Set shared state
-			FDebugDrawParams DrawParams;
-			DrawParams.World = &DebuggerWorld;
-			DrawParams.RootTransform = Transform;
-			DrawParams.DefaultLifeTime = 0.0f; // Single frame render
-			DrawParams.Mesh = Mesh;
-
 			for (const FTraceMotionMatchingStateDatabaseEntry& DbEntry : State.DatabaseEntries)
 			{
 				const UPoseSearchDatabase* Database = FTraceMotionMatchingState::GetObjectFromId<UPoseSearchDatabase>(DbEntry.DatabaseId);
 				if (Database && Database == CurrentDatabase && DbEntry.QueryVector.Num() == Database->Schema->SchemaCardinality && 
 					FAsyncPoseSearchDatabasesManagement::RequestAsyncBuildIndex(CurrentDatabase, ERequestAsyncBuildFlag::ContinueRequest))
 				{
-					DrawParams.Database = CurrentDatabase;
-					SetDrawFlags(DrawParams, Reflection ? Reflection->QueryDrawOptions : FPoseSearchDebuggerFeatureDrawOptions());
-					EnumAddFlags(DrawParams.Flags, EDebugDrawFlags::DrawQuery);
-					DrawFeatureVector(DrawParams, DbEntry.QueryVector);
-					EnumRemoveFlags(DrawParams.Flags, EDebugDrawFlags::DrawQuery);
+					FDebugDrawParams DrawParams(&DebuggerWorld, Mesh, CurrentDatabase, EDebugDrawFlags::DrawQuery);
+					DrawParams.DrawFeatureVector(DbEntry.QueryVector);
 					break;
 				}
 			}
@@ -403,25 +377,19 @@ void SDebuggerView::DrawFeatures(
 	}
 
 	// Draw selected poses
+	if (!Reflection || !Reflection->QueryDrawOptions.bDisable)
 	{
 		const TSharedPtr<SListView<TSharedRef<FDebuggerDatabaseRowData>>>& DatabaseRows = DatabaseView->GetDatabaseRows();
 		TArray<TSharedRef<FDebuggerDatabaseRowData>> SelectedRows = DatabaseRows->GetSelectedItems();
 	
-		// Red for non-active database view
-		FDebugDrawParams DrawParams;
-		DrawParams.World = &DebuggerWorld;
-		DrawParams.RootTransform = Transform;
-		DrawParams.DefaultLifeTime = 0.0f; // Single frame render
-		DrawParams.Mesh = Mesh;
-		SetDrawFlags(DrawParams, Reflection ? Reflection->SelectedPoseDrawOptions : FPoseSearchDebuggerFeatureDrawOptions());
-
 		// Draw any selected database vectors
 		for (const TSharedRef<FDebuggerDatabaseRowData>& Row : SelectedRows)
 		{
-			DrawParams.Database = Row->SourceDatabase.Get();
-			if (DrawParams.Database && FAsyncPoseSearchDatabasesManagement::RequestAsyncBuildIndex(DrawParams.Database, ERequestAsyncBuildFlag::ContinueRequest))
+			const UPoseSearchDatabase* RowDatabase = Row->SourceDatabase.Get();
+			if (FAsyncPoseSearchDatabasesManagement::RequestAsyncBuildIndex(RowDatabase, ERequestAsyncBuildFlag::ContinueRequest))
 			{
-				DrawFeatureVector(DrawParams, Row->PoseIdx);
+				FDebugDrawParams DrawParams(&DebuggerWorld, Mesh, RowDatabase);
+				DrawParams.DrawFeatureVector(Row->PoseIdx);
 			}
 		}
 	}
@@ -440,13 +408,8 @@ void SDebuggerView::DrawFeatures(
 			if (Database && FAsyncPoseSearchDatabasesManagement::RequestAsyncBuildIndex(Database, ERequestAsyncBuildFlag::ContinueRequest))
 			{
 				// Use the motion-matching state's pose idx, as the active row may be update-throttled at this point
-				FDebugDrawParams DrawParams;
-				DrawParams.World = &DebuggerWorld;
-				DrawParams.RootTransform = Transform;
-				DrawParams.DefaultLifeTime = 0.0f; // Single frame render
-				DrawParams.Mesh = Mesh;
-				DrawParams.Database = Database;
-				DrawFeatureVector(DrawParams, ActiveRows[0]->PoseIdx);
+				FDebugDrawParams DrawParams(&DebuggerWorld, Mesh, Database);
+				DrawParams.DrawFeatureVector(ActiveRows[0]->PoseIdx);
 			}
 		}
 	}
@@ -464,13 +427,8 @@ void SDebuggerView::DrawFeatures(
 			const UPoseSearchDatabase* Database = ContinuingRows[0]->SourceDatabase.Get();
 			if (Database && FAsyncPoseSearchDatabasesManagement::RequestAsyncBuildIndex(Database, ERequestAsyncBuildFlag::ContinueRequest))
 			{
-				FDebugDrawParams DrawParams;
-				DrawParams.World = &DebuggerWorld;
-				DrawParams.RootTransform = Transform;
-				DrawParams.DefaultLifeTime = 0.0f; // Single frame render
-				DrawParams.Mesh = Mesh;
-				DrawParams.Database = Database;
-				DrawFeatureVector(DrawParams, ContinuingRows[0]->PoseIdx);
+				FDebugDrawParams DrawParams(&DebuggerWorld, Mesh, Database);
+				DrawParams.DrawFeatureVector(ContinuingRows[0]->PoseIdx);
 			}
 		}
 	}
diff --git a/Engine/Plugins/Experimental/Animation/PoseSearch/Source/Runtime/Private/AnimNode_MotionMatching.cpp b/Engine/Plugins/Experimental/Animation/PoseSearch/Source/Runtime/Private/AnimNode_MotionMatching.cpp
--- a/Engine/Plugins/Experimental/Animation/PoseSearch/Source/Runtime/Private/AnimNode_MotionMatching.cpp
+++ b/Engine/Plugins/Experimental/Animation/PoseSearch/Source/Runtime/Private/AnimNode_MotionMatching.cpp
@@ -40,10 +40,6 @@ void FAnimNode_MotionMatching::Evaluate_AnyThread(FPoseContext& Output)
 
 	Source.Evaluate(Output);
 
-#if WITH_EDITORONLY_DATA
-	bWasEvaluated = true;
-#endif
-
 #if UE_POSE_SEARCH_TRACE_ENABLED
 	MotionMatchingState.RootMotionTransformDelta = FTransform::Identity;
 
@@ -135,27 +131,11 @@ void FAnimNode_MotionMatching::UpdateAssetPlayer(const FAnimationUpdateContext&
 	BlendStackNode.UpdatePlayRate(MotionMatchingState.WantedPlayRate);
 
 	Source.Update(Context);
-}
 
-bool FAnimNode_MotionMatching::HasPreUpdate() const
-{
-#if WITH_EDITORONLY_DATA
-	return true;
-#else
-	return false;
-#endif
-}
 
-void FAnimNode_MotionMatching::PreUpdate(const UAnimInstance* InAnimInstance)
-{
 #if WITH_EDITORONLY_DATA
-	using namespace UE::PoseSearch;
-
-	if (bWasEvaluated && bDebugDraw)
+	if (bDebugDraw)
 	{
-		USkeletalMeshComponent* SkeletalMeshComponent = InAnimInstance->GetSkelMeshComponent();
-		check(SkeletalMeshComponent);
-
 		const UE::PoseSearch::FSearchResult& CurResult = MotionMatchingState.CurrentSearchResult;
 
 #if WITH_EDITOR
@@ -166,61 +146,52 @@ void FAnimNode_MotionMatching::PreUpdate(const UAnimInstance* InAnimInstance)
 		else
 #endif // WITH_EDITOR
 		{
-			UE::PoseSearch::FDebugDrawParams DrawParams;
-			DrawParams.RootTransform = SkeletalMeshComponent->GetComponentTransform();
-			DrawParams.Database = CurResult.Database.Get();
-			DrawParams.World = SkeletalMeshComponent->GetWorld();
-			DrawParams.DefaultLifeTime = 0.0f;
+			if (bDebugDrawMatch)
+			{
+				UE::PoseSearch::FDebugDrawParams DrawParams(Context.AnimInstanceProxy, CurResult.Database.Get());
+				DrawParams.DrawFeatureVector(CurResult.PoseIdx);
+			}
 
-			if (DrawParams.CanDraw())
+			if (bDebugDrawQuery)
 			{
-				if (bDebugDrawMatch)
-				{
-					DrawFeatureVector(DrawParams, CurResult.PoseIdx);
-				}
-
-				if (bDebugDrawQuery)
-				{
-					EnumAddFlags(DrawParams.Flags, EDebugDrawFlags::DrawQuery);
-					DrawFeatureVector(DrawParams, CurResult.ComposedQuery.GetValues());
-				}
-
-				if (DrawParams.Database->PoseSearchMode == EPoseSearchMode::PCAKDTree_Compare)
-				{
-					FDebugFloatHistory& C = MotionMatchingState.SearchCostHistoryContinuing;
-					FDebugFloatHistory& B = MotionMatchingState.SearchCostHistoryBruteForce;
-					FDebugFloatHistory& K = MotionMatchingState.SearchCostHistoryKDTree;
-
-					C.AddSample(CurResult.ContinuingPoseCost.IsValid() ? CurResult.ContinuingPoseCost.GetTotalCost() : C.MaxValue);
-					B.AddSample(CurResult.BruteForcePoseCost.IsValid() ? CurResult.BruteForcePoseCost.GetTotalCost() : B.MaxValue);
-					K.AddSample(CurResult.PoseCost.IsValid() ? CurResult.PoseCost.GetTotalCost() : K.MaxValue);
-
-					// making SearchCostHistoryKDTree and SearchCostHistoryBruteForce min max consistent
-					const float MinValue = FMath::Min(C.MinValue, FMath::Min(B.MinValue, K.MinValue));
-					const float MaxValue = FMath::Max(C.MaxValue, FMath::Max(B.MaxValue, K.MaxValue));
-
-					C.MinValue = MinValue;
-					C.MaxValue = MaxValue;
-
-					B.MinValue = MinValue;
-					B.MaxValue = MaxValue;
-
-					K.MinValue = MinValue;
-					K.MaxValue = MaxValue;
-
-					const FVector2D DrawSize(150.f, 100.f);
-					const FTransform OffsetTransform(FRotator(0.f, 0.f, 0.f), FVector(-50.f, -75.f, 100.f));
-					const FTransform DrawTransform = OffsetTransform * DrawParams.RootTransform;
-
-					DrawDebugFloatHistory(*DrawParams.World, K, OffsetTransform * DrawParams.RootTransform, DrawSize, FColor(255, 192, 203, 160)); // pink
-					DrawDebugFloatHistory(*DrawParams.World, B, OffsetTransform * DrawParams.RootTransform, DrawSize, FColor(0, 0, 255, 160)); // blue
-					DrawDebugFloatHistory(*DrawParams.World, C, OffsetTransform * DrawParams.RootTransform, DrawSize, FColor(160, 160, 160, 160)); // gray
-				}
+				UE::PoseSearch::FDebugDrawParams DrawParams(Context.AnimInstanceProxy, CurResult.Database.Get(), EDebugDrawFlags::DrawQuery);
+				DrawParams.DrawFeatureVector(CurResult.ComposedQuery.GetValues());
 			}
-		}
-	}
 
-	bWasEvaluated = false;
+			// @todo: add ContinuingPoseCost / BruteForcePoseCost / PoseCost graphs to rewind debugger 
+			//if (DrawParams.Database->PoseSearchMode == EPoseSearchMode::PCAKDTree_Compare)
+			//{
+			//	FDebugFloatHistory& C = MotionMatchingState.SearchCostHistoryContinuing;
+			//	FDebugFloatHistory& B = MotionMatchingState.SearchCostHistoryBruteForce;
+			//	FDebugFloatHistory& K = MotionMatchingState.SearchCostHistoryKDTree;
+
+			//	C.AddSample(CurResult.ContinuingPoseCost.IsValid() ? CurResult.ContinuingPoseCost.GetTotalCost() : C.MaxValue);
+			//	B.AddSample(CurResult.BruteForcePoseCost.IsValid() ? CurResult.BruteForcePoseCost.GetTotalCost() : B.MaxValue);
+			//	K.AddSample(CurResult.PoseCost.IsValid() ? CurResult.PoseCost.GetTotalCost() : K.MaxValue);
+
+			//	// making SearchCostHistoryKDTree and SearchCostHistoryBruteForce min max consistent
+			//	const float MinValue = FMath::Min(C.MinValue, FMath::Min(B.MinValue, K.MinValue));
+			//	const float MaxValue = FMath::Max(C.MaxValue, FMath::Max(B.MaxValue, K.MaxValue));
+
+			//	C.MinValue = MinValue;
+			//	C.MaxValue = MaxValue;
+
+			//	B.MinValue = MinValue;
+			//	B.MaxValue = MaxValue;
+
+			//	K.MinValue = MinValue;
+			//	K.MaxValue = MaxValue;
+
+			//	const FVector2D DrawSize(150.f, 100.f);
+			//	const FTransform OffsetTransform(FRotator(0.f, 0.f, 0.f), FVector(-50.f, -75.f, 100.f));
+			//	const FTransform DrawTransform = OffsetTransform * DrawParams.GetRootTransform();
+
+			//	DrawDebugFloatHistory(*DrawParams.World, K, OffsetTransform * DrawParams.GetRootTransform(), DrawSize, FColor(255, 192, 203, 160)); // pink
+			//	DrawDebugFloatHistory(*DrawParams.World, B, OffsetTransform * DrawParams.GetRootTransform(), DrawSize, FColor(0, 0, 255, 160)); // blue
+			//	DrawDebugFloatHistory(*DrawParams.World, C, OffsetTransform * DrawParams.GetRootTransform(), DrawSize, FColor(160, 160, 160, 160)); // gray
+			//}
+		}
+	} 
 #endif
 }
 
diff --git a/Engine/Plugins/Experimental/Animation/PoseSearch/Source/Runtime/Private/AnimNode_PoseSearchHistoryCollector.cpp b/Engine/Plugins/Experimental/Animation/PoseSearch/Source/Runtime/Private/AnimNode_PoseSearchHistoryCollector.cpp
--- a/Engine/Plugins/Experimental/Animation/PoseSearch/Source/Runtime/Private/AnimNode_PoseSearchHistoryCollector.cpp
+++ b/Engine/Plugins/Experimental/Animation/PoseSearch/Source/Runtime/Private/AnimNode_PoseSearchHistoryCollector.cpp
@@ -9,7 +9,7 @@
 
 #define LOCTEXT_NAMESPACE "AnimNode_PoseSearchHistoryCollector"
 
-#if WITH_EDITORONLY_DATA && ENABLE_ANIM_DEBUG
+#if ENABLE_DRAW_DEBUG && ENABLE_ANIM_DEBUG
 TAutoConsoleVariable<bool> CVarAnimPoseHistoryDebugDraw(TEXT("a.AnimNode.PoseHistory.DebugDraw"), false, TEXT("Enable / Disable Pose History DebugDraw"));
 #endif
 
@@ -72,38 +72,6 @@ void FAnimNode_PoseSearchHistoryCollector_Base::CacheBones_AnyThread(const FAnim
 	PoseHistory.Init(PoseCount, PoseDuration, RequiredBones);
 }
 
-bool FAnimNode_PoseSearchHistoryCollector_Base::HasPreUpdate() const
-{
-	return WITH_EDITORONLY_DATA && ENABLE_ANIM_DEBUG;
-}
-
-void FAnimNode_PoseSearchHistoryCollector_Base::PreUpdate(const UAnimInstance* InAnimInstance)
-{
-	Super::PreUpdate(InAnimInstance);
-
-#if WITH_EDITORONLY_DATA && ENABLE_ANIM_DEBUG
-	check(InAnimInstance);
-	if (bWasEvaluated && CVarAnimPoseHistoryDebugDraw.GetValueOnAnyThread())
-	{
-		if (const USkeletalMeshComponent* SkeletalMeshComponent = InAnimInstance->GetSkelMeshComponent())
-		{
-			const USkinnedAsset* SkinnedAsset = SkeletalMeshComponent->GetSkinnedAsset();
-			const UWorld* World = SkeletalMeshComponent->GetWorld();
-			if (SkinnedAsset && World)
-			{
-				if (const USkeleton* Skeleton = SkinnedAsset->GetSkeleton())
-				{
-					PoseHistory.DebugDraw(World, Skeleton);
-				}
-			}
-		}
-	}
-
-	bWasEvaluated = false;
-#endif // WITH_EDITORONLY_DATA && ENABLE_ANIM_DEBUG
-}
-
-
 /////////////////////////////////////////////////////
 // FAnimNode_PoseSearchHistoryCollector
 
@@ -133,9 +101,12 @@ void FAnimNode_PoseSearchHistoryCollector::Evaluate_AnyThread(FPoseContext& Outp
 	ComponentSpacePose.InitPose(Output.Pose);
 	PoseHistory.Update(Output.AnimInstanceProxy->GetDeltaSeconds(), ComponentSpacePose, Output.AnimInstanceProxy->GetComponentTransform());
 
-#if WITH_EDITORONLY_DATA && ENABLE_ANIM_DEBUG
-	bWasEvaluated = true;
-#endif
+#if ENABLE_DRAW_DEBUG && ENABLE_ANIM_DEBUG
+	if (CVarAnimPoseHistoryDebugDraw.GetValueOnAnyThread())
+	{
+		PoseHistory.DebugDraw(*Output.AnimInstanceProxy);
+	}
+#endif // ENABLE_DRAW_DEBUG && ENABLE_ANIM_DEBUG
 }
 
 void FAnimNode_PoseSearchHistoryCollector::Update_AnyThread(const FAnimationUpdateContext& Context)
@@ -180,9 +151,12 @@ void FAnimNode_PoseSearchComponentSpaceHistoryCollector::EvaluateComponentSpace_
 
 	PoseHistory.Update(Output.AnimInstanceProxy->GetDeltaSeconds(), Output.Pose, Output.AnimInstanceProxy->GetComponentTransform());
 
-#if WITH_EDITORONLY_DATA && ENABLE_ANIM_DEBUG
-	bWasEvaluated = true;
-#endif
+#if ENABLE_DRAW_DEBUG && ENABLE_ANIM_DEBUG
+	if (CVarAnimPoseHistoryDebugDraw.GetValueOnAnyThread())
+	{
+		PoseHistory.DebugDraw(*Output.AnimInstanceProxy);
+	}
+#endif // ENABLE_DRAW_DEBUG && ENABLE_ANIM_DEBUG
 }
 
 void FAnimNode_PoseSearchComponentSpaceHistoryCollector::Update_AnyThread(const FAnimationUpdateContext& Context)
diff --git a/Engine/Plugins/Experimental/Animation/PoseSearch/Source/Runtime/Private/PoseSearchContext.cpp b/Engine/Plugins/Experimental/Animation/PoseSearch/Source/Runtime/Private/PoseSearchContext.cpp
--- a/Engine/Plugins/Experimental/Animation/PoseSearch/Source/Runtime/Private/PoseSearchContext.cpp
+++ b/Engine/Plugins/Experimental/Animation/PoseSearch/Source/Runtime/Private/PoseSearchContext.cpp
@@ -1,6 +1,7 @@
 // Copyright Epic Games, Inc. All Rights Reserved.
 
 #include "PoseSearch/PoseSearchContext.h"
+#include "Animation/AnimInstanceProxy.h"
 #include "Animation/MotionTrajectoryTypes.h"
 #include "AnimationRuntime.h"
 #include "PoseSearch/PoseSearchDatabase.h"
@@ -14,9 +15,29 @@ namespace UE::PoseSearch
 #if ENABLE_DRAW_DEBUG
 //////////////////////////////////////////////////////////////////////////
 // FDebugDrawParams
+FDebugDrawParams::FDebugDrawParams(FAnimInstanceProxy* InAnimInstanceProxy, const UPoseSearchDatabase* InDatabase, EDebugDrawFlags InFlags)
+: AnimInstanceProxy(InAnimInstanceProxy)
+, World(nullptr)
+, Mesh(nullptr)
+, Database(InDatabase)
+, Flags(InFlags)
+, CachedPositions()
+{
+}
+
+FDebugDrawParams::FDebugDrawParams(const UWorld* InWorld, const USkinnedMeshComponent* InMesh, const UPoseSearchDatabase* InDatabase, EDebugDrawFlags InFlags)
+: AnimInstanceProxy(nullptr)
+, World(InWorld)
+, Mesh(InMesh)
+, Database(InDatabase)
+, Flags(InFlags)
+, CachedPositions()
+{
+}
+
 bool FDebugDrawParams::CanDraw() const
 {
-	return World && Database && Database->Schema && Database->Schema->IsValid();
+	return (AnimInstanceProxy || World) && Database && Database->Schema && Database->Schema->IsValid();
 }
 
 FColor FDebugDrawParams::GetColor(int32 ColorPreset) const
@@ -86,21 +107,141 @@ FVector FDebugDrawParams::GetCachedPosition(float TimeOffset, int8 SchemaBoneIdx
 			return CachedPosition->Transform;
 		}
 
-		if (Mesh.IsValid() && SchemaBoneIdx >= 0)
+		if (Mesh && SchemaBoneIdx >= 0)
 		{
 			return Mesh->GetSocketTransform(Schema->BoneReferences[SchemaBoneIdx].BoneName).GetTranslation();
 		}
 	}
-	return RootTransform.GetTranslation();
+	return GetRootTransform().GetTranslation();
+}
+
+const FTransform& FDebugDrawParams::GetRootTransform() const
+{
+	if (AnimInstanceProxy)
+	{
+		return AnimInstanceProxy->GetComponentTransform();
+	}
+	
+	check(Mesh);
+	return Mesh->GetComponentTransform();
+}
+
+void FDebugDrawParams::DrawLine(const FVector& LineStart, const FVector& LineEnd, const FColor& Color, float Thickness) const
+{
+	if (AnimInstanceProxy)
+	{
+		AnimInstanceProxy->AnimDrawDebugLine(LineStart, LineEnd, Color, false, 0.f, Thickness, SDPG_Foreground);
+	}
+	else if (World)
+	{
+		DrawDebugLine(World, LineStart, LineEnd, Color, false, 0.f, SDPG_Foreground, Thickness);
+	}
+}
+
+void FDebugDrawParams::DrawPoint(const FVector& Position, const FColor& Color, float Thickness) const
+{
+	if (AnimInstanceProxy)
+	{
+		AnimInstanceProxy->AnimDrawDebugPoint(Position, Thickness, Color, false, 0.f, SDPG_Foreground);
+	}
+	else if (World)
+	{
+		DrawDebugPoint(World, Position, Thickness, Color, false, 0.f, SDPG_Foreground);
+	}
+}
+
+void FDebugDrawParams::DrawCircle(const FMatrix& TransformMatrix, float Radius, int32 Segments, const FColor& Color, float Thickness) const
+{
+	if (AnimInstanceProxy)
+	{
+		AnimInstanceProxy->AnimDrawDebugCircle(TransformMatrix.GetOrigin(), Radius, Segments, Color, TransformMatrix.GetScaledAxis(EAxis::X), false, 0.f, SDPG_Foreground, Thickness);
+	}
+	else if (World)
+	{
+		// @todo: use the DrawDebugCircle API with the up vector to communize with the AnimInstanceProxy call
+		DrawDebugCircle(World, TransformMatrix, Radius, Segments, Color, false, 0.f, SDPG_Foreground, Thickness);
+	}
+}
+
+void FDebugDrawParams::DrawCentripetalCatmullRomSpline(TConstArrayView<FVector> Points, TConstArrayView<FColor> Colors, float Alpha, int32 NumSamplesPerSegment, float Thickness) const
+{
+	const int32 NumPoints = Points.Num();
+	const int32 NumColors = Colors.Num();
+	if (NumPoints > 1)
+	{
+		auto GetT = [](float T, float Alpha, const FVector& P0, const FVector& P1)
+		{
+			const FVector P1P0 = P1 - P0;
+			const float Dot = P1P0 | P1P0;
+			const float Pow = FMath::Pow(Dot, Alpha * .5f);
+			return Pow + T;
+		};
+
+		auto LerpColor = [](FColor A, FColor B, float T) -> FColor
+		{
+			return FColor(
+				FMath::RoundToInt(float(A.R) * (1.f - T) + float(B.R) * T),
+				FMath::RoundToInt(float(A.G) * (1.f - T) + float(B.G) * T),
+				FMath::RoundToInt(float(A.B) * (1.f - T) + float(B.B) * T),
+				FMath::RoundToInt(float(A.A) * (1.f - T) + float(B.A) * T));
+		};
+
+		FVector PrevPoint = Points[0];
+		for (int i = 0; i < NumPoints - 1; ++i)
+		{
+			const FVector& P0 = Points[FMath::Max(i - 1, 0)];
+			const FVector& P1 = Points[i];
+			const FVector& P2 = Points[i + 1];
+			const FVector& P3 = Points[FMath::Min(i + 2, NumPoints - 1)];
+
+			const float T0 = 0.0f;
+			const float T1 = GetT(T0, Alpha, P0, P1);
+			const float T2 = GetT(T1, Alpha, P1, P2);
+			const float T3 = GetT(T2, Alpha, P2, P3);
+
+			const float T1T0 = T1 - T0;
+			const float T2T1 = T2 - T1;
+			const float T3T2 = T3 - T2;
+			const float T2T0 = T2 - T0;
+			const float T3T1 = T3 - T1;
+
+			const bool bIsNearlyZeroT1T0 = FMath::IsNearlyZero(T1T0, UE_KINDA_SMALL_NUMBER);
+			const bool bIsNearlyZeroT2T1 = FMath::IsNearlyZero(T2T1, UE_KINDA_SMALL_NUMBER);
+			const bool bIsNearlyZeroT3T2 = FMath::IsNearlyZero(T3T2, UE_KINDA_SMALL_NUMBER);
+			const bool bIsNearlyZeroT2T0 = FMath::IsNearlyZero(T2T0, UE_KINDA_SMALL_NUMBER);
+			const bool bIsNearlyZeroT3T1 = FMath::IsNearlyZero(T3T1, UE_KINDA_SMALL_NUMBER);
+
+			const FColor Color1 = Colors[FMath::Min(i, NumColors - 1)];
+			const FColor Color2 = Colors[FMath::Min(i + 1, NumColors - 1)];
+
+			for (int SampleIndex = 1; SampleIndex < NumSamplesPerSegment; ++SampleIndex)
+			{
+				const float ParametricDistance = float(SampleIndex) / float(NumSamplesPerSegment - 1);
+
+				const float T = FMath::Lerp(T1, T2, ParametricDistance);
+
+				const FVector A1 = bIsNearlyZeroT1T0 ? P0 : (T1 - T) / T1T0 * P0 + (T - T0) / T1T0 * P1;
+				const FVector A2 = bIsNearlyZeroT2T1 ? P1 : (T2 - T) / T2T1 * P1 + (T - T1) / T2T1 * P2;
+				const FVector A3 = bIsNearlyZeroT3T2 ? P2 : (T3 - T) / T3T2 * P2 + (T - T2) / T3T2 * P3;
+				const FVector B1 = bIsNearlyZeroT2T0 ? A1 : (T2 - T) / T2T0 * A1 + (T - T0) / T2T0 * A2;
+				const FVector B2 = bIsNearlyZeroT3T1 ? A2 : (T3 - T) / T3T1 * A2 + (T - T1) / T3T1 * A3;
+				const FVector Point = bIsNearlyZeroT2T1 ? B1 : (T2 - T) / T2T1 * B1 + (T - T1) / T2T1 * B2;
+
+				DrawLine(PrevPoint, Point, LerpColor(Color1, Color2, ParametricDistance));
+
+				PrevPoint = Point;
+			}
+		}
+	}
 }
 
-void DrawFeatureVector(FDebugDrawParams& DrawParams, TConstArrayView<float> PoseVector)
+void FDebugDrawParams::DrawFeatureVector(TConstArrayView<float> PoseVector)
 {
-	DrawParams.ClearCachedPositions();
+	ClearCachedPositions();
 
-	if (DrawParams.CanDraw())
+	if (CanDraw())
 	{
-		const UPoseSearchSchema* Schema = DrawParams.GetSchema();
+		const UPoseSearchSchema* Schema = GetSchema();
 		check(Schema);
 
 		if (PoseVector.Num() == Schema->SchemaCardinality)
@@ -109,7 +250,7 @@ void DrawFeatureVector(FDebugDrawParams& DrawParams, TConstArrayView<float> Pose
 			{
 				if (ChannelPtr)
 				{
-					ChannelPtr->PreDebugDraw(DrawParams, PoseVector);
+					ChannelPtr->PreDebugDraw(*this, PoseVector);
 				}
 			}
 
@@ -117,19 +258,19 @@ void DrawFeatureVector(FDebugDrawParams& DrawParams, TConstArrayView<float> Pose
 			{
 				if (ChannelPtr)
 				{
-					ChannelPtr->DebugDraw(DrawParams, PoseVector);
+					ChannelPtr->DebugDraw(*this, PoseVector);
 				}
 			}
 		}
 	}
 }
 
-void DrawFeatureVector(FDebugDrawParams& DrawParams, int32 PoseIdx)
+void FDebugDrawParams::DrawFeatureVector(int32 PoseIdx)
 {
 	// if we're editing the schema while in PIE with Rewind Debugger active, PoseIdx could be out of bound / stale
-	if (DrawParams.CanDraw() && PoseIdx >= 0 && PoseIdx < DrawParams.GetSearchIndex()->GetNumPoses())
+	if (CanDraw() && PoseIdx >= 0 && PoseIdx < GetSearchIndex()->GetNumPoses())
 	{
-		DrawFeatureVector(DrawParams, DrawParams.GetSearchIndex()->GetPoseValues(PoseIdx));
+		DrawFeatureVector(GetSearchIndex()->GetPoseValues(PoseIdx));
 	}
 }
 #endif // ENABLE_DRAW_DEBUG
diff --git a/Engine/Plugins/Experimental/Animation/PoseSearch/Source/Runtime/Private/PoseSearchFeatureChannel_FilterCrashingLegs.cpp b/Engine/Plugins/Experimental/Animation/PoseSearch/Source/Runtime/Private/PoseSearchFeatureChannel_FilterCrashingLegs.cpp
--- a/Engine/Plugins/Experimental/Animation/PoseSearch/Source/Runtime/Private/PoseSearchFeatureChannel_FilterCrashingLegs.cpp
+++ b/Engine/Plugins/Experimental/Animation/PoseSearch/Source/Runtime/Private/PoseSearchFeatureChannel_FilterCrashingLegs.cpp
@@ -74,10 +74,6 @@ void UPoseSearchFeatureChannel_FilterCrashingLegs::DebugDraw(const UE::PoseSearc
 
 	const float CrashingLegsValue = FFeatureVectorHelper::DecodeFloat(PoseVector, ChannelDataOffset);
 
-	const float LifeTime = DrawParams.DefaultLifeTime;
-	const uint8 DepthPriority = ESceneDepthPriorityGroup::SDPG_Foreground + 2;
-	const bool bPersistent = EnumHasAnyFlags(DrawParams.Flags, EDebugDrawFlags::Persistent);
-
 	const FVector LeftThighPosition = DrawParams.GetCachedPosition(0.f, LeftThighIdx);
 	const FVector RightThighPosition = DrawParams.GetCachedPosition(0.f, RightThighIdx);
 	const FVector LeftFootPosition = DrawParams.GetCachedPosition(0.f, LeftFootIdx);
@@ -101,12 +97,10 @@ void UPoseSearchFeatureChannel_FilterCrashingLegs::DebugDraw(const UE::PoseSearc
 
 	// when CrossingLegsValue is greater than .5 or lesser than -.5 we draw in Red or Orange
 	const float LerpValue = FMath::Min(CrashingLegsValue * 2.f, 1.f);
-	const FColor Color = LerpValue >= 0.f ?
-		LerpColor(FColor::Green, FColor::Red, LerpValue) :
-		LerpColor(FColor::Green, FColor::Orange, -LerpValue);
+	const FColor Color = LerpValue >= 0.f ? LerpColor(FColor::Green, FColor::Red, LerpValue) : LerpColor(FColor::Green, FColor::Orange, -LerpValue);
 
-	DrawDebugLine(DrawParams.World, LeftFootPosition, LeftFootPosition + CrossingLegsVector, Color, bPersistent, LifeTime, DepthPriority);
-	DrawDebugLine(DrawParams.World, RightFootPosition, RightFootPosition - CrossingLegsVector, Color, bPersistent, LifeTime, DepthPriority);
+	DrawParams.DrawLine(LeftFootPosition, LeftFootPosition + CrossingLegsVector, Color);
+	DrawParams.DrawLine(RightFootPosition, RightFootPosition - CrossingLegsVector, Color);
 }
 #endif // ENABLE_DRAW_DEBUG
 
diff --git a/Engine/Plugins/Experimental/Animation/PoseSearch/Source/Runtime/Private/PoseSearchFeatureChannel_Heading.cpp b/Engine/Plugins/Experimental/Animation/PoseSearch/Source/Runtime/Private/PoseSearchFeatureChannel_Heading.cpp
--- a/Engine/Plugins/Experimental/Animation/PoseSearch/Source/Runtime/Private/PoseSearchFeatureChannel_Heading.cpp
+++ b/Engine/Plugins/Experimental/Animation/PoseSearch/Source/Runtime/Private/PoseSearchFeatureChannel_Heading.cpp
@@ -29,7 +29,7 @@ FVector UPoseSearchFeatureChannel_Heading::GetAxis(const FQuat& Rotation) const
 	}
 
 	checkNoEntry();
-	return FVector(1.f, 0.f, 0.f);
+	return FVector::XAxisVector;
 }
 
 void UPoseSearchFeatureChannel_Heading::BuildQuery(UE::PoseSearch::FSearchContext& SearchContext, FPoseSearchFeatureVectorBuilder& InOutQuery) const
@@ -62,22 +62,12 @@ void UPoseSearchFeatureChannel_Heading::DebugDraw(const UE::PoseSearch::FDebugDr
 {
 	using namespace UE::PoseSearch;
 
-	const float LifeTime = DrawParams.DefaultLifeTime;
-	const uint8 DepthPriority = ESceneDepthPriorityGroup::SDPG_Foreground + 2;
-	const bool bPersistent = EnumHasAnyFlags(DrawParams.Flags, EDebugDrawFlags::Persistent);
 	const FColor Color = DrawParams.GetColor(ColorPresetIndex);
-	const FVector BoneHeading = DrawParams.RootTransform.GetRotation().RotateVector(FFeatureVectorHelper::DecodeVector(PoseVector, ChannelDataOffset, ComponentStripping));
+	const FVector BoneHeading = DrawParams.GetRootTransform().GetRotation().RotateVector(FFeatureVectorHelper::DecodeVector(PoseVector, ChannelDataOffset, ComponentStripping));
 	const FVector BonePos = DrawParams.GetCachedPosition(SampleTimeOffset, SchemaBoneIdx);
 
-	if (EnumHasAnyFlags(DrawParams.Flags, EDebugDrawFlags::DrawSearchIndex))
-	{
-		DrawDebugLine(DrawParams.World, BonePos, BonePos + BoneHeading * 15.f, Color, bPersistent, LifeTime, DepthPriority);
-	}
-	else
-	{
-		const float AdjustedThickness = EnumHasAnyFlags(DrawParams.Flags, EDebugDrawFlags::DrawFast) ? 0.0f : 1.f;
-		DrawDebugLine(DrawParams.World, BonePos, BonePos + BoneHeading * 15.f, Color, bPersistent, LifeTime, DepthPriority, AdjustedThickness);
-	}
+	DrawParams.DrawPoint(BonePos, Color, 3.f);
+	DrawParams.DrawLine(BonePos + BoneHeading * 4.f, BonePos + BoneHeading * 15.f, Color);
 }
 #endif // ENABLE_DRAW_DEBUG
 
diff --git a/Engine/Plugins/Experimental/Animation/PoseSearch/Source/Runtime/Private/PoseSearchFeatureChannel_Phase.cpp b/Engine/Plugins/Experimental/Animation/PoseSearch/Source/Runtime/Private/PoseSearchFeatureChannel_Phase.cpp
--- a/Engine/Plugins/Experimental/Animation/PoseSearch/Source/Runtime/Private/PoseSearchFeatureChannel_Phase.cpp
+++ b/Engine/Plugins/Experimental/Animation/PoseSearch/Source/Runtime/Private/PoseSearchFeatureChannel_Phase.cpp
@@ -316,25 +316,23 @@ void UPoseSearchFeatureChannel_Phase::DebugDraw(const UE::PoseSearch::FDebugDraw
 
 	static float ScaleFactor = 1.f;
 
-	const float LifeTime = DrawParams.DefaultLifeTime;
-	const uint8 DepthPriority = ESceneDepthPriorityGroup::SDPG_Foreground + 2;
-	const bool bPersistent = EnumHasAnyFlags(DrawParams.Flags, EDebugDrawFlags::Persistent);
+
 	const FColor Color = DrawParams.GetColor(ColorPresetIndex);
 
 	const FVector2D Phase = FFeatureVectorHelper::DecodeVector2D(PoseVector, ChannelDataOffset);
 	const FVector BonePos = DrawParams.GetCachedPosition(0.f, SchemaBoneIdx);
 
-	const FVector TransformXAxisVector = DrawParams.RootTransform.TransformVector(FVector::XAxisVector);
-	const FVector TransformYAxisVector = DrawParams.RootTransform.TransformVector(FVector::YAxisVector);
-	const FVector TransformZAxisVector = DrawParams.RootTransform.TransformVector(FVector::ZAxisVector);
+	const FVector TransformXAxisVector = DrawParams.GetRootTransform().TransformVector(FVector::XAxisVector);
+	const FVector TransformYAxisVector = DrawParams.GetRootTransform().TransformVector(FVector::YAxisVector);
+	const FVector TransformZAxisVector = DrawParams.GetRootTransform().TransformVector(FVector::ZAxisVector);
 
 	const FVector PhaseVector = (TransformZAxisVector * Phase.X + TransformYAxisVector * Phase.Y) * ScaleFactor;
-	DrawDebugLine(DrawParams.World, BonePos, BonePos + PhaseVector, Color, bPersistent, LifeTime, DepthPriority, 0.f);
+	DrawParams.DrawLine(BonePos, BonePos + PhaseVector, Color);
 
 	static int32 Segments = 32;
 	FMatrix CircleTransform;
 	CircleTransform.SetAxes(&TransformXAxisVector, &TransformYAxisVector, &TransformZAxisVector, &BonePos);
-	DrawDebugCircle(DrawParams.World, CircleTransform, PhaseVector.Length(), Segments, Color, bPersistent, LifeTime, DepthPriority, 0.f, false);
+	DrawParams.DrawCircle(CircleTransform, PhaseVector.Length(), Segments, Color);
 }
 #endif // ENABLE_DRAW_DEBUG
 
diff --git a/Engine/Plugins/Experimental/Animation/PoseSearch/Source/Runtime/Private/PoseSearchFeatureChannel_Position.cpp b/Engine/Plugins/Experimental/Animation/PoseSearch/Source/Runtime/Private/PoseSearchFeatureChannel_Position.cpp
--- a/Engine/Plugins/Experimental/Animation/PoseSearch/Source/Runtime/Private/PoseSearchFeatureChannel_Position.cpp
+++ b/Engine/Plugins/Experimental/Animation/PoseSearch/Source/Runtime/Private/PoseSearchFeatureChannel_Position.cpp
@@ -49,9 +49,9 @@ void UPoseSearchFeatureChannel_Position::PreDebugDraw(UE::PoseSearch::FDebugDraw
 	using namespace UE::PoseSearch;
 	
 	// if the SchemaOriginBoneIdx is not the root bone, the feature vector position doesn't represent a component space delta position for the SchemaBoneIdx, so we shouldn't collect it
-	if (!DrawParams.Database->Schema->BoneReferences[SchemaOriginBoneIdx].HasValidSetup())
+	if (!DrawParams.GetSchema()->BoneReferences[SchemaOriginBoneIdx].HasValidSetup())
 	{
-		const FVector BonePos = DrawParams.RootTransform.TransformPosition(FFeatureVectorHelper::DecodeVector(PoseVector, ChannelDataOffset, ComponentStripping));
+		const FVector BonePos = DrawParams.GetRootTransform().TransformPosition(FFeatureVectorHelper::DecodeVector(PoseVector, ChannelDataOffset, ComponentStripping));
 		DrawParams.AddCachedPosition(SampleTimeOffset, SchemaBoneIdx, BonePos);
 	}
 }
@@ -60,44 +60,20 @@ void UPoseSearchFeatureChannel_Position::DebugDraw(const UE::PoseSearch::FDebugD
 {
 	using namespace UE::PoseSearch;
 
-	const float LifeTime = DrawParams.DefaultLifeTime;
-	const uint8 DepthPriority = ESceneDepthPriorityGroup::SDPG_Foreground + 2;
-	const bool bPersistent = EnumHasAnyFlags(DrawParams.Flags, EDebugDrawFlags::Persistent);
 	const FColor Color = DrawParams.GetColor(ColorPresetIndex);
 
-	if (!DrawParams.Database->Schema->BoneReferences[SchemaOriginBoneIdx].HasValidSetup())
+	if (!DrawParams.GetSchema()->BoneReferences[SchemaOriginBoneIdx].HasValidSetup())
 	{
-		const FVector BonePos = DrawParams.RootTransform.TransformPosition(FFeatureVectorHelper::DecodeVector(PoseVector, ChannelDataOffset, ComponentStripping));
+		const FVector BonePos = DrawParams.GetRootTransform().TransformPosition(FFeatureVectorHelper::DecodeVector(PoseVector, ChannelDataOffset, ComponentStripping));
 
-		if (EnumHasAnyFlags(DrawParams.Flags, EDebugDrawFlags::DrawFast | EDebugDrawFlags::DrawSearchIndex))
-		{
-			DrawDebugPoint(DrawParams.World, BonePos, DrawParams.PointSize, Color, bPersistent, LifeTime, DepthPriority);
-		}
-		else
-		{
-			DrawDebugSphere(DrawParams.World, BonePos, 2.f, 8, Color, bPersistent, LifeTime, DepthPriority);
-		}
-
-		if (EnumHasAnyFlags(DrawParams.Flags, EDebugDrawFlags::DrawBoneNames))
-		{
-			const UPoseSearchSchema* Schema = DrawParams.GetSchema();
-			check(Schema && Schema->IsValid());
-			DrawDebugString(DrawParams.World, BonePos + FVector(0.0, 0.0, 10.0), Schema->BoneReferences[SchemaBoneIdx].BoneName.ToString(), nullptr, Color, LifeTime, false, 1.0f);
-		}
+		DrawParams.DrawPoint(BonePos, Color);
 	}
 	else
 	{
 		const FVector OriginBonePos = DrawParams.GetCachedPosition(SampleTimeOffset, SchemaOriginBoneIdx);
-		const FVector DeltaPos = DrawParams.RootTransform.TransformVector(FFeatureVectorHelper::DecodeVector(PoseVector, ChannelDataOffset, ComponentStripping));
-
-		DrawDebugLine(DrawParams.World, OriginBonePos, OriginBonePos + DeltaPos, Color, bPersistent, LifeTime, DepthPriority);
+		const FVector DeltaPos = DrawParams.GetRootTransform().TransformVector(FFeatureVectorHelper::DecodeVector(PoseVector, ChannelDataOffset, ComponentStripping));
 
-		if (EnumHasAnyFlags(DrawParams.Flags, EDebugDrawFlags::DrawBoneNames))
-		{
-			const UPoseSearchSchema* Schema = DrawParams.GetSchema();
-			check(Schema && Schema->IsValid());
-			DrawDebugString(DrawParams.World, OriginBonePos + DeltaPos + FVector(0.0, 0.0, 10.0), Schema->BoneReferences[SchemaBoneIdx].BoneName.ToString(), nullptr, Color, LifeTime, false, 1.0f);
-		}
+		DrawParams.DrawLine(OriginBonePos, OriginBonePos + DeltaPos, Color);
 	}	
 }
 #endif // ENABLE_DRAW_DEBUG
diff --git a/Engine/Plugins/Experimental/Animation/PoseSearch/Source/Runtime/Private/PoseSearchFeatureChannel_Trajectory.cpp b/Engine/Plugins/Experimental/Animation/PoseSearch/Source/Runtime/Private/PoseSearchFeatureChannel_Trajectory.cpp
--- a/Engine/Plugins/Experimental/Animation/PoseSearch/Source/Runtime/Private/PoseSearchFeatureChannel_Trajectory.cpp
+++ b/Engine/Plugins/Experimental/Animation/PoseSearch/Source/Runtime/Private/PoseSearchFeatureChannel_Trajectory.cpp
@@ -136,11 +136,7 @@ void UPoseSearchFeatureChannel_Trajectory::DebugDraw(const UE::PoseSearch::FDebu
 			PrevTimeOffset = CurrTimeOffset;
 		}
 
-		const float LifeTime = DrawParams.DefaultLifeTime;
-		const uint8 DepthPriority = ESceneDepthPriorityGroup::SDPG_Foreground + 2;
-		const bool bPersistent = EnumHasAnyFlags(DrawParams.Flags, EDebugDrawFlags::Persistent);
-
-		DrawCentripetalCatmullRomSpline(DrawParams.World, TrajSplinePos, TrajSplineColor, 0.5f, 8.f, bPersistent, LifeTime, DepthPriority, 0.f);
+		DrawParams.DrawCentripetalCatmullRomSpline(TrajSplinePos, TrajSplineColor, 0.5f, 8);
 	}
 
 	Super::DebugDraw(DrawParams, PoseVector);
diff --git a/Engine/Plugins/Experimental/Animation/PoseSearch/Source/Runtime/Private/PoseSearchFeatureChannel_Velocity.cpp b/Engine/Plugins/Experimental/Animation/PoseSearch/Source/Runtime/Private/PoseSearchFeatureChannel_Velocity.cpp
--- a/Engine/Plugins/Experimental/Animation/PoseSearch/Source/Runtime/Private/PoseSearchFeatureChannel_Velocity.cpp
+++ b/Engine/Plugins/Experimental/Animation/PoseSearch/Source/Runtime/Private/PoseSearchFeatureChannel_Velocity.cpp
@@ -53,25 +53,14 @@ void UPoseSearchFeatureChannel_Velocity::DebugDraw(const UE::PoseSearch::FDebugD
 {
 	using namespace UE::PoseSearch;
 
-	const float LifeTime = DrawParams.DefaultLifeTime;
-	const uint8 DepthPriority = ESceneDepthPriorityGroup::SDPG_Foreground + 2;
-	const bool bPersistent = EnumHasAnyFlags(DrawParams.Flags, EDebugDrawFlags::Persistent);
 	const FColor Color = DrawParams.GetColor(ColorPresetIndex);
 	const float LinearVelocityScale = bNormalize ? 15.f : 0.08f;
 
-	const FVector LinearVelocity = DrawParams.RootTransform.TransformVector(FFeatureVectorHelper::DecodeVector(PoseVector, ChannelDataOffset, ComponentStripping));
+	const FVector LinearVelocity = DrawParams.GetRootTransform().TransformVector(FFeatureVectorHelper::DecodeVector(PoseVector, ChannelDataOffset, ComponentStripping));
 	const FVector BoneVelDirection = LinearVelocity.GetSafeNormal();
 	const FVector BonePos = DrawParams.GetCachedPosition(SampleTimeOffset, SchemaBoneIdx);
 
-	if (EnumHasAnyFlags(DrawParams.Flags, EDebugDrawFlags::DrawSearchIndex))
-	{
-		DrawDebugLine(DrawParams.World, BonePos, BonePos + LinearVelocity * LinearVelocityScale, Color, bPersistent, LifeTime, DepthPriority);
-	}
-	else
-	{
-		const float AdjustedThickness = EnumHasAnyFlags(DrawParams.Flags, EDebugDrawFlags::DrawFast) ? 0.f : 1.f;
-		DrawDebugLine(DrawParams.World, BonePos + BoneVelDirection * 2.f, BonePos + LinearVelocity * LinearVelocityScale, Color, bPersistent, LifeTime, DepthPriority, AdjustedThickness);
-	}
+	DrawParams.DrawLine(BonePos, BonePos + LinearVelocity * LinearVelocityScale, Color);
 }
 #endif // ENABLE_DRAW_DEBUG
 
diff --git a/Engine/Plugins/Experimental/Animation/PoseSearch/Source/Runtime/Private/PoseSearchHistory.cpp b/Engine/Plugins/Experimental/Animation/PoseSearch/Source/Runtime/Private/PoseSearchHistory.cpp
--- a/Engine/Plugins/Experimental/Animation/PoseSearch/Source/Runtime/Private/PoseSearchHistory.cpp
+++ b/Engine/Plugins/Experimental/Animation/PoseSearch/Source/Runtime/Private/PoseSearchHistory.cpp
@@ -2,6 +2,7 @@
 
 #include "PoseSearch/PoseSearchHistory.h"
 #include "AnimationRuntime.h"
+#include "Animation/AnimInstanceProxy.h"
 #include "Animation/AnimNodeBase.h"
 #include "Animation/AnimRootMotionProvider.h"
 #include "BonePose.h"
@@ -243,10 +244,9 @@ float FPoseHistory::GetSampleTimeInterval() const
 	return TimeHorizon / (Entries.Max() - 1);
 }
 
-void FPoseHistory::DebugDraw(const UWorld* World, const USkeleton* Skeleton) const
+#if ENABLE_DRAW_DEBUG && ENABLE_ANIM_DEBUG
+void FPoseHistory::DebugDraw(FAnimInstanceProxy& AnimInstanceProxy) const
 {
-#if ENABLE_DRAW_DEBUG
-
 	auto LerpColor = [](FColor A, FColor B, float T) -> FColor
 	{
 		return FColor(
@@ -281,14 +281,14 @@ void FPoseHistory::DebugDraw(const UWorld* World, const USkeleton* Skeleton) con
 			{
 				const FTransform GlobalTransforms = Entry.ComponentSpaceTransforms[i] * Entry.RootTransform;
 
-				DrawDebugLine(World, PrevGlobalTransforms[i].GetTranslation(), GlobalTransforms.GetTranslation(), Color, false, 0.f, ESceneDepthPriorityGroup::SDPG_Foreground + 2);
+				AnimInstanceProxy.AnimDrawDebugLine(PrevGlobalTransforms[i].GetTranslation(), GlobalTransforms.GetTranslation(), Color, false, 0.f, ESceneDepthPriorityGroup::SDPG_Foreground);
 
 				PrevGlobalTransforms[i] = GlobalTransforms;
 			}
 		}
 	}
-#endif // ENABLE_DRAW_DEBUG
 }
+#endif // ENABLE_DRAW_DEBUG && ENABLE_ANIM_DEBUG
 
 //////////////////////////////////////////////////////////////////////////
 // FExtendedPoseHistory
@@ -409,9 +409,9 @@ void FExtendedPoseHistory::AddFuturePose(float SecondsInTheFuture, FCSPose<FComp
 	FutureEntries.InsertDefaulted_GetRef(LowerBoundIdx).Update(SecondsAgo, ComponentSpacePose, ComponentTransform, PoseHistory->GetBoneToTransformMap());
 }
 
-void FExtendedPoseHistory::DebugDraw(const UWorld* World, const USkeleton* Skeleton) const
+#if ENABLE_DRAW_DEBUG && ENABLE_ANIM_DEBUG
+void FExtendedPoseHistory::DebugDraw(FAnimInstanceProxy& AnimInstanceProxy) const
 {
-#if ENABLE_DRAW_DEBUG
 	check(PoseHistory);
 
 	auto LerpColor = [](FColor A, FColor B, float T) -> FColor
@@ -448,16 +448,16 @@ void FExtendedPoseHistory::DebugDraw(const UWorld* World, const USkeleton* Skele
 			{
 				const FTransform GlobalTransforms = Entry.ComponentSpaceTransforms[i] * Entry.RootTransform;
 
-				DrawDebugLine(World, PrevGlobalTransforms[i].GetTranslation(), GlobalTransforms.GetTranslation(), Color, false, 0.f, ESceneDepthPriorityGroup::SDPG_Foreground + 2);
+				AnimInstanceProxy.AnimDrawDebugLine(PrevGlobalTransforms[i].GetTranslation(), GlobalTransforms.GetTranslation(), Color, false, 0.f, ESceneDepthPriorityGroup::SDPG_Foreground);
 
 				PrevGlobalTransforms[i] = GlobalTransforms;
 			}
 		}
 	}
 
-	PoseHistory->DebugDraw(World, Skeleton);
-#endif // ENABLE_DRAW_DEBUG
+	PoseHistory->DebugDraw(AnimInstanceProxy);
 }
+#endif // ENABLE_DRAW_DEBUG && ENABLE_ANIM_DEBUG
 
 //////////////////////////////////////////////////////////////////////////
 // FPoseIndicesHistory
diff --git a/Engine/Plugins/Experimental/Animation/PoseSearch/Source/Runtime/Private/PoseSearchLibrary.cpp b/Engine/Plugins/Experimental/Animation/PoseSearch/Source/Runtime/Private/PoseSearchLibrary.cpp
--- a/Engine/Plugins/Experimental/Animation/PoseSearch/Source/Runtime/Private/PoseSearchLibrary.cpp
+++ b/Engine/Plugins/Experimental/Animation/PoseSearch/Source/Runtime/Private/PoseSearchLibrary.cpp
@@ -22,7 +22,7 @@
 
 #define LOCTEXT_NAMESPACE "PoseSearchLibrary"
 
-#if ENABLE_ANIM_DEBUG
+#if ENABLE_DRAW_DEBUG && ENABLE_ANIM_DEBUG
 TAutoConsoleVariable<int32> CVarAnimMotionMatchDrawQueryEnable(TEXT("a.MotionMatch.DrawQuery.Enable"), 0, TEXT("Enable / Disable MotionMatch Draw Query"));
 TAutoConsoleVariable<int32> CVarAnimMotionMatchDrawMatchEnable(TEXT("a.MotionMatch.DrawMatch.Enable"), 0, TEXT("Enable / Disable MotionMatch Draw Match"));
 TAutoConsoleVariable<int32> CVarAnimMotionMatchDrawHistoryEnable(TEXT("a.MotionMatch.DrawHistory.Enable"), 0, TEXT("Enable / Disable MotionMatch Draw History"));
@@ -420,7 +420,7 @@ void UPoseSearchLibrary::UpdateMotionMatchingState(
 }
 
 void UPoseSearchLibrary::MotionMatch(
-	const UAnimInstance* AnimInstance,
+	UAnimInstance* AnimInstance,
 	const UPoseSearchSearchableAsset* Searchable,
 	const FTrajectorySampleRange Trajectory,
 	const FName PoseHistoryName,
@@ -434,6 +434,18 @@ void UPoseSearchLibrary::MotionMatch(
 	float TimeToFutureAnimationStart,
 	const int DebugSessionUniqueIdentifier)
 {
+#if ENABLE_DRAW_DEBUG && ENABLE_ANIM_DEBUG
+	class UAnimInstanceProxyProvider : public UAnimInstance
+	{
+	public:
+		static FAnimInstanceProxy& GetAnimInstanceProxy(UAnimInstance* AnimInstance)
+		{
+			check(AnimInstance);
+			return static_cast<UAnimInstanceProxyProvider*>(AnimInstance)->GetProxyOnAnyThread<FAnimInstanceProxy>();
+		}
+	};
+#endif //ENABLE_DRAW_DEBUG && ENABLE_ANIM_DEBUG
+
 	using namespace UE::Anim;
 	using namespace UE::PoseSearch;
 
@@ -567,7 +579,7 @@ void UPoseSearchLibrary::MotionMatch(
 #if ENABLE_DRAW_DEBUG && ENABLE_ANIM_DEBUG
 					if (CVarAnimMotionMatchDrawHistoryEnable.GetValueOnAnyThread())
 					{
-						ExtendedPoseHistory.DebugDraw(AnimInstance->GetOwningComponent()->GetWorld(), BoneContainer.GetSkeletonAsset());
+						ExtendedPoseHistory.DebugDraw(UAnimInstanceProxyProvider::GetAnimInstanceProxy(AnimInstance));
 					}
 #endif // ENABLE_DRAW_DEBUG && ENABLE_ANIM_DEBUG
 				}
@@ -594,24 +606,16 @@ void UPoseSearchLibrary::MotionMatch(
 #if ENABLE_DRAW_DEBUG && ENABLE_ANIM_DEBUG
 			if (SearchResult.IsValid())
 			{
-				UE::PoseSearch::FDebugDrawParams DrawParams;
-				DrawParams.RootTransform = AnimInstance->GetOwningComponent()->GetComponentTransform();
-				DrawParams.Database = SearchResult.Database.Get();
-				DrawParams.World = AnimInstance->GetOwningComponent()->GetWorld();
-				DrawParams.DefaultLifeTime = 0.0f;
-
-				if (DrawParams.CanDraw())
+				if (CVarAnimMotionMatchDrawMatchEnable.GetValueOnAnyThread())
 				{
-					if (CVarAnimMotionMatchDrawMatchEnable.GetValueOnAnyThread())
-					{
-						DrawFeatureVector(DrawParams, SearchResult.PoseIdx);
-					}
+					UE::PoseSearch::FDebugDrawParams DrawParams(&UAnimInstanceProxyProvider::GetAnimInstanceProxy(AnimInstance), SearchResult.Database.Get());
+					DrawParams.DrawFeatureVector(SearchResult.PoseIdx);
+				}
 
-					if (CVarAnimMotionMatchDrawQueryEnable.GetValueOnAnyThread())
-					{
-						EnumAddFlags(DrawParams.Flags, EDebugDrawFlags::DrawQuery);
-						DrawFeatureVector(DrawParams, SearchResult.ComposedQuery.GetValues());
-					}
+				if (CVarAnimMotionMatchDrawQueryEnable.GetValueOnAnyThread())
+				{
+					UE::PoseSearch::FDebugDrawParams DrawParams(&UAnimInstanceProxyProvider::GetAnimInstanceProxy(AnimInstance), SearchResult.Database.Get(), EDebugDrawFlags::DrawQuery);
+					DrawParams.DrawFeatureVector(SearchResult.ComposedQuery.GetValues());
 				}
 			}
 #endif // ENABLE_DRAW_DEBUG && ENABLE_ANIM_DEBUG
diff --git a/Engine/Plugins/Experimental/Animation/PoseSearch/Source/Runtime/Public/PoseSearch/AnimNode_MotionMatching.h b/Engine/Plugins/Experimental/Animation/PoseSearch/Source/Runtime/Public/PoseSearch/AnimNode_MotionMatching.h
--- a/Engine/Plugins/Experimental/Animation/PoseSearch/Source/Runtime/Public/PoseSearch/AnimNode_MotionMatching.h
+++ b/Engine/Plugins/Experimental/Animation/PoseSearch/Source/Runtime/Public/PoseSearch/AnimNode_MotionMatching.h
@@ -57,8 +57,6 @@ public:
 	virtual void Initialize_AnyThread(const FAnimationInitializeContext& Context) override;
 	virtual void Evaluate_AnyThread(FPoseContext& Output) override;
 	virtual void GatherDebugData(FNodeDebugData& DebugData) override;
-	virtual bool HasPreUpdate() const override;
-	virtual void PreUpdate(const UAnimInstance* InAnimInstance) override;
 	// End of FAnimNode_Base interface
 
 private:
@@ -90,8 +88,5 @@ private:
 	// If true, "Relevant anim" nodes that look for the highest weighted animation in a state will ignore this node
 	UPROPERTY(EditAnywhere, Category=Relevancy, meta=(FoldProperty, PinHiddenByDefault))
 	bool bIgnoreForRelevancyTest = false;
-
-	// Whether this node was evaluated last frame
-	bool bWasEvaluated = false;
 #endif // WITH_EDITORONLY_DATA
 };
\ No newline at end of file
diff --git a/Engine/Plugins/Experimental/Animation/PoseSearch/Source/Runtime/Public/PoseSearch/AnimNode_PoseSearchHistoryCollector.h b/Engine/Plugins/Experimental/Animation/PoseSearch/Source/Runtime/Public/PoseSearch/AnimNode_PoseSearchHistoryCollector.h
--- a/Engine/Plugins/Experimental/Animation/PoseSearch/Source/Runtime/Public/PoseSearch/AnimNode_PoseSearchHistoryCollector.h
+++ b/Engine/Plugins/Experimental/Animation/PoseSearch/Source/Runtime/Public/PoseSearch/AnimNode_PoseSearchHistoryCollector.h
@@ -26,19 +26,12 @@ public:
 
 	// FAnimNode_Base interface
 	virtual void CacheBones_AnyThread(const FAnimationCacheBonesContext& Context) override;
-	virtual bool HasPreUpdate() const override;
-	virtual void PreUpdate(const UAnimInstance* InAnimInstance) override;
 	// End of FAnimNode_Base interface
 
 	const UE::PoseSearch::FPoseHistory& GetPoseHistory() const { return PoseHistory; }
 
 protected:
 	UE::PoseSearch::FPoseHistory PoseHistory;
-
-#if WITH_EDITORONLY_DATA && ENABLE_ANIM_DEBUG
-	// Whether this node was evaluated last frame
-	bool bWasEvaluated = false;
-#endif // WITH_EDITORONLY_DATA
 };
 
 USTRUCT(BlueprintInternalUseOnly)
diff --git a/Engine/Plugins/Experimental/Animation/PoseSearch/Source/Runtime/Public/PoseSearch/PoseSearchContext.h b/Engine/Plugins/Experimental/Animation/PoseSearch/Source/Runtime/Public/PoseSearch/PoseSearchContext.h
--- a/Engine/Plugins/Experimental/Animation/PoseSearch/Source/Runtime/Public/PoseSearch/PoseSearchContext.h
+++ b/Engine/Plugins/Experimental/Animation/PoseSearch/Source/Runtime/Public/PoseSearch/PoseSearchContext.h
@@ -21,26 +21,8 @@ enum class EDebugDrawFlags : uint32
 {
 	None = 0,
 
-	// Draw the entire search index as a point cloud
-	DrawSearchIndex = 1 << 0,
-
 	// Draw using Query colors form the schema / config
 	DrawQuery = 1 << 1,
-
-	/**
-	* Keep rendered data until the next call to FlushPersistentDebugLines().
-	* Combine with DrawSearchIndex to draw the search index only once.
-	*/
-	Persistent = 1 << 2,
-
-	// Label samples with their indices
-	DrawSampleLabels = 1 << 3,
-
-	// Draw Bone Names
-	DrawBoneNames = 1 << 5,
-
-	// Draws simpler shapes to improve performance
-	DrawFast = 1 << 6,
 };
 ENUM_CLASS_FLAGS(EDebugDrawFlags);
 
@@ -111,6 +93,11 @@ struct FCachedTransforms
 		CachedTransforms.Reset();
 	}
 
+	bool IsEmpty() const
+	{
+		return CachedTransforms.IsEmpty();
+	}
+
 private:
 	TArray<FCachedTransform<FTransformType>, TInlineAllocator<64>> CachedTransforms;
 };
@@ -118,19 +105,9 @@ private:
 #if ENABLE_DRAW_DEBUG
 struct POSESEARCH_API FDebugDrawParams
 {
-	const UWorld* World = nullptr;
-	const UPoseSearchDatabase* Database = nullptr;
-	EDebugDrawFlags Flags = EDebugDrawFlags::None;
-
-	float DefaultLifeTime = 5.f;
-	float PointSize = 1.f;
+	FDebugDrawParams(FAnimInstanceProxy* InAnimInstanceProxy, const UPoseSearchDatabase* InDatabase, EDebugDrawFlags InFlags = EDebugDrawFlags::None);
+	FDebugDrawParams(const UWorld* InWorld, const USkinnedMeshComponent* InMesh, const UPoseSearchDatabase* InDatabase, EDebugDrawFlags InFlags = EDebugDrawFlags::None);
 
-	FTransform RootTransform = FTransform::Identity;
-
-	// Optional Mesh for gathering SocketTransform(s)
-	TWeakObjectPtr<const USkinnedMeshComponent> Mesh = nullptr;
-
-	bool CanDraw() const;
 	FColor GetColor(int32 ColorPreset) const;
 	const FPoseSearchIndex* GetSearchIndex() const;
 	const UPoseSearchSchema* GetSchema() const;
@@ -138,14 +115,27 @@ struct POSESEARCH_API FDebugDrawParams
 	void ClearCachedPositions();
 	void AddCachedPosition(float TimeOffset, int8 SchemaBoneIdx, const FVector& Position);
 	FVector GetCachedPosition(float TimeOffset, int8 SchemaBoneIdx = -1) const;
+	const FTransform& GetRootTransform() const;
+
+	void DrawLine(const FVector& LineStart, const FVector& LineEnd, const FColor& Color, float Thickness = 0.f) const;
+	void DrawPoint(const FVector& Position, const FColor& Color, float Thickness = 6.f) const;
+	void DrawCircle(const FMatrix& TransformMatrix, float Radius, int32 Segments, const FColor& Color, float Thickness = 1.f) const;
+	void DrawCentripetalCatmullRomSpline(TConstArrayView<FVector> Points, TConstArrayView<FColor> Colors, float Alpha, int32 NumSamplesPerSegment, float Thickness = 1.f) const;
+	
+	void DrawFeatureVector(TConstArrayView<float> PoseVector);
+	void DrawFeatureVector(int32 PoseIdx);
 
 private:
+	bool CanDraw() const;
+
+	FAnimInstanceProxy* AnimInstanceProxy = nullptr;
+	const UWorld* World = nullptr;
+	const USkinnedMeshComponent* Mesh = nullptr;
+	const UPoseSearchDatabase* Database = nullptr;
+	EDebugDrawFlags Flags = EDebugDrawFlags::None;
 	FCachedTransforms<FVector> CachedPositions;
 };
 
-POSESEARCH_API void DrawFeatureVector(FDebugDrawParams& DrawParams, TConstArrayView<float> PoseVector);
-POSESEARCH_API void DrawFeatureVector(FDebugDrawParams& DrawParams, int32 PoseIdx);
-
 #endif // ENABLE_DRAW_DEBUG
 
 struct POSESEARCH_API FSearchContext
diff --git a/Engine/Plugins/Experimental/Animation/PoseSearch/Source/Runtime/Public/PoseSearch/PoseSearchHistory.h b/Engine/Plugins/Experimental/Animation/PoseSearch/Source/Runtime/Public/PoseSearch/PoseSearchHistory.h
--- a/Engine/Plugins/Experimental/Animation/PoseSearch/Source/Runtime/Public/PoseSearch/PoseSearchHistory.h
+++ b/Engine/Plugins/Experimental/Animation/PoseSearch/Source/Runtime/Public/PoseSearch/PoseSearchHistory.h
@@ -5,8 +5,10 @@
 #include "Animation/AnimNodeMessages.h"
 #include "BonePose.h"
 #include "Containers/RingBuffer.h"
+#include "DrawDebugHelpers.h"
 #include "UObject/ObjectKey.h"
 
+struct FAnimInstanceProxy;
 class USkeleton;
 class UWorld;
 
@@ -43,7 +45,7 @@ struct FPoseHistory : public IPoseHistory
 	void Init(int32 InNumPoses, float InTimeHorizon, const TArray<FBoneIndexType>& RequiredBones);
 	void Update(float SecondsElapsed, FCSPose<FCompactPose>& ComponentSpacePose, const FTransform& ComponentTransform);
 	float GetTimeHorizon() const { return TimeHorizon; }
-	void DebugDraw(const UWorld* World, const USkeleton* Skeleton) const;
+
 	const FBoneToTransformMap& GetBoneToTransformMap() const { return BoneToTransformMap; }
 	const FPoseHistoryEntries& GetEntries() const { return Entries; }
 
@@ -53,6 +55,10 @@ struct FPoseHistory : public IPoseHistory
 	virtual void GetRootTransformAtTime(float Time, FTransform& OutRootTransform, bool bExtrapolate = true) const override;
 	// End of IPoseHistory interface
 
+#if ENABLE_DRAW_DEBUG && ENABLE_ANIM_DEBUG
+	void DebugDraw(FAnimInstanceProxy& AnimInstanceProxy) const;
+#endif
+
 private:
 	FBoneToTransformMap BoneToTransformMap;
 	FPoseHistoryEntries Entries;
@@ -71,7 +77,10 @@ struct FExtendedPoseHistory : public IPoseHistory
 
 	void ResetFuturePoses();
 	void AddFuturePose(float SecondsInTheFuture, FCSPose<FCompactPose>& ComponentSpacePose, const FTransform& ComponentTransform);
-	void DebugDraw(const UWorld* World, const USkeleton* Skeleton) const;
+
+#if ENABLE_DRAW_DEBUG && ENABLE_ANIM_DEBUG
+	void DebugDraw(FAnimInstanceProxy& AnimInstanceProxy) const;
+#endif
 
 private:
 	const FPoseHistory* PoseHistory = nullptr;
diff --git a/Engine/Plugins/Experimental/Animation/PoseSearch/Source/Runtime/Public/PoseSearch/PoseSearchLibrary.h b/Engine/Plugins/Experimental/Animation/PoseSearch/Source/Runtime/Public/PoseSearch/PoseSearchLibrary.h
--- a/Engine/Plugins/Experimental/Animation/PoseSearch/Source/Runtime/Public/PoseSearch/PoseSearchLibrary.h
+++ b/Engine/Plugins/Experimental/Animation/PoseSearch/Source/Runtime/Public/PoseSearch/PoseSearchLibrary.h
@@ -107,12 +107,13 @@ struct POSESEARCH_API FMotionMatchingState
 	UPROPERTY(Transient)
 	FTransform RootMotionTransformDelta = FTransform::Identity;
 
-#if WITH_EDITORONLY_DATA
-	enum { SearchCostHistoryNumSamples = 200 };
-	FDebugFloatHistory SearchCostHistoryContinuing = FDebugFloatHistory(SearchCostHistoryNumSamples, 0, 0, true);
-	FDebugFloatHistory SearchCostHistoryBruteForce = FDebugFloatHistory(SearchCostHistoryNumSamples, 0, 0, true);
-	FDebugFloatHistory SearchCostHistoryKDTree = FDebugFloatHistory(SearchCostHistoryNumSamples, 0, 0, true);
-#endif
+	// @todo: add ContinuingPoseCost / BruteForcePoseCost / PoseCost graphs to rewind debugger 
+//#if WITH_EDITORONLY_DATA
+//	enum { SearchCostHistoryNumSamples = 200 };
+//	FDebugFloatHistory SearchCostHistoryContinuing = FDebugFloatHistory(SearchCostHistoryNumSamples, 0, 0, true);
+//	FDebugFloatHistory SearchCostHistoryBruteForce = FDebugFloatHistory(SearchCostHistoryNumSamples, 0, 0, true);
+//	FDebugFloatHistory SearchCostHistoryKDTree = FDebugFloatHistory(SearchCostHistoryNumSamples, 0, 0, true);
+//#endif
 
 	UE::PoseSearch::FPoseIndicesHistory PoseIndicesHistory;
 };
@@ -171,7 +172,7 @@ public:
 	*/
 	UFUNCTION(BlueprintPure, Category = "Animation|Pose Search", meta = (BlueprintThreadSafe, Keywords = "PoseMatch"))
 	static void MotionMatch(
-		const UAnimInstance* AnimInstance,
+		UAnimInstance* AnimInstance,
 		const UPoseSearchSearchableAsset* Searchable,
 		const FTrajectorySampleRange Trajectory,
 		const FName PoseHistoryName,
