motion matching - fix for database support of AnimMontages, removing SDatabaseAssetListItem::GenerateAddButtonWidget (leftover from db groups)

diff --git a/Engine/Plugins/Experimental/Animation/PoseSearch/Source/Editor/Private/PoseSearchDatabaseAssetListItem.cpp b/Engine/Plugins/Experimental/Animation/PoseSearch/Source/Editor/Private/PoseSearchDatabaseAssetListItem.cpp
--- a/Engine/Plugins/Experimental/Animation/PoseSearch/Source/Editor/Private/PoseSearchDatabaseAssetListItem.cpp
+++ b/Engine/Plugins/Experimental/Animation/PoseSearch/Source/Editor/Private/PoseSearchDatabaseAssetListItem.cpp
@@ -162,14 +162,6 @@ namespace UE::PoseSearch
 					.DecoratorStyleSet(&FAppStyle::Get())
 					.TextStyle(FAppStyle::Get(), "DetailsView.CategoryTextStyle")
 				]
-				+SHorizontalBox::Slot()
-				.AutoWidth()
-				.VAlign(VAlign_Center)
-				.HAlign(HAlign_Right)
-				.Padding(2, 0, 0, 0)
-				[
-					GenerateAddButtonWidget()
-				]
 			];
 		}
 		else
@@ -312,62 +304,6 @@ namespace UE::PoseSearch
 		return ItemWidget.ToSharedRef();
 	}
 
-	TSharedRef<SWidget> SDatabaseAssetListItem::GenerateAddButtonWidget()
-	{
-		FMenuBuilder AddOptions(true, nullptr);
-
-		AddOptions.AddMenuEntry(
-			LOCTEXT("AddSequence", "Add Sequence"),
-			LOCTEXT("AddSequenceTooltip", "Add new sequence to this group"),
-			FSlateIcon(),
-			FUIAction(FExecuteAction::CreateSP(this, &SDatabaseAssetListItem::OnAddSequence)),
-			NAME_None,
-			EUserInterfaceActionType::Button);
-
-		AddOptions.AddMenuEntry(
-			LOCTEXT("AddBlendSpaceOption", "Add Blend Space"),
-			LOCTEXT("AddBlendSpaceOptionTooltip", "Add new blend space to this group"),
-			FSlateIcon(),
-			FUIAction(FExecuteAction::CreateSP(this, &SDatabaseAssetListItem::OnAddBlendSpace)),
-			NAME_None,
-			EUserInterfaceActionType::Button);
-
-		AddOptions.AddMenuEntry(
-			LOCTEXT("AnimCompositeOption", "Add Anim Composite"),
-			LOCTEXT("AddAnimCompositeToDefaultGroupTooltip", "Add new anim composite to this group"),
-			FSlateIcon(),
-			FUIAction(FExecuteAction::CreateSP(this, &SDatabaseAssetListItem::OnAddAnimComposite)),
-			NAME_None,
-			EUserInterfaceActionType::Button);
-
-		AddOptions.AddMenuEntry(
-			LOCTEXT("AnimMontageOption", "Add Anim Montage"),
-			LOCTEXT("AddAnimMontageToDefaultGroupTooltip", "Add new anim montage to this group"),
-			FSlateIcon(),
-			FUIAction(FExecuteAction::CreateSP(this, &SDatabaseAssetListItem::OnAddAnimMontage)),
-			NAME_None,
-			EUserInterfaceActionType::Button);
-
-		TSharedPtr<SComboButton> AddButton;
-		SAssignNew(AddButton, SComboButton)
-		.ContentPadding(0)
-		.ComboButtonStyle(&FAppStyle::Get().GetWidgetStyle<FComboButtonStyle>("SimpleComboButton"))
-		.HasDownArrow(false)
-		.ButtonContent()
-		[
-			SNew(SImage)
-			.ColorAndOpacity(FSlateColor::UseForeground())
-			.Image(FAppStyle::Get().GetBrush("Icons.PlusCircle"))
-		]
-		.MenuContent()
-		[
-			AddOptions.MakeWidget()
-		];
-
-		return AddButton.ToSharedRef();
-	}
-
-
 	const FSlateBrush* SDatabaseAssetListItem::GetGroupBackgroundImage() const
 	{
 		if (STableRow<TSharedPtr<FDatabaseAssetTreeNode>>::IsHovered())
diff --git a/Engine/Plugins/Experimental/Animation/PoseSearch/Source/Editor/Private/PoseSearchDatabaseAssetTree.cpp b/Engine/Plugins/Experimental/Animation/PoseSearch/Source/Editor/Private/PoseSearchDatabaseAssetTree.cpp
--- a/Engine/Plugins/Experimental/Animation/PoseSearch/Source/Editor/Private/PoseSearchDatabaseAssetTree.cpp
+++ b/Engine/Plugins/Experimental/Animation/PoseSearch/Source/Editor/Private/PoseSearchDatabaseAssetTree.cpp
@@ -2,6 +2,7 @@
 
 #include "PoseSearchDatabaseAssetTree.h"
 #include "Animation/AnimComposite.h"
+#include "Animation/AnimMontage.h"
 #include "Animation/AnimSequence.h"
 #include "Animation/BlendSpace.h"
 #include "AssetSelection.h"
@@ -99,7 +100,7 @@ namespace UE::PoseSearch
 						.VAlign(VAlign_Center)
 						[
 							SNew(STextBlock)
-							.Text(FText::FromString(TEXT("Drag and drop Animation Sequences, Anim Composites or Blendspaces")))
+							.Text(FText::FromString(TEXT("Drag and drop Animation Sequences, Anim Composites, Blendspaces, or Anim Montages")))
 							.Font(FAppStyle::Get().GetFontStyle("DetailsView.CategoryFontStyle"))
 						]
 					]
@@ -135,7 +136,8 @@ namespace UE::PoseSearch
 						{
 							if (AssetClass->IsChildOf(UAnimSequence::StaticClass()) ||
 								AssetClass->IsChildOf(UAnimComposite::StaticClass()) ||
-								AssetClass->IsChildOf(UBlendSpace::StaticClass()))
+								AssetClass->IsChildOf(UBlendSpace::StaticClass()) ||
+								AssetClass->IsChildOf(UAnimMontage::StaticClass()))
 							{
 								Reply = FReply::Handled();
 								break;
@@ -350,7 +352,8 @@ namespace UE::PoseSearch
 				{
 					if (AssetClass->IsChildOf(UAnimSequence::StaticClass()) ||
 						AssetClass->IsChildOf(UAnimComposite::StaticClass()) ||
-						AssetClass->IsChildOf(UBlendSpace::StaticClass()))
+						AssetClass->IsChildOf(UBlendSpace::StaticClass()) ||
+						AssetClass->IsChildOf(UAnimMontage::StaticClass()))
 					{
 						ReturnedDropZone = EItemDropZone::OntoItem;
 						break;
@@ -425,6 +428,11 @@ namespace UE::PoseSearch
 						ViewModel->AddBlendSpaceToDatabase(Cast<UBlendSpace>(Asset));
 						++AddedAssets;
 					}
+					else if (AssetClass->IsChildOf(UAnimMontage::StaticClass()))
+					{
+						ViewModel->AddAnimMontageToDatabase(Cast<UAnimMontage>(Asset));
+						++AddedAssets;
+					}
 				}
 			}
 			
diff --git a/Engine/Plugins/Experimental/Animation/PoseSearch/Source/Editor/Public/PoseSearchDatabaseAssetListItem.h b/Engine/Plugins/Experimental/Animation/PoseSearch/Source/Editor/Public/PoseSearchDatabaseAssetListItem.h
--- a/Engine/Plugins/Experimental/Animation/PoseSearch/Source/Editor/Public/PoseSearchDatabaseAssetListItem.h
+++ b/Engine/Plugins/Experimental/Animation/PoseSearch/Source/Editor/Public/PoseSearchDatabaseAssetListItem.h
@@ -27,7 +27,6 @@ namespace UE::PoseSearch
 	protected:
 		FText GetName() const;
 		TSharedRef<SWidget> GenerateItemWidget();
-		TSharedRef<SWidget> GenerateAddButtonWidget();
 
 		const FSlateBrush* GetGroupBackgroundImage() const;
 		void ConstructGroupItem(const TSharedRef<STableViewBase>& OwnerTable);
