motion matching - introducing UPoseSearchSchema::bInjectAdditionalDebugChannels paired with UPoseSearchFeatureChannel::AddDependentChannels API as way for channels to inject additional (position for this initiali use case) channels to supplement metadata for debug drawing

diff --git a/Engine/Plugins/Experimental/Animation/PoseSearch/Source/Editor/Private/PoseSearchDebuggerDatabaseView.cpp b/Engine/Plugins/Experimental/Animation/PoseSearch/Source/Editor/Private/PoseSearchDebuggerDatabaseView.cpp
--- a/Engine/Plugins/Experimental/Animation/PoseSearch/Source/Editor/Private/PoseSearchDebuggerDatabaseView.cpp
+++ b/Engine/Plugins/Experimental/Animation/PoseSearch/Source/Editor/Private/PoseSearchDebuggerDatabaseView.cpp
@@ -37,12 +37,9 @@ void SDebuggerDatabaseView::Update(const FTraceMotionMatchingStateMessage& State
 				const UPoseSearchDatabase* Database = FTraceMotionMatchingState::GetObjectFromId<UPoseSearchDatabase>(DbEntry.DatabaseId);
 				if (FAsyncPoseSearchDatabasesManagement::RequestAsyncBuildIndex(Database, ERequestAsyncBuildFlag::ContinueRequest))
 				{
-					for (const TObjectPtr<UPoseSearchFeatureChannel>& ChannelPtr : Database->Schema->Channels)
+					for (const TObjectPtr<UPoseSearchFeatureChannel>& ChannelPtr : Database->Schema->GetChannels())
 					{
-						if (const UPoseSearchFeatureChannel* Channel = ChannelPtr.Get())
-						{
-							AnalyzeChannelRecursively(Channel, bIsVerbose);
-						}
+						AnalyzeChannelRecursively(ChannelPtr.Get(), bIsVerbose);
 					}
 				}
 			}
diff --git a/Engine/Plugins/Experimental/Animation/PoseSearch/Source/Runtime/Private/PoseSearchAssetIndexer.cpp b/Engine/Plugins/Experimental/Animation/PoseSearch/Source/Runtime/Private/PoseSearchAssetIndexer.cpp
--- a/Engine/Plugins/Experimental/Animation/PoseSearch/Source/Runtime/Private/PoseSearchAssetIndexer.cpp
+++ b/Engine/Plugins/Experimental/Animation/PoseSearch/Source/Runtime/Private/PoseSearchAssetIndexer.cpp
@@ -146,12 +146,9 @@ void FAssetIndexer::Process(int32 AssetIdx)
 	// Generate pose features data
 	if (IndexingContext.Schema->SchemaCardinality > 0)
 	{
-		for (const TObjectPtr<UPoseSearchFeatureChannel>& ChannelPtr : IndexingContext.Schema->Channels)
+		for (const TObjectPtr<UPoseSearchFeatureChannel>& ChannelPtr : IndexingContext.Schema->GetChannels())
 		{
-			if (ChannelPtr)
-			{
-				ChannelPtr->IndexAsset(*this);
-			}
+			ChannelPtr->IndexAsset(*this);
 		}
 	}
 
diff --git a/Engine/Plugins/Experimental/Animation/PoseSearch/Source/Runtime/Private/PoseSearchContext.cpp b/Engine/Plugins/Experimental/Animation/PoseSearch/Source/Runtime/Private/PoseSearchContext.cpp
--- a/Engine/Plugins/Experimental/Animation/PoseSearch/Source/Runtime/Private/PoseSearchContext.cpp
+++ b/Engine/Plugins/Experimental/Animation/PoseSearch/Source/Runtime/Private/PoseSearchContext.cpp
@@ -122,8 +122,12 @@ const FTransform& FDebugDrawParams::GetRootTransform() const
 		return AnimInstanceProxy->GetComponentTransform();
 	}
 	
-	check(Mesh);
-	return Mesh->GetComponentTransform();
+	if (Mesh)
+	{
+		return Mesh->GetComponentTransform();
+	}
+
+	return FTransform::Identity;
 }
 
 void FDebugDrawParams::DrawLine(const FVector& LineStart, const FVector& LineEnd, const FColor& Color, float Thickness) const
@@ -246,20 +250,14 @@ void FDebugDrawParams::DrawFeatureVector(TConstArrayView<float> PoseVector)
 
 		if (PoseVector.Num() == Schema->SchemaCardinality)
 		{
-			for (const TObjectPtr<UPoseSearchFeatureChannel>& ChannelPtr : Schema->Channels)
+			for (const TObjectPtr<UPoseSearchFeatureChannel>& ChannelPtr : Schema->GetChannels())
 			{
-				if (ChannelPtr)
-				{
-					ChannelPtr->PreDebugDraw(*this, PoseVector);
-				}
+				ChannelPtr->PreDebugDraw(*this, PoseVector);
 			}
 
-			for (const TObjectPtr<UPoseSearchFeatureChannel>& ChannelPtr : Schema->Channels)
+			for (const TObjectPtr<UPoseSearchFeatureChannel>& ChannelPtr : Schema->GetChannels())
 			{
-				if (ChannelPtr)
-				{
-					ChannelPtr->DebugDraw(*this, PoseVector);
-				}
+				ChannelPtr->DebugDraw(*this, PoseVector);
 			}
 		}
 	}
diff --git a/Engine/Plugins/Experimental/Animation/PoseSearch/Source/Runtime/Private/PoseSearchDatabase.cpp b/Engine/Plugins/Experimental/Animation/PoseSearch/Source/Runtime/Private/PoseSearchDatabase.cpp
--- a/Engine/Plugins/Experimental/Animation/PoseSearch/Source/Runtime/Private/PoseSearchDatabase.cpp
+++ b/Engine/Plugins/Experimental/Animation/PoseSearch/Source/Runtime/Private/PoseSearchDatabase.cpp
@@ -150,9 +150,9 @@ struct FPoseFilters
 			AllPoseFilters.Add(&NonSelectableIdxPoseFilter);
 		}
 
-		for (const IPoseFilter* ChannelPoseFilter : Schema->Channels)
+		for (const IPoseFilter* ChannelPoseFilter : Schema->GetChannels())
 		{
-			if (ChannelPoseFilter && ChannelPoseFilter->IsPoseFilterActive())
+			if (ChannelPoseFilter->IsPoseFilterActive())
 			{
 				AllPoseFilters.Add(ChannelPoseFilter);
 			}
diff --git a/Engine/Plugins/Experimental/Animation/PoseSearch/Source/Runtime/Private/PoseSearchDerivedData.cpp b/Engine/Plugins/Experimental/Animation/PoseSearch/Source/Runtime/Private/PoseSearchDerivedData.cpp
--- a/Engine/Plugins/Experimental/Animation/PoseSearch/Source/Runtime/Private/PoseSearchDerivedData.cpp
+++ b/Engine/Plugins/Experimental/Animation/PoseSearch/Source/Runtime/Private/PoseSearchDerivedData.cpp
@@ -122,12 +122,9 @@ private:
 		for (int32 SchemaIndex = 0; SchemaIndex < Schemas.Num(); ++SchemaIndex)
 		{
 			const UPoseSearchSchema* Schema = Schemas[SchemaIndex];
-			for (const TObjectPtr<UPoseSearchFeatureChannel>& ChannelPtr : Schema->Channels)
+			for (const TObjectPtr<UPoseSearchFeatureChannel>& ChannelPtr : Schema->GetChannels())
 			{
-				if (const UPoseSearchFeatureChannel* Channel = ChannelPtr.Get())
-				{
-					AnalyzeChannelRecursively(Channel, SchemaIndex, EntriesGroup);
-				}
+				AnalyzeChannelRecursively(ChannelPtr.Get(), SchemaIndex, EntriesGroup);
 			}
 		}
 	}
@@ -422,12 +419,9 @@ static void PreprocessSearchIndexWeights(FPoseSearchIndex& SearchIndex, const UP
 	const int32 NumDimensions = Schema->SchemaCardinality;
 	SearchIndex.WeightsSqrt.Init(1.f, NumDimensions);
 
-	for (const TObjectPtr<UPoseSearchFeatureChannel>& ChannelPtr : Schema->Channels)
+	for (const TObjectPtr<UPoseSearchFeatureChannel>& ChannelPtr : Schema->GetChannels())
 	{
-		if (ChannelPtr)
-		{
-			ChannelPtr->FillWeights(SearchIndex.WeightsSqrt);
-		}
+		ChannelPtr->FillWeights(SearchIndex.WeightsSqrt);
 	}
 
 	EPoseSearchDataPreprocessor DataPreprocessor = Schema->DataPreprocessor;
diff --git a/Engine/Plugins/Experimental/Animation/PoseSearch/Source/Runtime/Private/PoseSearchDerivedDataKey.cpp b/Engine/Plugins/Experimental/Animation/PoseSearch/Source/Runtime/Private/PoseSearchDerivedDataKey.cpp
--- a/Engine/Plugins/Experimental/Animation/PoseSearch/Source/Runtime/Private/PoseSearchDerivedDataKey.cpp
+++ b/Engine/Plugins/Experimental/Animation/PoseSearch/Source/Runtime/Private/PoseSearchDerivedDataKey.cpp
@@ -41,7 +41,7 @@ FKeyBuilder::FKeyBuilder(const UObject* Object, bool bUseDataVer)
 	if (bUseDataVer)
 	{
 		// used to invalidate the key without having to change POSESEARCHDB_DERIVEDDATA_VER all the times
-		int32 POSESEARCHDB_DERIVEDDATA_VER_SMALL = 109;
+		int32 POSESEARCHDB_DERIVEDDATA_VER_SMALL = 113;
 		FGuid VersionGuid = FDevSystemGuids::GetSystemGuid(FDevSystemGuids::Get().POSESEARCHDB_DERIVEDDATA_VER);
 
 		*this << VersionGuid;
diff --git a/Engine/Plugins/Experimental/Animation/PoseSearch/Source/Runtime/Private/PoseSearchFeatureChannel_FilterCrashingLegs.cpp b/Engine/Plugins/Experimental/Animation/PoseSearch/Source/Runtime/Private/PoseSearchFeatureChannel_FilterCrashingLegs.cpp
--- a/Engine/Plugins/Experimental/Animation/PoseSearch/Source/Runtime/Private/PoseSearchFeatureChannel_FilterCrashingLegs.cpp
+++ b/Engine/Plugins/Experimental/Animation/PoseSearch/Source/Runtime/Private/PoseSearchFeatureChannel_FilterCrashingLegs.cpp
@@ -6,6 +6,7 @@
 #include "PoseSearch/PoseSearchContext.h"
 #include "PoseSearch/PoseSearchDatabase.h"
 #include "PoseSearch/PoseSearchSchema.h"
+#include "PoseSearchFeatureChannel_Position.h"
 
 static void ComputeThighsSideAndForward(const FVector& RightThighPos, const FVector& LeftThighPos, FVector& ThighsSide, FVector& ThighsForward)
 {
@@ -38,6 +39,16 @@ void UPoseSearchFeatureChannel_FilterCrashingLegs::Finalize(UPoseSearchSchema* S
 	RightFootIdx = Schema->AddBoneReference(RightFoot);
 }
 
+void UPoseSearchFeatureChannel_FilterCrashingLegs::AddDependentChannels(UPoseSearchSchema* Schema) const
+{
+	if (Schema->bInjectAdditionalDebugChannels)
+	{
+		UPoseSearchFeatureChannel_Position::FindOrAddToSchema(Schema, LeftThigh.BoneName, 0.f, 0);
+		UPoseSearchFeatureChannel_Position::FindOrAddToSchema(Schema, RightThigh.BoneName, 0.f, 0);
+		UPoseSearchFeatureChannel_Position::FindOrAddToSchema(Schema, LeftFoot.BoneName, 0.f, 0);
+		UPoseSearchFeatureChannel_Position::FindOrAddToSchema(Schema, RightFoot.BoneName, 0.f, 0);
+	}
+}
 void UPoseSearchFeatureChannel_FilterCrashingLegs::BuildQuery(UE::PoseSearch::FSearchContext& SearchContext, FPoseSearchFeatureVectorBuilder& InOutQuery) const
 {
 	using namespace UE::PoseSearch;
diff --git a/Engine/Plugins/Experimental/Animation/PoseSearch/Source/Runtime/Private/PoseSearchFeatureChannel_FilterCrashingLegs.h b/Engine/Plugins/Experimental/Animation/PoseSearch/Source/Runtime/Private/PoseSearchFeatureChannel_FilterCrashingLegs.h
--- a/Engine/Plugins/Experimental/Animation/PoseSearch/Source/Runtime/Private/PoseSearchFeatureChannel_FilterCrashingLegs.h
+++ b/Engine/Plugins/Experimental/Animation/PoseSearch/Source/Runtime/Private/PoseSearchFeatureChannel_FilterCrashingLegs.h
@@ -56,6 +56,8 @@ public:
 	virtual void Finalize(UPoseSearchSchema* Schema) override;
 	virtual void BuildQuery(UE::PoseSearch::FSearchContext& SearchContext, FPoseSearchFeatureVectorBuilder& InOutQuery) const override;
 
+	virtual void AddDependentChannels(UPoseSearchSchema* Schema) const override;
+
 #if ENABLE_DRAW_DEBUG
 	virtual void DebugDraw(const UE::PoseSearch::FDebugDrawParams& DrawParams, TConstArrayView<float> PoseVector) const override;
 #endif // ENABLE_DRAW_DEBUG
diff --git a/Engine/Plugins/Experimental/Animation/PoseSearch/Source/Runtime/Private/PoseSearchFeatureChannel_Group.cpp b/Engine/Plugins/Experimental/Animation/PoseSearch/Source/Runtime/Private/PoseSearchFeatureChannel_Group.cpp
--- a/Engine/Plugins/Experimental/Animation/PoseSearch/Source/Runtime/Private/PoseSearchFeatureChannel_Group.cpp
+++ b/Engine/Plugins/Experimental/Animation/PoseSearch/Source/Runtime/Private/PoseSearchFeatureChannel_Group.cpp
@@ -16,6 +16,17 @@ void UPoseSearchFeatureChannel_GroupBase::Finalize(UPoseSearchSchema* Schema)
 	ChannelCardinality = Schema->SchemaCardinality - ChannelDataOffset;
 }
 
+void UPoseSearchFeatureChannel_GroupBase::AddDependentChannels(UPoseSearchSchema* Schema) const
+{
+	for (const TObjectPtr<UPoseSearchFeatureChannel>& SubChannelPtr : GetSubChannels())
+	{
+		if (SubChannelPtr)
+		{
+			SubChannelPtr->AddDependentChannels(Schema);
+		}
+	}
+}
+
 #if WITH_EDITOR
 void UPoseSearchFeatureChannel_GroupBase::FillWeights(TArray<float>& Weights) const
 {
diff --git a/Engine/Plugins/Experimental/Animation/PoseSearch/Source/Runtime/Private/PoseSearchFeatureChannel_Group.h b/Engine/Plugins/Experimental/Animation/PoseSearch/Source/Runtime/Private/PoseSearchFeatureChannel_Group.h
--- a/Engine/Plugins/Experimental/Animation/PoseSearch/Source/Runtime/Private/PoseSearchFeatureChannel_Group.h
+++ b/Engine/Plugins/Experimental/Animation/PoseSearch/Source/Runtime/Private/PoseSearchFeatureChannel_Group.h
@@ -16,6 +16,7 @@ public:
 	// UPoseSearchFeatureChannel interface
 	virtual void Finalize(UPoseSearchSchema* Schema) override;
 	virtual void BuildQuery(UE::PoseSearch::FSearchContext& SearchContext, FPoseSearchFeatureVectorBuilder& InOutQuery) const override;
+	virtual void AddDependentChannels(UPoseSearchSchema* Schema) const override; 
 
 #if ENABLE_DRAW_DEBUG
 	virtual void PreDebugDraw(UE::PoseSearch::FDebugDrawParams& DrawParams, TConstArrayView<float> PoseVector) const override;
diff --git a/Engine/Plugins/Experimental/Animation/PoseSearch/Source/Runtime/Private/PoseSearchFeatureChannel_Heading.cpp b/Engine/Plugins/Experimental/Animation/PoseSearch/Source/Runtime/Private/PoseSearchFeatureChannel_Heading.cpp
--- a/Engine/Plugins/Experimental/Animation/PoseSearch/Source/Runtime/Private/PoseSearchFeatureChannel_Heading.cpp
+++ b/Engine/Plugins/Experimental/Animation/PoseSearch/Source/Runtime/Private/PoseSearchFeatureChannel_Heading.cpp
@@ -7,6 +7,7 @@
 #include "PoseSearch/PoseSearchContext.h"
 #include "PoseSearch/PoseSearchDatabase.h"
 #include "PoseSearch/PoseSearchSchema.h"
+#include "PoseSearchFeatureChannel_Position.h"
 
 void UPoseSearchFeatureChannel_Heading::Finalize(UPoseSearchSchema* Schema)
 {
@@ -16,6 +17,14 @@ void UPoseSearchFeatureChannel_Heading::Finalize(UPoseSearchSchema* Schema)
 	SchemaBoneIdx = Schema->AddBoneReference(Bone);
 }
 
+void UPoseSearchFeatureChannel_Heading::AddDependentChannels(UPoseSearchSchema* Schema) const
+{
+	if (Schema->bInjectAdditionalDebugChannels)
+	{
+		UPoseSearchFeatureChannel_Position::FindOrAddToSchema(Schema, Bone.BoneName, SampleTimeOffset, ColorPresetIndex);
+	}
+}
+
 FVector UPoseSearchFeatureChannel_Heading::GetAxis(const FQuat& Rotation) const
 {
 	switch (HeadingAxis)
diff --git a/Engine/Plugins/Experimental/Animation/PoseSearch/Source/Runtime/Private/PoseSearchFeatureChannel_Heading.h b/Engine/Plugins/Experimental/Animation/PoseSearch/Source/Runtime/Private/PoseSearchFeatureChannel_Heading.h
--- a/Engine/Plugins/Experimental/Animation/PoseSearch/Source/Runtime/Private/PoseSearchFeatureChannel_Heading.h
+++ b/Engine/Plugins/Experimental/Animation/PoseSearch/Source/Runtime/Private/PoseSearchFeatureChannel_Heading.h
@@ -51,6 +51,8 @@ public:
 	virtual void Finalize(UPoseSearchSchema* Schema) override;
 	virtual void BuildQuery(UE::PoseSearch::FSearchContext& SearchContext, FPoseSearchFeatureVectorBuilder& InOutQuery) const override;
 
+	virtual void AddDependentChannels(UPoseSearchSchema* Schema) const override;
+
 #if ENABLE_DRAW_DEBUG
 	virtual void DebugDraw(const UE::PoseSearch::FDebugDrawParams& DrawParams, TConstArrayView<float> PoseVector) const override;
 #endif // ENABLE_DRAW_DEBUG
diff --git a/Engine/Plugins/Experimental/Animation/PoseSearch/Source/Runtime/Private/PoseSearchFeatureChannel_Phase.cpp b/Engine/Plugins/Experimental/Animation/PoseSearch/Source/Runtime/Private/PoseSearchFeatureChannel_Phase.cpp
--- a/Engine/Plugins/Experimental/Animation/PoseSearch/Source/Runtime/Private/PoseSearchFeatureChannel_Phase.cpp
+++ b/Engine/Plugins/Experimental/Animation/PoseSearch/Source/Runtime/Private/PoseSearchFeatureChannel_Phase.cpp
@@ -6,6 +6,7 @@
 #include "PoseSearch/PoseSearchContext.h"
 #include "PoseSearch/PoseSearchDatabase.h"
 #include "PoseSearch/PoseSearchSchema.h"
+#include "PoseSearchFeatureChannel_Position.h"
 
 #if WITH_EDITOR
 namespace UE::PoseSearch
@@ -287,6 +288,14 @@ void UPoseSearchFeatureChannel_Phase::Finalize(UPoseSearchSchema* Schema)
 	SchemaBoneIdx = Schema->AddBoneReference(Bone);
 }
 
+void UPoseSearchFeatureChannel_Phase::AddDependentChannels(UPoseSearchSchema* Schema) const
+{
+	if (Schema->bInjectAdditionalDebugChannels)
+	{
+		UPoseSearchFeatureChannel_Position::FindOrAddToSchema(Schema, Bone.BoneName, 0.f, ColorPresetIndex);
+	}
+}
+
 void UPoseSearchFeatureChannel_Phase::BuildQuery(UE::PoseSearch::FSearchContext& SearchContext, FPoseSearchFeatureVectorBuilder& InOutQuery) const
 {
 	using namespace UE::PoseSearch;
diff --git a/Engine/Plugins/Experimental/Animation/PoseSearch/Source/Runtime/Private/PoseSearchFeatureChannel_Phase.h b/Engine/Plugins/Experimental/Animation/PoseSearch/Source/Runtime/Private/PoseSearchFeatureChannel_Phase.h
--- a/Engine/Plugins/Experimental/Animation/PoseSearch/Source/Runtime/Private/PoseSearchFeatureChannel_Phase.h
+++ b/Engine/Plugins/Experimental/Animation/PoseSearch/Source/Runtime/Private/PoseSearchFeatureChannel_Phase.h
@@ -31,6 +31,8 @@ public:
 	virtual void Finalize(UPoseSearchSchema* Schema) override;
 	virtual void BuildQuery(UE::PoseSearch::FSearchContext& SearchContext, FPoseSearchFeatureVectorBuilder& InOutQuery) const override;
 
+	virtual void AddDependentChannels(UPoseSearchSchema* Schema) const override;
+
 #if ENABLE_DRAW_DEBUG
 	virtual void DebugDraw(const UE::PoseSearch::FDebugDrawParams& DrawParams, TConstArrayView<float> PoseVector) const override;
 #endif // ENABLE_DRAW_DEBUG
diff --git a/Engine/Plugins/Experimental/Animation/PoseSearch/Source/Runtime/Private/PoseSearchFeatureChannel_Position.cpp b/Engine/Plugins/Experimental/Animation/PoseSearch/Source/Runtime/Private/PoseSearchFeatureChannel_Position.cpp
--- a/Engine/Plugins/Experimental/Animation/PoseSearch/Source/Runtime/Private/PoseSearchFeatureChannel_Position.cpp
+++ b/Engine/Plugins/Experimental/Animation/PoseSearch/Source/Runtime/Private/PoseSearchFeatureChannel_Position.cpp
@@ -8,6 +8,30 @@
 #include "PoseSearch/PoseSearchDatabase.h"
 #include "PoseSearch/PoseSearchSchema.h"
 
+void UPoseSearchFeatureChannel_Position::FindOrAddToSchema(UPoseSearchSchema* Schema, const FName& InBoneName, float InSampleTimeOffset, int32 InColorPresetIndex)
+{
+	if (!Schema->FindChannel([&InBoneName, InSampleTimeOffset](const UPoseSearchFeatureChannel* Channel) -> const UPoseSearchFeatureChannel_Position*
+		{
+			if (const UPoseSearchFeatureChannel_Position* Position = Cast<UPoseSearchFeatureChannel_Position>(Channel))
+			{
+				if (Position->Bone.BoneName == InBoneName && Position->OriginBone.BoneName == NAME_None && Position->SampleTimeOffset == InSampleTimeOffset)
+				{
+					return Position;
+				}
+			}
+			return nullptr;
+		}))
+	{
+		UPoseSearchFeatureChannel_Position* Position = NewObject<UPoseSearchFeatureChannel_Position>(Schema, NAME_None, RF_Transient);
+		Position->Bone.BoneName = InBoneName;
+		Position->Weight = 0.f;
+		Position->SampleTimeOffset = InSampleTimeOffset;
+		Position->ColorPresetIndex = InColorPresetIndex;
+		Position->Finalize(Schema);
+		Schema->FinalizedChannels.Add(Position);
+	}
+}
+
 void UPoseSearchFeatureChannel_Position::Finalize(UPoseSearchSchema* Schema)
 {
 	ChannelDataOffset = Schema->SchemaCardinality;
@@ -18,6 +42,18 @@ void UPoseSearchFeatureChannel_Position::Finalize(UPoseSearchSchema* Schema)
 	SchemaOriginBoneIdx = Schema->AddBoneReference(OriginBone);
 }
 
+void UPoseSearchFeatureChannel_Position::AddDependentChannels(UPoseSearchSchema* Schema) const
+{
+	if (Schema->bInjectAdditionalDebugChannels)
+	{
+		if (OriginBone.BoneName != NAME_None)
+		{
+			UPoseSearchFeatureChannel_Position::FindOrAddToSchema(Schema, Bone.BoneName, SampleTimeOffset, ColorPresetIndex);
+			UPoseSearchFeatureChannel_Position::FindOrAddToSchema(Schema, OriginBone.BoneName, SampleTimeOffset, ColorPresetIndex);
+		}
+	}
+}
+
 void UPoseSearchFeatureChannel_Position::BuildQuery(UE::PoseSearch::FSearchContext& SearchContext, FPoseSearchFeatureVectorBuilder& InOutQuery) const
 {
 	using namespace UE::PoseSearch;
diff --git a/Engine/Plugins/Experimental/Animation/PoseSearch/Source/Runtime/Private/PoseSearchFeatureChannel_Position.h b/Engine/Plugins/Experimental/Animation/PoseSearch/Source/Runtime/Private/PoseSearchFeatureChannel_Position.h
--- a/Engine/Plugins/Experimental/Animation/PoseSearch/Source/Runtime/Private/PoseSearchFeatureChannel_Position.h
+++ b/Engine/Plugins/Experimental/Animation/PoseSearch/Source/Runtime/Private/PoseSearchFeatureChannel_Position.h
@@ -43,6 +43,8 @@ public:
 	virtual void Finalize(UPoseSearchSchema* Schema) override;
 	virtual void BuildQuery(UE::PoseSearch::FSearchContext& SearchContext, FPoseSearchFeatureVectorBuilder& InOutQuery) const override;
 
+	virtual void AddDependentChannels(UPoseSearchSchema* Schema) const override;
+
 #if ENABLE_DRAW_DEBUG
 	virtual void PreDebugDraw(UE::PoseSearch::FDebugDrawParams& DrawParams, TConstArrayView<float> PoseVector) const override;
 	virtual void DebugDraw(const UE::PoseSearch::FDebugDrawParams& DrawParams, TConstArrayView<float> PoseVector) const override;
@@ -53,4 +55,6 @@ public:
 	virtual void IndexAsset(UE::PoseSearch::FAssetIndexer& Indexer) const override;
 	virtual FString GetLabel() const override;
 #endif
+
+	static void FindOrAddToSchema(UPoseSearchSchema* Schema, const FName& InBoneName, float InSampleTimeOffset, int32 InColorPresetIndex);
 };
diff --git a/Engine/Plugins/Experimental/Animation/PoseSearch/Source/Runtime/Private/PoseSearchFeatureChannel_Velocity.cpp b/Engine/Plugins/Experimental/Animation/PoseSearch/Source/Runtime/Private/PoseSearchFeatureChannel_Velocity.cpp
--- a/Engine/Plugins/Experimental/Animation/PoseSearch/Source/Runtime/Private/PoseSearchFeatureChannel_Velocity.cpp
+++ b/Engine/Plugins/Experimental/Animation/PoseSearch/Source/Runtime/Private/PoseSearchFeatureChannel_Velocity.cpp
@@ -8,6 +8,7 @@
 #include "PoseSearch/PoseSearchDatabase.h"
 #include "PoseSearch/PoseSearchHistory.h"
 #include "PoseSearch/PoseSearchSchema.h"
+#include "PoseSearchFeatureChannel_Position.h"
 
 void UPoseSearchFeatureChannel_Velocity::Finalize(UPoseSearchSchema* Schema)
 {
@@ -18,6 +19,14 @@ void UPoseSearchFeatureChannel_Velocity::Finalize(UPoseSearchSchema* Schema)
 	SchemaBoneIdx = Schema->AddBoneReference(Bone);
 }
 
+void UPoseSearchFeatureChannel_Velocity::AddDependentChannels(UPoseSearchSchema* Schema) const
+{
+	if (Schema->bInjectAdditionalDebugChannels)
+	{
+		UPoseSearchFeatureChannel_Position::FindOrAddToSchema(Schema, Bone.BoneName, SampleTimeOffset, ColorPresetIndex);
+	}
+}
+
 void UPoseSearchFeatureChannel_Velocity::BuildQuery(UE::PoseSearch::FSearchContext& SearchContext, FPoseSearchFeatureVectorBuilder& InOutQuery) const
 {
 	using namespace UE::PoseSearch;
diff --git a/Engine/Plugins/Experimental/Animation/PoseSearch/Source/Runtime/Private/PoseSearchFeatureChannel_Velocity.h b/Engine/Plugins/Experimental/Animation/PoseSearch/Source/Runtime/Private/PoseSearchFeatureChannel_Velocity.h
--- a/Engine/Plugins/Experimental/Animation/PoseSearch/Source/Runtime/Private/PoseSearchFeatureChannel_Velocity.h
+++ b/Engine/Plugins/Experimental/Animation/PoseSearch/Source/Runtime/Private/PoseSearchFeatureChannel_Velocity.h
@@ -45,6 +45,8 @@ public:
 	virtual void Finalize(UPoseSearchSchema* Schema) override;
 	virtual void BuildQuery(UE::PoseSearch::FSearchContext& SearchContext, FPoseSearchFeatureVectorBuilder& InOutQuery) const override;
 
+	virtual void AddDependentChannels(UPoseSearchSchema* Schema) const override;
+
 #if ENABLE_DRAW_DEBUG
 	virtual void DebugDraw(const UE::PoseSearch::FDebugDrawParams& DrawParams, TConstArrayView<float> PoseVector) const override;
 #endif // ENABLE_DRAW_DEBUG
diff --git a/Engine/Plugins/Experimental/Animation/PoseSearch/Source/Runtime/Private/PoseSearchSchema.cpp b/Engine/Plugins/Experimental/Animation/PoseSearch/Source/Runtime/Private/PoseSearchSchema.cpp
--- a/Engine/Plugins/Experimental/Animation/PoseSearch/Source/Runtime/Private/PoseSearchSchema.cpp
+++ b/Engine/Plugins/Experimental/Animation/PoseSearch/Source/Runtime/Private/PoseSearchSchema.cpp
@@ -3,7 +3,6 @@
 #include "PoseSearch/PoseSearchSchema.h"
 #include "AnimationRuntime.h"
 #include "PoseSearch/PoseSearchDefines.h"
-#include "PoseSearch/PoseSearchFeatureChannel.h"
 #include "PoseSearch/PoseSearchResult.h"
 #include "UObject/ObjectSaveContext.h"
 
@@ -45,12 +44,9 @@ void UPoseSearchSchema::BuildQuery(UE::PoseSearch::FSearchContext& SearchContext
 
 	InOutQuery.Init(this);
 
-	for (const TObjectPtr<UPoseSearchFeatureChannel>& ChannelPtr : Channels)
+	for (const TObjectPtr<UPoseSearchFeatureChannel>& ChannelPtr : GetChannels())
 	{
-		if (ChannelPtr)
-		{
-			ChannelPtr->BuildQuery(SearchContext, InOutQuery);
-		}
+		ChannelPtr->BuildQuery(SearchContext, InOutQuery);
 	}
 }
 
@@ -70,14 +66,24 @@ void UPoseSearchSchema::Finalize()
 
 	SchemaCardinality = 0;
 
+	FinalizedChannels.Reset();
 	for (const TObjectPtr<UPoseSearchFeatureChannel>& ChannelPtr : Channels)
 	{
 		if (ChannelPtr)
 		{
+			FinalizedChannels.Add(ChannelPtr);
 			ChannelPtr->Finalize(this);
 		}
 	}
 
+	// AddDependentChannels can add channels to FinalizedChannels, so we need a while loop
+	int32 ChannelIndex = 0;
+	while (ChannelIndex < FinalizedChannels.Num())
+	{
+		FinalizedChannels[ChannelIndex]->AddDependentChannels(this);
+		++ChannelIndex;
+	}
+
 	ResolveBoneReferences();
 }
 
diff --git a/Engine/Plugins/Experimental/Animation/PoseSearch/Source/Runtime/Public/PoseSearch/PoseSearchFeatureChannel.h b/Engine/Plugins/Experimental/Animation/PoseSearch/Source/Runtime/Public/PoseSearch/PoseSearchFeatureChannel.h
--- a/Engine/Plugins/Experimental/Animation/PoseSearch/Source/Runtime/Public/PoseSearch/PoseSearchFeatureChannel.h
+++ b/Engine/Plugins/Experimental/Animation/PoseSearch/Source/Runtime/Public/PoseSearch/PoseSearchFeatureChannel.h
@@ -100,6 +100,9 @@ public:
 	virtual TArrayView<TObjectPtr<UPoseSearchFeatureChannel>> GetSubChannels() { return TArrayView<TObjectPtr<UPoseSearchFeatureChannel>>(); }
 	virtual TConstArrayView<TObjectPtr<UPoseSearchFeatureChannel>> GetSubChannels() const { return TConstArrayView<TObjectPtr<UPoseSearchFeatureChannel>>(); }
 
+	// @todo: should this API be under ENABLE_DRAW_DEBUG?
+	virtual void AddDependentChannels(UPoseSearchSchema* Schema) const {}
+
 #if ENABLE_DRAW_DEBUG
 	// API called before DebugDraw to collect shared channel informations such as decoded positions form the PoseVector
 	virtual void PreDebugDraw(UE::PoseSearch::FDebugDrawParams& DrawParams, TConstArrayView<float> PoseVector) const {}
diff --git a/Engine/Plugins/Experimental/Animation/PoseSearch/Source/Runtime/Public/PoseSearch/PoseSearchSchema.h b/Engine/Plugins/Experimental/Animation/PoseSearch/Source/Runtime/Public/PoseSearch/PoseSearchSchema.h
--- a/Engine/Plugins/Experimental/Animation/PoseSearch/Source/Runtime/Public/PoseSearch/PoseSearchSchema.h
+++ b/Engine/Plugins/Experimental/Animation/PoseSearch/Source/Runtime/Public/PoseSearch/PoseSearchSchema.h
@@ -5,17 +5,11 @@
 #include "BoneContainer.h"
 #include "Engine/DataAsset.h"
 #include "Interfaces/Interface_BoneReferenceSkeletonProvider.h"
+#include "PoseSearch/PoseSearchFeatureChannel.h"
 #include "PoseSearchSchema.generated.h"
 
 struct FBoneReference;
-struct FPoseSearchFeatureVectorBuilder;
 class UMirrorDataTable;
-class UPoseSearchFeatureChannel;
-
-namespace UE::PoseSearch
-{
-	struct FSearchContext;
-} // namespace UE::PoseSearch
 
 UENUM()
 enum class EPoseSearchDataPreprocessor : int32
@@ -59,6 +53,10 @@ public:
 	UPROPERTY(EditAnywhere, Instanced, BlueprintReadWrite, Category = "Schema")
 	TArray<TObjectPtr<UPoseSearchFeatureChannel>> Channels;
 
+	// FinalizedChannels gets populated with UPoseSearchFeatureChannel(s) from Channels and additional injected ones during the Finalize
+	UPROPERTY(Transient)
+	TArray<TObjectPtr<UPoseSearchFeatureChannel>> FinalizedChannels;
+
 	// If set, this schema will support mirroring pose search databases
 	UPROPERTY(EditAnywhere, Category = "Schema")
 	TObjectPtr<UMirrorDataTable> MirrorDataTable;
@@ -109,21 +107,28 @@ public:
 	UPROPERTY(EditAnywhere, Category = "Debug", meta = (ExcludeFromHash))
 	TArray<FPoseSearchSchemaColorPreset> ColorPresets;
 	
+	// if bInjectAdditionalDebugChannels is true, channels will be asked to injecting additional channels into this schema.
+	// the original intent is to add UPoseSearchFeatureChannel_Position(s) to help with the complexity of the debug drawing
+	// (the database will have all the necessary positions to draw lines at the right location and time)
+	UPROPERTY(EditAnywhere, Category = "Debug")
+	bool bInjectAdditionalDebugChannels;
+
 	bool IsValid () const;
 
 	float GetSamplingInterval() const { return 1.0f / SampleRate; }
 
+	TConstArrayView<TObjectPtr<UPoseSearchFeatureChannel>> GetChannels() const { return FinalizedChannels; }
+
+	template <typename FindPredicateType>
+	const UPoseSearchFeatureChannel* FindChannel(FindPredicateType FindPredicate) const
+	{
+		return FindChannelRecursive(GetChannels(), FindPredicate);
+	}
+
 	template<typename ChannelType>
 	const ChannelType* FindFirstChannelOfType() const
 	{
-		for (const TObjectPtr<UPoseSearchFeatureChannel>& ChannelPtr : Channels)
-		{
-			if (const ChannelType* Channel = Cast<const ChannelType>(ChannelPtr.Get()))
-			{
-				return Channel;
-			}
-		}
-		return nullptr;
+		return static_cast<const ChannelType*>(FindChannel([this](const UPoseSearchFeatureChannel* Channel) -> const UPoseSearchFeatureChannel* { return Cast<ChannelType>(Channel); }));
 	}
 
 	// UObject
@@ -146,6 +151,27 @@ public:
 	bool IsRootBone(int8 SchemaBoneIdx) const;
 	
 private:
+	template <typename FindPredicateType>
+	static const UPoseSearchFeatureChannel* FindChannelRecursive(TConstArrayView<TObjectPtr<UPoseSearchFeatureChannel>> Channels, FindPredicateType FindPredicate)
+	{
+		for (const TObjectPtr<UPoseSearchFeatureChannel>& ChannelPtr : Channels)
+		{
+			if (ChannelPtr)
+			{
+				if (const UPoseSearchFeatureChannel* Channel = FindPredicate(ChannelPtr))
+				{
+					return Channel;
+				}
+
+				if (const UPoseSearchFeatureChannel* Channel = FindChannelRecursive(ChannelPtr->GetSubChannels(), FindPredicate))
+				{
+					return Channel;
+				}
+			}
+		}
+		return nullptr;
+	}
+
 	void Finalize();
 	void ResolveBoneReferences();
 };
