motion matching
- removing PoseSearchFeatureChannel_PermutationType, and embedding the permutation logic into specific channels requiring it (only UPoseSearchFeatureChannel_Position for now)
- post Finalize channels (of schemas with valid skeleton) can only hold to valid bones (or defaulted with error message to root bone)
- automatically injecting missing UPoseSearchFeatureChannel_PermutationTime in schema requiring it (schemas with at least one channel returning GetPermutationTimeType() != EPermutationTimeType::UseSampleTime)
- WIP schema permutation debug drawing improvement

diff --git a/Engine/Plugins/Experimental/Animation/PoseSearch/Source/Runtime/Private/PoseSearchAssetIndexer.cpp b/Engine/Plugins/Experimental/Animation/PoseSearch/Source/Runtime/Private/PoseSearchAssetIndexer.cpp
--- a/Engine/Plugins/Experimental/Animation/PoseSearch/Source/Runtime/Private/PoseSearchAssetIndexer.cpp
+++ b/Engine/Plugins/Experimental/Animation/PoseSearch/Source/Runtime/Private/PoseSearchAssetIndexer.cpp
@@ -387,8 +387,12 @@ FTransform FAssetIndexer::CalculateComponentSpaceTransform(FAssetIndexer::Cached
 	return Entry.ComponentSpacePose.GetComponentSpaceTransform(CompactBoneIndex);
 }
 
-FQuat FAssetIndexer::GetSampleRotation(float SampleTimeOffset, int32 SampleIdx, int8 SchemaSampleBoneIdx, int8 SchemaOriginBoneIdx)
+FQuat FAssetIndexer::GetSampleRotation(float SampleTimeOffset, int32 SampleIdx, int8 SchemaSampleBoneIdx, int8 SchemaOriginBoneIdx, EPermutationTimeType PermutationTimeType)
 {
+	float PermutationSampleTimeOffset = 0.f;
+	float PermutationOriginTimeOffset = 0.f;
+	UPoseSearchFeatureChannel::GetPermutationTimeOffsets(PermutationTimeType, CalculatePermutationTimeOffset(), PermutationSampleTimeOffset, PermutationOriginTimeOffset);
+
 	const float Time = SampleIdx * IndexingContext.Schema->GetSamplingInterval();
 	const float SampleTime = Time + SampleTimeOffset + PermutationSampleTimeOffset;
 	const float OriginTime = Time + PermutationOriginTimeOffset;
@@ -415,8 +419,12 @@ FQuat FAssetIndexer::GetSampleRotation(float SampleTimeOffset, int32 SampleIdx,
 	return BoneTransform.GetRotation();
 }
 
-FVector FAssetIndexer::GetSamplePosition(float SampleTimeOffset, int32 SampleIdx, int8 SchemaSampleBoneIdx, int8 SchemaOriginBoneIdx)
+FVector FAssetIndexer::GetSamplePosition(float SampleTimeOffset, int32 SampleIdx, int8 SchemaSampleBoneIdx, int8 SchemaOriginBoneIdx, EPermutationTimeType PermutationTimeType)
 {
+	float PermutationSampleTimeOffset = 0.f;
+	float PermutationOriginTimeOffset = 0.f;
+	UPoseSearchFeatureChannel::GetPermutationTimeOffsets(PermutationTimeType, CalculatePermutationTimeOffset(), PermutationSampleTimeOffset, PermutationOriginTimeOffset);
+
 	const float Time = SampleIdx * IndexingContext.Schema->GetSamplingInterval();
 	const float SampleTime = Time + SampleTimeOffset + PermutationSampleTimeOffset;
 	const float OriginTime = Time + PermutationOriginTimeOffset;
@@ -456,8 +464,12 @@ FVector FAssetIndexer::GetSamplePositionInternal(float SampleTime, float OriginT
 	return RootBoneTransform.InverseTransformVector(DeltaBoneTranslation);
 }
 
-FVector FAssetIndexer::GetSampleVelocity(float SampleTimeOffset, int32 SampleIdx, int8 SchemaSampleBoneIdx, int8 SchemaOriginBoneIdx, bool bUseCharacterSpaceVelocities)
+FVector FAssetIndexer::GetSampleVelocity(float SampleTimeOffset, int32 SampleIdx, int8 SchemaSampleBoneIdx, int8 SchemaOriginBoneIdx, bool bUseCharacterSpaceVelocities, EPermutationTimeType PermutationTimeType)
 {
+	float PermutationSampleTimeOffset = 0.f;
+	float PermutationOriginTimeOffset = 0.f;
+	UPoseSearchFeatureChannel::GetPermutationTimeOffsets(PermutationTimeType, CalculatePermutationTimeOffset(), PermutationSampleTimeOffset, PermutationOriginTimeOffset);
+
 	const float Time = SampleIdx * IndexingContext.Schema->GetSamplingInterval();
 	const float SampleTime = Time + SampleTimeOffset + PermutationSampleTimeOffset;
 	const float OriginTime = Time + PermutationOriginTimeOffset;
@@ -498,20 +510,6 @@ const UPoseSearchSchema* FAssetIndexer::GetSchema() const
 	return IndexingContext.Schema;
 }
 
-void FAssetIndexer::SetPermutationTimeOffsets(float InPermutationSampleTimeOffset, float InPermutationOriginTimeOffset)
-{
-	// right now we disallow having nested channel controlling time offsets
-	check(PermutationSampleTimeOffset == 0.f && PermutationOriginTimeOffset == 0.f);
-	PermutationSampleTimeOffset = InPermutationSampleTimeOffset;
-	PermutationOriginTimeOffset = InPermutationOriginTimeOffset;
-}
-
-void FAssetIndexer::ResetPermutationTimeOffsets()
-{
-	PermutationSampleTimeOffset = 0.f;
-	PermutationOriginTimeOffset = 0.f;
-}
-
 float FAssetIndexer::CalculatePermutationTimeOffset() const
 {
 	const UPoseSearchSchema* Schema = GetSchema();
diff --git a/Engine/Plugins/Experimental/Animation/PoseSearch/Source/Runtime/Private/PoseSearchContext.cpp b/Engine/Plugins/Experimental/Animation/PoseSearch/Source/Runtime/Private/PoseSearchContext.cpp
--- a/Engine/Plugins/Experimental/Animation/PoseSearch/Source/Runtime/Private/PoseSearchContext.cpp
+++ b/Engine/Plugins/Experimental/Animation/PoseSearch/Source/Runtime/Private/PoseSearchContext.cpp
@@ -275,8 +275,12 @@ void FDebugDrawParams::DrawFeatureVector(int32 PoseIdx)
 
 //////////////////////////////////////////////////////////////////////////
 // FSearchContext
-FQuat FSearchContext::GetSampleRotation(float SampleTimeOffset, const UPoseSearchSchema* Schema, int8 SchemaSampleBoneIdx, int8 SchemaOriginBoneIdx, bool bUseHistoryRoot)
+FQuat FSearchContext::GetSampleRotation(float SampleTimeOffset, const UPoseSearchSchema* Schema, int8 SchemaSampleBoneIdx, int8 SchemaOriginBoneIdx, bool bUseHistoryRoot, EPermutationTimeType PermutationTimeType)
 {
+	float PermutationSampleTimeOffset = 0.f;
+	float PermutationOriginTimeOffset = 0.f;
+	UPoseSearchFeatureChannel::GetPermutationTimeOffsets(PermutationTimeType, DesiredPermutationTimeOffset, PermutationSampleTimeOffset, PermutationOriginTimeOffset);
+
 	// @todo: add support for SchemaSampleBoneIdx
 	if (SchemaOriginBoneIdx != RootSchemaBoneIdx)
 	{
@@ -297,15 +301,23 @@ FQuat FSearchContext::GetSampleRotation(float SampleTimeOffset, const UPoseSearc
 	return GetComponentSpaceTransform(SampleTime, Schema, SchemaSampleBoneIdx).GetRotation();
 }
 
-FVector FSearchContext::GetSamplePosition(float SampleTimeOffset, const UPoseSearchSchema* Schema, int8 SchemaSampleBoneIdx, int8 SchemaOriginBoneIdx, bool bUseHistoryRoot)
+FVector FSearchContext::GetSamplePosition(float SampleTimeOffset, const UPoseSearchSchema* Schema, int8 SchemaSampleBoneIdx, int8 SchemaOriginBoneIdx, bool bUseHistoryRoot, EPermutationTimeType PermutationTimeType)
 {
+	float PermutationSampleTimeOffset = 0.f;
+	float PermutationOriginTimeOffset = 0.f;
+	UPoseSearchFeatureChannel::GetPermutationTimeOffsets(PermutationTimeType, DesiredPermutationTimeOffset, PermutationSampleTimeOffset, PermutationOriginTimeOffset);
+
 	const float SampleTime = SampleTimeOffset + PermutationSampleTimeOffset;
 	const float OriginTime = PermutationOriginTimeOffset;
 	return GetSamplePositionInternal(SampleTime, OriginTime, Schema, SchemaSampleBoneIdx, SchemaOriginBoneIdx, bUseHistoryRoot);
 }
 
-FVector FSearchContext::GetSampleVelocity(float SampleTimeOffset, const UPoseSearchSchema* Schema, int8 SchemaSampleBoneIdx, int8 SchemaOriginBoneIdx, bool bUseCharacterSpaceVelocities, bool bUseHistoryRoot)
+FVector FSearchContext::GetSampleVelocity(float SampleTimeOffset, const UPoseSearchSchema* Schema, int8 SchemaSampleBoneIdx, int8 SchemaOriginBoneIdx, bool bUseCharacterSpaceVelocities, bool bUseHistoryRoot, EPermutationTimeType PermutationTimeType)
 {
+	float PermutationSampleTimeOffset = 0.f;
+	float PermutationOriginTimeOffset = 0.f;
+	UPoseSearchFeatureChannel::GetPermutationTimeOffsets(PermutationTimeType, DesiredPermutationTimeOffset, PermutationSampleTimeOffset, PermutationOriginTimeOffset);
+
 	const float SampleTime = SampleTimeOffset + PermutationSampleTimeOffset;
 	const float OriginTime = PermutationOriginTimeOffset;
 	const float FiniteDelta = History ? History->GetSampleTimeInterval() : 1 / 60.0f;
@@ -404,20 +416,6 @@ FVector FSearchContext::GetSamplePositionInternal(float SampleTime, float Origin
 	return RootBoneTransform.InverseTransformVector(DeltaBoneTranslation);
 }
 
-void FSearchContext::SetPermutationTimeOffsets(float InPermutationSampleTimeOffset, float InPermutationOriginTimeOffset)
-{
-	// right now we disallow having nested channel controlling time offsets
-	check(PermutationSampleTimeOffset == 0.f && PermutationOriginTimeOffset == 0.f);
-	PermutationSampleTimeOffset = InPermutationSampleTimeOffset;
-	PermutationOriginTimeOffset = InPermutationOriginTimeOffset;
-}
-
-void FSearchContext::ResetPermutationTimeOffsets()
-{
-	PermutationSampleTimeOffset = 0.f;
-	PermutationOriginTimeOffset = 0.f;
-}
-
 void FSearchContext::ClearCachedEntries()
 {
 	CachedTransforms.Reset();
diff --git a/Engine/Plugins/Experimental/Animation/PoseSearch/Source/Runtime/Private/PoseSearchDerivedDataKey.cpp b/Engine/Plugins/Experimental/Animation/PoseSearch/Source/Runtime/Private/PoseSearchDerivedDataKey.cpp
--- a/Engine/Plugins/Experimental/Animation/PoseSearch/Source/Runtime/Private/PoseSearchDerivedDataKey.cpp
+++ b/Engine/Plugins/Experimental/Animation/PoseSearch/Source/Runtime/Private/PoseSearchDerivedDataKey.cpp
@@ -41,7 +41,7 @@ FKeyBuilder::FKeyBuilder(const UObject* Object, bool bUseDataVer)
 	if (bUseDataVer)
 	{
 		// used to invalidate the key without having to change POSESEARCHDB_DERIVEDDATA_VER all the times
-		int32 POSESEARCHDB_DERIVEDDATA_VER_SMALL = 113;
+		int32 POSESEARCHDB_DERIVEDDATA_VER_SMALL = 121;
 		FGuid VersionGuid = FDevSystemGuids::GetSystemGuid(FDevSystemGuids::Get().POSESEARCHDB_DERIVEDDATA_VER);
 
 		*this << VersionGuid;
diff --git a/Engine/Plugins/Experimental/Animation/PoseSearch/Source/Runtime/Private/PoseSearchFeatureChannel.cpp b/Engine/Plugins/Experimental/Animation/PoseSearch/Source/Runtime/Private/PoseSearchFeatureChannel.cpp
--- a/Engine/Plugins/Experimental/Animation/PoseSearch/Source/Runtime/Private/PoseSearchFeatureChannel.cpp
+++ b/Engine/Plugins/Experimental/Animation/PoseSearch/Source/Runtime/Private/PoseSearchFeatureChannel.cpp
@@ -157,6 +157,30 @@ float FFeatureVectorHelper::DecodeFloat(TConstArrayView<float> Values, int32 Dat
 
 //////////////////////////////////////////////////////////////////////////
 // UPoseSearchFeatureChannel
+void UPoseSearchFeatureChannel::GetPermutationTimeOffsets(EPermutationTimeType PermutationTimeType, float DesiredPermutationTimeOffset, float& OutPermutationSampleTimeOffset, float& OutPermutationOriginTimeOffset)
+{
+	switch (PermutationTimeType)
+	{
+	case EPermutationTimeType::UseSampleTime:
+		OutPermutationSampleTimeOffset = 0.f;
+		OutPermutationOriginTimeOffset = 0.f;
+		break;
+	case EPermutationTimeType::UsePermutationTime:
+		OutPermutationSampleTimeOffset = DesiredPermutationTimeOffset;
+		OutPermutationOriginTimeOffset = DesiredPermutationTimeOffset;
+		break;
+	case EPermutationTimeType::UseSampleToPermutationTime:
+		OutPermutationSampleTimeOffset = DesiredPermutationTimeOffset;
+		OutPermutationOriginTimeOffset = 0.f;
+		break;
+	default:
+		checkNoEntry();
+		OutPermutationSampleTimeOffset = 0.f;
+		OutPermutationOriginTimeOffset = 0.f;
+		break;
+	}
+}
+
 #if WITH_EDITOR
 FString UPoseSearchFeatureChannel::GetLabel() const
 {
diff --git a/Engine/Plugins/Experimental/Animation/PoseSearch/Source/Runtime/Private/PoseSearchFeatureChannel_FilterCrashingLegs.cpp b/Engine/Plugins/Experimental/Animation/PoseSearch/Source/Runtime/Private/PoseSearchFeatureChannel_FilterCrashingLegs.cpp
--- a/Engine/Plugins/Experimental/Animation/PoseSearch/Source/Runtime/Private/PoseSearchFeatureChannel_FilterCrashingLegs.cpp
+++ b/Engine/Plugins/Experimental/Animation/PoseSearch/Source/Runtime/Private/PoseSearchFeatureChannel_FilterCrashingLegs.cpp
@@ -43,10 +43,10 @@ void UPoseSearchFeatureChannel_FilterCrashingLegs::AddDependentChannels(UPoseSea
 {
 	if (Schema->bInjectAdditionalDebugChannels)
 	{
-		UPoseSearchFeatureChannel_Position::FindOrAddToSchema(Schema, LeftThigh.BoneName, 0.f, 0);
-		UPoseSearchFeatureChannel_Position::FindOrAddToSchema(Schema, RightThigh.BoneName, 0.f, 0);
-		UPoseSearchFeatureChannel_Position::FindOrAddToSchema(Schema, LeftFoot.BoneName, 0.f, 0);
-		UPoseSearchFeatureChannel_Position::FindOrAddToSchema(Schema, RightFoot.BoneName, 0.f, 0);
+		UPoseSearchFeatureChannel_Position::FindOrAddToSchema(Schema, 0.f, 0, LeftThigh.BoneName);
+		UPoseSearchFeatureChannel_Position::FindOrAddToSchema(Schema, 0.f, 0, RightThigh.BoneName);
+		UPoseSearchFeatureChannel_Position::FindOrAddToSchema(Schema, 0.f, 0, LeftFoot.BoneName);
+		UPoseSearchFeatureChannel_Position::FindOrAddToSchema(Schema, 0.f, 0, RightFoot.BoneName);
 	}
 }
 void UPoseSearchFeatureChannel_FilterCrashingLegs::BuildQuery(UE::PoseSearch::FSearchContext& SearchContext, FPoseSearchFeatureVectorBuilder& InOutQuery) const
diff --git a/Engine/Plugins/Experimental/Animation/PoseSearch/Source/Runtime/Private/PoseSearchFeatureChannel_Heading.cpp b/Engine/Plugins/Experimental/Animation/PoseSearch/Source/Runtime/Private/PoseSearchFeatureChannel_Heading.cpp
--- a/Engine/Plugins/Experimental/Animation/PoseSearch/Source/Runtime/Private/PoseSearchFeatureChannel_Heading.cpp
+++ b/Engine/Plugins/Experimental/Animation/PoseSearch/Source/Runtime/Private/PoseSearchFeatureChannel_Heading.cpp
@@ -21,7 +21,7 @@ void UPoseSearchFeatureChannel_Heading::AddDependentChannels(UPoseSearchSchema*
 {
 	if (Schema->bInjectAdditionalDebugChannels)
 	{
-		UPoseSearchFeatureChannel_Position::FindOrAddToSchema(Schema, Bone.BoneName, SampleTimeOffset, ColorPresetIndex);
+		UPoseSearchFeatureChannel_Position::FindOrAddToSchema(Schema, SampleTimeOffset, ColorPresetIndex, Bone.BoneName);
 	}
 }
 
@@ -47,8 +47,8 @@ void UPoseSearchFeatureChannel_Heading::BuildQuery(UE::PoseSearch::FSearchContex
 
 	const bool bIsCurrentResultValid = SearchContext.CurrentResult.IsValid() && SearchContext.CurrentResult.Database->Schema == InOutQuery.GetSchema();
 	const bool bSkip = InputQueryPose != EInputQueryPose::UseCharacterPose && bIsCurrentResultValid;
-	const bool bBoneValid = InOutQuery.GetSchema()->BoneReferences[SchemaBoneIdx].HasValidSetup();
-	if (bSkip || (!SearchContext.History && bBoneValid))
+	const bool bIsRootBone = InOutQuery.GetSchema()->IsRootBone(SchemaBoneIdx);
+	if (bSkip || (!SearchContext.History && !bIsRootBone))
 	{
 		if (bIsCurrentResultValid)
 		{
@@ -61,7 +61,7 @@ void UPoseSearchFeatureChannel_Heading::BuildQuery(UE::PoseSearch::FSearchContex
 	else
 	{
 		// calculating the BoneRotation in component space for the bone indexed by SchemaBoneIdx
-		const FQuat BoneRotation = SearchContext.GetSampleRotation(SampleTimeOffset, InOutQuery.GetSchema(), SchemaBoneIdx, RootSchemaBoneIdx, bBoneValid);
+		const FQuat BoneRotation = SearchContext.GetSampleRotation(SampleTimeOffset, InOutQuery.GetSchema(), SchemaBoneIdx, RootSchemaBoneIdx, !bIsRootBone);
 		FFeatureVectorHelper::EncodeVector(InOutQuery.EditValues(), ChannelDataOffset, GetAxis(BoneRotation), ComponentStripping);
 	}
 }
@@ -132,11 +132,12 @@ FString UPoseSearchFeatureChannel_Heading::GetLabel() const
 		Label.Append(TEXT("_xy"));
 	}
 
-	const FBoneReference& BoneReference = GetSchema()->BoneReferences[SchemaBoneIdx];
-	if (BoneReference.HasValidSetup())
+	const UPoseSearchSchema* Schema = GetSchema();
+	check(Schema);
+	if (!Schema->IsRootBone(SchemaBoneIdx))
 	{
 		Label.Append(TEXT("_"));
-		Label.Append(BoneReference.BoneName.ToString());
+		Label.Append(Schema->BoneReferences[SchemaBoneIdx].BoneName.ToString());
 	}
 
 	Label.Appendf(TEXT(" %.1f"), SampleTimeOffset);
diff --git a/Engine/Plugins/Experimental/Animation/PoseSearch/Source/Runtime/Private/PoseSearchFeatureChannel_PermutationTime.cpp b/Engine/Plugins/Experimental/Animation/PoseSearch/Source/Runtime/Private/PoseSearchFeatureChannel_PermutationTime.cpp
--- a/Engine/Plugins/Experimental/Animation/PoseSearch/Source/Runtime/Private/PoseSearchFeatureChannel_PermutationTime.cpp
+++ b/Engine/Plugins/Experimental/Animation/PoseSearch/Source/Runtime/Private/PoseSearchFeatureChannel_PermutationTime.cpp
@@ -5,6 +5,17 @@
 #include "PoseSearch/PoseSearchContext.h"
 #include "PoseSearch/PoseSearchSchema.h"
 
+void UPoseSearchFeatureChannel_PermutationTime::FindOrAddToSchema(UPoseSearchSchema* Schema)
+{
+	if (Schema->FindFirstChannelOfType<UPoseSearchFeatureChannel_PermutationTime>() == nullptr)
+	{
+		UPoseSearchFeatureChannel_PermutationTime* PermutationTime = NewObject<UPoseSearchFeatureChannel_PermutationTime>(Schema, NAME_None, RF_Transient);
+		PermutationTime->Weight = 0.f;
+		PermutationTime->Finalize(Schema);
+		Schema->FinalizedChannels.Add(PermutationTime);
+	}
+}
+
 void UPoseSearchFeatureChannel_PermutationTime::Finalize(UPoseSearchSchema* Schema)
 {
 	ChannelDataOffset = Schema->SchemaCardinality;
diff --git a/Engine/Plugins/Experimental/Animation/PoseSearch/Source/Runtime/Private/PoseSearchFeatureChannel_PermutationTime.h b/Engine/Plugins/Experimental/Animation/PoseSearch/Source/Runtime/Private/PoseSearchFeatureChannel_PermutationTime.h
--- a/Engine/Plugins/Experimental/Animation/PoseSearch/Source/Runtime/Private/PoseSearchFeatureChannel_PermutationTime.h
+++ b/Engine/Plugins/Experimental/Animation/PoseSearch/Source/Runtime/Private/PoseSearchFeatureChannel_PermutationTime.h
@@ -23,4 +23,6 @@ public:
 	virtual void IndexAsset(UE::PoseSearch::FAssetIndexer& Indexer) const override;
 	virtual FString GetLabel() const override;
 #endif
+
+	static void FindOrAddToSchema(UPoseSearchSchema* Schema);
 };
diff --git a/Engine/Plugins/Experimental/Animation/PoseSearch/Source/Runtime/Private/PoseSearchFeatureChannel_PermutationType.cpp b/Engine/Plugins/Experimental/Animation/PoseSearch/Source/Runtime/Private/PoseSearchFeatureChannel_PermutationType.cpp
deleted file mode 100644
diff --git a/Engine/Plugins/Experimental/Animation/PoseSearch/Source/Runtime/Private/PoseSearchFeatureChannel_PermutationType.h b/Engine/Plugins/Experimental/Animation/PoseSearch/Source/Runtime/Private/PoseSearchFeatureChannel_PermutationType.h
deleted file mode 100644
diff --git a/Engine/Plugins/Experimental/Animation/PoseSearch/Source/Runtime/Private/PoseSearchFeatureChannel_Phase.cpp b/Engine/Plugins/Experimental/Animation/PoseSearch/Source/Runtime/Private/PoseSearchFeatureChannel_Phase.cpp
--- a/Engine/Plugins/Experimental/Animation/PoseSearch/Source/Runtime/Private/PoseSearchFeatureChannel_Phase.cpp
+++ b/Engine/Plugins/Experimental/Animation/PoseSearch/Source/Runtime/Private/PoseSearchFeatureChannel_Phase.cpp
@@ -292,7 +292,7 @@ void UPoseSearchFeatureChannel_Phase::AddDependentChannels(UPoseSearchSchema* Sc
 {
 	if (Schema->bInjectAdditionalDebugChannels)
 	{
-		UPoseSearchFeatureChannel_Position::FindOrAddToSchema(Schema, Bone.BoneName, 0.f, ColorPresetIndex);
+		UPoseSearchFeatureChannel_Position::FindOrAddToSchema(Schema, 0.f, ColorPresetIndex, Bone.BoneName);
 	}
 }
 
@@ -405,11 +405,13 @@ FString UPoseSearchFeatureChannel_Phase::GetLabel() const
 	}
 
 	Label.Append(TEXT("Pha"));
-	const FBoneReference& BoneReference = GetSchema()->BoneReferences[SchemaBoneIdx];
-	if (BoneReference.HasValidSetup())
+
+	const UPoseSearchSchema* Schema = GetSchema();
+	check(Schema);
+	if (!Schema->IsRootBone(SchemaBoneIdx))
 	{
 		Label.Append(TEXT("_"));
-		Label.Append(BoneReference.BoneName.ToString());
+		Label.Append(Schema->BoneReferences[SchemaBoneIdx].BoneName.ToString());
 	}
 
 	return Label.ToString();
diff --git a/Engine/Plugins/Experimental/Animation/PoseSearch/Source/Runtime/Private/PoseSearchFeatureChannel_Position.cpp b/Engine/Plugins/Experimental/Animation/PoseSearch/Source/Runtime/Private/PoseSearchFeatureChannel_Position.cpp
--- a/Engine/Plugins/Experimental/Animation/PoseSearch/Source/Runtime/Private/PoseSearchFeatureChannel_Position.cpp
+++ b/Engine/Plugins/Experimental/Animation/PoseSearch/Source/Runtime/Private/PoseSearchFeatureChannel_Position.cpp
@@ -8,13 +8,13 @@
 #include "PoseSearch/PoseSearchDatabase.h"
 #include "PoseSearch/PoseSearchSchema.h"
 
-void UPoseSearchFeatureChannel_Position::FindOrAddToSchema(UPoseSearchSchema* Schema, const FName& InBoneName, float InSampleTimeOffset, int32 InColorPresetIndex)
+void UPoseSearchFeatureChannel_Position::FindOrAddToSchema(UPoseSearchSchema* Schema, float SampleTimeOffset, int32 ColorPresetIndex, const FName& BoneName, EPermutationTimeType PermutationTimeType)
 {
-	if (!Schema->FindChannel([&InBoneName, InSampleTimeOffset](const UPoseSearchFeatureChannel* Channel) -> const UPoseSearchFeatureChannel_Position*
+	if (!Schema->FindChannel([SampleTimeOffset, &BoneName, PermutationTimeType](const UPoseSearchFeatureChannel* Channel) -> const UPoseSearchFeatureChannel_Position*
 		{
 			if (const UPoseSearchFeatureChannel_Position* Position = Cast<UPoseSearchFeatureChannel_Position>(Channel))
 			{
-				if (Position->Bone.BoneName == InBoneName && Position->OriginBone.BoneName == NAME_None && Position->SampleTimeOffset == InSampleTimeOffset)
+				if (Position->Bone.BoneName == BoneName && Position->OriginBone.BoneName == NAME_None && Position->SampleTimeOffset == SampleTimeOffset && Position->PermutationTimeType == PermutationTimeType)
 				{
 					return Position;
 				}
@@ -23,10 +23,11 @@ void UPoseSearchFeatureChannel_Position::FindOrAddToSchema(UPoseSearchSchema* Sc
 		}))
 	{
 		UPoseSearchFeatureChannel_Position* Position = NewObject<UPoseSearchFeatureChannel_Position>(Schema, NAME_None, RF_Transient);
-		Position->Bone.BoneName = InBoneName;
+		Position->Bone.BoneName = BoneName;
 		Position->Weight = 0.f;
-		Position->SampleTimeOffset = InSampleTimeOffset;
-		Position->ColorPresetIndex = InColorPresetIndex;
+		Position->SampleTimeOffset = SampleTimeOffset;
+		Position->ColorPresetIndex = ColorPresetIndex;
+		Position->PermutationTimeType = PermutationTimeType;
 		Position->Finalize(Schema);
 		Schema->FinalizedChannels.Add(Position);
 	}
@@ -46,10 +47,10 @@ void UPoseSearchFeatureChannel_Position::AddDependentChannels(UPoseSearchSchema*
 {
 	if (Schema->bInjectAdditionalDebugChannels)
 	{
-		if (OriginBone.BoneName != NAME_None)
+		if (!Schema->IsRootBone(SchemaOriginBoneIdx))
 		{
-			UPoseSearchFeatureChannel_Position::FindOrAddToSchema(Schema, Bone.BoneName, SampleTimeOffset, ColorPresetIndex);
-			UPoseSearchFeatureChannel_Position::FindOrAddToSchema(Schema, OriginBone.BoneName, SampleTimeOffset, ColorPresetIndex);
+			const EPermutationTimeType DependentChannelsPermutationTimeType = PermutationTimeType == EPermutationTimeType::UsePermutationTime ? EPermutationTimeType::UseSampleToPermutationTime : EPermutationTimeType::UseSampleTime;
+			UPoseSearchFeatureChannel_Position::FindOrAddToSchema(Schema, 0.f, ColorPresetIndex, OriginBone.BoneName, DependentChannelsPermutationTimeType);
 		}
 	}
 }
@@ -61,8 +62,8 @@ void UPoseSearchFeatureChannel_Position::BuildQuery(UE::PoseSearch::FSearchConte
 	check(InOutQuery.GetSchema());
 	const bool bIsCurrentResultValid = SearchContext.CurrentResult.IsValid() && SearchContext.CurrentResult.Database->Schema == InOutQuery.GetSchema();
 	const bool bSkip = InputQueryPose != EInputQueryPose::UseCharacterPose && bIsCurrentResultValid;
-	const bool bBoneValid = InOutQuery.GetSchema()->BoneReferences[SchemaBoneIdx].HasValidSetup();
-	if (bSkip || (!SearchContext.History && bBoneValid))
+	const bool bIsRootBone = InOutQuery.GetSchema()->IsRootBone(SchemaBoneIdx);
+	if (bSkip || (!SearchContext.History && !bIsRootBone))
 	{
 		if (bIsCurrentResultValid)
 		{
@@ -74,7 +75,7 @@ void UPoseSearchFeatureChannel_Position::BuildQuery(UE::PoseSearch::FSearchConte
 	else
 	{
 		// calculating the BonePosition in component space for the bone indexed by SchemaBoneIdx
-		const FVector BonePosition = SearchContext.GetSamplePosition(SampleTimeOffset, InOutQuery.GetSchema(), SchemaBoneIdx, SchemaOriginBoneIdx, bBoneValid);
+		const FVector BonePosition = SearchContext.GetSamplePosition(SampleTimeOffset, InOutQuery.GetSchema(), SchemaBoneIdx, SchemaOriginBoneIdx, !bIsRootBone, PermutationTimeType);
 		FFeatureVectorHelper::EncodeVector(InOutQuery.EditValues(), ChannelDataOffset, BonePosition, ComponentStripping);
 	}
 }
@@ -85,7 +86,7 @@ void UPoseSearchFeatureChannel_Position::PreDebugDraw(UE::PoseSearch::FDebugDraw
 	using namespace UE::PoseSearch;
 	
 	// if the SchemaOriginBoneIdx is not the root bone, the feature vector position doesn't represent a component space delta position for the SchemaBoneIdx, so we shouldn't collect it
-	if (!DrawParams.GetSchema()->BoneReferences[SchemaOriginBoneIdx].HasValidSetup())
+	if (DrawParams.GetSchema()->IsRootBone(SchemaOriginBoneIdx) && PermutationTimeType != EPermutationTimeType::UsePermutationTime)
 	{
 		const FVector BonePos = DrawParams.GetRootTransform().TransformPosition(FFeatureVectorHelper::DecodeVector(PoseVector, ChannelDataOffset, ComponentStripping));
 		DrawParams.AddCachedPosition(SampleTimeOffset, SchemaBoneIdx, BonePos);
@@ -98,18 +99,20 @@ void UPoseSearchFeatureChannel_Position::DebugDraw(const UE::PoseSearch::FDebugD
 
 	const FColor Color = DrawParams.GetColor(ColorPresetIndex);
 
-	if (!DrawParams.GetSchema()->BoneReferences[SchemaOriginBoneIdx].HasValidSetup())
+	const FVector FeaturesVector = FFeatureVectorHelper::DecodeVector(PoseVector, ChannelDataOffset, ComponentStripping);
+	if (DrawParams.GetSchema()->IsRootBone(SchemaOriginBoneIdx))
 	{
-		const FVector BonePos = DrawParams.GetRootTransform().TransformPosition(FFeatureVectorHelper::DecodeVector(PoseVector, ChannelDataOffset, ComponentStripping));
-
+		const FVector BonePos = DrawParams.GetRootTransform().TransformPosition(FeaturesVector);
 		DrawParams.DrawPoint(BonePos, Color);
 	}
 	else
 	{
 		const FVector OriginBonePos = DrawParams.GetCachedPosition(SampleTimeOffset, SchemaOriginBoneIdx);
-		const FVector DeltaPos = DrawParams.GetRootTransform().TransformVector(FFeatureVectorHelper::DecodeVector(PoseVector, ChannelDataOffset, ComponentStripping));
-
-		DrawParams.DrawLine(OriginBonePos, OriginBonePos + DeltaPos, Color);
+		const FVector DeltaPos = DrawParams.GetRootTransform().TransformVector(FeaturesVector);
+		const FVector BonePos = OriginBonePos + DeltaPos;
+		DrawParams.DrawLine(OriginBonePos, BonePos, Color);
+		DrawParams.DrawPoint(OriginBonePos, Color);
+		DrawParams.DrawPoint(BonePos, Color);
 	}	
 }
 #endif // ENABLE_DRAW_DEBUG
@@ -129,7 +132,7 @@ void UPoseSearchFeatureChannel_Position::IndexAsset(UE::PoseSearch::FAssetIndexe
 
 	for (int32 SampleIdx = Indexer.GetBeginSampleIdx(); SampleIdx != Indexer.GetEndSampleIdx(); ++SampleIdx)
 	{
-		const FVector BonePosition = Indexer.GetSamplePosition(SampleTimeOffset, SampleIdx, SchemaBoneIdx, SchemaOriginBoneIdx);
+		const FVector BonePosition = Indexer.GetSamplePosition(SampleTimeOffset, SampleIdx, SchemaBoneIdx, SchemaOriginBoneIdx, PermutationTimeType);
 		FFeatureVectorHelper::EncodeVector(Indexer.GetPoseVector(SampleIdx), ChannelDataOffset, BonePosition, ComponentStripping);
 	}
 }
@@ -154,18 +157,18 @@ FString UPoseSearchFeatureChannel_Position::GetLabel() const
 		Label.Append(TEXT("_xy"));
 	}
 
-	const FBoneReference& BoneReference = GetSchema()->BoneReferences[SchemaBoneIdx];
-	if (BoneReference.HasValidSetup())
+	const UPoseSearchSchema* Schema = GetSchema();
+	check(Schema);
+	if (!Schema->IsRootBone(SchemaBoneIdx))
 	{
 		Label.Append(TEXT("_"));
-		Label.Append(BoneReference.BoneName.ToString());
+		Label.Append(Schema->BoneReferences[SchemaBoneIdx].BoneName.ToString());
 	}
 
-	const FBoneReference& OriginBoneReference = GetSchema()->BoneReferences[SchemaOriginBoneIdx];
-	if (OriginBoneReference.HasValidSetup())
+	if (!Schema->IsRootBone(SchemaOriginBoneIdx))
 	{
 		Label.Append(TEXT("_"));
-		Label.Append(OriginBoneReference.BoneName.ToString());
+		Label.Append(Schema->BoneReferences[SchemaOriginBoneIdx].BoneName.ToString());
 	}
 
 	Label.Appendf(TEXT(" %.1f"), SampleTimeOffset);
diff --git a/Engine/Plugins/Experimental/Animation/PoseSearch/Source/Runtime/Private/PoseSearchFeatureChannel_Position.h b/Engine/Plugins/Experimental/Animation/PoseSearch/Source/Runtime/Private/PoseSearchFeatureChannel_Position.h
--- a/Engine/Plugins/Experimental/Animation/PoseSearch/Source/Runtime/Private/PoseSearchFeatureChannel_Position.h
+++ b/Engine/Plugins/Experimental/Animation/PoseSearch/Source/Runtime/Private/PoseSearchFeatureChannel_Position.h
@@ -39,10 +39,14 @@ public:
 	UPROPERTY(EditAnywhere, Category = "Settings")
 	EComponentStrippingVector ComponentStripping = EComponentStrippingVector::None;
 
+	UPROPERTY(EditAnywhere, Category = "Settings")
+	EPermutationTimeType PermutationTimeType = EPermutationTimeType::UseSampleTime;
+
 	// UPoseSearchFeatureChannel interface
 	virtual void Finalize(UPoseSearchSchema* Schema) override;
 	virtual void BuildQuery(UE::PoseSearch::FSearchContext& SearchContext, FPoseSearchFeatureVectorBuilder& InOutQuery) const override;
 
+	virtual EPermutationTimeType GetPermutationTimeType() const override { return PermutationTimeType; }
 	virtual void AddDependentChannels(UPoseSearchSchema* Schema) const override;
 
 #if ENABLE_DRAW_DEBUG
@@ -56,5 +60,5 @@ public:
 	virtual FString GetLabel() const override;
 #endif
 
-	static void FindOrAddToSchema(UPoseSearchSchema* Schema, const FName& InBoneName, float InSampleTimeOffset, int32 InColorPresetIndex);
+	static void FindOrAddToSchema(UPoseSearchSchema* Schema, float SampleTimeOffset, int32 ColorPresetIndex, const FName& BoneName = NAME_None, EPermutationTimeType PermutationTimeType = EPermutationTimeType::UseSampleTime);
 };
diff --git a/Engine/Plugins/Experimental/Animation/PoseSearch/Source/Runtime/Private/PoseSearchFeatureChannel_Velocity.cpp b/Engine/Plugins/Experimental/Animation/PoseSearch/Source/Runtime/Private/PoseSearchFeatureChannel_Velocity.cpp
--- a/Engine/Plugins/Experimental/Animation/PoseSearch/Source/Runtime/Private/PoseSearchFeatureChannel_Velocity.cpp
+++ b/Engine/Plugins/Experimental/Animation/PoseSearch/Source/Runtime/Private/PoseSearchFeatureChannel_Velocity.cpp
@@ -23,7 +23,7 @@ void UPoseSearchFeatureChannel_Velocity::AddDependentChannels(UPoseSearchSchema*
 {
 	if (Schema->bInjectAdditionalDebugChannels)
 	{
-		UPoseSearchFeatureChannel_Position::FindOrAddToSchema(Schema, Bone.BoneName, SampleTimeOffset, ColorPresetIndex);
+		UPoseSearchFeatureChannel_Position::FindOrAddToSchema(Schema, SampleTimeOffset, ColorPresetIndex, Bone.BoneName);
 	}
 }
 
@@ -33,8 +33,8 @@ void UPoseSearchFeatureChannel_Velocity::BuildQuery(UE::PoseSearch::FSearchConte
 
 	const bool bIsCurrentResultValid = SearchContext.CurrentResult.IsValid() && SearchContext.CurrentResult.Database->Schema == InOutQuery.GetSchema();
 	const bool bSkip = InputQueryPose != EInputQueryPose::UseCharacterPose && bIsCurrentResultValid;
-	const bool bBoneValid = InOutQuery.GetSchema()->BoneReferences[SchemaBoneIdx].HasValidSetup();
-	if (bSkip || (!SearchContext.History && bBoneValid))
+	const bool bIsRootBone = InOutQuery.GetSchema()->IsRootBone(SchemaBoneIdx);
+	if (bSkip || (!SearchContext.History && !bIsRootBone))
 	{
 		if (bIsCurrentResultValid)
 		{
@@ -47,7 +47,7 @@ void UPoseSearchFeatureChannel_Velocity::BuildQuery(UE::PoseSearch::FSearchConte
 	else
 	{
 		// calculating the LinearVelocity for the bone indexed by SchemaBoneIdx
-		FVector LinearVelocity = SearchContext.GetSampleVelocity(SampleTimeOffset, InOutQuery.GetSchema(), SchemaBoneIdx, RootSchemaBoneIdx, bUseCharacterSpaceVelocities, bBoneValid);
+		FVector LinearVelocity = SearchContext.GetSampleVelocity(SampleTimeOffset, InOutQuery.GetSchema(), SchemaBoneIdx, RootSchemaBoneIdx, bUseCharacterSpaceVelocities, !bIsRootBone);
 		if (bNormalize)
 		{
 			LinearVelocity = LinearVelocity.GetClampedToMaxSize(1.f);
@@ -121,11 +121,12 @@ FString UPoseSearchFeatureChannel_Velocity::GetLabel() const
 		Label.Append(TEXT("_xy"));
 	}
 
-	const FBoneReference& BoneReference = GetSchema()->BoneReferences[SchemaBoneIdx];
-	if (BoneReference.HasValidSetup())
+	const UPoseSearchSchema* Schema = GetSchema();
+	check(Schema);
+	if (!Schema->IsRootBone(SchemaBoneIdx))
 	{
 		Label.Append(TEXT("_"));
-		Label.Append(BoneReference.BoneName.ToString());
+		Label.Append(Schema->BoneReferences[SchemaBoneIdx].BoneName.ToString());
 	}
 
 	Label.Appendf(TEXT(" %.1f"), SampleTimeOffset);
diff --git a/Engine/Plugins/Experimental/Animation/PoseSearch/Source/Runtime/Private/PoseSearchSchema.cpp b/Engine/Plugins/Experimental/Animation/PoseSearch/Source/Runtime/Private/PoseSearchSchema.cpp
--- a/Engine/Plugins/Experimental/Animation/PoseSearch/Source/Runtime/Private/PoseSearchSchema.cpp
+++ b/Engine/Plugins/Experimental/Animation/PoseSearch/Source/Runtime/Private/PoseSearchSchema.cpp
@@ -4,6 +4,7 @@
 #include "AnimationRuntime.h"
 #include "PoseSearch/PoseSearchDefines.h"
 #include "PoseSearch/PoseSearchResult.h"
+#include "PoseSearchFeatureChannel_PermutationTime.h"
 #include "UObject/ObjectSaveContext.h"
 
 bool UPoseSearchSchema::IsValid() const
@@ -11,33 +12,6 @@ bool UPoseSearchSchema::IsValid() const
 	return Skeleton != nullptr;
 }
 
-void UPoseSearchSchema::ResolveBoneReferences()
-{
-	BoneIndicesWithParents.Reset();
-	if (Skeleton)
-	{
-		// Initialize references to obtain bone indices and fill out bone index array
-		for (FBoneReference& BoneRef : BoneReferences)
-		{
-			BoneRef.Initialize(Skeleton);
-			if (BoneRef.HasValidSetup())
-			{
-				BoneIndicesWithParents.Add(BoneRef.BoneIndex);
-			}
-		}
-
-		// Build separate index array with parent indices guaranteed to be present. Sort for EnsureParentsPresent.
-		BoneIndicesWithParents.Sort();
-		FAnimationRuntime::EnsureParentsPresent(BoneIndicesWithParents, Skeleton->GetReferenceSkeleton());
-	}
-
-	// BoneIndicesWithParents should at least contain the root to support mirroring root motion
-	if (BoneIndicesWithParents.IsEmpty())
-	{
-		BoneIndicesWithParents.Add(0);
-	}
-}
-
 void UPoseSearchSchema::BuildQuery(UE::PoseSearch::FSearchContext& SearchContext, FPoseSearchFeatureVectorBuilder& InOutQuery) const
 {
 	QUICK_SCOPE_CYCLE_COUNTER(STAT_PoseSearch_BuildQuery);
@@ -52,12 +26,61 @@ void UPoseSearchSchema::BuildQuery(UE::PoseSearch::FSearchContext& SearchContext
 
 FBoneIndexType UPoseSearchSchema::GetBoneIndexType(int8 SchemaBoneIdx) const
 {
-	return SchemaBoneIdx >= 0 && BoneReferences[SchemaBoneIdx].HasValidSetup() ? BoneReferences[SchemaBoneIdx].BoneIndex : RootBoneIndexType;
+	if (SchemaBoneIdx < 0)
+	{
+		return RootBoneIndexType;
+	}
+	check(BoneReferences[SchemaBoneIdx].HasValidSetup());
+	return BoneReferences[SchemaBoneIdx].BoneIndex;
 }
 
 bool UPoseSearchSchema::IsRootBone(int8 SchemaBoneIdx) const
 {
-	return SchemaBoneIdx < 0 || !BoneReferences[SchemaBoneIdx].HasValidSetup();
+	if (SchemaBoneIdx < 0)
+	{
+		return true;
+	}
+
+	check(BoneReferences[SchemaBoneIdx].HasValidSetup());
+	return BoneReferences[SchemaBoneIdx].BoneIndex == int32(RootBoneIndexType);
+}
+
+int8 UPoseSearchSchema::AddBoneReference(const FBoneReference& BoneReference)
+{
+	int32 SchemaBoneIdx = 0;
+	if (Skeleton)
+	{
+		bool bDefaultToRootBone = true;
+		FBoneReference TempBoneReference = BoneReference;
+		if (TempBoneReference.BoneName != NAME_None)
+		{
+			TempBoneReference.Initialize(Skeleton);
+			if (!TempBoneReference.HasValidSetup())
+			{
+				UE_LOG(LogPoseSearch, Error, TEXT("UPoseSearchSchema::AddBoneReference: couldn't initialize FBoneReference '%s' with Skeleton '%s' in UPoseSearchSchema '%s'. Defaulting to root bone instead"),
+					*TempBoneReference.BoneName.ToString(), *GetNameSafe(Skeleton), *GetNameSafe(this));
+			}
+			else
+			{
+				bDefaultToRootBone = false;
+			}
+		}
+
+		if (bDefaultToRootBone)
+		{
+			TempBoneReference.BoneName = Skeleton->GetReferenceSkeleton().GetBoneName(int32(RootBoneIndexType));
+			TempBoneReference.Initialize(Skeleton);
+			check(TempBoneReference.HasValidSetup());
+		}
+
+		SchemaBoneIdx = BoneReferences.AddUnique(TempBoneReference);
+		check(SchemaBoneIdx >= 0 && SchemaBoneIdx < 128);
+	}
+	else if (BoneReferences.IsEmpty())
+	{
+		BoneReferences.Add(FBoneReference());
+	}
+	return int8(SchemaBoneIdx);
 }
 
 void UPoseSearchSchema::Finalize()
@@ -84,7 +107,37 @@ void UPoseSearchSchema::Finalize()
 		++ChannelIndex;
 	}
 
-	ResolveBoneReferences();
+	for (const TObjectPtr<UPoseSearchFeatureChannel>& ChannelPtr : FinalizedChannels)
+	{
+		check(ChannelPtr);
+		if (ChannelPtr->GetPermutationTimeType() != EPermutationTimeType::UseSampleTime)
+		{
+			// there's at least one channel that uses UsePermutationTime or UseSampleToPermutationTime: we automatically add a UPoseSearchFeatureChannel_PermutationTime if not already in the schema
+			UPoseSearchFeatureChannel_PermutationTime::FindOrAddToSchema(this);
+			break;
+		}
+	}
+
+	BoneIndicesWithParents.Reset();
+	if (Skeleton)
+	{
+		// Initialize references to obtain bone indices and fill out bone index array
+		for (FBoneReference& BoneRef : BoneReferences)
+		{
+			check(BoneRef.HasValidSetup());
+			BoneIndicesWithParents.Add(BoneRef.BoneIndex);
+		}
+
+		// Build separate index array with parent indices guaranteed to be present. Sort for EnsureParentsPresent.
+		BoneIndicesWithParents.Sort();
+		FAnimationRuntime::EnsureParentsPresent(BoneIndicesWithParents, Skeleton->GetReferenceSkeleton());
+	}
+
+	// BoneIndicesWithParents should at least contain the root to support mirroring root motion
+	if (BoneIndicesWithParents.IsEmpty())
+	{
+		BoneIndicesWithParents.Add(0);
+	}
 }
 
 void UPoseSearchSchema::PreSave(FObjectPreSaveContext ObjectSaveContext)
@@ -99,13 +152,6 @@ void UPoseSearchSchema::PostLoad()
 	Finalize();
 }
 
-int8 UPoseSearchSchema::AddBoneReference(const FBoneReference& BoneReference)
-{
-	const int32 SchemaBoneIdx = BoneReferences.AddUnique(BoneReference);
-	check(SchemaBoneIdx >= 0 && SchemaBoneIdx < 128);
-	return int8(SchemaBoneIdx);
-}
-
 USkeleton* UPoseSearchSchema::GetSkeleton(bool& bInvalidSkeletonIsError, const IPropertyHandle* PropertyHandle)
 {
 	bInvalidSkeletonIsError = false;
diff --git a/Engine/Plugins/Experimental/Animation/PoseSearch/Source/Runtime/Public/PoseSearch/PoseSearchAssetIndexer.h b/Engine/Plugins/Experimental/Animation/PoseSearch/Source/Runtime/Public/PoseSearch/PoseSearchAssetIndexer.h
--- a/Engine/Plugins/Experimental/Animation/PoseSearch/Source/Runtime/Public/PoseSearch/PoseSearchAssetIndexer.h
+++ b/Engine/Plugins/Experimental/Animation/PoseSearch/Source/Runtime/Public/PoseSearch/PoseSearchAssetIndexer.h
@@ -7,6 +7,7 @@
 #include "BonePose.h"
 #include "CoreMinimal.h"
 #include "PoseSearch/PoseSearchDefines.h"
+#include "PoseSearch/PoseSearchFeatureChannel.h"
 #include "PoseSearch/PoseSearchSchema.h"
 
 struct FPoseSearchIndexAsset;
@@ -47,9 +48,9 @@ public:
 	void Process(int32 AssetIdx);
 	const FStats& GetStats() const { return Stats; }
 
-	FQuat GetSampleRotation(float SampleTimeOffset, int32 SampleIdx, int8 SchemaSampleBoneIdx = RootSchemaBoneIdx, int8 SchemaOriginBoneIdx = RootSchemaBoneIdx);
-	FVector GetSamplePosition(float SampleTimeOffset, int32 SampleIdx, int8 SchemaSampleBoneIdx = RootSchemaBoneIdx, int8 SchemaOriginBoneIdx = RootSchemaBoneIdx);
-	FVector GetSampleVelocity(float SampleTimeOffset, int32 SampleIdx, int8 SchemaSampleBoneIdx = RootSchemaBoneIdx, int8 SchemaOriginBoneIdx = RootSchemaBoneIdx, bool bUseCharacterSpaceVelocities = true);
+	FQuat GetSampleRotation(float SampleTimeOffset, int32 SampleIdx, int8 SchemaSampleBoneIdx = RootSchemaBoneIdx, int8 SchemaOriginBoneIdx = RootSchemaBoneIdx, EPermutationTimeType PermutationTimeType = EPermutationTimeType::UseSampleTime);
+	FVector GetSamplePosition(float SampleTimeOffset, int32 SampleIdx, int8 SchemaSampleBoneIdx = RootSchemaBoneIdx, int8 SchemaOriginBoneIdx = RootSchemaBoneIdx, EPermutationTimeType PermutationTimeType = EPermutationTimeType::UseSampleTime);
+	FVector GetSampleVelocity(float SampleTimeOffset, int32 SampleIdx, int8 SchemaSampleBoneIdx = RootSchemaBoneIdx, int8 SchemaOriginBoneIdx = RootSchemaBoneIdx, bool bUseCharacterSpaceVelocities = true, EPermutationTimeType PermutationTimeType = EPermutationTimeType::UseSampleTime);
 
 	int32 GetBeginSampleIdx() const { return FirstIndexedSample; }
 	int32 GetEndSampleIdx() const { return LastIndexedSample + 1; }
@@ -57,11 +58,8 @@ public:
 	
 	TArrayView<float> GetPoseVector(int32 SampleIdx) const;
 	const UPoseSearchSchema* GetSchema() const;
-
-	void SetPermutationTimeOffsets(float InPermutationSampleTimeOffset, float InPermutationOriginTimeOffset);
-	void ResetPermutationTimeOffsets();
 	float CalculatePermutationTimeOffset() const;
-	
+
 private:
 	int32 GetVectorIdx(int32 SampleIdx) const;
 	FTransform GetTransform(float SampleTime, bool& bClamped, int8 SchemaBoneIdx = RootSchemaBoneIdx);
@@ -99,10 +97,6 @@ private:
 	int32 FirstIndexedSample = 0;
 	int32 LastIndexedSample = 0;
 
-	// time offsets controlled by sampling data permutations
-	float PermutationSampleTimeOffset = 0.f;
-	float PermutationOriginTimeOffset = 0.f;
-
 	TArrayView<float> FeatureVectorTable;
 	TArrayView<FPoseSearchPoseMetadata> PoseMetadata;
 
diff --git a/Engine/Plugins/Experimental/Animation/PoseSearch/Source/Runtime/Public/PoseSearch/PoseSearchContext.h b/Engine/Plugins/Experimental/Animation/PoseSearch/Source/Runtime/Public/PoseSearch/PoseSearchContext.h
--- a/Engine/Plugins/Experimental/Animation/PoseSearch/Source/Runtime/Public/PoseSearch/PoseSearchContext.h
+++ b/Engine/Plugins/Experimental/Animation/PoseSearch/Source/Runtime/Public/PoseSearch/PoseSearchContext.h
@@ -4,12 +4,12 @@
 
 #include "DrawDebugHelpers.h"
 #include "PoseSearch/PoseSearchDefines.h"
+#include "PoseSearch/PoseSearchFeatureChannel.h"
 #include "PoseSearch/PoseSearchIndex.h"
 #include "PoseSearch/PoseSearchResult.h"
 
 struct FTrajectorySampleRange;
 class UPoseSearchDatabase;
-class UPoseSearchSchema;
 
 namespace UE::PoseSearch
 {
@@ -50,8 +50,8 @@ struct POSESEARCH_API FCachedTransform
 {
 	FCachedTransform()
 		: SampleTime(0.f)
-		, BoneIndexType(-1)
-		, Transform()
+		, BoneIndexType(RootBoneIndexType)
+		, Transform(FTransformType::Identity)
 	{
 	}
 
@@ -62,13 +62,12 @@ struct POSESEARCH_API FCachedTransform
 	{
 	}
 
-	float SampleTime;
+	float SampleTime = 0.f;
 	
-	// if -1 it represents the root bone
-	FBoneIndexType BoneIndexType;
+	FBoneIndexType BoneIndexType = RootBoneIndexType;
 
 	// associated transform to BoneIndexType in ComponentSpace (except for the root bone stored in global space)
-	FTransformType Transform;
+	FTransformType Transform = FTransformType::Identity;
 };
 
 template<typename FTransformType>
@@ -149,9 +148,9 @@ struct POSESEARCH_API FSearchContext
 	// can the continuing pose advance? (if not we skip evaluating it)
 	bool bCanAdvance = true;
 
-	FQuat GetSampleRotation(float SampleTimeOffset, const UPoseSearchSchema* Schema, int8 SchemaSampleBoneIdx = RootSchemaBoneIdx, int8 SchemaOriginBoneIdx = RootSchemaBoneIdx, bool bUseHistoryRoot = false);
-	FVector GetSamplePosition(float SampleTimeOffset, const UPoseSearchSchema* Schema, int8 SchemaSampleBoneIdx = RootSchemaBoneIdx, int8 SchemaOriginBoneIdx = RootSchemaBoneIdx, bool bUseHistoryRoot = false);
-	FVector GetSampleVelocity(float SampleTimeOffset, const UPoseSearchSchema* Schema, int8 SchemaSampleBoneIdx = RootSchemaBoneIdx, int8 SchemaOriginBoneIdx = RootSchemaBoneIdx, bool bUseCharacterSpaceVelocities = true, bool bUseHistoryRoot = false);
+	FQuat GetSampleRotation(float SampleTimeOffset, const UPoseSearchSchema* Schema, int8 SchemaSampleBoneIdx = RootSchemaBoneIdx, int8 SchemaOriginBoneIdx = RootSchemaBoneIdx, bool bUseHistoryRoot = false, EPermutationTimeType PermutationTimeType = EPermutationTimeType::UseSampleTime);
+	FVector GetSamplePosition(float SampleTimeOffset, const UPoseSearchSchema* Schema, int8 SchemaSampleBoneIdx = RootSchemaBoneIdx, int8 SchemaOriginBoneIdx = RootSchemaBoneIdx, bool bUseHistoryRoot = false, EPermutationTimeType PermutationTimeType = EPermutationTimeType::UseSampleTime);
+	FVector GetSampleVelocity(float SampleTimeOffset, const UPoseSearchSchema* Schema, int8 SchemaSampleBoneIdx = RootSchemaBoneIdx, int8 SchemaOriginBoneIdx = RootSchemaBoneIdx, bool bUseCharacterSpaceVelocities = true, bool bUseHistoryRoot = false, EPermutationTimeType PermutationTimeType = EPermutationTimeType::UseSampleTime);
 
 	void ClearCachedEntries();
 
@@ -168,9 +167,6 @@ struct POSESEARCH_API FSearchContext
 	TConstArrayView<float> GetCurrentResultPoseVector() const;
 	TConstArrayView<float> GetCurrentResultNextPoseVector() const;
 
-	void SetPermutationTimeOffsets(float InPermutationSampleTimeOffset, float InPermutationOriginTimeOffset);
-	void ResetPermutationTimeOffsets();
-	
 	float DesiredPermutationTimeOffset = 0.f;
 
 	const FPoseIndicesHistory* PoseIndicesHistory = nullptr;
@@ -186,10 +182,6 @@ private:
 
 	float CurrentBestTotalCost = MAX_flt;
 	
-	// time offsets controlled by sampling data permutations
-	float PermutationSampleTimeOffset = 0.f;
-	float PermutationOriginTimeOffset = 0.f;
-
 #if UE_POSE_SEARCH_TRACE_ENABLED
 
 public:
diff --git a/Engine/Plugins/Experimental/Animation/PoseSearch/Source/Runtime/Public/PoseSearch/PoseSearchDefines.h b/Engine/Plugins/Experimental/Animation/PoseSearch/Source/Runtime/Public/PoseSearch/PoseSearchDefines.h
--- a/Engine/Plugins/Experimental/Animation/PoseSearch/Source/Runtime/Public/PoseSearch/PoseSearchDefines.h
+++ b/Engine/Plugins/Experimental/Animation/PoseSearch/Source/Runtime/Public/PoseSearch/PoseSearchDefines.h
@@ -22,6 +22,6 @@ constexpr EParallelForFlags ParallelForFlags = EParallelForFlags::None;
 DECLARE_LOG_CATEGORY_EXTERN(LogPoseSearch, Log, All);
 
 static constexpr int8 RootSchemaBoneIdx = -1;
-static constexpr FBoneIndexType RootBoneIndexType = 0xFFFF;
+static constexpr FBoneIndexType RootBoneIndexType = 0;
 
 
diff --git a/Engine/Plugins/Experimental/Animation/PoseSearch/Source/Runtime/Public/PoseSearch/PoseSearchFeatureChannel.h b/Engine/Plugins/Experimental/Animation/PoseSearch/Source/Runtime/Public/PoseSearch/PoseSearchFeatureChannel.h
--- a/Engine/Plugins/Experimental/Animation/PoseSearch/Source/Runtime/Public/PoseSearch/PoseSearchFeatureChannel.h
+++ b/Engine/Plugins/Experimental/Animation/PoseSearch/Source/Runtime/Public/PoseSearch/PoseSearchFeatureChannel.h
@@ -37,6 +37,25 @@ enum class EInputQueryPose : uint8
 	UseInterpolatedContinuingPose,
 };
 
+// this enumeration controls the channel sampling time:
+// for example if a channel specifies a bone and an origin bone (used to generate the reference system of the features associated to the bone),
+// bone and origin bone will be evaluated at potentially different times:
+UENUM()
+enum class EPermutationTimeType : uint8
+{
+	// bone and origin bone are sampled at the same sample time (plus eventual SampleTimeOffset for the bone):
+	// it's defined as the current animation evaluation time
+	UseSampleTime,
+
+	// bone and origin bone are sampled at the same permutation time (plus eventual SampleTimeOffset for the bone):
+	// it's defined as SamplingTime (as UseSampleTime) + Schema->PermutationsTimeOffset + PermutationIndex / Schema->PermutationsSampleRate
+	// where PermutationIndex is in range [0, Schema->NumberOfPermutations)
+	UsePermutationTime,
+
+	// bone is evaluated at sample time (and plus eventual SampleTimeOffset) and origin bone is evaluated at permutation time
+	UseSampleToPermutationTime,
+};
+
 namespace UE::PoseSearch
 {
 
@@ -103,6 +122,9 @@ public:
 	// @todo: should this API be under ENABLE_DRAW_DEBUG?
 	virtual void AddDependentChannels(UPoseSearchSchema* Schema) const {}
 
+	virtual EPermutationTimeType GetPermutationTimeType() const { return EPermutationTimeType::UseSampleTime; }
+	static void GetPermutationTimeOffsets(EPermutationTimeType PermutationTimeType, float DesiredPermutationTimeOffset, float& OutPermutationSampleTimeOffset, float& OutPermutationOriginTimeOffset);
+
 #if ENABLE_DRAW_DEBUG
 	// API called before DebugDraw to collect shared channel informations such as decoded positions form the PoseVector
 	virtual void PreDebugDraw(UE::PoseSearch::FDebugDrawParams& DrawParams, TConstArrayView<float> PoseVector) const {}
diff --git a/Engine/Plugins/Experimental/Animation/PoseSearch/Source/Runtime/Public/PoseSearch/PoseSearchSchema.h b/Engine/Plugins/Experimental/Animation/PoseSearch/Source/Runtime/Public/PoseSearch/PoseSearchSchema.h
--- a/Engine/Plugins/Experimental/Animation/PoseSearch/Source/Runtime/Public/PoseSearch/PoseSearchSchema.h
+++ b/Engine/Plugins/Experimental/Animation/PoseSearch/Source/Runtime/Public/PoseSearch/PoseSearchSchema.h
@@ -173,5 +173,4 @@ private:
 	}
 
 	void Finalize();
-	void ResolveBoneReferences();
 };
