motion matching
- removed UPoseSearchFeatureChannel::PreDebugDraw
- deprecated channels color presets in favor of holding on to a FLinearColor DebugColor
- removed FDebugDrawParams::CachedPositions, and implemented FDebugDrawParams::ExtractPosition to calculate required channels debug drawing positions
- cleaned up UPoseSearchSchema::AddBoneReference and UPoseSearchSchema::Finalize initializing only on valid SchemaBoneIdx

diff --git a/Engine/Plugins/Experimental/Animation/PoseSearch/Source/Runtime/Private/PoseSearchContext.cpp b/Engine/Plugins/Experimental/Animation/PoseSearch/Source/Runtime/Private/PoseSearchContext.cpp
--- a/Engine/Plugins/Experimental/Animation/PoseSearch/Source/Runtime/Private/PoseSearchContext.cpp
+++ b/Engine/Plugins/Experimental/Animation/PoseSearch/Source/Runtime/Private/PoseSearchContext.cpp
@@ -5,9 +5,9 @@
 #include "Animation/MotionTrajectoryTypes.h"
 #include "AnimationRuntime.h"
 #include "PoseSearch/PoseSearchDatabase.h"
-#include "PoseSearch/PoseSearchFeatureChannel.h"
 #include "PoseSearch/PoseSearchHistory.h"
 #include "PoseSearch/PoseSearchSchema.h"
+#include "PoseSearchFeatureChannel_Position.h"
 
 namespace UE::PoseSearch
 {
@@ -21,7 +21,6 @@ FDebugDrawParams::FDebugDrawParams(FAnimInstanceProxy* InAnimInstanceProxy, cons
 , Mesh(nullptr)
 , Database(InDatabase)
 , Flags(InFlags)
-, CachedPositions()
 {
 }
 
@@ -31,7 +30,6 @@ FDebugDrawParams::FDebugDrawParams(const UWorld* InWorld, const USkinnedMeshComp
 , Mesh(InMesh)
 , Database(InDatabase)
 , Flags(InFlags)
-, CachedPositions()
 {
 }
 
@@ -40,41 +38,6 @@ bool FDebugDrawParams::CanDraw() const
 	return (AnimInstanceProxy || World) && Database && Database->Schema && Database->Schema->IsValid();
 }
 
-FColor FDebugDrawParams::GetColor(int32 ColorPreset) const
-{
-	FLinearColor Color = FLinearColor::Red;
-
-	const UPoseSearchSchema* Schema = GetSchema();
-	if (!Schema || !Schema->IsValid())
-	{
-		Color = FLinearColor::Red;
-	}
-	else if (ColorPreset < 0 || ColorPreset >= Schema->ColorPresets.Num())
-	{
-		if (EnumHasAnyFlags(Flags, EDebugDrawFlags::DrawQuery))
-		{
-			Color = FLinearColor::Blue;
-		}
-		else
-		{
-			Color = FLinearColor::Green;
-		}
-	}
-	else
-	{
-		if (EnumHasAnyFlags(Flags, EDebugDrawFlags::DrawQuery))
-		{
-			Color = Schema->ColorPresets[ColorPreset].Query;
-		}
-		else
-		{
-			Color = Schema->ColorPresets[ColorPreset].Result;
-		}
-	}
-
-	return Color.ToFColor(true);
-}
-
 const FPoseSearchIndex* FDebugDrawParams::GetSearchIndex() const
 {
 	return Database ? &Database->GetSearchIndex() : nullptr;
@@ -85,26 +48,39 @@ const UPoseSearchSchema* FDebugDrawParams::GetSchema() const
 	return Database ? Database->Schema : nullptr;
 }
 
-void FDebugDrawParams::ClearCachedPositions()
+FVector FDebugDrawParams::ExtractPosition(TConstArrayView<float> PoseVector, const UPoseSearchFeatureChannel_Position* Position) const
 {
-	CachedPositions.Reset();
-}
-
-void FDebugDrawParams::AddCachedPosition(float TimeOffset, int8 SchemaBoneIdx, const FVector& Position)
-{
-	if (const UPoseSearchSchema* Schema = GetSchema())
-	{
-		CachedPositions.Add(TimeOffset, Schema->GetBoneIndexType(SchemaBoneIdx), Position);
-	}
+	check(Position);
+	const FVector BonePosition = FFeatureVectorHelper::DecodeVector(PoseVector, Position->GetChannelDataOffset(), Position->ComponentStripping);
+	const FVector WorldBonePosition = GetRootTransform().TransformPosition(BonePosition);
+	return WorldBonePosition;
 }
 
-FVector FDebugDrawParams::GetCachedPosition(float TimeOffset, int8 SchemaBoneIdx) const
+FVector FDebugDrawParams::ExtractPosition(TConstArrayView<float> PoseVector, float SampleTimeOffset, int8 SchemaBoneIdx, EPermutationTimeType PermutationTimeType) const
 {
+	// we don't wanna ask for a SchemaOriginBoneIdx in the future or past
+	check(PermutationTimeType != EPermutationTimeType::UsePermutationTime);
 	if (const UPoseSearchSchema* Schema = GetSchema())
 	{
-		if (auto CachedPosition = CachedPositions.Find(TimeOffset, Schema->GetBoneIndexType(SchemaBoneIdx)))
+		// looking for a UPoseSearchFeatureChannel_Position that matches the TimeOffset and SchemaBoneIdx,
+		// with SchemaOriginBoneIdx to be the root bone and the appropriate PermutationTimeType 
+		if (const UPoseSearchFeatureChannel_Position* FoundPosition = static_cast<const UPoseSearchFeatureChannel_Position*>(
+				Schema->FindChannel([SampleTimeOffset, SchemaBoneIdx, PermutationTimeType, Schema](const UPoseSearchFeatureChannel* Channel) -> const UPoseSearchFeatureChannel_Position*
+			{
+				if (const UPoseSearchFeatureChannel_Position* Position = Cast<UPoseSearchFeatureChannel_Position>(Channel))
+				{
+					if (Position->SchemaBoneIdx == SchemaBoneIdx &&
+						Position->SampleTimeOffset == SampleTimeOffset &&
+						Position->PermutationTimeType == PermutationTimeType &&
+						Schema->IsRootBone(Position->SchemaOriginBoneIdx))
+					{
+						return Position;
+					}
+				}
+		return nullptr;
+			})))
 		{
-			return CachedPosition->Transform;
+			return ExtractPosition(PoseVector, FoundPosition);
 		}
 
 		if (Mesh && SchemaBoneIdx >= 0)
@@ -241,8 +217,6 @@ void FDebugDrawParams::DrawCentripetalCatmullRomSpline(TConstArrayView<FVector>
 
 void FDebugDrawParams::DrawFeatureVector(TConstArrayView<float> PoseVector)
 {
-	ClearCachedPositions();
-
 	if (CanDraw())
 	{
 		const UPoseSearchSchema* Schema = GetSchema();
@@ -250,11 +224,6 @@ void FDebugDrawParams::DrawFeatureVector(TConstArrayView<float> PoseVector)
 
 		if (PoseVector.Num() == Schema->SchemaCardinality)
 		{
-			for (const TObjectPtr<UPoseSearchFeatureChannel>& ChannelPtr : Schema->GetChannels())
-			{
-				ChannelPtr->PreDebugDraw(*this, PoseVector);
-			}
-
 			for (const TObjectPtr<UPoseSearchFeatureChannel>& ChannelPtr : Schema->GetChannels())
 			{
 				ChannelPtr->DebugDraw(*this, PoseVector);
diff --git a/Engine/Plugins/Experimental/Animation/PoseSearch/Source/Runtime/Private/PoseSearchFeatureChannel_FilterCrashingLegs.cpp b/Engine/Plugins/Experimental/Animation/PoseSearch/Source/Runtime/Private/PoseSearchFeatureChannel_FilterCrashingLegs.cpp
--- a/Engine/Plugins/Experimental/Animation/PoseSearch/Source/Runtime/Private/PoseSearchFeatureChannel_FilterCrashingLegs.cpp
+++ b/Engine/Plugins/Experimental/Animation/PoseSearch/Source/Runtime/Private/PoseSearchFeatureChannel_FilterCrashingLegs.cpp
@@ -43,10 +43,10 @@ void UPoseSearchFeatureChannel_FilterCrashingLegs::AddDependentChannels(UPoseSea
 {
 	if (Schema->bInjectAdditionalDebugChannels)
 	{
-		UPoseSearchFeatureChannel_Position::FindOrAddToSchema(Schema, 0.f, 0, LeftThigh.BoneName);
-		UPoseSearchFeatureChannel_Position::FindOrAddToSchema(Schema, 0.f, 0, RightThigh.BoneName);
-		UPoseSearchFeatureChannel_Position::FindOrAddToSchema(Schema, 0.f, 0, LeftFoot.BoneName);
-		UPoseSearchFeatureChannel_Position::FindOrAddToSchema(Schema, 0.f, 0, RightFoot.BoneName);
+		UPoseSearchFeatureChannel_Position::FindOrAddToSchema(Schema, 0.f, LeftThigh.BoneName);
+		UPoseSearchFeatureChannel_Position::FindOrAddToSchema(Schema, 0.f, RightThigh.BoneName);
+		UPoseSearchFeatureChannel_Position::FindOrAddToSchema(Schema, 0.f, LeftFoot.BoneName);
+		UPoseSearchFeatureChannel_Position::FindOrAddToSchema(Schema, 0.f, RightFoot.BoneName);
 	}
 }
 void UPoseSearchFeatureChannel_FilterCrashingLegs::BuildQuery(UE::PoseSearch::FSearchContext& SearchContext, FPoseSearchFeatureVectorBuilder& InOutQuery) const
@@ -85,10 +85,10 @@ void UPoseSearchFeatureChannel_FilterCrashingLegs::DebugDraw(const UE::PoseSearc
 
 	const float CrashingLegsValue = FFeatureVectorHelper::DecodeFloat(PoseVector, ChannelDataOffset);
 
-	const FVector LeftThighPosition = DrawParams.GetCachedPosition(0.f, LeftThighIdx);
-	const FVector RightThighPosition = DrawParams.GetCachedPosition(0.f, RightThighIdx);
-	const FVector LeftFootPosition = DrawParams.GetCachedPosition(0.f, LeftFootIdx);
-	const FVector RightFootPosition = DrawParams.GetCachedPosition(0.f, RightFootIdx);
+	const FVector LeftThighPosition = DrawParams.ExtractPosition(PoseVector, 0.f, LeftThighIdx);
+	const FVector RightThighPosition = DrawParams.ExtractPosition(PoseVector, 0.f, RightThighIdx);
+	const FVector LeftFootPosition = DrawParams.ExtractPosition(PoseVector, 0.f, LeftFootIdx);
+	const FVector RightFootPosition = DrawParams.ExtractPosition(PoseVector, 0.f, RightFootIdx);
 
 	const float FeetDistance = (RightFootPosition - LeftFootPosition).Length();
 
diff --git a/Engine/Plugins/Experimental/Animation/PoseSearch/Source/Runtime/Private/PoseSearchFeatureChannel_Group.cpp b/Engine/Plugins/Experimental/Animation/PoseSearch/Source/Runtime/Private/PoseSearchFeatureChannel_Group.cpp
--- a/Engine/Plugins/Experimental/Animation/PoseSearch/Source/Runtime/Private/PoseSearchFeatureChannel_Group.cpp
+++ b/Engine/Plugins/Experimental/Animation/PoseSearch/Source/Runtime/Private/PoseSearchFeatureChannel_Group.cpp
@@ -63,17 +63,6 @@ void UPoseSearchFeatureChannel_GroupBase::BuildQuery(UE::PoseSearch::FSearchCont
 }
 
 #if ENABLE_DRAW_DEBUG
-void UPoseSearchFeatureChannel_GroupBase::PreDebugDraw(UE::PoseSearch::FDebugDrawParams& DrawParams, TConstArrayView<float> PoseVector) const
-{
-	for (const TObjectPtr<UPoseSearchFeatureChannel>& SubChannelPtr : GetSubChannels())
-	{
-		if (const UPoseSearchFeatureChannel* SubChannel = SubChannelPtr.Get())
-		{
-			SubChannel->PreDebugDraw(DrawParams, PoseVector);
-		}
-	}
-}
-
 void UPoseSearchFeatureChannel_GroupBase::DebugDraw(const UE::PoseSearch::FDebugDrawParams& DrawParams, TConstArrayView<float> PoseVector) const
 {
 	for (const TObjectPtr<UPoseSearchFeatureChannel>& SubChannelPtr : GetSubChannels())
diff --git a/Engine/Plugins/Experimental/Animation/PoseSearch/Source/Runtime/Private/PoseSearchFeatureChannel_Group.h b/Engine/Plugins/Experimental/Animation/PoseSearch/Source/Runtime/Private/PoseSearchFeatureChannel_Group.h
--- a/Engine/Plugins/Experimental/Animation/PoseSearch/Source/Runtime/Private/PoseSearchFeatureChannel_Group.h
+++ b/Engine/Plugins/Experimental/Animation/PoseSearch/Source/Runtime/Private/PoseSearchFeatureChannel_Group.h
@@ -19,7 +19,6 @@ public:
 	virtual void AddDependentChannels(UPoseSearchSchema* Schema) const override; 
 
 #if ENABLE_DRAW_DEBUG
-	virtual void PreDebugDraw(UE::PoseSearch::FDebugDrawParams& DrawParams, TConstArrayView<float> PoseVector) const override;
 	virtual void DebugDraw(const UE::PoseSearch::FDebugDrawParams& DrawParams, TConstArrayView<float> PoseVector) const override;
 #endif // ENABLE_DRAW_DEBUG
 
diff --git a/Engine/Plugins/Experimental/Animation/PoseSearch/Source/Runtime/Private/PoseSearchFeatureChannel_Heading.cpp b/Engine/Plugins/Experimental/Animation/PoseSearch/Source/Runtime/Private/PoseSearchFeatureChannel_Heading.cpp
--- a/Engine/Plugins/Experimental/Animation/PoseSearch/Source/Runtime/Private/PoseSearchFeatureChannel_Heading.cpp
+++ b/Engine/Plugins/Experimental/Animation/PoseSearch/Source/Runtime/Private/PoseSearchFeatureChannel_Heading.cpp
@@ -21,7 +21,7 @@ void UPoseSearchFeatureChannel_Heading::AddDependentChannels(UPoseSearchSchema*
 {
 	if (Schema->bInjectAdditionalDebugChannels)
 	{
-		UPoseSearchFeatureChannel_Position::FindOrAddToSchema(Schema, SampleTimeOffset, ColorPresetIndex, Bone.BoneName);
+		UPoseSearchFeatureChannel_Position::FindOrAddToSchema(Schema, SampleTimeOffset, Bone.BoneName);
 	}
 }
 
@@ -71,9 +71,9 @@ void UPoseSearchFeatureChannel_Heading::DebugDraw(const UE::PoseSearch::FDebugDr
 {
 	using namespace UE::PoseSearch;
 
-	const FColor Color = DrawParams.GetColor(ColorPresetIndex);
+	const FColor Color = DebugColor.ToFColor(true);
 	const FVector BoneHeading = DrawParams.GetRootTransform().GetRotation().RotateVector(FFeatureVectorHelper::DecodeVector(PoseVector, ChannelDataOffset, ComponentStripping));
-	const FVector BonePos = DrawParams.GetCachedPosition(SampleTimeOffset, SchemaBoneIdx);
+	const FVector BonePos = DrawParams.ExtractPosition(PoseVector, SampleTimeOffset, SchemaBoneIdx);
 
 	DrawParams.DrawPoint(BonePos, Color, 3.f);
 	DrawParams.DrawLine(BonePos + BoneHeading * 4.f, BonePos + BoneHeading * 15.f, Color);
diff --git a/Engine/Plugins/Experimental/Animation/PoseSearch/Source/Runtime/Private/PoseSearchFeatureChannel_Heading.h b/Engine/Plugins/Experimental/Animation/PoseSearch/Source/Runtime/Private/PoseSearchFeatureChannel_Heading.h
--- a/Engine/Plugins/Experimental/Animation/PoseSearch/Source/Runtime/Private/PoseSearchFeatureChannel_Heading.h
+++ b/Engine/Plugins/Experimental/Animation/PoseSearch/Source/Runtime/Private/PoseSearchFeatureChannel_Heading.h
@@ -39,7 +39,7 @@ public:
 	int8 SchemaBoneIdx = 0;
 
 	UPROPERTY(EditAnywhere, Category = "Settings", meta = (ExcludeFromHash))
-	int32 ColorPresetIndex = 0;
+	FLinearColor DebugColor = FLinearColor::White;
 
 	UPROPERTY(EditAnywhere, Category = "Settings")
 	EInputQueryPose InputQueryPose = EInputQueryPose::UseContinuingPose;
diff --git a/Engine/Plugins/Experimental/Animation/PoseSearch/Source/Runtime/Private/PoseSearchFeatureChannel_Phase.cpp b/Engine/Plugins/Experimental/Animation/PoseSearch/Source/Runtime/Private/PoseSearchFeatureChannel_Phase.cpp
--- a/Engine/Plugins/Experimental/Animation/PoseSearch/Source/Runtime/Private/PoseSearchFeatureChannel_Phase.cpp
+++ b/Engine/Plugins/Experimental/Animation/PoseSearch/Source/Runtime/Private/PoseSearchFeatureChannel_Phase.cpp
@@ -292,7 +292,7 @@ void UPoseSearchFeatureChannel_Phase::AddDependentChannels(UPoseSearchSchema* Sc
 {
 	if (Schema->bInjectAdditionalDebugChannels)
 	{
-		UPoseSearchFeatureChannel_Position::FindOrAddToSchema(Schema, 0.f, ColorPresetIndex, Bone.BoneName);
+		UPoseSearchFeatureChannel_Position::FindOrAddToSchema(Schema, 0.f, Bone.BoneName);
 	}
 }
 
@@ -326,10 +326,10 @@ void UPoseSearchFeatureChannel_Phase::DebugDraw(const UE::PoseSearch::FDebugDraw
 	static float ScaleFactor = 1.f;
 
 
-	const FColor Color = DrawParams.GetColor(ColorPresetIndex);
+	const FColor Color = DebugColor.ToFColor(true);
 
 	const FVector2D Phase = FFeatureVectorHelper::DecodeVector2D(PoseVector, ChannelDataOffset);
-	const FVector BonePos = DrawParams.GetCachedPosition(0.f, SchemaBoneIdx);
+	const FVector BonePos = DrawParams.ExtractPosition(PoseVector, 0.f, SchemaBoneIdx);
 
 	const FVector TransformXAxisVector = DrawParams.GetRootTransform().TransformVector(FVector::XAxisVector);
 	const FVector TransformYAxisVector = DrawParams.GetRootTransform().TransformVector(FVector::YAxisVector);
diff --git a/Engine/Plugins/Experimental/Animation/PoseSearch/Source/Runtime/Private/PoseSearchFeatureChannel_Phase.h b/Engine/Plugins/Experimental/Animation/PoseSearch/Source/Runtime/Private/PoseSearchFeatureChannel_Phase.h
--- a/Engine/Plugins/Experimental/Animation/PoseSearch/Source/Runtime/Private/PoseSearchFeatureChannel_Phase.h
+++ b/Engine/Plugins/Experimental/Animation/PoseSearch/Source/Runtime/Private/PoseSearchFeatureChannel_Phase.h
@@ -22,7 +22,7 @@ public:
 	int8 SchemaBoneIdx = 0;
 
 	UPROPERTY(EditAnywhere, Category = "Settings", meta = (ExcludeFromHash))
-	int32 ColorPresetIndex = 0;
+	FLinearColor DebugColor = FLinearColor::Yellow;
 
 	UPROPERTY(EditAnywhere, Category = "Settings")
 	EInputQueryPose InputQueryPose = EInputQueryPose::UseContinuingPose;
diff --git a/Engine/Plugins/Experimental/Animation/PoseSearch/Source/Runtime/Private/PoseSearchFeatureChannel_Pose.cpp b/Engine/Plugins/Experimental/Animation/PoseSearch/Source/Runtime/Private/PoseSearchFeatureChannel_Pose.cpp
--- a/Engine/Plugins/Experimental/Animation/PoseSearch/Source/Runtime/Private/PoseSearchFeatureChannel_Pose.cpp
+++ b/Engine/Plugins/Experimental/Animation/PoseSearch/Source/Runtime/Private/PoseSearchFeatureChannel_Pose.cpp
@@ -24,7 +24,7 @@ void UPoseSearchFeatureChannel_Pose::Finalize(UPoseSearchSchema* Schema)
 			Position->Bone = SampledBone.Reference;
 			Position->Weight = SampledBone.Weight * Weight;
 			Position->SampleTimeOffset = 0.f;
-			Position->ColorPresetIndex = SampledBone.ColorPresetIndex;
+			Position->DebugColor = SampledBone.DebugColor;
 			Position->InputQueryPose = InputQueryPose;
 			SubChannels.Add(Position);
 		}
@@ -36,7 +36,7 @@ void UPoseSearchFeatureChannel_Pose::Finalize(UPoseSearchSchema* Schema)
 			HeadingX->Weight = SampledBone.Weight * Weight;
 			HeadingX->SampleTimeOffset = 0.f;
 			HeadingX->HeadingAxis = EHeadingAxis::X;
-			HeadingX->ColorPresetIndex = SampledBone.ColorPresetIndex;
+			HeadingX->DebugColor = SampledBone.DebugColor;
 			HeadingX->InputQueryPose = InputQueryPose;
 			SubChannels.Add(HeadingX);
 
@@ -45,7 +45,7 @@ void UPoseSearchFeatureChannel_Pose::Finalize(UPoseSearchSchema* Schema)
 			HeadingY->Weight = SampledBone.Weight * Weight;
 			HeadingY->SampleTimeOffset = 0.f;
 			HeadingY->HeadingAxis = EHeadingAxis::Y;
-			HeadingY->ColorPresetIndex = SampledBone.ColorPresetIndex;
+			HeadingY->DebugColor = SampledBone.DebugColor;
 			HeadingY->InputQueryPose = InputQueryPose;
 			SubChannels.Add(HeadingY);
 		}
@@ -56,7 +56,7 @@ void UPoseSearchFeatureChannel_Pose::Finalize(UPoseSearchSchema* Schema)
 			Velocity->Bone = SampledBone.Reference;
 			Velocity->Weight = SampledBone.Weight * Weight;
 			Velocity->SampleTimeOffset = 0.f;
-			Velocity->ColorPresetIndex = SampledBone.ColorPresetIndex;
+			Velocity->DebugColor = SampledBone.DebugColor;
 			Velocity->InputQueryPose = InputQueryPose;
 			Velocity->bUseCharacterSpaceVelocities = bUseCharacterSpaceVelocities;
 			SubChannels.Add(Velocity);
@@ -67,7 +67,7 @@ void UPoseSearchFeatureChannel_Pose::Finalize(UPoseSearchSchema* Schema)
 			UPoseSearchFeatureChannel_Phase* Phase = NewObject<UPoseSearchFeatureChannel_Phase>(this, NAME_None, RF_Transient);
 			Phase->Bone = SampledBone.Reference;
 			Phase->Weight = SampledBone.Weight * Weight;
-			Phase->ColorPresetIndex = SampledBone.ColorPresetIndex;
+			Phase->DebugColor = SampledBone.DebugColor;
 			Phase->InputQueryPose = InputQueryPose;
 			SubChannels.Add(Phase);
 		}
diff --git a/Engine/Plugins/Experimental/Animation/PoseSearch/Source/Runtime/Private/PoseSearchFeatureChannel_Pose.h b/Engine/Plugins/Experimental/Animation/PoseSearch/Source/Runtime/Private/PoseSearchFeatureChannel_Pose.h
--- a/Engine/Plugins/Experimental/Animation/PoseSearch/Source/Runtime/Private/PoseSearchFeatureChannel_Pose.h
+++ b/Engine/Plugins/Experimental/Animation/PoseSearch/Source/Runtime/Private/PoseSearchFeatureChannel_Pose.h
@@ -32,7 +32,7 @@ struct POSESEARCH_API FPoseSearchBone
 	float Weight = 1.f;
 
 	UPROPERTY(EditAnywhere, Category = Config, meta=(ExcludeFromHash))
-	int32 ColorPresetIndex = 0;
+	FLinearColor DebugColor = FLinearColor::Green;
 };
 
 // UPoseSearchFeatureChannel_Pose
diff --git a/Engine/Plugins/Experimental/Animation/PoseSearch/Source/Runtime/Private/PoseSearchFeatureChannel_Position.cpp b/Engine/Plugins/Experimental/Animation/PoseSearch/Source/Runtime/Private/PoseSearchFeatureChannel_Position.cpp
--- a/Engine/Plugins/Experimental/Animation/PoseSearch/Source/Runtime/Private/PoseSearchFeatureChannel_Position.cpp
+++ b/Engine/Plugins/Experimental/Animation/PoseSearch/Source/Runtime/Private/PoseSearchFeatureChannel_Position.cpp
@@ -8,7 +8,7 @@
 #include "PoseSearch/PoseSearchDatabase.h"
 #include "PoseSearch/PoseSearchSchema.h"
 
-void UPoseSearchFeatureChannel_Position::FindOrAddToSchema(UPoseSearchSchema* Schema, float SampleTimeOffset, int32 ColorPresetIndex, const FName& BoneName, EPermutationTimeType PermutationTimeType)
+void UPoseSearchFeatureChannel_Position::FindOrAddToSchema(UPoseSearchSchema* Schema, float SampleTimeOffset, const FName& BoneName, EPermutationTimeType PermutationTimeType)
 {
 	if (!Schema->FindChannel([SampleTimeOffset, &BoneName, PermutationTimeType](const UPoseSearchFeatureChannel* Channel) -> const UPoseSearchFeatureChannel_Position*
 		{
@@ -26,7 +26,8 @@ void UPoseSearchFeatureChannel_Position::FindOrAddToSchema(UPoseSearchSchema* Sc
 		Position->Bone.BoneName = BoneName;
 		Position->Weight = 0.f;
 		Position->SampleTimeOffset = SampleTimeOffset;
-		Position->ColorPresetIndex = ColorPresetIndex;
+		// @todo: perhaps add a tunable color for injected channels
+		Position->DebugColor = FLinearColor::Gray;
 		Position->PermutationTimeType = PermutationTimeType;
 		Position->Finalize(Schema);
 		Schema->FinalizedChannels.Add(Position);
@@ -50,7 +51,7 @@ void UPoseSearchFeatureChannel_Position::AddDependentChannels(UPoseSearchSchema*
 		if (!Schema->IsRootBone(SchemaOriginBoneIdx))
 		{
 			const EPermutationTimeType DependentChannelsPermutationTimeType = PermutationTimeType == EPermutationTimeType::UsePermutationTime ? EPermutationTimeType::UseSampleToPermutationTime : EPermutationTimeType::UseSampleTime;
-			UPoseSearchFeatureChannel_Position::FindOrAddToSchema(Schema, 0.f, ColorPresetIndex, OriginBone.BoneName, DependentChannelsPermutationTimeType);
+			UPoseSearchFeatureChannel_Position::FindOrAddToSchema(Schema, 0.f, OriginBone.BoneName, DependentChannelsPermutationTimeType);
 		}
 	}
 }
@@ -81,23 +82,11 @@ void UPoseSearchFeatureChannel_Position::BuildQuery(UE::PoseSearch::FSearchConte
 }
 
 #if ENABLE_DRAW_DEBUG
-void UPoseSearchFeatureChannel_Position::PreDebugDraw(UE::PoseSearch::FDebugDrawParams& DrawParams, TConstArrayView<float> PoseVector) const
-{
-	using namespace UE::PoseSearch;
-	
-	// if the SchemaOriginBoneIdx is not the root bone, the feature vector position doesn't represent a component space delta position for the SchemaBoneIdx, so we shouldn't collect it
-	if (DrawParams.GetSchema()->IsRootBone(SchemaOriginBoneIdx) && PermutationTimeType != EPermutationTimeType::UsePermutationTime)
-	{
-		const FVector BonePos = DrawParams.GetRootTransform().TransformPosition(FFeatureVectorHelper::DecodeVector(PoseVector, ChannelDataOffset, ComponentStripping));
-		DrawParams.AddCachedPosition(SampleTimeOffset, SchemaBoneIdx, BonePos);
-	}
-}
-
 void UPoseSearchFeatureChannel_Position::DebugDraw(const UE::PoseSearch::FDebugDrawParams& DrawParams, TConstArrayView<float> PoseVector) const
 {
 	using namespace UE::PoseSearch;
 
-	const FColor Color = DrawParams.GetColor(ColorPresetIndex);
+	const FColor Color = DebugColor.ToFColor(true);
 
 	const FVector FeaturesVector = FFeatureVectorHelper::DecodeVector(PoseVector, ChannelDataOffset, ComponentStripping);
 	if (DrawParams.GetSchema()->IsRootBone(SchemaOriginBoneIdx))
@@ -107,7 +96,8 @@ void UPoseSearchFeatureChannel_Position::DebugDraw(const UE::PoseSearch::FDebugD
 	}
 	else
 	{
-		const FVector OriginBonePos = DrawParams.GetCachedPosition(SampleTimeOffset, SchemaOriginBoneIdx);
+		const EPermutationTimeType TimeType = PermutationTimeType == EPermutationTimeType::UsePermutationTime ? EPermutationTimeType::UseSampleToPermutationTime : EPermutationTimeType::UseSampleTime;
+		const FVector OriginBonePos = DrawParams.ExtractPosition(PoseVector, SampleTimeOffset, SchemaOriginBoneIdx, TimeType);
 		const FVector DeltaPos = DrawParams.GetRootTransform().TransformVector(FeaturesVector);
 		const FVector BonePos = OriginBonePos + DeltaPos;
 		DrawParams.DrawLine(OriginBonePos, BonePos, Color);
diff --git a/Engine/Plugins/Experimental/Animation/PoseSearch/Source/Runtime/Private/PoseSearchFeatureChannel_Position.h b/Engine/Plugins/Experimental/Animation/PoseSearch/Source/Runtime/Private/PoseSearchFeatureChannel_Position.h
--- a/Engine/Plugins/Experimental/Animation/PoseSearch/Source/Runtime/Private/PoseSearchFeatureChannel_Position.h
+++ b/Engine/Plugins/Experimental/Animation/PoseSearch/Source/Runtime/Private/PoseSearchFeatureChannel_Position.h
@@ -31,7 +31,7 @@ public:
 	int8 SchemaOriginBoneIdx = 0;
 
 	UPROPERTY(EditAnywhere, Category = "Settings", meta = (ExcludeFromHash))
-	int32 ColorPresetIndex = 0;
+	FLinearColor DebugColor = FLinearColor::Blue;
 
 	UPROPERTY(EditAnywhere, Category = "Settings")
 	EInputQueryPose InputQueryPose = EInputQueryPose::UseContinuingPose;
@@ -50,7 +50,6 @@ public:
 	virtual void AddDependentChannels(UPoseSearchSchema* Schema) const override;
 
 #if ENABLE_DRAW_DEBUG
-	virtual void PreDebugDraw(UE::PoseSearch::FDebugDrawParams& DrawParams, TConstArrayView<float> PoseVector) const override;
 	virtual void DebugDraw(const UE::PoseSearch::FDebugDrawParams& DrawParams, TConstArrayView<float> PoseVector) const override;
 #endif // ENABLE_DRAW_DEBUG
 
@@ -60,5 +59,5 @@ public:
 	virtual FString GetLabel() const override;
 #endif
 
-	static void FindOrAddToSchema(UPoseSearchSchema* Schema, float SampleTimeOffset, int32 ColorPresetIndex, const FName& BoneName = NAME_None, EPermutationTimeType PermutationTimeType = EPermutationTimeType::UseSampleTime);
+	static void FindOrAddToSchema(UPoseSearchSchema* Schema, float SampleTimeOffset, const FName& BoneName = NAME_None, EPermutationTimeType PermutationTimeType = EPermutationTimeType::UseSampleTime);
 };
diff --git a/Engine/Plugins/Experimental/Animation/PoseSearch/Source/Runtime/Private/PoseSearchFeatureChannel_Trajectory.cpp b/Engine/Plugins/Experimental/Animation/PoseSearch/Source/Runtime/Private/PoseSearchFeatureChannel_Trajectory.cpp
--- a/Engine/Plugins/Experimental/Animation/PoseSearch/Source/Runtime/Private/PoseSearchFeatureChannel_Trajectory.cpp
+++ b/Engine/Plugins/Experimental/Animation/PoseSearch/Source/Runtime/Private/PoseSearchFeatureChannel_Trajectory.cpp
@@ -17,10 +17,10 @@ UPoseSearchFeatureChannel_Trajectory::UPoseSearchFeatureChannel_Trajectory()
 #if WITH_EDITOR
 	// defaulting UPoseSearchFeatureChannel_Trajectory to a meaningful locomotion setup
 	Weight = 7.f;
-	Samples.Add(FPoseSearchTrajectorySample({ -0.4f, int32(EPoseSearchTrajectoryFlags::PositionXY), 0.4f, 0 }));
-	Samples.Add(FPoseSearchTrajectorySample({ 0.f, int32(EPoseSearchTrajectoryFlags::VelocityXY | EPoseSearchTrajectoryFlags::FacingDirectionXY), 2.f }));
-	Samples.Add(FPoseSearchTrajectorySample({ 0.35f, int32(EPoseSearchTrajectoryFlags::PositionXY | EPoseSearchTrajectoryFlags::FacingDirectionXY), 0.7f }));
-	Samples.Add(FPoseSearchTrajectorySample({ 0.7f, int32(EPoseSearchTrajectoryFlags::VelocityXY | EPoseSearchTrajectoryFlags::PositionXY | EPoseSearchTrajectoryFlags::FacingDirectionXY), 0.5f }));
+	Samples.Add(FPoseSearchTrajectorySample({ -0.4f, int32(EPoseSearchTrajectoryFlags::PositionXY), 0.4f, FLinearColor::Red }));
+	Samples.Add(FPoseSearchTrajectorySample({ 0.f, int32(EPoseSearchTrajectoryFlags::VelocityXY | EPoseSearchTrajectoryFlags::FacingDirectionXY), 2.f, FLinearColor::Blue }));
+	Samples.Add(FPoseSearchTrajectorySample({ 0.35f, int32(EPoseSearchTrajectoryFlags::PositionXY | EPoseSearchTrajectoryFlags::FacingDirectionXY), 0.7f, FLinearColor::Blue }));
+	Samples.Add(FPoseSearchTrajectorySample({ 0.7f, int32(EPoseSearchTrajectoryFlags::VelocityXY | EPoseSearchTrajectoryFlags::PositionXY | EPoseSearchTrajectoryFlags::FacingDirectionXY), 0.5f, FLinearColor::Blue }));
 #endif
 }
 
@@ -35,7 +35,7 @@ void UPoseSearchFeatureChannel_Trajectory::Finalize(UPoseSearchSchema* Schema)
 			UPoseSearchFeatureChannel_Position* Position = NewObject<UPoseSearchFeatureChannel_Position>(this, NAME_None, RF_Transient);
 			Position->Weight = Sample.Weight * Weight;
 			Position->SampleTimeOffset = Sample.Offset;
-			Position->ColorPresetIndex = Sample.ColorPresetIndex;
+			Position->DebugColor = Sample.DebugColor;
 			Position->InputQueryPose = EInputQueryPose::UseCharacterPose;
 			if (EnumHasAnyFlags(Sample.Flags, EPoseSearchTrajectoryFlags::PositionXY))
 			{
@@ -49,7 +49,7 @@ void UPoseSearchFeatureChannel_Trajectory::Finalize(UPoseSearchSchema* Schema)
 			UPoseSearchFeatureChannel_Velocity* Velocity = NewObject<UPoseSearchFeatureChannel_Velocity>(this, NAME_None, RF_Transient);
 			Velocity->Weight = Sample.Weight * Weight;
 			Velocity->SampleTimeOffset = Sample.Offset;
-			Velocity->ColorPresetIndex = Sample.ColorPresetIndex;
+			Velocity->DebugColor = Sample.DebugColor;
 			Velocity->InputQueryPose = EInputQueryPose::UseCharacterPose;
 			Velocity->bUseCharacterSpaceVelocities = false;
 			if (EnumHasAnyFlags(Sample.Flags, EPoseSearchTrajectoryFlags::VelocityXY))
@@ -64,7 +64,7 @@ void UPoseSearchFeatureChannel_Trajectory::Finalize(UPoseSearchSchema* Schema)
 			UPoseSearchFeatureChannel_Velocity* Velocity = NewObject<UPoseSearchFeatureChannel_Velocity>(this, NAME_None, RF_Transient);
 			Velocity->Weight = Sample.Weight * Weight;
 			Velocity->SampleTimeOffset = Sample.Offset;
-			Velocity->ColorPresetIndex = Sample.ColorPresetIndex;
+			Velocity->DebugColor = Sample.DebugColor;
 			Velocity->InputQueryPose = EInputQueryPose::UseCharacterPose;
 			Velocity->bUseCharacterSpaceVelocities = false;
 			Velocity->bNormalize = true;
@@ -80,7 +80,7 @@ void UPoseSearchFeatureChannel_Trajectory::Finalize(UPoseSearchSchema* Schema)
 			UPoseSearchFeatureChannel_Heading* Heading = NewObject<UPoseSearchFeatureChannel_Heading>(this, NAME_None, RF_Transient);
 			Heading->Weight = Sample.Weight * Weight;
 			Heading->SampleTimeOffset = Sample.Offset;
-			Heading->ColorPresetIndex = Sample.ColorPresetIndex;
+			Heading->DebugColor = Sample.DebugColor;
 			Heading->InputQueryPose = EInputQueryPose::UseCharacterPose;
 			if (EnumHasAnyFlags(Sample.Flags, EPoseSearchTrajectoryFlags::FacingDirectionXY))
 			{
@@ -121,17 +121,17 @@ void UPoseSearchFeatureChannel_Trajectory::DebugDraw(const UE::PoseSearch::FDebu
 		for (int32 i = 0; i < Positions.Num(); ++i)
 		{
 			const float CurrTimeOffset = Positions[i]->SampleTimeOffset;
-			const int32 CurrColorPresetIndex = Positions[i]->ColorPresetIndex;
+			const FColor Color = Positions[i]->DebugColor.ToFColor(true);;
 
 			if (PrevTimeOffset * CurrTimeOffset < UE_KINDA_SMALL_NUMBER)
 			{
 				// we jumped from negative to positive time offset without having a zero time offset. so we add the zero
-				TrajSplinePos.Add(DrawParams.GetCachedPosition(0.f));
-				TrajSplineColor.Add(DrawParams.GetColor(CurrColorPresetIndex));
+				TrajSplinePos.Add(DrawParams.ExtractPosition(PoseVector, 0.f));
+				TrajSplineColor.Add(Color);
 			}
 
-			TrajSplinePos.Add(DrawParams.GetCachedPosition(CurrTimeOffset));
-			TrajSplineColor.Add(DrawParams.GetColor(CurrColorPresetIndex));
+			TrajSplinePos.Add(DrawParams.ExtractPosition(PoseVector, CurrTimeOffset));
+			TrajSplineColor.Add(Color);
 
 			PrevTimeOffset = CurrTimeOffset;
 		}
diff --git a/Engine/Plugins/Experimental/Animation/PoseSearch/Source/Runtime/Private/PoseSearchFeatureChannel_Trajectory.h b/Engine/Plugins/Experimental/Animation/PoseSearch/Source/Runtime/Private/PoseSearchFeatureChannel_Trajectory.h
--- a/Engine/Plugins/Experimental/Animation/PoseSearch/Source/Runtime/Private/PoseSearchFeatureChannel_Trajectory.h
+++ b/Engine/Plugins/Experimental/Animation/PoseSearch/Source/Runtime/Private/PoseSearchFeatureChannel_Trajectory.h
@@ -39,7 +39,7 @@ struct POSESEARCH_API FPoseSearchTrajectorySample
 	float Weight = 1.f;
 
 	UPROPERTY(EditAnywhere, Category = Config, meta = (ExcludeFromHash))
-	int32 ColorPresetIndex = 0;
+	FLinearColor DebugColor = FLinearColor::Blue;
 };
 UCLASS(BlueprintType, EditInlineNew, meta = (DisplayName = "Trajectory Channel"), CollapseCategories)
 class POSESEARCH_API UPoseSearchFeatureChannel_Trajectory : public UPoseSearchFeatureChannel_GroupBase
diff --git a/Engine/Plugins/Experimental/Animation/PoseSearch/Source/Runtime/Private/PoseSearchFeatureChannel_Velocity.cpp b/Engine/Plugins/Experimental/Animation/PoseSearch/Source/Runtime/Private/PoseSearchFeatureChannel_Velocity.cpp
--- a/Engine/Plugins/Experimental/Animation/PoseSearch/Source/Runtime/Private/PoseSearchFeatureChannel_Velocity.cpp
+++ b/Engine/Plugins/Experimental/Animation/PoseSearch/Source/Runtime/Private/PoseSearchFeatureChannel_Velocity.cpp
@@ -23,7 +23,7 @@ void UPoseSearchFeatureChannel_Velocity::AddDependentChannels(UPoseSearchSchema*
 {
 	if (Schema->bInjectAdditionalDebugChannels)
 	{
-		UPoseSearchFeatureChannel_Position::FindOrAddToSchema(Schema, SampleTimeOffset, ColorPresetIndex, Bone.BoneName);
+		UPoseSearchFeatureChannel_Position::FindOrAddToSchema(Schema, SampleTimeOffset, Bone.BoneName);
 	}
 }
 
@@ -62,12 +62,12 @@ void UPoseSearchFeatureChannel_Velocity::DebugDraw(const UE::PoseSearch::FDebugD
 {
 	using namespace UE::PoseSearch;
 
-	const FColor Color = DrawParams.GetColor(ColorPresetIndex);
+	const FColor Color = DebugColor.ToFColor(true);
 	const float LinearVelocityScale = bNormalize ? 15.f : 0.08f;
 
 	const FVector LinearVelocity = DrawParams.GetRootTransform().TransformVector(FFeatureVectorHelper::DecodeVector(PoseVector, ChannelDataOffset, ComponentStripping));
 	const FVector BoneVelDirection = LinearVelocity.GetSafeNormal();
-	const FVector BonePos = DrawParams.GetCachedPosition(SampleTimeOffset, SchemaBoneIdx);
+	const FVector BonePos = DrawParams.ExtractPosition(PoseVector, SampleTimeOffset, SchemaBoneIdx);
 
 	DrawParams.DrawLine(BonePos, BonePos + LinearVelocity * LinearVelocityScale, Color);
 }
diff --git a/Engine/Plugins/Experimental/Animation/PoseSearch/Source/Runtime/Private/PoseSearchFeatureChannel_Velocity.h b/Engine/Plugins/Experimental/Animation/PoseSearch/Source/Runtime/Private/PoseSearchFeatureChannel_Velocity.h
--- a/Engine/Plugins/Experimental/Animation/PoseSearch/Source/Runtime/Private/PoseSearchFeatureChannel_Velocity.h
+++ b/Engine/Plugins/Experimental/Animation/PoseSearch/Source/Runtime/Private/PoseSearchFeatureChannel_Velocity.h
@@ -25,7 +25,7 @@ public:
 	int8 SchemaBoneIdx = 0;
 
 	UPROPERTY(EditAnywhere, Category = "Settings", meta = (ExcludeFromHash))
-	int32 ColorPresetIndex = 0;
+	FLinearColor DebugColor = FLinearColor::Green;
 
 	UPROPERTY(EditAnywhere, Category = "Settings")
 	EInputQueryPose InputQueryPose = EInputQueryPose::UseContinuingPose;
diff --git a/Engine/Plugins/Experimental/Animation/PoseSearch/Source/Runtime/Private/PoseSearchSchema.cpp b/Engine/Plugins/Experimental/Animation/PoseSearch/Source/Runtime/Private/PoseSearchSchema.cpp
--- a/Engine/Plugins/Experimental/Animation/PoseSearch/Source/Runtime/Private/PoseSearchSchema.cpp
+++ b/Engine/Plugins/Experimental/Animation/PoseSearch/Source/Runtime/Private/PoseSearchSchema.cpp
@@ -48,48 +48,53 @@ bool UPoseSearchSchema::IsRootBone(int8 SchemaBoneIdx) const
 int8 UPoseSearchSchema::AddBoneReference(const FBoneReference& BoneReference)
 {
 	int32 SchemaBoneIdx = 0;
-	if (Skeleton)
+	check(Skeleton);
+	
+	bool bDefaultToRootBone = true;
+	FBoneReference TempBoneReference = BoneReference;
+	if (TempBoneReference.BoneName != NAME_None)
 	{
-		bool bDefaultToRootBone = true;
-		FBoneReference TempBoneReference = BoneReference;
-		if (TempBoneReference.BoneName != NAME_None)
+		TempBoneReference.Initialize(Skeleton);
+		if (!TempBoneReference.HasValidSetup())
 		{
-			TempBoneReference.Initialize(Skeleton);
-			if (!TempBoneReference.HasValidSetup())
-			{
-				UE_LOG(LogPoseSearch, Error, TEXT("UPoseSearchSchema::AddBoneReference: couldn't initialize FBoneReference '%s' with Skeleton '%s' in UPoseSearchSchema '%s'. Defaulting to root bone instead"),
-					*TempBoneReference.BoneName.ToString(), *GetNameSafe(Skeleton), *GetNameSafe(this));
-			}
-			else
-			{
-				bDefaultToRootBone = false;
-			}
+			UE_LOG(LogPoseSearch, Error, TEXT("UPoseSearchSchema::AddBoneReference: couldn't initialize FBoneReference '%s' with Skeleton '%s' in UPoseSearchSchema '%s'. Defaulting to root bone instead"),
+				*TempBoneReference.BoneName.ToString(), *GetNameSafe(Skeleton), *GetNameSafe(this));
 		}
-
-		if (bDefaultToRootBone)
+		else
 		{
-			TempBoneReference.BoneName = Skeleton->GetReferenceSkeleton().GetBoneName(int32(RootBoneIndexType));
-			TempBoneReference.Initialize(Skeleton);
-			check(TempBoneReference.HasValidSetup());
+			bDefaultToRootBone = false;
 		}
-
-		SchemaBoneIdx = BoneReferences.AddUnique(TempBoneReference);
-		check(SchemaBoneIdx >= 0 && SchemaBoneIdx < 128);
 	}
-	else if (BoneReferences.IsEmpty())
+
+	if (bDefaultToRootBone)
 	{
-		BoneReferences.Add(FBoneReference());
+		TempBoneReference.BoneName = Skeleton->GetReferenceSkeleton().GetBoneName(int32(RootBoneIndexType));
+		TempBoneReference.Initialize(Skeleton);
+		check(TempBoneReference.HasValidSetup());
 	}
+
+	SchemaBoneIdx = BoneReferences.AddUnique(TempBoneReference);
+	check(SchemaBoneIdx >= 0 && SchemaBoneIdx < 128);
 	return int8(SchemaBoneIdx);
 }
 
 void UPoseSearchSchema::Finalize()
 {
 	BoneReferences.Reset();
-
+	BoneIndicesWithParents.Reset();
+	FinalizedChannels.Reset();
 	SchemaCardinality = 0;
 
-	FinalizedChannels.Reset();
+	if (!Skeleton)
+	{
+		UE_LOG(LogPoseSearch, Error, TEXT("UPoseSearchSchema::Finalize: couldn't Finalize '%s' because Skeleton is null"), *GetNameSafe(this));
+		return;
+	}
+
+	// adding as first bone reference the root bone
+	const int8 SchemaBoneIdx = AddBoneReference(FBoneReference());
+	check(SchemaBoneIdx == RootSchemaBoneIdx);
+
 	for (const TObjectPtr<UPoseSearchFeatureChannel>& ChannelPtr : Channels)
 	{
 		if (ChannelPtr)
@@ -117,27 +122,18 @@ void UPoseSearchSchema::Finalize()
 			break;
 		}
 	}
-
-	BoneIndicesWithParents.Reset();
-	if (Skeleton)
+	
+	// Initialize references to obtain bone indices and fill out bone index array
+	for (FBoneReference& BoneRef : BoneReferences)
 	{
-		// Initialize references to obtain bone indices and fill out bone index array
-		for (FBoneReference& BoneRef : BoneReferences)
-		{
-			check(BoneRef.HasValidSetup());
-			BoneIndicesWithParents.Add(BoneRef.BoneIndex);
-		}
-
-		// Build separate index array with parent indices guaranteed to be present. Sort for EnsureParentsPresent.
-		BoneIndicesWithParents.Sort();
-		FAnimationRuntime::EnsureParentsPresent(BoneIndicesWithParents, Skeleton->GetReferenceSkeleton());
+		check(BoneRef.HasValidSetup());
+		BoneIndicesWithParents.Add(BoneRef.BoneIndex);
 	}
 
-	// BoneIndicesWithParents should at least contain the root to support mirroring root motion
-	if (BoneIndicesWithParents.IsEmpty())
-	{
-		BoneIndicesWithParents.Add(0);
-	}
+	// Build separate index array with parent indices guaranteed to be present. Sort for EnsureParentsPresent.
+	check(!BoneIndicesWithParents.IsEmpty());
+	BoneIndicesWithParents.Sort();
+	FAnimationRuntime::EnsureParentsPresent(BoneIndicesWithParents, Skeleton->GetReferenceSkeleton());
 }
 
 void UPoseSearchSchema::PreSave(FObjectPreSaveContext ObjectSaveContext)
diff --git a/Engine/Plugins/Experimental/Animation/PoseSearch/Source/Runtime/Public/PoseSearch/PoseSearchContext.h b/Engine/Plugins/Experimental/Animation/PoseSearch/Source/Runtime/Public/PoseSearch/PoseSearchContext.h
--- a/Engine/Plugins/Experimental/Animation/PoseSearch/Source/Runtime/Public/PoseSearch/PoseSearchContext.h
+++ b/Engine/Plugins/Experimental/Animation/PoseSearch/Source/Runtime/Public/PoseSearch/PoseSearchContext.h
@@ -10,6 +10,7 @@
 
 struct FTrajectorySampleRange;
 class UPoseSearchDatabase;
+class UPoseSearchFeatureChannel_Position;
 
 namespace UE::PoseSearch
 {
@@ -107,13 +108,11 @@ struct POSESEARCH_API FDebugDrawParams
 	FDebugDrawParams(FAnimInstanceProxy* InAnimInstanceProxy, const UPoseSearchDatabase* InDatabase, EDebugDrawFlags InFlags = EDebugDrawFlags::None);
 	FDebugDrawParams(const UWorld* InWorld, const USkinnedMeshComponent* InMesh, const UPoseSearchDatabase* InDatabase, EDebugDrawFlags InFlags = EDebugDrawFlags::None);
 
-	FColor GetColor(int32 ColorPreset) const;
 	const FPoseSearchIndex* GetSearchIndex() const;
 	const UPoseSearchSchema* GetSchema() const;
 
-	void ClearCachedPositions();
-	void AddCachedPosition(float TimeOffset, int8 SchemaBoneIdx, const FVector& Position);
-	FVector GetCachedPosition(float TimeOffset, int8 SchemaBoneIdx = -1) const;
+	FVector ExtractPosition(TConstArrayView<float> PoseVector, const UPoseSearchFeatureChannel_Position* Position) const;
+	FVector ExtractPosition(TConstArrayView<float> PoseVector, float SampleTimeOffset, int8 SchemaBoneIdx = RootSchemaBoneIdx, EPermutationTimeType PermutationTimeType = EPermutationTimeType::UseSampleTime) const;
 	const FTransform& GetRootTransform() const;
 
 	void DrawLine(const FVector& LineStart, const FVector& LineEnd, const FColor& Color, float Thickness = 0.f) const;
@@ -132,7 +131,6 @@ private:
 	const USkinnedMeshComponent* Mesh = nullptr;
 	const UPoseSearchDatabase* Database = nullptr;
 	EDebugDrawFlags Flags = EDebugDrawFlags::None;
-	FCachedTransforms<FVector> CachedPositions;
 };
 
 #endif // ENABLE_DRAW_DEBUG
diff --git a/Engine/Plugins/Experimental/Animation/PoseSearch/Source/Runtime/Public/PoseSearch/PoseSearchDefines.h b/Engine/Plugins/Experimental/Animation/PoseSearch/Source/Runtime/Public/PoseSearch/PoseSearchDefines.h
--- a/Engine/Plugins/Experimental/Animation/PoseSearch/Source/Runtime/Public/PoseSearch/PoseSearchDefines.h
+++ b/Engine/Plugins/Experimental/Animation/PoseSearch/Source/Runtime/Public/PoseSearch/PoseSearchDefines.h
@@ -21,7 +21,7 @@ constexpr EParallelForFlags ParallelForFlags = EParallelForFlags::None;
 
 DECLARE_LOG_CATEGORY_EXTERN(LogPoseSearch, Log, All);
 
-static constexpr int8 RootSchemaBoneIdx = -1;
+static constexpr int8 RootSchemaBoneIdx = 0;
 static constexpr FBoneIndexType RootBoneIndexType = 0;
 
 
diff --git a/Engine/Plugins/Experimental/Animation/PoseSearch/Source/Runtime/Public/PoseSearch/PoseSearchFeatureChannel.h b/Engine/Plugins/Experimental/Animation/PoseSearch/Source/Runtime/Public/PoseSearch/PoseSearchFeatureChannel.h
--- a/Engine/Plugins/Experimental/Animation/PoseSearch/Source/Runtime/Public/PoseSearch/PoseSearchFeatureChannel.h
+++ b/Engine/Plugins/Experimental/Animation/PoseSearch/Source/Runtime/Public/PoseSearch/PoseSearchFeatureChannel.h
@@ -126,9 +126,6 @@ public:
 	static void GetPermutationTimeOffsets(EPermutationTimeType PermutationTimeType, float DesiredPermutationTimeOffset, float& OutPermutationSampleTimeOffset, float& OutPermutationOriginTimeOffset);
 
 #if ENABLE_DRAW_DEBUG
-	// API called before DebugDraw to collect shared channel informations such as decoded positions form the PoseVector
-	virtual void PreDebugDraw(UE::PoseSearch::FDebugDrawParams& DrawParams, TConstArrayView<float> PoseVector) const {}
-
 	// Draw this channel's data for the given pose vector
 	virtual void DebugDraw(const UE::PoseSearch::FDebugDrawParams& DrawParams, TConstArrayView<float> PoseVector) const {}
 #endif //ENABLE_DRAW_DEBUG
diff --git a/Engine/Plugins/Experimental/Animation/PoseSearch/Source/Runtime/Public/PoseSearch/PoseSearchSchema.h b/Engine/Plugins/Experimental/Animation/PoseSearch/Source/Runtime/Public/PoseSearch/PoseSearchSchema.h
--- a/Engine/Plugins/Experimental/Animation/PoseSearch/Source/Runtime/Public/PoseSearch/PoseSearchSchema.h
+++ b/Engine/Plugins/Experimental/Animation/PoseSearch/Source/Runtime/Public/PoseSearch/PoseSearchSchema.h
@@ -22,18 +22,6 @@ enum class EPoseSearchDataPreprocessor : int32
 	Invalid = Num UMETA(Hidden)
 };
 
-USTRUCT()
-struct FPoseSearchSchemaColorPreset
-{
-	GENERATED_BODY()
-
-	UPROPERTY(EditAnywhere, Category = "Colors", meta = (ExcludeFromHash))
-	FLinearColor Query = FLinearColor::Blue;
-
-	UPROPERTY(EditAnywhere, Category = "Colors", meta = (ExcludeFromHash))
-	FLinearColor Result = FLinearColor::Yellow;
-};
-
 /**
 * Specifies the format of a pose search index. At runtime, queries are built according to the schema for searching.
 */
@@ -104,9 +92,6 @@ public:
 	UPROPERTY(EditAnywhere, Category = "Permutations")
 	float PermutationsTimeOffset = 0.f;
 
-	UPROPERTY(EditAnywhere, Category = "Debug", meta = (ExcludeFromHash))
-	TArray<FPoseSearchSchemaColorPreset> ColorPresets;
-	
 	// if bInjectAdditionalDebugChannels is true, channels will be asked to injecting additional channels into this schema.
 	// the original intent is to add UPoseSearchFeatureChannel_Position(s) to help with the complexity of the debug drawing
 	// (the database will have all the necessary positions to draw lines at the right location and time)
