motion matching - exposing looping and root motion in database reflection

diff --git a/Engine/Plugins/Experimental/Animation/PoseSearch/Source/Editor/Private/PoseSearchDatabaseEditor.cpp b/Engine/Plugins/Experimental/Animation/PoseSearch/Source/Editor/Private/PoseSearchDatabaseEditor.cpp
--- a/Engine/Plugins/Experimental/Animation/PoseSearch/Source/Editor/Private/PoseSearchDatabaseEditor.cpp
+++ b/Engine/Plugins/Experimental/Animation/PoseSearch/Source/Editor/Private/PoseSearchDatabaseEditor.cpp
@@ -525,7 +525,9 @@ namespace UE::PoseSearch
 				{
 					UPoseSearchDatabaseSequenceReflection* NewSelectionReflection = NewObject<UPoseSearchDatabaseSequenceReflection>();
 					NewSelectionReflection->AddToRoot();
-					NewSelectionReflection->Sequence = *DatabaseSequence;
+					static_cast<FPoseSearchDatabaseSequence&>(NewSelectionReflection->Sequence) = *DatabaseSequence;
+					NewSelectionReflection->Sequence.bLooping = DatabaseSequence->IsLooping();
+					NewSelectionReflection->Sequence.bHasRootMotion = DatabaseSequence->IsRootMotionEnabled();
 					NewSelectionReflection->SetSourceLink(SelectedItem, AssetTreeWidget);
 					NewSelectionReflection->SetFlags(RF_Transactional);
 
@@ -535,7 +537,9 @@ namespace UE::PoseSearch
 				{
 					UPoseSearchDatabaseAnimCompositeReflection* NewSelectionReflection = NewObject<UPoseSearchDatabaseAnimCompositeReflection>();
 					NewSelectionReflection->AddToRoot();
-					NewSelectionReflection->AnimComposite = *DatabaseAnimComposite;
+					static_cast<FPoseSearchDatabaseAnimComposite&>(NewSelectionReflection->AnimComposite) = *DatabaseAnimComposite;
+					NewSelectionReflection->AnimComposite.bLooping = DatabaseAnimComposite->IsLooping();
+					NewSelectionReflection->AnimComposite.bHasRootMotion = DatabaseAnimComposite->IsRootMotionEnabled();
 					NewSelectionReflection->SetSourceLink(SelectedItem, AssetTreeWidget);
 					NewSelectionReflection->SetFlags(RF_Transactional);
 					
@@ -545,7 +549,9 @@ namespace UE::PoseSearch
 				{
 					UPoseSearchDatabaseBlendSpaceReflection* NewSelectionReflection = NewObject<UPoseSearchDatabaseBlendSpaceReflection>();
 					NewSelectionReflection->AddToRoot();
-					NewSelectionReflection->BlendSpace = *DatabaseBlendSpace;
+					static_cast<FPoseSearchDatabaseBlendSpace&>(NewSelectionReflection->BlendSpace) = *DatabaseBlendSpace;
+					NewSelectionReflection->BlendSpace.bLooping = DatabaseBlendSpace->IsLooping();
+					NewSelectionReflection->BlendSpace.bHasRootMotion = DatabaseBlendSpace->IsRootMotionEnabled();
 					NewSelectionReflection->SetSourceLink(SelectedItem, AssetTreeWidget);
 					NewSelectionReflection->SetFlags(RF_Transactional);
 					
@@ -555,7 +561,9 @@ namespace UE::PoseSearch
 				{
 					UPoseSearchDatabaseAnimMontageReflection* NewSelectionReflection = NewObject<UPoseSearchDatabaseAnimMontageReflection>();
 					NewSelectionReflection->AddToRoot();
-					NewSelectionReflection->AnimMontage = *DatabaseAnimMontage;
+					static_cast<FPoseSearchDatabaseAnimMontage&>(NewSelectionReflection->AnimMontage) = *DatabaseAnimMontage;
+					NewSelectionReflection->AnimMontage.bLooping = DatabaseAnimMontage->IsLooping();
+					NewSelectionReflection->AnimMontage.bHasRootMotion = DatabaseAnimMontage->IsRootMotionEnabled();
 					NewSelectionReflection->SetSourceLink(SelectedItem, AssetTreeWidget);
 					NewSelectionReflection->SetFlags(RF_Transactional);
 
diff --git a/Engine/Plugins/Experimental/Animation/PoseSearch/Source/Editor/Public/PoseSearchDatabaseEditorReflection.h b/Engine/Plugins/Experimental/Animation/PoseSearch/Source/Editor/Public/PoseSearchDatabaseEditorReflection.h
--- a/Engine/Plugins/Experimental/Animation/PoseSearch/Source/Editor/Public/PoseSearchDatabaseEditorReflection.h
+++ b/Engine/Plugins/Experimental/Animation/PoseSearch/Source/Editor/Public/PoseSearchDatabaseEditorReflection.h
@@ -29,6 +29,18 @@ protected:
 	TSharedPtr<UE::PoseSearch::SDatabaseAssetTree> AssetTreeWidget;
 };
 
+USTRUCT()
+struct FPoseSearchDatabaseSequenceEx : public FPoseSearchDatabaseSequence
+{
+	GENERATED_BODY()
+
+	UPROPERTY(VisibleAnywhere, Category="Sequence", meta = (DisplayPriority = 10))
+	bool bLooping = false;
+
+	UPROPERTY(VisibleAnywhere, Category="Sequence", meta = (DisplayPriority = 11))
+	bool bHasRootMotion = false;
+};
+
 UCLASS()
 class UPoseSearchDatabaseSequenceReflection : public UPoseSearchDatabaseReflectionBase
 {
@@ -36,13 +48,25 @@ class UPoseSearchDatabaseSequenceReflection : public UPoseSearchDatabaseReflecti
 
 public:
 	UPROPERTY(EditAnywhere, Category = "Selected Sequence")
-	FPoseSearchDatabaseSequence Sequence;
+	FPoseSearchDatabaseSequenceEx Sequence;
 
 #if WITH_EDITOR
 	virtual void PostEditChangeProperty(struct FPropertyChangedEvent& PropertyChangedEvent) override;
 #endif
 };
 
+USTRUCT()
+struct FPoseSearchDatabaseBlendSpaceEx : public FPoseSearchDatabaseBlendSpace
+{
+	GENERATED_BODY()
+
+	UPROPERTY(VisibleAnywhere, Category="Sequence", meta = (DisplayPriority = 10))
+	bool bLooping = false;
+
+	UPROPERTY(VisibleAnywhere, Category="Sequence", meta = (DisplayPriority = 11))
+	bool bHasRootMotion = false;
+};
+
 UCLASS()
 class UPoseSearchDatabaseBlendSpaceReflection : public UPoseSearchDatabaseReflectionBase
 {
@@ -50,13 +74,25 @@ class UPoseSearchDatabaseBlendSpaceReflection : public UPoseSearchDatabaseReflec
 
 public:
 	UPROPERTY(EditAnywhere, Category = "Selected Blend Space")
-	FPoseSearchDatabaseBlendSpace BlendSpace;
+	FPoseSearchDatabaseBlendSpaceEx BlendSpace;
 
 #if WITH_EDITOR
 	virtual void PostEditChangeProperty(struct FPropertyChangedEvent& PropertyChangedEvent) override;
 #endif
 };
 
+USTRUCT()
+struct FPoseSearchDatabaseAnimCompositeEx : public FPoseSearchDatabaseAnimComposite
+{
+	GENERATED_BODY()
+
+	UPROPERTY(VisibleAnywhere, Category="Sequence", meta = (DisplayPriority = 10))
+	bool bLooping = false;
+
+	UPROPERTY(VisibleAnywhere, Category="Sequence", meta = (DisplayPriority = 11))
+	bool bHasRootMotion = false;
+};
+
 UCLASS()
 class UPoseSearchDatabaseAnimCompositeReflection : public UPoseSearchDatabaseReflectionBase
 {
@@ -64,13 +100,25 @@ class UPoseSearchDatabaseAnimCompositeReflection : public UPoseSearchDatabaseRef
 
 public:
 	UPROPERTY(EditAnywhere, Category = "Selected Anim Composite")
-	FPoseSearchDatabaseAnimComposite AnimComposite;
+	FPoseSearchDatabaseAnimCompositeEx AnimComposite;
 
 #if WITH_EDITOR
 	virtual void PostEditChangeProperty(struct FPropertyChangedEvent& PropertyChangedEvent) override;
 #endif
 };
 
+USTRUCT()
+struct FPoseSearchDatabaseAnimMontageEx : public FPoseSearchDatabaseAnimMontage
+{
+	GENERATED_BODY()
+
+	UPROPERTY(VisibleAnywhere, Category="Sequence", meta = (DisplayPriority = 10))
+	bool bLooping = false;
+
+	UPROPERTY(VisibleAnywhere, Category="Sequence", meta = (DisplayPriority = 11))
+	bool bHasRootMotion = false;
+};
+
 UCLASS()
 class UPoseSearchDatabaseAnimMontageReflection : public UPoseSearchDatabaseReflectionBase
 {
@@ -78,7 +126,7 @@ class UPoseSearchDatabaseAnimMontageReflection : public UPoseSearchDatabaseRefle
 
 public:
 	UPROPERTY(EditAnywhere, Category = "Selected Anim Montage")
-	FPoseSearchDatabaseAnimMontage AnimMontage;
+	FPoseSearchDatabaseAnimMontageEx AnimMontage;
 
 #if WITH_EDITOR
 	virtual void PostEditChangeProperty(struct FPropertyChangedEvent& PropertyChangedEvent) override;
diff --git a/Engine/Plugins/Experimental/Animation/PoseSearch/Source/Runtime/Public/PoseSearch/PoseSearchDatabase.h b/Engine/Plugins/Experimental/Animation/PoseSearch/Source/Runtime/Public/PoseSearch/PoseSearchDatabase.h
--- a/Engine/Plugins/Experimental/Animation/PoseSearch/Source/Runtime/Public/PoseSearch/PoseSearchDatabase.h
+++ b/Engine/Plugins/Experimental/Animation/PoseSearch/Source/Runtime/Public/PoseSearch/PoseSearchDatabase.h
@@ -74,16 +74,16 @@ struct POSESEARCH_API FPoseSearchDatabaseSequence : public FPoseSearchDatabaseAn
 	GENERATED_BODY()
 	virtual ~FPoseSearchDatabaseSequence() = default;
 
-	UPROPERTY(EditAnywhere, Category="Sequence")
+	UPROPERTY(EditAnywhere, Category="Sequence", meta = (DisplayPriority = 0))
 	TObjectPtr<UAnimSequence> Sequence;
 
-	UPROPERTY(EditAnywhere, Category = "Sequence")
+	UPROPERTY(EditAnywhere, Category = "Sequence", meta = (DisplayPriority = 3))
 	bool bEnabled = true;
 
-	UPROPERTY(EditAnywhere, Category="Sequence")
+	UPROPERTY(EditAnywhere, Category="Sequence", meta = (DisplayPriority = 1))
 	FFloatInterval SamplingRange = FFloatInterval(0.0f, 0.0f);
 
-	UPROPERTY(EditAnywhere, Category = "Sequence")
+	UPROPERTY(EditAnywhere, Category = "Sequence", meta = (DisplayPriority = 2))
 	EPoseSearchMirrorOption MirrorOption = EPoseSearchMirrorOption::UnmirroredOnly;
 
 	UAnimationAsset* GetAnimationAsset() const override;
@@ -103,24 +103,24 @@ struct POSESEARCH_API FPoseSearchDatabaseBlendSpace : public FPoseSearchDatabase
 	GENERATED_BODY()
 	virtual ~FPoseSearchDatabaseBlendSpace() = default;
 
-	UPROPERTY(EditAnywhere, Category = "BlendSpace")
+	UPROPERTY(EditAnywhere, Category = "BlendSpace", meta = (DisplayPriority = 0))
 	TObjectPtr<UBlendSpace> BlendSpace;
 
-	UPROPERTY(EditAnywhere, Category = "BlendSpace")
+	UPROPERTY(EditAnywhere, Category = "BlendSpace", meta = (DisplayPriority = 5))
 	bool bEnabled = true;
 
-	UPROPERTY(EditAnywhere, Category = "BlendSpace")
+	UPROPERTY(EditAnywhere, Category = "BlendSpace", meta = (DisplayPriority = 1))
 	EPoseSearchMirrorOption MirrorOption = EPoseSearchMirrorOption::UnmirroredOnly;
 
 	// If to use the blendspace grid locations as parameter sample locations.
 	// When enabled, NumberOfHorizontalSamples and NumberOfVerticalSamples are ignored.
-	UPROPERTY(EditAnywhere, Category = "BlendSpace")
+	UPROPERTY(EditAnywhere, Category = "BlendSpace", meta = (DisplayPriority = 4))
 	bool bUseGridForSampling = false;
 
-	UPROPERTY(EditAnywhere, Category = "BlendSpace", meta = (EditCondition = "!bUseGridForSampling", EditConditionHides, ClampMin = "1", UIMin = "1", UIMax = "25"))
+	UPROPERTY(EditAnywhere, Category = "BlendSpace", meta = (EditCondition = "!bUseGridForSampling", EditConditionHides, ClampMin = "1", UIMin = "1", UIMax = "25", DisplayPriority = 2))
 	int32 NumberOfHorizontalSamples = 9;
 
-	UPROPERTY(EditAnywhere, Category = "BlendSpace", meta = (EditCondition = "!bUseGridForSampling", EditConditionHides, ClampMin = "1", UIMin = "1", UIMax = "25"))
+	UPROPERTY(EditAnywhere, Category = "BlendSpace", meta = (EditCondition = "!bUseGridForSampling", EditConditionHides, ClampMin = "1", UIMin = "1", UIMax = "25", DisplayPriority = 3))
 	int32 NumberOfVerticalSamples = 2;
 
 	UAnimationAsset* GetAnimationAsset() const override;
@@ -143,16 +143,16 @@ struct POSESEARCH_API FPoseSearchDatabaseAnimComposite : public FPoseSearchDatab
 	GENERATED_BODY()
 	virtual ~FPoseSearchDatabaseAnimComposite() = default;
 
-	UPROPERTY(EditAnywhere, Category = "AnimComposite")
+	UPROPERTY(EditAnywhere, Category = "AnimComposite", meta = (DisplayPriority = 0))
 	TObjectPtr<UAnimComposite> AnimComposite;
 
-	UPROPERTY(EditAnywhere, Category = "AnimComposite")
+	UPROPERTY(EditAnywhere, Category = "AnimComposite", meta = (DisplayPriority = 3))
 	bool bEnabled = true;
 
-	UPROPERTY(EditAnywhere, Category = "AnimComposite")
+	UPROPERTY(EditAnywhere, Category = "AnimComposite", meta = (DisplayPriority = 1))
 	FFloatInterval SamplingRange = FFloatInterval(0.0f, 0.0f);
 
-	UPROPERTY(EditAnywhere, Category = "AnimComposite")
+	UPROPERTY(EditAnywhere, Category = "AnimComposite", meta = (DisplayPriority = 2))
 	EPoseSearchMirrorOption MirrorOption = EPoseSearchMirrorOption::UnmirroredOnly;
 
 	UAnimationAsset* GetAnimationAsset() const override;
@@ -172,16 +172,16 @@ struct POSESEARCH_API FPoseSearchDatabaseAnimMontage : public FPoseSearchDatabas
 	GENERATED_BODY()
 	virtual ~FPoseSearchDatabaseAnimMontage() = default;
 
-	UPROPERTY(EditAnywhere, Category="AnimMontage")
+	UPROPERTY(EditAnywhere, Category="AnimMontage", meta = (DisplayPriority = 0))
 	TObjectPtr<UAnimMontage> AnimMontage;
 
-	UPROPERTY(EditAnywhere, Category = "AnimMontage")
+	UPROPERTY(EditAnywhere, Category = "AnimMontage", meta = (DisplayPriority = 3))
 	bool bEnabled = true;
 
-	UPROPERTY(EditAnywhere, Category="AnimMontage")
+	UPROPERTY(EditAnywhere, Category="AnimMontage", meta = (DisplayPriority = 1))
 	FFloatInterval SamplingRange = FFloatInterval(0.0f, 0.0f);
 
-	UPROPERTY(EditAnywhere, Category = "AnimMontage")
+	UPROPERTY(EditAnywhere, Category = "AnimMontage", meta = (DisplayPriority = 2))
 	EPoseSearchMirrorOption MirrorOption = EPoseSearchMirrorOption::UnmirroredOnly;
 
 	UAnimationAsset* GetAnimationAsset() const override;
