motion matching - adding SDatabaseAssetListItem::OnMouseButtonDoubleClick to open the related asset editor

diff --git a/Engine/Plugins/Experimental/Animation/PoseSearch/Source/Editor/Private/PoseSearchDatabaseAssetListItem.cpp b/Engine/Plugins/Experimental/Animation/PoseSearch/Source/Editor/Private/PoseSearchDatabaseAssetListItem.cpp
--- a/Engine/Plugins/Experimental/Animation/PoseSearch/Source/Editor/Private/PoseSearchDatabaseAssetListItem.cpp
+++ b/Engine/Plugins/Experimental/Animation/PoseSearch/Source/Editor/Private/PoseSearchDatabaseAssetListItem.cpp
@@ -26,6 +26,8 @@
 
 #include "ScopedTransaction.h"
 #include "Styling/StyleColors.h"
+#include "Subsystems/AssetEditorSubsystem.h"
+#include "Editor.h"
 
 #define LOCTEXT_NAMESPACE "SDatabaseAssetListItem"
 
@@ -110,6 +112,24 @@ namespace UE::PoseSearch
 		SkeletonView.Pin()->RefreshTreeView(false);
 	}
 
+	FReply SDatabaseAssetListItem::OnMouseButtonDoubleClick(const FGeometry& InMyGeometry, const FPointerEvent& InMouseEvent)
+	{
+		TSharedPtr<FDatabaseAssetTreeNode> Node = WeakAssetTreeNode.Pin();
+		TSharedPtr<FDatabaseViewModel> ViewModel = EditorViewModel.Pin();
+		if (const UPoseSearchDatabase* Database = ViewModel->GetPoseSearchDatabase())
+		{
+			if (const FPoseSearchDatabaseAnimationAssetBase* DatabaseAnimationAsset = Database->GetAnimationAssetBase(Node->SourceAssetIdx))
+			{
+				if (UAssetEditorSubsystem* AssetEditorSS = GEditor->GetEditorSubsystem<UAssetEditorSubsystem>())
+				{
+					AssetEditorSS->OpenEditorForAsset(DatabaseAnimationAsset->GetAnimationAsset());
+				}
+			}
+		}
+
+		return STableRow<TSharedPtr<FDatabaseAssetTreeNode>>::OnMouseButtonDoubleClick(InMyGeometry, InMouseEvent);
+	}
+
 	FText SDatabaseAssetListItem::GetName() const
 	{
 		TSharedPtr<FDatabaseAssetTreeNode> Node = WeakAssetTreeNode.Pin();
diff --git a/Engine/Plugins/Experimental/Animation/PoseSearch/Source/Editor/Public/PoseSearchDatabaseAssetListItem.h b/Engine/Plugins/Experimental/Animation/PoseSearch/Source/Editor/Public/PoseSearchDatabaseAssetListItem.h
--- a/Engine/Plugins/Experimental/Animation/PoseSearch/Source/Editor/Public/PoseSearchDatabaseAssetListItem.h
+++ b/Engine/Plugins/Experimental/Animation/PoseSearch/Source/Editor/Public/PoseSearchDatabaseAssetListItem.h
@@ -36,6 +36,7 @@ namespace UE::PoseSearch
 		void OnAddBlendSpace();
 		void OnAddAnimComposite();
 		void OnAddAnimMontage();
+		virtual FReply OnMouseButtonDoubleClick(const FGeometry& InMyGeometry, const FPointerEvent& InMouseEvent) override;
 
 		EVisibility GetSelectedActorIconVisbility() const;
 
