motion matching - adding units to UPoseSearchDatabaseStatistics Kinematic Information

diff --git a/Engine/Plugins/Experimental/Animation/PoseSearch/Source/Editor/Private/PoseSearchDatabaseEditorReflection.cpp b/Engine/Plugins/Experimental/Animation/PoseSearch/Source/Editor/Private/PoseSearchDatabaseEditorReflection.cpp
--- a/Engine/Plugins/Experimental/Animation/PoseSearch/Source/Editor/Private/PoseSearchDatabaseEditorReflection.cpp
+++ b/Engine/Plugins/Experimental/Animation/PoseSearch/Source/Editor/Private/PoseSearchDatabaseEditorReflection.cpp
@@ -133,10 +133,11 @@ void UPoseSearchDatabaseStatistics::Initialize(const UPoseSearchDatabase* PoseSe
 			
 		// Kinematic Information
 	
-		AverageSpeed = SearchIndex.Stats.AverageSpeed;
-		MaxSpeed = SearchIndex.Stats.MaxSpeed;
-		AverageAcceleration = SearchIndex.Stats.AverageAcceleration;
-		MaxAcceleration = SearchIndex.Stats.MaxAcceleration;
+		// using FText instead of meta = (ForceUnits = "cm/s") to keep properties consistent
+		AverageSpeed = FText::Format(LOCTEXT("StatsAverageSpeed", "{0} cm/s"), SearchIndex.Stats.AverageSpeed);
+		MaxSpeed = FText::Format(LOCTEXT("StatsMaxSpeed", "{0} cm/s"), SearchIndex.Stats.MaxSpeed);
+		AverageAcceleration = FText::Format(LOCTEXT("StatsAverageAcceleration", "{0} cm/s²"), SearchIndex.Stats.AverageAcceleration);
+		MaxAcceleration = FText::Format(LOCTEXT("StatsMaxAcceleration", "{0} cm/s²"), SearchIndex.Stats.MaxAcceleration);
 
 		// Principal Component Analysis
 			
diff --git a/Engine/Plugins/Experimental/Animation/PoseSearch/Source/Editor/Public/PoseSearchDatabaseEditorReflection.h b/Engine/Plugins/Experimental/Animation/PoseSearch/Source/Editor/Public/PoseSearchDatabaseEditorReflection.h
--- a/Engine/Plugins/Experimental/Animation/PoseSearch/Source/Editor/Public/PoseSearchDatabaseEditorReflection.h
+++ b/Engine/Plugins/Experimental/Animation/PoseSearch/Source/Editor/Public/PoseSearchDatabaseEditorReflection.h
@@ -157,7 +157,6 @@ struct FPoseSearchDatabaseMemoryStats
 	FText AssetsSize;
 	
 	void Initialize(const UPoseSearchDatabase* PoseSearchDatabase);
-	static FText ToMemoryBudgetText(int32 Size);
 };
 
 UCLASS()
@@ -187,16 +186,16 @@ public:
 	// Kinematic Information
 
 	UPROPERTY(VisibleAnywhere, Category = "Kinematic Information")
-	float AverageSpeed;
+	FText AverageSpeed;
 
 	UPROPERTY(VisibleAnywhere, Category = "Kinematic Information")
-	float MaxSpeed;
+	FText MaxSpeed;
 
 	UPROPERTY(VisibleAnywhere, Category = "Kinematic Information")
-	float AverageAcceleration;
+	FText AverageAcceleration;
 
 	UPROPERTY(VisibleAnywhere, Category = "Kinematic Information")
-	float MaxAcceleration;
+	FText MaxAcceleration;
 
 	// Principal Component Analysis Information
 
