motion matching - quantizing FDatabaseViewModel::UpdatePreviewActors animation "position" to reflect database samples

diff --git a/Engine/Plugins/Experimental/Animation/PoseSearch/Source/Editor/Private/PoseSearchDatabaseViewModel.cpp b/Engine/Plugins/Experimental/Animation/PoseSearch/Source/Editor/Private/PoseSearchDatabaseViewModel.cpp
--- a/Engine/Plugins/Experimental/Animation/PoseSearch/Source/Editor/Private/PoseSearchDatabaseViewModel.cpp
+++ b/Engine/Plugins/Experimental/Animation/PoseSearch/Source/Editor/Private/PoseSearchDatabaseViewModel.cpp
@@ -46,12 +46,6 @@ namespace UE::PoseSearch
 		return Sampler.Get();
 	}
 
-	float FDatabasePreviewActor::GetScaledTime(float Time) const
-	{
-		check(Sampler);
-		return Sampler->GetScaledTime(Time);
-	}
-
 	UDebugSkelMeshComponent* FDatabasePreviewActor::GetDebugSkelMeshComponent()
 	{
 		if (IsValid())
@@ -345,9 +339,20 @@ namespace UE::PoseSearch
 
 			float CurrentTime = 0.f;
 			FAnimationRuntime::AdvanceTime(false, PlayTime, CurrentTime, PreviewAsset->GetPlayLength());
-			const float CurrentScaledTime = PreviewActor.GetScaledTime(CurrentTime);
 
+			// time to pose index
 			const FPoseSearchIndexAsset& IndexAsset = SearchIndex.Assets[PreviewActor.IndexAssetIndex];
+			PreviewActor.CurrentPoseIndex = PoseSearchDatabase->GetPoseIndexFromTime(CurrentTime, IndexAsset);
+
+			// pose index to quantized time
+			if (PreviewActor.CurrentPoseIndex >= 0)
+			{
+				CurrentTime = PoseSearchDatabase->GetAssetTime(PreviewActor.CurrentPoseIndex);
+			}
+
+			// quantized time to scaled quantized time
+			const float CurrentScaledTime = Sampler->GetScaledTime(CurrentTime);
+
 			AnimInstance->SetPosition(CurrentScaledTime);
 			AnimInstance->SetPlayRate(0.f);
 			AnimInstance->SetBlendSpacePosition(IndexAsset.BlendParameters);
@@ -359,7 +364,6 @@ namespace UE::PoseSearch
 			}
 			PreviewActor.Actor->SetActorTransform(RootMotion);
 
-			PreviewActor.CurrentPoseIndex = PoseSearchDatabase->GetPoseIndexFromTime(CurrentTime, IndexAsset);
 		}
 	}
 
diff --git a/Engine/Plugins/Experimental/Animation/PoseSearch/Source/Editor/Public/PoseSearchDatabaseViewModel.h b/Engine/Plugins/Experimental/Animation/PoseSearch/Source/Editor/Public/PoseSearchDatabaseViewModel.h
--- a/Engine/Plugins/Experimental/Animation/PoseSearch/Source/Editor/Public/PoseSearchDatabaseViewModel.h
+++ b/Engine/Plugins/Experimental/Animation/PoseSearch/Source/Editor/Public/PoseSearchDatabaseViewModel.h
@@ -52,7 +52,6 @@ namespace UE::PoseSearch
 		bool IsValid() const;
 		void Process();
 		const FAssetSamplerBase* GetSampler() const;
-		float GetScaledTime(float Time) const;
 		UDebugSkelMeshComponent* GetDebugSkelMeshComponent();
 		UAnimPreviewInstance* GetAnimPreviewInstance();
 	};
