motion matching - removing FPoseSearchExtrapolationParameters

diff --git a/Engine/Plugins/Experimental/Animation/PoseSearch/Source/Editor/Private/PoseSearchDatabaseViewModel.cpp b/Engine/Plugins/Experimental/Animation/PoseSearch/Source/Editor/Private/PoseSearchDatabaseViewModel.cpp
--- a/Engine/Plugins/Experimental/Animation/PoseSearch/Source/Editor/Private/PoseSearchDatabaseViewModel.cpp
+++ b/Engine/Plugins/Experimental/Animation/PoseSearch/Source/Editor/Private/PoseSearchDatabaseViewModel.cpp
@@ -206,7 +206,6 @@ namespace UE::PoseSearch
 		if (const FPoseSearchDatabaseSequence* DatabaseSequence = DatabaseAsset.GetPtr<FPoseSearchDatabaseSequence>())
 		{
 			FSequenceBaseSampler::FInput Input;
-			Input.ExtrapolationParameters = PoseSearchDatabase->ExtrapolationParameters;
 			Input.SequenceBase = DatabaseSequence->Sequence;
 
 			TSharedPtr<FSequenceBaseSampler> Sampler = MakeShared<FSequenceBaseSampler>();
@@ -216,7 +215,6 @@ namespace UE::PoseSearch
 		else if (const FPoseSearchDatabaseAnimComposite* DatabaseAnimComposite = DatabaseAsset.GetPtr<FPoseSearchDatabaseAnimComposite>())
 		{
 			FSequenceBaseSampler::FInput Input;
-			Input.ExtrapolationParameters = PoseSearchDatabase->ExtrapolationParameters;
 			Input.SequenceBase = DatabaseAnimComposite->AnimComposite;
 
 			TSharedPtr<FSequenceBaseSampler> Sampler = MakeShared<FSequenceBaseSampler>();
@@ -227,7 +225,6 @@ namespace UE::PoseSearch
 		{
 			FBlendSpaceSampler::FInput Input;
 			Input.BoneContainer = BoneContainer;
-			Input.ExtrapolationParameters = PoseSearchDatabase->ExtrapolationParameters;
 			Input.BlendSpace = DatabaseBlendSpace->BlendSpace;
 			Input.BlendParameters = IndexAsset.BlendParameters;
 
@@ -238,7 +235,6 @@ namespace UE::PoseSearch
 		else if (const FPoseSearchDatabaseAnimMontage* DatabaseAnimMontage = DatabaseAsset.GetPtr<FPoseSearchDatabaseAnimMontage>())
 		{
 			FAnimMontageSampler::FInput Input;
-			Input.ExtrapolationParameters = PoseSearchDatabase->ExtrapolationParameters;
 			Input.AnimMontage = DatabaseAnimMontage->AnimMontage;
 
 			TSharedPtr<FAnimMontageSampler> Sampler = MakeShared<FAnimMontageSampler>();
diff --git a/Engine/Plugins/Experimental/Animation/PoseSearch/Source/Runtime/Private/PoseSearchAssetSampler.cpp b/Engine/Plugins/Experimental/Animation/PoseSearch/Source/Runtime/Private/PoseSearchAssetSampler.cpp
--- a/Engine/Plugins/Experimental/Animation/PoseSearch/Source/Runtime/Private/PoseSearchAssetSampler.cpp
+++ b/Engine/Plugins/Experimental/Animation/PoseSearch/Source/Runtime/Private/PoseSearchAssetSampler.cpp
@@ -20,35 +20,11 @@ static FTransform ExtrapolateRootMotion(
 	FTransform SampleToExtrapolate,
 	float SampleStart, 
 	float SampleEnd, 
-	float ExtrapolationTime,
-	const FPoseSearchExtrapolationParameters& ExtrapolationParameters)
+	float ExtrapolationTime)
 {
 	const float SampleDelta = SampleEnd - SampleStart;
 	check(!FMath::IsNearlyZero(SampleDelta));
 
-	const FVector LinearVelocityToExtrapolate = SampleToExtrapolate.GetTranslation() / SampleDelta;
-	const float LinearSpeedToExtrapolate = LinearVelocityToExtrapolate.Size();
-	const bool bCanExtrapolateTranslation =
-		LinearSpeedToExtrapolate >= ExtrapolationParameters.LinearSpeedThreshold;
-
-	const float AngularSpeedToExtrapolateRad = SampleToExtrapolate.GetRotation().GetAngle() / SampleDelta;
-	const bool bCanExtrapolateRotation = FMath::RadiansToDegrees(AngularSpeedToExtrapolateRad) >= ExtrapolationParameters.AngularSpeedThreshold;
-
-	if (!bCanExtrapolateTranslation && !bCanExtrapolateRotation)
-	{
-		return FTransform::Identity;
-	}
-
-	if (!bCanExtrapolateTranslation)
-	{
-		SampleToExtrapolate.SetTranslation(FVector::ZeroVector);
-	}
-
-	if (!bCanExtrapolateRotation)
-	{
-		SampleToExtrapolate.SetRotation(FQuat::Identity);
-	}
-
 	// converting ExtrapolationTime to a positive number to avoid dealing with the negative extrapolation and inverting
 	// transforms later on.
 	const float AbsExtrapolationTime = FMath::Abs(ExtrapolationTime);
@@ -130,8 +106,6 @@ FTransform FSequenceBaseSampler::ExtractRootTransform(float Time) const
 	}
 	else
 	{
-		const float ExtrapolationSampleTime = Input.ExtrapolationParameters.SampleTime;
-
 		const float PlayLength = SequenceBase->GetPlayLength();
 		const float ClampedTime = FMath::Clamp(Time, 0.0f, PlayLength);
 		const float ExtrapolationTime = Time - ClampedTime;
@@ -145,8 +119,7 @@ FTransform FSequenceBaseSampler::ExtractRootTransform(float Time) const
 			const FTransform ExtrapolatedRootMotion = UE::PoseSearch::ExtrapolateRootMotion(
 				SampleToExtrapolate,
 				0.0f, ExtrapolationSampleTime,
-				ExtrapolationTime,
-				Input.ExtrapolationParameters);
+				ExtrapolationTime);
 			RootTransform = ExtrapolatedRootMotion;
 		}
 		else
@@ -162,8 +135,7 @@ FTransform FSequenceBaseSampler::ExtractRootTransform(float Time) const
 				const FTransform ExtrapolatedRootMotion = UE::PoseSearch::ExtrapolateRootMotion(
 					SampleToExtrapolate,
 					PlayLength - ExtrapolationSampleTime, PlayLength,
-					ExtrapolationTime,
-					Input.ExtrapolationParameters);
+					ExtrapolationTime);
 				RootTransform = ExtrapolatedRootMotion * RootTransform;
 			}
 		}
@@ -177,7 +149,6 @@ void FSequenceBaseSampler::ExtractPoseSearchNotifyStates(
 	TArray<UAnimNotifyState_PoseSearchBase*>& NotifyStates) const
 {
 	// getting pose search notifies in an interval of size ExtractionInterval, centered on Time
-	constexpr float ExtractionInterval = 1.0f / 120.0f;
 	FAnimNotifyContext NotifyContext;
 	Input.SequenceBase->GetAnimNotifies(Time - (ExtractionInterval * 0.5f), ExtractionInterval, NotifyContext);
 
@@ -275,8 +246,6 @@ FTransform FBlendSpaceSampler::ExtractRootTransform(float Time) const
 	}
 	else
 	{
-		const float ExtrapolationSampleTime = Input.ExtrapolationParameters.SampleTime;
-
 		const float ClampedTime = FMath::Clamp(Time, 0.0f, PlayLength);
 		const float ExtrapolationTime = Time - ClampedTime;
 
@@ -290,8 +259,7 @@ FTransform FBlendSpaceSampler::ExtractRootTransform(float Time) const
 			const FTransform ExtrapolatedRootMotion = UE::PoseSearch::ExtrapolateRootMotion(
 				SampleToExtrapolate,
 				0.0f, ExtrapolationSampleTime,
-				ExtrapolationTime,
-				Input.ExtrapolationParameters);
+				ExtrapolationTime);
 			RootTransform = ExtrapolatedRootMotion;
 		}
 		else
@@ -307,8 +275,7 @@ FTransform FBlendSpaceSampler::ExtractRootTransform(float Time) const
 				const FTransform ExtrapolatedRootMotion = UE::PoseSearch::ExtrapolateRootMotion(
 					SampleToExtrapolate,
 					PlayLength - ExtrapolationSampleTime, PlayLength,
-					ExtrapolationTime,
-					Input.ExtrapolationParameters);
+					ExtrapolationTime);
 				RootTransform = ExtrapolatedRootMotion * RootTransform;
 			}
 		}
@@ -349,8 +316,6 @@ void FBlendSpaceSampler::ExtractPoseSearchNotifyStates(
 		check(HighestWeightIndex != -1);
 
 		// getting pose search notifies in an interval of size ExtractionInterval, centered on Time
-		constexpr float ExtractionInterval = 1.0f / 120.0f;
-
 		float SampleTime = Time * (BlendSamples[HighestWeightIndex].Animation->GetPlayLength() / PlayLength);
 
 		// Get notifies for highest weighted
@@ -622,8 +587,6 @@ FTransform FAnimMontageSampler::ExtractRootTransform(float Time) const
 	}
 	else
 	{
-		const float ExtrapolationSampleTime = Input.ExtrapolationParameters.SampleTime;
-
 		const float PlayLength = GetPlayLength();
 		const float ClampedTime = FMath::Clamp(Time, 0.f, PlayLength);
 		const float ExtrapolationTime = Time - ClampedTime;
@@ -637,8 +600,7 @@ FTransform FAnimMontageSampler::ExtractRootTransform(float Time) const
 			const FTransform ExtrapolatedRootMotion = UE::PoseSearch::ExtrapolateRootMotion(
 				SampleToExtrapolate,
 				0.0f, ExtrapolationSampleTime,
-				ExtrapolationTime,
-				Input.ExtrapolationParameters);
+				ExtrapolationTime);
 			RootTransform = ExtrapolatedRootMotion;
 		}
 		else
@@ -654,8 +616,7 @@ FTransform FAnimMontageSampler::ExtractRootTransform(float Time) const
 				const FTransform ExtrapolatedRootMotion = UE::PoseSearch::ExtrapolateRootMotion(
 					SampleToExtrapolate,
 					PlayLength - ExtrapolationSampleTime, PlayLength,
-					ExtrapolationTime,
-					Input.ExtrapolationParameters);
+					ExtrapolationTime);
 				RootTransform = ExtrapolatedRootMotion * RootTransform;
 			}
 		}
@@ -667,7 +628,6 @@ FTransform FAnimMontageSampler::ExtractRootTransform(float Time) const
 void FAnimMontageSampler::ExtractPoseSearchNotifyStates(float Time, TArray<UAnimNotifyState_PoseSearchBase*>& NotifyStates) const
 {
 	// getting pose search notifies in an interval of size ExtractionInterval, centered on Time
-	constexpr float ExtractionInterval = 1.0f / 120.0f;
 	FAnimNotifyContext NotifyContext;
 	Input.AnimMontage->GetAnimNotifies(Time - (ExtractionInterval * 0.5f), ExtractionInterval, NotifyContext);
 
diff --git a/Engine/Plugins/Experimental/Animation/PoseSearch/Source/Runtime/Private/PoseSearchDatabaseIndexingContext.cpp b/Engine/Plugins/Experimental/Animation/PoseSearch/Source/Runtime/Private/PoseSearchDatabaseIndexingContext.cpp
--- a/Engine/Plugins/Experimental/Animation/PoseSearch/Source/Runtime/Private/PoseSearchDatabaseIndexingContext.cpp
+++ b/Engine/Plugins/Experimental/Animation/PoseSearch/Source/Runtime/Private/PoseSearchDatabaseIndexingContext.cpp
@@ -74,7 +74,6 @@ bool FDatabaseIndexingContext::IndexDatabase(FPoseSearchIndexBase& SearchIndexBa
 			if (Sequence && !SamplerMap.Contains(Sequence))
 			{
 				FSequenceBaseSampler::FInput Input;
-				Input.ExtrapolationParameters = Database.ExtrapolationParameters;
 				Input.SequenceBase = Sequence;
 
 				SamplerMap.Add(Sequence, Samplers.Num());
@@ -110,7 +109,6 @@ bool FDatabaseIndexingContext::IndexDatabase(FPoseSearchIndexBase& SearchIndexBa
 						{
 							FBlendSpaceSampler::FInput Input;
 							Input.BoneContainer = BoneContainer;
-							Input.ExtrapolationParameters = Database.ExtrapolationParameters;
 							Input.BlendSpace = DatabaseBlendSpace->BlendSpace;
 							Input.BlendParameters = BlendParameters;
 
@@ -129,7 +127,6 @@ bool FDatabaseIndexingContext::IndexDatabase(FPoseSearchIndexBase& SearchIndexBa
 			if (DatabaseAnimMontage->AnimMontage && !SamplerMap.Contains(DatabaseAnimMontage->AnimMontage))
 			{
 				FAnimMontageSampler::FInput Input;
-				Input.ExtrapolationParameters = Database.ExtrapolationParameters;
 				Input.AnimMontage = DatabaseAnimMontage->AnimMontage;
 
 				SamplerMap.Add(DatabaseAnimMontage->AnimMontage, Samplers.Num());
diff --git a/Engine/Plugins/Experimental/Animation/PoseSearch/Source/Runtime/Private/PoseSearchDerivedDataKey.cpp b/Engine/Plugins/Experimental/Animation/PoseSearch/Source/Runtime/Private/PoseSearchDerivedDataKey.cpp
--- a/Engine/Plugins/Experimental/Animation/PoseSearch/Source/Runtime/Private/PoseSearchDerivedDataKey.cpp
+++ b/Engine/Plugins/Experimental/Animation/PoseSearch/Source/Runtime/Private/PoseSearchDerivedDataKey.cpp
@@ -41,7 +41,7 @@ FKeyBuilder::FKeyBuilder(const UObject* Object, bool bUseDataVer)
 	if (bUseDataVer)
 	{
 		// used to invalidate the key without having to change POSESEARCHDB_DERIVEDDATA_VER all the times
-		int32 POSESEARCHDB_DERIVEDDATA_VER_SMALL = 121;
+		int32 POSESEARCHDB_DERIVEDDATA_VER_SMALL = 123;
 		FGuid VersionGuid = FDevSystemGuids::GetSystemGuid(FDevSystemGuids::Get().POSESEARCHDB_DERIVEDDATA_VER);
 
 		*this << VersionGuid;
diff --git a/Engine/Plugins/Experimental/Animation/PoseSearch/Source/Runtime/Public/PoseSearch/PoseSearchAssetSampler.h b/Engine/Plugins/Experimental/Animation/PoseSearch/Source/Runtime/Public/PoseSearch/PoseSearchAssetSampler.h
--- a/Engine/Plugins/Experimental/Animation/PoseSearch/Source/Runtime/Public/PoseSearch/PoseSearchAssetSampler.h
+++ b/Engine/Plugins/Experimental/Animation/PoseSearch/Source/Runtime/Public/PoseSearch/PoseSearchAssetSampler.h
@@ -3,25 +3,6 @@
 #pragma once
 
 #include "BonePose.h"
-#include "PoseSearchAssetSampler.generated.h"
-
-USTRUCT()
-struct POSESEARCH_API FPoseSearchExtrapolationParameters
-{
-	GENERATED_BODY()
-
-	// If the angular root motion speed in degrees is below this value, it will be treated as zero.
-	UPROPERTY(EditAnywhere, Category = "Settings")
-	float AngularSpeedThreshold = 1.0f;
-
-	// If the root motion linear speed is below this value, it will be treated as zero.
-	UPROPERTY(EditAnywhere, Category = "Settings")
-	float LinearSpeedThreshold = 1.0f;
-
-	// Time from sequence start/end used to extrapolate the trajectory.
-	UPROPERTY(EditAnywhere, Category = "Settings")
-	float SampleTime = 0.05f;
-};
 
 struct FAnimationPoseData;
 struct FAnimExtractContext;
@@ -62,6 +43,9 @@ struct POSESEARCH_API FAssetSamplerBase : public TSharedFromThis<FAssetSamplerBa
 	virtual const UAnimationAsset* GetAsset() const { return nullptr; }
 
 	virtual void Process() { }
+
+	const float ExtrapolationSampleTime = 1.f / 30.f;
+	const float ExtractionInterval = 1.0f / 120.0f;
 };
 
 // Sampler working with UAnimSequenceBase so it can be used for UAnimSequence as well as UAnimComposite.
@@ -70,7 +54,6 @@ struct POSESEARCH_API FSequenceBaseSampler : public FAssetSamplerBase
 	struct FInput
 	{
 		TWeakObjectPtr<const UAnimSequenceBase> SequenceBase;
-		FPoseSearchExtrapolationParameters ExtrapolationParameters;
 	} Input;
 
 	void Init(const FInput& Input);
@@ -95,7 +78,6 @@ struct POSESEARCH_API FBlendSpaceSampler : public FAssetSamplerBase
 		FBoneContainer BoneContainer;
 		TWeakObjectPtr<const UBlendSpace> BlendSpace;
 		int32 RootTransformSamplingRate = 30;
-		FPoseSearchExtrapolationParameters ExtrapolationParameters;
 		FVector BlendParameters;
 	} Input;
 
@@ -133,7 +115,6 @@ struct POSESEARCH_API FAnimMontageSampler : public FAssetSamplerBase
 	{
 		// @todo: add support for SlotName / multiple SlotAnimTracks
 		TWeakObjectPtr<const UAnimMontage> AnimMontage;
-		FPoseSearchExtrapolationParameters ExtrapolationParameters;
 	} Input;
 
 	void Init(const FInput& Input);
diff --git a/Engine/Plugins/Experimental/Animation/PoseSearch/Source/Runtime/Public/PoseSearch/PoseSearchDatabase.h b/Engine/Plugins/Experimental/Animation/PoseSearch/Source/Runtime/Public/PoseSearch/PoseSearchDatabase.h
--- a/Engine/Plugins/Experimental/Animation/PoseSearch/Source/Runtime/Public/PoseSearch/PoseSearchDatabase.h
+++ b/Engine/Plugins/Experimental/Animation/PoseSearch/Source/Runtime/Public/PoseSearch/PoseSearchDatabase.h
@@ -216,9 +216,6 @@ public:
 	UPROPERTY(EditAnywhere, BlueprintReadOnly, Category="Database", DisplayName="Config")
 	TObjectPtr<const UPoseSearchSchema> Schema;
 
-	UPROPERTY(EditAnywhere, Category = "Database")
-	FPoseSearchExtrapolationParameters ExtrapolationParameters;
-
 	UPROPERTY(EditAnywhere, Category = "Database")
 	FPoseSearchExcludeFromDatabaseParameters ExcludeFromDatabaseParameters;
 
