motion matching - removing db editor "ShowAnimationNone"

diff --git a/Engine/Plugins/Experimental/Animation/PoseSearch/Source/Editor/Private/PoseSearchDatabaseEditorCommands.cpp b/Engine/Plugins/Experimental/Animation/PoseSearch/Source/Editor/Private/PoseSearchDatabaseEditorCommands.cpp
--- a/Engine/Plugins/Experimental/Animation/PoseSearch/Source/Editor/Private/PoseSearchDatabaseEditorCommands.cpp
+++ b/Engine/Plugins/Experimental/Animation/PoseSearch/Source/Editor/Private/PoseSearchDatabaseEditorCommands.cpp
@@ -20,9 +20,6 @@ namespace UE::PoseSearch
 			ShowPoseFeaturesDetailed, "Detailed Features", "Show all features with mode details",
 			EUserInterfaceActionType::RadioButton, FInputChord());
 
-		UI_COMMAND(
-			ShowAnimationNone, "None", "Don't show animations",
-			EUserInterfaceActionType::RadioButton, FInputChord());
 		UI_COMMAND(
 			ShowAnimationOriginalOnly, "Original Only", "Show only original animations",
 			EUserInterfaceActionType::RadioButton, FInputChord());
diff --git a/Engine/Plugins/Experimental/Animation/PoseSearch/Source/Editor/Private/PoseSearchDatabaseViewModel.cpp b/Engine/Plugins/Experimental/Animation/PoseSearch/Source/Editor/Private/PoseSearchDatabaseViewModel.cpp
--- a/Engine/Plugins/Experimental/Animation/PoseSearch/Source/Editor/Private/PoseSearchDatabaseViewModel.cpp
+++ b/Engine/Plugins/Experimental/Animation/PoseSearch/Source/Editor/Private/PoseSearchDatabaseViewModel.cpp
@@ -103,11 +103,6 @@ namespace UE::PoseSearch
 
 		MaxPreviewPlayLength = 0.0f;
 
-		if (AnimationPreviewMode == EAnimationPreviewMode::None && SelectedNodes.IsEmpty())
-		{
-			return;
-		}
-
 		if (FAsyncPoseSearchDatabasesManagement::RequestAsyncBuildIndex(PoseSearchDatabase, ERequestAsyncBuildFlag::ContinueRequest))
 		{
 			FBoneContainer BoneContainer;
diff --git a/Engine/Plugins/Experimental/Animation/PoseSearch/Source/Editor/Private/SPoseSearchDatabaseViewport.cpp b/Engine/Plugins/Experimental/Animation/PoseSearch/Source/Editor/Private/SPoseSearchDatabaseViewport.cpp
--- a/Engine/Plugins/Experimental/Animation/PoseSearch/Source/Editor/Private/SPoseSearchDatabaseViewport.cpp
+++ b/Engine/Plugins/Experimental/Animation/PoseSearch/Source/Editor/Private/SPoseSearchDatabaseViewport.cpp
@@ -73,18 +73,6 @@ namespace UE::PoseSearch
 				&FDatabaseViewModel::IsPoseFeaturesDrawMode,
 				EFeaturesDrawMode::Detailed));
 
-		CommandList->MapAction(
-			Commands.ShowAnimationNone,
-			FExecuteAction::CreateSP(
-				ViewModelRef,
-				&FDatabaseViewModel::OnSetAnimationPreviewMode,
-				EAnimationPreviewMode::None),
-			FCanExecuteAction(),
-			FIsActionChecked::CreateSP(
-				ViewModelRef,
-				&FDatabaseViewModel::IsAnimationPreviewMode,
-				EAnimationPreviewMode::None));
-
 		CommandList->MapAction(
 			Commands.ShowAnimationOriginalOnly,
 			FExecuteAction::CreateSP(
diff --git a/Engine/Plugins/Experimental/Animation/PoseSearch/Source/Editor/Private/SPoseSearchDatabaseViewportToolbar.cpp b/Engine/Plugins/Experimental/Animation/PoseSearch/Source/Editor/Private/SPoseSearchDatabaseViewportToolbar.cpp
--- a/Engine/Plugins/Experimental/Animation/PoseSearch/Source/Editor/Private/SPoseSearchDatabaseViewportToolbar.cpp
+++ b/Engine/Plugins/Experimental/Animation/PoseSearch/Source/Editor/Private/SPoseSearchDatabaseViewportToolbar.cpp
@@ -55,7 +55,6 @@ namespace UE::PoseSearch
 
 				SubMenuBuilder.BeginSection("Animations", LOCTEXT("ShowMenu_AnimationsLabel", "Animations"));
 				{
-					SubMenuBuilder.AddMenuEntry(Commands.ShowAnimationNone);
 					SubMenuBuilder.AddMenuEntry(Commands.ShowAnimationOriginalOnly);
 					SubMenuBuilder.AddMenuEntry(Commands.ShowAnimationOriginalAndMirrored);
 				}
diff --git a/Engine/Plugins/Experimental/Animation/PoseSearch/Source/Editor/Public/PoseSearchDatabaseEditorCommands.h b/Engine/Plugins/Experimental/Animation/PoseSearch/Source/Editor/Public/PoseSearchDatabaseEditorCommands.h
--- a/Engine/Plugins/Experimental/Animation/PoseSearch/Source/Editor/Public/PoseSearchDatabaseEditorCommands.h
+++ b/Engine/Plugins/Experimental/Animation/PoseSearch/Source/Editor/Public/PoseSearchDatabaseEditorCommands.h
@@ -30,7 +30,6 @@ namespace UE::PoseSearch
 		TSharedPtr<FUICommandInfo> ShowPoseFeaturesAll;
 		TSharedPtr<FUICommandInfo> ShowPoseFeaturesDetailed;
 
-		TSharedPtr<FUICommandInfo> ShowAnimationNone;
 		TSharedPtr<FUICommandInfo> ShowAnimationOriginalOnly;
 		TSharedPtr<FUICommandInfo> ShowAnimationOriginalAndMirrored;
 		
diff --git a/Engine/Plugins/Experimental/Animation/PoseSearch/Source/Editor/Public/PoseSearchDatabaseViewModel.h b/Engine/Plugins/Experimental/Animation/PoseSearch/Source/Editor/Public/PoseSearchDatabaseViewModel.h
--- a/Engine/Plugins/Experimental/Animation/PoseSearch/Source/Editor/Public/PoseSearchDatabaseViewModel.h
+++ b/Engine/Plugins/Experimental/Animation/PoseSearch/Source/Editor/Public/PoseSearchDatabaseViewModel.h
@@ -34,7 +34,6 @@ namespace UE::PoseSearch
 
 	enum class EAnimationPreviewMode : uint8
 	{
-		None,
 		OriginalOnly,
 		OriginalAndMirrored
 	};
