motion matching
- db looping fix on clipped assets
- db opening aniamtion assets at the playback time on double click

diff --git a/Engine/Plugins/Experimental/Animation/PoseSearch/Source/Editor/Private/PoseSearchDatabaseAssetListItem.cpp b/Engine/Plugins/Experimental/Animation/PoseSearch/Source/Editor/Private/PoseSearchDatabaseAssetListItem.cpp
--- a/Engine/Plugins/Experimental/Animation/PoseSearch/Source/Editor/Private/PoseSearchDatabaseAssetListItem.cpp
+++ b/Engine/Plugins/Experimental/Animation/PoseSearch/Source/Editor/Private/PoseSearchDatabaseAssetListItem.cpp
@@ -1,33 +1,34 @@
 // Copyright Epic Games, Inc. All Rights Reserved.
 
 #include "PoseSearchDatabaseAssetListItem.h"
-#include "PoseSearchDatabaseViewModel.h"
 
-#include "Animation/AnimSequence.h"
 #include "Animation/AnimComposite.h"
+#include "Animation/AnimSequence.h"
 #include "Animation/BlendSpace.h"
-
-#include "Misc/TransactionObjectEvent.h"
-#include "DragAndDrop/AssetDragDropOp.h"
-#include "Misc/FeedbackContext.h"
+#include "Animation/DebugSkelMeshComponent.h"
+#include "AnimPreviewInstance.h"
 #include "AssetSelection.h"
 #include "ClassIconFinder.h"
 #include "DetailColumnSizeData.h"
-#include "PoseSearch/PoseSearchDatabase.h"
-#include "PoseSearchDatabaseAssetTree.h"
-
-#include "Widgets/Text/SRichTextBlock.h"
-#include "Widgets/Input/SCheckBox.h"
-#include "Widgets/Input/SSearchBox.h"
+#include "DragAndDrop/AssetDragDropOp.h"
+#include "Editor.h"
 #include "Framework/Commands/GenericCommands.h"
 #include "Framework/MultiBox/MultiBoxBuilder.h"
+#include "IAnimationEditor.h"
+#include "IPersonaToolkit.h"
+#include "Misc/FeedbackContext.h"
+#include "Misc/TransactionObjectEvent.h"
+#include "PoseSearch/PoseSearchDatabase.h"
+#include "PoseSearchDatabaseAssetTree.h"
+#include "PoseSearchDatabaseViewModel.h"
+#include "ScopedTransaction.h"
 #include "SPositiveActionButton.h"
 #include "Styling/AppStyle.h"
-
-#include "ScopedTransaction.h"
 #include "Styling/StyleColors.h"
 #include "Subsystems/AssetEditorSubsystem.h"
-#include "Editor.h"
+#include "Widgets/Input/SCheckBox.h"
+#include "Widgets/Input/SSearchBox.h"
+#include "Widgets/Text/SRichTextBlock.h"
 
 #define LOCTEXT_NAMESPACE "SDatabaseAssetListItem"
 
@@ -122,7 +123,34 @@ namespace UE::PoseSearch
 			{
 				if (UAssetEditorSubsystem* AssetEditorSS = GEditor->GetEditorSubsystem<UAssetEditorSubsystem>())
 				{
-					AssetEditorSS->OpenEditorForAsset(DatabaseAnimationAsset->GetAnimationAsset());
+					UAnimationAsset* AnimationAsset = DatabaseAnimationAsset->GetAnimationAsset();
+
+					AssetEditorSS->OpenEditorForAsset(AnimationAsset);
+
+					if (IAssetEditorInstance* Editor = AssetEditorSS->FindEditorForAsset(AnimationAsset, true))
+					{
+						if (Editor->GetEditorName() == "AnimationEditor")
+						{
+							float AnimationAssetTime = 0.f;
+							FVector AnimationAssetBlendParameters = FVector::ZeroVector;
+							TArray<const FPoseSearchIndexAsset*> SearchIndexAssets;
+							if (Database->GetPoseSearchIndexAssets(Node->SourceAssetIdx, SearchIndexAssets) && !SearchIndexAssets.IsEmpty())
+							{
+								const FPoseSearchIndexAsset* SearchIndexAsset = SearchIndexAssets[0];
+								check(SearchIndexAsset);
+								AnimationAssetTime = ViewModel->GetPlayTime() + SearchIndexAsset->SamplingInterval.Min;
+								AnimationAssetBlendParameters = SearchIndexAsset->BlendParameters;
+							}
+
+							const IAnimationEditor* AnimationEditor = static_cast<IAnimationEditor*>(Editor);
+							const UDebugSkelMeshComponent* PreviewComponent = AnimationEditor->GetPersonaToolkit()->GetPreviewMeshComponent();
+
+							// Open asset paused and at specific time as seen on the pose search debugger.
+							PreviewComponent->PreviewInstance->SetPosition(AnimationAssetTime);
+							PreviewComponent->PreviewInstance->SetPlaying(false);
+							PreviewComponent->PreviewInstance->SetBlendSpacePosition(AnimationAssetBlendParameters);
+						}
+					}
 				}
 			}
 		}
diff --git a/Engine/Plugins/Experimental/Animation/PoseSearch/Source/Editor/Private/PoseSearchDatabaseViewModel.cpp b/Engine/Plugins/Experimental/Animation/PoseSearch/Source/Editor/Private/PoseSearchDatabaseViewModel.cpp
--- a/Engine/Plugins/Experimental/Animation/PoseSearch/Source/Editor/Private/PoseSearchDatabaseViewModel.cpp
+++ b/Engine/Plugins/Experimental/Animation/PoseSearch/Source/Editor/Private/PoseSearchDatabaseViewModel.cpp
@@ -124,10 +124,7 @@ namespace UE::PoseSearch
 					FDatabasePreviewActor PreviewActor = SpawnPreviewActor(IndexAssetIndex, BoneContainer);
 					if (PreviewActor.IsValid())
 					{
-						if (const UAnimationAsset* AnimationAsset = PreviewActor.GetAnimPreviewInstance()->GetAnimationAsset())
-						{
-							MaxPreviewPlayLength = FMath::Max(MaxPreviewPlayLength, AnimationAsset->GetPlayLength());
-						}
+						MaxPreviewPlayLength = FMath::Max(MaxPreviewPlayLength, IndexAsset.SamplingInterval.Max - IndexAsset.SamplingInterval.Min);
 						PreviewActors.Add(PreviewActor);
 					}
 				}
@@ -329,10 +326,11 @@ namespace UE::PoseSearch
 			}
 
 			float CurrentTime = 0.f;
-			FAnimationRuntime::AdvanceTime(false, PlayTime, CurrentTime, PreviewAsset->GetPlayLength());
+			const FPoseSearchIndexAsset& IndexAsset = SearchIndex.Assets[PreviewActor.IndexAssetIndex];
+			float CurrentPlayTime = PlayTime + IndexAsset.SamplingInterval.Min;
+			FAnimationRuntime::AdvanceTime(false, CurrentPlayTime, CurrentTime, IndexAsset.SamplingInterval.Max);
 
 			// time to pose index
-			const FPoseSearchIndexAsset& IndexAsset = SearchIndex.Assets[PreviewActor.IndexAssetIndex];
 			PreviewActor.CurrentPoseIndex = PoseSearchDatabase->GetPoseIndexFromTime(CurrentTime, IndexAsset);
 
 			// pose index to quantized time
diff --git a/Engine/Plugins/Experimental/Animation/PoseSearch/Source/Runtime/Private/PoseSearchDatabase.cpp b/Engine/Plugins/Experimental/Animation/PoseSearch/Source/Runtime/Private/PoseSearchDatabase.cpp
--- a/Engine/Plugins/Experimental/Animation/PoseSearch/Source/Runtime/Private/PoseSearchDatabase.cpp
+++ b/Engine/Plugins/Experimental/Animation/PoseSearch/Source/Runtime/Private/PoseSearchDatabase.cpp
@@ -243,8 +243,10 @@ UClass* FPoseSearchDatabaseSequence::GetAnimationAssetStaticClass() const
 
 bool FPoseSearchDatabaseSequence::IsLooping() const
 {
-	// @todo: we should check if against the sample range (look for GetEffectiveSamplingRange()...)
-	return Sequence ? Sequence->bLoop : false;
+	return Sequence &&
+		Sequence->bLoop &&
+		SamplingRange.Min == 0.f &&
+		SamplingRange.Max == 0.f;
 }
 
 const FString FPoseSearchDatabaseSequence::GetName() const
@@ -271,7 +273,7 @@ UClass* FPoseSearchDatabaseBlendSpace::GetAnimationAssetStaticClass() const
 
 bool FPoseSearchDatabaseBlendSpace::IsLooping() const
 {
-	return BlendSpace ? BlendSpace->bLoop : false;
+	return BlendSpace && BlendSpace->bLoop;
 }
 
 const FString FPoseSearchDatabaseBlendSpace::GetName() const
@@ -361,7 +363,10 @@ UClass* FPoseSearchDatabaseAnimComposite::GetAnimationAssetStaticClass() const
 
 bool FPoseSearchDatabaseAnimComposite::IsLooping() const
 {
-	return AnimComposite ? AnimComposite->bLoop : false;
+	return AnimComposite &&
+		AnimComposite->bLoop &&
+		SamplingRange.Min == 0.f &&
+		SamplingRange.Max == 0.f;
 }
 
 const FString FPoseSearchDatabaseAnimComposite::GetName() const
@@ -388,7 +393,10 @@ UClass* FPoseSearchDatabaseAnimMontage::GetAnimationAssetStaticClass() const
 
 bool FPoseSearchDatabaseAnimMontage::IsLooping() const
 {
-	return AnimMontage ? AnimMontage->bLoop : false;
+	return AnimMontage &&
+		AnimMontage->bLoop &&
+		SamplingRange.Min == 0.f &&
+		SamplingRange.Max == 0.f;
 }
 
 const FString FPoseSearchDatabaseAnimMontage::GetName() const
@@ -626,6 +634,26 @@ bool UPoseSearchDatabase::IsCachedCookedPlatformDataLoaded(const ITargetPlatform
 	check(IsInGameThread());
 	return FAsyncPoseSearchDatabasesManagement::RequestAsyncBuildIndex(this, ERequestAsyncBuildFlag::ContinueRequest);
 }
+
+// returns all the FPoseSearchIndexAsset associated with the AnimationAssetIndex
+bool UPoseSearchDatabase::GetPoseSearchIndexAssets(int32 AnimationAssetIndex, TArray<const FPoseSearchIndexAsset*>& OutSearchIndexAssets) const
+{
+	using namespace UE::PoseSearch;
+	OutSearchIndexAssets.Reset();
+	if (FAsyncPoseSearchDatabasesManagement::RequestAsyncBuildIndex(this, ERequestAsyncBuildFlag::ContinueRequest))
+	{
+		const FPoseSearchIndex& SearchIndex = GetSearchIndex();
+		for (const FPoseSearchIndexAsset& SearchIndexAsset : SearchIndex.Assets)
+		{
+			if (SearchIndexAsset.SourceAssetIdx == AnimationAssetIndex)
+			{
+				OutSearchIndexAssets.Add(&SearchIndexAsset);
+			}
+		}
+		return true;
+	}
+	return false;
+}
 #endif // WITH_EDITOR
 
 void UPoseSearchDatabase::PostSaveRoot(FObjectPostSaveRootContext ObjectSaveContext)
diff --git a/Engine/Plugins/Experimental/Animation/PoseSearch/Source/Runtime/Public/PoseSearch/PoseSearchDatabase.h b/Engine/Plugins/Experimental/Animation/PoseSearch/Source/Runtime/Public/PoseSearch/PoseSearchDatabase.h
--- a/Engine/Plugins/Experimental/Animation/PoseSearch/Source/Runtime/Public/PoseSearch/PoseSearchDatabase.h
+++ b/Engine/Plugins/Experimental/Animation/PoseSearch/Source/Runtime/Public/PoseSearch/PoseSearchDatabase.h
@@ -60,7 +60,7 @@ struct POSESEARCH_API FPoseSearchDatabaseAnimationAssetBase
 	virtual UAnimationAsset* GetAnimationAsset() const { return nullptr; }
 	virtual UClass* GetAnimationAssetStaticClass() const { return nullptr; }
 	virtual bool IsLooping() const { return false; }
-	virtual const FString GetName() const { return {}; }
+	virtual const FString GetName() const { return FString(); }
 	virtual bool IsEnabled() const { return false; }
 	virtual void SetIsEnabled(bool bInIsEnabled) {}
 	virtual bool IsRootMotionEnabled() const { return false; }
@@ -81,7 +81,7 @@ struct POSESEARCH_API FPoseSearchDatabaseSequence : public FPoseSearchDatabaseAn
 	bool bEnabled = true;
 
 	UPROPERTY(EditAnywhere, Category="Sequence", meta = (DisplayPriority = 1))
-	FFloatInterval SamplingRange = FFloatInterval(0.0f, 0.0f);
+	FFloatInterval SamplingRange = FFloatInterval(0.f, 0.f);
 
 	UPROPERTY(EditAnywhere, Category = "Sequence", meta = (DisplayPriority = 2))
 	EPoseSearchMirrorOption MirrorOption = EPoseSearchMirrorOption::UnmirroredOnly;
@@ -150,7 +150,7 @@ struct POSESEARCH_API FPoseSearchDatabaseAnimComposite : public FPoseSearchDatab
 	bool bEnabled = true;
 
 	UPROPERTY(EditAnywhere, Category = "AnimComposite", meta = (DisplayPriority = 1))
-	FFloatInterval SamplingRange = FFloatInterval(0.0f, 0.0f);
+	FFloatInterval SamplingRange = FFloatInterval(0.f, 0.f);
 
 	UPROPERTY(EditAnywhere, Category = "AnimComposite", meta = (DisplayPriority = 2))
 	EPoseSearchMirrorOption MirrorOption = EPoseSearchMirrorOption::UnmirroredOnly;
@@ -179,7 +179,7 @@ struct POSESEARCH_API FPoseSearchDatabaseAnimMontage : public FPoseSearchDatabas
 	bool bEnabled = true;
 
 	UPROPERTY(EditAnywhere, Category="AnimMontage", meta = (DisplayPriority = 1))
-	FFloatInterval SamplingRange = FFloatInterval(0.0f, 0.0f);
+	FFloatInterval SamplingRange = FFloatInterval(0.f, 0.f);
 
 	UPROPERTY(EditAnywhere, Category = "AnimMontage", meta = (DisplayPriority = 2))
 	EPoseSearchMirrorOption MirrorOption = EPoseSearchMirrorOption::UnmirroredOnly;
@@ -292,6 +292,9 @@ public:
 	virtual void BeginCacheForCookedPlatformData(const ITargetPlatform* TargetPlatform) override;
 	virtual bool IsCachedCookedPlatformDataLoaded(const ITargetPlatform* TargetPlatform) override;
 
+	// returns all the FPoseSearchIndexAsset associated with the AnimationAssetIndex
+	bool GetPoseSearchIndexAssets(int32 AnimationAssetIndex, TArray<const FPoseSearchIndexAsset*>& OutSearchIndexAssets) const;
+
 private:
 	DECLARE_MULTICAST_DELEGATE(FOnDerivedDataRebuildMulticaster);
 	FOnDerivedDataRebuildMulticaster OnDerivedDataRebuild;
