motion matching - channels and schema property DisplayPriority tuning

diff --git a/Engine/Plugins/Experimental/Animation/PoseSearch/Source/Runtime/Private/PoseSearchFeatureChannel_Heading.h b/Engine/Plugins/Experimental/Animation/PoseSearch/Source/Runtime/Private/PoseSearchFeatureChannel_Heading.h
--- a/Engine/Plugins/Experimental/Animation/PoseSearch/Source/Runtime/Private/PoseSearchFeatureChannel_Heading.h
+++ b/Engine/Plugins/Experimental/Animation/PoseSearch/Source/Runtime/Private/PoseSearchFeatureChannel_Heading.h
@@ -38,7 +38,7 @@ public:
 	UPROPERTY(Transient)
 	int8 SchemaBoneIdx = 0;
 
-	UPROPERTY(EditAnywhere, Category = "Settings", meta = (ExcludeFromHash))
+	UPROPERTY(EditAnywhere, Category = "Settings", meta = (ExcludeFromHash, DisplayPriority = 0))
 	FLinearColor DebugColor = FLinearColor::White;
 
 	UPROPERTY(EditAnywhere, Category = "Settings")
diff --git a/Engine/Plugins/Experimental/Animation/PoseSearch/Source/Runtime/Private/PoseSearchFeatureChannel_Phase.h b/Engine/Plugins/Experimental/Animation/PoseSearch/Source/Runtime/Private/PoseSearchFeatureChannel_Phase.h
--- a/Engine/Plugins/Experimental/Animation/PoseSearch/Source/Runtime/Private/PoseSearchFeatureChannel_Phase.h
+++ b/Engine/Plugins/Experimental/Animation/PoseSearch/Source/Runtime/Private/PoseSearchFeatureChannel_Phase.h
@@ -21,7 +21,7 @@ public:
 	UPROPERTY(Transient)
 	int8 SchemaBoneIdx = 0;
 
-	UPROPERTY(EditAnywhere, Category = "Settings", meta = (ExcludeFromHash))
+	UPROPERTY(EditAnywhere, Category = "Settings", meta = (ExcludeFromHash, DisplayPriority = 0))
 	FLinearColor DebugColor = FLinearColor::Yellow;
 
 	UPROPERTY(EditAnywhere, Category = "Settings")
diff --git a/Engine/Plugins/Experimental/Animation/PoseSearch/Source/Runtime/Private/PoseSearchFeatureChannel_Pose.h b/Engine/Plugins/Experimental/Animation/PoseSearch/Source/Runtime/Private/PoseSearchFeatureChannel_Pose.h
--- a/Engine/Plugins/Experimental/Animation/PoseSearch/Source/Runtime/Private/PoseSearchFeatureChannel_Pose.h
+++ b/Engine/Plugins/Experimental/Animation/PoseSearch/Source/Runtime/Private/PoseSearchFeatureChannel_Pose.h
@@ -31,7 +31,7 @@ struct POSESEARCH_API FPoseSearchBone
 	UPROPERTY(EditAnywhere, Category = Config)
 	float Weight = 1.f;
 
-	UPROPERTY(EditAnywhere, Category = Config, meta=(ExcludeFromHash))
+	UPROPERTY(EditAnywhere, Category = Config, meta=(ExcludeFromHash, DisplayPriority = 0))
 	FLinearColor DebugColor = FLinearColor::Green;
 };
 
diff --git a/Engine/Plugins/Experimental/Animation/PoseSearch/Source/Runtime/Private/PoseSearchFeatureChannel_Position.h b/Engine/Plugins/Experimental/Animation/PoseSearch/Source/Runtime/Private/PoseSearchFeatureChannel_Position.h
--- a/Engine/Plugins/Experimental/Animation/PoseSearch/Source/Runtime/Private/PoseSearchFeatureChannel_Position.h
+++ b/Engine/Plugins/Experimental/Animation/PoseSearch/Source/Runtime/Private/PoseSearchFeatureChannel_Position.h
@@ -30,7 +30,7 @@ public:
 	UPROPERTY(Transient)
 	int8 SchemaOriginBoneIdx = 0;
 
-	UPROPERTY(EditAnywhere, Category = "Settings", meta = (ExcludeFromHash))
+	UPROPERTY(EditAnywhere, Category = "Settings", meta = (ExcludeFromHash, DisplayPriority = 0))
 	FLinearColor DebugColor = FLinearColor::Blue;
 
 	UPROPERTY(EditAnywhere, Category = "Settings")
diff --git a/Engine/Plugins/Experimental/Animation/PoseSearch/Source/Runtime/Private/PoseSearchFeatureChannel_Trajectory.h b/Engine/Plugins/Experimental/Animation/PoseSearch/Source/Runtime/Private/PoseSearchFeatureChannel_Trajectory.h
--- a/Engine/Plugins/Experimental/Animation/PoseSearch/Source/Runtime/Private/PoseSearchFeatureChannel_Trajectory.h
+++ b/Engine/Plugins/Experimental/Animation/PoseSearch/Source/Runtime/Private/PoseSearchFeatureChannel_Trajectory.h
@@ -38,7 +38,7 @@ struct POSESEARCH_API FPoseSearchTrajectorySample
 	UPROPERTY(EditAnywhere, Category = Config)
 	float Weight = 1.f;
 
-	UPROPERTY(EditAnywhere, Category = Config, meta = (ExcludeFromHash))
+	UPROPERTY(EditAnywhere, Category = Config, meta = (ExcludeFromHash, DisplayPriority = 0))
 	FLinearColor DebugColor = FLinearColor::Blue;
 };
 UCLASS(BlueprintType, EditInlineNew, meta = (DisplayName = "Trajectory Channel"), CollapseCategories)
diff --git a/Engine/Plugins/Experimental/Animation/PoseSearch/Source/Runtime/Private/PoseSearchFeatureChannel_Velocity.h b/Engine/Plugins/Experimental/Animation/PoseSearch/Source/Runtime/Private/PoseSearchFeatureChannel_Velocity.h
--- a/Engine/Plugins/Experimental/Animation/PoseSearch/Source/Runtime/Private/PoseSearchFeatureChannel_Velocity.h
+++ b/Engine/Plugins/Experimental/Animation/PoseSearch/Source/Runtime/Private/PoseSearchFeatureChannel_Velocity.h
@@ -24,7 +24,7 @@ public:
 	UPROPERTY(Transient)
 	int8 SchemaBoneIdx = 0;
 
-	UPROPERTY(EditAnywhere, Category = "Settings", meta = (ExcludeFromHash))
+	UPROPERTY(EditAnywhere, Category = "Settings", meta = (ExcludeFromHash, DisplayPriority = 0))
 	FLinearColor DebugColor = FLinearColor::Green;
 
 	UPROPERTY(EditAnywhere, Category = "Settings")
diff --git a/Engine/Plugins/Experimental/Animation/PoseSearch/Source/Runtime/Public/PoseSearch/PoseSearchSchema.h b/Engine/Plugins/Experimental/Animation/PoseSearch/Source/Runtime/Public/PoseSearch/PoseSearchSchema.h
--- a/Engine/Plugins/Experimental/Animation/PoseSearch/Source/Runtime/Public/PoseSearch/PoseSearchSchema.h
+++ b/Engine/Plugins/Experimental/Animation/PoseSearch/Source/Runtime/Public/PoseSearch/PoseSearchSchema.h
@@ -32,10 +32,10 @@ class POSESEARCH_API UPoseSearchSchema : public UDataAsset, public IBoneReferenc
 
 public:
 	// @todo: used only for indexing: cache it somewhere else
-	UPROPERTY(EditAnywhere, Category = "Schema")
+	UPROPERTY(EditAnywhere, Category = "Schema", meta = (DisplayPriority = 0))
 	TObjectPtr<USkeleton> Skeleton;
 
-	UPROPERTY(EditAnywhere, meta = (ClampMin = "1", ClampMax = "240"), Category = "Schema")
+	UPROPERTY(EditAnywhere, Category = "Schema", meta = (DisplayPriority = 3, ClampMin = "1", ClampMax = "240"))
 	int32 SampleRate = 30;
 
 	UPROPERTY(EditAnywhere, Instanced, BlueprintReadWrite, Category = "Schema")
@@ -46,10 +46,10 @@ public:
 	TArray<TObjectPtr<UPoseSearchFeatureChannel>> FinalizedChannels;
 
 	// If set, this schema will support mirroring pose search databases
-	UPROPERTY(EditAnywhere, Category = "Schema")
+	UPROPERTY(EditAnywhere, Category = "Schema", meta = (DisplayPriority = 1))
 	TObjectPtr<UMirrorDataTable> MirrorDataTable;
 
-	UPROPERTY(EditAnywhere, Category = "Schema")
+	UPROPERTY(EditAnywhere, Category = "Schema", meta = (DisplayPriority = 2))
 	EPoseSearchDataPreprocessor DataPreprocessor = EPoseSearchDataPreprocessor::Normalize;
 
 	UPROPERTY(Transient)
@@ -80,11 +80,11 @@ public:
 	float LoopingCostBias = 0.f;
 
 	// how many times the animation assets of the database using this schema will be indexed
-	UPROPERTY(EditAnywhere, meta = (ClampMin = "1"), Category = "Permutations")
+	UPROPERTY(EditAnywhere, Category = "Permutations", meta = (ClampMin = "1"))
 	int32 NumberOfPermutations = 1;
 
 	// delta time between every permutation indexing
-	UPROPERTY(EditAnywhere, meta = (ClampMin = "1", ClampMax = "240"), Category = "Permutations")
+	UPROPERTY(EditAnywhere, Category = "Permutations", meta = (ClampMin = "1", ClampMax = "240", EditCondition = "NumberOfPermutations > 1", EditConditionHides))
 	int32 PermutationsSampleRate = 30;
 
 	// starting offset of the "PermutationTime" from the "SamplingTime" of the first permutation.
