motion matching - adding UPoseSearchNormalizationSet

diff --git a/Engine/Plugins/Experimental/Animation/PoseSearch/Source/Editor/Private/PoseSearchEditor.cpp b/Engine/Plugins/Experimental/Animation/PoseSearch/Source/Editor/Private/PoseSearchEditor.cpp
--- a/Engine/Plugins/Experimental/Animation/PoseSearch/Source/Editor/Private/PoseSearchEditor.cpp
+++ b/Engine/Plugins/Experimental/Animation/PoseSearch/Source/Editor/Private/PoseSearchEditor.cpp
@@ -56,6 +56,7 @@ private:
 
 	TSharedPtr<IAssetTypeActions> PoseSearchDatabaseActions;
 	TSharedPtr<IAssetTypeActions> PoseSearchSchemaActions;
+	TSharedPtr<IAssetTypeActions> PoseSearchNormalizationSetActions;
 };
 
 void FEditorModule::StartupModule()
@@ -83,6 +84,10 @@ void FEditorModule::StartupModule()
 		// Register UPoseSearchSchema Type Actions 
 		PoseSearchSchemaActions = MakeShared<FSchemaTypeActions>();
 		AssetTools.RegisterAssetTypeActions(PoseSearchSchemaActions.ToSharedRef());
+
+		// Register UPoseSearchSchema Type Actions 
+		PoseSearchNormalizationSetActions = MakeShared<FNormalizationSetTypeActions>();
+		AssetTools.RegisterAssetTypeActions(PoseSearchNormalizationSetActions.ToSharedRef());
 	}
 
 	RegisterPropertyTypeCustomizations();
@@ -105,6 +110,7 @@ void FEditorModule::ShutdownModule()
 		IAssetTools& AssetTools = FModuleManager::GetModuleChecked<FAssetToolsModule>("AssetTools").Get();
 		AssetTools.UnregisterAssetTypeActions(PoseSearchDatabaseActions.ToSharedRef());
 		AssetTools.UnregisterAssetTypeActions(PoseSearchSchemaActions.ToSharedRef());
+		AssetTools.UnregisterAssetTypeActions(PoseSearchNormalizationSetActions.ToSharedRef());
 	}
 
 	// Unregister Ed Mode
diff --git a/Engine/Plugins/Experimental/Animation/PoseSearch/Source/Editor/Private/PoseSearchNormalizationSetFactory.cpp b/Engine/Plugins/Experimental/Animation/PoseSearch/Source/Editor/Private/PoseSearchNormalizationSetFactory.cpp
new file mode 100644
--- /dev/null
+++ b/Engine/Plugins/Experimental/Animation/PoseSearch/Source/Editor/Private/PoseSearchNormalizationSetFactory.cpp
@@ -0,0 +1,26 @@
+// Copyright Epic Games, Inc. All Rights Reserved.
+
+#include "PoseSearchNormalizationSetFactory.h"
+#include "PoseSearch/PoseSearchNormalizationSet.h"
+
+#define LOCTEXT_NAMESPACE "PoseSearchNormalizationSetFactory"
+
+UPoseSearchNormalizationSetFactory::UPoseSearchNormalizationSetFactory(const FObjectInitializer& ObjectInitializer)
+	: Super(ObjectInitializer)
+{
+	bCreateNew = true;
+	bEditAfterNew = true;
+	SupportedClass = UPoseSearchNormalizationSet::StaticClass();
+}
+
+UObject* UPoseSearchNormalizationSetFactory::FactoryCreateNew(UClass* Class, UObject* InParent, FName Name, EObjectFlags Flags, UObject* Context, FFeedbackContext* Warn)
+{
+	return NewObject<UPoseSearchNormalizationSet>(InParent, Class, Name, Flags);
+}
+
+FString UPoseSearchNormalizationSetFactory::GetDefaultNewAssetName() const
+{
+	return FString(TEXT("NewPoseSearchNormalizationSet"));
+}
+
+#undef LOCTEXT_NAMESPACE
\ No newline at end of file
diff --git a/Engine/Plugins/Experimental/Animation/PoseSearch/Source/Editor/Private/PoseSearchTypeActions.cpp b/Engine/Plugins/Experimental/Animation/PoseSearch/Source/Editor/Private/PoseSearchTypeActions.cpp
--- a/Engine/Plugins/Experimental/Animation/PoseSearch/Source/Editor/Private/PoseSearchTypeActions.cpp
+++ b/Engine/Plugins/Experimental/Animation/PoseSearch/Source/Editor/Private/PoseSearchTypeActions.cpp
@@ -3,6 +3,7 @@
 #include "PoseSearchTypeActions.h"
 #include "PoseSearchDatabaseEditor.h"
 #include "PoseSearch/PoseSearchDatabase.h"
+#include "PoseSearch/PoseSearchNormalizationSet.h"
 #include "PoseSearch/PoseSearchSchema.h"
 
 #define LOCTEXT_NAMESPACE "PoseSearchTypeActions"
@@ -92,6 +93,38 @@ namespace UE::PoseSearch
 		};
 		return SubMenus;
 	}
+
+	//////////////////////////////////////////////////////////////////////////
+	// FNormalizationSetTypeActions
+
+	FText FNormalizationSetTypeActions::GetName() const
+	{
+		return LOCTEXT("PoseSearchNormalizationSetTypeActionsName", "Normalization Set");
+	}
+
+	FColor FNormalizationSetTypeActions::GetTypeColor() const
+	{
+		return PoseSearchAssetColor;
+	}
+
+	UClass* FNormalizationSetTypeActions::GetSupportedClass() const
+	{
+		return UPoseSearchNormalizationSet::StaticClass();
+	}
+
+	uint32 FNormalizationSetTypeActions::GetCategories()
+	{
+		return EAssetTypeCategories::Animation;
+	}
+
+	const TArray<FText>& FNormalizationSetTypeActions::GetSubMenus() const
+	{
+		static const TArray<FText> SubMenus
+		{
+			PoseSearchSubMenuName
+		};
+		return SubMenus;
+	}
 }
 
 #undef LOCTEXT_NAMESPACE
diff --git a/Engine/Plugins/Experimental/Animation/PoseSearch/Source/Editor/Public/PoseSearchNormalizationSetFactory.h b/Engine/Plugins/Experimental/Animation/PoseSearch/Source/Editor/Public/PoseSearchNormalizationSetFactory.h
new file mode 100644
--- /dev/null
+++ b/Engine/Plugins/Experimental/Animation/PoseSearch/Source/Editor/Public/PoseSearchNormalizationSetFactory.h
@@ -0,0 +1,19 @@
+// Copyright Epic Games, Inc. All Rights Reserved.
+
+#pragma once
+
+#include "CoreMinimal.h"
+#include "UObject/ObjectMacros.h"
+#include "Factories/Factory.h"
+#include "PoseSearchNormalizationSetFactory.generated.h"
+
+UCLASS()
+class POSESEARCHEDITOR_API UPoseSearchNormalizationSetFactory : public UFactory
+{
+	GENERATED_UCLASS_BODY()
+
+	// UFactory interface
+	virtual UObject* FactoryCreateNew(UClass* Class, UObject* InParent, FName Name, EObjectFlags Flags, UObject* Context, FFeedbackContext* Warn) override;
+	virtual FString GetDefaultNewAssetName() const override;
+	// End of UFactory interface
+};
diff --git a/Engine/Plugins/Experimental/Animation/PoseSearch/Source/Editor/Public/PoseSearchTypeActions.h b/Engine/Plugins/Experimental/Animation/PoseSearch/Source/Editor/Public/PoseSearchTypeActions.h
--- a/Engine/Plugins/Experimental/Animation/PoseSearch/Source/Editor/Public/PoseSearchTypeActions.h
+++ b/Engine/Plugins/Experimental/Animation/PoseSearch/Source/Editor/Public/PoseSearchTypeActions.h
@@ -33,4 +33,16 @@ namespace UE::PoseSearch
 		virtual uint32 GetCategories() override;
 		virtual const TArray<FText>& GetSubMenus() const override;
 	};
+
+	class FNormalizationSetTypeActions : public FAssetTypeActions_Base
+	{
+	public:
+
+		FNormalizationSetTypeActions() {}
+		virtual FText GetName() const override;
+		virtual FColor GetTypeColor() const override;
+		virtual UClass* GetSupportedClass() const override;
+		virtual uint32 GetCategories() override;
+		virtual const TArray<FText>& GetSubMenus() const override;
+	};
 }
diff --git a/Engine/Plugins/Experimental/Animation/PoseSearch/Source/Runtime/Private/PoseSearchDerivedData.cpp b/Engine/Plugins/Experimental/Animation/PoseSearch/Source/Runtime/Private/PoseSearchDerivedData.cpp
--- a/Engine/Plugins/Experimental/Animation/PoseSearch/Source/Runtime/Private/PoseSearchDerivedData.cpp
+++ b/Engine/Plugins/Experimental/Animation/PoseSearch/Source/Runtime/Private/PoseSearchDerivedData.cpp
@@ -20,6 +20,7 @@
 #include "PoseSearch/PoseSearchDefines.h"
 #include "PoseSearch/PoseSearchDerivedDataKey.h"
 #include "PoseSearch/PoseSearchFeatureChannel.h"
+#include "PoseSearch/PoseSearchNormalizationSet.h"
 #include "PoseSearch/PoseSearchSchema.h"
 #include "PoseSearchEigenHelper.h"
 #include "ProfilingDebugging/CookStats.h"
@@ -899,13 +900,7 @@ void FPoseSearchDatabaseAsyncCacheTask::OnGetComplete(UE::DerivedData::FCacheGet
 				IndexBaseDatabases.Add(Database); // the first one is always this Database
 				if (Database->NormalizationSet)
 				{
-					for (auto OtherDatabase : Database->NormalizationSet->Databases)
-					{
-						if (OtherDatabase)
-						{
-							IndexBaseDatabases.AddUnique(OtherDatabase);
-						}
-					}
+					Database->NormalizationSet->AddUniqueDatabases(IndexBaseDatabases);
 				}
 
 				// @todo: DDC or parallelize this code
diff --git a/Engine/Plugins/Experimental/Animation/PoseSearch/Source/Runtime/Private/PoseSearchNormalizationSet.cpp b/Engine/Plugins/Experimental/Animation/PoseSearch/Source/Runtime/Private/PoseSearchNormalizationSet.cpp
new file mode 100644
--- /dev/null
+++ b/Engine/Plugins/Experimental/Animation/PoseSearch/Source/Runtime/Private/PoseSearchNormalizationSet.cpp
@@ -0,0 +1,15 @@
+// Copyright Epic Games, Inc. All Rights Reserved.
+
+#include "PoseSearch/PoseSearchNormalizationSet.h"
+#include "PoseSearch/PoseSearchDatabase.h"
+
+void UPoseSearchNormalizationSet::AddUniqueDatabases(TArray<TWeakObjectPtr<const UPoseSearchDatabase>>& UniqueDatabases) const
+{
+	for (const TObjectPtr<const UPoseSearchDatabase>& Database : Databases)
+	{
+		if (Database)
+		{
+			UniqueDatabases.AddUnique(Database);
+		}
+	}
+}
\ No newline at end of file
diff --git a/Engine/Plugins/Experimental/Animation/PoseSearch/Source/Runtime/Public/PoseSearch/PoseSearchDatabase.h b/Engine/Plugins/Experimental/Animation/PoseSearch/Source/Runtime/Public/PoseSearch/PoseSearchDatabase.h
--- a/Engine/Plugins/Experimental/Animation/PoseSearch/Source/Runtime/Public/PoseSearch/PoseSearchDatabase.h
+++ b/Engine/Plugins/Experimental/Animation/PoseSearch/Source/Runtime/Public/PoseSearch/PoseSearchDatabase.h
@@ -12,6 +12,7 @@ class UAnimationAsset;
 class UAnimComposite;
 class UAnimMontage;
 class UBlendSpace;
+class UPoseSearchNormalizationSet;
 
 UENUM()
 enum class EPoseSearchMode : int32
@@ -194,18 +195,6 @@ struct POSESEARCH_API FPoseSearchDatabaseAnimMontage : public FPoseSearchDatabas
 	EPoseSearchMirrorOption GetMirrorOption() const override { return MirrorOption; }
 };
 
-
-/** A data asset for indexing a collection of animation sequences. */
-UCLASS(BlueprintType, Category = "Animation|Pose Search", Experimental, meta = (DisplayName = "Normalization Set"))
-class POSESEARCH_API UNormalizationSetAsset : public UDataAsset
-{
-	GENERATED_BODY()
-public:
-
-	UPROPERTY(EditAnywhere, BlueprintReadOnly, Category = "NormalizationSet")
-	TArray<TObjectPtr<const UPoseSearchDatabase>> Databases; // @todo: it should be a UPoseSearchSearchableAsset, and have the UPoseSearchSearchableAsset iterate over all it's contained UPoseSearchSearchableAsset recursively (without duplicates)
-};
-
 /** A data asset for indexing a collection of animation sequences. */
 UCLASS(BlueprintType, Category = "Animation|Pose Search", Experimental, meta = (DisplayName = "Motion Database"))
 class POSESEARCH_API UPoseSearchDatabase : public UPoseSearchSearchableAsset
@@ -248,7 +237,7 @@ public:
 	bool bSkipSearchIfPossible = true;
 
 	UPROPERTY(EditAnywhere, BlueprintReadOnly, Category = "Database")
-	TObjectPtr<const UNormalizationSetAsset> NormalizationSet;
+	TObjectPtr<const UPoseSearchNormalizationSet> NormalizationSet;
 
 private:
 	UPROPERTY(Transient, meta = (ExcludeFromHash))
diff --git a/Engine/Plugins/Experimental/Animation/PoseSearch/Source/Runtime/Public/PoseSearch/PoseSearchNormalizationSet.h b/Engine/Plugins/Experimental/Animation/PoseSearch/Source/Runtime/Public/PoseSearch/PoseSearchNormalizationSet.h
new file mode 100644
--- /dev/null
+++ b/Engine/Plugins/Experimental/Animation/PoseSearch/Source/Runtime/Public/PoseSearch/PoseSearchNormalizationSet.h
@@ -0,0 +1,21 @@
+// Copyright Epic Games, Inc. All Rights Reserved.
+
+#pragma once
+
+#include "CoreMinimal.h"
+#include "Engine/DataAsset.h"
+#include "PoseSearchNormalizationSet.generated.h"
+
+class UPoseSearchDatabase;
+
+UCLASS(BlueprintType, Category = "Animation|Pose Search", Experimental, meta = (DisplayName = "Normalization Set"))
+class POSESEARCH_API UPoseSearchNormalizationSet : public UDataAsset
+{
+	GENERATED_BODY()
+	
+public:
+	UPROPERTY(EditAnywhere, BlueprintReadOnly, Category = "NormalizationSet")
+	TArray<TObjectPtr<const UPoseSearchDatabase>> Databases;
+
+	void AddUniqueDatabases(TArray<TWeakObjectPtr<const UPoseSearchDatabase>>& UniqueDatabases) const;
+};
\ No newline at end of file
