motion matching - defaulting UPoseSearchSchema for a meaningful locomotion setup

diff --git a/Engine/Plugins/Experimental/Animation/PoseSearch/Source/Editor/Private/PoseSearchSchemaFactory.cpp b/Engine/Plugins/Experimental/Animation/PoseSearch/Source/Editor/Private/PoseSearchSchemaFactory.cpp
--- a/Engine/Plugins/Experimental/Animation/PoseSearch/Source/Editor/Private/PoseSearchSchemaFactory.cpp
+++ b/Engine/Plugins/Experimental/Animation/PoseSearch/Source/Editor/Private/PoseSearchSchemaFactory.cpp
@@ -2,7 +2,8 @@
 
 #include "PoseSearchSchemaFactory.h"
 #include "PoseSearch/PoseSearchSchema.h"
-#include "PoseSearch/PoseSearchSchema.h"
+#include "PoseSearchFeatureChannel_Pose.h"
+#include "PoseSearchFeatureChannel_Trajectory.h"
 
 #define LOCTEXT_NAMESPACE "PoseSearchEditor"
 
@@ -16,7 +17,13 @@ UPoseSearchSchemaFactory::UPoseSearchSchemaFactory(const FObjectInitializer& Obj
 
 UObject* UPoseSearchSchemaFactory::FactoryCreateNew(UClass* Class, UObject* InParent, FName Name, EObjectFlags Flags, UObject* Context, FFeedbackContext* Warn)
 {
-	return NewObject<UPoseSearchSchema>(InParent, Class, Name, Flags);
+	UPoseSearchSchema* Schema = NewObject<UPoseSearchSchema>(InParent, Class, Name, Flags);
+
+	// defaulting UPoseSearchSchema for a meaningful locomotion setup
+	Schema->Channels.Add(NewObject<UPoseSearchFeatureChannel_Trajectory>(Schema, NAME_None));
+	Schema->Channels.Add(NewObject<UPoseSearchFeatureChannel_Pose>(Schema, NAME_None));
+
+	return Schema;
 }
 
 FString UPoseSearchSchemaFactory::GetDefaultNewAssetName() const
diff --git a/Engine/Plugins/Experimental/Animation/PoseSearch/Source/Runtime/Private/PoseSearchDerivedDataKey.cpp b/Engine/Plugins/Experimental/Animation/PoseSearch/Source/Runtime/Private/PoseSearchDerivedDataKey.cpp
--- a/Engine/Plugins/Experimental/Animation/PoseSearch/Source/Runtime/Private/PoseSearchDerivedDataKey.cpp
+++ b/Engine/Plugins/Experimental/Animation/PoseSearch/Source/Runtime/Private/PoseSearchDerivedDataKey.cpp
@@ -41,7 +41,7 @@ FKeyBuilder::FKeyBuilder(const UObject* Object, bool bUseDataVer)
 	if (bUseDataVer)
 	{
 		// used to invalidate the key without having to change POSESEARCHDB_DERIVEDDATA_VER all the times
-		int32 POSESEARCHDB_DERIVEDDATA_VER_SMALL = 123;
+		int32 POSESEARCHDB_DERIVEDDATA_VER_SMALL = 127;
 		FGuid VersionGuid = FDevSystemGuids::GetSystemGuid(FDevSystemGuids::Get().POSESEARCHDB_DERIVEDDATA_VER);
 
 		*this << VersionGuid;
diff --git a/Engine/Plugins/Experimental/Animation/PoseSearch/Source/Runtime/Private/PoseSearchFeatureChannel_Pose.cpp b/Engine/Plugins/Experimental/Animation/PoseSearch/Source/Runtime/Private/PoseSearchFeatureChannel_Pose.cpp
--- a/Engine/Plugins/Experimental/Animation/PoseSearch/Source/Runtime/Private/PoseSearchFeatureChannel_Pose.cpp
+++ b/Engine/Plugins/Experimental/Animation/PoseSearch/Source/Runtime/Private/PoseSearchFeatureChannel_Pose.cpp
@@ -11,6 +11,16 @@
 
 #define LOCTEXT_NAMESPACE "PoseSearchFeatureChannels"
 
+UPoseSearchFeatureChannel_Pose::UPoseSearchFeatureChannel_Pose()
+{
+#if WITH_EDITOR
+	// defaulting UPoseSearchFeatureChannel_Pose for a meaningful locomotion setup
+	Weight = 1.f;
+	SampledBones.Add(FPoseSearchBone({ {"foot_l"}, int32(EPoseSearchBoneFlags::Position | EPoseSearchBoneFlags::Velocity), 1.f, FLinearColor::Green }));
+	SampledBones.Add(FPoseSearchBone({ {"foot_r"}, int32(EPoseSearchBoneFlags::Position | EPoseSearchBoneFlags::Velocity), 1.f, FLinearColor::Green}));
+#endif
+}
+
 void UPoseSearchFeatureChannel_Pose::Finalize(UPoseSearchSchema* Schema)
 {
 	SubChannels.Reset();
diff --git a/Engine/Plugins/Experimental/Animation/PoseSearch/Source/Runtime/Private/PoseSearchFeatureChannel_Pose.h b/Engine/Plugins/Experimental/Animation/PoseSearch/Source/Runtime/Private/PoseSearchFeatureChannel_Pose.h
--- a/Engine/Plugins/Experimental/Animation/PoseSearch/Source/Runtime/Private/PoseSearchFeatureChannel_Pose.h
+++ b/Engine/Plugins/Experimental/Animation/PoseSearch/Source/Runtime/Private/PoseSearchFeatureChannel_Pose.h
@@ -61,6 +61,8 @@ public:
 	UPROPERTY(Transient)
 	TArray<TObjectPtr<UPoseSearchFeatureChannel>> SubChannels;
 
+	UPoseSearchFeatureChannel_Pose();
+
 	// UPoseSearchFeatureChannel_GroupBase interface
 	virtual TArrayView<TObjectPtr<UPoseSearchFeatureChannel>> GetSubChannels() override { return SubChannels; }
 	virtual TConstArrayView<TObjectPtr<UPoseSearchFeatureChannel>> GetSubChannels() const override { return SubChannels; }
diff --git a/Engine/Plugins/Experimental/Animation/PoseSearch/Source/Runtime/Private/PoseSearchFeatureChannel_Trajectory.cpp b/Engine/Plugins/Experimental/Animation/PoseSearch/Source/Runtime/Private/PoseSearchFeatureChannel_Trajectory.cpp
--- a/Engine/Plugins/Experimental/Animation/PoseSearch/Source/Runtime/Private/PoseSearchFeatureChannel_Trajectory.cpp
+++ b/Engine/Plugins/Experimental/Animation/PoseSearch/Source/Runtime/Private/PoseSearchFeatureChannel_Trajectory.cpp
@@ -15,7 +15,7 @@
 UPoseSearchFeatureChannel_Trajectory::UPoseSearchFeatureChannel_Trajectory()
 {
 #if WITH_EDITOR
-	// defaulting UPoseSearchFeatureChannel_Trajectory to a meaningful locomotion setup
+	// defaulting UPoseSearchFeatureChannel_Trajectory for a meaningful locomotion setup
 	Weight = 7.f;
 	Samples.Add(FPoseSearchTrajectorySample({ -0.4f, int32(EPoseSearchTrajectoryFlags::PositionXY), 0.4f, FLinearColor::Red }));
 	Samples.Add(FPoseSearchTrajectorySample({ 0.f, int32(EPoseSearchTrajectoryFlags::VelocityXY | EPoseSearchTrajectoryFlags::FacingDirectionXY), 2.f, FLinearColor::Blue }));
diff --git a/Engine/Plugins/Experimental/Animation/PoseSearch/Source/Runtime/Public/PoseSearch/PoseSearchSchema.h b/Engine/Plugins/Experimental/Animation/PoseSearch/Source/Runtime/Public/PoseSearch/PoseSearchSchema.h
--- a/Engine/Plugins/Experimental/Animation/PoseSearch/Source/Runtime/Public/PoseSearch/PoseSearchSchema.h
+++ b/Engine/Plugins/Experimental/Animation/PoseSearch/Source/Runtime/Public/PoseSearch/PoseSearchSchema.h
@@ -31,7 +31,6 @@ class POSESEARCH_API UPoseSearchSchema : public UDataAsset, public IBoneReferenc
 	GENERATED_BODY()
 
 public:
-	// @todo: used only for indexing: cache it somewhere else
 	UPROPERTY(EditAnywhere, Category = "Schema", meta = (DisplayPriority = 0))
 	TObjectPtr<USkeleton> Skeleton;
 
@@ -55,7 +54,6 @@ public:
 	UPROPERTY(Transient)
 	int32 SchemaCardinality = 0;
 
-	// @todo: used only for indexing: cache it somewhere else
 	UPROPERTY(Transient)
 	TArray<FBoneReference> BoneReferences;
 
@@ -77,7 +75,7 @@ public:
 
 	// cost added to all poses from looping assets of databases that uses this schema
 	UPROPERTY(EditAnywhere, Category = "Bias")
-	float LoopingCostBias = 0.f;
+	float LoopingCostBias = 0.01f;
 
 	// how many times the animation assets of the database using this schema will be indexed
 	UPROPERTY(EditAnywhere, Category = "Permutations", meta = (ClampMin = "1"))
