motion matching - tooltips, renaming, default tuning

diff --git a/Engine/Plugins/Experimental/Animation/PoseSearch/Source/Editor/Public/PoseSearchDatabaseEditorReflection.h b/Engine/Plugins/Experimental/Animation/PoseSearch/Source/Editor/Public/PoseSearchDatabaseEditorReflection.h
--- a/Engine/Plugins/Experimental/Animation/PoseSearch/Source/Editor/Public/PoseSearchDatabaseEditorReflection.h
+++ b/Engine/Plugins/Experimental/Animation/PoseSearch/Source/Editor/Public/PoseSearchDatabaseEditorReflection.h
@@ -34,9 +34,12 @@ struct FPoseSearchDatabaseSequenceEx : public FPoseSearchDatabaseSequence
 {
 	GENERATED_BODY()
 
+	// Is this animation set as looping? If you want to change this you need to change it in the base Anim Sequence.
+	// Changing the sampling range will disable looping
 	UPROPERTY(VisibleAnywhere, Category="Sequence", meta = (DisplayPriority = 10))
 	bool bLooping = false;
 
+	// Does this animation have root motion enabled ? If you want to change this you need to change it in the base Anim Sequence.
 	UPROPERTY(VisibleAnywhere, Category="Sequence", meta = (DisplayPriority = 11))
 	bool bHasRootMotion = false;
 };
@@ -60,9 +63,12 @@ struct FPoseSearchDatabaseBlendSpaceEx : public FPoseSearchDatabaseBlendSpace
 {
 	GENERATED_BODY()
 
+	// Is this animation set as looping? If you want to change this you need to change it in the base Anim Sequence.
+	// Changing the sampling range will disable looping
 	UPROPERTY(VisibleAnywhere, Category="Sequence", meta = (DisplayPriority = 10))
 	bool bLooping = false;
 
+	// Does this animation have root motion enabled ? If you want to change this you need to change it in the base Anim Sequence.
 	UPROPERTY(VisibleAnywhere, Category="Sequence", meta = (DisplayPriority = 11))
 	bool bHasRootMotion = false;
 };
@@ -86,9 +92,12 @@ struct FPoseSearchDatabaseAnimCompositeEx : public FPoseSearchDatabaseAnimCompos
 {
 	GENERATED_BODY()
 
+	// Is this animation set as looping? If you want to change this you need to change it in the base Anim Sequence.
+	// Changing the sampling range will disable looping
 	UPROPERTY(VisibleAnywhere, Category="Sequence", meta = (DisplayPriority = 10))
 	bool bLooping = false;
 
+	// Does this animation have root motion enabled ? If you want to change this you need to change it in the base Anim Sequence.
 	UPROPERTY(VisibleAnywhere, Category="Sequence", meta = (DisplayPriority = 11))
 	bool bHasRootMotion = false;
 };
@@ -112,9 +121,12 @@ struct FPoseSearchDatabaseAnimMontageEx : public FPoseSearchDatabaseAnimMontage
 {
 	GENERATED_BODY()
 
+	// Is this animation set as looping? If you want to change this you need to change it in the base Anim Sequence.
+	// Changing the sampling range will disable looping
 	UPROPERTY(VisibleAnywhere, Category="Sequence", meta = (DisplayPriority = 10))
 	bool bLooping = false;
 
+	// Does this animation have root motion enabled ? If you want to change this you need to change it in the base Anim Sequence.
 	UPROPERTY(VisibleAnywhere, Category="Sequence", meta = (DisplayPriority = 11))
 	bool bHasRootMotion = false;
 };
@@ -166,59 +178,68 @@ class UPoseSearchDatabaseStatistics : public UObject
 
 public:
 	
-	// General Information
-
+	// Number of Animation Sequences in the database.
 	UPROPERTY(VisibleAnywhere, Category = "General Information")
 	uint32 AnimationSequences;
 
+	// Number of total animation poses in frames in the database.
 	UPROPERTY(VisibleAnywhere, Category = "General Information")
 	uint32 TotalAnimationPosesInFrames;
 
+	// Number of total animation poses in time in the database.
 	UPROPERTY(VisibleAnywhere, Category = "General Information")
 	FText TotalAnimationPosesInTime;
 
+	// Amount of animation frames that are searchable in the database (this will exclude frames that have been removed using Sampling Range).
 	UPROPERTY(VisibleAnywhere, Category = "General Information")
 	uint32 SearchableFrames;
 
+	// Amount of animation in time that are searchable in the database (this will exclude time that have been removed using Sampling Range).
 	UPROPERTY(VisibleAnywhere, Category = "General Information")
 	FText SearchableTime;
 
-	// Kinematic Information
-
+	// Average speed of the characters trajectory across all animations in the database.
 	UPROPERTY(VisibleAnywhere, Category = "Kinematic Information")
 	FText AverageSpeed;
 
+	// Highest speed of the characters trajectory across all animations in the database.
 	UPROPERTY(VisibleAnywhere, Category = "Kinematic Information")
 	FText MaxSpeed;
 
+	// The average acceleration of the characters trajectory across all the animations in the database.
 	UPROPERTY(VisibleAnywhere, Category = "Kinematic Information")
 	FText AverageAcceleration;
 
+	// The max acceleration of the characters trajectory across all the animations in the database.
 	UPROPERTY(VisibleAnywhere, Category = "Kinematic Information")
 	FText MaxAcceleration;
 
-	// Principal Component Analysis Information
-
+	// When Pose Search Mode is set to PCAKDTree this value represents how well the variance of the dataset is explained within the chosen Number Of Principal Components:
+	// the higher, the better the quality (statistically more significant result) out of the kdtree search.
 	UPROPERTY(VisibleAnywhere, Category = "Principal Component Analysis Information", meta = (Units = "Percent"))
 	float ExplainedVariance;
 
-	// Memory information
-	
+	// aggregated total memory used by this database.
 	UPROPERTY(VisibleAnywhere, Category = "Memory Information")
 	FText EstimatedDatabaseSize;
 
+	// partial memory size used to store the pose feature vectors.
 	UPROPERTY(VisibleAnywhere, Category = "Memory Information")
 	FText ValuesSize;
 
+	// partial memory size used to store the pose feature vectors in PCA space.
 	UPROPERTY(VisibleAnywhere, Category = "Memory Information")
 	FText PCAValuesSize;
 
+	// partial memory size used by the kdtree.
 	UPROPERTY(VisibleAnywhere, Category = "Memory Information")
 	FText KDTreeSize;
 
+	// partial memory size used to store database metadata.
 	UPROPERTY(VisibleAnywhere, Category = "Memory Information")
 	FText PoseMetadataSize;
 	
+	// partial memory size used to animation data sub ranges, mirror state, blend parameters.
 	UPROPERTY(VisibleAnywhere, Category = "Memory Information")
 	FText AssetsSize;
 
diff --git a/Engine/Plugins/Experimental/Animation/PoseSearch/Source/Runtime/Private/PoseSearchFeatureChannel_Pose.h b/Engine/Plugins/Experimental/Animation/PoseSearch/Source/Runtime/Private/PoseSearchFeatureChannel_Pose.h
--- a/Engine/Plugins/Experimental/Animation/PoseSearch/Source/Runtime/Private/PoseSearchFeatureChannel_Pose.h
+++ b/Engine/Plugins/Experimental/Animation/PoseSearch/Source/Runtime/Private/PoseSearchFeatureChannel_Pose.h
@@ -25,6 +25,7 @@ struct POSESEARCH_API FPoseSearchBone
 	UPROPERTY(EditAnywhere, Category = Config)
 	FBoneReference Reference;
 
+	// This allows the user to define what information from the channel you want to compare to.
 	UPROPERTY(EditAnywhere, meta = (Bitmask, BitmaskEnum = "/Script/PoseSearch.EPoseSearchBoneFlags"), Category = Config)
 	int32 Flags = int32(EPoseSearchBoneFlags::Position);
 
@@ -45,6 +46,7 @@ public:
 	UPROPERTY(EditAnywhere, Category = "Settings")
 	float Weight = 1.f;
 
+	// List of skeletal joints and associated Flags (Velocity, Position, etc) to sample.
 	UPROPERTY(EditAnywhere, Category = "Settings")
 	TArray<FPoseSearchBone> SampledBones;
 
diff --git a/Engine/Plugins/Experimental/Animation/PoseSearch/Source/Runtime/Private/PoseSearchFeatureChannel_Trajectory.h b/Engine/Plugins/Experimental/Animation/PoseSearch/Source/Runtime/Private/PoseSearchFeatureChannel_Trajectory.h
--- a/Engine/Plugins/Experimental/Animation/PoseSearch/Source/Runtime/Private/PoseSearchFeatureChannel_Trajectory.h
+++ b/Engine/Plugins/Experimental/Animation/PoseSearch/Source/Runtime/Private/PoseSearchFeatureChannel_Trajectory.h
@@ -32,6 +32,7 @@ struct POSESEARCH_API FPoseSearchTrajectorySample
 	UPROPERTY(EditAnywhere, Category = Config)
 	float Offset = 0.f; // offset in time or distance depending on UPoseSearchFeatureChannel_Trajectory.Domain
 
+	// This allows the user to define what information from the channel you want to compare to.
 	UPROPERTY(EditAnywhere, meta = (Bitmask, BitmaskEnum = "/Script/PoseSearch.EPoseSearchTrajectoryFlags"), Category = Config)
 	int32 Flags = int32(EPoseSearchTrajectoryFlags::Position);
 
diff --git a/Engine/Plugins/Experimental/Animation/PoseSearch/Source/Runtime/Private/PoseSearchLibrary.cpp b/Engine/Plugins/Experimental/Animation/PoseSearch/Source/Runtime/Private/PoseSearchLibrary.cpp
--- a/Engine/Plugins/Experimental/Animation/PoseSearch/Source/Runtime/Private/PoseSearchLibrary.cpp
+++ b/Engine/Plugins/Experimental/Animation/PoseSearch/Source/Runtime/Private/PoseSearchLibrary.cpp
@@ -137,8 +137,7 @@ void FMotionMatchingState::JumpToPose(const FAnimationUpdateContext& Context, co
 	// requesting inertial blending only if blendstack is disabled
 	if (Settings.MaxActiveBlends <= 0)
 	{
-		const float JumpBlendTime = ComputeJumpBlendTime(Result, Settings);
-		RequestInertialBlend(Context, JumpBlendTime);
+		RequestInertialBlend(Context, Settings.BlendTime);
 	}
 
 	// Remember which pose and sequence we're playing from the database
@@ -151,7 +150,7 @@ void FMotionMatchingState::UpdateWantedPlayRate(const UE::PoseSearch::FSearchCon
 {
 	if (CurrentSearchResult.IsValid())
 	{
-		if (!FMath::IsNearlyEqual(Settings.PlayRateMin, 1.f, UE_KINDA_SMALL_NUMBER) || !FMath::IsNearlyEqual(Settings.PlayRateMax, 1.f, UE_KINDA_SMALL_NUMBER))
+		if (!FMath::IsNearlyEqual(Settings.PlayRate.Min, 1.f, UE_KINDA_SMALL_NUMBER) || !FMath::IsNearlyEqual(Settings.PlayRate.Max, 1.f, UE_KINDA_SMALL_NUMBER))
 		{
 			if (const FPoseSearchFeatureVectorBuilder* PoseSearchFeatureVectorBuilder = SearchContext.GetCachedQuery(CurrentSearchResult.Database->Schema))
 			{
@@ -160,8 +159,8 @@ void FMotionMatchingState::UpdateWantedPlayRate(const UE::PoseSearch::FSearchCon
 					TConstArrayView<float> QueryData = PoseSearchFeatureVectorBuilder->GetValues();
 					TConstArrayView<float> ResultData = CurrentSearchResult.Database->GetSearchIndex().GetPoseValues(CurrentSearchResult.PoseIdx);
 					const float EstimatedSpeedRatio = TrajectoryChannel->GetEstimatedSpeedRatio(QueryData, ResultData);
-					check(Settings.PlayRateMin <= Settings.PlayRateMax);
-					WantedPlayRate = FMath::Clamp(EstimatedSpeedRatio, Settings.PlayRateMin, Settings.PlayRateMax);
+					check(Settings.PlayRate.Min <= Settings.PlayRate.Max);
+					WantedPlayRate = FMath::Clamp(EstimatedSpeedRatio, Settings.PlayRate.Min, Settings.PlayRate.Max);
 				}
 				else
 				{
@@ -174,23 +173,6 @@ void FMotionMatchingState::UpdateWantedPlayRate(const UE::PoseSearch::FSearchCon
 	}
 }
 
-float FMotionMatchingState::ComputeJumpBlendTime(const UE::PoseSearch::FSearchResult& Result, const FMotionMatchingSettings& Settings) const
-{
-	const FPoseSearchIndexAsset* SearchIndexAsset = CurrentSearchResult.GetSearchIndexAsset();
-
-	// Use alternate blend time when changing between mirrored and unmirrored
-	float JumpBlendTime = Settings.BlendTime;
-	if (SearchIndexAsset && Settings.MirrorChangeBlendTime > 0.0f)
-	{
-		if (Result.GetSearchIndexAsset(true)->bMirrored != SearchIndexAsset->bMirrored)
-		{
-			JumpBlendTime = Settings.MirrorChangeBlendTime;
-		}
-	}
-
-	return JumpBlendTime;
-}
-
 void UPoseSearchLibrary::TraceMotionMatchingState(
 	const UPoseSearchSearchableAsset* Searchable,
 	UE::PoseSearch::FSearchContext& SearchContext,
diff --git a/Engine/Plugins/Experimental/Animation/PoseSearch/Source/Runtime/Public/PoseSearch/AnimNode_MotionMatching.h b/Engine/Plugins/Experimental/Animation/PoseSearch/Source/Runtime/Public/PoseSearch/AnimNode_MotionMatching.h
--- a/Engine/Plugins/Experimental/Animation/PoseSearch/Source/Runtime/Public/PoseSearch/AnimNode_MotionMatching.h
+++ b/Engine/Plugins/Experimental/Animation/PoseSearch/Source/Runtime/Public/PoseSearch/AnimNode_MotionMatching.h
@@ -20,24 +20,23 @@ public:
 	UPROPERTY()
 	FPoseLink Source;
 
-	// Collection of animations for motion matching
+	// Collection of animations for motion matching.
 	UPROPERTY(EditAnywhere, BlueprintReadWrite, Category=Settings, meta=(PinShownByDefault))
 	TObjectPtr<const UPoseSearchSearchableAsset> Searchable = nullptr;
 
-	// Motion trajectory samples for pose search queries. Expected to be in the space of the SkeletalMeshComponent.
+	// Motion Trajectory samples for pose search queries in Motion Matching.These are expected to be in the space of the SkeletalMeshComponent.This is provided with the CharacterMovementTrajectory Component output.
 	UPROPERTY(EditAnywhere, BlueprintReadWrite, Category=Settings, meta=(PinShownByDefault))
 	FTrajectorySampleRange Trajectory;
 
-	// Settings for the core motion matching algorithm evaluation
+	// Settings for the core motion matching node.
 	UPROPERTY(EditAnywhere, BlueprintReadWrite, Category=Settings, meta=(PinHiddenByDefault))
 	FMotionMatchingSettings Settings;
 
-	// Reset the motion matching state if we have become relevant to the graph
-	// after not being ticked on the previous frame(s)
+	// Reset the motion matching selection state if it has become relevant to the graph after not being updated on previous frames.
 	UPROPERTY(EditAnywhere, BlueprintReadWrite, Category = Settings, meta = (PinHiddenByDefault))
 	bool bResetOnBecomingRelevant = true;
 
-	// if true the continuing pose will be invalidated
+	// If set to true, the continuing pose will be invalidated. This is useful if you want to force a re-selection of the animation segment instead of continuing with the previous segment, even if it has a better score.
 	UPROPERTY(EditAnywhere, BlueprintReadWrite, Category = Settings, meta = (PinHiddenByDefault))
 	bool bForceInterrupt = false;
 
diff --git a/Engine/Plugins/Experimental/Animation/PoseSearch/Source/Runtime/Public/PoseSearch/PoseSearchDatabase.h b/Engine/Plugins/Experimental/Animation/PoseSearch/Source/Runtime/Public/PoseSearch/PoseSearchDatabase.h
--- a/Engine/Plugins/Experimental/Animation/PoseSearch/Source/Runtime/Public/PoseSearch/PoseSearchDatabase.h
+++ b/Engine/Plugins/Experimental/Animation/PoseSearch/Source/Runtime/Public/PoseSearch/PoseSearchDatabase.h
@@ -17,10 +17,17 @@ class UPoseSearchNormalizationSet;
 UENUM()
 enum class EPoseSearchMode : int32
 {
+	// Database searches will be evaluated extensively. the system will evaluate all the indexed poses to search for the best one.
 	BruteForce,
+
+	// Optimized search mode: the database projects the poses into a PCA space using only the most significant "NumberOfPrincipalComponents" dimensions, and construct a kdtree to facilitate the search.
 	PCAKDTree,
-	PCAKDTree_Validate,	// runs PCAKDTree and performs validation tests
-	PCAKDTree_Compare,	// compares BruteForce vs PCAKDTree
+
+	// Debug functionality performing diagnostics and validation during searches using PCAKDTree.
+	PCAKDTree_Validate,
+	
+	// Debug functionality to compare BruteForce vs PCAKDTree.
+	PCAKDTree_Compare,
 
 	Num UMETA(Hidden),
 	Invalid = Num UMETA(Hidden)
@@ -42,13 +49,15 @@ struct FPoseSearchExcludeFromDatabaseParameters
 {
 	GENERATED_BODY()
 
-	// Excluding the beginning of sequences can help ensure an exact past trajectory is used when building the features
-	UPROPERTY(EditAnywhere, Category = "Settings")
+	// Determines how much of the start of an animation segment is preserved for blending in seconds.
+	// Excluding the beginning of animation segments can help ensure an exact past trajectory is used when building the channels.
+	UPROPERTY(EditAnywhere, Category = "Settings", meta = (DisplayName = "Anim Start Interval"))
 	float SequenceStartInterval = 0.0f;
 
-	// Excluding the end of sequences help ensure an exact future trajectory, and also prevents the selection of
-	// a sequence which will end too soon to be worth selecting.
-	UPROPERTY(EditAnywhere, Category = "Settings")
+	// Determines how much of the end of an animation segment is preserved for blending in seconds.
+	// Excluding the end of animation segments helps ensure an exact future trajectory,
+	// and also prevents the selection of a anim segment which will end too soon to be worth selecting.
+	UPROPERTY(EditAnywhere, Category = "Settings", meta = (DisplayName = "Anim End Interval"))
 	float SequenceEndInterval = 0.3f;
 };
 
@@ -78,12 +87,18 @@ struct POSESEARCH_API FPoseSearchDatabaseSequence : public FPoseSearchDatabaseAn
 	UPROPERTY(EditAnywhere, Category="Sequence", meta = (DisplayPriority = 0))
 	TObjectPtr<UAnimSequence> Sequence;
 
+	// This allows users to enable or exclude animations from this database. Useful for debugging.
 	UPROPERTY(EditAnywhere, Category = "Sequence", meta = (DisplayPriority = 3))
 	bool bEnabled = true;
 
+	// It allows users to set a time range to an individual animation sequence in the database. 
+	// This is effectively trimming the beginning and end of the animation in the database (not in the original sequence).
+	// If set to [0, 0] it will be the entire frame range of the original sequence.
 	UPROPERTY(EditAnywhere, Category="Sequence", meta = (DisplayPriority = 1))
 	FFloatInterval SamplingRange = FFloatInterval(0.f, 0.f);
 
+	// This allows users to set if this animation is original only (no mirrored data), original and mirrored, or only the mirrored version of this animation.
+	// It requires the mirror table to be set up in the config file.
 	UPROPERTY(EditAnywhere, Category = "Sequence", meta = (DisplayPriority = 2))
 	EPoseSearchMirrorOption MirrorOption = EPoseSearchMirrorOption::UnmirroredOnly;
 
@@ -107,20 +122,24 @@ struct POSESEARCH_API FPoseSearchDatabaseBlendSpace : public FPoseSearchDatabase
 	UPROPERTY(EditAnywhere, Category = "BlendSpace", meta = (DisplayPriority = 0))
 	TObjectPtr<UBlendSpace> BlendSpace;
 
+	// This allows users to enable or exclude animations from this database. Useful for debugging.
 	UPROPERTY(EditAnywhere, Category = "BlendSpace", meta = (DisplayPriority = 5))
 	bool bEnabled = true;
 
+	// This allows users to set if this animation is original only (no mirrored data), original and mirrored, or only the mirrored version of this animation.
+	// It requires the mirror table to be set up in the config file.
 	UPROPERTY(EditAnywhere, Category = "BlendSpace", meta = (DisplayPriority = 1))
 	EPoseSearchMirrorOption MirrorOption = EPoseSearchMirrorOption::UnmirroredOnly;
 
-	// If to use the blendspace grid locations as parameter sample locations.
-	// When enabled, NumberOfHorizontalSamples and NumberOfVerticalSamples are ignored.
+	// When turned on, this will use the set grid samples of the blend space asset for sampling. This will override the Number of Horizontal/Vertical Samples.
 	UPROPERTY(EditAnywhere, Category = "BlendSpace", meta = (DisplayPriority = 4))
 	bool bUseGridForSampling = false;
 
+	// Sets the number of horizontal samples in the blend space to pull the animation data coverage from. The larger the samples the more the data, but also the more memory and performance it takes.
 	UPROPERTY(EditAnywhere, Category = "BlendSpace", meta = (EditCondition = "!bUseGridForSampling", EditConditionHides, ClampMin = "1", UIMin = "1", UIMax = "25", DisplayPriority = 2))
 	int32 NumberOfHorizontalSamples = 9;
 
+	// Sets the number of vertical samples in the blend space to pull the animation data coverage from.The larger the samples the more the data, but also the more memory and performance it takes.
 	UPROPERTY(EditAnywhere, Category = "BlendSpace", meta = (EditCondition = "!bUseGridForSampling", EditConditionHides, ClampMin = "1", UIMin = "1", UIMax = "25", DisplayPriority = 3))
 	int32 NumberOfVerticalSamples = 2;
 
@@ -147,12 +166,18 @@ struct POSESEARCH_API FPoseSearchDatabaseAnimComposite : public FPoseSearchDatab
 	UPROPERTY(EditAnywhere, Category = "AnimComposite", meta = (DisplayPriority = 0))
 	TObjectPtr<UAnimComposite> AnimComposite;
 
+	// This allows users to enable or exclude animations from this database. Useful for debugging.
 	UPROPERTY(EditAnywhere, Category = "AnimComposite", meta = (DisplayPriority = 3))
 	bool bEnabled = true;
 
+	// It allows users to set a time range to an individual animation sequence in the database. 
+	// This is effectively trimming the beginning and end of the animation in the database (not in the original sequence).
+	// If set to [0, 0] it will be the entire frame range of the original sequence.
 	UPROPERTY(EditAnywhere, Category = "AnimComposite", meta = (DisplayPriority = 1))
 	FFloatInterval SamplingRange = FFloatInterval(0.f, 0.f);
 
+	// This allows users to set if this animation is original only (no mirrored data), original and mirrored, or only the mirrored version of this animation.
+	// It requires the mirror table to be set up in the config file.
 	UPROPERTY(EditAnywhere, Category = "AnimComposite", meta = (DisplayPriority = 2))
 	EPoseSearchMirrorOption MirrorOption = EPoseSearchMirrorOption::UnmirroredOnly;
 
@@ -176,12 +201,18 @@ struct POSESEARCH_API FPoseSearchDatabaseAnimMontage : public FPoseSearchDatabas
 	UPROPERTY(EditAnywhere, Category="AnimMontage", meta = (DisplayPriority = 0))
 	TObjectPtr<UAnimMontage> AnimMontage;
 
+	// This allows users to enable or exclude animations from this database. Useful for debugging.
 	UPROPERTY(EditAnywhere, Category = "AnimMontage", meta = (DisplayPriority = 3))
 	bool bEnabled = true;
 
+	// It allows users to set a time range to an individual animation sequence in the database. 
+	// This is effectively trimming the beginning and end of the animation in the database (not in the original sequence).
+	// If set to [0, 0] it will be the entire frame range of the original sequence.
 	UPROPERTY(EditAnywhere, Category="AnimMontage", meta = (DisplayPriority = 1))
 	FFloatInterval SamplingRange = FFloatInterval(0.f, 0.f);
 
+	// This allows users to set if this animation is original only (no mirrored data), original and mirrored, or only the mirrored version of this animation.
+	// It requires the mirror table to be set up in the config file.
 	UPROPERTY(EditAnywhere, Category = "AnimMontage", meta = (DisplayPriority = 2))
 	EPoseSearchMirrorOption MirrorOption = EPoseSearchMirrorOption::UnmirroredOnly;
 
@@ -201,10 +232,12 @@ class POSESEARCH_API UPoseSearchDatabase : public UPoseSearchSearchableAsset
 {
 	GENERATED_BODY()
 public:
-	// Motion Database Config asset to use with this database.
+
+	// The Motion Database Config sets what channels this database will use to match against (bones, trajectory and what properties of those you’re interested in, such as position and velocity).
 	UPROPERTY(EditAnywhere, BlueprintReadOnly, Category="Database", DisplayName="Config")
 	TObjectPtr<const UPoseSearchSchema> Schema;
 
+	// These settings allow users to trim the start and end of animations in the database to preserve start/end frames for blending, and prevent the system from selecting the very last frames before it blends out.
 	UPROPERTY(EditAnywhere, Category = "Database")
 	FPoseSearchExcludeFromDatabaseParameters ExcludeFromDatabaseParameters;
 
@@ -220,28 +253,35 @@ public:
 	UPROPERTY(EditAnywhere, Category="Database")
 	TArray<FInstancedStruct> AnimationAssets;
 
+	// This dictates how the database will perform the search.
 	UPROPERTY(EditAnywhere, Category = "Performance")
 	EPoseSearchMode PoseSearchMode = EPoseSearchMode::PCAKDTree;
 
+	// Number of dimensions used to create the kdtree. More dimensions allows a better explanation of the variance of the dataset that usually translates in better search results, but will imply more memory usage and worse performances.
 	UPROPERTY(EditAnywhere, Category = "Performance", meta = (EditCondition = "PoseSearchMode != EPoseSearchMode::BruteForce", EditConditionHides, ClampMin = "1", ClampMax = "64", UIMin = "1", UIMax = "64"))
 	int32 NumberOfPrincipalComponents = 4;
 
 	UPROPERTY(EditAnywhere, Category = "Performance", meta = (EditCondition = "PoseSearchMode != EPoseSearchMode::BruteForce", EditConditionHides, ClampMin = "1", ClampMax = "256", UIMin = "1", UIMax = "256"))
 	int32 KDTreeMaxLeafSize = 16;
 	
+	// Out of a kdtree search, results will have only an approximate cost, so the database search will select the best “KDTree Query Num Neighbors” poses to perform the full cost analysis, and be able to elect the best pose.
 	UPROPERTY(EditAnywhere, Category = "Performance", meta = (EditCondition = "PoseSearchMode != EPoseSearchMode::BruteForce", EditConditionHides, ClampMin = "1", ClampMax = "600", UIMin = "1", UIMax = "600"))
 	int32 KDTreeQueryNumNeighbors = 200;
 
-	// if true, this database search will be skipped if cannot decrease the pose cost, and poses will not be listed into the PoseSearchDebugger
+	// When evaluating multiple searches, including the continuing pose search, the system keeps track of the best pose and associated cost.
+	// if the current database cannot possibly improve the current cost, the database search will be skipped entirely.
 	UPROPERTY(EditAnywhere, Category = "Performance")
 	bool bSkipSearchIfPossible = true;
 
+	// This optional asset defines a list of databases you want to normalize together. Without it, it would be difficult to compare costs from separately normalized databases containing different types of animation,
+	// like only idles versus only runs animations, given that the range of movement would be dramatically different.
 	UPROPERTY(EditAnywhere, BlueprintReadOnly, Category = "Database")
 	TObjectPtr<const UPoseSearchNormalizationSet> NormalizationSet;
 
 private:
+	// Do not use it directly. Use GetSearchIndex / SetSearchIndex interact with it and validate that is ok to do so.
 	UPROPERTY(Transient, meta = (ExcludeFromHash))
-	FPoseSearchIndex SearchIndexPrivate; // do not use it directly. use GetSearchIndex / SetSearchIndex interact with it and validate that is ok to do so
+	FPoseSearchIndex SearchIndexPrivate;
 
 public:
 	virtual ~UPoseSearchDatabase();
diff --git a/Engine/Plugins/Experimental/Animation/PoseSearch/Source/Runtime/Public/PoseSearch/PoseSearchFeatureChannel.h b/Engine/Plugins/Experimental/Animation/PoseSearch/Source/Runtime/Public/PoseSearch/PoseSearchFeatureChannel.h
--- a/Engine/Plugins/Experimental/Animation/PoseSearch/Source/Runtime/Public/PoseSearch/PoseSearchFeatureChannel.h
+++ b/Engine/Plugins/Experimental/Animation/PoseSearch/Source/Runtime/Public/PoseSearch/PoseSearchFeatureChannel.h
@@ -14,26 +14,26 @@ struct FPoseSearchFeatureVectorBuilder;
 UENUM()
 enum class EComponentStrippingVector : uint8
 {
-	// no component stripping
+	// No component stripping.
 	None,
 
-	// stripping X and Y components (matching only on the horizontal plane) 
+	// Stripping X and Y components (matching only on the horizontal plane).
 	StripXY,
 
-	// stripping Z (matching only vertically - caring only about the height of the feature) 
+	// Stripping Z (matching only vertically - caring only about the height of the feature).
 	StripZ,
 };
 
 UENUM()
 enum class EInputQueryPose : uint8
 {
-	// use character pose to compose the query
+	// Use character pose to compose the query.
 	UseCharacterPose,
 
-	// if available reuse continuing pose from the database to compose the query or else UseCharacterPose
+	// If available reuse continuing pose from the database to compose the query, or else UseCharacterPose.
 	UseContinuingPose,
 
-	// if available reuse and interpolate continuing pose from the database to compose the query or else UseCharacterPose
+	// If available reuse and interpolate continuing pose from the database to compose the query, or else UseCharacterPose.
 	UseInterpolatedContinuingPose,
 };
 
@@ -43,16 +43,16 @@ enum class EInputQueryPose : uint8
 UENUM()
 enum class EPermutationTimeType : uint8
 {
-	// bone and origin bone are sampled at the same sample time (plus eventual SampleTimeOffset for the bone):
-	// it's defined as the current animation evaluation time
+	// Bone and origin bone are sampled at the same sample time (plus eventual SampleTimeOffset for the bone):
+	// it's defined as the current animation evaluation time.
 	UseSampleTime,
 
-	// bone and origin bone are sampled at the same permutation time (plus eventual SampleTimeOffset for the bone):
+	// Bone and origin bone are sampled at the same permutation time (plus eventual SampleTimeOffset for the bone):
 	// it's defined as SamplingTime (as UseSampleTime) + Schema->PermutationsTimeOffset + PermutationIndex / Schema->PermutationsSampleRate
-	// where PermutationIndex is in range [0, Schema->NumberOfPermutations)
+	// where PermutationIndex is in range [0, Schema->NumberOfPermutations).
 	UsePermutationTime,
 
-	// bone is evaluated at sample time (and plus eventual SampleTimeOffset) and origin bone is evaluated at permutation time
+	// Bone is evaluated at sample time (and plus eventual SampleTimeOffset) and origin bone is evaluated at permutation time.
 	UseSampleToPermutationTime,
 };
 
diff --git a/Engine/Plugins/Experimental/Animation/PoseSearch/Source/Runtime/Public/PoseSearch/PoseSearchLibrary.h b/Engine/Plugins/Experimental/Animation/PoseSearch/Source/Runtime/Public/PoseSearch/PoseSearchLibrary.h
--- a/Engine/Plugins/Experimental/Animation/PoseSearch/Source/Runtime/Public/PoseSearch/PoseSearchLibrary.h
+++ b/Engine/Plugins/Experimental/Animation/PoseSearch/Source/Runtime/Public/PoseSearch/PoseSearchLibrary.h
@@ -28,41 +28,37 @@ struct POSESEARCH_API FMotionMatchingSettings
 {
 	GENERATED_BODY()
 
-	// Time in seconds to blend out to the new pose. Uses either inertial blending, requiring an Inertialization node after this node, or the internal blend stack, if MaxActiveBlends is greter than zero
+	// Time in seconds to blend out to the new pose. Uses either inertial blending, requiring an Inertialization node after this node, or the internal blend stack, if MaxActiveBlends is greater than zero.
 	UPROPERTY(EditAnywhere, BlueprintReadWrite, Category=Settings, meta=(ClampMin="0"))
 	float BlendTime = 0.2f;
 
-	// Number of max active blendin animation in the blend stack. If MaxActiveBlends is zero then blend stack is disabled
+	// Number of max active animation segments being blended together in the blend stack. If MaxActiveBlends is zero then the blend stack is disabled.
 	UPROPERTY(EditAnywhere, BlueprintReadWrite, Category=Settings, meta=(ClampMin="0"))
 	int32 MaxActiveBlends = 4;
 
+	// Set Blend Profiles (editable in the skeleton) to determine how the blending is distributed among your character's bones. It could be used to differentiate between upper body and lower body to blend timing.
 	UPROPERTY(EditAnywhere, BlueprintReadWrite, Category=Settings, meta=(UseAsBlendProfile = true))
 	TObjectPtr<UBlendProfile> BlendProfile;
 
+	// How the blend is applied over time to the bones. Common selections are linear, ease in, ease out, and ease in and out.
 	UPROPERTY(EditAnywhere, BlueprintReadWrite, Category = Settings)
 	EAlphaBlendOption BlendOption = EAlphaBlendOption::Linear;
 
-	// If the pose jump requires a mirroring change and this value is greater than 0, it will be used instead of BlendTime
-	UPROPERTY(EditAnywhere, BlueprintReadWrite, Category=Settings, meta=(ClampMin="0", DislayAfter="BlendTime"))
-	float MirrorChangeBlendTime = 0.0f;
-	
-	// Don't jump to poses that are less than this many seconds away
+	// Don't jump to poses of the same segment that are less than this many seconds away.
 	UPROPERTY(EditAnywhere, BlueprintReadWrite, Category=Settings, meta=(ClampMin="0"))
 	float PoseJumpThresholdTime = 0.f;
 
-	// Don't jump to poses that has been selected previously within this many seconds in the past
+	// Prevent re-selection of poses that have been selected previously within this much time (in seconds) in the past. This is across all animation segments that have been selected within this time range.
 	UPROPERTY(EditAnywhere, BlueprintReadWrite, Category = Settings, meta = (ClampMin = "0"))
 	float PoseReselectHistory = 0.3f;
 
-	// Minimum amount of time to wait between pose search queries
+	// Minimum amount of time to wait between searching for a new pose segment. It allows users to define how often the system searches, default for locomotion is searching every update, but you may only want to search once for other situations, like jump.
 	UPROPERTY(EditAnywhere, BlueprintReadWrite, Category=Settings, meta=(ClampMin="0"))
 	float SearchThrottleTime = 0.f;
 
+	// Effective range of play rate that can be applied to the animations to account for discrepancies in estimated velocity between the movement modeland the animation.
 	UPROPERTY(EditAnywhere, BlueprintReadWrite, Category = Settings, meta = (ClampMin = "0.5", ClampMax = "1.0", UIMin = "0.5", UIMax = "1.0"))
-	float PlayRateMin = 1.f;
-
-	UPROPERTY(EditAnywhere, BlueprintReadWrite, Category = Settings, meta = (ClampMin = "1.0", ClampMax = "2.0", UIMin = "1.0", UIMax = "2.0"))
-	float PlayRateMax = 1.f;
+	FFloatInterval PlayRate = FFloatInterval(1.f, 1.f);
 };
 
 USTRUCT(BlueprintType, Category="Animation|Pose Search")
@@ -84,8 +80,6 @@ struct POSESEARCH_API FMotionMatchingState
 	// Internally stores the 'jump' to a new pose/sequence index and asset time for evaluation
 	void JumpToPose(const FAnimationUpdateContext& Context, const FMotionMatchingSettings& Settings, const UE::PoseSearch::FSearchResult& Result);
 
-	float ComputeJumpBlendTime(const UE::PoseSearch::FSearchResult& Result, const FMotionMatchingSettings& Settings) const;
-
 	void UpdateWantedPlayRate(const UE::PoseSearch::FSearchContext& SearchContext, const FMotionMatchingSettings& Settings);
 
 	UE::PoseSearch::FSearchResult CurrentSearchResult;
@@ -94,7 +88,7 @@ struct POSESEARCH_API FMotionMatchingState
 	UPROPERTY(Transient)
 	float ElapsedPoseSearchTime = 0.f;
 
-	// wanted PlayRate to have the selected animation playing at the estimated requested speed from the query
+	// wanted PlayRate to have the selected animation playing at the estimated requested speed from the query.
 	UPROPERTY(Transient)
 	float WantedPlayRate = 1.f;
 
@@ -102,8 +96,7 @@ struct POSESEARCH_API FMotionMatchingState
 	UPROPERTY(VisibleAnywhere, BlueprintReadOnly, Category=State)
 	bool bJumpedToPose = false;
 
-	// Root motion delta for currently playing animation. Only required
-	// when UE_POSE_SEARCH_TRACE_ENABLED is active
+	// Root motion delta for currently playing animation. Only required when UE_POSE_SEARCH_TRACE_ENABLED is active.
 	UPROPERTY(Transient)
 	FTransform RootMotionTransformDelta = FTransform::Identity;
 
diff --git a/Engine/Plugins/Experimental/Animation/PoseSearch/Source/Runtime/Public/PoseSearch/PoseSearchSchema.h b/Engine/Plugins/Experimental/Animation/PoseSearch/Source/Runtime/Public/PoseSearch/PoseSearchSchema.h
--- a/Engine/Plugins/Experimental/Animation/PoseSearch/Source/Runtime/Public/PoseSearch/PoseSearchSchema.h
+++ b/Engine/Plugins/Experimental/Animation/PoseSearch/Source/Runtime/Public/PoseSearch/PoseSearchSchema.h
@@ -14,8 +14,13 @@ class UMirrorDataTable;
 UENUM()
 enum class EPoseSearchDataPreprocessor : int32
 {
+	// The data will be left untouched.
 	None,
+
+	// The data will be normalized against its deviation, and the user weights will be normalized to be a unitary vector.
 	Normalize,
+
+	// The data will be normalized against its deviation
 	NormalizeOnlyByDeviation,
 
 	Num UMETA(Hidden),
@@ -31,23 +36,27 @@ class POSESEARCH_API UPoseSearchSchema : public UDataAsset, public IBoneReferenc
 	GENERATED_BODY()
 
 public:
+	// Skeleton Reference for Motion Matching Database assets. Must be set to a compatible skeleton to the animation data in the database.
 	UPROPERTY(EditAnywhere, Category = "Schema", meta = (DisplayPriority = 0))
 	TObjectPtr<USkeleton> Skeleton;
 
+	// The update rate at which we sample the animation data in the database. The higher the SampleRate the more refined your searches will be, but the more memory will be required
 	UPROPERTY(EditAnywhere, Category = "Schema", meta = (DisplayPriority = 3, ClampMin = "1", ClampMax = "240"))
 	int32 SampleRate = 30;
 
+	// Channels itemize the cost breakdown of the config in simpler parts such as position or velocity of a bones, or phase of limbs. The total cost of a query against an indexed database pose will be the sum of the combined channel costs
 	UPROPERTY(EditAnywhere, Instanced, BlueprintReadWrite, Category = "Schema")
 	TArray<TObjectPtr<UPoseSearchFeatureChannel>> Channels;
 
-	// FinalizedChannels gets populated with UPoseSearchFeatureChannel(s) from Channels and additional injected ones during the Finalize
+	// FinalizedChannels gets populated with UPoseSearchFeatureChannel(s) from Channels and additional injected ones during the Finalize.
 	UPROPERTY(Transient)
 	TArray<TObjectPtr<UPoseSearchFeatureChannel>> FinalizedChannels;
 
-	// If set, this schema will support mirroring pose search databases
+	// Setting up and assigning a mirror data table will allow all your assets in your database to access the mirrored version of the data. This is required for mirroring to work with Motion Matching.
 	UPROPERTY(EditAnywhere, Category = "Schema", meta = (DisplayPriority = 1))
 	TObjectPtr<UMirrorDataTable> MirrorDataTable;
 
+	// Type of operation performed to the full pose features dataset
 	UPROPERTY(EditAnywhere, Category = "Schema", meta = (DisplayPriority = 2))
 	EPoseSearchDataPreprocessor DataPreprocessor = EPoseSearchDataPreprocessor::Normalize;
 
@@ -60,39 +69,43 @@ public:
 	UPROPERTY(Transient)
 	TArray<uint16> BoneIndicesWithParents;
 
-	// cost added to the continuing pose from databases that uses this schema
+	// Cost added to the continuing pose from databases that uses this config. This allows users to apply a cost bias (positive or negative) to the continuing pose.
+	// This is useful to help the system stay in one animation segment longer, or shorter depending on how you set this bias.
+	// Negative values make it more likely to be picked, or stayed in, positive values make it less likely to be picked or stay in.
 	UPROPERTY(EditAnywhere, Category = "Bias")
 	float ContinuingPoseCostBias = 0.f;
 
-	// base cost added to all poses from databases that uses this schema. it can be overridden by UAnimNotifyState_PoseSearchModifyCost
+	// Base Cost added or removed to all poses from databases that use this config. It can be overridden by Anim Notify: Pose Search Modify Cost at the frame level of animation data.
+	// Negative values make it more likely to be picked, or stayed in, Positive values make it less likely to be picked or stay in.
 	UPROPERTY(EditAnywhere, Category = "Bias")
 	float BaseCostBias = 0.f;
 
-	// If there's a mirroring mismatch between the currently playing asset and a search candidate, this cost will be 
-	// added to the candidate, making it less likely to be selected
+	// If there’s a mirroring mismatch between the currently playing asset and a search candidate, this cost will be added or removed to the candidate,
+	// making it more or less likely to be selected. Negative values make it more likely to be picked, or stayed in, Positive values make it less likely to be picked or stay in.
 	UPROPERTY(EditAnywhere, Category = "Bias")
 	float MirrorMismatchCostBias = 0.f;
 
-	// cost added to all poses from looping assets of databases that uses this schema
+	// Cost added to all looping animation assets in a database that uses this config. This allows users to make it more or less likely to pick the looping animation segments.
+	// Negative values make it more likely to be picked, or stayed in, Positive values make it less likely to be picked or stay in.
 	UPROPERTY(EditAnywhere, Category = "Bias")
-	float LoopingCostBias = 0.01f;
+	float LoopingCostBias = -0.005f;
 
-	// how many times the animation assets of the database using this schema will be indexed
+	// How many times the animation assets of the database using this schema will be indexed.
 	UPROPERTY(EditAnywhere, Category = "Permutations", meta = (ClampMin = "1"))
 	int32 NumberOfPermutations = 1;
 
-	// delta time between every permutation indexing
+	// Delta time between every permutation indexing.
 	UPROPERTY(EditAnywhere, Category = "Permutations", meta = (ClampMin = "1", ClampMax = "240", EditCondition = "NumberOfPermutations > 1", EditConditionHides))
 	int32 PermutationsSampleRate = 30;
 
-	// starting offset of the "PermutationTime" from the "SamplingTime" of the first permutation.
-	// subsequent permutations will have PermutationTime = SamplingTime + PermutationsTimeOffset + PermutationIndex / PermutationsSampleRate
+	// Starting offset of the "PermutationTime" from the "SamplingTime" of the first permutation.
+	// subsequent permutations will have PermutationTime = SamplingTime + PermutationsTimeOffset + PermutationIndex / PermutationsSampleRate.
 	UPROPERTY(EditAnywhere, Category = "Permutations")
 	float PermutationsTimeOffset = 0.f;
 
-	// if bInjectAdditionalDebugChannels is true, channels will be asked to injecting additional channels into this schema.
+	// If bInjectAdditionalDebugChannels is true, channels will be asked to inject additional channels into this schema.
 	// the original intent is to add UPoseSearchFeatureChannel_Position(s) to help with the complexity of the debug drawing
-	// (the database will have all the necessary positions to draw lines at the right location and time)
+	// (the database will have all the necessary positions to draw lines at the right location and time).
 	UPROPERTY(EditAnywhere, Category = "Debug")
 	bool bInjectAdditionalDebugChannels;
 
