motion matching - fix for normalization of different databases sharing the same schema

diff --git a/Engine/Plugins/Experimental/Animation/PoseSearch/Source/Runtime/Private/PoseSearchDerivedData.cpp b/Engine/Plugins/Experimental/Animation/PoseSearch/Source/Runtime/Private/PoseSearchDerivedData.cpp
--- a/Engine/Plugins/Experimental/Animation/PoseSearch/Source/Runtime/Private/PoseSearchDerivedData.cpp
+++ b/Engine/Plugins/Experimental/Animation/PoseSearch/Source/Runtime/Private/PoseSearchDerivedData.cpp
@@ -59,36 +59,12 @@ private:
 		bool bEntryFound = false;
 		for (FEntries& Entries : EntriesGroup)
 		{
-			check(Entries[0].Channel != Channel);
 			if (Entries[0].Channel->CanBeNormalizedWith(Channel))
 			{
-				#if DO_CHECK
-				check(Entries[0].Channel->GetChannelCardinality() == Channel->GetChannelCardinality());
-
-				// Channel can be normalized with ALL the channels already collected in Entry. Let's test that it's true
-				for (int EntryIndex = 1; EntryIndex < Entries.Num(); ++EntryIndex)
-				{
-					check(Entries[EntryIndex].Channel != Channel);
-					check(Entries[EntryIndex].Channel->CanBeNormalizedWith(Channel));
-				}
-				#endif // DO_CHECK
-
 				Entries.Add({ SchemaIndex, Channel });
 				bEntryFound = true;
-				#if !DO_CHECK
 				break;
-				#endif // !DO_CHECK
-			}
-			#if DO_CHECK
-			else
-			{
-				for (int EntryIndex = 1; EntryIndex < Entries.Num(); ++EntryIndex)
-				{
-					check(Entries[EntryIndex].Channel != Channel);
-					check(!Entries[EntryIndex].Channel->CanBeNormalizedWith(Channel));
-				}
 			}
-			#endif // DO_CHECK
 		}
 
 		if (!bEntryFound)
diff --git a/Engine/Plugins/Experimental/Animation/PoseSearch/Source/Runtime/Private/PoseSearchFeatureChannel.cpp b/Engine/Plugins/Experimental/Animation/PoseSearch/Source/Runtime/Private/PoseSearchFeatureChannel.cpp
--- a/Engine/Plugins/Experimental/Animation/PoseSearch/Source/Runtime/Private/PoseSearchFeatureChannel.cpp
+++ b/Engine/Plugins/Experimental/Animation/PoseSearch/Source/Runtime/Private/PoseSearchFeatureChannel.cpp
@@ -196,6 +196,11 @@ FString UPoseSearchFeatureChannel::GetLabel() const
 
 bool UPoseSearchFeatureChannel::CanBeNormalizedWith(const UPoseSearchFeatureChannel* Other) const
 {
+	if (this == Other)
+	{
+		return true;
+	}
+
 	if (ChannelCardinality != Other->ChannelCardinality)
 	{
 		return false;
