motion matching - fix for FAssetIndexer::GetEntry component space pose extraction

diff --git a/Engine/Plugins/Experimental/Animation/PoseSearch/Source/Runtime/Private/PoseSearchAssetIndexer.cpp b/Engine/Plugins/Experimental/Animation/PoseSearch/Source/Runtime/Private/PoseSearchAssetIndexer.cpp
--- a/Engine/Plugins/Experimental/Animation/PoseSearch/Source/Runtime/Private/PoseSearchAssetIndexer.cpp
+++ b/Engine/Plugins/Experimental/Animation/PoseSearch/Source/Runtime/Private/PoseSearchAssetIndexer.cpp
@@ -331,6 +331,7 @@ FAssetIndexer::CachedEntry& FAssetIndexer::GetEntry(float SampleTime)
 		Pose.SetBoneContainer(&BoneContainer);
 
 		Sample.Clip->ExtractPose(ExtractionCtx, AnimPoseData);
+		Pose[FCompactPoseBoneIndex(RootBoneIndexType)].SetIdentity();
 
 		if (IndexingContext.bMirrored)
 		{
