motion matching - fix for FMotionMatchingSettings::PlayRate UIMax clamping

diff --git a/Engine/Plugins/Experimental/Animation/PoseSearch/Source/Runtime/Public/PoseSearch/PoseSearchLibrary.h b/Engine/Plugins/Experimental/Animation/PoseSearch/Source/Runtime/Public/PoseSearch/PoseSearchLibrary.h
--- a/Engine/Plugins/Experimental/Animation/PoseSearch/Source/Runtime/Public/PoseSearch/PoseSearchLibrary.h
+++ b/Engine/Plugins/Experimental/Animation/PoseSearch/Source/Runtime/Public/PoseSearch/PoseSearchLibrary.h
@@ -57,7 +57,7 @@ struct POSESEARCH_API FMotionMatchingSettings
 	float SearchThrottleTime = 0.f;
 
 	// Effective range of play rate that can be applied to the animations to account for discrepancies in estimated velocity between the movement modeland the animation.
-	UPROPERTY(EditAnywhere, BlueprintReadWrite, Category = Settings, meta = (ClampMin = "0.5", ClampMax = "1.0", UIMin = "0.5", UIMax = "1.0"))
+	UPROPERTY(EditAnywhere, BlueprintReadWrite, Category = Settings, meta = (ClampMin = "0.5", ClampMax = "2.0", UIMin = "0.5", UIMax = "2.0"))
 	FFloatInterval PlayRate = FFloatInterval(1.f, 1.f);
 };
 
