motion matching - fix for FSequenceBaseSampler::GetTotalRootTransform (Debugging Motion Matching in 5.1, looping animations don't loop root correctly)

diff --git a/Engine/Plugins/Experimental/Animation/PoseSearch/Source/Runtime/Private/PoseSearchAssetSampler.cpp b/Engine/Plugins/Experimental/Animation/PoseSearch/Source/Runtime/Private/PoseSearchAssetSampler.cpp
--- a/Engine/Plugins/Experimental/Animation/PoseSearch/Source/Runtime/Private/PoseSearchAssetSampler.cpp
+++ b/Engine/Plugins/Experimental/Animation/PoseSearch/Source/Runtime/Private/PoseSearchAssetSampler.cpp
@@ -89,8 +89,8 @@ void FSequenceBaseSampler::ExtractPose(const FAnimExtractContext& ExtractionCtx,
 
 FTransform FSequenceBaseSampler::GetTotalRootTransform() const
 {
-	const FTransform InitialRootTransform = Input.SequenceBase->ExtractRootTrackTransform(0.f, nullptr);
-	const FTransform LastRootTransform = Input.SequenceBase->ExtractRootTrackTransform(Input.SequenceBase->GetPlayLength(), nullptr);
+	const FTransform InitialRootTransform = ExtractRootTransform(0.f);
+	const FTransform LastRootTransform = ExtractRootTransform(GetPlayLength());
 	const FTransform TotalRootTransform = LastRootTransform.GetRelativeTransform(InitialRootTransform);
 	return TotalRootTransform;
 }
diff --git a/Engine/Plugins/Experimental/Animation/PoseSearch/Source/Runtime/Private/PoseSearchDerivedDataKey.cpp b/Engine/Plugins/Experimental/Animation/PoseSearch/Source/Runtime/Private/PoseSearchDerivedDataKey.cpp
--- a/Engine/Plugins/Experimental/Animation/PoseSearch/Source/Runtime/Private/PoseSearchDerivedDataKey.cpp
+++ b/Engine/Plugins/Experimental/Animation/PoseSearch/Source/Runtime/Private/PoseSearchDerivedDataKey.cpp
@@ -41,7 +41,7 @@ FKeyBuilder::FKeyBuilder(const UObject* Object, bool bUseDataVer)
 	if (bUseDataVer)
 	{
 		// used to invalidate the key without having to change POSESEARCHDB_DERIVEDDATA_VER all the times
-		int32 POSESEARCHDB_DERIVEDDATA_VER_SMALL = 127;
+		int32 POSESEARCHDB_DERIVEDDATA_VER_SMALL = 131;
 		FGuid VersionGuid = FDevSystemGuids::GetSystemGuid(FDevSystemGuids::Get().POSESEARCHDB_DERIVEDDATA_VER);
 
 		*this << VersionGuid;
