motion matching - QOL improvements

diff --git a/Engine/Plugins/Experimental/Animation/MotionTrajectory/Source/MotionTrajectory/Private/MotionTrajectory.cpp b/Engine/Plugins/Experimental/Animation/MotionTrajectory/Source/MotionTrajectory/Private/MotionTrajectory.cpp
--- a/Engine/Plugins/Experimental/Animation/MotionTrajectory/Source/MotionTrajectory/Private/MotionTrajectory.cpp
+++ b/Engine/Plugins/Experimental/Animation/MotionTrajectory/Source/MotionTrajectory/Private/MotionTrajectory.cpp
@@ -31,7 +31,7 @@ UMotionTrajectoryComponent::UMotionTrajectoryComponent(const FObjectInitializer&
 FTrajectorySample UMotionTrajectoryComponent::CalcWorldSpacePresentTrajectorySample(float DeltaTime) const
 {
 	const APawn* Pawn = TryGetOwnerPawn();
-	checkf(false, TEXT("UMotionTrajectoryComponent::GetPresentTrajectory for Pawn: %s requires implementation."), Pawn ? *(Pawn->GetHumanReadableName()) : TEXT("NULL"));
+	UE_LOG(LogMotionTrajectory, Error, TEXT("UMotionTrajectoryComponent::GetPresentTrajectory for Pawn: %s requires implementation."), Pawn ? *(Pawn->GetHumanReadableName()) : TEXT("NULL"));
 	return FTrajectorySample();
 }
 
diff --git a/Engine/Plugins/Experimental/Animation/PoseSearch/Source/Editor/Private/PoseSearchMeshComponent.cpp b/Engine/Plugins/Experimental/Animation/PoseSearch/Source/Editor/Private/PoseSearchMeshComponent.cpp
--- a/Engine/Plugins/Experimental/Animation/PoseSearch/Source/Editor/Private/PoseSearchMeshComponent.cpp
+++ b/Engine/Plugins/Experimental/Animation/PoseSearch/Source/Editor/Private/PoseSearchMeshComponent.cpp
@@ -7,6 +7,7 @@
 #include "Animation/BlendSpace.h"
 #include "Animation/MirrorDataTable.h"
 #include "Engine/SkinnedAsset.h"
+#include "PoseSearch/PoseSearchContext.h"
 
 void UPoseSearchMeshComponent::Initialize(const FTransform& InComponentToWorld)
 {
@@ -153,14 +154,7 @@ void UPoseSearchMeshComponent::UpdatePose(const FUpdateContext& UpdateContext)
 
 	if (UpdateContext.bMirrored)
 	{
-		const EAxis::Type MirrorAxis = UpdateContext.MirrorDataTable->MirrorAxis;
-		FVector T = LastRootMotionDelta.GetTranslation();
-		T = FAnimationRuntime::MirrorVector(T, MirrorAxis);
-		const FQuat ReferenceRotation = (*UpdateContext.ComponentSpaceRefRotations)[FCompactPoseBoneIndex(0)];
-		FQuat Q = LastRootMotionDelta.GetRotation();
-		Q = FAnimationRuntime::MirrorQuat(Q, MirrorAxis);
-		Q *= FAnimationRuntime::MirrorQuat(ReferenceRotation, MirrorAxis).Inverse() * ReferenceRotation;
-		LastRootMotionDelta = FTransform(Q, T, LastRootMotionDelta.GetScale3D());
+		LastRootMotionDelta = UE::PoseSearch::MirrorTransform(LastRootMotionDelta, UpdateContext.MirrorDataTable->MirrorAxis, (*UpdateContext.ComponentSpaceRefRotations)[FCompactPoseBoneIndex(RootBoneIndexType)]);
 	}
 
 	const FTransform ComponentTransform = LastRootMotionDelta * StartingTransform;
diff --git a/Engine/Plugins/Experimental/Animation/PoseSearch/Source/Runtime/Private/PoseSearchAssetSampler.cpp b/Engine/Plugins/Experimental/Animation/PoseSearch/Source/Runtime/Private/PoseSearchAssetSampler.cpp
--- a/Engine/Plugins/Experimental/Animation/PoseSearch/Source/Runtime/Private/PoseSearchAssetSampler.cpp
+++ b/Engine/Plugins/Experimental/Animation/PoseSearch/Source/Runtime/Private/PoseSearchAssetSampler.cpp
@@ -479,15 +479,22 @@ void FBlendSpaceSampler::ProcessRootTransform()
 
 		const UE::Anim::IAnimRootMotionProvider* RootMotionProvider = UE::Anim::IAnimRootMotionProvider::Get();
 
-		if (ensureMsgf(RootMotionProvider, TEXT("Could not get Root Motion Provider.")))
+		if (RootMotionProvider)
 		{
-			if (ensureMsgf(RootMotionProvider->HasRootMotion(StackAttributeContainer), TEXT("Blend Space had no Root Motion Attribute.")))
+			if (RootMotionProvider->HasRootMotion(StackAttributeContainer))
 			{
 				FTransform RootMotionDelta;
 				RootMotionProvider->ExtractRootMotion(StackAttributeContainer, RootMotionDelta);
-
 				RootMotionAccumulation = RootMotionDelta * RootMotionAccumulation;
 			}
+			else
+			{
+				UE_LOG(LogPoseSearch, Error, TEXT("FBlendSpaceSampler::ProcessRootTransform: Blend Space '%s' has no Root Motion Attribute"), *Input.BlendSpace->GetName());
+			}
+		}
+		else
+		{
+			UE_LOG(LogPoseSearch, Error, TEXT("FBlendSpaceSampler::ProcessRootTransform: Could not get Root Motion Provider for BlendSpace '%s'"), *Input.BlendSpace->GetName());
 		}
 
 		AccumulatedRootTransform[SampleIdx] = RootMotionAccumulation;
diff --git a/Engine/Plugins/Experimental/Animation/PoseSearch/Source/Runtime/Private/PoseSearchContext.cpp b/Engine/Plugins/Experimental/Animation/PoseSearch/Source/Runtime/Private/PoseSearchContext.cpp
--- a/Engine/Plugins/Experimental/Animation/PoseSearch/Source/Runtime/Private/PoseSearchContext.cpp
+++ b/Engine/Plugins/Experimental/Animation/PoseSearch/Source/Runtime/Private/PoseSearchContext.cpp
@@ -108,38 +108,47 @@ const FTransform& FDebugDrawParams::GetRootTransform() const
 
 void FDebugDrawParams::DrawLine(const FVector& LineStart, const FVector& LineEnd, const FColor& Color, float Thickness) const
 {
-	if (AnimInstanceProxy)
-	{
-		AnimInstanceProxy->AnimDrawDebugLine(LineStart, LineEnd, Color, false, 0.f, Thickness, SDPG_Foreground);
-	}
-	else if (World)
+	if (Color.A > 0)
 	{
-		DrawDebugLine(World, LineStart, LineEnd, Color, false, 0.f, SDPG_Foreground, Thickness);
+		if (AnimInstanceProxy)
+		{
+			AnimInstanceProxy->AnimDrawDebugLine(LineStart, LineEnd, Color, false, 0.f, Thickness, SDPG_Foreground);
+		}
+		else if (World)
+		{
+			DrawDebugLine(World, LineStart, LineEnd, Color, false, 0.f, SDPG_Foreground, Thickness);
+		}
 	}
 }
 
 void FDebugDrawParams::DrawPoint(const FVector& Position, const FColor& Color, float Thickness) const
 {
-	if (AnimInstanceProxy)
-	{
-		AnimInstanceProxy->AnimDrawDebugPoint(Position, Thickness, Color, false, 0.f, SDPG_Foreground);
-	}
-	else if (World)
+	if (Color.A > 0)
 	{
-		DrawDebugPoint(World, Position, Thickness, Color, false, 0.f, SDPG_Foreground);
+		if (AnimInstanceProxy)
+		{
+			AnimInstanceProxy->AnimDrawDebugPoint(Position, Thickness, Color, false, 0.f, SDPG_Foreground);
+		}
+		else if (World)
+		{
+			DrawDebugPoint(World, Position, Thickness, Color, false, 0.f, SDPG_Foreground);
+		}
 	}
 }
 
 void FDebugDrawParams::DrawCircle(const FMatrix& TransformMatrix, float Radius, int32 Segments, const FColor& Color, float Thickness) const
 {
-	if (AnimInstanceProxy)
+	if (Color.A > 0)
 	{
-		AnimInstanceProxy->AnimDrawDebugCircle(TransformMatrix.GetOrigin(), Radius, Segments, Color, TransformMatrix.GetScaledAxis(EAxis::X), false, 0.f, SDPG_Foreground, Thickness);
-	}
-	else if (World)
-	{
-		// @todo: use the DrawDebugCircle API with the up vector to communize with the AnimInstanceProxy call
-		DrawDebugCircle(World, TransformMatrix, Radius, Segments, Color, false, 0.f, SDPG_Foreground, Thickness);
+		if (AnimInstanceProxy)
+		{
+			AnimInstanceProxy->AnimDrawDebugCircle(TransformMatrix.GetOrigin(), Radius, Segments, Color, TransformMatrix.GetScaledAxis(EAxis::X), false, 0.f, SDPG_Foreground, Thickness);
+		}
+		else if (World)
+		{
+			// @todo: use the DrawDebugCircle API with the up vector to communize with the AnimInstanceProxy call
+			DrawDebugCircle(World, TransformMatrix, Radius, Segments, Color, false, 0.f, SDPG_Foreground, Thickness);
+		}
 	}
 }
 
@@ -461,4 +470,13 @@ TConstArrayView<float> FSearchContext::GetCurrentResultNextPoseVector() const
 	return SearchIndex.GetPoseValues(CurrentResult.NextPoseIdx);
 }
 
+FTransform MirrorTransform(const FTransform& InTransform, EAxis::Type MirrorAxis, const FQuat& ReferenceRotation)
+{
+	const FVector T = FAnimationRuntime::MirrorVector(InTransform.GetTranslation(), MirrorAxis);
+	const FQuat Q = FAnimationRuntime::MirrorQuat(InTransform.GetRotation(), MirrorAxis);
+	const FQuat QR = Q * FAnimationRuntime::MirrorQuat(ReferenceRotation, MirrorAxis).Inverse() * ReferenceRotation;
+	const FTransform Result = FTransform(QR, T, InTransform.GetScale3D());
+	return Result;
+}
+
 } // namespace UE::PoseSearch
\ No newline at end of file
diff --git a/Engine/Plugins/Experimental/Animation/PoseSearch/Source/Runtime/Private/PoseSearchDatabaseIndexingContext.cpp b/Engine/Plugins/Experimental/Animation/PoseSearch/Source/Runtime/Private/PoseSearchDatabaseIndexingContext.cpp
--- a/Engine/Plugins/Experimental/Animation/PoseSearch/Source/Runtime/Private/PoseSearchDatabaseIndexingContext.cpp
+++ b/Engine/Plugins/Experimental/Animation/PoseSearch/Source/Runtime/Private/PoseSearchDatabaseIndexingContext.cpp
@@ -10,6 +10,7 @@
 #include "Animation/MirrorDataTable.h"
 #include "DerivedDataRequestOwner.h"
 #include "InstancedStruct.h"
+#include "PoseSearch/PoseSearchContext.h"
 #include "PoseSearch/PoseSearchDatabase.h"
 #include "PoseSearch/PoseSearchDefines.h"
 #include "PoseSearch/PoseSearchSchema.h"
@@ -36,15 +37,7 @@ void FAssetSamplingContext::Init(const UMirrorDataTable* InMirrorDataTable, cons
 
 FTransform FAssetSamplingContext::MirrorTransform(const FTransform& InTransform) const
 {
-	const EAxis::Type MirrorAxis = MirrorDataTable->MirrorAxis;
-	FVector T = InTransform.GetTranslation();
-	T = FAnimationRuntime::MirrorVector(T, MirrorAxis);
-	const FQuat ReferenceRotation = ComponentSpaceRefRotations[FCompactPoseBoneIndex(0)];
-	FQuat Q = InTransform.GetRotation();
-	Q = FAnimationRuntime::MirrorQuat(Q, MirrorAxis);
-	Q *= FAnimationRuntime::MirrorQuat(ReferenceRotation, MirrorAxis).Inverse() * ReferenceRotation;
-	FTransform Result = FTransform(Q, T, InTransform.GetScale3D());
-	return Result;
+	return UE::PoseSearch::MirrorTransform(InTransform, MirrorDataTable->MirrorAxis, ComponentSpaceRefRotations[FCompactPoseBoneIndex(RootBoneIndexType)]);
 }
 
 //////////////////////////////////////////////////////////////////////////
diff --git a/Engine/Plugins/Experimental/Animation/PoseSearch/Source/Runtime/Public/PoseSearch/PoseSearchContext.h b/Engine/Plugins/Experimental/Animation/PoseSearch/Source/Runtime/Public/PoseSearch/PoseSearchContext.h
--- a/Engine/Plugins/Experimental/Animation/PoseSearch/Source/Runtime/Public/PoseSearch/PoseSearchContext.h
+++ b/Engine/Plugins/Experimental/Animation/PoseSearch/Source/Runtime/Public/PoseSearch/PoseSearchContext.h
@@ -230,4 +230,6 @@ public:
 #endif
 };
 
+POSESEARCH_API FTransform MirrorTransform(const FTransform& InTransform, EAxis::Type MirrorAxis, const FQuat& ReferenceRotation);
+
 } // namespace UE::PoseSearch
diff --git a/Engine/Plugins/Experimental/Animation/PoseSearch/Source/Runtime/Public/PoseSearch/PoseSearchSchema.h b/Engine/Plugins/Experimental/Animation/PoseSearch/Source/Runtime/Public/PoseSearch/PoseSearchSchema.h
--- a/Engine/Plugins/Experimental/Animation/PoseSearch/Source/Runtime/Public/PoseSearch/PoseSearchSchema.h
+++ b/Engine/Plugins/Experimental/Animation/PoseSearch/Source/Runtime/Public/PoseSearch/PoseSearchSchema.h
@@ -73,7 +73,7 @@ public:
 	// This is useful to help the system stay in one animation segment longer, or shorter depending on how you set this bias.
 	// Negative values make it more likely to be picked, or stayed in, positive values make it less likely to be picked or stay in.
 	UPROPERTY(EditAnywhere, Category = "Bias")
-	float ContinuingPoseCostBias = 0.f;
+	float ContinuingPoseCostBias = -0.01f;
 
 	// Base Cost added or removed to all poses from databases that use this config. It can be overridden by Anim Notify: Pose Search Modify Cost at the frame level of animation data.
 	// Negative values make it more likely to be picked, or stayed in, Positive values make it less likely to be picked or stay in.
