motion matching - exposing  UPoseSearchDatabaseStatistics::ConfigCardinality in PoseSearchDatabaseEditorReflection

diff --git a/Engine/Plugins/Experimental/Animation/PoseSearch/Source/Editor/Private/PoseSearchDatabaseEditorReflection.cpp b/Engine/Plugins/Experimental/Animation/PoseSearch/Source/Editor/Private/PoseSearchDatabaseEditorReflection.cpp
--- a/Engine/Plugins/Experimental/Animation/PoseSearch/Source/Editor/Private/PoseSearchDatabaseEditorReflection.cpp
+++ b/Engine/Plugins/Experimental/Animation/PoseSearch/Source/Editor/Private/PoseSearchDatabaseEditorReflection.cpp
@@ -120,17 +120,16 @@ void UPoseSearchDatabaseStatistics::Initialize(const UPoseSearchDatabase* PoseSe
 		TotalAnimationPosesInFrames = SearchIndex.GetNumPoses();
 		TotalAnimationPosesInTime = FText::Format(TimeFormat, static_cast<double>(TotalAnimationPosesInFrames) / SampleRate);
 			
+		uint32 NumOfSearchablePoses = 0;
+		for (const FPoseSearchPoseMetadata& PoseMetadata : SearchIndex.PoseMetadata)
 		{
-			uint32 NumOfSearchablePoses = 0;
-			for (const FPoseSearchPoseMetadata & PoseMetadata : SearchIndex.PoseMetadata)
-			{
-				NumOfSearchablePoses += PoseMetadata.Flags != EPoseSearchPoseFlags::BlockTransition;
-			}
-				
-			SearchableFrames = NumOfSearchablePoses;
-			SearchableTime = FText::Format(TimeFormat, static_cast<double>(NumOfSearchablePoses) / SampleRate);
+			NumOfSearchablePoses += PoseMetadata.Flags != EPoseSearchPoseFlags::BlockTransition;
 		}
-			
+		SearchableFrames = NumOfSearchablePoses;
+		SearchableTime = FText::Format(TimeFormat, static_cast<double>(NumOfSearchablePoses) / SampleRate);
+	
+		ConfigCardinality = PoseSearchDatabase->Schema->SchemaCardinality;
+
 		// Kinematic Information
 	
 		// using FText instead of meta = (ForceUnits = "cm/s") to keep properties consistent
diff --git a/Engine/Plugins/Experimental/Animation/PoseSearch/Source/Editor/Public/PoseSearchDatabaseEditorReflection.h b/Engine/Plugins/Experimental/Animation/PoseSearch/Source/Editor/Public/PoseSearchDatabaseEditorReflection.h
--- a/Engine/Plugins/Experimental/Animation/PoseSearch/Source/Editor/Public/PoseSearchDatabaseEditorReflection.h
+++ b/Engine/Plugins/Experimental/Animation/PoseSearch/Source/Editor/Public/PoseSearchDatabaseEditorReflection.h
@@ -198,6 +198,10 @@ public:
 	UPROPERTY(VisibleAnywhere, Category = "General Information")
 	FText SearchableTime;
 
+	// Cardinality for the database config (how many floats per pose to store the pose features data)
+	UPROPERTY(VisibleAnywhere, Category = "General Information")
+	uint32 ConfigCardinality;
+
 	// Average speed of the characters trajectory across all animations in the database.
 	UPROPERTY(VisibleAnywhere, Category = "Kinematic Information")
 	FText AverageSpeed;
