motion matching - memory optimizations: stripping pose features Values from database index and reconstructing the data via FPoseSearchIndex::GetReconstructedPoseValues if the PoseSearchMode is PCAKDTree and KDTreeQueryNumNeighbors is 1

diff --git a/Engine/Plugins/Experimental/Animation/PoseSearch/Source/Editor/Private/PoseSearchDebuggerDatabaseColumns.h b/Engine/Plugins/Experimental/Animation/PoseSearch/Source/Editor/Private/PoseSearchDebuggerDatabaseColumns.h
--- a/Engine/Plugins/Experimental/Animation/PoseSearch/Source/Editor/Private/PoseSearchDebuggerDatabaseColumns.h
+++ b/Engine/Plugins/Experimental/Animation/PoseSearch/Source/Editor/Private/PoseSearchDebuggerDatabaseColumns.h
@@ -345,7 +345,6 @@ struct FChannelBreakdownCostColumn : ITextColumn
 	int32 BreakdownCostIndex = INDEX_NONE;
 };
 
-#if WITH_EDITORONLY_DATA
 struct FCostModifier : ITextColumn
 {
 	using ITextColumn::ITextColumn;
@@ -365,7 +364,6 @@ struct FCostModifier : ITextColumn
 		return FText::AsNumber(Row->PoseCost.GetCostAddend());
 	}
 };
-#endif // WITH_EDITORONLY_DATA
 
 struct FMirrored : ITextColumn
 {
diff --git a/Engine/Plugins/Experimental/Animation/PoseSearch/Source/Editor/Private/PoseSearchDebuggerDatabaseView.cpp b/Engine/Plugins/Experimental/Animation/PoseSearch/Source/Editor/Private/PoseSearchDebuggerDatabaseView.cpp
--- a/Engine/Plugins/Experimental/Animation/PoseSearch/Source/Editor/Private/PoseSearchDebuggerDatabaseView.cpp
+++ b/Engine/Plugins/Experimental/Animation/PoseSearch/Source/Editor/Private/PoseSearchDebuggerDatabaseView.cpp
@@ -193,7 +193,7 @@ void SDebuggerDatabaseView::Update(const FTraceMotionMatchingStateMessage& State
 					Row->PoseCost = PoseEntry.Cost;
 
 					Row->CostVector.SetNum(Database->Schema->SchemaCardinality);
-					TConstArrayView<float> PoseValues = SearchIndex.GetPoseValuesSafe(PoseEntry.DbPoseIdx);
+					const TArray<float> PoseValues = SearchIndex.GetPoseValuesSafe(PoseEntry.DbPoseIdx);
 
 					// in case we modify the schema while PIE is paused and displaying the Pose Search Editor, we could end up with a stale State with a DbEntry.QueryVector saved with the previous schema
 					// so the cardinality of DbEntry.QueryVector and PoseValues don't match. In that case we just use PoseValues as query to have all costs set to zero
@@ -365,9 +365,7 @@ void SDebuggerDatabaseView::Update(const FTraceMotionMatchingStateMessage& State
 			AddColumn(MakeShared<FChannelBreakdownCostColumn>(ColumnIdx++, LabelIdx++, Label));
 		}
 
-#if WITH_EDITORONLY_DATA
 		AddColumn(MakeShared<FCostModifier>(ColumnIdx++));
-#endif // WITH_EDITORONLY_DATA
 		AddColumn(MakeShared<FFrame>(ColumnIdx++));
 		AddColumn(MakeShared<FMirrored>(ColumnIdx++));
 		AddColumn(MakeShared<FLooping>(ColumnIdx++));
diff --git a/Engine/Plugins/Experimental/Animation/PoseSearch/Source/Runtime/Private/PoseSearchContext.cpp b/Engine/Plugins/Experimental/Animation/PoseSearch/Source/Runtime/Private/PoseSearchContext.cpp
--- a/Engine/Plugins/Experimental/Animation/PoseSearch/Source/Runtime/Private/PoseSearchContext.cpp
+++ b/Engine/Plugins/Experimental/Animation/PoseSearch/Source/Runtime/Private/PoseSearchContext.cpp
@@ -243,16 +243,50 @@ void FDebugDrawParams::DrawFeatureVector(TConstArrayView<float> PoseVector)
 
 void FDebugDrawParams::DrawFeatureVector(int32 PoseIdx)
 {
-	// if we're editing the schema while in PIE with Rewind Debugger active, PoseIdx could be out of bound / stale
-	if (CanDraw() && PoseIdx >= 0 && PoseIdx < GetSearchIndex()->GetNumPoses())
+	if (CanDraw())
 	{
-		DrawFeatureVector(GetSearchIndex()->GetPoseValues(PoseIdx));
+		DrawFeatureVector(GetSearchIndex()->GetPoseValuesSafe(PoseIdx));
 	}
 }
 #endif // ENABLE_DRAW_DEBUG
 
 //////////////////////////////////////////////////////////////////////////
 // FSearchContext
+FSearchContext::FSearchContext(const FTrajectorySampleRange* InTrajectory, const IPoseHistory* InHistory, float InDesiredPermutationTimeOffset, const FPoseIndicesHistory* InPoseIndicesHistory,
+	EPoseSearchBooleanRequest InQueryMirrorRequest, const FSearchResult& InCurrentResult, float InPoseJumpThresholdTime, bool bInForceInterrupt, bool bInCanAdvance)
+: Trajectory(InTrajectory)
+, History(InHistory)
+, DesiredPermutationTimeOffset(InDesiredPermutationTimeOffset)
+, PoseIndicesHistory(InPoseIndicesHistory)
+, QueryMirrorRequest(InQueryMirrorRequest)
+, CurrentResult(InCurrentResult)
+, PoseJumpThresholdTime(InPoseJumpThresholdTime)
+, bForceInterrupt(bInForceInterrupt)
+, bCanAdvance(bInCanAdvance)
+{
+	if (CurrentResult.IsValid())
+	{
+		const FPoseSearchIndex& SearchIndex = CurrentResult.Database->GetSearchIndex();
+		if (SearchIndex.Values.IsEmpty())
+		{
+			const int32 NumDimensions = CurrentResult.Database->Schema->SchemaCardinality;
+			CurrentResultPoseVectorData.AddUninitialized(NumDimensions * 3);
+
+			CurrentResultNextPoseVector = SearchIndex.GetReconstructedPoseValues(CurrentResult.NextPoseIdx, MakeArrayView(CurrentResultPoseVectorData.GetData(), NumDimensions));
+			CurrentResultPoseVector = SearchIndex.GetReconstructedPoseValues(CurrentResult.PoseIdx, MakeArrayView(CurrentResultPoseVectorData.GetData() + NumDimensions, NumDimensions));
+			CurrentResultPrevPoseVector = SearchIndex.GetReconstructedPoseValues(CurrentResult.PrevPoseIdx, MakeArrayView(CurrentResultPoseVectorData.GetData() + 2 * NumDimensions, NumDimensions));
+		}
+		else
+		{
+			CurrentResultNextPoseVector = SearchIndex.GetPoseValues(CurrentResult.NextPoseIdx);
+			CurrentResultPoseVector = SearchIndex.GetPoseValues(CurrentResult.PoseIdx);
+			CurrentResultPrevPoseVector = SearchIndex.GetPoseValues(CurrentResult.PrevPoseIdx);
+		}
+
+		check(CurrentResultNextPoseVector.Num() == CurrentResultPoseVector.Num() && CurrentResultPrevPoseVector.Num() == CurrentResultPoseVector.Num());
+	}
+}
+
 FQuat FSearchContext::GetSampleRotation(float SampleTimeOffset, const UPoseSearchSchema* Schema, int8 SchemaSampleBoneIdx, int8 SchemaOriginBoneIdx, bool bUseHistoryRoot, EPermutationTimeType PermutationTimeType)
 {
 	float PermutationSampleTimeOffset = 0.f;
@@ -449,27 +483,6 @@ bool FSearchContext::IsCurrentResultFromDatabase(const UPoseSearchDatabase* Data
 	return CurrentResult.IsValid() && CurrentResult.Database == Database;
 }
 
-TConstArrayView<float> FSearchContext::GetCurrentResultPrevPoseVector() const
-{
-	check(CurrentResult.IsValid());
-	const FPoseSearchIndex& SearchIndex = CurrentResult.Database->GetSearchIndex();
-	return SearchIndex.GetPoseValues(CurrentResult.PrevPoseIdx);
-}
-
-TConstArrayView<float> FSearchContext::GetCurrentResultPoseVector() const
-{
-	check(CurrentResult.IsValid());
-	const FPoseSearchIndex& SearchIndex = CurrentResult.Database->GetSearchIndex();
-	return SearchIndex.GetPoseValues(CurrentResult.PoseIdx);
-}
-
-TConstArrayView<float> FSearchContext::GetCurrentResultNextPoseVector() const
-{
-	check(CurrentResult.IsValid());
-	const FPoseSearchIndex& SearchIndex = CurrentResult.Database->GetSearchIndex();
-	return SearchIndex.GetPoseValues(CurrentResult.NextPoseIdx);
-}
-
 FTransform MirrorTransform(const FTransform& InTransform, EAxis::Type MirrorAxis, const FQuat& ReferenceRotation)
 {
 	const FVector T = FAnimationRuntime::MirrorVector(InTransform.GetTranslation(), MirrorAxis);
diff --git a/Engine/Plugins/Experimental/Animation/PoseSearch/Source/Runtime/Private/PoseSearchDatabase.cpp b/Engine/Plugins/Experimental/Animation/PoseSearch/Source/Runtime/Private/PoseSearchDatabase.cpp
--- a/Engine/Plugins/Experimental/Animation/PoseSearch/Source/Runtime/Private/PoseSearchDatabase.cpp
+++ b/Engine/Plugins/Experimental/Animation/PoseSearch/Source/Runtime/Private/PoseSearchDatabase.cpp
@@ -27,22 +27,26 @@ namespace UE::PoseSearch
 {
 
 typedef TArray<size_t, TInlineAllocator<128>> FNonSelectableIdx;
-static void PopulateNonSelectableIdx(FNonSelectableIdx& NonSelectableIdx, FSearchContext& SearchContext, const UPoseSearchDatabase* Database, TConstArrayView<float> QueryValues)
+static void PopulateNonSelectableIdx(FNonSelectableIdx& NonSelectableIdx, FSearchContext& SearchContext, const UPoseSearchDatabase* Database
+#if UE_POSE_SEARCH_TRACE_ENABLED
+	, TConstArrayView<float> QueryValues
+#endif //UE_POSE_SEARCH_TRACE_ENABLED
+)
 {
 	check(Database);
 #if UE_POSE_SEARCH_TRACE_ENABLED
 	const FPoseSearchIndex& SearchIndex = Database->GetSearchIndex();
 #endif
 
-	const FPoseSearchIndexAsset* CurrentIndexAsset = SearchContext.CurrentResult.GetSearchIndexAsset();
-	if (CurrentIndexAsset && SearchContext.IsCurrentResultFromDatabase(Database) && SearchContext.PoseJumpThresholdTime > 0.f)
+	const FPoseSearchIndexAsset* CurrentIndexAsset = SearchContext.GetCurrentResult().GetSearchIndexAsset();
+	if (CurrentIndexAsset && SearchContext.IsCurrentResultFromDatabase(Database) && SearchContext.GetPoseJumpThresholdTime() > 0.f)
 	{
-		const int32 PoseJumpIndexThreshold = FMath::FloorToInt(SearchContext.PoseJumpThresholdTime / Database->Schema->GetSamplingInterval());
+		const int32 PoseJumpIndexThreshold = FMath::FloorToInt(SearchContext.GetPoseJumpThresholdTime() / Database->Schema->GetSamplingInterval());
 		const bool IsLooping = Database->IsSourceAssetLooping(*CurrentIndexAsset);
 
 		for (int32 i = -PoseJumpIndexThreshold; i <= -1; ++i)
 		{
-			int32 PoseIdx = SearchContext.CurrentResult.PoseIdx + i;
+			int32 PoseIdx = SearchContext.GetCurrentResult().PoseIdx + i;
 			bool bIsPoseInRange = false;
 			if (IsLooping)
 			{
@@ -63,7 +67,8 @@ static void PopulateNonSelectableIdx(FNonSelectableIdx& NonSelectableIdx, FSearc
 				NonSelectableIdx.Add(PoseIdx);
 
 #if UE_POSE_SEARCH_TRACE_ENABLED
-				const FPoseSearchCost PoseCost = SearchIndex.ComparePoses(PoseIdx, SearchContext.QueryMirrorRequest, EPoseComparisonFlags::None, Database->Schema->MirrorMismatchCostBias, QueryValues);
+				const TArray<float> PoseValues = SearchIndex.GetPoseValuesSafe(PoseIdx);
+				const FPoseSearchCost PoseCost = SearchIndex.ComparePoses(PoseIdx, SearchContext.GetQueryMirrorRequest(), EPoseComparisonFlags::None, Database->Schema->MirrorMismatchCostBias, PoseValues, QueryValues);
 				SearchContext.BestCandidates.Add(PoseCost, PoseIdx, Database, EPoseCandidateFlags::DiscardedBy_PoseJumpThresholdTime);
 #endif
 			}
@@ -75,7 +80,7 @@ static void PopulateNonSelectableIdx(FNonSelectableIdx& NonSelectableIdx, FSearc
 
 		for (int32 i = 0; i <= PoseJumpIndexThreshold; ++i)
 		{
-			int32 PoseIdx = SearchContext.CurrentResult.PoseIdx + i;
+			int32 PoseIdx = SearchContext.GetCurrentResult().PoseIdx + i;
 			bool bIsPoseInRange = false;
 			if (IsLooping)
 			{
@@ -96,7 +101,7 @@ static void PopulateNonSelectableIdx(FNonSelectableIdx& NonSelectableIdx, FSearc
 				NonSelectableIdx.Add(PoseIdx);
 
 #if UE_POSE_SEARCH_TRACE_ENABLED
-				const FPoseSearchCost PoseCost = SearchIndex.ComparePoses(PoseIdx, SearchContext.QueryMirrorRequest, EPoseComparisonFlags::None, Database->Schema->MirrorMismatchCostBias, QueryValues);
+				const FPoseSearchCost PoseCost = SearchIndex.ComparePoses(PoseIdx, SearchContext.GetQueryMirrorRequest(), EPoseComparisonFlags::None, Database->Schema->MirrorMismatchCostBias, SearchIndex.GetPoseValuesSafe(PoseIdx), QueryValues);
 				SearchContext.BestCandidates.Add(PoseCost, PoseIdx, Database, EPoseCandidateFlags::DiscardedBy_PoseJumpThresholdTime);
 #endif
 			}
@@ -107,10 +112,10 @@ static void PopulateNonSelectableIdx(FNonSelectableIdx& NonSelectableIdx, FSearc
 		}
 	}
 
-	if (SearchContext.PoseIndicesHistory)
+	if (SearchContext.GetPoseIndicesHistory())
 	{
 		const FObjectKey DatabaseKey(Database);
-		for (auto It = SearchContext.PoseIndicesHistory->IndexToTime.CreateConstIterator(); It; ++It)
+		for (auto It = SearchContext.GetPoseIndicesHistory()->IndexToTime.CreateConstIterator(); It; ++It)
 		{
 			const FHistoricalPoseIndex& HistoricalPoseIndex = It.Key();
 			if (HistoricalPoseIndex.DatabaseKey == DatabaseKey)
@@ -123,7 +128,7 @@ static void PopulateNonSelectableIdx(FNonSelectableIdx& NonSelectableIdx, FSearc
 				// if we're editing the database and removing assets it's possible that the PoseIndicesHistory contains invalid pose indexes
 				if (HistoricalPoseIndex.PoseIndex < SearchIndex.GetNumPoses())
 				{
-					const FPoseSearchCost PoseCost = SearchIndex.ComparePoses(HistoricalPoseIndex.PoseIndex, SearchContext.QueryMirrorRequest, EPoseComparisonFlags::None, Database->Schema->MirrorMismatchCostBias, QueryValues);
+					const FPoseSearchCost PoseCost = SearchIndex.ComparePoses(HistoricalPoseIndex.PoseIndex, SearchContext.GetQueryMirrorRequest(), EPoseComparisonFlags::None, Database->Schema->MirrorMismatchCostBias, SearchIndex.GetPoseValuesSafe(HistoricalPoseIndex.PoseIndex), QueryValues);
 					SearchContext.BestCandidates.Add(PoseCost, HistoricalPoseIndex.PoseIndex, Database, EPoseCandidateFlags::DiscardedBy_PoseReselectHistory);
 				}
 #endif
@@ -159,13 +164,12 @@ struct FPoseFilters
 		}
 	}
 
-	bool AreFiltersValid(const FPoseSearchIndex& SearchIndex, TConstArrayView<float> QueryValues, int32 PoseIdx, const FPoseSearchPoseMetadata& Metadata
+	bool AreFiltersValid(const FPoseSearchIndex& SearchIndex, TConstArrayView<float> PoseValues, TConstArrayView<float> QueryValues, int32 PoseIdx, const FPoseSearchPoseMetadata& Metadata
 #if UE_POSE_SEARCH_TRACE_ENABLED
 		, UE::PoseSearch::FSearchContext& SearchContext, const UPoseSearchDatabase* Database
 #endif
 	) const
 	{
-		TConstArrayView<float> PoseValues = SearchIndex.GetPoseValues(PoseIdx);
 		for (const IPoseFilter* PoseFilter : AllPoseFilters)
 		{
 			if (!PoseFilter->IsPoseValid(PoseValues, QueryValues, PoseIdx, Metadata))
@@ -177,12 +181,12 @@ struct FPoseFilters
 				}
 				else if (PoseFilter == &BlockTransitionPoseFilter)
 				{
-					const FPoseSearchCost PoseCost = SearchIndex.ComparePoses(PoseIdx, SearchContext.QueryMirrorRequest, EPoseComparisonFlags::None, Database->Schema->MirrorMismatchCostBias, QueryValues);
+					const FPoseSearchCost PoseCost = SearchIndex.ComparePoses(PoseIdx, SearchContext.GetQueryMirrorRequest(), EPoseComparisonFlags::None, Database->Schema->MirrorMismatchCostBias, PoseValues, QueryValues);
 					SearchContext.BestCandidates.Add(PoseCost, PoseIdx, Database, EPoseCandidateFlags::DiscardedBy_BlockTransition);
 				}
 				else
 				{
-					const FPoseSearchCost PoseCost = SearchIndex.ComparePoses(PoseIdx, SearchContext.QueryMirrorRequest, EPoseComparisonFlags::None, Database->Schema->MirrorMismatchCostBias, QueryValues);
+					const FPoseSearchCost PoseCost = SearchIndex.ComparePoses(PoseIdx, SearchContext.GetQueryMirrorRequest(), EPoseComparisonFlags::None, Database->Schema->MirrorMismatchCostBias, PoseValues, QueryValues);
 					SearchContext.BestCandidates.Add(PoseCost, PoseIdx, Database, EPoseCandidateFlags::DiscardedBy_PoseFilter);
 				}
 #endif
@@ -580,8 +584,6 @@ bool UPoseSearchDatabase::GetSkipSearchIfPossible() const
 	return bSkipSearchIfPossible;
 }
 
-
-
 void UPoseSearchDatabase::PostLoad()
 {
 #if WITH_EDITORONLY_DATA
@@ -762,7 +764,8 @@ UE::PoseSearch::FSearchResult UPoseSearchDatabase::SearchPCAKDTree(UE::PoseSearc
 	RowMajorVectorMap WeightedQueryValues((float*)FMemory_Alloca(NumDimensions * sizeof(float)), 1, NumDimensions);
 	RowMajorVectorMap CenteredQueryValues((float*)FMemory_Alloca(NumDimensions * sizeof(float)), 1, NumDimensions);
 	RowMajorVectorMap ProjectedQueryValues((float*)FMemory_Alloca(ClampedNumberOfPrincipalComponents * sizeof(float)), 1, ClampedNumberOfPrincipalComponents);
-	
+	TArrayView<float> ReconstructedPoseValuesBuffer((float*)FMemory_Alloca(NumDimensions * sizeof(float)), NumDimensions);
+
 	// KDTree in PCA space search
 	if (PoseSearchMode == EPoseSearchMode::PCAKDTree_Validate)
 	{
@@ -805,7 +808,11 @@ UE::PoseSearch::FSearchResult UPoseSearchDatabase::SearchPCAKDTree(UE::PoseSearc
 		TConstArrayView<float> QueryValues = Result.ComposedQuery.GetValues();
 
 		FNonSelectableIdx NonSelectableIdx;
-		PopulateNonSelectableIdx(NonSelectableIdx, SearchContext, this, QueryValues);
+		PopulateNonSelectableIdx(NonSelectableIdx, SearchContext, this
+#if UE_POSE_SEARCH_TRACE_ENABLED
+			, QueryValues
+#endif // UE_POSE_SEARCH_TRACE_ENABLED
+		);
 		FKDTree::KNNResultSet ResultSet(ClampedKDTreeQueryNumNeighbors, ResultIndexes, ResultDistanceSqr, NonSelectableIdx);
 
 		check(QueryValues.Num() == NumDimensions);
@@ -827,13 +834,15 @@ UE::PoseSearch::FSearchResult UPoseSearchDatabase::SearchPCAKDTree(UE::PoseSearc
 		for (size_t ResultIndex = 0; ResultIndex < ResultSet.Num(); ++ResultIndex)
 		{
 			const int32 PoseIdx = ResultIndexes[ResultIndex];
-			if (PoseFilters.AreFiltersValid(SearchIndex, QueryValues, PoseIdx, SearchIndex.PoseMetadata[PoseIdx]
+			const TConstArrayView<float> PoseValues = SearchIndex.Values.IsEmpty() ? SearchIndex.GetReconstructedPoseValues(PoseIdx, ReconstructedPoseValuesBuffer) : SearchIndex.GetPoseValues(PoseIdx);
+
+			if (PoseFilters.AreFiltersValid(SearchIndex, PoseValues, QueryValues, PoseIdx, SearchIndex.PoseMetadata[PoseIdx]
 #if UE_POSE_SEARCH_TRACE_ENABLED
 				, SearchContext, this
 #endif
 			))
 			{
-				const FPoseSearchCost PoseCost = SearchIndex.ComparePoses(PoseIdx, SearchContext.QueryMirrorRequest, EPoseComparisonFlags::None, Schema->MirrorMismatchCostBias, QueryValues);
+				const FPoseSearchCost PoseCost = SearchIndex.ComparePoses(PoseIdx, SearchContext.GetQueryMirrorRequest(), EPoseComparisonFlags::None, Schema->MirrorMismatchCostBias, PoseValues, QueryValues);
 				if (PoseCost < Result.PoseCost)
 				{
 					Result.PoseCost = PoseCost;
@@ -886,19 +895,26 @@ UE::PoseSearch::FSearchResult UPoseSearchDatabase::SearchBruteForce(UE::PoseSear
 		TConstArrayView<float> QueryValues = Result.ComposedQuery.GetValues();
 
 		FNonSelectableIdx NonSelectableIdx;
-		PopulateNonSelectableIdx(NonSelectableIdx, SearchContext, this, QueryValues);
+		PopulateNonSelectableIdx(NonSelectableIdx, SearchContext, this
+#if UE_POSE_SEARCH_TRACE_ENABLED
+			, QueryValues
+#endif // UE_POSE_SEARCH_TRACE_ENABLED
+		);
 		check(Algo::IsSorted(NonSelectableIdx));
 
+		const int32 NumDimensions = Schema->SchemaCardinality;
+		TArrayView<float> ReconstructedPoseValuesBuffer((float*)FMemory_Alloca(NumDimensions * sizeof(float)), NumDimensions);
 		const FPoseFilters PoseFilters(Schema, NonSelectableIdx, SearchIndex.OverallFlags);
 		for (int32 PoseIdx = 0; PoseIdx < SearchIndex.GetNumPoses(); ++PoseIdx)
 		{
-			if (PoseFilters.AreFiltersValid(SearchIndex, QueryValues, PoseIdx, SearchIndex.PoseMetadata[PoseIdx]
+			const TConstArrayView<float> PoseValues = SearchIndex.Values.IsEmpty() ? SearchIndex.GetReconstructedPoseValues(PoseIdx, ReconstructedPoseValuesBuffer) : SearchIndex.GetPoseValues(PoseIdx);
+			if (PoseFilters.AreFiltersValid(SearchIndex, PoseValues, QueryValues, PoseIdx, SearchIndex.PoseMetadata[PoseIdx]
 #if UE_POSE_SEARCH_TRACE_ENABLED
 				, SearchContext, this
 #endif
 			))
 			{
-				const FPoseSearchCost PoseCost = SearchIndex.ComparePoses(PoseIdx, SearchContext.QueryMirrorRequest, EPoseComparisonFlags::None, Schema->MirrorMismatchCostBias, QueryValues);
+				const FPoseSearchCost PoseCost = SearchIndex.ComparePoses(PoseIdx, SearchContext.GetQueryMirrorRequest(), EPoseComparisonFlags::None, Schema->MirrorMismatchCostBias, PoseValues, QueryValues);
 				if (PoseCost < Result.PoseCost)
 				{
 					Result.PoseCost = PoseCost;
diff --git a/Engine/Plugins/Experimental/Animation/PoseSearch/Source/Runtime/Private/PoseSearchDerivedData.cpp b/Engine/Plugins/Experimental/Animation/PoseSearch/Source/Runtime/Private/PoseSearchDerivedData.cpp
--- a/Engine/Plugins/Experimental/Animation/PoseSearch/Source/Runtime/Private/PoseSearchDerivedData.cpp
+++ b/Engine/Plugins/Experimental/Animation/PoseSearch/Source/Runtime/Private/PoseSearchDerivedData.cpp
@@ -522,6 +522,21 @@ static void PreprocessSearchIndexPCAData(FPoseSearchIndex& SearchIndex, int32 Nu
 				const float Error = (ReconstructedValues - MapValues.row(RowIndex)).squaredNorm();
 				check(Error < UE_KINDA_SMALL_NUMBER);
 			}
+
+			TArray<float> ReconstructedPoseValues;
+			ReconstructedPoseValues.SetNumZeroed(NumDimensions);
+			for (int32 PoseIdx = 0; PoseIdx < NumPoses; ++PoseIdx)
+			{
+				SearchIndex.GetReconstructedPoseValues(PoseIdx, ReconstructedPoseValues);
+				TConstArrayView<float> PoseValues = SearchIndex.GetPoseValues(PoseIdx);
+
+				check(ReconstructedPoseValues.Num() == PoseValues.Num());
+				Eigen::Map<const Eigen::ArrayXf> VA(ReconstructedPoseValues.GetData(), ReconstructedPoseValues.Num());
+				Eigen::Map<const Eigen::ArrayXf> VB(PoseValues.GetData(), PoseValues.Num());
+
+				const float Error = (VA - VB).square().sum();
+				check(Error < UE_KINDA_SMALL_NUMBER);
+			}
 		}
 	}
 }
@@ -972,6 +987,12 @@ void FPoseSearchDatabaseAsyncCacheTask::OnGetComplete(UE::DerivedData::FCacheGet
 					return;
 				}
 
+				// removing SearchIndex.Values and relying on FPoseSearchIndex::GetReconstructedPoseValues to reconstruct the Values data from the PCAValues
+				if (Database->PoseSearchMode == EPoseSearchMode::PCAKDTree && Database->KDTreeQueryNumNeighbors <= 1)
+				{
+					SearchIndex.Values.Reset();
+				}
+
 				UE_LOG(LogPoseSearch, Log, TEXT("%s - %s BuildIndex Succeeded"), *LexToString(FullIndexKey.Hash), *Database->GetName());
 
 				// putting SearchIndex to DDC
diff --git a/Engine/Plugins/Experimental/Animation/PoseSearch/Source/Runtime/Private/PoseSearchFeatureChannel_FilterCrashingLegs.cpp b/Engine/Plugins/Experimental/Animation/PoseSearch/Source/Runtime/Private/PoseSearchFeatureChannel_FilterCrashingLegs.cpp
--- a/Engine/Plugins/Experimental/Animation/PoseSearch/Source/Runtime/Private/PoseSearchFeatureChannel_FilterCrashingLegs.cpp
+++ b/Engine/Plugins/Experimental/Animation/PoseSearch/Source/Runtime/Private/PoseSearchFeatureChannel_FilterCrashingLegs.cpp
@@ -53,17 +53,17 @@ void UPoseSearchFeatureChannel_FilterCrashingLegs::BuildQuery(UE::PoseSearch::FS
 {
 	using namespace UE::PoseSearch;
 
-	const bool bIsCurrentResultValid = SearchContext.CurrentResult.IsValid() && SearchContext.CurrentResult.Database->Schema == InOutQuery.GetSchema();
+	const bool bIsCurrentResultValid = SearchContext.GetCurrentResult().IsValid() && SearchContext.GetCurrentResult().Database->Schema == InOutQuery.GetSchema();
 	const bool bSkip = InputQueryPose != EInputQueryPose::UseCharacterPose && bIsCurrentResultValid;
 	
-	if (bSkip || !SearchContext.History)
+	if (bSkip || !SearchContext.GetHistory())
 	{
 		if (bIsCurrentResultValid)
 		{
-			const float LerpValue = InputQueryPose == EInputQueryPose::UseInterpolatedContinuingPose ? SearchContext.CurrentResult.LerpValue : 0.f;
+			const float LerpValue = InputQueryPose == EInputQueryPose::UseInterpolatedContinuingPose ? SearchContext.GetCurrentResult().LerpValue : 0.f;
 			FFeatureVectorHelper::EncodeFloat(InOutQuery.EditValues(), ChannelDataOffset, SearchContext.GetCurrentResultPrevPoseVector(), SearchContext.GetCurrentResultPoseVector(), SearchContext.GetCurrentResultNextPoseVector(), LerpValue);
 		}
-		// else leave the InOutQuery set to zero since the SearchContext.History is invalid and it'll fail if we continue
+		// else leave the InOutQuery set to zero since the SearchContext.GetHistory() is invalid and it'll fail if we continue
 	}
 	else
 	{
diff --git a/Engine/Plugins/Experimental/Animation/PoseSearch/Source/Runtime/Private/PoseSearchFeatureChannel_Heading.cpp b/Engine/Plugins/Experimental/Animation/PoseSearch/Source/Runtime/Private/PoseSearchFeatureChannel_Heading.cpp
--- a/Engine/Plugins/Experimental/Animation/PoseSearch/Source/Runtime/Private/PoseSearchFeatureChannel_Heading.cpp
+++ b/Engine/Plugins/Experimental/Animation/PoseSearch/Source/Runtime/Private/PoseSearchFeatureChannel_Heading.cpp
@@ -45,18 +45,18 @@ void UPoseSearchFeatureChannel_Heading::BuildQuery(UE::PoseSearch::FSearchContex
 {
 	using namespace UE::PoseSearch;
 
-	const bool bIsCurrentResultValid = SearchContext.CurrentResult.IsValid() && SearchContext.CurrentResult.Database->Schema == InOutQuery.GetSchema();
+	const bool bIsCurrentResultValid = SearchContext.GetCurrentResult().IsValid() && SearchContext.GetCurrentResult().Database->Schema == InOutQuery.GetSchema();
 	const bool bSkip = InputQueryPose != EInputQueryPose::UseCharacterPose && bIsCurrentResultValid;
 	const bool bIsRootBone = InOutQuery.GetSchema()->IsRootBone(SchemaBoneIdx);
-	if (bSkip || (!SearchContext.History && !bIsRootBone))
+	if (bSkip || (!SearchContext.GetHistory() && !bIsRootBone))
 	{
 		if (bIsCurrentResultValid)
 		{
-			const float LerpValue = InputQueryPose == EInputQueryPose::UseInterpolatedContinuingPose ? SearchContext.CurrentResult.LerpValue : 0.f;
+			const float LerpValue = InputQueryPose == EInputQueryPose::UseInterpolatedContinuingPose ? SearchContext.GetCurrentResult().LerpValue : 0.f;
 			// @todo: we should normalize if LerpValue != 0
 			FFeatureVectorHelper::EncodeVector(InOutQuery.EditValues(), ChannelDataOffset, SearchContext.GetCurrentResultPrevPoseVector(), SearchContext.GetCurrentResultPoseVector(), SearchContext.GetCurrentResultNextPoseVector(), LerpValue, true, ComponentStripping);
 		}
-		// else leave the InOutQuery set to zero since the SearchContext.History is invalid and it'll fail if we continue
+		// else leave the InOutQuery set to zero since the SearchContext.GetHistory() is invalid and it'll fail if we continue
 	}
 	else
 	{
diff --git a/Engine/Plugins/Experimental/Animation/PoseSearch/Source/Runtime/Private/PoseSearchFeatureChannel_PermutationTime.cpp b/Engine/Plugins/Experimental/Animation/PoseSearch/Source/Runtime/Private/PoseSearchFeatureChannel_PermutationTime.cpp
--- a/Engine/Plugins/Experimental/Animation/PoseSearch/Source/Runtime/Private/PoseSearchFeatureChannel_PermutationTime.cpp
+++ b/Engine/Plugins/Experimental/Animation/PoseSearch/Source/Runtime/Private/PoseSearchFeatureChannel_PermutationTime.cpp
@@ -26,7 +26,7 @@ void UPoseSearchFeatureChannel_PermutationTime::Finalize(UPoseSearchSchema* Sche
 void UPoseSearchFeatureChannel_PermutationTime::BuildQuery(UE::PoseSearch::FSearchContext& SearchContext, FPoseSearchFeatureVectorBuilder& InOutQuery) const
 {
 	using namespace UE::PoseSearch;
-	FFeatureVectorHelper::EncodeFloat(InOutQuery.EditValues(), ChannelDataOffset, SearchContext.DesiredPermutationTimeOffset);
+	FFeatureVectorHelper::EncodeFloat(InOutQuery.EditValues(), ChannelDataOffset, SearchContext.GetDesiredPermutationTimeOffset());
 }
 
 #if WITH_EDITOR
diff --git a/Engine/Plugins/Experimental/Animation/PoseSearch/Source/Runtime/Private/PoseSearchFeatureChannel_Phase.cpp b/Engine/Plugins/Experimental/Animation/PoseSearch/Source/Runtime/Private/PoseSearchFeatureChannel_Phase.cpp
--- a/Engine/Plugins/Experimental/Animation/PoseSearch/Source/Runtime/Private/PoseSearchFeatureChannel_Phase.cpp
+++ b/Engine/Plugins/Experimental/Animation/PoseSearch/Source/Runtime/Private/PoseSearchFeatureChannel_Phase.cpp
@@ -300,16 +300,16 @@ void UPoseSearchFeatureChannel_Phase::BuildQuery(UE::PoseSearch::FSearchContext&
 {
 	using namespace UE::PoseSearch;
 
-	const bool bIsCurrentResultValid = SearchContext.CurrentResult.IsValid() && SearchContext.CurrentResult.Database->Schema == InOutQuery.GetSchema();
+	const bool bIsCurrentResultValid = SearchContext.GetCurrentResult().IsValid() && SearchContext.GetCurrentResult().Database->Schema == InOutQuery.GetSchema();
 	const bool bSkip = InputQueryPose != EInputQueryPose::UseCharacterPose && bIsCurrentResultValid;
-	if (bSkip || !SearchContext.History)
+	if (bSkip || !SearchContext.GetHistory())
 	{
 		if (bIsCurrentResultValid)
 		{
-			const float LerpValue = InputQueryPose == EInputQueryPose::UseInterpolatedContinuingPose ? SearchContext.CurrentResult.LerpValue : 0.f;
+			const float LerpValue = InputQueryPose == EInputQueryPose::UseInterpolatedContinuingPose ? SearchContext.GetCurrentResult().LerpValue : 0.f;
 			FFeatureVectorHelper::EncodeVector2D(InOutQuery.EditValues(), ChannelDataOffset, SearchContext.GetCurrentResultPrevPoseVector(), SearchContext.GetCurrentResultPoseVector(), SearchContext.GetCurrentResultNextPoseVector(), LerpValue);
 		}
-		// else leave the InOutQuery set to zero since the SearchContext.History is invalid and it'll fail if we continue
+		// else leave the InOutQuery set to zero since the SearchContext.GetHistory() is invalid and it'll fail if we continue
 	}
 	else
 	{
diff --git a/Engine/Plugins/Experimental/Animation/PoseSearch/Source/Runtime/Private/PoseSearchFeatureChannel_Position.cpp b/Engine/Plugins/Experimental/Animation/PoseSearch/Source/Runtime/Private/PoseSearchFeatureChannel_Position.cpp
--- a/Engine/Plugins/Experimental/Animation/PoseSearch/Source/Runtime/Private/PoseSearchFeatureChannel_Position.cpp
+++ b/Engine/Plugins/Experimental/Animation/PoseSearch/Source/Runtime/Private/PoseSearchFeatureChannel_Position.cpp
@@ -61,17 +61,17 @@ void UPoseSearchFeatureChannel_Position::BuildQuery(UE::PoseSearch::FSearchConte
 	using namespace UE::PoseSearch;
 
 	check(InOutQuery.GetSchema());
-	const bool bIsCurrentResultValid = SearchContext.CurrentResult.IsValid() && SearchContext.CurrentResult.Database->Schema == InOutQuery.GetSchema();
+	const bool bIsCurrentResultValid = SearchContext.GetCurrentResult().IsValid() && SearchContext.GetCurrentResult().Database->Schema == InOutQuery.GetSchema();
 	const bool bSkip = InputQueryPose != EInputQueryPose::UseCharacterPose && bIsCurrentResultValid;
 	const bool bIsRootBone = InOutQuery.GetSchema()->IsRootBone(SchemaBoneIdx);
-	if (bSkip || (!SearchContext.History && !bIsRootBone))
+	if (bSkip || (!SearchContext.GetHistory() && !bIsRootBone))
 	{
 		if (bIsCurrentResultValid)
 		{
-			const float LerpValue = InputQueryPose == EInputQueryPose::UseInterpolatedContinuingPose ? SearchContext.CurrentResult.LerpValue : 0.f;
+			const float LerpValue = InputQueryPose == EInputQueryPose::UseInterpolatedContinuingPose ? SearchContext.GetCurrentResult().LerpValue : 0.f;
 			FFeatureVectorHelper::EncodeVector(InOutQuery.EditValues(), ChannelDataOffset, SearchContext.GetCurrentResultPrevPoseVector(), SearchContext.GetCurrentResultPoseVector(), SearchContext.GetCurrentResultNextPoseVector(), LerpValue, false, ComponentStripping);
 		}
-		// else leave the InOutQuery set to zero since the SearchContext.History is invalid and it'll fail if we continue
+		// else leave the InOutQuery set to zero since the SearchContext.GetHistory() is invalid and it'll fail if we continue
 	}
 	else
 	{
diff --git a/Engine/Plugins/Experimental/Animation/PoseSearch/Source/Runtime/Private/PoseSearchFeatureChannel_Velocity.cpp b/Engine/Plugins/Experimental/Animation/PoseSearch/Source/Runtime/Private/PoseSearchFeatureChannel_Velocity.cpp
--- a/Engine/Plugins/Experimental/Animation/PoseSearch/Source/Runtime/Private/PoseSearchFeatureChannel_Velocity.cpp
+++ b/Engine/Plugins/Experimental/Animation/PoseSearch/Source/Runtime/Private/PoseSearchFeatureChannel_Velocity.cpp
@@ -31,18 +31,18 @@ void UPoseSearchFeatureChannel_Velocity::BuildQuery(UE::PoseSearch::FSearchConte
 {
 	using namespace UE::PoseSearch;
 
-	const bool bIsCurrentResultValid = SearchContext.CurrentResult.IsValid() && SearchContext.CurrentResult.Database->Schema == InOutQuery.GetSchema();
+	const bool bIsCurrentResultValid = SearchContext.GetCurrentResult().IsValid() && SearchContext.GetCurrentResult().Database->Schema == InOutQuery.GetSchema();
 	const bool bSkip = InputQueryPose != EInputQueryPose::UseCharacterPose && bIsCurrentResultValid;
 	const bool bIsRootBone = InOutQuery.GetSchema()->IsRootBone(SchemaBoneIdx);
-	if (bSkip || (!SearchContext.History && !bIsRootBone))
+	if (bSkip || (!SearchContext.GetHistory() && !bIsRootBone))
 	{
 		if (bIsCurrentResultValid)
 		{
-			const float LerpValue = InputQueryPose == EInputQueryPose::UseInterpolatedContinuingPose ? SearchContext.CurrentResult.LerpValue : 0.f;
+			const float LerpValue = InputQueryPose == EInputQueryPose::UseInterpolatedContinuingPose ? SearchContext.GetCurrentResult().LerpValue : 0.f;
 			// @todo: we should normalize if EPoseSearchVelocityFlags::Normalize && LerpValue != 0
 			FFeatureVectorHelper::EncodeVector(InOutQuery.EditValues(), ChannelDataOffset, SearchContext.GetCurrentResultPrevPoseVector(), SearchContext.GetCurrentResultPoseVector(), SearchContext.GetCurrentResultNextPoseVector(), LerpValue, false, ComponentStripping);
 		}
-		// else leave the InOutQuery set to zero since the SearchContext.History is invalid and it'll fail if we continue
+		// else leave the InOutQuery set to zero since the SearchContext.GetHistory() is invalid and it'll fail if we continue
 	}
 	else
 	{
diff --git a/Engine/Plugins/Experimental/Animation/PoseSearch/Source/Runtime/Private/PoseSearchHistory.cpp b/Engine/Plugins/Experimental/Animation/PoseSearch/Source/Runtime/Private/PoseSearchHistory.cpp
--- a/Engine/Plugins/Experimental/Animation/PoseSearch/Source/Runtime/Private/PoseSearchHistory.cpp
+++ b/Engine/Plugins/Experimental/Animation/PoseSearch/Source/Runtime/Private/PoseSearchHistory.cpp
@@ -298,6 +298,11 @@ void FExtendedPoseHistory::Init(const FPoseHistory* InPoseHistory)
 	PoseHistory = InPoseHistory;
 }
 
+bool FExtendedPoseHistory::IsInitialized() const
+{
+	return PoseHistory != nullptr;
+}
+
 float FExtendedPoseHistory::GetSampleTimeInterval() const
 {
 	check(PoseHistory);
diff --git a/Engine/Plugins/Experimental/Animation/PoseSearch/Source/Runtime/Private/PoseSearchIndex.cpp b/Engine/Plugins/Experimental/Animation/PoseSearch/Source/Runtime/Private/PoseSearchIndex.cpp
--- a/Engine/Plugins/Experimental/Animation/PoseSearch/Source/Runtime/Private/PoseSearchIndex.cpp
+++ b/Engine/Plugins/Experimental/Animation/PoseSearch/Source/Runtime/Private/PoseSearchIndex.cpp
@@ -143,24 +143,62 @@ void FPoseSearchIndex::Reset()
 TConstArrayView<float> FPoseSearchIndex::GetPoseValues(int32 PoseIdx) const
 {
 	const int32 SchemaCardinality = WeightsSqrt.Num();
-	check(PoseIdx >= 0 && PoseIdx < GetNumPoses() && SchemaCardinality > 0);
+	check(!Values.IsEmpty() && PoseIdx >= 0 && PoseIdx < GetNumPoses() && SchemaCardinality > 0);
 	const int32 ValueOffset = PoseIdx * SchemaCardinality;
 	return MakeArrayView(&Values[ValueOffset], SchemaCardinality);
 }
 
-TConstArrayView<float> FPoseSearchIndex::GetPoseValuesSafe(int32 PoseIdx) const
+TConstArrayView<float> FPoseSearchIndex::GetReconstructedPoseValues(int32 PoseIdx, TArrayView<float> BufferUsedForReconstruction) const
 {
+	using namespace UE::PoseSearch;
+
+	const int32 NumDimensions = WeightsSqrt.Num();
+	const int32 NumPoses = GetNumPoses();
+	check(PoseIdx >= 0 && PoseIdx < NumPoses&& NumDimensions > 0);
+	check(BufferUsedForReconstruction.Num() == NumDimensions);
+
+	const int32 NumberOfPrincipalComponents = PCAValues.Num() / NumPoses;
+	check(NumPoses * NumberOfPrincipalComponents == PCAValues.Num());
+
+	const RowMajorVectorMapConst MapWeightsSqrt(WeightsSqrt.GetData(), 1, NumDimensions);
+	const ColMajorMatrixMapConst MapPCAProjectionMatrix(PCAProjectionMatrix.GetData(), NumDimensions, NumberOfPrincipalComponents);
+	const RowMajorVectorMapConst MapMean(Mean.GetData(), 1, NumDimensions);
+	const RowMajorMatrixMapConst MapPCAValues(PCAValues.GetData(), NumPoses, NumberOfPrincipalComponents);
+
+	const RowMajorVector ReciprocalWeightsSqrt = MapWeightsSqrt.cwiseInverse();
+	const RowMajorVector WeightedReconstructedValues = MapPCAValues.row(PoseIdx) * MapPCAProjectionMatrix.transpose() + MapMean;
+
+	RowMajorVectorMap ReconstructedPoseValues(BufferUsedForReconstruction.GetData(), 1, NumDimensions);
+	ReconstructedPoseValues = WeightedReconstructedValues.array() * ReciprocalWeightsSqrt.array();
+
+	return BufferUsedForReconstruction;
+}
+
+#if WITH_EDITOR || ENABLE_DRAW_DEBUG
+TArray<float> FPoseSearchIndex::GetPoseValuesSafe(int32 PoseIdx) const
+{
+	TArray<float> PoseValues;
 	if (PoseIdx >= 0 && PoseIdx < GetNumPoses())
 	{
-		return GetPoseValues(PoseIdx);
+		if (Values.IsEmpty())
+		{
+			const int32 NumDimensions = WeightsSqrt.Num();
+			PoseValues.SetNumUninitialized(NumDimensions);
+			GetReconstructedPoseValues(PoseIdx, PoseValues);
+		}
+		else
+		{
+			PoseValues = GetPoseValues(PoseIdx);
+		}
 	}
-	return TConstArrayView<float>();
+	return PoseValues;
 }
+#endif // WITH_EDITOR || ENABLE_DRAW_DEBUG
 
-FPoseSearchCost FPoseSearchIndex::ComparePoses(int32 PoseIdx, EPoseSearchBooleanRequest QueryMirrorRequest, UE::PoseSearch::EPoseComparisonFlags PoseComparisonFlags, float MirrorMismatchCostBias, TConstArrayView<float> QueryValues) const
+FPoseSearchCost FPoseSearchIndex::ComparePoses(int32 PoseIdx, EPoseSearchBooleanRequest QueryMirrorRequest, UE::PoseSearch::EPoseComparisonFlags PoseComparisonFlags, float MirrorMismatchCostBias, TConstArrayView<float> PoseValues, TConstArrayView<float> QueryValues) const
 {
 	// base dissimilarity cost representing how the associated PoseIdx differ, in a weighted way, from the query pose (QueryValues)
-	const float DissimilarityCost = UE::PoseSearch::CompareFeatureVectors(GetPoseValues(PoseIdx), QueryValues, WeightsSqrt);
+	const float DissimilarityCost = UE::PoseSearch::CompareFeatureVectors(PoseValues, QueryValues, WeightsSqrt);
 
 	// cost addend associated to a mismatch in mirror state between query and analyzed PoseIdx
 	float MirrorMismatchAddend = 0.f;
diff --git a/Engine/Plugins/Experimental/Animation/PoseSearch/Source/Runtime/Private/PoseSearchLibrary.cpp b/Engine/Plugins/Experimental/Animation/PoseSearch/Source/Runtime/Private/PoseSearchLibrary.cpp
--- a/Engine/Plugins/Experimental/Animation/PoseSearch/Source/Runtime/Private/PoseSearchLibrary.cpp
+++ b/Engine/Plugins/Experimental/Animation/PoseSearch/Source/Runtime/Private/PoseSearchLibrary.cpp
@@ -261,11 +261,11 @@ void UPoseSearchLibrary::TraceMotionMatchingState(
 		TraceState.CurrentPoseEntryIdx = PoseEntryIdx;
 	}
 
-	if (DeltaTime > SMALL_NUMBER && SearchContext.Trajectory)
+	if (DeltaTime > SMALL_NUMBER && SearchContext.GetTrajectory())
 	{
 		// simulation
-		const FTrajectorySample PrevSample = SearchContext.Trajectory->GetSampleAtTime(-DeltaTime);
-		const FTrajectorySample CurrSample = SearchContext.Trajectory->GetSampleAtTime(0.f);
+		const FTrajectorySample PrevSample = SearchContext.GetTrajectory()->GetSampleAtTime(-DeltaTime);
+		const FTrajectorySample CurrSample = SearchContext.GetTrajectory()->GetSampleAtTime(0.f);
 
 		const FTransform SimDelta = CurrSample.Transform.GetRelativeTransform(PrevSample.Transform);
 
@@ -315,27 +315,23 @@ void UPoseSearchLibrary::UpdateMotionMatchingState(
 	const FSearchResult LastResult = InOutMotionMatchingState.CurrentSearchResult;
 #endif
 
-	// Build the search context
-	FSearchContext SearchContext;
-	SearchContext.Trajectory = &Trajectory;
-	SearchContext.bForceInterrupt = bForceInterrupt;
-	SearchContext.bCanAdvance = InOutMotionMatchingState.CanAdvance(DeltaTime);
-	SearchContext.CurrentResult = InOutMotionMatchingState.CurrentSearchResult;
-	SearchContext.PoseJumpThresholdTime = Settings.PoseJumpThresholdTime;
-	SearchContext.PoseIndicesHistory = &InOutMotionMatchingState.PoseIndicesHistory;
-
+	const IPoseHistory* History = nullptr;
 	if (IPoseHistoryProvider* PoseHistoryProvider = Context.GetMessage<IPoseHistoryProvider>())
 	{
-		SearchContext.History = &PoseHistoryProvider->GetPoseHistory();
+		History = &PoseHistoryProvider->GetPoseHistory();
 	}
-
+	
+	EPoseSearchBooleanRequest QueryMirrorRequest = EPoseSearchBooleanRequest::Indifferent;
 	if (const FPoseSearchIndexAsset* CurrentIndexAsset = InOutMotionMatchingState.CurrentSearchResult.GetSearchIndexAsset())
 	{
-		SearchContext.QueryMirrorRequest = CurrentIndexAsset->bMirrored ? EPoseSearchBooleanRequest::TrueValue : EPoseSearchBooleanRequest::FalseValue;
+		QueryMirrorRequest = CurrentIndexAsset->bMirrored ? EPoseSearchBooleanRequest::TrueValue : EPoseSearchBooleanRequest::FalseValue;
 	}
 
+	FSearchContext SearchContext(&Trajectory, History, 0.f, &InOutMotionMatchingState.PoseIndicesHistory, QueryMirrorRequest, 
+		InOutMotionMatchingState.CurrentSearchResult, Settings.PoseJumpThresholdTime, bForceInterrupt, InOutMotionMatchingState.CanAdvance(DeltaTime));
+
 	// If we can't advance or enough time has elapsed since the last pose jump then search
-	const bool bSearch = !SearchContext.bCanAdvance || (InOutMotionMatchingState.ElapsedPoseSearchTime >= Settings.SearchThrottleTime);
+	const bool bSearch = !SearchContext.CanAdvance() || (InOutMotionMatchingState.ElapsedPoseSearchTime >= Settings.SearchThrottleTime);
 	if (bSearch)
 	{
 		InOutMotionMatchingState.ElapsedPoseSearchTime = 0.f;
@@ -343,18 +339,25 @@ void UPoseSearchLibrary::UpdateMotionMatchingState(
 		// Evaluate continuing pose
 		FPoseSearchCost ContinuingPoseCost;
 		FPoseSearchFeatureVectorBuilder ContinuingPoseComposedQuery;
-		if (!SearchContext.bForceInterrupt && SearchContext.bCanAdvance && SearchContext.CurrentResult.Database.IsValid()
+		if (!SearchContext.IsForceInterrupt() && SearchContext.CanAdvance() && SearchContext.GetCurrentResult().Database.IsValid()
 #if WITH_EDITOR
-			&& FAsyncPoseSearchDatabasesManagement::RequestAsyncBuildIndex(SearchContext.CurrentResult.Database.Get(), ERequestAsyncBuildFlag::ContinueRequest)
+			&& FAsyncPoseSearchDatabasesManagement::RequestAsyncBuildIndex(SearchContext.GetCurrentResult().Database.Get(), ERequestAsyncBuildFlag::ContinueRequest)
 #endif
 		)
 		{
-			const UPoseSearchDatabase* ContinuingPoseDatabase = SearchContext.CurrentResult.Database.Get();
+			const UPoseSearchDatabase* ContinuingPoseDatabase = SearchContext.GetCurrentResult().Database.Get();
 			if (ensure(ContinuingPoseDatabase->Schema))
 			{	
 				SearchContext.GetOrBuildQuery(ContinuingPoseDatabase->Schema, ContinuingPoseComposedQuery);
-				ContinuingPoseCost = ContinuingPoseDatabase->GetSearchIndex().ComparePoses(SearchContext.CurrentResult.PoseIdx, SearchContext.QueryMirrorRequest,
-					EPoseComparisonFlags::ContinuingPose, ContinuingPoseDatabase->Schema->MirrorMismatchCostBias, ContinuingPoseComposedQuery.GetValues());
+
+				const FPoseSearchIndex& SearchIndex = ContinuingPoseDatabase->GetSearchIndex();
+				const int32 PoseIdx = SearchContext.GetCurrentResult().PoseIdx;
+				const int32 NumDimensions = ContinuingPoseDatabase->Schema->SchemaCardinality;
+				TArrayView<float> ReconstructedPoseValuesBuffer((float*)FMemory_Alloca(NumDimensions * sizeof(float)), NumDimensions);
+				const TConstArrayView<float> PoseValues = SearchIndex.Values.IsEmpty() ? SearchIndex.GetReconstructedPoseValues(PoseIdx, ReconstructedPoseValuesBuffer) : SearchIndex.GetPoseValues(PoseIdx);
+
+				ContinuingPoseCost = ContinuingPoseDatabase->GetSearchIndex().ComparePoses(SearchContext.GetCurrentResult().PoseIdx, SearchContext.GetQueryMirrorRequest(),
+					EPoseComparisonFlags::ContinuingPose, ContinuingPoseDatabase->Schema->MirrorMismatchCostBias, PoseValues, ContinuingPoseComposedQuery.GetValues());
 				SearchContext.UpdateCurrentBestCost(ContinuingPoseCost);
 			}
 		}
@@ -454,11 +457,7 @@ void UPoseSearchLibrary::MotionMatch(
 
 	if (Searchable)
 	{
-		// Build the search context
-		FSearchContext SearchContext;
-		SearchContext.DesiredPermutationTimeOffset = TimeToFutureAnimationStart;
-		SearchContext.Trajectory = &Trajectory;
-
+		// ExtendedPoseHistory will hold future poses to match AssetSamplerBase (at FutureAnimationStartTime) TimeToFutureAnimationStart seconds in the future
 		FExtendedPoseHistory ExtendedPoseHistory;
 		if (AnimInstance)
 		{
@@ -468,14 +467,12 @@ void UPoseSearchLibrary::MotionMatch(
 				{
 					if (const FAnimNode_PoseSearchHistoryCollector_Base* PoseHistoryNode = TagSubsystem->FindNodeByTag<FAnimNode_PoseSearchHistoryCollector_Base>(PoseHistoryName, AnimInstance))
 					{
-						const FPoseHistory* PoseHistory = &PoseHistoryNode->GetPoseHistory();
-						ExtendedPoseHistory.Init(PoseHistory);
-						SearchContext.History = PoseHistory;
+						ExtendedPoseHistory.Init(&PoseHistoryNode->GetPoseHistory());
 					}
 				}
 			}
 
-			if (!SearchContext.History)
+			if (!ExtendedPoseHistory.IsInitialized())
 			{
 				if (FutureAnimation)
 				{
@@ -525,9 +522,6 @@ void UPoseSearchLibrary::MotionMatch(
 
 				if (AssetSamplerBase)
 				{
-					// ExtendedPoseHistory will hold future poses to match AssetSamplerBase (at FutureAnimationStartTime) TimeToFutureAnimationStart seconds in the future
-					SearchContext.History = &ExtendedPoseHistory;
-
 					FCompactPose Pose;
 					FBlendedCurve UnusedCurve;
 					FStackAttributeContainer UnusedAtrribute;
@@ -569,6 +563,8 @@ void UPoseSearchLibrary::MotionMatch(
 		}
 
 		// @todo: finish set up SearchContext by exposing or calculating additional members
+		FSearchContext SearchContext(&Trajectory, ExtendedPoseHistory.IsInitialized() ? &ExtendedPoseHistory : nullptr, TimeToFutureAnimationStart);
+
 		FSearchResult SearchResult = Searchable->Search(SearchContext);
 		if (SearchResult.IsValid())
 		{
diff --git a/Engine/Plugins/Experimental/Animation/PoseSearch/Source/Runtime/Public/PoseSearch/PoseSearchContext.h b/Engine/Plugins/Experimental/Animation/PoseSearch/Source/Runtime/Public/PoseSearch/PoseSearchContext.h
--- a/Engine/Plugins/Experimental/Animation/PoseSearch/Source/Runtime/Public/PoseSearch/PoseSearchContext.h
+++ b/Engine/Plugins/Experimental/Animation/PoseSearch/Source/Runtime/Public/PoseSearch/PoseSearchContext.h
@@ -137,14 +137,9 @@ private:
 
 struct POSESEARCH_API FSearchContext
 {
-	EPoseSearchBooleanRequest QueryMirrorRequest = EPoseSearchBooleanRequest::Indifferent;
-	const IPoseHistory* History = nullptr;
-	const FTrajectorySampleRange* Trajectory = nullptr;
-	FSearchResult CurrentResult;
-	float PoseJumpThresholdTime = 0.f;
-	bool bForceInterrupt = false;
-	// can the continuing pose advance? (if not we skip evaluating it)
-	bool bCanAdvance = true;
+	FSearchContext(const FTrajectorySampleRange* InTrajectory, const IPoseHistory* InHistory, float InDesiredPermutationTimeOffset,
+		const FPoseIndicesHistory* InPoseIndicesHistory = nullptr, EPoseSearchBooleanRequest InQueryMirrorRequest = EPoseSearchBooleanRequest::Indifferent,
+		const FSearchResult& InCurrentResult = FSearchResult(), float InPoseJumpThresholdTime = 0.f, bool bInForceInterrupt = false, bool bInCanAdvance = true);
 
 	FQuat GetSampleRotation(float SampleTimeOffset, const UPoseSearchSchema* Schema, int8 SchemaSampleBoneIdx = RootSchemaBoneIdx, int8 SchemaOriginBoneIdx = RootSchemaBoneIdx, bool bUseHistoryRoot = false, EPermutationTimeType PermutationTimeType = EPermutationTimeType::UseSampleTime);
 	FVector GetSamplePosition(float SampleTimeOffset, const UPoseSearchSchema* Schema, int8 SchemaSampleBoneIdx = RootSchemaBoneIdx, int8 SchemaOriginBoneIdx = RootSchemaBoneIdx, bool bUseHistoryRoot = false, EPermutationTimeType PermutationTimeType = EPermutationTimeType::UseSampleTime);
@@ -161,19 +156,42 @@ struct POSESEARCH_API FSearchContext
 
 	bool IsCurrentResultFromDatabase(const UPoseSearchDatabase* Database) const;
 
-	TConstArrayView<float> GetCurrentResultPrevPoseVector() const;
-	TConstArrayView<float> GetCurrentResultPoseVector() const;
-	TConstArrayView<float> GetCurrentResultNextPoseVector() const;
-
-	float DesiredPermutationTimeOffset = 0.f;
+	TConstArrayView<float> GetCurrentResultPrevPoseVector() const { return CurrentResultPrevPoseVector; }
+	TConstArrayView<float> GetCurrentResultPoseVector() const { return CurrentResultPoseVector; }
+	TConstArrayView<float> GetCurrentResultNextPoseVector() const { return CurrentResultNextPoseVector; }
 
-	const FPoseIndicesHistory* PoseIndicesHistory = nullptr;
+	const FSearchResult& GetCurrentResult() const { return CurrentResult; }
+	float GetPoseJumpThresholdTime() const { return PoseJumpThresholdTime; }
+	EPoseSearchBooleanRequest GetQueryMirrorRequest() const { return QueryMirrorRequest; }
+	const FPoseIndicesHistory* GetPoseIndicesHistory() const { return PoseIndicesHistory; }
+	const IPoseHistory* GetHistory() const { return History; }
+	float GetDesiredPermutationTimeOffset() const { return DesiredPermutationTimeOffset; }
+	const FTrajectorySampleRange* GetTrajectory() const { return Trajectory; }
+	bool CanAdvance() const { return bCanAdvance; }
+	bool IsForceInterrupt() const { return bForceInterrupt; }
 
 private:
 	FTransform GetTransform(float SampleTime, const UPoseSearchSchema* Schema, int8 SchemaBoneIdx = RootSchemaBoneIdx, bool bUseHistoryRoot = false);
 	FTransform GetComponentSpaceTransform(float SampleTime, const UPoseSearchSchema* Schema, int8 SchemaSampleBoneIdx = RootSchemaBoneIdx);
 	FVector GetSamplePositionInternal(float SampleTime, float OriginTime, const UPoseSearchSchema* Schema, int8 SchemaSampleBoneIdx = RootSchemaBoneIdx, int8 SchemaOriginBoneIdx = RootSchemaBoneIdx, bool bUseHistoryRoot = false);
 
+	const FTrajectorySampleRange* Trajectory = nullptr;
+	const IPoseHistory* History = nullptr;
+	float DesiredPermutationTimeOffset = 0.f;
+	const FPoseIndicesHistory* PoseIndicesHistory = nullptr;
+	EPoseSearchBooleanRequest QueryMirrorRequest = EPoseSearchBooleanRequest::Indifferent;
+	FSearchResult CurrentResult;
+	float PoseJumpThresholdTime = 0.f;
+	bool bForceInterrupt = false;
+	// can the continuing pose advance? (if not we skip evaluating it)
+	bool bCanAdvance = true;
+
+	// cached previous, current, and next pose data
+	TConstArrayView<float> CurrentResultNextPoseVector;
+	TConstArrayView<float> CurrentResultPoseVector;
+	TConstArrayView<float> CurrentResultPrevPoseVector;
+	TArray<float> CurrentResultPoseVectorData;
+
 	// transforms cached in component space
 	FCachedTransforms<FTransform> CachedTransforms;
 	TArray<FPoseSearchFeatureVectorBuilder, TInlineAllocator<8>> CachedQueries;
diff --git a/Engine/Plugins/Experimental/Animation/PoseSearch/Source/Runtime/Public/PoseSearch/PoseSearchHistory.h b/Engine/Plugins/Experimental/Animation/PoseSearch/Source/Runtime/Public/PoseSearch/PoseSearchHistory.h
--- a/Engine/Plugins/Experimental/Animation/PoseSearch/Source/Runtime/Public/PoseSearch/PoseSearchHistory.h
+++ b/Engine/Plugins/Experimental/Animation/PoseSearch/Source/Runtime/Public/PoseSearch/PoseSearchHistory.h
@@ -69,6 +69,8 @@ struct FExtendedPoseHistory : public IPoseHistory
 {
 	void Init(const FPoseHistory* InPoseHistory);
 
+	bool IsInitialized() const;
+
 	// IPoseHistory interface
 	virtual float GetSampleTimeInterval() const override;
 	virtual bool GetComponentSpaceTransformAtTime(float Time, FBoneIndexType BoneIndexType, FTransform& OutBoneTransform, bool bExtrapolate = true) const override;
diff --git a/Engine/Plugins/Experimental/Animation/PoseSearch/Source/Runtime/Public/PoseSearch/PoseSearchIndex.h b/Engine/Plugins/Experimental/Animation/PoseSearch/Source/Runtime/Public/PoseSearch/PoseSearchIndex.h
--- a/Engine/Plugins/Experimental/Animation/PoseSearch/Source/Runtime/Public/PoseSearch/PoseSearchIndex.h
+++ b/Engine/Plugins/Experimental/Animation/PoseSearch/Source/Runtime/Public/PoseSearch/PoseSearchIndex.h
@@ -226,8 +226,11 @@ struct POSESEARCH_API FPoseSearchIndex : public FPoseSearchIndexBase
 
 	void Reset();
 	TConstArrayView<float> GetPoseValues(int32 PoseIdx) const;
-	TConstArrayView<float> GetPoseValuesSafe(int32 PoseIdx) const;
-	FPoseSearchCost ComparePoses(int32 PoseIdx, EPoseSearchBooleanRequest QueryMirrorRequest, UE::PoseSearch::EPoseComparisonFlags PoseComparisonFlags, float MirrorMismatchCostBias, TConstArrayView<float> QueryValues) const;
+	TConstArrayView<float> GetReconstructedPoseValues(int32 PoseIdx, TArrayView<float> BufferUsedForReconstruction) const;
+#if WITH_EDITOR || ENABLE_DRAW_DEBUG
+	TArray<float> GetPoseValuesSafe(int32 PoseIdx) const;
+#endif // WITH_EDITOR || ENABLE_DRAW_DEBUG
+	FPoseSearchCost ComparePoses(int32 PoseIdx, EPoseSearchBooleanRequest QueryMirrorRequest, UE::PoseSearch::EPoseComparisonFlags PoseComparisonFlags, float MirrorMismatchCostBias, TConstArrayView<float> PoseValues, TConstArrayView<float> QueryValues) const;
 
 	friend FArchive& operator<<(FArchive& Ar, FPoseSearchIndex& Index);
 };
