motion matching - database indexing fix when using normalization set

diff --git a/Engine/Plugins/Experimental/Animation/PoseSearch/Source/Runtime/Private/PoseSearchDerivedData.cpp b/Engine/Plugins/Experimental/Animation/PoseSearch/Source/Runtime/Private/PoseSearchDerivedData.cpp
--- a/Engine/Plugins/Experimental/Animation/PoseSearch/Source/Runtime/Private/PoseSearchDerivedData.cpp
+++ b/Engine/Plugins/Experimental/Animation/PoseSearch/Source/Runtime/Private/PoseSearchDerivedData.cpp
@@ -269,7 +269,7 @@ static void FindValidSequenceIntervals(const UAnimSequenceBase* SequenceBase, FF
 	}
 }
 
-static void InitSearchIndexAssets(FPoseSearchIndexBase& SearchIndex, UPoseSearchDatabase* Database)
+static void InitSearchIndexAssets(FPoseSearchIndexBase& SearchIndex, const UPoseSearchDatabase* Database)
 {
 	using namespace UE::PoseSearch;
 
@@ -901,7 +901,7 @@ void FPoseSearchDatabaseAsyncCacheTask::OnGetComplete(UE::DerivedData::FCacheGet
 				Schemas.AddDefaulted(IndexBaseDatabases.Num());
 				for (int32 IndexBaseIdx = 0; IndexBaseIdx < IndexBaseDatabases.Num(); ++IndexBaseIdx)
 				{
-					auto IndexBaseDatabase = IndexBaseDatabases[IndexBaseIdx];
+					TWeakObjectPtr<const UPoseSearchDatabase> IndexBaseDatabase = IndexBaseDatabases[IndexBaseIdx];
 					FPoseSearchIndexBase& SearchIndexBase = SearchIndexBases[IndexBaseIdx];
 					Schemas[IndexBaseIdx] = IndexBaseDatabase->Schema;
 
@@ -928,7 +928,7 @@ void FPoseSearchDatabaseAsyncCacheTask::OnGetComplete(UE::DerivedData::FCacheGet
 					}
 
 					// Building all the related FPoseSearchBaseIndex first
-					InitSearchIndexAssets(SearchIndexBase, Database.Get());
+					InitSearchIndexAssets(SearchIndexBase, IndexBaseDatabase.Get());
 
 					if (Owner.IsCanceled())
 					{
