motion matching - UPoseSearchLibrary::MotionMatch: removed unnecessary UE_LOG(s) for not set FutureAnimation

diff --git a/Engine/Plugins/Experimental/Animation/PoseSearch/Source/Runtime/Private/PoseSearchLibrary.cpp b/Engine/Plugins/Experimental/Animation/PoseSearch/Source/Runtime/Private/PoseSearchLibrary.cpp
--- a/Engine/Plugins/Experimental/Animation/PoseSearch/Source/Runtime/Private/PoseSearchLibrary.cpp
+++ b/Engine/Plugins/Experimental/Animation/PoseSearch/Source/Runtime/Private/PoseSearchLibrary.cpp
@@ -442,19 +442,6 @@ void UPoseSearchLibrary::MotionMatch(
 	bIsMirrored = false;
 	BlendParameters = FVector::ZeroVector;
 
-	if (FutureAnimationStartTime < FiniteDelta)
-	{
-		UE_LOG(LogPoseSearch, Warning, TEXT("UPoseSearchLibrary::MotionMatch - provided FutureAnimationStartTime (%f) is too small to be able to calculate velocities. Clamping it to minimum value of %f"), TimeToFutureAnimationStart, FiniteDelta);
-		FutureAnimationStartTime = FiniteDelta;
-	}
-
-	const float MinTimeToFutureAnimationStart = FiniteDelta + UE_KINDA_SMALL_NUMBER;
-	if (TimeToFutureAnimationStart < MinTimeToFutureAnimationStart)
-	{
-		UE_LOG(LogPoseSearch, Warning, TEXT("UPoseSearchLibrary::MotionMatch - provided TimeToFutureAnimationStart (%f) is too small. Clamping it to minimum value of %f"), TimeToFutureAnimationStart, MinTimeToFutureAnimationStart);
-		TimeToFutureAnimationStart = MinTimeToFutureAnimationStart;
-	}
-
 	if (Searchable)
 	{
 		// ExtendedPoseHistory will hold future poses to match AssetSamplerBase (at FutureAnimationStartTime) TimeToFutureAnimationStart seconds in the future
@@ -530,6 +517,19 @@ void UPoseSearchLibrary::MotionMatch(
 					UnusedCurve.InitFrom(BoneContainer);
 					Pose.SetBoneContainer(&BoneContainer);
 
+					if (FutureAnimationStartTime < FiniteDelta)
+					{
+						UE_LOG(LogPoseSearch, Warning, TEXT("UPoseSearchLibrary::MotionMatch - provided FutureAnimationStartTime (%f) is too small to be able to calculate velocities. Clamping it to minimum value of %f"), FutureAnimationStartTime, FiniteDelta);
+						FutureAnimationStartTime = FiniteDelta;
+					}
+
+					const float MinTimeToFutureAnimationStart = FiniteDelta + UE_KINDA_SMALL_NUMBER;
+					if (TimeToFutureAnimationStart < MinTimeToFutureAnimationStart)
+					{
+						UE_LOG(LogPoseSearch, Warning, TEXT("UPoseSearchLibrary::MotionMatch - provided TimeToFutureAnimationStart (%f) is too small. Clamping it to minimum value of %f"), TimeToFutureAnimationStart, MinTimeToFutureAnimationStart);
+						TimeToFutureAnimationStart = MinTimeToFutureAnimationStart;
+					}
+
 					// extracting 2 poses to be able to calculate velocities
 					for (int i = 0; i < 2; ++i)
 					{
