motion matching - database metadata memory improvements

diff --git a/Engine/Plugins/Experimental/Animation/PoseSearch/Source/Editor/Private/PoseSearchDatabaseEditorReflection.cpp b/Engine/Plugins/Experimental/Animation/PoseSearch/Source/Editor/Private/PoseSearchDatabaseEditorReflection.cpp
--- a/Engine/Plugins/Experimental/Animation/PoseSearch/Source/Editor/Private/PoseSearchDatabaseEditorReflection.cpp
+++ b/Engine/Plugins/Experimental/Animation/PoseSearch/Source/Editor/Private/PoseSearchDatabaseEditorReflection.cpp
@@ -123,7 +123,10 @@ void UPoseSearchDatabaseStatistics::Initialize(const UPoseSearchDatabase* PoseSe
 		uint32 NumOfSearchablePoses = 0;
 		for (const FPoseSearchPoseMetadata& PoseMetadata : SearchIndex.PoseMetadata)
 		{
-			NumOfSearchablePoses += PoseMetadata.Flags != EPoseSearchPoseFlags::BlockTransition;
+			if (!PoseMetadata.IsBlockTransition())
+			{
+				++NumOfSearchablePoses;
+			}
 		}
 		SearchableFrames = NumOfSearchablePoses;
 		SearchableTime = FText::Format(TimeFormat, static_cast<double>(NumOfSearchablePoses) / SampleRate);
diff --git a/Engine/Plugins/Experimental/Animation/PoseSearch/Source/Runtime/Private/PoseSearchAssetIndexer.cpp b/Engine/Plugins/Experimental/Animation/PoseSearch/Source/Runtime/Private/PoseSearchAssetIndexer.cpp
--- a/Engine/Plugins/Experimental/Animation/PoseSearch/Source/Runtime/Private/PoseSearchAssetIndexer.cpp
+++ b/Engine/Plugins/Experimental/Animation/PoseSearch/Source/Runtime/Private/PoseSearchAssetIndexer.cpp
@@ -110,7 +110,7 @@ void FAssetIndexer::Process(int32 AssetIdx)
 		const float SampleTime = FMath::Min(SampleIdx * IndexingContext.Schema->GetSamplingInterval(), SequenceLength);
 		float CostAddend = IndexingContext.Schema->BaseCostBias;
 		float ContinuingPoseCostAddend = IndexingContext.Schema->ContinuingPoseCostBias;
-		EPoseSearchPoseFlags Flags = EPoseSearchPoseFlags::None;
+		bool bBlockTransition = false;
 
 		TArray<UAnimNotifyState_PoseSearchBase*> NotifyStates;
 		IndexingContext.AssetSampler->ExtractPoseSearchNotifyStates(SampleTime, NotifyStates);
@@ -118,16 +118,12 @@ void FAssetIndexer::Process(int32 AssetIdx)
 		{
 			if (PoseSearchNotify->GetClass()->IsChildOf<UAnimNotifyState_PoseSearchBlockTransition>())
 			{
-				EnumAddFlags(Flags, EPoseSearchPoseFlags::BlockTransition);
+				bBlockTransition = true;
 			}
 			else if (const UAnimNotifyState_PoseSearchModifyCost* ModifyCostNotify = Cast<const UAnimNotifyState_PoseSearchModifyCost>(PoseSearchNotify))
 			{
 				CostAddend = ModifyCostNotify->CostAddend;
 			}
-			else if (const UAnimNotifyState_PoseSearchOverrideContinuingPoseCostBias* ContinuingPoseCostBias = Cast<const UAnimNotifyState_PoseSearchOverrideContinuingPoseCostBias>(PoseSearchNotify))
-			{
-				ContinuingPoseCostAddend = ContinuingPoseCostBias->CostAddend;
-			}
 		}
 
 		if (IndexingContext.AssetSampler->IsLoopable())
@@ -135,12 +131,7 @@ void FAssetIndexer::Process(int32 AssetIdx)
 			CostAddend += IndexingContext.Schema->LoopingCostBias;
 		}
 
-		const int32 VectorIdx = GetVectorIdx(SampleIdx);
-		FPoseSearchPoseMetadata& Metadata = PoseMetadata[VectorIdx];
-		Metadata.Flags = Flags;
-		Metadata.CostAddend = CostAddend;
-		Metadata.ContinuingPoseCostAddend = ContinuingPoseCostAddend;
-		Metadata.AssetIndex = AssetIdx;
+		PoseMetadata[GetVectorIdx(SampleIdx)].Init(AssetIdx, bBlockTransition, CostAddend);
 	}
 
 	// Generate pose features data
diff --git a/Engine/Plugins/Experimental/Animation/PoseSearch/Source/Runtime/Private/PoseSearchDatabase.cpp b/Engine/Plugins/Experimental/Animation/PoseSearch/Source/Runtime/Private/PoseSearchDatabase.cpp
--- a/Engine/Plugins/Experimental/Animation/PoseSearch/Source/Runtime/Private/PoseSearchDatabase.cpp
+++ b/Engine/Plugins/Experimental/Animation/PoseSearch/Source/Runtime/Private/PoseSearchDatabase.cpp
@@ -68,7 +68,7 @@ static void PopulateNonSelectableIdx(FNonSelectableIdx& NonSelectableIdx, FSearc
 
 #if UE_POSE_SEARCH_TRACE_ENABLED
 				const TArray<float> PoseValues = SearchIndex.GetPoseValuesSafe(PoseIdx);
-				const FPoseSearchCost PoseCost = SearchIndex.ComparePoses(PoseIdx, SearchContext.GetQueryMirrorRequest(), EPoseComparisonFlags::None, Database->Schema->MirrorMismatchCostBias, PoseValues, QueryValues);
+				const FPoseSearchCost PoseCost = SearchIndex.ComparePoses(PoseIdx, SearchContext.GetQueryMirrorRequest(), 0.f, Database->Schema->MirrorMismatchCostBias, PoseValues, QueryValues);
 				SearchContext.BestCandidates.Add(PoseCost, PoseIdx, Database, EPoseCandidateFlags::DiscardedBy_PoseJumpThresholdTime);
 #endif
 			}
@@ -101,7 +101,7 @@ static void PopulateNonSelectableIdx(FNonSelectableIdx& NonSelectableIdx, FSearc
 				NonSelectableIdx.Add(PoseIdx);
 
 #if UE_POSE_SEARCH_TRACE_ENABLED
-				const FPoseSearchCost PoseCost = SearchIndex.ComparePoses(PoseIdx, SearchContext.GetQueryMirrorRequest(), EPoseComparisonFlags::None, Database->Schema->MirrorMismatchCostBias, SearchIndex.GetPoseValuesSafe(PoseIdx), QueryValues);
+				const FPoseSearchCost PoseCost = SearchIndex.ComparePoses(PoseIdx, SearchContext.GetQueryMirrorRequest(), 0.f, Database->Schema->MirrorMismatchCostBias, SearchIndex.GetPoseValuesSafe(PoseIdx), QueryValues);
 				SearchContext.BestCandidates.Add(PoseCost, PoseIdx, Database, EPoseCandidateFlags::DiscardedBy_PoseJumpThresholdTime);
 #endif
 			}
@@ -128,7 +128,7 @@ static void PopulateNonSelectableIdx(FNonSelectableIdx& NonSelectableIdx, FSearc
 				// if we're editing the database and removing assets it's possible that the PoseIndicesHistory contains invalid pose indexes
 				if (HistoricalPoseIndex.PoseIndex < SearchIndex.GetNumPoses())
 				{
-					const FPoseSearchCost PoseCost = SearchIndex.ComparePoses(HistoricalPoseIndex.PoseIndex, SearchContext.GetQueryMirrorRequest(), EPoseComparisonFlags::None, Database->Schema->MirrorMismatchCostBias, SearchIndex.GetPoseValuesSafe(HistoricalPoseIndex.PoseIndex), QueryValues);
+					const FPoseSearchCost PoseCost = SearchIndex.ComparePoses(HistoricalPoseIndex.PoseIndex, SearchContext.GetQueryMirrorRequest(), 0.f, Database->Schema->MirrorMismatchCostBias, SearchIndex.GetPoseValuesSafe(HistoricalPoseIndex.PoseIndex), QueryValues);
 					SearchContext.BestCandidates.Add(PoseCost, HistoricalPoseIndex.PoseIndex, Database, EPoseCandidateFlags::DiscardedBy_PoseReselectHistory);
 				}
 #endif
@@ -141,11 +141,11 @@ static void PopulateNonSelectableIdx(FNonSelectableIdx& NonSelectableIdx, FSearc
 
 struct FPoseFilters
 {
-	FPoseFilters(const UPoseSearchSchema* Schema, TConstArrayView<size_t> NonSelectableIdx, EPoseSearchPoseFlags OverallFlags)
+	FPoseFilters(const UPoseSearchSchema* Schema, TConstArrayView<size_t> NonSelectableIdx, bool bAnyBlockTransition)
 	{
 		NonSelectableIdxPoseFilter.NonSelectableIdx = NonSelectableIdx;
 
-		if (EnumHasAnyFlags(OverallFlags, EPoseSearchPoseFlags::BlockTransition))
+		if (bAnyBlockTransition)
 		{
 			AllPoseFilters.Add(&BlockTransitionPoseFilter);
 		}
@@ -181,12 +181,12 @@ struct FPoseFilters
 				}
 				else if (PoseFilter == &BlockTransitionPoseFilter)
 				{
-					const FPoseSearchCost PoseCost = SearchIndex.ComparePoses(PoseIdx, SearchContext.GetQueryMirrorRequest(), EPoseComparisonFlags::None, Database->Schema->MirrorMismatchCostBias, PoseValues, QueryValues);
+					const FPoseSearchCost PoseCost = SearchIndex.ComparePoses(PoseIdx, SearchContext.GetQueryMirrorRequest(), 0.f, Database->Schema->MirrorMismatchCostBias, PoseValues, QueryValues);
 					SearchContext.BestCandidates.Add(PoseCost, PoseIdx, Database, EPoseCandidateFlags::DiscardedBy_BlockTransition);
 				}
 				else
 				{
-					const FPoseSearchCost PoseCost = SearchIndex.ComparePoses(PoseIdx, SearchContext.GetQueryMirrorRequest(), EPoseComparisonFlags::None, Database->Schema->MirrorMismatchCostBias, PoseValues, QueryValues);
+					const FPoseSearchCost PoseCost = SearchIndex.ComparePoses(PoseIdx, SearchContext.GetQueryMirrorRequest(), 0.f, Database->Schema->MirrorMismatchCostBias, PoseValues, QueryValues);
 					SearchContext.BestCandidates.Add(PoseCost, PoseIdx, Database, EPoseCandidateFlags::DiscardedBy_PoseFilter);
 				}
 #endif
@@ -221,7 +221,7 @@ private:
 
 		virtual bool IsPoseValid(TConstArrayView<float> PoseValues, TConstArrayView<float> QueryValues, int32 PoseIdx, const FPoseSearchPoseMetadata& Metadata) const override
 		{
-			return !EnumHasAnyFlags(Metadata.Flags, EPoseSearchPoseFlags::BlockTransition);
+			return !Metadata.IsBlockTransition();
 		}
 	};
 
@@ -830,7 +830,7 @@ UE::PoseSearch::FSearchResult UPoseSearchDatabase::SearchPCAKDTree(UE::PoseSearc
 		SearchIndex.KDTree.FindNeighbors(ResultSet, ProjectedQueryValues.data());
 
 		// NonSelectableIdx are already filtered out inside the kdtree search
-		const FPoseFilters PoseFilters(Schema, TConstArrayView<size_t>(), SearchIndex.OverallFlags);
+		const FPoseFilters PoseFilters(Schema, TConstArrayView<size_t>(), SearchIndex.bAnyBlockTransition);
 		for (size_t ResultIndex = 0; ResultIndex < ResultSet.Num(); ++ResultIndex)
 		{
 			const int32 PoseIdx = ResultIndexes[ResultIndex];
@@ -842,7 +842,7 @@ UE::PoseSearch::FSearchResult UPoseSearchDatabase::SearchPCAKDTree(UE::PoseSearc
 #endif
 			))
 			{
-				const FPoseSearchCost PoseCost = SearchIndex.ComparePoses(PoseIdx, SearchContext.GetQueryMirrorRequest(), EPoseComparisonFlags::None, Schema->MirrorMismatchCostBias, PoseValues, QueryValues);
+				const FPoseSearchCost PoseCost = SearchIndex.ComparePoses(PoseIdx, SearchContext.GetQueryMirrorRequest(), 0.f, Schema->MirrorMismatchCostBias, PoseValues, QueryValues);
 				if (PoseCost < Result.PoseCost)
 				{
 					Result.PoseCost = PoseCost;
@@ -904,7 +904,7 @@ UE::PoseSearch::FSearchResult UPoseSearchDatabase::SearchBruteForce(UE::PoseSear
 
 		const int32 NumDimensions = Schema->SchemaCardinality;
 		TArrayView<float> ReconstructedPoseValuesBuffer((float*)FMemory_Alloca(NumDimensions * sizeof(float)), NumDimensions);
-		const FPoseFilters PoseFilters(Schema, NonSelectableIdx, SearchIndex.OverallFlags);
+		const FPoseFilters PoseFilters(Schema, NonSelectableIdx, SearchIndex.bAnyBlockTransition);
 		for (int32 PoseIdx = 0; PoseIdx < SearchIndex.GetNumPoses(); ++PoseIdx)
 		{
 			const TConstArrayView<float> PoseValues = SearchIndex.Values.IsEmpty() ? SearchIndex.GetReconstructedPoseValues(PoseIdx, ReconstructedPoseValuesBuffer) : SearchIndex.GetPoseValues(PoseIdx);
@@ -914,7 +914,7 @@ UE::PoseSearch::FSearchResult UPoseSearchDatabase::SearchBruteForce(UE::PoseSear
 #endif
 			))
 			{
-				const FPoseSearchCost PoseCost = SearchIndex.ComparePoses(PoseIdx, SearchContext.GetQueryMirrorRequest(), EPoseComparisonFlags::None, Schema->MirrorMismatchCostBias, PoseValues, QueryValues);
+				const FPoseSearchCost PoseCost = SearchIndex.ComparePoses(PoseIdx, SearchContext.GetQueryMirrorRequest(), 0.f, Schema->MirrorMismatchCostBias, PoseValues, QueryValues);
 				if (PoseCost < Result.PoseCost)
 				{
 					Result.PoseCost = PoseCost;
diff --git a/Engine/Plugins/Experimental/Animation/PoseSearch/Source/Runtime/Private/PoseSearchDatabaseIndexingContext.cpp b/Engine/Plugins/Experimental/Animation/PoseSearch/Source/Runtime/Private/PoseSearchDatabaseIndexingContext.cpp
--- a/Engine/Plugins/Experimental/Animation/PoseSearch/Source/Runtime/Private/PoseSearchDatabaseIndexingContext.cpp
+++ b/Engine/Plugins/Experimental/Animation/PoseSearch/Source/Runtime/Private/PoseSearchDatabaseIndexingContext.cpp
@@ -229,10 +229,14 @@ bool FDatabaseIndexingContext::IndexDatabase(FPoseSearchIndexBase& SearchIndexBa
 	}
 
 	// Joining Metadata.Flags into OverallFlags
-	SearchIndexBase.OverallFlags = EPoseSearchPoseFlags::None;
+	SearchIndexBase.bAnyBlockTransition = false;
 	for (const FPoseSearchPoseMetadata& Metadata : SearchIndexBase.PoseMetadata)
 	{
-		SearchIndexBase.OverallFlags |= Metadata.Flags;
+		if (Metadata.IsBlockTransition())
+		{
+			SearchIndexBase.bAnyBlockTransition = true;
+			break;
+		}
 	}
 
 	// Joining Stats
@@ -263,9 +267,9 @@ bool FDatabaseIndexingContext::IndexDatabase(FPoseSearchIndexBase& SearchIndexBa
 		SearchIndexBase.MinCostAddend = MAX_FLT;
 		for (const FPoseSearchPoseMetadata& PoseMetadata : SearchIndexBase.PoseMetadata)
 		{
-			if (PoseMetadata.CostAddend < SearchIndexBase.MinCostAddend)
+			if (PoseMetadata.GetCostAddend() < SearchIndexBase.MinCostAddend)
 			{
-				SearchIndexBase.MinCostAddend = PoseMetadata.CostAddend;
+				SearchIndexBase.MinCostAddend = PoseMetadata.GetCostAddend();
 			}
 		}
 	}
diff --git a/Engine/Plugins/Experimental/Animation/PoseSearch/Source/Runtime/Private/PoseSearchDerivedDataKey.cpp b/Engine/Plugins/Experimental/Animation/PoseSearch/Source/Runtime/Private/PoseSearchDerivedDataKey.cpp
--- a/Engine/Plugins/Experimental/Animation/PoseSearch/Source/Runtime/Private/PoseSearchDerivedDataKey.cpp
+++ b/Engine/Plugins/Experimental/Animation/PoseSearch/Source/Runtime/Private/PoseSearchDerivedDataKey.cpp
@@ -41,7 +41,7 @@ FKeyBuilder::FKeyBuilder(const UObject* Object, bool bUseDataVer)
 	if (bUseDataVer)
 	{
 		// used to invalidate the key without having to change POSESEARCHDB_DERIVEDDATA_VER all the times
-		int32 POSESEARCHDB_DERIVEDDATA_VER_SMALL = 131;
+		int32 POSESEARCHDB_DERIVEDDATA_VER_SMALL = 137;
 		FGuid VersionGuid = FDevSystemGuids::GetSystemGuid(FDevSystemGuids::Get().POSESEARCHDB_DERIVEDDATA_VER);
 
 		*this << VersionGuid;
diff --git a/Engine/Plugins/Experimental/Animation/PoseSearch/Source/Runtime/Private/PoseSearchIndex.cpp b/Engine/Plugins/Experimental/Animation/PoseSearch/Source/Runtime/Private/PoseSearchIndex.cpp
--- a/Engine/Plugins/Experimental/Animation/PoseSearch/Source/Runtime/Private/PoseSearchIndex.cpp
+++ b/Engine/Plugins/Experimental/Animation/PoseSearch/Source/Runtime/Private/PoseSearchIndex.cpp
@@ -35,10 +35,8 @@ void CompareFeatureVectors(TConstArrayView<float> A, TConstArrayView<float> B, T
 // FPoseSearchPoseMetadata
 FArchive& operator<<(FArchive& Ar, FPoseSearchPoseMetadata& Metadata)
 {
-	Ar << Metadata.Flags;
+	Ar << Metadata.Data;
 	Ar << Metadata.CostAddend;
-	Ar << Metadata.ContinuingPoseCostAddend;
-	Ar << Metadata.AssetIndex;
 	return Ar;
 }
 
@@ -71,7 +69,7 @@ FArchive& operator<<(FArchive& Ar, FPoseSearchStats& Stats)
 // FPoseSearchBaseIndex
 const FPoseSearchIndexAsset& FPoseSearchIndexBase::GetAssetForPose(int32 PoseIdx) const
 {
-	const int32 AssetIndex = PoseMetadata[PoseIdx].AssetIndex;
+	const uint32 AssetIndex = PoseMetadata[PoseIdx].GetAssetIndex();
 	return Assets[AssetIndex];
 }
 
@@ -79,7 +77,7 @@ const FPoseSearchIndexAsset* FPoseSearchIndexBase::GetAssetForPoseSafe(int32 Pos
 {
 	if (PoseMetadata.IsValidIndex(PoseIdx))
 	{
-		const int32 AssetIndex = PoseMetadata[PoseIdx].AssetIndex;
+		const uint32 AssetIndex = PoseMetadata[PoseIdx].GetAssetIndex();
 		if (Assets.IsValidIndex(AssetIndex))
 		{
 			return &Assets[AssetIndex];
@@ -102,7 +100,7 @@ FArchive& operator<<(FArchive& Ar, FPoseSearchIndexBase& Index)
 {
 	Ar << Index.Values;
 	Ar << Index.PoseMetadata;
-	Ar << Index.OverallFlags;
+	Ar << Index.bAnyBlockTransition;
 	Ar << Index.Assets;
 	Ar << Index.MinCostAddend;
 	Ar << Index.Stats;
@@ -193,7 +191,7 @@ TArray<float> FPoseSearchIndex::GetPoseValuesSafe(int32 PoseIdx) const
 	return PoseValues;
 }
 
-FPoseSearchCost FPoseSearchIndex::ComparePoses(int32 PoseIdx, EPoseSearchBooleanRequest QueryMirrorRequest, UE::PoseSearch::EPoseComparisonFlags PoseComparisonFlags, float MirrorMismatchCostBias, TConstArrayView<float> PoseValues, TConstArrayView<float> QueryValues) const
+FPoseSearchCost FPoseSearchIndex::ComparePoses(int32 PoseIdx, EPoseSearchBooleanRequest QueryMirrorRequest, float ContinuingPoseCostBias, float MirrorMismatchCostBias, TConstArrayView<float> PoseValues, TConstArrayView<float> QueryValues) const
 {
 	// base dissimilarity cost representing how the associated PoseIdx differ, in a weighted way, from the query pose (QueryValues)
 	const float DissimilarityCost = UE::PoseSearch::CompareFeatureVectors(PoseValues, QueryValues, WeightsSqrt);
@@ -212,15 +210,9 @@ FPoseSearchCost FPoseSearchIndex::ComparePoses(int32 PoseIdx, EPoseSearchBoolean
 		}
 	}
 
-	const FPoseSearchPoseMetadata& PoseIdxMetadata = PoseMetadata[PoseIdx];
-
 	// cost addend associated to Schema->BaseCostBias or overriden by UAnimNotifyState_PoseSearchModifyCost
-	const float NotifyAddend = PoseIdxMetadata.CostAddend;
-
-	// cost addend associated to Schema->ContinuingPoseCostBias or overriden by UAnimNotifyState_PoseSearchOverrideContinuingPoseCostBias
-	const float ContinuingPoseCostAddend = EnumHasAnyFlags(PoseComparisonFlags, UE::PoseSearch::EPoseComparisonFlags::ContinuingPose) ? PoseIdxMetadata.ContinuingPoseCostAddend : 0.f;
-
-	return FPoseSearchCost(DissimilarityCost, NotifyAddend, MirrorMismatchAddend, ContinuingPoseCostAddend);
+	const float NotifyAddend = PoseMetadata[PoseIdx].GetCostAddend();
+	return FPoseSearchCost(DissimilarityCost, NotifyAddend, MirrorMismatchAddend, ContinuingPoseCostBias);
 }
 
 FArchive& operator<<(FArchive& Ar, FPoseSearchIndex& Index)
diff --git a/Engine/Plugins/Experimental/Animation/PoseSearch/Source/Runtime/Private/PoseSearchLibrary.cpp b/Engine/Plugins/Experimental/Animation/PoseSearch/Source/Runtime/Private/PoseSearchLibrary.cpp
--- a/Engine/Plugins/Experimental/Animation/PoseSearch/Source/Runtime/Private/PoseSearchLibrary.cpp
+++ b/Engine/Plugins/Experimental/Animation/PoseSearch/Source/Runtime/Private/PoseSearchLibrary.cpp
@@ -6,12 +6,14 @@
 #include "Animation/AnimNode_Inertialization.h"
 #include "Animation/AnimNode_SequencePlayer.h"
 #include "Animation/AnimRootMotionProvider.h"
+#include "Animation/AnimComposite.h"
 #include "Animation/AnimSequence.h"
 #include "Animation/AnimSubsystem_Tag.h"
 #include "Animation/BlendSpace.h"
 #include "InstancedStruct.h"
 #include "PoseSearch/AnimNode_MotionMatching.h"
 #include "PoseSearch/AnimNode_PoseSearchHistoryCollector.h"
+#include "PoseSearch/PoseSearchAnimNotifies.h"
 #include "PoseSearch/PoseSearchDatabase.h"
 #include "PoseSearch/PoseSearchDerivedData.h"
 #include "PoseSearch/PoseSearchSchema.h"
@@ -356,8 +358,58 @@ void UPoseSearchLibrary::UpdateMotionMatchingState(
 				TArrayView<float> ReconstructedPoseValuesBuffer((float*)FMemory_Alloca(NumDimensions * sizeof(float)), NumDimensions);
 				const TConstArrayView<float> PoseValues = SearchIndex.Values.IsEmpty() ? SearchIndex.GetReconstructedPoseValues(PoseIdx, ReconstructedPoseValuesBuffer) : SearchIndex.GetPoseValues(PoseIdx);
 
+				const float SampleTime = ContinuingPoseDatabase->GetAssetTime(PoseIdx);
+
+				const FPoseSearchIndexAsset& SearchIndexAsset = SearchIndex.GetAssetForPose(PoseIdx);
+				const FInstancedStruct& DatabaseAssetStruct = ContinuingPoseDatabase->GetAnimationAssetStruct(SearchIndexAsset);
+
+				// @todo: create a unified sampler to sample any animation asset
+				// @todo: change ExtractPoseSearchNotifyStates api to avoid NotifyStates allocation
+
+				// extracting notifies from the database animation asset at time SampleTime to search for UAnimNotifyState_PoseSearchOverrideContinuingPoseCostBias eventually overriding the schema ContinuingPoseCostBias
+				TArray<UAnimNotifyState_PoseSearchBase*> NotifyStates;
+				if (const FPoseSearchDatabaseSequence* DatabaseSequence = DatabaseAssetStruct.GetPtr<FPoseSearchDatabaseSequence>())
+				{
+					FSequenceBaseSampler SequenceBaseSampler;
+					FSequenceBaseSampler::FInput Input;
+					Input.SequenceBase = DatabaseSequence->Sequence;
+					SequenceBaseSampler.Init(Input);
+					SequenceBaseSampler.ExtractPoseSearchNotifyStates(SampleTime, NotifyStates);
+				}
+				else if (const FPoseSearchDatabaseAnimComposite* DatabaseAnimComposite = DatabaseAssetStruct.GetPtr<FPoseSearchDatabaseAnimComposite>())
+				{
+					FSequenceBaseSampler SequenceBaseSampler;
+					FSequenceBaseSampler::FInput Input;
+					Input.SequenceBase = DatabaseAnimComposite->AnimComposite;
+					SequenceBaseSampler.Init(Input);
+					SequenceBaseSampler.ExtractPoseSearchNotifyStates(SampleTime, NotifyStates);
+				}
+				else if (const FPoseSearchDatabaseBlendSpace* DatabaseBlendSpace = DatabaseAssetStruct.GetPtr<FPoseSearchDatabaseBlendSpace>())
+				{
+					FBlendSpaceSampler BlendSpaceSampler;
+					FBlendSpaceSampler::FInput Input;
+					Input.BlendSpace = DatabaseBlendSpace->BlendSpace;
+					Input.BlendParameters = SearchIndexAsset.BlendParameters;
+					BlendSpaceSampler.Init(Input);
+					BlendSpaceSampler.ExtractPoseSearchNotifyStates(SampleTime, NotifyStates);
+				}
+				else
+				{
+					checkNoEntry();
+				}
+
+				float ContinuingPoseCostBias = ContinuingPoseDatabase->Schema->ContinuingPoseCostBias;
+				for (const UAnimNotifyState_PoseSearchBase* PoseSearchNotify : NotifyStates)
+				{
+					if (const UAnimNotifyState_PoseSearchOverrideContinuingPoseCostBias* ContinuingPoseCostBiasNotify = Cast<const UAnimNotifyState_PoseSearchOverrideContinuingPoseCostBias>(PoseSearchNotify))
+					{
+						ContinuingPoseCostBias = ContinuingPoseCostBiasNotify->CostAddend;
+						break;
+					}
+				}
+
 				ContinuingPoseCost = ContinuingPoseDatabase->GetSearchIndex().ComparePoses(SearchContext.GetCurrentResult().PoseIdx, SearchContext.GetQueryMirrorRequest(),
-					EPoseComparisonFlags::ContinuingPose, ContinuingPoseDatabase->Schema->MirrorMismatchCostBias, PoseValues, ContinuingPoseComposedQuery.GetValues());
+					ContinuingPoseCostBias, ContinuingPoseDatabase->Schema->MirrorMismatchCostBias, PoseValues, ContinuingPoseComposedQuery.GetValues());
 				SearchContext.UpdateCurrentBestCost(ContinuingPoseCost);
 			}
 		}
diff --git a/Engine/Plugins/Experimental/Animation/PoseSearch/Source/Runtime/Public/PoseSearch/PoseSearchIndex.h b/Engine/Plugins/Experimental/Animation/PoseSearch/Source/Runtime/Public/PoseSearch/PoseSearchIndex.h
--- a/Engine/Plugins/Experimental/Animation/PoseSearch/Source/Runtime/Public/PoseSearch/PoseSearchIndex.h
+++ b/Engine/Plugins/Experimental/Animation/PoseSearch/Source/Runtime/Public/PoseSearch/PoseSearchIndex.h
@@ -12,13 +12,6 @@ namespace UE::PoseSearch
 float CompareFeatureVectors(TConstArrayView<float> A, TConstArrayView<float> B, TConstArrayView<float> WeightsSqrt);
 POSESEARCH_API void CompareFeatureVectors(TConstArrayView<float> A, TConstArrayView<float> B, TConstArrayView<float> WeightsSqrt, TArrayView<float> Result);
 
-enum class EPoseComparisonFlags : int32
-{
-	None = 0,
-	ContinuingPose = 1 << 0,
-};
-ENUM_CLASS_FLAGS(EPoseComparisonFlags);
-
 } // namespace UE::PoseSearch
 
 UENUM()
@@ -32,16 +25,6 @@ enum class EPoseSearchBooleanRequest : uint8
 	Invalid = Num UMETA(Hidden)
 };
 
-UENUM()
-enum class EPoseSearchPoseFlags : uint8
-{
-	None = 0,
-
-	// Don't return this pose as a search result
-	BlockTransition = 1 << 0,
-};
-ENUM_CLASS_FLAGS(EPoseSearchPoseFlags);
-
 /**
  * This is kept for each pose in the search index along side the feature vector values and is used to influence the search.
  */
@@ -50,22 +33,37 @@ struct POSESEARCH_API FPoseSearchPoseMetadata
 {
 	GENERATED_BODY()
 
+private:
+	// bits 0-30 represent the AssetIndex, bit 31 represents bBlockTransition
 	UPROPERTY(meta = (NeverInHash))
-	EPoseSearchPoseFlags Flags = EPoseSearchPoseFlags::None;
+	uint32 Data = 0;
 
-	// @todo: consider float16
 	UPROPERTY(meta = (NeverInHash))
-	float CostAddend = 0.0f;
+	float CostAddend = 0.f;
 
-	// @todo: consider float16
-	UPROPERTY(meta = (NeverInHash))
-	float ContinuingPoseCostAddend = 0.0f;
+public:
+	void Init(uint32 AssetIndex, bool bBlockTransition, float InCostAddend)
+	{
+		check((AssetIndex & (1 << 31)) == 0);
+		Data = AssetIndex | (bBlockTransition ? 1 << 31 : 0);
 
-	// @todo: consider uint16
-	UPROPERTY(meta = (NeverInHash))
-	int32 AssetIndex = INDEX_NONE;
+		CostAddend = InCostAddend;
+	}
+
+	bool IsBlockTransition() const
+	{
+		return Data & (1 << 31);
+	}
 
-	bool IsInitialized() const { return AssetIndex != INDEX_NONE; }
+	uint32 GetAssetIndex() const
+	{
+		return Data & ~(1 << 31);
+	}
+
+	float GetCostAddend() const
+	{
+		return CostAddend;
+	}
 
 	friend FArchive& operator<<(FArchive& Ar, FPoseSearchPoseMetadata& Metadata);
 };
@@ -160,7 +158,7 @@ struct POSESEARCH_API FPoseSearchIndexBase
 	TArray<FPoseSearchPoseMetadata> PoseMetadata;
 
 	UPROPERTY(meta = (NeverInHash))
-	EPoseSearchPoseFlags OverallFlags = EPoseSearchPoseFlags::None;
+	bool bAnyBlockTransition = false;
 
 	UPROPERTY(meta = (NeverInHash))
 	TArray<FPoseSearchIndexAsset> Assets;
@@ -228,7 +226,7 @@ struct POSESEARCH_API FPoseSearchIndex : public FPoseSearchIndexBase
 	TConstArrayView<float> GetPoseValues(int32 PoseIdx) const;
 	TConstArrayView<float> GetReconstructedPoseValues(int32 PoseIdx, TArrayView<float> BufferUsedForReconstruction) const;
 	TArray<float> GetPoseValuesSafe(int32 PoseIdx) const;
-	FPoseSearchCost ComparePoses(int32 PoseIdx, EPoseSearchBooleanRequest QueryMirrorRequest, UE::PoseSearch::EPoseComparisonFlags PoseComparisonFlags, float MirrorMismatchCostBias, TConstArrayView<float> PoseValues, TConstArrayView<float> QueryValues) const;
+	FPoseSearchCost ComparePoses(int32 PoseIdx, EPoseSearchBooleanRequest QueryMirrorRequest, float ContinuingPoseCostBias, float MirrorMismatchCostBias, TConstArrayView<float> PoseValues, TConstArrayView<float> QueryValues) const;
 
 	friend FArchive& operator<<(FArchive& Ar, FPoseSearchIndex& Index);
 };
