motion matching - sharing a unified FAnimationAssetSampler across all the supported animation asset to simplify code complexity and reduce memory allocations

diff --git a/Engine/Plugins/Experimental/Animation/PoseSearch/Source/Editor/Private/PoseSearchDatabaseViewModel.cpp b/Engine/Plugins/Experimental/Animation/PoseSearch/Source/Editor/Private/PoseSearchDatabaseViewModel.cpp
--- a/Engine/Plugins/Experimental/Animation/PoseSearch/Source/Editor/Private/PoseSearchDatabaseViewModel.cpp
+++ b/Engine/Plugins/Experimental/Animation/PoseSearch/Source/Editor/Private/PoseSearchDatabaseViewModel.cpp
@@ -36,14 +36,7 @@ namespace UE::PoseSearch
 
 	void FDatabasePreviewActor::Process()
 	{
-		check(Sampler);
-		Sampler->Process();
-	}
-
-	const FAssetSamplerBase* FDatabasePreviewActor::GetSampler() const
-	{
-		check(Sampler);
-		return Sampler.Get();
+		Sampler.Process();
 	}
 
 	UDebugSkelMeshComponent* FDatabasePreviewActor::GetDebugSkelMeshComponent()
@@ -195,49 +188,9 @@ namespace UE::PoseSearch
 		}
 
 		const FInstancedStruct& DatabaseAsset = PoseSearchDatabase->GetAnimationAssetStruct(IndexAsset);
-		if (const FPoseSearchDatabaseSequence* DatabaseSequence = DatabaseAsset.GetPtr<FPoseSearchDatabaseSequence>())
-		{
-			FSequenceBaseSampler::FInput Input;
-			Input.SequenceBase = DatabaseSequence->Sequence;
-
-			TSharedPtr<FSequenceBaseSampler> Sampler = MakeShared<FSequenceBaseSampler>();
-			Sampler->Init(Input);
-			PreviewActor.Sampler = Sampler;
-		}
-		else if (const FPoseSearchDatabaseAnimComposite* DatabaseAnimComposite = DatabaseAsset.GetPtr<FPoseSearchDatabaseAnimComposite>())
-		{
-			FSequenceBaseSampler::FInput Input;
-			Input.SequenceBase = DatabaseAnimComposite->AnimComposite;
-
-			TSharedPtr<FSequenceBaseSampler> Sampler = MakeShared<FSequenceBaseSampler>();
-			Sampler->Init(Input);
-			PreviewActor.Sampler = Sampler;
-		}
-		else if (const FPoseSearchDatabaseBlendSpace* DatabaseBlendSpace = DatabaseAsset.GetPtr<FPoseSearchDatabaseBlendSpace>())
-		{
-			FBlendSpaceSampler::FInput Input;
-			Input.BoneContainer = BoneContainer;
-			Input.BlendSpace = DatabaseBlendSpace->BlendSpace;
-			Input.BlendParameters = IndexAsset.BlendParameters;
-
-			TSharedPtr<FBlendSpaceSampler> Sampler = MakeShared<FBlendSpaceSampler>();
-			Sampler->Init(Input);
-			PreviewActor.Sampler = Sampler;
-		}
-		else if (const FPoseSearchDatabaseAnimMontage* DatabaseAnimMontage = DatabaseAsset.GetPtr<FPoseSearchDatabaseAnimMontage>())
-		{
-			FAnimMontageSampler::FInput Input;
-			Input.AnimMontage = DatabaseAnimMontage->AnimMontage;
-
-			TSharedPtr<FAnimMontageSampler> Sampler = MakeShared<FAnimMontageSampler>();
-			Sampler->Init(Input);
-			PreviewActor.Sampler = Sampler;
-		}
-		else
-		{
-			checkNoEntry();
-		}
-
+		const FPoseSearchDatabaseAnimationAssetBase* DatabaseAnimationAssetBase = DatabaseAsset.GetPtr<FPoseSearchDatabaseAnimationAssetBase>();
+		check(DatabaseAnimationAssetBase);
+		PreviewActor.Sampler.Init(DatabaseAnimationAssetBase->GetAnimationAsset(), IndexAsset.BlendParameters, BoneContainer);
 		PreviewActor.IndexAssetIndex = IndexAssetIndex;
 		PreviewActor.CurrentPoseIndex = INDEX_NONE;
 
@@ -319,8 +272,7 @@ namespace UE::PoseSearch
 			}
 
 			const UAnimationAsset* PreviewAsset = AnimInstance->GetAnimationAsset();
-			const FAssetSamplerBase* Sampler = PreviewActor.GetSampler();
-			if (!PreviewAsset || !Sampler)
+			if (!PreviewAsset || !PreviewActor.Sampler.IsInitialized())
 			{
 				continue;
 			}
@@ -340,13 +292,13 @@ namespace UE::PoseSearch
 			}
 
 			// quantized time to scaled quantized time
-			const float CurrentScaledTime = Sampler->GetScaledTime(CurrentTime);
+			const float CurrentScaledTime = PreviewActor.Sampler.GetScaledTime(CurrentTime);
 
 			AnimInstance->SetPosition(CurrentScaledTime);
 			AnimInstance->SetPlayRate(0.f);
 			AnimInstance->SetBlendSpacePosition(IndexAsset.BlendParameters);
 
-			FTransform RootMotion = Sampler->ExtractRootTransform(CurrentTime);
+			FTransform RootMotion = PreviewActor.Sampler.ExtractRootTransform(CurrentTime);
 			if (AnimInstance->GetMirrorDataTable())
 			{
 				RootMotion = MirrorRootMotion(RootMotion, AnimInstance->GetMirrorDataTable());
diff --git a/Engine/Plugins/Experimental/Animation/PoseSearch/Source/Editor/Public/PoseSearchDatabaseViewModel.h b/Engine/Plugins/Experimental/Animation/PoseSearch/Source/Editor/Public/PoseSearchDatabaseViewModel.h
--- a/Engine/Plugins/Experimental/Animation/PoseSearch/Source/Editor/Public/PoseSearchDatabaseViewModel.h
+++ b/Engine/Plugins/Experimental/Animation/PoseSearch/Source/Editor/Public/PoseSearchDatabaseViewModel.h
@@ -45,12 +45,10 @@ namespace UE::PoseSearch
 		TWeakObjectPtr<AActor> Actor = nullptr;
 		int32 IndexAssetIndex = INDEX_NONE;
 		int32 CurrentPoseIndex = INDEX_NONE;
-
-		TSharedPtr<FAssetSamplerBase> Sampler;
+		FAnimationAssetSampler Sampler;
 
 		bool IsValid() const;
 		void Process();
-		const FAssetSamplerBase* GetSampler() const;
 		UDebugSkelMeshComponent* GetDebugSkelMeshComponent();
 		UAnimPreviewInstance* GetAnimPreviewInstance();
 	};
diff --git a/Engine/Plugins/Experimental/Animation/PoseSearch/Source/Runtime/Private/PoseSearchAssetIndexer.cpp b/Engine/Plugins/Experimental/Animation/PoseSearch/Source/Runtime/Private/PoseSearchAssetIndexer.cpp
--- a/Engine/Plugins/Experimental/Animation/PoseSearch/Source/Runtime/Private/PoseSearchAssetIndexer.cpp
+++ b/Engine/Plugins/Experimental/Animation/PoseSearch/Source/Runtime/Private/PoseSearchAssetIndexer.cpp
@@ -208,13 +208,11 @@ FAssetIndexer::FSampleInfo FAssetIndexer::GetSampleInfo(float SampleTime) const
 
 	const FSamplingParam SamplingParam = WrapOrClampSamplingParam(bCanWrap, PlayLength, MainRelativeTime);
 
-	Sample.Clip = IndexingContext.AssetSampler;
-
 	if (FMath::Abs(SamplingParam.Extrapolation) > SMALL_NUMBER)
 	{
 		Sample.bClamped = true;
 		Sample.ClipTime = SamplingParam.WrappedParam + SamplingParam.Extrapolation;
-		Sample.RootTransform = Sample.Clip->ExtractRootTransform(Sample.ClipTime);
+		Sample.RootTransform = IndexingContext.AssetSampler->ExtractRootTransform(Sample.ClipTime);
 	}
 	else
 	{
@@ -222,7 +220,7 @@ FAssetIndexer::FSampleInfo FAssetIndexer::GetSampleInfo(float SampleTime) const
 		Sample.RootTransform = FTransform::Identity;
 
 		// Find the remaining motion deltas after wrapping
-		FTransform RootMotionRemainder = Sample.Clip->ExtractRootTransform(Sample.ClipTime);
+		FTransform RootMotionRemainder = IndexingContext.AssetSampler->ExtractRootTransform(Sample.ClipTime);
 
 		const bool bNegativeSampleTime = SampleTime < 0.f;
 		if (SamplingParam.NumCycles > 0 || bNegativeSampleTime)
@@ -291,8 +289,8 @@ FAssetIndexer::CachedEntry& FAssetIndexer::GetEntry(float SampleTime)
 		float CurrentTime = Sample.ClipTime;
 		float PreviousTime = CurrentTime - SamplingContext->FiniteDelta;
 
-		const bool bLoopable = Sample.Clip->IsLoopable();
-		const float PlayLength = Sample.Clip->GetPlayLength();
+		const bool bLoopable = IndexingContext.AssetSampler->IsLoopable();
+		const float PlayLength = IndexingContext.AssetSampler->GetPlayLength();
 		if (!bLoopable)
 		{
 			// if not loopable we clamp the pose at time zero or PlayLength
@@ -321,7 +319,7 @@ FAssetIndexer::CachedEntry& FAssetIndexer::GetEntry(float SampleTime)
 		UnusedCurve.InitFrom(BoneContainer);
 		Pose.SetBoneContainer(&BoneContainer);
 
-		Sample.Clip->ExtractPose(ExtractionCtx, AnimPoseData);
+		IndexingContext.AssetSampler->ExtractPose(ExtractionCtx, AnimPoseData);
 		Pose[FCompactPoseBoneIndex(RootBoneIndexType)].SetIdentity();
 
 		if (IndexingContext.bMirrored)
diff --git a/Engine/Plugins/Experimental/Animation/PoseSearch/Source/Runtime/Private/PoseSearchAssetSampler.cpp b/Engine/Plugins/Experimental/Animation/PoseSearch/Source/Runtime/Private/PoseSearchAssetSampler.cpp
--- a/Engine/Plugins/Experimental/Animation/PoseSearch/Source/Runtime/Private/PoseSearchAssetSampler.cpp
+++ b/Engine/Plugins/Experimental/Animation/PoseSearch/Source/Runtime/Private/PoseSearchAssetSampler.cpp
@@ -16,11 +16,7 @@ namespace UE::PoseSearch
 //////////////////////////////////////////////////////////////////////////
 // Root motion extrapolation
 
-static FTransform ExtrapolateRootMotion(
-	FTransform SampleToExtrapolate,
-	float SampleStart, 
-	float SampleEnd, 
-	float ExtrapolationTime)
+static FTransform ExtrapolateRootMotion(FTransform SampleToExtrapolate, float SampleStart, float SampleEnd, float ExtrapolationTime)
 {
 	const float SampleDelta = SampleEnd - SampleStart;
 	check(!FMath::IsNearlyZero(SampleDelta));
@@ -55,322 +51,50 @@ static FTransform ExtrapolateRootMotion(
 	return ExtrapolatedRootMotion;
 }
 
-//////////////////////////////////////////////////////////////////////////
-// FSequenceBaseSampler
-void FSequenceBaseSampler::Init(const FInput& InInput)
-{
-	check(InInput.SequenceBase.Get());
-	Input = InInput;
-}
-
-void FSequenceBaseSampler::Process()
-{
-}
-
-float FSequenceBaseSampler::GetPlayLength() const
-{
-	return Input.SequenceBase->GetPlayLength();
-}
-
-float FSequenceBaseSampler::GetScaledTime(float Time) const
-{
-	return Time;
-}
-
-bool FSequenceBaseSampler::IsLoopable() const
-{
-	return Input.SequenceBase->bLoop;
-}
-
-void FSequenceBaseSampler::ExtractPose(const FAnimExtractContext& ExtractionCtx, FAnimationPoseData& OutAnimPoseData) const
-{
-	Input.SequenceBase->GetAnimationPose(OutAnimPoseData, ExtractionCtx);
-}
-
-FTransform FSequenceBaseSampler::GetTotalRootTransform() const
-{
-	const FTransform InitialRootTransform = ExtractRootTransform(0.f);
-	const FTransform LastRootTransform = ExtractRootTransform(GetPlayLength());
-	const FTransform TotalRootTransform = LastRootTransform.GetRelativeTransform(InitialRootTransform);
-	return TotalRootTransform;
-}
-
-FTransform FSequenceBaseSampler::ExtractRootTransform(float Time) const
-{
-	FTransform RootTransform = FTransform::Identity;
-	const UAnimSequenceBase* SequenceBase = Input.SequenceBase.Get();
-
-	if (IsLoopable())
-	{
-		RootTransform = SequenceBase->ExtractRootMotion(0.0f, Time, true);
-	}
-	else
-	{
-		const float PlayLength = SequenceBase->GetPlayLength();
-		const float ClampedTime = FMath::Clamp(Time, 0.0f, PlayLength);
-		const float ExtrapolationTime = Time - ClampedTime;
-
-		// If Time is less than zero, ExtrapolationTime will be negative. In this case, we extrapolate the beginning of the 
-		// animation to estimate where the root would be at Time
-		if (ExtrapolationTime < -SMALL_NUMBER)
-		{
-			FTransform SampleToExtrapolate = SequenceBase->ExtractRootMotionFromRange(0.0f, ExtrapolationSampleTime);
-
-			const FTransform ExtrapolatedRootMotion = UE::PoseSearch::ExtrapolateRootMotion(
-				SampleToExtrapolate,
-				0.0f, ExtrapolationSampleTime,
-				ExtrapolationTime);
-			RootTransform = ExtrapolatedRootMotion;
-		}
-		else
-		{
-			RootTransform = SequenceBase->ExtractRootMotionFromRange(0.0f, ClampedTime);
-
-			// If Time is greater than PlayLength, ExtrapolationTime will be a positive number. In this case, we extrapolate
-			// the end of the animation to estimate where the root would be at Time
-			if (ExtrapolationTime > SMALL_NUMBER)
-			{
-				FTransform SampleToExtrapolate = SequenceBase->ExtractRootMotionFromRange(PlayLength - ExtrapolationSampleTime, PlayLength);
-
-				const FTransform ExtrapolatedRootMotion = UE::PoseSearch::ExtrapolateRootMotion(
-					SampleToExtrapolate,
-					PlayLength - ExtrapolationSampleTime, PlayLength,
-					ExtrapolationTime);
-				RootTransform = ExtrapolatedRootMotion * RootTransform;
-			}
-		}
-	}
-
-	return RootTransform;
-}
-
-void FSequenceBaseSampler::ExtractPoseSearchNotifyStates(
-	float Time, 
-	TArray<UAnimNotifyState_PoseSearchBase*>& NotifyStates) const
-{
-	// getting pose search notifies in an interval of size ExtractionInterval, centered on Time
-	FAnimNotifyContext NotifyContext;
-	Input.SequenceBase->GetAnimNotifies(Time - (ExtractionInterval * 0.5f), ExtractionInterval, NotifyContext);
-
-	// check which notifies actually overlap Time and are of the right base type
-	for (const FAnimNotifyEventReference& EventReference : NotifyContext.ActiveNotifies)
-	{
-		const FAnimNotifyEvent* NotifyEvent = EventReference.GetNotify();
-		if (!NotifyEvent)
-		{
-			continue;
-		}
-
-		if (NotifyEvent->GetTriggerTime() > Time ||
-			NotifyEvent->GetEndTriggerTime() < Time)
-		{
-			continue;
-		}
-
-		UAnimNotifyState_PoseSearchBase* PoseSearchAnimNotify = 
-			Cast<UAnimNotifyState_PoseSearchBase>(NotifyEvent->NotifyStateClass);
-		if (PoseSearchAnimNotify)
-		{
-			NotifyStates.Add(PoseSearchAnimNotify);
-		}
-	}
-}
-
-const UAnimationAsset* FSequenceBaseSampler::GetAsset() const
-{
-	return Input.SequenceBase.Get();
-}
-
-//////////////////////////////////////////////////////////////////////////
-// FBlendSpaceSampler
-void FBlendSpaceSampler::Init(const FInput& InInput)
-{
-	check(InInput.BlendSpace.Get());
-	Input = InInput;
-}
-
-void FBlendSpaceSampler::Process()
-{
-	FMemMark Mark(FMemStack::Get());
-
-	ProcessPlayLength();
-	ProcessRootTransform();
-}
-
-float FBlendSpaceSampler::GetScaledTime(float Time) const
-{
-	const float ScaledTime = GetPlayLength() > UE_KINDA_SMALL_NUMBER ? Time / GetPlayLength() : 0.f;
-	return ScaledTime;
-}
-
-bool FBlendSpaceSampler::IsLoopable() const
-{
-	return Input.BlendSpace->bLoop;
-}
-
-void FBlendSpaceSampler::ExtractPose(const FAnimExtractContext& ExtractionCtx, FAnimationPoseData& OutAnimPoseData) const
-{
-	TArray<FBlendSampleData> BlendSamples;
-	int32 TriangulationIndex = 0;
-	Input.BlendSpace->GetSamplesFromBlendInput(Input.BlendParameters, BlendSamples, TriangulationIndex, true);
-
-	for (int32 BlendSampleIdex = 0; BlendSampleIdex < BlendSamples.Num(); BlendSampleIdex++)
-	{
-		float Scale = BlendSamples[BlendSampleIdex].Animation->GetPlayLength() / PlayLength;
-
-		FDeltaTimeRecord BlendSampleDeltaTimeRecord;
-		BlendSampleDeltaTimeRecord.Set(ExtractionCtx.DeltaTimeRecord.GetPrevious() * Scale, ExtractionCtx.DeltaTimeRecord.Delta * Scale);
-
-		BlendSamples[BlendSampleIdex].DeltaTimeRecord = BlendSampleDeltaTimeRecord;
-		BlendSamples[BlendSampleIdex].PreviousTime = ExtractionCtx.DeltaTimeRecord.GetPrevious() * Scale;
-		BlendSamples[BlendSampleIdex].Time = ExtractionCtx.CurrentTime * Scale;
-	}
-
-	Input.BlendSpace->GetAnimationPose(BlendSamples, ExtractionCtx, OutAnimPoseData);
-}
-
-FTransform FBlendSpaceSampler::GetTotalRootTransform() const
-{
-	const FTransform InitialRootTransform = ExtractBlendSpaceRootTrackTransform(0.f);
-	const FTransform LastRootTransform = ExtractBlendSpaceRootTrackTransform(PlayLength);
-	const FTransform TotalRootTransform = LastRootTransform.GetRelativeTransform(InitialRootTransform);
-	return TotalRootTransform;
-}
-
-FTransform FBlendSpaceSampler::ExtractRootTransform(float Time) const
-{
-	FTransform RootTransform = FTransform::Identity;
-	if (IsLoopable())
-	{
-		RootTransform = ExtractBlendSpaceRootMotion(0.0f, Time, true);
-	}
-	else
-	{
-		const float ClampedTime = FMath::Clamp(Time, 0.0f, PlayLength);
-		const float ExtrapolationTime = Time - ClampedTime;
-
-
-		// If Time is less than zero, ExtrapolationTime will be negative. In this case, we extrapolate the beginning of the 
-		// animation to estimate where the root would be at Time
-		if (ExtrapolationTime < -SMALL_NUMBER)
-		{
-			FTransform SampleToExtrapolate = ExtractBlendSpaceRootMotionFromRange(0.0f, ExtrapolationSampleTime);
-
-			const FTransform ExtrapolatedRootMotion = UE::PoseSearch::ExtrapolateRootMotion(
-				SampleToExtrapolate,
-				0.0f, ExtrapolationSampleTime,
-				ExtrapolationTime);
-			RootTransform = ExtrapolatedRootMotion;
-		}
-		else
-		{
-			RootTransform = ExtractBlendSpaceRootMotionFromRange(0.0f, ClampedTime);
-
-			// If Time is greater than PlayLength, ExtrapolationTime will be a positive number. In this case, we extrapolate
-			// the end of the animation to estimate where the root would be at Time
-			if (ExtrapolationTime > SMALL_NUMBER)
-			{
-				FTransform SampleToExtrapolate = ExtractBlendSpaceRootMotionFromRange(PlayLength - ExtrapolationSampleTime, PlayLength);
-
-				const FTransform ExtrapolatedRootMotion = UE::PoseSearch::ExtrapolateRootMotion(
-					SampleToExtrapolate,
-					PlayLength - ExtrapolationSampleTime, PlayLength,
-					ExtrapolationTime);
-				RootTransform = ExtrapolatedRootMotion * RootTransform;
-			}
-		}
-	}
-
-	return RootTransform;
-}
-
-static int32 GetHighestWeightSample(const TArray<struct FBlendSampleData>& SampleDataList)
+static FTransform ExtractRootTransformInternal(const UAnimMontage* AnimMontage, float StartTime, float EndTime)
 {
-	int32 HighestWeightIndex = 0;
-	float HighestWeight = SampleDataList[HighestWeightIndex].GetClampedWeight();
-	for (int32 I = 1; I < SampleDataList.Num(); I++)
+	// @todo: add support for SlotName / multiple SlotAnimTracks
+	if (AnimMontage->SlotAnimTracks.Num() != 1)
 	{
-		if (SampleDataList[I].GetClampedWeight() > HighestWeight)
-		{
-			HighestWeightIndex = I;
-			HighestWeight = SampleDataList[I].GetClampedWeight();
-		}
+		UE_LOG(LogPoseSearch, Error, TEXT("ExtractRootTransformInternal: so far we support only montages with one SlotAnimTracks. %s has %d"), *AnimMontage->GetName(), AnimMontage->SlotAnimTracks.Num());
+		return FTransform::Identity;
 	}
-	return HighestWeightIndex;
-}
 
-void FBlendSpaceSampler::ExtractPoseSearchNotifyStates(
-	float Time,
-	TArray<UAnimNotifyState_PoseSearchBase*>& NotifyStates) const
-{
-	if (Input.BlendSpace->NotifyTriggerMode == ENotifyTriggerMode::HighestWeightedAnimation)
+	const FAnimTrack& RootMotionAnimTrack = AnimMontage->SlotAnimTracks[0].AnimTrack;
+	TArray<FRootMotionExtractionStep> RootMotionExtractionSteps;
+	RootMotionAnimTrack.GetRootMotionExtractionStepsForTrackRange(RootMotionExtractionSteps, StartTime, EndTime);
+	FRootMotionMovementParams AccumulatedRootMotionParams;
+	for (const FRootMotionExtractionStep& CurStep : RootMotionExtractionSteps)
 	{
-		// Set up blend samples
-		TArray<FBlendSampleData> BlendSamples;
-		int32 TriangulationIndex = 0;
-		Input.BlendSpace->GetSamplesFromBlendInput(Input.BlendParameters, BlendSamples, TriangulationIndex, true);
-
-		// Find highest weighted
-		const int32 HighestWeightIndex = GetHighestWeightSample(BlendSamples);
-
-		check(HighestWeightIndex != -1);
-
-		// getting pose search notifies in an interval of size ExtractionInterval, centered on Time
-		float SampleTime = Time * (BlendSamples[HighestWeightIndex].Animation->GetPlayLength() / PlayLength);
-
-		// Get notifies for highest weighted
-		FAnimNotifyContext NotifyContext;
-		BlendSamples[HighestWeightIndex].Animation->GetAnimNotifies(
-			(SampleTime - (ExtractionInterval * 0.5f)),
-			ExtractionInterval, 
-			NotifyContext);
-
-		// check which notifies actually overlap Time and are of the right base type
-		for (const FAnimNotifyEventReference& EventReference : NotifyContext.ActiveNotifies)
+		if (CurStep.AnimSequence)
 		{
-			const FAnimNotifyEvent* NotifyEvent = EventReference.GetNotify();
-			if (!NotifyEvent)
-			{
-				continue;
-			}
-
-			if (NotifyEvent->GetTriggerTime() > SampleTime ||
-				NotifyEvent->GetEndTriggerTime() < SampleTime)
-			{
-				continue;
-			}
-
-			UAnimNotifyState_PoseSearchBase* PoseSearchAnimNotify =
-				Cast<UAnimNotifyState_PoseSearchBase>(NotifyEvent->NotifyStateClass);
-			if (PoseSearchAnimNotify)
-			{
-				NotifyStates.Add(PoseSearchAnimNotify);
-			}
+			AccumulatedRootMotionParams.Accumulate(CurStep.AnimSequence->ExtractRootMotionFromRange(CurStep.StartPosition, CurStep.EndPosition));
 		}
 	}
+	return AccumulatedRootMotionParams.GetRootMotionTransform();
 }
 
-FTransform FBlendSpaceSampler::ExtractBlendSpaceRootTrackTransform(float Time) const
+static FTransform ExtractBlendSpaceRootTrackTransform(float Time, const TArray<FTransform>& AccumulatedRootTransform, int32 RootTransformSamplingRate)
 {
 	checkf(AccumulatedRootTransform.Num() > 0, TEXT("ProcessRootTransform must be run first"));
 
-	const int32 Index = Time * Input.RootTransformSamplingRate;
+	const int32 Index = Time * RootTransformSamplingRate;
 	const int32 FirstIndexClamped = FMath::Clamp(Index + 0, 0, AccumulatedRootTransform.Num() - 1);
 	const int32 SecondIndexClamped = FMath::Clamp(Index + 1, 0, AccumulatedRootTransform.Num() - 1);
-	const float Alpha = FMath::Fmod(Time * Input.RootTransformSamplingRate, 1.0f);
+	const float Alpha = FMath::Fmod(Time * RootTransformSamplingRate, 1.0f);
 	FTransform OutputTransform;
 	OutputTransform.Blend(AccumulatedRootTransform[FirstIndexClamped], AccumulatedRootTransform[SecondIndexClamped], Alpha);
 	return OutputTransform;
 }
 
-FTransform FBlendSpaceSampler::ExtractBlendSpaceRootMotionFromRange(float StartTrackPosition, float EndTrackPosition) const
+static FTransform ExtractBlendSpaceRootMotionFromRange(float StartTrackPosition, float EndTrackPosition, const TArray<FTransform>& AccumulatedRootTransform, int32 RootTransformSamplingRate)
 {
 	checkf(AccumulatedRootTransform.Num() > 0, TEXT("ProcessRootTransform must be run first"));
 
-	FTransform RootTransformRefPose = ExtractBlendSpaceRootTrackTransform(0.0f);
+	FTransform RootTransformRefPose = ExtractBlendSpaceRootTrackTransform(0.f, AccumulatedRootTransform, RootTransformSamplingRate);
 
-	FTransform StartTransform = ExtractBlendSpaceRootTrackTransform(StartTrackPosition);
-	FTransform EndTransform = ExtractBlendSpaceRootTrackTransform(EndTrackPosition);
+	FTransform StartTransform = ExtractBlendSpaceRootTrackTransform(StartTrackPosition, AccumulatedRootTransform, RootTransformSamplingRate);
+	FTransform EndTransform = ExtractBlendSpaceRootTrackTransform(EndTrackPosition, AccumulatedRootTransform, RootTransformSamplingRate);
 
 	// Transform to Component Space
 	const FTransform RootToComponent = RootTransformRefPose.Inverse();
@@ -380,7 +104,7 @@ FTransform FBlendSpaceSampler::ExtractBlendSpaceRootMotionFromRange(float StartT
 	return EndTransform.GetRelativeTransform(StartTransform);
 }
 
-FTransform FBlendSpaceSampler::ExtractBlendSpaceRootMotion(float StartTime, float DeltaTime, bool bAllowLooping) const
+static FTransform ExtractBlendSpaceRootMotion(float StartTime, float DeltaTime, bool bAllowLooping, float CachedPlayLength, const TArray<FTransform>& AccumulatedRootTransform, int32 RootTransformSamplingRate)
 {
 	FRootMotionMovementParams RootMotionParams;
 
@@ -395,13 +119,14 @@ FTransform FBlendSpaceSampler::ExtractBlendSpaceRootMotion(float StartTime, floa
 		do
 		{
 			// Disable looping here. Advance to desired position, or beginning / end of animation 
-			const ETypeAdvanceAnim AdvanceType = FAnimationRuntime::AdvanceTime(false, DesiredDeltaMove, CurrentPosition, PlayLength);
+			check(CachedPlayLength >= 0.f);
+			const ETypeAdvanceAnim AdvanceType = FAnimationRuntime::AdvanceTime(false, DesiredDeltaMove, CurrentPosition, CachedPlayLength);
 
 			// Verify position assumptions
 			//ensureMsgf(bPlayingBackwards ? (CurrentPosition <= PreviousPosition) : (CurrentPosition >= PreviousPosition), TEXT("in Animation %s(Skeleton %s) : bPlayingBackwards(%d), PreviousPosition(%0.2f), Current Position(%0.2f)"),
 			//	*GetName(), *GetNameSafe(GetSkeleton()), bPlayingBackwards, PreviousPosition, CurrentPosition);
 
-			RootMotionParams.Accumulate(ExtractBlendSpaceRootMotionFromRange(PreviousPosition, CurrentPosition));
+			RootMotionParams.Accumulate(ExtractBlendSpaceRootMotionFromRange(PreviousPosition, CurrentPosition, AccumulatedRootTransform, RootTransformSamplingRate));
 
 			// If we've hit the end of the animation, and we're allowed to loop, keep going.
 			if ((AdvanceType == ETAA_Finished) && bAllowLooping)
@@ -409,7 +134,7 @@ FTransform FBlendSpaceSampler::ExtractBlendSpaceRootMotion(float StartTime, floa
 				const float ActualDeltaMove = (CurrentPosition - PreviousPosition);
 				DesiredDeltaMove -= ActualDeltaMove;
 
-				PreviousPosition = bPlayingBackwards ? PlayLength : 0.f;
+				PreviousPosition = bPlayingBackwards ? CachedPlayLength : 0.f;
 				CurrentPosition = PreviousPosition;
 			}
 			else
@@ -422,221 +147,435 @@ FTransform FBlendSpaceSampler::ExtractBlendSpaceRootMotion(float StartTime, floa
 	return RootMotionParams.GetRootMotionTransform();
 }
 
-void FBlendSpaceSampler::ProcessPlayLength()
+static void ProcessPlayLength(const UBlendSpace* BlendSpace, const FVector& BlendParameters, float& CachedPlayLength)
 {
 	TArray<FBlendSampleData> BlendSamples;
 	int32 TriangulationIndex = 0;
-	Input.BlendSpace->GetSamplesFromBlendInput(Input.BlendParameters, BlendSamples, TriangulationIndex, true);
-
-	PlayLength = Input.BlendSpace->GetAnimationLengthFromSampleData(BlendSamples);
+	if (BlendSpace->GetSamplesFromBlendInput(BlendParameters, BlendSamples, TriangulationIndex, true))
+	{
+		CachedPlayLength = BlendSpace->GetAnimationLengthFromSampleData(BlendSamples);
+	}
 }
 
-void FBlendSpaceSampler::ProcessRootTransform()
+static void ProcessRootTransform(const UBlendSpace* BlendSpace, const FVector& BlendParameters, float CachedPlayLength, const FBoneContainer& BoneContainer,
+	int32 RootTransformSamplingRate, bool bIsLoopable, TArray<FTransform>& AccumulatedRootTransform)
 {
 	// Pre-compute root motion
-
-	int32 NumRootSamples = FMath::Max(PlayLength * Input.RootTransformSamplingRate + 1, 1);
+	check(CachedPlayLength >= 0.f);
+	int32 NumRootSamples = FMath::Max(CachedPlayLength * RootTransformSamplingRate + 1, 1);
 	AccumulatedRootTransform.SetNumUninitialized(NumRootSamples);
 
 	TArray<FBlendSampleData> BlendSamples;
 	int32 TriangulationIndex = 0;
-	Input.BlendSpace->GetSamplesFromBlendInput(Input.BlendParameters, BlendSamples, TriangulationIndex, true);
-
-	FTransform RootMotionAccumulation = FTransform::Identity;
-
-	AccumulatedRootTransform[0] = RootMotionAccumulation;
-
-	for (int32 SampleIdx = 1; SampleIdx < NumRootSamples; ++SampleIdx)
+	if (BlendSpace->GetSamplesFromBlendInput(BlendParameters, BlendSamples, TriangulationIndex, true))
 	{
-		float PreviousTime = float(SampleIdx - 1) / Input.RootTransformSamplingRate;
-		float CurrentTime = float(SampleIdx - 0) / Input.RootTransformSamplingRate;
+		FTransform RootMotionAccumulation = FTransform::Identity;
 
-		FDeltaTimeRecord DeltaTimeRecord;
-		DeltaTimeRecord.Set(PreviousTime, CurrentTime - PreviousTime);
-		FAnimExtractContext ExtractionCtx(static_cast<double>(CurrentTime), true, DeltaTimeRecord, IsLoopable());
+		AccumulatedRootTransform[0] = RootMotionAccumulation;
 
-		for (int32 BlendSampleIdex = 0; BlendSampleIdex < BlendSamples.Num(); BlendSampleIdex++)
+		for (int32 SampleIdx = 1; SampleIdx < NumRootSamples; ++SampleIdx)
 		{
-			float Scale = BlendSamples[BlendSampleIdex].Animation->GetPlayLength() / PlayLength;
+			float PreviousTime = float(SampleIdx - 1) / RootTransformSamplingRate;
+			float CurrentTime = float(SampleIdx - 0) / RootTransformSamplingRate;
 
-			FDeltaTimeRecord BlendSampleDeltaTimeRecord;
-			BlendSampleDeltaTimeRecord.Set(DeltaTimeRecord.GetPrevious() * Scale, DeltaTimeRecord.Delta * Scale);
+			FDeltaTimeRecord DeltaTimeRecord;
+			DeltaTimeRecord.Set(PreviousTime, CurrentTime - PreviousTime);
+			FAnimExtractContext ExtractionCtx(static_cast<double>(CurrentTime), true, DeltaTimeRecord, bIsLoopable);
 
-			BlendSamples[BlendSampleIdex].DeltaTimeRecord = BlendSampleDeltaTimeRecord;
-			BlendSamples[BlendSampleIdex].PreviousTime = PreviousTime * Scale;
-			BlendSamples[BlendSampleIdex].Time = CurrentTime * Scale;
-		}
+			for (int32 BlendSampleIdex = 0; BlendSampleIdex < BlendSamples.Num(); BlendSampleIdex++)
+			{
+				float Scale = BlendSamples[BlendSampleIdex].Animation->GetPlayLength() / CachedPlayLength;
 
-		FCompactPose Pose;
-		FBlendedCurve BlendedCurve;
-		UE::Anim::FStackAttributeContainer StackAttributeContainer;
-		FAnimationPoseData AnimPoseData(Pose, BlendedCurve, StackAttributeContainer);
+				FDeltaTimeRecord BlendSampleDeltaTimeRecord;
+				BlendSampleDeltaTimeRecord.Set(DeltaTimeRecord.GetPrevious() * Scale, DeltaTimeRecord.Delta * Scale);
 
-		Pose.SetBoneContainer(&Input.BoneContainer);
-		BlendedCurve.InitFrom(Input.BoneContainer);
+				BlendSamples[BlendSampleIdex].DeltaTimeRecord = BlendSampleDeltaTimeRecord;
+				BlendSamples[BlendSampleIdex].PreviousTime = PreviousTime * Scale;
+				BlendSamples[BlendSampleIdex].Time = CurrentTime * Scale;
+			}
 
-		Input.BlendSpace->GetAnimationPose(BlendSamples, ExtractionCtx, AnimPoseData);
+			FCompactPose Pose;
+			FBlendedCurve BlendedCurve;
+			UE::Anim::FStackAttributeContainer StackAttributeContainer;
+			FAnimationPoseData AnimPoseData(Pose, BlendedCurve, StackAttributeContainer);
 
-		const UE::Anim::IAnimRootMotionProvider* RootMotionProvider = UE::Anim::IAnimRootMotionProvider::Get();
+			Pose.SetBoneContainer(&BoneContainer);
+			BlendedCurve.InitFrom(BoneContainer);
 
-		if (RootMotionProvider)
-		{
-			if (RootMotionProvider->HasRootMotion(StackAttributeContainer))
+			BlendSpace->GetAnimationPose(BlendSamples, ExtractionCtx, AnimPoseData);
+
+			const UE::Anim::IAnimRootMotionProvider* RootMotionProvider = UE::Anim::IAnimRootMotionProvider::Get();
+
+			if (RootMotionProvider)
 			{
-				FTransform RootMotionDelta;
-				RootMotionProvider->ExtractRootMotion(StackAttributeContainer, RootMotionDelta);
-				RootMotionAccumulation = RootMotionDelta * RootMotionAccumulation;
+				if (RootMotionProvider->HasRootMotion(StackAttributeContainer))
+				{
+					FTransform RootMotionDelta;
+					RootMotionProvider->ExtractRootMotion(StackAttributeContainer, RootMotionDelta);
+					RootMotionAccumulation = RootMotionDelta * RootMotionAccumulation;
+				}
+				else
+				{
+					UE_LOG(LogPoseSearch, Error, TEXT("ProcessRootTransform: Blend Space '%s' has no Root Motion Attribute"), *BlendSpace->GetName());
+				}
 			}
 			else
 			{
-				UE_LOG(LogPoseSearch, Error, TEXT("FBlendSpaceSampler::ProcessRootTransform: Blend Space '%s' has no Root Motion Attribute"), *Input.BlendSpace->GetName());
+				UE_LOG(LogPoseSearch, Error, TEXT("ProcessRootTransform: Could not get Root Motion Provider for BlendSpace '%s'"), *BlendSpace->GetName());
 			}
+
+			AccumulatedRootTransform[SampleIdx] = RootMotionAccumulation;
 		}
-		else
+	}
+}
+
+static int32 GetHighestWeightSample(const TArray<struct FBlendSampleData>& SampleDataList)
+{
+	check(!SampleDataList.IsEmpty());
+	int32 HighestWeightIndex = 0;
+	float HighestWeight = SampleDataList[HighestWeightIndex].GetClampedWeight();
+	for (int32 I = 1; I < SampleDataList.Num(); I++)
+	{
+		if (SampleDataList[I].GetClampedWeight() > HighestWeight)
 		{
-			UE_LOG(LogPoseSearch, Error, TEXT("FBlendSpaceSampler::ProcessRootTransform: Could not get Root Motion Provider for BlendSpace '%s'"), *Input.BlendSpace->GetName());
+			HighestWeightIndex = I;
+			HighestWeight = SampleDataList[I].GetClampedWeight();
 		}
-
-		AccumulatedRootTransform[SampleIdx] = RootMotionAccumulation;
 	}
+	return HighestWeightIndex;
 }
 
-const UAnimationAsset* FBlendSpaceSampler::GetAsset() const
+//////////////////////////////////////////////////////////////////////////
+// FAssetSamplerBase
+FAnimationAssetSampler::FAnimationAssetSampler(TWeakObjectPtr<const UAnimationAsset> InAnimationAsset, const FVector& InBlendParameters, const FBoneContainer& InBoneContainer, int32 InRootTransformSamplingRate)
 {
-	return Input.BlendSpace.Get();
+	Init(InAnimationAsset, InBlendParameters, InBoneContainer, InRootTransformSamplingRate);
 }
 
-//////////////////////////////////////////////////////////////////////////
-// FAnimMontageSampler
-void FAnimMontageSampler::Init(const FInput& InInput)
+void FAnimationAssetSampler::Init(TWeakObjectPtr<const UAnimationAsset> InAnimationAsset, const FVector& InBlendParameters, const FBoneContainer& InBoneContainer, int32 InRootTransformSamplingRate)
 {
-	check(InInput.AnimMontage.Get());
-	Input = InInput;
+	AnimationAsset = InAnimationAsset;
+	BlendParameters = InBlendParameters;
+	BoneContainer = InBoneContainer;
+	RootTransformSamplingRate = InRootTransformSamplingRate;
 }
 
-void FAnimMontageSampler::Process()
+bool FAnimationAssetSampler::IsInitialized() const
 {
+	return AnimationAsset != nullptr;
 }
 
-float FAnimMontageSampler::GetPlayLength() const
+const UAnimationAsset* FAnimationAssetSampler::GetAsset() const
 {
-	return Input.AnimMontage->GetPlayLength();
+	return AnimationAsset.Get();
 }
 
-float FAnimMontageSampler::GetScaledTime(float Time) const
+float FAnimationAssetSampler::GetScaledTime(float Time) const
 {
+	// Asset player time for blend spaces is normalized [0, 1] so we convert it to scaled time by dividing it by CachedPlayLength
+	if (Cast<UBlendSpace>(AnimationAsset.Get()))
+	{
+		check(CachedPlayLength >= 0.f);
+		return CachedPlayLength > UE_KINDA_SMALL_NUMBER ? Time / CachedPlayLength : 0.f;
+	}
+
 	return Time;
 }
 
-bool FAnimMontageSampler::IsLoopable() const
+float FAnimationAssetSampler::GetPlayLength() const
 {
-	return Input.AnimMontage->bLoop;
+	if (Cast<UBlendSpace>(AnimationAsset.Get()))
+	{
+		check(CachedPlayLength >= 0.f);
+		return CachedPlayLength;
+	}
+	return AnimationAsset->GetPlayLength();
 }
 
-void FAnimMontageSampler::ExtractPose(const FAnimExtractContext& ExtractionCtx, FAnimationPoseData& OutAnimPoseData) const
+bool FAnimationAssetSampler::IsLoopable() const
 {
-	// @todo: add support for SlotName / multiple SlotAnimTracks
-	if (Input.AnimMontage->SlotAnimTracks.Num() != 1)
+	if (const UAnimSequenceBase* SequenceBase = Cast<UAnimSequenceBase>(AnimationAsset.Get()))
+	{
+		return SequenceBase->bLoop;
+	}
+
+	if (const UBlendSpace* BlendSpace = Cast<UBlendSpace>(AnimationAsset.Get()))
 	{
-		UE_LOG(LogPoseSearch, Error, TEXT("FAnimMontageSampler::ExtractPose: so far we support only montages with one SlotAnimTracks. %s has %d"), *Input.AnimMontage->GetName(), Input.AnimMontage->SlotAnimTracks.Num());
-		OutAnimPoseData.GetPose().ResetToRefPose();
-		return;
+		return BlendSpace->bLoop;
 	}
 
-	Input.AnimMontage->SlotAnimTracks[0].AnimTrack.GetAnimationPose(OutAnimPoseData, ExtractionCtx);
+	checkNoEntry();
+	return false;
 }
 
-FTransform FAnimMontageSampler::GetTotalRootTransform() const
+FTransform FAnimationAssetSampler::GetTotalRootTransform() const
 {
-	// @todo: add support for SlotName / multiple SlotAnimTracks
-	if (Input.AnimMontage->SlotAnimTracks.Num() != 1)
+	if (Cast<UBlendSpace>(AnimationAsset.Get()))
 	{
-		UE_LOG(LogPoseSearch, Error, TEXT("FAnimMontageSampler::GetTotalRootTransform: so far we support only montages with one SlotAnimTracks. %s has %d"), *Input.AnimMontage->GetName(), Input.AnimMontage->SlotAnimTracks.Num());
-		return FTransform::Identity;
+		check(CachedPlayLength >= 0.f);
+		const FTransform InitialRootTransform = ExtractBlendSpaceRootTrackTransform(0.f, AccumulatedRootTransform, RootTransformSamplingRate);
+		const FTransform LastRootTransform = ExtractBlendSpaceRootTrackTransform(CachedPlayLength, AccumulatedRootTransform, RootTransformSamplingRate);
+		const FTransform TotalRootTransform = LastRootTransform.GetRelativeTransform(InitialRootTransform);
+		return TotalRootTransform;
+	}
+
+	if (const UAnimMontage* AnimMontage = Cast<UAnimMontage>(AnimationAsset.Get()))
+	{
+		// @todo: add support for SlotName / multiple SlotAnimTracks
+		if (AnimMontage->SlotAnimTracks.Num() != 1)
+		{
+			UE_LOG(LogPoseSearch, Error, TEXT("FAssetSamplerBase::GetTotalRootTransform: so far we support only montages with one SlotAnimTracks. %s has %d"), *AnimMontage->GetName(), AnimMontage->SlotAnimTracks.Num());
+			return FTransform::Identity;
+		}
+
+		// @todo: optimize me
+		const FTransform InitialRootTransform = ExtractRootTransform(0.f);
+		const FTransform LastRootTransform = ExtractRootTransform(GetPlayLength());
+		const FTransform TotalRootTransform = LastRootTransform.GetRelativeTransform(InitialRootTransform);
+		return TotalRootTransform;
 	}
 
-	// @todo: optimize me
 	const FTransform InitialRootTransform = ExtractRootTransform(0.f);
 	const FTransform LastRootTransform = ExtractRootTransform(GetPlayLength());
 	const FTransform TotalRootTransform = LastRootTransform.GetRelativeTransform(InitialRootTransform);
 	return TotalRootTransform;
 }
 
-FTransform FAnimMontageSampler::ExtractRootTransformInternal(float StartTime, float EndTime) const
+void FAnimationAssetSampler::ExtractPose(const FAnimExtractContext& ExtractionCtx, FAnimationPoseData& OutAnimPoseData) const
 {
-	// @todo: add support for SlotName / multiple SlotAnimTracks
-	if (Input.AnimMontage->SlotAnimTracks.Num() != 1)
+	if (const UBlendSpace* BlendSpace = Cast<UBlendSpace>(AnimationAsset.Get()))
 	{
-		UE_LOG(LogPoseSearch, Error, TEXT("FAnimMontageSampler::ExtractRootTransform: so far we support only montages with one SlotAnimTracks. %s has %d"), *Input.AnimMontage->GetName(), Input.AnimMontage->SlotAnimTracks.Num());
-		return FTransform::Identity;
-	}
+		TArray<FBlendSampleData> BlendSamples;
+		int32 TriangulationIndex = 0;
+		if (BlendSpace->GetSamplesFromBlendInput(BlendParameters, BlendSamples, TriangulationIndex, true))
+		{
+			check(CachedPlayLength >= 0.f);
+			for (int32 BlendSampleIdex = 0; BlendSampleIdex < BlendSamples.Num(); BlendSampleIdex++)
+			{
+				float Scale = BlendSamples[BlendSampleIdex].Animation->GetPlayLength() / CachedPlayLength;
 
-	const FAnimTrack& RootMotionAnimTrack = Input.AnimMontage->SlotAnimTracks[0].AnimTrack;
-	TArray<FRootMotionExtractionStep> RootMotionExtractionSteps;
-	RootMotionAnimTrack.GetRootMotionExtractionStepsForTrackRange(RootMotionExtractionSteps, StartTime, EndTime);
-	FRootMotionMovementParams AccumulatedRootMotionParams;
-	for (const FRootMotionExtractionStep& CurStep : RootMotionExtractionSteps)
+				FDeltaTimeRecord BlendSampleDeltaTimeRecord;
+				BlendSampleDeltaTimeRecord.Set(ExtractionCtx.DeltaTimeRecord.GetPrevious() * Scale, ExtractionCtx.DeltaTimeRecord.Delta * Scale);
+
+				BlendSamples[BlendSampleIdex].DeltaTimeRecord = BlendSampleDeltaTimeRecord;
+				BlendSamples[BlendSampleIdex].PreviousTime = ExtractionCtx.DeltaTimeRecord.GetPrevious() * Scale;
+				BlendSamples[BlendSampleIdex].Time = ExtractionCtx.CurrentTime * Scale;
+			}
+
+			BlendSpace->GetAnimationPose(BlendSamples, ExtractionCtx, OutAnimPoseData);
+		}
+	}
+	else if (const UAnimMontage* AnimMontage = Cast<UAnimMontage>(AnimationAsset.Get()))
 	{
-		if (CurStep.AnimSequence)
+		// @todo: add support for SlotName / multiple SlotAnimTracks
+		if (AnimMontage->SlotAnimTracks.Num() != 1)
 		{
-			AccumulatedRootMotionParams.Accumulate(CurStep.AnimSequence->ExtractRootMotionFromRange(CurStep.StartPosition, CurStep.EndPosition));
+			UE_LOG(LogPoseSearch, Error, TEXT("FAnimMontageSampler::ExtractPose: so far we support only montages with one SlotAnimTracks. %s has %d"), *AnimMontage->GetName(), AnimMontage->SlotAnimTracks.Num());
+			OutAnimPoseData.GetPose().ResetToRefPose();
+			return;
 		}
+
+		AnimMontage->SlotAnimTracks[0].AnimTrack.GetAnimationPose(OutAnimPoseData, ExtractionCtx);
+	}
+	else if (const UAnimSequenceBase* SequenceBase = Cast<UAnimSequenceBase>(AnimationAsset.Get()))
+	{
+		SequenceBase->GetAnimationPose(OutAnimPoseData, ExtractionCtx);
+	}
+	else
+	{
+		checkNoEntry();
 	}
-	return AccumulatedRootMotionParams.GetRootMotionTransform();
 }
 
-FTransform FAnimMontageSampler::ExtractRootTransform(float Time) const
+FTransform FAnimationAssetSampler::ExtractRootTransform(float Time) const
 {
 	FTransform RootTransform = FTransform::Identity;
-	if (IsLoopable())
+	
+	if (const UBlendSpace* BlendSpace = Cast<UBlendSpace>(AnimationAsset.Get()))
 	{
-		RootTransform = ExtractRootTransformInternal(0.f, Time);
+		if (IsLoopable())
+		{
+			RootTransform = ExtractBlendSpaceRootMotion(0.0f, Time, true, CachedPlayLength, AccumulatedRootTransform, RootTransformSamplingRate);
+		}
+		else
+		{
+			check(CachedPlayLength >= 0.f);
+			const float ClampedTime = FMath::Clamp(Time, 0.0f, CachedPlayLength);
+			const float ExtrapolationTime = Time - ClampedTime;
+
+			// If Time is less than zero, ExtrapolationTime will be negative. In this case, we extrapolate the beginning of the 
+			// animation to estimate where the root would be at Time
+			if (ExtrapolationTime < -SMALL_NUMBER)
+			{
+				FTransform SampleToExtrapolate = ExtractBlendSpaceRootMotionFromRange(0.0f, ExtrapolationSampleTime, AccumulatedRootTransform, RootTransformSamplingRate);
+
+				const FTransform ExtrapolatedRootMotion = UE::PoseSearch::ExtrapolateRootMotion(
+					SampleToExtrapolate,
+					0.0f, ExtrapolationSampleTime,
+					ExtrapolationTime);
+				RootTransform = ExtrapolatedRootMotion;
+			}
+			else
+			{
+				RootTransform = ExtractBlendSpaceRootMotionFromRange(0.0f, ClampedTime, AccumulatedRootTransform, RootTransformSamplingRate);
+
+				// If Time is greater than PlayLength, ExtrapolationTime will be a positive number. In this case, we extrapolate
+				// the end of the animation to estimate where the root would be at Time
+				if (ExtrapolationTime > SMALL_NUMBER)
+				{
+					FTransform SampleToExtrapolate = ExtractBlendSpaceRootMotionFromRange(CachedPlayLength - ExtrapolationSampleTime, CachedPlayLength, AccumulatedRootTransform, RootTransformSamplingRate);
+
+					const FTransform ExtrapolatedRootMotion = UE::PoseSearch::ExtrapolateRootMotion(
+						SampleToExtrapolate,
+						CachedPlayLength - ExtrapolationSampleTime, CachedPlayLength,
+						ExtrapolationTime);
+					RootTransform = ExtrapolatedRootMotion * RootTransform;
+				}
+			}
+		}
 	}
-	else
+	else if (const UAnimMontage* AnimMontage = Cast<UAnimMontage>(AnimationAsset.Get()))
 	{
-		const float PlayLength = GetPlayLength();
-		const float ClampedTime = FMath::Clamp(Time, 0.f, PlayLength);
-		const float ExtrapolationTime = Time - ClampedTime;
-
-		// If Time is less than zero, ExtrapolationTime will be negative. In this case, we extrapolate the beginning of the 
-		// animation to estimate where the root would be at Time
-		if (ExtrapolationTime < -SMALL_NUMBER)
+		if (IsLoopable())
 		{
-			FTransform SampleToExtrapolate = ExtractRootTransformInternal(0.f, ExtrapolationSampleTime);
+			RootTransform = ExtractRootTransformInternal(AnimMontage, 0.f, Time);
+		}
+		else
+		{
+			const float PlayLength = GetPlayLength();
+			const float ClampedTime = FMath::Clamp(Time, 0.f, PlayLength);
+			const float ExtrapolationTime = Time - ClampedTime;
+
+			// If Time is less than zero, ExtrapolationTime will be negative. In this case, we extrapolate the beginning of the 
+			// animation to estimate where the root would be at Time
+			if (ExtrapolationTime < -SMALL_NUMBER)
+			{
+				FTransform SampleToExtrapolate = ExtractRootTransformInternal(AnimMontage, 0.f, ExtrapolationSampleTime);
 
-			const FTransform ExtrapolatedRootMotion = UE::PoseSearch::ExtrapolateRootMotion(
-				SampleToExtrapolate,
-				0.0f, ExtrapolationSampleTime,
-				ExtrapolationTime);
-			RootTransform = ExtrapolatedRootMotion;
+				const FTransform ExtrapolatedRootMotion = UE::PoseSearch::ExtrapolateRootMotion(
+					SampleToExtrapolate,
+					0.0f, ExtrapolationSampleTime,
+					ExtrapolationTime);
+				RootTransform = ExtrapolatedRootMotion;
+			}
+			else
+			{
+				RootTransform = ExtractRootTransformInternal(AnimMontage, 0.f, ClampedTime);
+
+				// If Time is greater than PlayLength, ExtrapolationTime will be a positive number. In this case, we extrapolate
+				// the end of the animation to estimate where the root would be at Time
+				if (ExtrapolationTime > SMALL_NUMBER)
+				{
+					FTransform SampleToExtrapolate = ExtractRootTransformInternal(AnimMontage, PlayLength - ExtrapolationSampleTime, PlayLength);
+
+					const FTransform ExtrapolatedRootMotion = UE::PoseSearch::ExtrapolateRootMotion(
+						SampleToExtrapolate,
+						PlayLength - ExtrapolationSampleTime, PlayLength,
+						ExtrapolationTime);
+					RootTransform = ExtrapolatedRootMotion * RootTransform;
+				}
+			}
+		}
+	}
+	else if (const UAnimSequenceBase* SequenceBase = Cast<UAnimSequenceBase>(AnimationAsset.Get()))
+	{
+		if (IsLoopable())
+		{
+			RootTransform = SequenceBase->ExtractRootMotion(0.0f, Time, true);
 		}
 		else
 		{
-			RootTransform = ExtractRootTransformInternal(0.f, ClampedTime);
+			const float PlayLength = GetPlayLength();
+			const float ClampedTime = FMath::Clamp(Time, 0.0f, PlayLength);
+			const float ExtrapolationTime = Time - ClampedTime;
 
-			// If Time is greater than PlayLength, ExtrapolationTime will be a positive number. In this case, we extrapolate
-			// the end of the animation to estimate where the root would be at Time
-			if (ExtrapolationTime > SMALL_NUMBER)
+			// If Time is less than zero, ExtrapolationTime will be negative. In this case, we extrapolate the beginning of the 
+			// animation to estimate where the root would be at Time
+			if (ExtrapolationTime < -SMALL_NUMBER)
 			{
-				FTransform SampleToExtrapolate = ExtractRootTransformInternal(PlayLength - ExtrapolationSampleTime, PlayLength);
+				FTransform SampleToExtrapolate = SequenceBase->ExtractRootMotionFromRange(0.0f, ExtrapolationSampleTime);
 
 				const FTransform ExtrapolatedRootMotion = UE::PoseSearch::ExtrapolateRootMotion(
 					SampleToExtrapolate,
-					PlayLength - ExtrapolationSampleTime, PlayLength,
+					0.0f, ExtrapolationSampleTime,
 					ExtrapolationTime);
-				RootTransform = ExtrapolatedRootMotion * RootTransform;
+				RootTransform = ExtrapolatedRootMotion;
+			}
+			else
+			{
+				RootTransform = SequenceBase->ExtractRootMotionFromRange(0.0f, ClampedTime);
+
+				// If Time is greater than PlayLength, ExtrapolationTime will be a positive number. In this case, we extrapolate
+				// the end of the animation to estimate where the root would be at Time
+				if (ExtrapolationTime > SMALL_NUMBER)
+				{
+					FTransform SampleToExtrapolate = SequenceBase->ExtractRootMotionFromRange(PlayLength - ExtrapolationSampleTime, PlayLength);
+
+					const FTransform ExtrapolatedRootMotion = UE::PoseSearch::ExtrapolateRootMotion(
+						SampleToExtrapolate,
+						PlayLength - ExtrapolationSampleTime, PlayLength,
+						ExtrapolationTime);
+					RootTransform = ExtrapolatedRootMotion * RootTransform;
+				}
 			}
 		}
 	}
+	else
+	{
+		checkNoEntry();
+	}
 
 	return RootTransform;
 }
 
-void FAnimMontageSampler::ExtractPoseSearchNotifyStates(float Time, TArray<UAnimNotifyState_PoseSearchBase*>& NotifyStates) const
+void FAnimationAssetSampler::Process()
+{
+	if (const UBlendSpace* BlendSpace = Cast<UBlendSpace>(AnimationAsset.Get()))
+	{
+		FMemMark Mark(FMemStack::Get());
+
+		ProcessPlayLength(BlendSpace, BlendParameters, CachedPlayLength);
+		ProcessRootTransform(BlendSpace, BlendParameters, CachedPlayLength, BoneContainer, RootTransformSamplingRate, IsLoopable(), AccumulatedRootTransform);
+	}
+}
+
+void FAnimationAssetSampler::ExtractPoseSearchNotifyStates(float Time, TArray<UAnimNotifyState_PoseSearchBase*>& NotifyStates) const
 {
-	// getting pose search notifies in an interval of size ExtractionInterval, centered on Time
+	float SampleTime = Time;
 	FAnimNotifyContext NotifyContext;
-	Input.AnimMontage->GetAnimNotifies(Time - (ExtractionInterval * 0.5f), ExtractionInterval, NotifyContext);
+	if (const UBlendSpace* BlendSpace = Cast<UBlendSpace>(AnimationAsset.Get()))
+	{
+		if (BlendSpace->NotifyTriggerMode == ENotifyTriggerMode::HighestWeightedAnimation)
+		{
+			// Set up blend samples
+			TArray<FBlendSampleData> BlendSamples;
+			int32 TriangulationIndex = 0;
+			if (BlendSpace->GetSamplesFromBlendInput(BlendParameters, BlendSamples, TriangulationIndex, true))
+			{
+				// Find highest weighted
+				const int32 HighestWeightIndex = GetHighestWeightSample(BlendSamples);
+
+				// getting pose search notifies in an interval of size ExtractionInterval, centered on Time
+				check(CachedPlayLength >= 0.f);
+				SampleTime = Time * (BlendSamples[HighestWeightIndex].Animation->GetPlayLength() / CachedPlayLength);
+
+				// Get notifies for highest weighted
+				BlendSamples[HighestWeightIndex].Animation->GetAnimNotifies(
+					(SampleTime - (ExtractionInterval * 0.5f)),
+					ExtractionInterval,
+					NotifyContext);
+			}
+		}
+	}
+	else if (const UAnimSequenceBase* SequenceBase = Cast<UAnimSequenceBase>(AnimationAsset.Get()))
+	{
+		// getting pose search notifies in an interval of size ExtractionInterval, centered on Time
+		SequenceBase->GetAnimNotifies(Time - (ExtractionInterval * 0.5f), ExtractionInterval, NotifyContext);
+	}
+	else
+	{
+		checkNoEntry();
+	}
 
 	// check which notifies actually overlap Time and are of the right base type
 	for (const FAnimNotifyEventReference& EventReference : NotifyContext.ActiveNotifies)
@@ -647,24 +586,17 @@ void FAnimMontageSampler::ExtractPoseSearchNotifyStates(float Time, TArray<UAnim
 			continue;
 		}
 
-		if (NotifyEvent->GetTriggerTime() > Time ||
-			NotifyEvent->GetEndTriggerTime() < Time)
+		// @todo: is this condition necessary? can we just rely on the ExtractionInterval?
+		if (NotifyEvent->GetTriggerTime() > SampleTime || NotifyEvent->GetEndTriggerTime() < SampleTime)
 		{
 			continue;
 		}
 
-		UAnimNotifyState_PoseSearchBase* PoseSearchAnimNotify =
-			Cast<UAnimNotifyState_PoseSearchBase>(NotifyEvent->NotifyStateClass);
-		if (PoseSearchAnimNotify)
+		if (UAnimNotifyState_PoseSearchBase* PoseSearchAnimNotify = Cast<UAnimNotifyState_PoseSearchBase>(NotifyEvent->NotifyStateClass))
 		{
 			NotifyStates.Add(PoseSearchAnimNotify);
 		}
 	}
 }
 
-const UAnimationAsset* FAnimMontageSampler::GetAsset() const
-{
-	return Input.AnimMontage.Get();
-}
-
 } // namespace UE::PoseSearch
diff --git a/Engine/Plugins/Experimental/Animation/PoseSearch/Source/Runtime/Private/PoseSearchDatabaseIndexingContext.cpp b/Engine/Plugins/Experimental/Animation/PoseSearch/Source/Runtime/Private/PoseSearchDatabaseIndexingContext.cpp
--- a/Engine/Plugins/Experimental/Animation/PoseSearch/Source/Runtime/Private/PoseSearchDatabaseIndexingContext.cpp
+++ b/Engine/Plugins/Experimental/Animation/PoseSearch/Source/Runtime/Private/PoseSearchDatabaseIndexingContext.cpp
@@ -58,34 +58,11 @@ bool FDatabaseIndexingContext::IndexDatabase(FPoseSearchIndexBase& SearchIndexBa
 	}
 
 	// Prepare samplers for all animation assets.
-	TMap<const UAnimationAsset*, int32> SamplerMap;
-	TMap<TPair<const UBlendSpace*, FVector>, int32> BlendSpaceSamplerMap;
+	Samplers.Reset();
+	TMap<TPair<const UAnimationAsset*, FVector>, int32> SamplerMap;
 	for (const FInstancedStruct& DatabaseAssetStruct : Database.AnimationAssets)
 	{
-		auto AddSequenceBaseSampler = [&](const UAnimSequenceBase* Sequence)
-		{
-			if (Sequence && !SamplerMap.Contains(Sequence))
-			{
-				FSequenceBaseSampler::FInput Input;
-				Input.SequenceBase = Sequence;
-
-				SamplerMap.Add(Sequence, Samplers.Num());
-
-				TSharedPtr<FSequenceBaseSampler> Sampler = MakeShared<FSequenceBaseSampler>();
-				Sampler->Init(Input);
-				Samplers.Add(Sampler);
-			}
-		};
-
-		if (const FPoseSearchDatabaseSequence* DatabaseSequence = DatabaseAssetStruct.GetPtr<FPoseSearchDatabaseSequence>())
-		{
-			AddSequenceBaseSampler(DatabaseSequence->Sequence);
-		}
-		else if (const FPoseSearchDatabaseAnimComposite* DatabaseAnimComposite = DatabaseAssetStruct.GetPtr<FPoseSearchDatabaseAnimComposite>())
-		{
-			AddSequenceBaseSampler(DatabaseAnimComposite->AnimComposite);
-		}
-		else if (const FPoseSearchDatabaseBlendSpace* DatabaseBlendSpace = DatabaseAssetStruct.GetPtr<FPoseSearchDatabaseBlendSpace>())
+		if (const FPoseSearchDatabaseBlendSpace* DatabaseBlendSpace = DatabaseAssetStruct.GetPtr<FPoseSearchDatabaseBlendSpace>())
 		{
 			if (DatabaseBlendSpace->BlendSpace)
 			{
@@ -98,44 +75,29 @@ bool FDatabaseIndexingContext::IndexDatabase(FPoseSearchIndexBase& SearchIndexBa
 					{
 						const FVector BlendParameters = DatabaseBlendSpace->BlendParameterForSampleRanges(HorizontalIndex, VerticalIndex);
 
-						if (!BlendSpaceSamplerMap.Contains({ DatabaseBlendSpace->BlendSpace, BlendParameters }))
+						if (!SamplerMap.Contains({ DatabaseBlendSpace->BlendSpace, BlendParameters }))
 						{
-							FBlendSpaceSampler::FInput Input;
-							Input.BoneContainer = BoneContainer;
-							Input.BlendSpace = DatabaseBlendSpace->BlendSpace;
-							Input.BlendParameters = BlendParameters;
-
-							BlendSpaceSamplerMap.Add({ DatabaseBlendSpace->BlendSpace, BlendParameters }, Samplers.Num());
-
-							TSharedPtr<FBlendSpaceSampler> Sampler = MakeShared<FBlendSpaceSampler>();
-							Sampler->Init(Input);
-							Samplers.Add(Sampler);
+							SamplerMap.Add({ DatabaseBlendSpace->BlendSpace, BlendParameters }, Samplers.Num());
+							Samplers.Emplace(DatabaseBlendSpace->BlendSpace, BlendParameters, BoneContainer);
 						}
 					}
 				}
 			}
 		}
-		else if (const FPoseSearchDatabaseAnimMontage* DatabaseAnimMontage = DatabaseAssetStruct.GetPtr<FPoseSearchDatabaseAnimMontage>())
+		else if (const FPoseSearchDatabaseAnimationAssetBase* DatabaseAnimationAssetBase = DatabaseAssetStruct.GetPtr<FPoseSearchDatabaseAnimationAssetBase>())
 		{
-			if (DatabaseAnimMontage->AnimMontage && !SamplerMap.Contains(DatabaseAnimMontage->AnimMontage))
+			if (const UAnimationAsset* AnimationAsset = DatabaseAnimationAssetBase->GetAnimationAsset())
 			{
-				FAnimMontageSampler::FInput Input;
-				Input.AnimMontage = DatabaseAnimMontage->AnimMontage;
-
-				SamplerMap.Add(DatabaseAnimMontage->AnimMontage, Samplers.Num());
-
-				TSharedPtr<FAnimMontageSampler> Sampler = MakeShared<FAnimMontageSampler>();
-				Sampler->Init(Input);
-				Samplers.Add(Sampler);
+				if (!SamplerMap.Contains({ AnimationAsset, FVector::ZeroVector }))
+				{
+					SamplerMap.Add({ AnimationAsset, FVector::ZeroVector }, Samplers.Num());
+					Samplers.Emplace(AnimationAsset);
+				}
 			}
 		}
-		else
-		{
-			checkNoEntry();
-		}
 	}
 
-	ParallelFor(Samplers.Num(), [this](int32 SamplerIdx) { Samplers[SamplerIdx]->Process(); }, ParallelForFlags);
+	ParallelFor(Samplers.Num(), [this](int32 SamplerIdx) { Samplers[SamplerIdx].Process(); }, ParallelForFlags);
 
 	if (Owner.IsCanceled())
 	{
@@ -158,38 +120,11 @@ bool FDatabaseIndexingContext::IndexDatabase(FPoseSearchIndexBase& SearchIndexBa
 		IndexerContext.bMirrored = SearchIndexAsset.bMirrored;
 
 		const FInstancedStruct& DatabaseAsset = Database.GetAnimationAssetStruct(SearchIndexAsset.SourceAssetIdx);
-		if (const FPoseSearchDatabaseSequence* DatabaseSequence = DatabaseAsset.GetPtr<FPoseSearchDatabaseSequence>())
-		{
-			if (DatabaseSequence->Sequence)
-			{
-				const float SequenceLength = DatabaseSequence->Sequence->GetPlayLength();
-				IndexerContext.AssetSampler = Samplers[SamplerMap[DatabaseSequence->Sequence]];
-			}
-		}
-		else if (const FPoseSearchDatabaseAnimComposite* DatabaseAnimComposite = DatabaseAsset.GetPtr<FPoseSearchDatabaseAnimComposite>())
-		{
-			if (DatabaseAnimComposite->AnimComposite)
-			{
-				IndexerContext.AssetSampler = Samplers[SamplerMap[DatabaseAnimComposite->AnimComposite]];
-			}
-		}
-		else if (const FPoseSearchDatabaseBlendSpace* DatabaseBlendSpace = DatabaseAsset.GetPtr<FPoseSearchDatabaseBlendSpace>())
-		{
-			if (DatabaseBlendSpace->BlendSpace)
-			{
-				IndexerContext.AssetSampler = Samplers[BlendSpaceSamplerMap[{DatabaseBlendSpace->BlendSpace, SearchIndexAsset.BlendParameters}]];
-			}
-		}
-		else if (const FPoseSearchDatabaseAnimMontage* DatabaseAnimMontage = DatabaseAsset.GetPtr<FPoseSearchDatabaseAnimMontage>())
-		{
-			if (DatabaseAnimMontage->AnimMontage)
-			{
-				IndexerContext.AssetSampler = Samplers[SamplerMap[DatabaseAnimMontage->AnimMontage]];
-			}
-		}
-		else
+		const FPoseSearchDatabaseAnimationAssetBase* DatabaseAnimationAssetBase = DatabaseAsset.GetPtr<FPoseSearchDatabaseAnimationAssetBase>();
+		check(DatabaseAnimationAssetBase);
+		if (DatabaseAnimationAssetBase->GetAnimationAsset())
 		{
-			checkNoEntry();
+			IndexerContext.AssetSampler = &Samplers[SamplerMap[{ DatabaseAnimationAssetBase->GetAnimationAsset(), SearchIndexAsset.BlendParameters }]];
 		}
 
 		const int32 NewIndexerIdx = Indexers.Emplace(IndexerContext, BoneContainer, SearchIndexAsset);
diff --git a/Engine/Plugins/Experimental/Animation/PoseSearch/Source/Runtime/Private/PoseSearchDatabaseIndexingContext.h b/Engine/Plugins/Experimental/Animation/PoseSearch/Source/Runtime/Private/PoseSearchDatabaseIndexingContext.h
--- a/Engine/Plugins/Experimental/Animation/PoseSearch/Source/Runtime/Private/PoseSearchDatabaseIndexingContext.h
+++ b/Engine/Plugins/Experimental/Animation/PoseSearch/Source/Runtime/Private/PoseSearchDatabaseIndexingContext.h
@@ -39,7 +39,7 @@ struct FDatabaseIndexingContext
 
 private:
 	FAssetSamplingContext SamplingContext;
-	TArray<TSharedPtr<FAssetSamplerBase>> Samplers;
+	TArray<FAnimationAssetSampler> Samplers;
 	TArray<FAssetIndexer> Indexers;
 };
 
diff --git a/Engine/Plugins/Experimental/Animation/PoseSearch/Source/Runtime/Private/PoseSearchDerivedDataKey.cpp b/Engine/Plugins/Experimental/Animation/PoseSearch/Source/Runtime/Private/PoseSearchDerivedDataKey.cpp
--- a/Engine/Plugins/Experimental/Animation/PoseSearch/Source/Runtime/Private/PoseSearchDerivedDataKey.cpp
+++ b/Engine/Plugins/Experimental/Animation/PoseSearch/Source/Runtime/Private/PoseSearchDerivedDataKey.cpp
@@ -41,7 +41,7 @@ FKeyBuilder::FKeyBuilder(const UObject* Object, bool bUseDataVer)
 	if (bUseDataVer)
 	{
 		// used to invalidate the key without having to change POSESEARCHDB_DERIVEDDATA_VER all the times
-		int32 POSESEARCHDB_DERIVEDDATA_VER_SMALL = 137;
+		int32 POSESEARCHDB_DERIVEDDATA_VER_SMALL = 139;
 		FGuid VersionGuid = FDevSystemGuids::GetSystemGuid(FDevSystemGuids::Get().POSESEARCHDB_DERIVEDDATA_VER);
 
 		*this << VersionGuid;
diff --git a/Engine/Plugins/Experimental/Animation/PoseSearch/Source/Runtime/Private/PoseSearchLibrary.cpp b/Engine/Plugins/Experimental/Animation/PoseSearch/Source/Runtime/Private/PoseSearchLibrary.cpp
--- a/Engine/Plugins/Experimental/Animation/PoseSearch/Source/Runtime/Private/PoseSearchLibrary.cpp
+++ b/Engine/Plugins/Experimental/Animation/PoseSearch/Source/Runtime/Private/PoseSearchLibrary.cpp
@@ -49,73 +49,35 @@ void FMotionMatchingState::AdjustAssetTime(float AssetTime)
 
 bool FMotionMatchingState::CanAdvance(float DeltaTime) const
 {
-	if (!CurrentSearchResult.IsValid())
-	{
-		return false;
-	}
-
-	const FPoseSearchIndexAsset* SearchIndexAsset = CurrentSearchResult.GetSearchIndexAsset(true);
-
-	const FInstancedStruct& DatabaseAsset = CurrentSearchResult.Database->GetAnimationAssetStruct(*SearchIndexAsset);
-	if (const FPoseSearchDatabaseSequence* DatabaseSequence = DatabaseAsset.GetPtr<FPoseSearchDatabaseSequence>())
+	if (CurrentSearchResult.IsValid())
 	{
-		const float AssetLength = DatabaseSequence->Sequence->GetPlayLength();
-
+		ETypeAdvanceAnim AdvanceType = ETypeAdvanceAnim::ETAA_Default;
 		float SteppedTime = CurrentSearchResult.AssetTime;
-		ETypeAdvanceAnim AdvanceType = FAnimationRuntime::AdvanceTime(
-			DatabaseSequence->Sequence->bLoop,
-			DeltaTime,
-			SteppedTime,
-			AssetLength);
-
-		if (AdvanceType != ETAA_Finished)
+		const FPoseSearchIndexAsset* SearchIndexAsset = CurrentSearchResult.GetSearchIndexAsset(true);
+		const FInstancedStruct& DatabaseAsset = CurrentSearchResult.Database->GetAnimationAssetStruct(*SearchIndexAsset);
+		if (const FPoseSearchDatabaseBlendSpace* DatabaseBlendSpace = DatabaseAsset.GetPtr<FPoseSearchDatabaseBlendSpace>())
 		{
-			return SearchIndexAsset->SamplingInterval.Contains(SteppedTime);
-		}
-	}
-	else if (const FPoseSearchDatabaseAnimComposite* DatabaseAnimComposite = DatabaseAsset.GetPtr<FPoseSearchDatabaseAnimComposite>())
-	{
-		const float AssetLength = DatabaseAnimComposite->GetAnimationAsset()->GetPlayLength();
+			TArray<FBlendSampleData> BlendSamples;
+			int32 TriangulationIndex = 0;
+			DatabaseBlendSpace->BlendSpace->GetSamplesFromBlendInput(SearchIndexAsset->BlendParameters, BlendSamples, TriangulationIndex, true);
 
-		float SteppedTime = CurrentSearchResult.AssetTime;
-		ETypeAdvanceAnim AdvanceType = FAnimationRuntime::AdvanceTime(
-			DatabaseAnimComposite->IsLooping(),
-			DeltaTime,
-			SteppedTime,
-			AssetLength);
+			const float PlayLength = DatabaseBlendSpace->BlendSpace->GetAnimationLengthFromSampleData(BlendSamples);
 
-		if (AdvanceType != ETAA_Finished)
+			// Asset player time for blend spaces is normalized [0, 1] so we need to convert it back to real time before we advance it
+			SteppedTime = CurrentSearchResult.AssetTime * PlayLength;
+			AdvanceType = FAnimationRuntime::AdvanceTime(DatabaseBlendSpace->IsLooping(), DeltaTime, SteppedTime, PlayLength);
+		}
+		else if (const FPoseSearchDatabaseAnimationAssetBase* DatabaseAnimationAssetBase = DatabaseAsset.GetPtr<FPoseSearchDatabaseAnimationAssetBase>())
 		{
-			return SearchIndexAsset->SamplingInterval.Contains(SteppedTime);
+			const float AssetLength = DatabaseAnimationAssetBase->GetAnimationAsset()->GetPlayLength();
+			AdvanceType = FAnimationRuntime::AdvanceTime(DatabaseAnimationAssetBase->IsLooping(), DeltaTime, SteppedTime, AssetLength);
 		}
-	}
-	else if (const FPoseSearchDatabaseBlendSpace* DatabaseBlendSpace = DatabaseAsset.GetPtr<FPoseSearchDatabaseBlendSpace>())
-	{
-		TArray<FBlendSampleData> BlendSamples;
-		int32 TriangulationIndex = 0;
-		DatabaseBlendSpace->BlendSpace->GetSamplesFromBlendInput(SearchIndexAsset->BlendParameters, BlendSamples, TriangulationIndex, true);
-
-		float PlayLength = DatabaseBlendSpace->BlendSpace->GetAnimationLengthFromSampleData(BlendSamples);
-
-		// Asset player time for blendspaces is normalized [0, 1] so we need to convert 
-		// to a real time before we advance it
-		float RealTime = CurrentSearchResult.AssetTime * PlayLength;
-		float SteppedTime = RealTime;
-		ETypeAdvanceAnim AdvanceType = FAnimationRuntime::AdvanceTime(
-			DatabaseBlendSpace->BlendSpace->bLoop,
-			DeltaTime,
-			SteppedTime,
-			PlayLength);
 
 		if (AdvanceType != ETAA_Finished)
 		{
 			return SearchIndexAsset->SamplingInterval.Contains(SteppedTime);
 		}
 	}
-	else
-	{
-		checkNoEntry();
-	}
 
 	return false;
 }
@@ -358,45 +320,16 @@ void UPoseSearchLibrary::UpdateMotionMatchingState(
 				TArrayView<float> ReconstructedPoseValuesBuffer((float*)FMemory_Alloca(NumDimensions * sizeof(float)), NumDimensions);
 				const TConstArrayView<float> PoseValues = SearchIndex.Values.IsEmpty() ? SearchIndex.GetReconstructedPoseValues(PoseIdx, ReconstructedPoseValuesBuffer) : SearchIndex.GetPoseValues(PoseIdx);
 
-				const float SampleTime = ContinuingPoseDatabase->GetAssetTime(PoseIdx);
-
+				// extracting notifies from the database animation asset at time SampleTime to search for UAnimNotifyState_PoseSearchOverrideContinuingPoseCostBias eventually overriding the schema ContinuingPoseCostBias
 				const FPoseSearchIndexAsset& SearchIndexAsset = SearchIndex.GetAssetForPose(PoseIdx);
-				const FInstancedStruct& DatabaseAssetStruct = ContinuingPoseDatabase->GetAnimationAssetStruct(SearchIndexAsset);
+				const FPoseSearchDatabaseAnimationAssetBase* DatabaseAnimationAssetBase = ContinuingPoseDatabase->GetAnimationAssetStruct(SearchIndexAsset).GetPtr<FPoseSearchDatabaseAnimationAssetBase>();
+				check(DatabaseAnimationAssetBase);
+				const FAnimationAssetSampler SequenceBaseSampler(DatabaseAnimationAssetBase->GetAnimationAsset(), SearchIndexAsset.BlendParameters);
+				const float SampleTime = ContinuingPoseDatabase->GetAssetTime(PoseIdx);
 
-				// @todo: create a unified sampler to sample any animation asset
 				// @todo: change ExtractPoseSearchNotifyStates api to avoid NotifyStates allocation
-
-				// extracting notifies from the database animation asset at time SampleTime to search for UAnimNotifyState_PoseSearchOverrideContinuingPoseCostBias eventually overriding the schema ContinuingPoseCostBias
 				TArray<UAnimNotifyState_PoseSearchBase*> NotifyStates;
-				if (const FPoseSearchDatabaseSequence* DatabaseSequence = DatabaseAssetStruct.GetPtr<FPoseSearchDatabaseSequence>())
-				{
-					FSequenceBaseSampler SequenceBaseSampler;
-					FSequenceBaseSampler::FInput Input;
-					Input.SequenceBase = DatabaseSequence->Sequence;
-					SequenceBaseSampler.Init(Input);
-					SequenceBaseSampler.ExtractPoseSearchNotifyStates(SampleTime, NotifyStates);
-				}
-				else if (const FPoseSearchDatabaseAnimComposite* DatabaseAnimComposite = DatabaseAssetStruct.GetPtr<FPoseSearchDatabaseAnimComposite>())
-				{
-					FSequenceBaseSampler SequenceBaseSampler;
-					FSequenceBaseSampler::FInput Input;
-					Input.SequenceBase = DatabaseAnimComposite->AnimComposite;
-					SequenceBaseSampler.Init(Input);
-					SequenceBaseSampler.ExtractPoseSearchNotifyStates(SampleTime, NotifyStates);
-				}
-				else if (const FPoseSearchDatabaseBlendSpace* DatabaseBlendSpace = DatabaseAssetStruct.GetPtr<FPoseSearchDatabaseBlendSpace>())
-				{
-					FBlendSpaceSampler BlendSpaceSampler;
-					FBlendSpaceSampler::FInput Input;
-					Input.BlendSpace = DatabaseBlendSpace->BlendSpace;
-					Input.BlendParameters = SearchIndexAsset.BlendParameters;
-					BlendSpaceSampler.Init(Input);
-					BlendSpaceSampler.ExtractPoseSearchNotifyStates(SampleTime, NotifyStates);
-				}
-				else
-				{
-					checkNoEntry();
-				}
+				SequenceBaseSampler.ExtractPoseSearchNotifyStates(SampleTime, NotifyStates);
 
 				float ContinuingPoseCostBias = ContinuingPoseDatabase->Schema->ContinuingPoseCostBias;
 				for (const UAnimNotifyState_PoseSearchBase* PoseSearchNotify : NotifyStates)
@@ -466,6 +399,7 @@ void UPoseSearchLibrary::MotionMatch(
 	bool& bLoop,
 	bool& bIsMirrored,
 	FVector& BlendParameters,
+	float& SearchCost,
 	const UAnimationAsset* FutureAnimation,
 	float FutureAnimationStartTime,
 	float TimeToFutureAnimationStart,
@@ -493,6 +427,7 @@ void UPoseSearchLibrary::MotionMatch(
 	bLoop = false;
 	bIsMirrored = false;
 	BlendParameters = FVector::ZeroVector;
+	SearchCost = MAX_flt;
 
 	if (Searchable)
 	{
@@ -525,92 +460,58 @@ void UPoseSearchLibrary::MotionMatch(
 			else if (FutureAnimation)
 			{
 				const FBoneContainer& BoneContainer = AnimInstance->GetRequiredBonesOnAnyThread();
+				// @todo... add input BlendParameters to support sampling FutureAnimation blendspaces
+				const FAnimationAssetSampler Sampler(FutureAnimation, FVector::ZeroVector, BoneContainer);
 
-				FBlendSpaceSampler BlendSpaceSampler;
-				FSequenceBaseSampler SequenceBaseSampler;
-				FAnimMontageSampler AnimMontageSampler;
-				const FAssetSamplerBase* AssetSamplerBase = nullptr;
-				if (const UBlendSpace* BlendSpace = Cast<UBlendSpace>(FutureAnimation))
-				{
-					FBlendSpaceSampler::FInput Input;
-					Input.BoneContainer = BoneContainer;
-					Input.BlendSpace = BlendSpace;
-					// @todo... add BlendParameters
-					//Input.BlendParameters = BlendParameters;
-					BlendSpaceSampler.Init(Input);
-					AssetSamplerBase = &BlendSpaceSampler;
-				}
-				else if (const UAnimMontage* AnimMontage = Cast<UAnimMontage>(FutureAnimation))
-				{
-					FAnimMontageSampler::FInput Input;
-					Input.AnimMontage = AnimMontage;
-					AnimMontageSampler.Init(Input);
-					AssetSamplerBase = &AnimMontageSampler;
-				}
-				else if (const UAnimSequenceBase* SequenceBase = Cast<UAnimSequenceBase>(FutureAnimation))
+				FCompactPose Pose;
+				FBlendedCurve UnusedCurve;
+				FStackAttributeContainer UnusedAtrribute;
+				FAnimationPoseData AnimPoseData = { Pose, UnusedCurve, UnusedAtrribute };
+
+				UnusedCurve.InitFrom(BoneContainer);
+				Pose.SetBoneContainer(&BoneContainer);
+
+				if (FutureAnimationStartTime < FiniteDelta)
 				{
-					FSequenceBaseSampler::FInput Input;
-					Input.SequenceBase = SequenceBase;
-					SequenceBaseSampler.Init(Input);
-					AssetSamplerBase = &SequenceBaseSampler;
+					UE_LOG(LogPoseSearch, Warning, TEXT("UPoseSearchLibrary::MotionMatch - provided FutureAnimationStartTime (%f) is too small to be able to calculate velocities. Clamping it to minimum value of %f"), FutureAnimationStartTime, FiniteDelta);
+					FutureAnimationStartTime = FiniteDelta;
 				}
-				else
+
+				const float MinTimeToFutureAnimationStart = FiniteDelta + UE_KINDA_SMALL_NUMBER;
+				if (TimeToFutureAnimationStart < MinTimeToFutureAnimationStart)
 				{
-					UE_LOG(LogPoseSearch, Error, TEXT("UPoseSearchLibrary::MotionMatch - unsupported FutureAnimation '%s'"), *GetNameSafe(FutureAnimation));
+					UE_LOG(LogPoseSearch, Warning, TEXT("UPoseSearchLibrary::MotionMatch - provided TimeToFutureAnimationStart (%f) is too small. Clamping it to minimum value of %f"), TimeToFutureAnimationStart, MinTimeToFutureAnimationStart);
+					TimeToFutureAnimationStart = MinTimeToFutureAnimationStart;
 				}
 
-				if (AssetSamplerBase)
+				// extracting 2 poses to be able to calculate velocities
+				for (int i = 0; i < 2; ++i)
 				{
-					FCompactPose Pose;
-					FBlendedCurve UnusedCurve;
-					FStackAttributeContainer UnusedAtrribute;
-					FAnimationPoseData AnimPoseData = { Pose, UnusedCurve, UnusedAtrribute };
-
-					UnusedCurve.InitFrom(BoneContainer);
-					Pose.SetBoneContainer(&BoneContainer);
+					const float ExtractionTime = FutureAnimationStartTime + (i - 1) * FiniteDelta;
+					const float FutureAnimationTime = TimeToFutureAnimationStart + (i - 1) * FiniteDelta;
 
-					if (FutureAnimationStartTime < FiniteDelta)
-					{
-						UE_LOG(LogPoseSearch, Warning, TEXT("UPoseSearchLibrary::MotionMatch - provided FutureAnimationStartTime (%f) is too small to be able to calculate velocities. Clamping it to minimum value of %f"), FutureAnimationStartTime, FiniteDelta);
-						FutureAnimationStartTime = FiniteDelta;
-					}
-
-					const float MinTimeToFutureAnimationStart = FiniteDelta + UE_KINDA_SMALL_NUMBER;
-					if (TimeToFutureAnimationStart < MinTimeToFutureAnimationStart)
-					{
-						UE_LOG(LogPoseSearch, Warning, TEXT("UPoseSearchLibrary::MotionMatch - provided TimeToFutureAnimationStart (%f) is too small. Clamping it to minimum value of %f"), TimeToFutureAnimationStart, MinTimeToFutureAnimationStart);
-						TimeToFutureAnimationStart = MinTimeToFutureAnimationStart;
-					}
-
-					// extracting 2 poses to be able to calculate velocities
-					for (int i = 0; i < 2; ++i)
-					{
-						const float ExtractionTime = FutureAnimationStartTime + (i - 1) * FiniteDelta;
-						const float FutureAnimationTime = TimeToFutureAnimationStart + (i - 1) * FiniteDelta;
-
-						FDeltaTimeRecord DeltaTimeRecord;
-						DeltaTimeRecord.Set(ExtractionTime - FiniteDelta, FiniteDelta);
-						FAnimExtractContext ExtractionCtx(double(ExtractionTime), false, DeltaTimeRecord, false);
+					FDeltaTimeRecord DeltaTimeRecord;
+					DeltaTimeRecord.Set(ExtractionTime - FiniteDelta, FiniteDelta);
+					FAnimExtractContext ExtractionCtx(double(ExtractionTime), false, DeltaTimeRecord, false);
 
-						AssetSamplerBase->ExtractPose(ExtractionCtx, AnimPoseData);
+					Sampler.ExtractPose(ExtractionCtx, AnimPoseData);
 
-						FCSPose<FCompactPose> ComponentSpacePose;
-						ComponentSpacePose.InitPose(Pose);
+					FCSPose<FCompactPose> ComponentSpacePose;
+					ComponentSpacePose.InitPose(Pose);
 
-						const FTrajectorySample TrajectorySample = Trajectory.GetSampleAtTime(ExtractionTime);
-						const FTransform& ComponentTransform = AnimInstance->GetOwningComponent()->GetComponentTransform();
-						const FTransform FutureComponentTransform = TrajectorySample.Transform * ComponentTransform;
+					const FTrajectorySample TrajectorySample = Trajectory.GetSampleAtTime(ExtractionTime);
+					const FTransform& ComponentTransform = AnimInstance->GetOwningComponent()->GetComponentTransform();
+					const FTransform FutureComponentTransform = TrajectorySample.Transform * ComponentTransform;
 
-						ExtendedPoseHistory.AddFuturePose(FutureAnimationTime, ComponentSpacePose, FutureComponentTransform);
-					}
+					ExtendedPoseHistory.AddFuturePose(FutureAnimationTime, ComponentSpacePose, FutureComponentTransform);
+				}
 
 #if ENABLE_DRAW_DEBUG && ENABLE_ANIM_DEBUG
-					if (CVarAnimMotionMatchDrawHistoryEnable.GetValueOnAnyThread())
-					{
-						ExtendedPoseHistory.DebugDraw(UAnimInstanceProxyProvider::GetAnimInstanceProxy(AnimInstance));
-					}
-#endif // ENABLE_DRAW_DEBUG && ENABLE_ANIM_DEBUG
+				if (CVarAnimMotionMatchDrawHistoryEnable.GetValueOnAnyThread())
+				{
+					ExtendedPoseHistory.DebugDraw(UAnimInstanceProxyProvider::GetAnimInstanceProxy(AnimInstance));
 				}
+#endif // ENABLE_DRAW_DEBUG && ENABLE_ANIM_DEBUG
 			}
 		}
 
@@ -628,6 +529,7 @@ void UPoseSearchLibrary::MotionMatch(
 				bLoop = DatabaseAsset->IsLooping();
 				bIsMirrored = SearchIndexAsset->bMirrored;
 				BlendParameters = SearchIndexAsset->BlendParameters;
+				SearchCost = SearchResult.PoseCost.GetTotalCost();
 			}
 		}
 
diff --git a/Engine/Plugins/Experimental/Animation/PoseSearch/Source/Runtime/Public/PoseSearch/PoseSearchAssetIndexer.h b/Engine/Plugins/Experimental/Animation/PoseSearch/Source/Runtime/Public/PoseSearch/PoseSearchAssetIndexer.h
--- a/Engine/Plugins/Experimental/Animation/PoseSearch/Source/Runtime/Public/PoseSearch/PoseSearchAssetIndexer.h
+++ b/Engine/Plugins/Experimental/Animation/PoseSearch/Source/Runtime/Public/PoseSearch/PoseSearchAssetIndexer.h
@@ -16,7 +16,7 @@ struct FPoseSearchPoseMetadata;
 namespace UE::PoseSearch
 {
 
-struct FAssetSamplerBase;
+struct FAnimationAssetSampler;
 struct FAssetSamplingContext;
 
 /**
@@ -26,9 +26,9 @@ struct FAssetIndexingContext
 {
 	const FAssetSamplingContext* SamplingContext = nullptr;
 	const UPoseSearchSchema* Schema = nullptr;
-	TSharedPtr<FAssetSamplerBase> AssetSampler;
+	const FAnimationAssetSampler* AssetSampler = nullptr;
 	bool bMirrored = false;
-	FFloatInterval RequestedSamplingRange = FFloatInterval(0.0f, 0.0f);
+	FFloatInterval RequestedSamplingRange = FFloatInterval(0.f, 0.f);
 };
 
 class FAssetIndexer
@@ -68,7 +68,6 @@ private:
 
 	struct FSampleInfo
 	{
-		TSharedPtr<FAssetSamplerBase> Clip;
 		FTransform RootTransform;
 		float ClipTime = 0.f;
 		bool bClamped = false;
diff --git a/Engine/Plugins/Experimental/Animation/PoseSearch/Source/Runtime/Public/PoseSearch/PoseSearchAssetSampler.h b/Engine/Plugins/Experimental/Animation/PoseSearch/Source/Runtime/Public/PoseSearch/PoseSearchAssetSampler.h
--- a/Engine/Plugins/Experimental/Animation/PoseSearch/Source/Runtime/Public/PoseSearch/PoseSearchAssetSampler.h
+++ b/Engine/Plugins/Experimental/Animation/PoseSearch/Source/Runtime/Public/PoseSearch/PoseSearchAssetSampler.h
@@ -7,132 +7,54 @@
 struct FAnimationPoseData;
 struct FAnimExtractContext;
 class UAnimationAsset;
-class UAnimInstance;
 class UAnimNotifyState_PoseSearchBase;
-class UBlendSpace;
-class UMirrorDataTable;
 
 namespace UE::PoseSearch
 {
 
 /**
- * Helper interface for sampling data from animation assets
+ * Helper for sampling data from animation assets
  */
-struct POSESEARCH_API FAssetSamplerBase : public TSharedFromThis<FAssetSamplerBase>
+struct POSESEARCH_API FAnimationAssetSampler
 {
-	virtual ~FAssetSamplerBase() {}
+	FAnimationAssetSampler(TWeakObjectPtr<const UAnimationAsset> InAnimationAsset = nullptr, const FVector& InBlendParameters = FVector::ZeroVector, const FBoneContainer& InBoneContainer = FBoneContainer(), int32 InRootTransformSamplingRate = 30);
+	void Init(TWeakObjectPtr<const UAnimationAsset> InAnimationAsset, const FVector& InBlendParameters = FVector::ZeroVector, const FBoneContainer& InBoneContainer = FBoneContainer(), int32 InRootTransformSamplingRate = 30);
 
-	virtual float GetPlayLength() const { return 0.f; }
-	virtual float GetScaledTime(float Time) const { return Time; }
-	virtual bool IsLoopable() const { return false; }
+	bool IsInitialized() const;
+	float GetPlayLength() const;
+	float GetScaledTime(float Time) const;
+	bool IsLoopable() const;
 
 	// Gets the final root transformation at the end of the asset's playback time
-	virtual FTransform GetTotalRootTransform() const { return FTransform::Identity; }
+	FTransform GetTotalRootTransform() const;
 
 	// Extracts pose for this asset for a given context
-	virtual void ExtractPose(const FAnimExtractContext& ExtractionCtx, FAnimationPoseData& OutAnimPoseData) const { }
+	void ExtractPose(const FAnimExtractContext& ExtractionCtx, FAnimationPoseData& OutAnimPoseData) const;
 
 	// Extracts root transform at the given time, using the extremities of the sequence to extrapolate beyond the 
 	// sequence limits when Time is less than zero or greater than the sequence length.
-	virtual FTransform ExtractRootTransform(float Time) const { return FTransform::Identity; }
+	FTransform ExtractRootTransform(float Time) const;
 
 	// Extracts notify states inheriting from UAnimNotifyState_PoseSearchBase present in the sequence at Time.
 	// The function does not empty NotifyStates before adding new notifies!
-	virtual void ExtractPoseSearchNotifyStates(float Time, TArray<UAnimNotifyState_PoseSearchBase*>& NotifyStates) const { }
+	void ExtractPoseSearchNotifyStates(float Time, TArray<UAnimNotifyState_PoseSearchBase*>& NotifyStates) const;
 
-	virtual const UAnimationAsset* GetAsset() const { return nullptr; }
+	const UAnimationAsset* GetAsset() const;
 
-	virtual void Process() { }
+	void Process();
 
-	const float ExtrapolationSampleTime = 1.f / 30.f;
-	const float ExtractionInterval = 1.0f / 120.0f;
-};
-
-// Sampler working with UAnimSequenceBase so it can be used for UAnimSequence as well as UAnimComposite.
-struct POSESEARCH_API FSequenceBaseSampler : public FAssetSamplerBase
-{
-	struct FInput
-	{
-		TWeakObjectPtr<const UAnimSequenceBase> SequenceBase;
-	} Input;
-
-	void Init(const FInput& Input);
-	virtual void Process() override;
-
-	virtual float GetPlayLength() const override;
-	virtual float GetScaledTime(float Time) const override;
-	virtual bool IsLoopable() const override;
-
-	virtual FTransform GetTotalRootTransform() const override;
-
-	virtual void ExtractPose(const FAnimExtractContext& ExtractionCtx, FAnimationPoseData& OutAnimPoseData) const override;
-	virtual FTransform ExtractRootTransform(float Time) const override;
-	virtual void ExtractPoseSearchNotifyStates(float Time, TArray<class UAnimNotifyState_PoseSearchBase*>& NotifyStates) const override;
-	virtual const UAnimationAsset* GetAsset() const override;
-};
-
-struct POSESEARCH_API FBlendSpaceSampler : public FAssetSamplerBase
-{
-	struct FInput
-	{
-		FBoneContainer BoneContainer;
-		TWeakObjectPtr<const UBlendSpace> BlendSpace;
-		int32 RootTransformSamplingRate = 30;
-		FVector BlendParameters;
-	} Input;
-
-	void Init(const FInput& Input);
-	virtual void Process() override;
-
-	virtual float GetPlayLength() const override { return PlayLength; }
-	virtual float GetScaledTime(float Time) const override;
-	virtual bool IsLoopable() const override;
+protected:
+	TWeakObjectPtr<const UAnimationAsset> AnimationAsset;
 
-	virtual FTransform GetTotalRootTransform() const override;
-
-	virtual void ExtractPose(const FAnimExtractContext& ExtractionCtx, FAnimationPoseData& OutAnimPoseData) const override;
-	virtual FTransform ExtractRootTransform(float Time) const override;
-	virtual void ExtractPoseSearchNotifyStates(float Time, TArray<class UAnimNotifyState_PoseSearchBase*>& NotifyStates) const override;
-
-	virtual const UAnimationAsset* GetAsset() const override;
-
-private:
-	float PlayLength = 0.0f;
+	// members used to sample blend spaces only!
+	FVector BlendParameters = FVector::ZeroVector;
+	FBoneContainer BoneContainer;
+	int32 RootTransformSamplingRate = 30;
+	float CachedPlayLength = -1.f;
 	TArray<FTransform> AccumulatedRootTransform;
 
-	void ProcessPlayLength();
-	void ProcessRootTransform();
-
-	// Extracts the pre-computed blend space root transform. ProcessRootTransform must be run first.
-	FTransform ExtractBlendSpaceRootTrackTransform(float Time) const;
-	FTransform ExtractBlendSpaceRootMotion(float StartTime, float DeltaTime, bool bAllowLooping) const;
-	FTransform ExtractBlendSpaceRootMotionFromRange(float StartTrackPosition, float EndTrackPosition) const;
-};
-
-struct POSESEARCH_API FAnimMontageSampler : public FAssetSamplerBase
-{
-	struct FInput
-	{
-		// @todo: add support for SlotName / multiple SlotAnimTracks
-		TWeakObjectPtr<const UAnimMontage> AnimMontage;
-	} Input;
-
-	void Init(const FInput& Input);
-	virtual void Process() override;
-
-	virtual float GetPlayLength() const override;
-	virtual float GetScaledTime(float Time) const override;
-	virtual bool IsLoopable() const override;
-
-	virtual FTransform GetTotalRootTransform() const override;
-
-	virtual void ExtractPose(const FAnimExtractContext& ExtractionCtx, FAnimationPoseData& OutAnimPoseData) const override;
-	virtual FTransform ExtractRootTransform(float Time) const override;
-	virtual void ExtractPoseSearchNotifyStates(float Time, TArray<class UAnimNotifyState_PoseSearchBase*>& NotifyStates) const override;
-	virtual const UAnimationAsset* GetAsset() const override;
-
-private:
-	FTransform ExtractRootTransformInternal(float StartTime, float EndTime) const;
+	const float ExtrapolationSampleTime = 1.f / 30.f;
+	const float ExtractionInterval = 1.0f / 120.0f;
 };
 
 } // namespace UE::PoseSearch
diff --git a/Engine/Plugins/Experimental/Animation/PoseSearch/Source/Runtime/Public/PoseSearch/PoseSearchLibrary.h b/Engine/Plugins/Experimental/Animation/PoseSearch/Source/Runtime/Public/PoseSearch/PoseSearchLibrary.h
--- a/Engine/Plugins/Experimental/Animation/PoseSearch/Source/Runtime/Public/PoseSearch/PoseSearchLibrary.h
+++ b/Engine/Plugins/Experimental/Animation/PoseSearch/Source/Runtime/Public/PoseSearch/PoseSearchLibrary.h
@@ -158,6 +158,7 @@ public:
 	* @param bLoop							Output selected animation looping state
 	* @param bIsMirrored					Output selected animation mirror state
 	* @param BlendParameters				Output selected animation blend space parameters (if SelectedAnimation is a blend space)
+	* @param SearchCost						Output search associated cost
 	* @param FutureAnimation				Input animation we want to match after TimeToFutureAnimationStart seconds
 	* @param FutureAnimationStartTime		Input start time for the first pose of FutureAnimation
 	* @param TimeToFutureAnimationStart		Input time in seconds before start playing FutureAnimation (from FutureAnimationStartTime seconds)
@@ -174,6 +175,7 @@ public:
 		bool& bLoop,
 		bool& bIsMirrored,
 		FVector& BlendParameters,
+		float& SearchCost,
 		const UAnimationAsset* FutureAnimation = nullptr,
 		float FutureAnimationStartTime = 0.f,
 		float TimeToFutureAnimationStart = 0.f,
