Anim Node Functions for Motion Matching that can push one or more databases to search.

diff --git a/Engine/Plugins/Experimental/Animation/PoseSearch/Source/Runtime/Private/AnimNode_MotionMatching.cpp b/Engine/Plugins/Experimental/Animation/PoseSearch/Source/Runtime/Private/AnimNode_MotionMatching.cpp
--- a/Engine/Plugins/Experimental/Animation/PoseSearch/Source/Runtime/Private/AnimNode_MotionMatching.cpp
+++ b/Engine/Plugins/Experimental/Animation/PoseSearch/Source/Runtime/Private/AnimNode_MotionMatching.cpp
@@ -78,13 +78,13 @@ void FAnimNode_MotionMatching::UpdateAssetPlayer(const FAnimationUpdateContext&
 #if WITH_EDITOR
 		if (FAsyncPoseSearchDatabasesManagement::RequestAsyncBuildIndex(MotionMatchingState.CurrentSearchResult.Database.Get(), ERequestAsyncBuildFlag::ContinueRequest))
 		{
-			const UPoseSearchDatabase* Database = MotionMatchingState.CurrentSearchResult.Database.Get();
-			const FPoseSearchIndex& SearchIndex = Database->GetSearchIndex();
+			const UPoseSearchDatabase* CurrentResultDatabase = MotionMatchingState.CurrentSearchResult.Database.Get();
+			const FPoseSearchIndex& SearchIndex = CurrentResultDatabase->GetSearchIndex();
 			if (!SearchIndex.IsValidPoseIndex(MotionMatchingState.CurrentSearchResult.PrevPoseIdx) ||
 				!SearchIndex.IsValidPoseIndex(MotionMatchingState.CurrentSearchResult.NextPoseIdx) ||
-				Database->Schema != MotionMatchingState.CurrentSearchResult.ComposedQuery.GetSchema())
+				CurrentResultDatabase->Schema != MotionMatchingState.CurrentSearchResult.ComposedQuery.GetSchema())
 			{
-				// MotionMatchingState is out of sync with Database: we need to reset the MM state. This could happen if PIE is paused, and we edit the database,
+				// MotionMatchingState is out of sync with CurrentResultDatabase: we need to reset the MM state. This could happen if PIE is paused, and we edit the database,
 				// so FAnimNode_MotionMatching::UpdateAssetPlayer is never called and FAsyncPoseSearchDatabasesManagement::RequestAsyncBuildIndex never returns false here
 				MotionMatchingState.Reset();
 			}
@@ -103,27 +103,44 @@ void FAnimNode_MotionMatching::UpdateAssetPlayer(const FAnimationUpdateContext&
 	}
 	UpdateCounter.SynchronizeWith(Context.AnimInstanceProxy->GetUpdateCounter());
 
+	// If the Database property hasn't been overridden, set it as the only database to search.
+	if (!bOverrideDatabaseInput && Database)
+	{
+		DatabasesToSearch.Reset(1);
+		DatabasesToSearch.Add(Database);
+	}
+	// @todo: Delete this after removing Searchable.
+	else if (Searchable)
+	{
+		if (!bOverrideDatabaseInput && Searchable->IsA(UPoseSearchDatabase::StaticClass()))
+		{
+			DatabasesToSearch.Reset(1);
+			DatabasesToSearch.Add(Cast<UPoseSearchDatabase>(Searchable));
+		}
+		UE_LOG(LogPoseSearch, Warning, TEXT("FAnimNode_MotionMatching: The Searchable variable will be removed soon."));
+	}
+
 	// Execute core motion matching algorithm
 	UPoseSearchLibrary::UpdateMotionMatchingState(
 		Context,
-		Searchable,
+		DatabasesToSearch,
 		Trajectory,
 		Settings,
 		MotionMatchingState,
-		bForceInterrupt
+		bForceInterrupt | bForceInterruptNextUpdate
 	);
 
 	// If a new pose is requested, blend into the new asset via BlendStackNode
 	if (MotionMatchingState.bJumpedToPose)
 	{
 		const FPoseSearchIndexAsset* SearchIndexAsset = MotionMatchingState.CurrentSearchResult.GetSearchIndexAsset();
-		const UPoseSearchDatabase* Database = MotionMatchingState.CurrentSearchResult.Database.Get();
-		if (SearchIndexAsset && Database && Database->Schema)
+		const UPoseSearchDatabase* CurrentResultDatabase = MotionMatchingState.CurrentSearchResult.Database.Get();
+		if (SearchIndexAsset && CurrentResultDatabase && CurrentResultDatabase->Schema)
 		{
-			if (const FPoseSearchDatabaseAnimationAssetBase* DatabaseAsset = Database->GetAnimationAssetBase(*SearchIndexAsset))
+			if (const FPoseSearchDatabaseAnimationAssetBase* DatabaseAsset = CurrentResultDatabase->GetAnimationAssetBase(*SearchIndexAsset))
 			{
 				BlendStackNode.BlendTo(DatabaseAsset->GetAnimationAsset(), MotionMatchingState.CurrentSearchResult.AssetTime,
-					DatabaseAsset->IsLooping(), SearchIndexAsset->bMirrored, Database->Schema->MirrorDataTable.Get(),
+					DatabaseAsset->IsLooping(), SearchIndexAsset->bMirrored, CurrentResultDatabase->Schema->MirrorDataTable.Get(),
 					Settings.MaxActiveBlends, Settings.BlendTime, Settings.BlendProfile, Settings.BlendOption, SearchIndexAsset->BlendParameters, MotionMatchingState.WantedPlayRate);
 			}
 		}
@@ -193,6 +210,8 @@ void FAnimNode_MotionMatching::UpdateAssetPlayer(const FAnimationUpdateContext&
 		}
 	} 
 #endif
+
+	bForceInterruptNextUpdate = false;
 }
 
 void FAnimNode_MotionMatching::GatherDebugData(FNodeDebugData& DebugData)
@@ -200,6 +219,84 @@ void FAnimNode_MotionMatching::GatherDebugData(FNodeDebugData& DebugData)
 	Source.GatherDebugData(DebugData);
 }
 
+void FAnimNode_MotionMatching::SetDatabaseToSearch(UPoseSearchDatabase* InDatabase, bool bForceInterruptIfNew)
+{
+	if (DatabasesToSearch.Num() == 1 && DatabasesToSearch[0] == InDatabase)
+	{
+		UE_LOG(LogPoseSearch, Verbose, TEXT("FAnimNode_MotionMatching::SetDatabaseToSearch - Database(%s) is already set."), *GetNameSafe(InDatabase));
+	}
+	else
+	{
+		DatabasesToSearch.Reset();
+		bOverrideDatabaseInput = false;
+		if (InDatabase)
+		{
+			DatabasesToSearch.Add(InDatabase);
+			bOverrideDatabaseInput = true;
+		}
+
+		bForceInterruptNextUpdate |= bForceInterruptIfNew;
+
+		UE_LOG(LogPoseSearch, Verbose, TEXT("FAnimNode_MotionMatching::SetDatabaseToSearch - Setting to Database(%s), bForceInterruptIfNew(%d)."), *GetNameSafe(InDatabase), bForceInterruptIfNew);
+	}
+}
+
+void FAnimNode_MotionMatching::SetDatabasesToSearch(const TArray<UPoseSearchDatabase*>& InDatabases, bool bForceInterruptIfNew)
+{
+	// Check if InDatabases and DatabasesToSearch are the same.
+	bool bDatabasesAlreadySet = true;
+	if (DatabasesToSearch.Num() != InDatabases.Num())
+	{
+		bDatabasesAlreadySet = false;
+	}
+	else
+	{
+		for (int32 Index = 0; Index < InDatabases.Num(); ++Index)
+		{
+			if (DatabasesToSearch[Index] != InDatabases[Index])
+			{
+				bDatabasesAlreadySet = false;
+				break;
+			}
+		}
+	}
+
+	if (bDatabasesAlreadySet)
+	{
+		UE_LOG(LogPoseSearch, Verbose, TEXT("FAnimNode_MotionMatching::SetDatabasesToSearch - Databases(#%d) already set."), InDatabases.Num());
+	}
+	else
+	{
+		DatabasesToSearch.Reset();
+		bOverrideDatabaseInput = false;
+		if (!InDatabases.IsEmpty())
+		{
+			DatabasesToSearch.Append(InDatabases);
+			bOverrideDatabaseInput = true;
+		}
+
+		bForceInterruptNextUpdate |= bForceInterruptIfNew;
+
+		UE_LOG(LogPoseSearch, Verbose, TEXT("FAnimNode_MotionMatching::SetDatabaseToSearch - Setting to Databases(#%d), bForceInterruptIfNew(%d)."), InDatabases.Num(), bForceInterruptIfNew);
+	}
+}
+
+void FAnimNode_MotionMatching::ResetDatabasesToSearch(bool bInForceInterrupt)
+{
+	DatabasesToSearch.Reset();
+	bOverrideDatabaseInput = false;
+	bForceInterruptNextUpdate = bInForceInterrupt;
+
+	UE_LOG(LogPoseSearch, Verbose, TEXT("FAnimNode_MotionMatching::ResetDatabasesToSearch - Resetting databases, bInForceInterrupt(%d)."), bInForceInterrupt);
+}
+
+void FAnimNode_MotionMatching::ForceInterruptNextUpdate()
+{
+	bForceInterruptNextUpdate = true;
+
+	UE_LOG(LogPoseSearch, Verbose, TEXT("FAnimNode_MotionMatching::ForceInterruptNextUpdate - Forcing interrupt."));
+}
+
 // FAnimNode_AssetPlayerBase interface
 float FAnimNode_MotionMatching::GetAccumulatedTime() const
 {
diff --git a/Engine/Plugins/Experimental/Animation/PoseSearch/Source/Runtime/Private/MotionMatchingAnimNodeLibrary.cpp b/Engine/Plugins/Experimental/Animation/PoseSearch/Source/Runtime/Private/MotionMatchingAnimNodeLibrary.cpp
new file mode 100644
--- /dev/null
+++ b/Engine/Plugins/Experimental/Animation/PoseSearch/Source/Runtime/Private/MotionMatchingAnimNodeLibrary.cpp
@@ -0,0 +1,60 @@
+// Copyright Epic Games, Inc. All Rights Reserved.
+
+#include "PoseSearch/MotionMatchingAnimNodeLibrary.h"
+
+#include "PoseSearch/AnimNode_MotionMatching.h"
+#include "PoseSearch/PoseSearchDefines.h"
+
+FMotionMatchingAnimNodeReference UMotionMatchingAnimNodeLibrary::ConvertToMotionMatchingNode(const FAnimNodeReference& Node, EAnimNodeReferenceConversionResult& Result)
+{
+	return FAnimNodeReference::ConvertToType<FMotionMatchingAnimNodeReference>(Node, Result);
+}
+
+void UMotionMatchingAnimNodeLibrary::SetDatabaseToSearch(const FMotionMatchingAnimNodeReference& MotionMatchingNode, UPoseSearchDatabase* Database, bool bForceInterruptIfNew)
+{
+	FAnimNode_MotionMatching* MotionMatchingNodePtr = MotionMatchingNode.GetAnimNodePtr<FAnimNode_MotionMatching>();
+	if (MotionMatchingNodePtr == nullptr)
+	{
+		UE_LOG(LogPoseSearch, Warning, TEXT("UMotionMatchingAnimNodeLibrary::SetDatabase called on an invalid context or with an invalid type"));
+		return;
+	}
+
+	MotionMatchingNodePtr->SetDatabaseToSearch(Database, bForceInterruptIfNew);
+}
+
+void UMotionMatchingAnimNodeLibrary::SetDatabasesToSearch(const FMotionMatchingAnimNodeReference& MotionMatchingNode, const TArray<UPoseSearchDatabase*>& Databases, bool bForceInterruptIfNew)
+{
+	FAnimNode_MotionMatching* MotionMatchingNodePtr = MotionMatchingNode.GetAnimNodePtr<FAnimNode_MotionMatching>();
+	if (MotionMatchingNodePtr == nullptr)
+	{
+		UE_LOG(LogPoseSearch, Warning, TEXT("UMotionMatchingAnimNodeLibrary::SetDatabases called on an invalid context or with an invalid type"));
+		return;
+	}
+
+	MotionMatchingNodePtr->SetDatabasesToSearch(Databases, bForceInterruptIfNew);
+}
+
+void UMotionMatchingAnimNodeLibrary::ResetDatabasesToSearch(const FMotionMatchingAnimNodeReference& MotionMatchingNode, bool bForceInterrupt)
+{
+	FAnimNode_MotionMatching* MotionMatchingNodePtr = MotionMatchingNode.GetAnimNodePtr<FAnimNode_MotionMatching>();
+	if (MotionMatchingNodePtr == nullptr)
+	{
+		UE_LOG(LogPoseSearch, Warning, TEXT("UMotionMatchingAnimNodeLibrary::ResetDatabasesToSearch called on an invalid context or with an invalid type"));
+		return;
+	}
+
+	MotionMatchingNodePtr->ResetDatabasesToSearch(bForceInterrupt);
+
+}
+
+void UMotionMatchingAnimNodeLibrary::ForceInterruptNextUpdate(const FMotionMatchingAnimNodeReference& MotionMatchingNode)
+{
+	FAnimNode_MotionMatching* MotionMatchingNodePtr = MotionMatchingNode.GetAnimNodePtr<FAnimNode_MotionMatching>();
+	if (MotionMatchingNodePtr == nullptr)
+	{
+		UE_LOG(LogPoseSearch, Warning, TEXT("UMotionMatchingAnimNodeLibrary::ForceInterruptOneFrame called on an invalid context or with an invalid type"));
+		return;
+	}
+
+	MotionMatchingNodePtr->ForceInterruptNextUpdate();
+}
\ No newline at end of file
diff --git a/Engine/Plugins/Experimental/Animation/PoseSearch/Source/Runtime/Private/PoseSearchLibrary.cpp b/Engine/Plugins/Experimental/Animation/PoseSearch/Source/Runtime/Private/PoseSearchLibrary.cpp
--- a/Engine/Plugins/Experimental/Animation/PoseSearch/Source/Runtime/Private/PoseSearchLibrary.cpp
+++ b/Engine/Plugins/Experimental/Animation/PoseSearch/Source/Runtime/Private/PoseSearchLibrary.cpp
@@ -252,7 +252,7 @@ void UPoseSearchLibrary::TraceMotionMatchingState(
 
 void UPoseSearchLibrary::UpdateMotionMatchingState(
 	const FAnimationUpdateContext& Context,
-	const UPoseSearchSearchableAsset* Searchable,
+	const TArray<TObjectPtr<const UPoseSearchDatabase>>& Databases,
 	const FTrajectorySampleRange& Trajectory,
 	const FMotionMatchingSettings& Settings,
 	FMotionMatchingState& InOutMotionMatchingState,
@@ -262,11 +262,11 @@ void UPoseSearchLibrary::UpdateMotionMatchingState(
 
 	using namespace UE::PoseSearch;
 
-	if (!Searchable)
+	if (Databases.IsEmpty())
 	{
 		Context.LogMessage(
 			EMessageSeverity::Error,
-			LOCTEXT("NoSearchable", "No searchable asset provided for motion matching."));
+			LOCTEXT("NoDatabases", "No database assets provided for motion matching."));
 		return;
 	}
 
@@ -347,11 +347,19 @@ void UPoseSearchLibrary::UpdateMotionMatchingState(
 			}
 		}
 
-		// Search the database for the nearest match to the updated query vector
-		FSearchResult SearchResult = Searchable->Search(SearchContext);
-		if (SearchResult.PoseCost.IsValid())
+		FSearchResult SearchResult;
+		for (TObjectPtr<const UPoseSearchDatabase> Database : Databases)
 		{
-			SearchContext.UpdateCurrentBestCost(SearchResult.PoseCost);
+			if (ensure(Database))
+			{
+				FSearchResult NewSearchResult = Database->Search(SearchContext);
+				if (NewSearchResult.PoseCost.GetTotalCost() < SearchResult.PoseCost.GetTotalCost())
+				{
+					SearchResult = NewSearchResult;
+					SearchContext.UpdateCurrentBestCost(SearchResult.PoseCost);
+				}
+			}
+
 		}
 
 		if (SearchResult.PoseCost.GetTotalCost() < ContinuingPoseCost.GetTotalCost())
@@ -383,7 +391,8 @@ void UPoseSearchLibrary::UpdateMotionMatchingState(
 	// Record debugger details
 	if (IsTracing(Context))
 	{
-		TraceMotionMatchingState(Searchable, SearchContext, InOutMotionMatchingState.CurrentSearchResult, LastResult, InOutMotionMatchingState.ElapsedPoseSearchTime,
+		// @todo: Update this to account for multiple databases.
+		TraceMotionMatchingState(Databases[0], SearchContext, InOutMotionMatchingState.CurrentSearchResult, LastResult, InOutMotionMatchingState.ElapsedPoseSearchTime,
 			InOutMotionMatchingState.RootMotionTransformDelta, Context.AnimInstanceProxy->GetAnimInstanceObject(), Context.GetCurrentNodeId(), DeltaTime, bSearch);
 	}
 #endif
@@ -391,7 +400,7 @@ void UPoseSearchLibrary::UpdateMotionMatchingState(
 
 void UPoseSearchLibrary::MotionMatch(
 	UAnimInstance* AnimInstance,
-	const UPoseSearchSearchableAsset* Searchable,
+	const UPoseSearchDatabase* Database,
 	const FTrajectorySampleRange Trajectory,
 	const FName PoseHistoryName,
 	UAnimationAsset*& SelectedAnimation,
@@ -429,7 +438,7 @@ void UPoseSearchLibrary::MotionMatch(
 	BlendParameters = FVector::ZeroVector;
 	SearchCost = MAX_flt;
 
-	if (Searchable)
+	if (Database)
 	{
 		// ExtendedPoseHistory will hold future poses to match AssetSamplerBase (at FutureAnimationStartTime) TimeToFutureAnimationStart seconds in the future
 		FExtendedPoseHistory ExtendedPoseHistory;
@@ -518,7 +527,7 @@ void UPoseSearchLibrary::MotionMatch(
 		// @todo: finish set up SearchContext by exposing or calculating additional members
 		FSearchContext SearchContext(&Trajectory, ExtendedPoseHistory.IsInitialized() ? &ExtendedPoseHistory : nullptr, TimeToFutureAnimationStart);
 
-		FSearchResult SearchResult = Searchable->Search(SearchContext);
+		FSearchResult SearchResult = Database->Search(SearchContext);
 		if (SearchResult.IsValid())
 		{
 			const FPoseSearchIndexAsset* SearchIndexAsset = SearchResult.GetSearchIndexAsset();
@@ -553,7 +562,7 @@ void UPoseSearchLibrary::MotionMatch(
 #endif // ENABLE_DRAW_DEBUG && ENABLE_ANIM_DEBUG
 
 #if UE_POSE_SEARCH_TRACE_ENABLED
-			TraceMotionMatchingState(Searchable, SearchContext, SearchResult, FSearchResult(), 0.f, FTransform::Identity, AnimInstance, DebugSessionUniqueIdentifier, AnimInstance->GetDeltaSeconds(), true);
+			TraceMotionMatchingState(Database, SearchContext, SearchResult, FSearchResult(), 0.f, FTransform::Identity, AnimInstance, DebugSessionUniqueIdentifier, AnimInstance->GetDeltaSeconds(), true);
 #endif // UE_POSE_SEARCH_TRACE_ENABLED
 		}
 	}
diff --git a/Engine/Plugins/Experimental/Animation/PoseSearch/Source/Runtime/Public/PoseSearch/AnimNode_MotionMatching.h b/Engine/Plugins/Experimental/Animation/PoseSearch/Source/Runtime/Public/PoseSearch/AnimNode_MotionMatching.h
--- a/Engine/Plugins/Experimental/Animation/PoseSearch/Source/Runtime/Public/PoseSearch/AnimNode_MotionMatching.h
+++ b/Engine/Plugins/Experimental/Animation/PoseSearch/Source/Runtime/Public/PoseSearch/AnimNode_MotionMatching.h
@@ -9,6 +9,7 @@
 #include "PoseSearch/PoseSearchLibrary.h"
 #include "AnimNode_MotionMatching.generated.h"
 
+class UPoseSearchDatabase;
 class UPoseSearchSearchableAsset;
 
 USTRUCT(BlueprintInternalUseOnly)
@@ -16,28 +17,33 @@ struct POSESEARCH_API FAnimNode_MotionMatching : public FAnimNode_AssetPlayerBas
 {
 	GENERATED_BODY()
 
-public:
+private:
 	UPROPERTY()
 	FPoseLink Source;
 
+	// @todo: Delete this after updating content.
 	// Collection of animations for motion matching.
-	UPROPERTY(EditAnywhere, BlueprintReadWrite, Category=Settings, meta=(PinShownByDefault))
+	UPROPERTY(EditAnywhere, Category=Settings, meta=(PinHiddenByDefault))
 	TObjectPtr<const UPoseSearchSearchableAsset> Searchable = nullptr;
 
+	// The database to search. This can be overridden by Anim Node Functions such as "On Become Relevant" and "On Update" via SetDatabaseToSearch/SetDatabasesToSearch.
+	UPROPERTY(EditAnywhere, Category = Settings, meta = (PinShownByDefault))
+	TObjectPtr<const UPoseSearchDatabase> Database = nullptr;
+
 	// Motion Trajectory samples for pose search queries in Motion Matching.These are expected to be in the space of the SkeletalMeshComponent.This is provided with the CharacterMovementTrajectory Component output.
-	UPROPERTY(EditAnywhere, BlueprintReadWrite, Category=Settings, meta=(PinShownByDefault))
+	UPROPERTY(EditAnywhere, Category=Settings, meta=(PinShownByDefault))
 	FTrajectorySampleRange Trajectory;
 
 	// Settings for the core motion matching node.
-	UPROPERTY(EditAnywhere, BlueprintReadWrite, Category=Settings, meta=(PinHiddenByDefault))
+	UPROPERTY(EditAnywhere, Category=Settings, meta=(PinHiddenByDefault))
 	FMotionMatchingSettings Settings;
 
 	// Reset the motion matching selection state if it has become relevant to the graph after not being updated on previous frames.
-	UPROPERTY(EditAnywhere, BlueprintReadWrite, Category = Settings, meta = (PinHiddenByDefault))
+	UPROPERTY(EditAnywhere, Category = Settings, meta = (PinHiddenByDefault))
 	bool bResetOnBecomingRelevant = true;
 
 	// If set to true, the continuing pose will be invalidated. This is useful if you want to force a re-selection of the animation segment instead of continuing with the previous segment, even if it has a better score.
-	UPROPERTY(EditAnywhere, BlueprintReadWrite, Category = Settings, meta = (PinHiddenByDefault))
+	UPROPERTY(EditAnywhere, Category = Settings, meta = (PinHiddenByDefault))
 	bool bForceInterrupt = false;
 
 #if WITH_EDITORONLY_DATA
@@ -51,6 +57,8 @@ public:
 	bool bDebugDrawMatch = true;
 #endif
 
+public:
+
 	// FAnimNode_Base interface
 	// @todo: implement CacheBones_AnyThread to rebind the schema bones
 	virtual void Initialize_AnyThread(const FAnimationInitializeContext& Context) override;
@@ -58,6 +66,18 @@ public:
 	virtual void GatherDebugData(FNodeDebugData& DebugData) override;
 	// End of FAnimNode_Base interface
 
+	// Search InDatabase instead of the Database property on this node. Use bForceInterruptIfNew to ignore the continuing pose if InDatabase is new.
+	void SetDatabaseToSearch(UPoseSearchDatabase* InDatabase, bool bForceInterruptIfNew);
+
+	// Search InDatabases instead of the Database property on the node. Use bForceInterruptIfNew to ignore the continuing pose if InDatabases is new.
+	void SetDatabasesToSearch(const TArray<UPoseSearchDatabase*>& InDatabases, bool bForceInterruptIfNew);
+
+	// Reset the effects of SetDatabaseToSearch/SetDatabasesToSearch and use the Database property on this node.
+	void ResetDatabasesToSearch(bool bInForceInterrupt);
+
+	// Ignore the continuing pose on the next update and force a search.
+	void ForceInterruptNextUpdate();
+
 private:
 
 	FAnimNode_BlendStack_Standalone BlendStackNode;
@@ -68,6 +88,16 @@ private:
 	// Update Counter for detecting being relevant
 	FGraphTraversalCounter UpdateCounter;
 
+	// List of databases this node is searching.
+	UPROPERTY()
+	TArray<TObjectPtr<const UPoseSearchDatabase>> DatabasesToSearch;
+
+	// Ignore the continuing pose on the next update and use the best result from DatabasesToSearch. This is set back to false after each update.
+	bool bForceInterruptNextUpdate = false;
+
+	// True if the Database property on this node has been overridden by SetDatabaseToSearch/SetDatabasesToSearch.
+	bool bOverrideDatabaseInput = false;
+
 	// FAnimNode_AssetPlayerBase
 protected:
 	// FAnimNode_AssetPlayerBase interface
diff --git a/Engine/Plugins/Experimental/Animation/PoseSearch/Source/Runtime/Public/PoseSearch/MotionMatchingAnimNodeLibrary.h b/Engine/Plugins/Experimental/Animation/PoseSearch/Source/Runtime/Public/PoseSearch/MotionMatchingAnimNodeLibrary.h
new file mode 100644
--- /dev/null
+++ b/Engine/Plugins/Experimental/Animation/PoseSearch/Source/Runtime/Public/PoseSearch/MotionMatchingAnimNodeLibrary.h
@@ -0,0 +1,79 @@
+// Copyright Epic Games, Inc. All Rights Reserved.
+
+#pragma once
+
+#include "Animation/AnimExecutionContext.h"
+#include "Animation/AnimNodeReference.h"
+#include "Kismet/BlueprintFunctionLibrary.h"
+#include "MotionMatchingAnimNodeLibrary.generated.h"
+
+struct FAnimNode_MotionMatching;
+class UPoseSearchDatabase;
+
+USTRUCT(BlueprintType)
+struct FMotionMatchingAnimNodeReference : public FAnimNodeReference
+{
+	GENERATED_BODY()
+
+	typedef FAnimNode_MotionMatching FInternalNodeType;
+};
+
+// Exposes operations that can be run on a Motion Matching node via Anim Node Functions such as "On Become Relevant" and "On Update".
+UCLASS(Experimental)
+class POSESEARCH_API UMotionMatchingAnimNodeLibrary : public UBlueprintFunctionLibrary
+{
+	GENERATED_BODY()
+
+public:
+	/** Get a motion matching node context from an anim node context */
+	UFUNCTION(BlueprintCallable, Category = "Animation|MotionMatching", meta = (BlueprintThreadSafe, ExpandEnumAsExecs = "Result"))
+	static FMotionMatchingAnimNodeReference ConvertToMotionMatchingNode(const FAnimNodeReference& Node, EAnimNodeReferenceConversionResult& Result);
+
+	/** Get a motion matching node context from an anim node context (pure) */
+	UFUNCTION(BlueprintPure, Category = "Animation|MotionMatching", meta = (BlueprintThreadSafe, DisplayName = "Convert to Motion Matching Node"))
+	static void ConvertToMotionMatchingNodePure(const FAnimNodeReference& Node, FMotionMatchingAnimNodeReference& MotionMatchingNode, bool& Result)
+	{
+		EAnimNodeReferenceConversionResult ConversionResult;
+		MotionMatchingNode = ConvertToMotionMatchingNode(Node, ConversionResult);
+		Result = (ConversionResult == EAnimNodeReferenceConversionResult::Succeeded);
+	}
+
+	/**
+	 * Set the database to search on the motion matching node. This overrides the Database property on the motion matching node.
+	 * @param MotionMatchingNode - The motion matching node to operate on.
+	 * @param Database - The database for the motion matching node to search.
+	 * @param bForceInterruptIfNew - If true, ignore the continuing pose (the current clip that's playing) and force a search of the new database. 
+		If false, the continuing pose will continue to play until a better match is found in the new database. This setting is ignored if the 
+		motion matching node is already searching this database.
+	 */
+	UFUNCTION(BlueprintCallable, Category = "Animation|MotionMatching", meta = (BlueprintThreadSafe))
+	static void SetDatabaseToSearch(const FMotionMatchingAnimNodeReference& MotionMatchingNode, UPoseSearchDatabase* Database, bool bForceInterruptIfNew);
+
+	/**
+	 * Set the database to search on the motion matching node. This overrides the Database property on the motion matching node.
+	 * @param MotionMatchingNode - The motion matching node to operate on.
+	 * @param Databases - Array of databases for the motion matching node to search.
+	 * @param bForceInterruptIfNew - If true, ignore the continuing pose (the current clip that's playing) and force a search of the new databases. 
+		If false, the continuing pose will continue to play until a better match is found in one of the new databases. This setting is ignored if the 
+		motion matching node is already searching this array of databases (note: a subset of databases or the same databases in a different order is
+		considered to be a new array of databases).
+	 */
+	UFUNCTION(BlueprintCallable, Category = "Animation|MotionMatching", meta = (BlueprintThreadSafe))
+	static void SetDatabasesToSearch(const FMotionMatchingAnimNodeReference& MotionMatchingNode, const TArray<UPoseSearchDatabase*>& Databases, bool bForceInterruptIfNew);
+
+	/**
+	 * Clear the effects of SetDatabaseToSearch/SetDatabasesToSearch and resume searching the Database property on the motion matching node.
+	 * @param MotionMatchingNode - The motion matching node to operate on.
+	 * @param bForceInterrupt - Force a search after the reset. If false, the continuing pose (the current clip that's playing) will continue
+		to play until a better match is found from the Database property.
+	 */
+	UFUNCTION(BlueprintCallable, Category = "Animation|MotionMatching", meta = (BlueprintThreadSafe))
+	static void ResetDatabasesToSearch(const FMotionMatchingAnimNodeReference& MotionMatchingNode, bool bForceInterrupt);
+
+	/**
+	 * Ignore the continuing pose (the current clip that's playing) and force a new search.
+	 * @param MotionMatchingNode - The motion matching node to operate on.
+	 */
+	UFUNCTION(BlueprintCallable, Category = "Animation|MotionMatching", meta = (BlueprintThreadSafe))
+	static void ForceInterruptNextUpdate(const FMotionMatchingAnimNodeReference& MotionMatchingNode);
+};
\ No newline at end of file
diff --git a/Engine/Plugins/Experimental/Animation/PoseSearch/Source/Runtime/Public/PoseSearch/PoseSearchLibrary.h b/Engine/Plugins/Experimental/Animation/PoseSearch/Source/Runtime/Public/PoseSearch/PoseSearchLibrary.h
--- a/Engine/Plugins/Experimental/Animation/PoseSearch/Source/Runtime/Public/PoseSearch/PoseSearchLibrary.h
+++ b/Engine/Plugins/Experimental/Animation/PoseSearch/Source/Runtime/Public/PoseSearch/PoseSearchLibrary.h
@@ -132,7 +132,7 @@ public:
 	* Implementation of the core motion matching algorithm
 	*
 	* @param Context						Input animation update context providing access to the proxy and delta time
-	* @param Searchable						Input collection of animations for motion matching
+	* @param Databases						Input array of databases to search
 	* @param Trajectory						Input motion trajectory samples for pose search queries
 	* @param Settings						Input motion matching algorithm configuration settings
 	* @param InOutMotionMatchingState		Input/Output encapsulated motion matching algorithm and state
@@ -140,7 +140,7 @@ public:
 	*/
 	static void UpdateMotionMatchingState(
 		const FAnimationUpdateContext& Context,
-		const UPoseSearchSearchableAsset* Searchable,
+		const TArray<TObjectPtr<const UPoseSearchDatabase>>& Databases,
 		const FTrajectorySampleRange& Trajectory,
 		const FMotionMatchingSettings& Settings,
 		FMotionMatchingState& InOutMotionMatchingState,
@@ -150,7 +150,7 @@ public:
 	* Implementation of the core motion matching algorithm
 	*
 	* @param AnimInstance					Input animation instance
-	* @param Searchable						Input searchable asset
+	* @param Database						Input database to search
 	* @param Trajectory						Input motion trajectory samples for pose search queries
 	* @param PoseHistoryName				Input tag of the associated PoseSearchHistoryCollector node in the anim graph
 	* @param SelectedAnimation				Output selected animation from the searchable asset
@@ -167,7 +167,7 @@ public:
 	UFUNCTION(BlueprintPure, Category = "Animation|Pose Search", meta = (BlueprintThreadSafe, Keywords = "PoseMatch"))
 	static void MotionMatch(
 		UAnimInstance* AnimInstance,
-		const UPoseSearchSearchableAsset* Searchable,
+		const UPoseSearchDatabase* Database,
 		const FTrajectorySampleRange Trajectory,
 		const FName PoseHistoryName,
 		UAnimationAsset*& SelectedAnimation,
