Add simple mechanism to change the preview mesh in the Pose Search Database editor. We may want to customize the Preview Scene Settings to be consistent with other animation editors, but this change is simple and unblocks most workflows.

diff --git a/Engine/Plugins/Experimental/Animation/PoseSearch/Source/Editor/Private/PoseSearchDatabaseViewModel.cpp b/Engine/Plugins/Experimental/Animation/PoseSearch/Source/Editor/Private/PoseSearchDatabaseViewModel.cpp
--- a/Engine/Plugins/Experimental/Animation/PoseSearch/Source/Editor/Private/PoseSearchDatabaseViewModel.cpp
+++ b/Engine/Plugins/Experimental/Animation/PoseSearch/Source/Editor/Private/PoseSearchDatabaseViewModel.cpp
@@ -207,7 +207,8 @@ namespace UE::PoseSearch
 		Mesh->PreviewInstance = AnimInstance;
 		AnimInstance->InitializeAnimation();
 
-		Mesh->SetSkeletalMesh(PoseSearchDatabase->Schema->Skeleton->GetPreviewMesh(true));
+		USkeletalMesh* DatabasePreviewMesh = PoseSearchDatabase->PreviewMesh;
+		Mesh->SetSkeletalMesh(DatabasePreviewMesh ? DatabasePreviewMesh : PoseSearchDatabase->Schema->Skeleton->GetPreviewMesh(true));
 		Mesh->EnablePreview(true, PreviewAsset);
 		
 		AnimInstance->SetAnimationAsset(PreviewAsset, false, 0.0f);
diff --git a/Engine/Plugins/Experimental/Animation/PoseSearch/Source/Runtime/Public/PoseSearch/PoseSearchDatabase.h b/Engine/Plugins/Experimental/Animation/PoseSearch/Source/Runtime/Public/PoseSearch/PoseSearchDatabase.h
--- a/Engine/Plugins/Experimental/Animation/PoseSearch/Source/Runtime/Public/PoseSearch/PoseSearchDatabase.h
+++ b/Engine/Plugins/Experimental/Animation/PoseSearch/Source/Runtime/Public/PoseSearch/PoseSearchDatabase.h
@@ -248,6 +248,11 @@ public:
 	TArray<FPoseSearchDatabaseSequence> Sequences_DEPRECATED;
 	UPROPERTY()
 	TArray<FPoseSearchDatabaseBlendSpace> BlendSpaces_DEPRECATED;
+
+	// If null, the default preview mesh for the skeleton will be used. Otherwise, this will be used in preview scenes.
+	// @todo: Move this to be a setting in the Pose Search Database editor. 
+	UPROPERTY(EditAnywhere, Category = "Preview")
+	TObjectPtr<USkeletalMesh> PreviewMesh = nullptr;
 #endif // WITH_EDITORONLY_DATA
 
 	UPROPERTY(EditAnywhere, Category="Database")
