motion matching - relying on  FPoseSearchDatabaseAnimationAssetBase rather than derived classes

diff --git a/Engine/Plugins/Experimental/Animation/PoseSearch/Source/Editor/Private/PoseSearchDebuggerDatabaseView.cpp b/Engine/Plugins/Experimental/Animation/PoseSearch/Source/Editor/Private/PoseSearchDebuggerDatabaseView.cpp
--- a/Engine/Plugins/Experimental/Animation/PoseSearch/Source/Editor/Private/PoseSearchDebuggerDatabaseView.cpp
+++ b/Engine/Plugins/Experimental/Animation/PoseSearch/Source/Editor/Private/PoseSearchDebuggerDatabaseView.cpp
@@ -206,32 +206,17 @@ void SDebuggerDatabaseView::Update(const FTraceMotionMatchingStateMessage& State
 						Row->AssetName = DatabaseAsset->GetName();
 						Row->AssetPath = DatabaseAsset->GetAnimationAsset() ? DatabaseAsset->GetAnimationAsset()->GetPathName() : "";
 						Row->bLooping = DatabaseAsset->IsLooping();
-						Row->BlendParameters = FVector::Zero();
+						Row->BlendParameters = SearchIndexAsset->BlendParameters;
 						Row->AnimFrame = 0;
 						Row->AnimPercentage = 0.0f;
 
-						if (const FPoseSearchDatabaseSequence* DatabaseSequence = DatabaseAssetStruct.GetPtr<FPoseSearchDatabaseSequence>())
+						if (const FPoseSearchDatabaseAnimationAssetBase* DatabaseAnimationAssetBase = DatabaseAssetStruct.GetPtr<FPoseSearchDatabaseAnimationAssetBase>())
 						{
-							Row->AnimFrame = DatabaseSequence->Sequence->GetFrameAtTime(Time);
-							Row->AnimPercentage = Time / DatabaseSequence->Sequence->GetPlayLength();
-						}
-						else if (const FPoseSearchDatabaseAnimComposite* DatabaseAnimComposite = DatabaseAssetStruct.GetPtr<FPoseSearchDatabaseAnimComposite>())
-						{
-							Row->AnimFrame = DatabaseAnimComposite->AnimComposite->GetFrameAtTime(Time);
-							Row->AnimPercentage = Time / DatabaseAnimComposite->AnimComposite->GetPlayLength();
-						}
-						else if (const FPoseSearchDatabaseBlendSpace* DatabaseBlendSpace = DatabaseAssetStruct.GetPtr<FPoseSearchDatabaseBlendSpace>())
-						{
-							Row->BlendParameters = SearchIndexAsset->BlendParameters;
-						}
-						else if (const FPoseSearchDatabaseAnimMontage* DatabaseAnimMontage = DatabaseAssetStruct.GetPtr<FPoseSearchDatabaseAnimMontage>())
-						{
-							Row->AnimFrame = DatabaseAnimMontage->AnimMontage->GetFrameAtTime(Time);
-							Row->AnimPercentage = Time / DatabaseAnimMontage->AnimMontage->GetPlayLength();
-						}
-						else
-						{
-							checkNoEntry();
+							if (const UAnimSequenceBase* SequenceBase = Cast<UAnimSequenceBase>(DatabaseAnimationAssetBase->GetAnimationAsset()))
+							{
+								Row->AnimFrame = SequenceBase->GetFrameAtTime(Time);
+								Row->AnimPercentage = Time / SequenceBase->GetPlayLength();
+							}
 						}
 					}
 				}
diff --git a/Engine/Plugins/Experimental/Animation/PoseSearch/Source/Runtime/Private/PoseSearchDerivedData.cpp b/Engine/Plugins/Experimental/Animation/PoseSearch/Source/Runtime/Private/PoseSearchDerivedData.cpp
--- a/Engine/Plugins/Experimental/Animation/PoseSearch/Source/Runtime/Private/PoseSearchDerivedData.cpp
+++ b/Engine/Plugins/Experimental/Animation/PoseSearch/Source/Runtime/Private/PoseSearchDerivedData.cpp
@@ -290,47 +290,7 @@ static void InitSearchIndexAssets(FPoseSearchIndexBase& SearchIndex, const UPose
 			const bool bAddUnmirrored = DatabaseAsset->GetMirrorOption() == EPoseSearchMirrorOption::UnmirroredOnly || DatabaseAsset->GetMirrorOption() == EPoseSearchMirrorOption::UnmirroredAndMirrored;
 			const bool bAddMirrored = DatabaseAsset->GetMirrorOption() == EPoseSearchMirrorOption::MirroredOnly || DatabaseAsset->GetMirrorOption() == EPoseSearchMirrorOption::UnmirroredAndMirrored;
 
-			if (const FPoseSearchDatabaseSequence* DatabaseSequence = DatabaseAssetStruct.GetPtr<FPoseSearchDatabaseSequence>())
-			{
-				ValidRanges.Reset();
-				FindValidSequenceIntervals(DatabaseSequence->Sequence, DatabaseSequence->SamplingRange, DatabaseSequence->IsLooping(), Database->ExcludeFromDatabaseParameters, ValidRanges);
-				for (const FFloatRange& Range : ValidRanges)
-				{
-					for (int32 PermutationIdx = 0; PermutationIdx < Database->Schema->NumberOfPermutations; ++PermutationIdx)
-					{
-						if (bAddUnmirrored)
-						{
-							SearchIndex.Assets.Emplace(AnimationAssetIndex, false, FFloatInterval(Range.GetLowerBoundValue(), Range.GetUpperBoundValue()), PermutationIdx);
-						}
-
-						if (bAddMirrored)
-						{
-							SearchIndex.Assets.Emplace(AnimationAssetIndex, true, FFloatInterval(Range.GetLowerBoundValue(), Range.GetUpperBoundValue()), PermutationIdx);
-						}
-					}
-				}
-			}
-			else if (const FPoseSearchDatabaseAnimComposite* DatabaseAnimComposite = DatabaseAssetStruct.GetPtr<FPoseSearchDatabaseAnimComposite>())
-			{
-				ValidRanges.Reset();
-				FindValidSequenceIntervals(DatabaseAnimComposite->AnimComposite, DatabaseAnimComposite->SamplingRange, DatabaseAnimComposite->IsLooping(), Database->ExcludeFromDatabaseParameters, ValidRanges);
-				for (const FFloatRange& Range : ValidRanges)
-				{
-					for (int32 PermutationIdx = 0; PermutationIdx < Database->Schema->NumberOfPermutations; ++PermutationIdx)
-					{
-						if (bAddUnmirrored)
-						{
-							SearchIndex.Assets.Emplace(AnimationAssetIndex, false, FFloatInterval(Range.GetLowerBoundValue(), Range.GetUpperBoundValue()), PermutationIdx);
-						}
-
-						if (bAddMirrored)
-						{
-							SearchIndex.Assets.Emplace(AnimationAssetIndex, true, FFloatInterval(Range.GetLowerBoundValue(), Range.GetUpperBoundValue()), PermutationIdx);
-						}
-					}
-				}
-			}
-			else if (const FPoseSearchDatabaseBlendSpace* DatabaseBlendSpace = DatabaseAssetStruct.GetPtr<FPoseSearchDatabaseBlendSpace>())
+			if (const FPoseSearchDatabaseBlendSpace* DatabaseBlendSpace = DatabaseAssetStruct.GetPtr<FPoseSearchDatabaseBlendSpace>())
 			{
 				int32 HorizontalBlendNum, VerticalBlendNum;
 				DatabaseBlendSpace->GetBlendSpaceParameterSampleRanges(HorizontalBlendNum, VerticalBlendNum);
@@ -363,10 +323,12 @@ static void InitSearchIndexAssets(FPoseSearchIndexBase& SearchIndex, const UPose
 					}
 				}
 			}
-			else if (const FPoseSearchDatabaseAnimMontage* DatabaseAnimMontage = DatabaseAssetStruct.GetPtr<FPoseSearchDatabaseAnimMontage>())
+			// support for FPoseSearchDatabaseSequence, FPoseSearchDatabaseAnimComposite, FPoseSearchDatabaseAnimMontage
+			else if (const UAnimSequenceBase* SequenceBase = Cast<UAnimSequenceBase>(DatabaseAsset->GetAnimationAsset()))
 			{
 				ValidRanges.Reset();
-				FindValidSequenceIntervals(DatabaseAnimMontage->AnimMontage, DatabaseAnimMontage->SamplingRange, DatabaseAnimMontage->IsLooping(), Database->ExcludeFromDatabaseParameters, ValidRanges);
+
+				FindValidSequenceIntervals(SequenceBase, DatabaseAsset->GetSamplingRange(), DatabaseAsset->IsLooping(), Database->ExcludeFromDatabaseParameters, ValidRanges);
 				for (const FFloatRange& Range : ValidRanges)
 				{
 					for (int32 PermutationIdx = 0; PermutationIdx < Database->Schema->NumberOfPermutations; ++PermutationIdx)
diff --git a/Engine/Plugins/Experimental/Animation/PoseSearch/Source/Runtime/Public/PoseSearch/PoseSearchDatabase.h b/Engine/Plugins/Experimental/Animation/PoseSearch/Source/Runtime/Public/PoseSearch/PoseSearchDatabase.h
--- a/Engine/Plugins/Experimental/Animation/PoseSearch/Source/Runtime/Public/PoseSearch/PoseSearchDatabase.h
+++ b/Engine/Plugins/Experimental/Animation/PoseSearch/Source/Runtime/Public/PoseSearch/PoseSearchDatabase.h
@@ -75,6 +75,8 @@ struct POSESEARCH_API FPoseSearchDatabaseAnimationAssetBase
 	virtual void SetIsEnabled(bool bInIsEnabled) {}
 	virtual bool IsRootMotionEnabled() const { return false; }
 	virtual EPoseSearchMirrorOption GetMirrorOption() const { return EPoseSearchMirrorOption::Invalid; }
+	// [0, 0] represents the entire frame range of the original animation.
+	virtual FFloatInterval GetSamplingRange() const { return FFloatInterval(0.f, 0.f); }
 };
 
 /** A sequence entry in a UPoseSearchDatabase. */
@@ -110,6 +112,7 @@ struct POSESEARCH_API FPoseSearchDatabaseSequence : public FPoseSearchDatabaseAn
 	void SetIsEnabled(bool bInIsEnabled) override { bEnabled = bInIsEnabled; }
 	bool IsRootMotionEnabled() const override;
 	EPoseSearchMirrorOption GetMirrorOption() const override { return MirrorOption; }
+	FFloatInterval GetSamplingRange() const override { return SamplingRange; }
 };
 
 /** An blend space entry in a UPoseSearchDatabase. */
@@ -189,6 +192,7 @@ struct POSESEARCH_API FPoseSearchDatabaseAnimComposite : public FPoseSearchDatab
 	void SetIsEnabled(bool bInIsEnabled) override { bEnabled = bInIsEnabled; }
 	bool IsRootMotionEnabled() const override;
 	EPoseSearchMirrorOption GetMirrorOption() const override { return MirrorOption; }
+	FFloatInterval GetSamplingRange() const override { return SamplingRange; }
 };
 
 /** An anim montage entry in a UPoseSearchDatabase. */
@@ -224,6 +228,7 @@ struct POSESEARCH_API FPoseSearchDatabaseAnimMontage : public FPoseSearchDatabas
 	void SetIsEnabled(bool bInIsEnabled) override { bEnabled = bInIsEnabled; }
 	bool IsRootMotionEnabled() const override;
 	EPoseSearchMirrorOption GetMirrorOption() const override { return MirrorOption; }
+	FFloatInterval GetSamplingRange() const override { return SamplingRange; }
 };
 
 /** A data asset for indexing a collection of animation sequences. */
