motion matching - removing UPoseSearchSearchableAsset

diff --git a/Engine/Plugins/Experimental/Animation/PoseSearch/Source/Editor/Private/PoseSearchDebuggerView.cpp b/Engine/Plugins/Experimental/Animation/PoseSearch/Source/Editor/Private/PoseSearchDebuggerView.cpp
--- a/Engine/Plugins/Experimental/Animation/PoseSearch/Source/Editor/Private/PoseSearchDebuggerView.cpp
+++ b/Engine/Plugins/Experimental/Animation/PoseSearch/Source/Editor/Private/PoseSearchDebuggerView.cpp
@@ -306,7 +306,7 @@ bool SDebuggerView::UpdateNodeSelection()
 				.Padding(10.0f)
 				[
 					SNew(SButton)
-					.Text(FText::FromString(Model->GetSearchableAsset()->GetName()))
+					.Text(FText::FromString(Model->GetCurrentDatabase()->GetName()))
 					.HAlign(HAlign_Center)
 					.VAlign(VAlign_Center)
 					.ContentPadding(10.0f)
diff --git a/Engine/Plugins/Experimental/Animation/PoseSearch/Source/Editor/Private/PoseSearchDebuggerViewModel.cpp b/Engine/Plugins/Experimental/Animation/PoseSearch/Source/Editor/Private/PoseSearchDebuggerViewModel.cpp
--- a/Engine/Plugins/Experimental/Animation/PoseSearch/Source/Editor/Private/PoseSearchDebuggerViewModel.cpp
+++ b/Engine/Plugins/Experimental/Animation/PoseSearch/Source/Editor/Private/PoseSearchDebuggerViewModel.cpp
@@ -46,16 +46,6 @@ const UPoseSearchDatabase* FDebuggerViewModel::GetCurrentDatabase() const
 	return ActiveMotionMatchingState ? ActiveMotionMatchingState->GetCurrentDatabase() : nullptr;
 }
 
-const UPoseSearchSearchableAsset* FDebuggerViewModel::GetSearchableAsset() const
-{
-	if (ActiveMotionMatchingState)
-	{
-		return FTraceMotionMatchingState::GetObjectFromId<UPoseSearchSearchableAsset>(ActiveMotionMatchingState->SearchableAssetId);
-	}
-
-	return nullptr;
-}
-
 void FDebuggerViewModel::ShowSelectedSkeleton(const UPoseSearchDatabase* Database, int32 DbPoseIdx, float Time)
 {
 	UPoseSearchMeshComponent* Component = Skeletons[SelectedPose].Component.Get();
diff --git a/Engine/Plugins/Experimental/Animation/PoseSearch/Source/Editor/Private/PoseSearchDebuggerViewModel.h b/Engine/Plugins/Experimental/Animation/PoseSearch/Source/Editor/Private/PoseSearchDebuggerViewModel.h
--- a/Engine/Plugins/Experimental/Animation/PoseSearch/Source/Editor/Private/PoseSearchDebuggerViewModel.h
+++ b/Engine/Plugins/Experimental/Animation/PoseSearch/Source/Editor/Private/PoseSearchDebuggerViewModel.h
@@ -11,7 +11,6 @@
 struct FInstancedStruct;
 class IRewindDebugger;
 class UPoseSearchDatabase;
-class UPoseSearchSearchableAsset;
 
 namespace UE::PoseSearch
 {
@@ -42,7 +41,6 @@ public:
 	// Used for view callbacks
     const FTraceMotionMatchingStateMessage* GetMotionMatchingState() const;
 	const UPoseSearchDatabase* GetCurrentDatabase() const;
-	const UPoseSearchSearchableAsset* GetSearchableAsset() const;
 	const TArray<int32>* GetNodeIds() const;
 	int32 GetNodesNum() const;
 	const FTransform* GetRootTransform() const;
diff --git a/Engine/Plugins/Experimental/Animation/PoseSearch/Source/Runtime/Private/AnimNode_MotionMatching.cpp b/Engine/Plugins/Experimental/Animation/PoseSearch/Source/Runtime/Private/AnimNode_MotionMatching.cpp
--- a/Engine/Plugins/Experimental/Animation/PoseSearch/Source/Runtime/Private/AnimNode_MotionMatching.cpp
+++ b/Engine/Plugins/Experimental/Animation/PoseSearch/Source/Runtime/Private/AnimNode_MotionMatching.cpp
@@ -109,16 +109,6 @@ void FAnimNode_MotionMatching::UpdateAssetPlayer(const FAnimationUpdateContext&
 		DatabasesToSearch.Reset(1);
 		DatabasesToSearch.Add(Database);
 	}
-	// @todo: Delete this after removing Searchable.
-	else if (Searchable)
-	{
-		if (!bOverrideDatabaseInput && Searchable->IsA(UPoseSearchDatabase::StaticClass()))
-		{
-			DatabasesToSearch.Reset(1);
-			DatabasesToSearch.Add(Cast<UPoseSearchDatabase>(Searchable));
-		}
-		UE_LOG(LogPoseSearch, Warning, TEXT("FAnimNode_MotionMatching: The Searchable variable will be removed soon."));
-	}
 
 	// Execute core motion matching algorithm
 	UPoseSearchLibrary::UpdateMotionMatchingState(
diff --git a/Engine/Plugins/Experimental/Animation/PoseSearch/Source/Runtime/Private/PoseSearchLibrary.cpp b/Engine/Plugins/Experimental/Animation/PoseSearch/Source/Runtime/Private/PoseSearchLibrary.cpp
--- a/Engine/Plugins/Experimental/Animation/PoseSearch/Source/Runtime/Private/PoseSearchLibrary.cpp
+++ b/Engine/Plugins/Experimental/Animation/PoseSearch/Source/Runtime/Private/PoseSearchLibrary.cpp
@@ -138,7 +138,7 @@ void FMotionMatchingState::UpdateWantedPlayRate(const UE::PoseSearch::FSearchCon
 }
 
 void UPoseSearchLibrary::TraceMotionMatchingState(
-	const UPoseSearchSearchableAsset* Searchable,
+	const UPoseSearchDatabase* Database,
 	UE::PoseSearch::FSearchContext& SearchContext,
 	const UE::PoseSearch::FSearchResult& CurrentResult,
 	const UE::PoseSearch::FSearchResult& LastResult,
@@ -241,7 +241,7 @@ void UPoseSearchLibrary::TraceMotionMatchingState(
 		TraceState.AnimAngularVelocity = FMath::RadiansToDegrees(RootMotionTransformDelta.GetRotation().GetAngle()) / DeltaTime;
 	}
 
-	TraceState.SearchableAssetId = FTraceMotionMatchingState::GetIdFromObject(Searchable);
+	TraceState.SearchableAssetId = FTraceMotionMatchingState::GetIdFromObject(Database);
 	TraceState.ElapsedPoseSearchTime = ElapsedPoseSearchTime;
 	TraceState.AssetPlayerTime = CurrentResult.AssetTime;
 	TraceState.DeltaTime = DeltaTime;
diff --git a/Engine/Plugins/Experimental/Animation/PoseSearch/Source/Runtime/Public/PoseSearch/AnimNode_MotionMatching.h b/Engine/Plugins/Experimental/Animation/PoseSearch/Source/Runtime/Public/PoseSearch/AnimNode_MotionMatching.h
--- a/Engine/Plugins/Experimental/Animation/PoseSearch/Source/Runtime/Public/PoseSearch/AnimNode_MotionMatching.h
+++ b/Engine/Plugins/Experimental/Animation/PoseSearch/Source/Runtime/Public/PoseSearch/AnimNode_MotionMatching.h
@@ -10,7 +10,6 @@
 #include "AnimNode_MotionMatching.generated.h"
 
 class UPoseSearchDatabase;
-class UPoseSearchSearchableAsset;
 
 USTRUCT(BlueprintInternalUseOnly)
 struct POSESEARCH_API FAnimNode_MotionMatching : public FAnimNode_AssetPlayerBase
@@ -21,11 +20,6 @@ private:
 	UPROPERTY()
 	FPoseLink Source;
 
-	// @todo: Delete this after updating content.
-	// Collection of animations for motion matching.
-	UPROPERTY(EditAnywhere, Category=Settings, meta=(PinHiddenByDefault))
-	TObjectPtr<const UPoseSearchSearchableAsset> Searchable = nullptr;
-
 	// The database to search. This can be overridden by Anim Node Functions such as "On Become Relevant" and "On Update" via SetDatabaseToSearch/SetDatabasesToSearch.
 	UPROPERTY(EditAnywhere, Category = Settings, meta = (PinShownByDefault))
 	TObjectPtr<const UPoseSearchDatabase> Database = nullptr;
diff --git a/Engine/Plugins/Experimental/Animation/PoseSearch/Source/Runtime/Public/PoseSearch/PoseSearchDatabase.h b/Engine/Plugins/Experimental/Animation/PoseSearch/Source/Runtime/Public/PoseSearch/PoseSearchDatabase.h
--- a/Engine/Plugins/Experimental/Animation/PoseSearch/Source/Runtime/Public/PoseSearch/PoseSearchDatabase.h
+++ b/Engine/Plugins/Experimental/Animation/PoseSearch/Source/Runtime/Public/PoseSearch/PoseSearchDatabase.h
@@ -2,9 +2,10 @@
 
 #pragma once
 
+#include "Engine/DataAsset.h"
 #include "PoseSearch/PoseSearchAssetSampler.h"
 #include "PoseSearch/PoseSearchIndex.h"
-#include "PoseSearch/PoseSearchSearchableAsset.h"
+#include "PoseSearch/PoseSearchResult.h"
 #include "PoseSearchDatabase.generated.h"
 
 struct FInstancedStruct;
@@ -14,6 +15,11 @@ class UAnimMontage;
 class UBlendSpace;
 class UPoseSearchNormalizationSet;
 
+namespace UE::PoseSearch
+{
+	struct FSearchContext;
+} // namespace UE::PoseSearch
+
 UENUM()
 enum class EPoseSearchMode : int32
 {
@@ -233,7 +239,7 @@ struct POSESEARCH_API FPoseSearchDatabaseAnimMontage : public FPoseSearchDatabas
 
 /** A data asset for indexing a collection of animation sequences. */
 UCLASS(BlueprintType, Category = "Animation|Pose Search", Experimental, meta = (DisplayName = "Motion Database"))
-class POSESEARCH_API UPoseSearchDatabase : public UPoseSearchSearchableAsset
+class POSESEARCH_API UPoseSearchDatabase : public UDataAsset
 {
 	GENERATED_BODY()
 public:
@@ -323,9 +329,7 @@ public:
 	virtual void Serialize(FArchive& Ar) override;
 	// End UObject
 	
-	// Begin UPoseSearchSearchableAsset
-	virtual UE::PoseSearch::FSearchResult Search(UE::PoseSearch::FSearchContext& SearchContext) const override;
-	// End UPoseSearchSearchableAsset
+	UE::PoseSearch::FSearchResult Search(UE::PoseSearch::FSearchContext& SearchContext) const;
 
 #if WITH_EDITOR
 	virtual void BeginCacheForCookedPlatformData(const ITargetPlatform* TargetPlatform) override;
diff --git a/Engine/Plugins/Experimental/Animation/PoseSearch/Source/Runtime/Public/PoseSearch/PoseSearchLibrary.h b/Engine/Plugins/Experimental/Animation/PoseSearch/Source/Runtime/Public/PoseSearch/PoseSearchLibrary.h
--- a/Engine/Plugins/Experimental/Animation/PoseSearch/Source/Runtime/Public/PoseSearch/PoseSearchLibrary.h
+++ b/Engine/Plugins/Experimental/Animation/PoseSearch/Source/Runtime/Public/PoseSearch/PoseSearchLibrary.h
@@ -21,7 +21,6 @@ namespace UE::PoseSearch
 
 struct FAnimationUpdateContext;
 struct FTrajectorySampleRange;
-class UPoseSearchSearchableAsset;
 
 USTRUCT(BlueprintType, Category = "Animation|Pose Search")
 struct POSESEARCH_API FMotionMatchingSettings
@@ -116,7 +115,7 @@ class POSESEARCH_API UPoseSearchLibrary : public UBlueprintFunctionLibrary
 {
 	GENERATED_BODY()
 	static void TraceMotionMatchingState(
-		const UPoseSearchSearchableAsset* Searchable,
+		const UPoseSearchDatabase* Database,
 		UE::PoseSearch::FSearchContext& SearchContext,
 		const UE::PoseSearch::FSearchResult& CurrentResult,
 		const UE::PoseSearch::FSearchResult& LastResult,
diff --git a/Engine/Plugins/Experimental/Animation/PoseSearch/Source/Runtime/Public/PoseSearch/PoseSearchSearchableAsset.h b/Engine/Plugins/Experimental/Animation/PoseSearch/Source/Runtime/Public/PoseSearch/PoseSearchSearchableAsset.h
deleted file mode 100644
