motion matching - fix for posesearchdatabase BruteForce DDC indexing

diff --git a/Engine/Plugins/Experimental/Animation/PoseSearch/Source/Runtime/Private/PoseSearchDerivedData.cpp b/Engine/Plugins/Experimental/Animation/PoseSearch/Source/Runtime/Private/PoseSearchDerivedData.cpp
--- a/Engine/Plugins/Experimental/Animation/PoseSearch/Source/Runtime/Private/PoseSearchDerivedData.cpp
+++ b/Engine/Plugins/Experimental/Animation/PoseSearch/Source/Runtime/Private/PoseSearchDerivedData.cpp
@@ -822,7 +822,7 @@ void FPoseSearchDatabaseAsyncCacheTask::OnGetComplete(UE::DerivedData::FCacheGet
 		// cache can be corrupted in case the version of the derived data cache has not being updated while 
 		// developing channels that changes their cardinality without impacting any asset properties
 		// so to account for this, we just reindex the database and update the associated DDC 
-		if (!SearchIndex.IsEmpty() && SearchIndex.WeightsSqrt.Num() == Database->Schema->SchemaCardinality && SearchIndex.KDTree.Impl)
+		if (!SearchIndex.IsEmpty() && SearchIndex.WeightsSqrt.Num() == Database->Schema->SchemaCardinality)
 		{
 			UE_LOG(LogPoseSearch, Log, TEXT("%s - %s BuildIndex From Cache"), *LexToString(FullIndexKey.Hash), *Database->GetName());
 		}
