Replace bool parameters on Motion Matching node with cvars

diff --git a/Engine/Plugins/Experimental/Animation/PoseSearch/Source/Runtime/Private/AnimNode_MotionMatching.cpp b/Engine/Plugins/Experimental/Animation/PoseSearch/Source/Runtime/Private/AnimNode_MotionMatching.cpp
--- a/Engine/Plugins/Experimental/Animation/PoseSearch/Source/Runtime/Private/AnimNode_MotionMatching.cpp
+++ b/Engine/Plugins/Experimental/Animation/PoseSearch/Source/Runtime/Private/AnimNode_MotionMatching.cpp
@@ -19,6 +19,11 @@
 
 #define LOCTEXT_NAMESPACE "AnimNode_MotionMatching"
 
+#if ENABLE_ANIM_DEBUG
+static TAutoConsoleVariable<int32> CVarAnimNodeMotionMatchingDrawQuery(TEXT("a.AnimNode.MotionMatching.DebugDrawQuery"), 0, TEXT("Draw input query"));
+static TAutoConsoleVariable<int32> CVarAnimNodeMotionMatchingDrawCurResult(TEXT("a.AnimNode.MotionMatching.DebugDrawCurResult"), 0, TEXT("Draw current result"));
+#endif
+
 /////////////////////////////////////////////////////
 // FAnimNode_MotionMatching
 
@@ -141,7 +146,9 @@ void FAnimNode_MotionMatching::UpdateAssetPlayer(const FAnimationUpdateContext&
 
 
 #if WITH_EDITORONLY_DATA
-	if (bDebugDraw)
+	const bool bDebugDrawQuery = CVarAnimNodeMotionMatchingDrawQuery.GetValueOnAnyThread() > 0;
+	const bool bDebugDrawCurResult = CVarAnimNodeMotionMatchingDrawCurResult.GetValueOnAnyThread() > 0;
+	if (bDebugDrawQuery || bDebugDrawCurResult)
 	{
 		const UE::PoseSearch::FSearchResult& CurResult = MotionMatchingState.CurrentSearchResult;
 
@@ -153,7 +160,7 @@ void FAnimNode_MotionMatching::UpdateAssetPlayer(const FAnimationUpdateContext&
 		else
 #endif // WITH_EDITOR
 		{
-			if (bDebugDrawMatch)
+			if (bDebugDrawCurResult)
 			{
 				UE::PoseSearch::FDebugDrawParams DrawParams(Context.AnimInstanceProxy, CurResult.Database.Get());
 				DrawParams.DrawFeatureVector(CurResult.PoseIdx);
diff --git a/Engine/Plugins/Experimental/Animation/PoseSearch/Source/Runtime/Public/PoseSearch/AnimNode_MotionMatching.h b/Engine/Plugins/Experimental/Animation/PoseSearch/Source/Runtime/Public/PoseSearch/AnimNode_MotionMatching.h
--- a/Engine/Plugins/Experimental/Animation/PoseSearch/Source/Runtime/Public/PoseSearch/AnimNode_MotionMatching.h
+++ b/Engine/Plugins/Experimental/Animation/PoseSearch/Source/Runtime/Public/PoseSearch/AnimNode_MotionMatching.h
@@ -40,17 +40,6 @@ private:
 	UPROPERTY(EditAnywhere, Category = Settings, meta = (PinHiddenByDefault))
 	bool bForceInterrupt = false;
 
-#if WITH_EDITORONLY_DATA
-	UPROPERTY(EditAnywhere, Category=Debug, meta = (PinShownByDefault))
-	bool bDebugDraw = false;
-
-	UPROPERTY(EditAnywhere, Category=Debug, meta = (PinShownByDefault))
-	bool bDebugDrawQuery = true;
-
-	UPROPERTY(EditAnywhere, Category=Debug, meta = (PinShownByDefault))
-	bool bDebugDrawMatch = true;
-#endif
-
 public:
 
 	// FAnimNode_Base interface
