motion matching - less verbose FSearchContext::GetComponentSpaceTransform

diff --git a/Engine/Plugins/Experimental/Animation/PoseSearch/Source/Runtime/Private/PoseSearchContext.cpp b/Engine/Plugins/Experimental/Animation/PoseSearch/Source/Runtime/Private/PoseSearchContext.cpp
--- a/Engine/Plugins/Experimental/Animation/PoseSearch/Source/Runtime/Private/PoseSearchContext.cpp
+++ b/Engine/Plugins/Experimental/Animation/PoseSearch/Source/Runtime/Private/PoseSearchContext.cpp
@@ -386,13 +386,17 @@ FTransform FSearchContext::GetComponentSpaceTransform(float SampleTime, const UP
 		FTransform BoneComponentSpaceTransform;
 		if (!History->GetComponentSpaceTransformAtTime(SampleTime, BoneIndexType, BoneComponentSpaceTransform))
 		{
-			FName BoneName;
 			if (const USkeleton* Skeleton = Schema->Skeleton)
 			{
-				BoneName = Skeleton->GetReferenceSkeleton().GetBoneName(BoneIndexType);
+				if (!History->IsEmpty())
+				{
+					UE_LOG(LogPoseSearch, Warning, TEXT("FSearchContext::GetComponentSpaceTransform - Couldn't find BoneIndexType %d (%s) requested by %s"), BoneIndexType, *Skeleton->GetReferenceSkeleton().GetBoneName(BoneIndexType).ToString(), *Schema->GetName());
+				}
+			}
+			else
+			{
+				UE_LOG(LogPoseSearch, Warning, TEXT("FSearchContext::GetComponentSpaceTransform - Schema '%s' Skeleton is not properly set"), *Schema->GetName());
 			}
-
-			UE_LOG(LogPoseSearch, Warning, TEXT("FSearchContext::GetComponentSpaceTransform - Couldn't find BoneIndexType %d (%s) requested by %s"), BoneIndexType, *BoneName.ToString(), *Schema->GetName());
 		}
 
 		CachedTransforms.Add(SampleTime, BoneIndexType, BoneComponentSpaceTransform);
diff --git a/Engine/Plugins/Experimental/Animation/PoseSearch/Source/Runtime/Private/PoseSearchHistory.cpp b/Engine/Plugins/Experimental/Animation/PoseSearch/Source/Runtime/Private/PoseSearchHistory.cpp
--- a/Engine/Plugins/Experimental/Animation/PoseSearch/Source/Runtime/Private/PoseSearchHistory.cpp
+++ b/Engine/Plugins/Experimental/Animation/PoseSearch/Source/Runtime/Private/PoseSearchHistory.cpp
@@ -201,6 +201,11 @@ void FPoseHistory::GetRootTransformAtTime(float Time, FTransform& OutRootTransfo
 	}
 }
 
+bool FPoseHistory::IsEmpty() const
+{
+	return Entries.IsEmpty();
+}
+
 void FPoseHistory::Update(float SecondsElapsed, FCSPose<FCompactPose>& ComponentSpacePose, const FTransform& ComponentTransform)
 {
 	// Age our elapsed times
@@ -399,6 +404,12 @@ void FExtendedPoseHistory::GetRootTransformAtTime(float Time, FTransform& OutRoo
 	PoseHistory->GetRootTransformAtTime(Time, OutRootTransform, bExtrapolate);
 }
 
+bool FExtendedPoseHistory::IsEmpty() const
+{
+	check(PoseHistory);
+	return PoseHistory->IsEmpty() && FutureEntries.IsEmpty();
+}
+
 void FExtendedPoseHistory::ResetFuturePoses()
 {
 	FutureEntries.Reset();
diff --git a/Engine/Plugins/Experimental/Animation/PoseSearch/Source/Runtime/Public/PoseSearch/PoseSearchHistory.h b/Engine/Plugins/Experimental/Animation/PoseSearch/Source/Runtime/Public/PoseSearch/PoseSearchHistory.h
--- a/Engine/Plugins/Experimental/Animation/PoseSearch/Source/Runtime/Public/PoseSearch/PoseSearchHistory.h
+++ b/Engine/Plugins/Experimental/Animation/PoseSearch/Source/Runtime/Public/PoseSearch/PoseSearchHistory.h
@@ -26,6 +26,7 @@ struct POSESEARCH_API IPoseHistory
 	virtual float GetSampleTimeInterval() const = 0;
 	virtual bool GetComponentSpaceTransformAtTime(float Time, FBoneIndexType BoneIndexType, FTransform& OutBoneTransform, bool bExtrapolate = true) const = 0;
 	virtual void GetRootTransformAtTime(float Time, FTransform& OutRootTransform, bool bExtrapolate = true) const = 0;
+	virtual bool IsEmpty() const = 0;
 };
 
 struct FPoseHistoryEntry
@@ -53,6 +54,7 @@ struct FPoseHistory : public IPoseHistory
 	virtual float GetSampleTimeInterval() const override;
 	virtual bool GetComponentSpaceTransformAtTime(float Time, FBoneIndexType BoneIndexType, FTransform& OutBoneTransform, bool bExtrapolate = true) const override;
 	virtual void GetRootTransformAtTime(float Time, FTransform& OutRootTransform, bool bExtrapolate = true) const override;
+	virtual bool IsEmpty() const override;
 	// End of IPoseHistory interface
 
 #if ENABLE_DRAW_DEBUG && ENABLE_ANIM_DEBUG
@@ -75,6 +77,7 @@ struct FExtendedPoseHistory : public IPoseHistory
 	virtual float GetSampleTimeInterval() const override;
 	virtual bool GetComponentSpaceTransformAtTime(float Time, FBoneIndexType BoneIndexType, FTransform& OutBoneTransform, bool bExtrapolate = true) const override;
 	virtual void GetRootTransformAtTime(float Time, FTransform& OutRootTransform, bool bExtrapolate = true) const override;
+	virtual bool IsEmpty() const override;
 	// End of IPoseHistory interface
 
 	void ResetFuturePoses();
