motion matching - fix for CVarAnimNodeMotionMatchingDrawQuery': undeclared identifier

diff --git a/Engine/Plugins/Experimental/Animation/PoseSearch/Source/Runtime/Private/AnimNode_MotionMatching.cpp b/Engine/Plugins/Experimental/Animation/PoseSearch/Source/Runtime/Private/AnimNode_MotionMatching.cpp
--- a/Engine/Plugins/Experimental/Animation/PoseSearch/Source/Runtime/Private/AnimNode_MotionMatching.cpp
+++ b/Engine/Plugins/Experimental/Animation/PoseSearch/Source/Runtime/Private/AnimNode_MotionMatching.cpp
@@ -145,7 +145,7 @@ void FAnimNode_MotionMatching::UpdateAssetPlayer(const FAnimationUpdateContext&
 	Source.Update(Context);
 
 
-#if WITH_EDITORONLY_DATA
+#if WITH_EDITORONLY_DATA && ENABLE_ANIM_DEBUG
 	const bool bDebugDrawQuery = CVarAnimNodeMotionMatchingDrawQuery.GetValueOnAnyThread() > 0;
 	const bool bDebugDrawCurResult = CVarAnimNodeMotionMatchingDrawCurResult.GetValueOnAnyThread() > 0;
 	if (bDebugDrawQuery || bDebugDrawCurResult)
