motion matching - pose search debugger integration with database editor

diff --git a/Engine/Plugins/Experimental/Animation/PoseSearch/Source/Editor/Private/PoseSearchDatabaseAssetListItem.cpp b/Engine/Plugins/Experimental/Animation/PoseSearch/Source/Editor/Private/PoseSearchDatabaseAssetListItem.cpp
--- a/Engine/Plugins/Experimental/Animation/PoseSearch/Source/Editor/Private/PoseSearchDatabaseAssetListItem.cpp
+++ b/Engine/Plugins/Experimental/Animation/PoseSearch/Source/Editor/Private/PoseSearchDatabaseAssetListItem.cpp
@@ -133,15 +133,8 @@ namespace UE::PoseSearch
 						{
 							float AnimationAssetTime = 0.f;
 							FVector AnimationAssetBlendParameters = FVector::ZeroVector;
-							TArray<const FPoseSearchIndexAsset*> SearchIndexAssets;
-							if (Database->GetPoseSearchIndexAssets(Node->SourceAssetIdx, SearchIndexAssets) && !SearchIndexAssets.IsEmpty())
-							{
-								const FPoseSearchIndexAsset* SearchIndexAsset = SearchIndexAssets[0];
-								check(SearchIndexAsset);
-								AnimationAssetTime = ViewModel->GetPlayTime() + SearchIndexAsset->SamplingInterval.Min;
-								AnimationAssetBlendParameters = SearchIndexAsset->BlendParameters;
-							}
-
+							ViewModel->GetAnimationTime(Node->SourceAssetIdx, AnimationAssetTime, AnimationAssetBlendParameters);
+							
 							const IAnimationEditor* AnimationEditor = static_cast<IAnimationEditor*>(Editor);
 							const UDebugSkelMeshComponent* PreviewComponent = AnimationEditor->GetPersonaToolkit()->GetPreviewMeshComponent();
 
diff --git a/Engine/Plugins/Experimental/Animation/PoseSearch/Source/Editor/Private/PoseSearchDatabaseAssetTree.cpp b/Engine/Plugins/Experimental/Animation/PoseSearch/Source/Editor/Private/PoseSearchDatabaseAssetTree.cpp
--- a/Engine/Plugins/Experimental/Animation/PoseSearch/Source/Editor/Private/PoseSearchDatabaseAssetTree.cpp
+++ b/Engine/Plugins/Experimental/Animation/PoseSearch/Source/Editor/Private/PoseSearchDatabaseAssetTree.cpp
@@ -730,19 +730,23 @@ namespace UE::PoseSearch
 		EditorViewModel.Pin()->BuildSearchIndex();
 	}
 
-	TWeakPtr<FDatabaseAssetTreeNode> SDatabaseAssetTree::SetSelectedItem(int32 SourceAssetIdx)
+	void SDatabaseAssetTree::SetSelectedItem(int32 SourceAssetIdx, bool bClearSelection)
 	{
-		for (TSharedPtr<FDatabaseAssetTreeNode> & Node : AllNodes)
+		if (bClearSelection)
 		{
-			if (Node->SourceAssetIdx == SourceAssetIdx)
+			TreeView->ClearSelection();
+		}
+
+		if (SourceAssetIdx >= 0)
+		{
+			for (TSharedPtr<FDatabaseAssetTreeNode>& Node : AllNodes)
 			{
-				TreeView->ClearSelection();
-				TreeView->SetItemSelection(Node, true);
-				return Node;
+				if (Node->SourceAssetIdx == SourceAssetIdx)
+				{
+					TreeView->SetItemSelection(Node, true);
+				}
 			}
 		}
-
-		return TWeakPtr<FDatabaseAssetTreeNode>();
 	}
 }
 
diff --git a/Engine/Plugins/Experimental/Animation/PoseSearch/Source/Editor/Private/PoseSearchDatabaseEditor.cpp b/Engine/Plugins/Experimental/Animation/PoseSearch/Source/Editor/Private/PoseSearchDatabaseEditor.cpp
--- a/Engine/Plugins/Experimental/Animation/PoseSearch/Source/Editor/Private/PoseSearchDatabaseEditor.cpp
+++ b/Engine/Plugins/Experimental/Animation/PoseSearch/Source/Editor/Private/PoseSearchDatabaseEditor.cpp
@@ -9,6 +9,7 @@
 #include "PoseSearchDatabaseEditorReflection.h"
 #include "PoseSearchEditor.h"
 #include "PoseSearch/PoseSearchDerivedData.h"
+#include "Framework/Application/SlateApplication.h"
 #include "GameFramework/WorldSettings.h"
 #include "AdvancedPreviewSceneModule.h"
 #include "InstancedStruct.h"
@@ -25,7 +26,6 @@
 namespace UE::PoseSearch
 {
 	const FName PoseSearchDatabaseEditorAppName = FName(TEXT("PoseSearchDatabaseEditorApp"));
-	constexpr double ViewRangeSlack = 0.25;
 
 	// Tab identifiers
 	struct FDatabaseEditorTabs
@@ -55,16 +55,13 @@ namespace UE::PoseSearch
 		return ViewModel.IsValid() ? ViewModel->GetPoseSearchDatabase() : nullptr;
 	}
 
-	void FDatabaseEditor::SetSelectedAsset(int32 SourceAssetIdx)
+	void FDatabaseEditor::SetSelectedPoseIdx(int32 PoseIdx)
 	{
 		if (ViewModel.IsValid())
 		{
-			const TWeakPtr<FDatabaseAssetTreeNode> Node = AssetTreeWidget->SetSelectedItem(SourceAssetIdx);
-			
-			if (Node.IsValid())
-			{
-				ViewModel->SetSelectedNode(Node.Pin());
-			}
+			const bool bClearSelection = !FSlateApplication::Get().GetModifierKeys().IsControlDown() || ViewModel->IsEditorSelection();
+			const int32 SourceAssetIdx = ViewModel->SetSelectedNode(PoseIdx, bClearSelection);
+			AssetTreeWidget->SetSelectedItem(SourceAssetIdx, bClearSelection);
 		}
 	}
 
@@ -142,15 +139,22 @@ namespace UE::PoseSearch
 				PreviewScene.ToSharedRef());
 			
 			PreviewWidget = SNew(SDatabasePreview, PreviewArgs)
-				.SliderScrubTime_Lambda([this]() { return ViewModel->GetPlayTime(); })
+				.SliderColor_Lambda([this]()
+					{
+						return ViewModel->IsEditorSelection() ? FLinearColor::Red.CopyWithNewOpacity(0.5f) : FLinearColor::Blue.CopyWithNewOpacity(0.5f);
+					})
+				.SliderScrubTime_Lambda([this]()
+					{
+						return ViewModel->GetPlayTime();
+					})
 				.SliderViewRange_Lambda([this]() 
-				{ 
-					return TRange<double>(-ViewRangeSlack, ViewModel->GetMaxPreviewPlayLength() + ViewRangeSlack);
-				})
+					{ 
+						return ViewModel->GetPreviewPlayRange();
+					})
 				.OnSliderScrubPositionChanged_Lambda([this](float NewScrubPosition, bool bScrubbing)
-				{
-					ViewModel->SetPlayTime(NewScrubPosition, !bScrubbing);
-				})
+					{
+						ViewModel->SetPlayTime(NewScrubPosition, !bScrubbing);
+					})
 				.OnBackwardEnd_Raw(this, &FDatabaseEditor::PreviewBackwardEnd)
 				.OnBackwardStep_Raw(this, &FDatabaseEditor::PreviewBackwardStep)
 				.OnBackward_Raw(this, &FDatabaseEditor::PreviewBackward)
@@ -590,7 +594,10 @@ namespace UE::PoseSearch
 			}
 		}
 
-		ViewModel->SetSelectedNodes(SelectedItems);
+		if (SelectionType != ESelectInfo::Direct)
+		{
+			ViewModel->SetSelectedNodes(SelectedItems);
+		}
 	}
 
 	FDatabaseEditor::FSelectionWidget& FDatabaseEditor::FindOrAddSelectionWidget(const UScriptStruct* ScriptStructType)
diff --git a/Engine/Plugins/Experimental/Animation/PoseSearch/Source/Editor/Private/PoseSearchDatabaseViewModel.cpp b/Engine/Plugins/Experimental/Animation/PoseSearch/Source/Editor/Private/PoseSearchDatabaseViewModel.cpp
--- a/Engine/Plugins/Experimental/Animation/PoseSearch/Source/Editor/Private/PoseSearchDatabaseViewModel.cpp
+++ b/Engine/Plugins/Experimental/Animation/PoseSearch/Source/Editor/Private/PoseSearchDatabaseViewModel.cpp
@@ -80,53 +80,7 @@ namespace UE::PoseSearch
 
 		RemovePreviewActors();
 
-		PoseSearchDatabase->RegisterOnDerivedDataRebuild(UPoseSearchDatabase::FOnDerivedDataRebuild::CreateSP(this, &FDatabaseViewModel::ResetPreviewActors));
-	}
-
-	void FDatabaseViewModel::ResetPreviewActors()
-	{
-		PlayTime = 0.0f;
-		DeltaTimeMultiplier = 1.0f;
-		RespawnPreviewActors();
-	}
-
-	void FDatabaseViewModel::RespawnPreviewActors()
-	{
-		RemovePreviewActors();
-
-		MaxPreviewPlayLength = 0.0f;
-
-		if (FAsyncPoseSearchDatabasesManagement::RequestAsyncBuildIndex(PoseSearchDatabase, ERequestAsyncBuildFlag::ContinueRequest))
-		{
-			FBoneContainer BoneContainer;
-			BoneContainer.InitializeTo(PoseSearchDatabase->Schema->BoneIndicesWithParents, UE::Anim::FCurveFilterSettings(UE::Anim::ECurveFilterMode::DisallowAll), *PoseSearchDatabase->Schema->Skeleton);
-
-			TSet<int32> AssociatedAssetIndices;
-			for (const TSharedPtr<FDatabaseAssetTreeNode>& SelectedNode : SelectedNodes)
-			{
-				AssociatedAssetIndices.Add(SelectedNode->SourceAssetIdx);
-			}
-
-			const FPoseSearchIndex& SearchIndex = PoseSearchDatabase->GetSearchIndex();
-			for (int32 IndexAssetIndex = 0; IndexAssetIndex < SearchIndex.Assets.Num(); ++IndexAssetIndex)
-			{
-				const FPoseSearchIndexAsset& IndexAsset = SearchIndex.Assets[IndexAssetIndex];
-				if ((AnimationPreviewMode == EAnimationPreviewMode::OriginalAndMirrored || !IndexAsset.bMirrored) &&
-					AssociatedAssetIndices.Contains(IndexAsset.SourceAssetIdx))
-				{
-					FDatabasePreviewActor PreviewActor = SpawnPreviewActor(IndexAssetIndex, BoneContainer);
-					if (PreviewActor.IsValid())
-					{
-						MaxPreviewPlayLength = FMath::Max(MaxPreviewPlayLength, IndexAsset.SamplingInterval.Max - IndexAsset.SamplingInterval.Min);
-						PreviewActors.Add(PreviewActor);
-					}
-				}
-			}
-
-			ParallelFor(PreviewActors.Num(), [this](int32 PreviewActorIndex) { PreviewActors[PreviewActorIndex].Process(); }, ParallelForFlags);
-
-			UpdatePreviewActors();
-		}
+		PoseSearchDatabase->RegisterOnDerivedDataRebuild(UPoseSearchDatabase::FOnDerivedDataRebuild::CreateSP(this, &FDatabaseViewModel::RemovePreviewActors));
 	}
 
 	void FDatabaseViewModel::BuildSearchIndex()
@@ -137,12 +91,12 @@ namespace UE::PoseSearch
 
 	void FDatabaseViewModel::PreviewBackwardEnd()
 	{
-		PlayTime = 0.0f;
+		PlayTime = MinPreviewPlayLength;
 	}
 
 	void FDatabaseViewModel::PreviewBackwardStep()
 	{
-		PlayTime = FMath::Clamp(PlayTime - StepDeltaTime, 0.0f, MaxPreviewPlayLength);
+		PlayTime = FMath::Clamp(PlayTime - StepDeltaTime, MinPreviewPlayLength, MaxPreviewPlayLength);
 		DeltaTimeMultiplier = 0.0f;
 	}
 
@@ -163,7 +117,7 @@ namespace UE::PoseSearch
 
 	void FDatabaseViewModel::PreviewForwardStep()
 	{
-		PlayTime = FMath::Clamp(PlayTime + StepDeltaTime, 0.0f, MaxPreviewPlayLength);
+		PlayTime = FMath::Clamp(PlayTime + StepDeltaTime, MinPreviewPlayLength, MaxPreviewPlayLength);
 		DeltaTimeMultiplier = 0.0f;
 	}
 
@@ -172,7 +126,7 @@ namespace UE::PoseSearch
 		PlayTime = MaxPreviewPlayLength;
 	}
 
-	FDatabasePreviewActor FDatabaseViewModel::SpawnPreviewActor(int32 IndexAssetIndex, const FBoneContainer& BoneContainer)
+	FDatabasePreviewActor FDatabaseViewModel::SpawnPreviewActor(int32 IndexAssetIndex, const FBoneContainer& BoneContainer, float PlayTimeOffset)
 	{
 		FDatabasePreviewActor PreviewActor;
 
@@ -193,6 +147,7 @@ namespace UE::PoseSearch
 		PreviewActor.Sampler.Init(DatabaseAnimationAssetBase->GetAnimationAsset(), IndexAsset.BlendParameters, BoneContainer);
 		PreviewActor.IndexAssetIndex = IndexAssetIndex;
 		PreviewActor.CurrentPoseIndex = INDEX_NONE;
+		PreviewActor.PlayTimeOffset = PlayTimeOffset;
 
 		FActorSpawnParameters Params;
 		Params.SpawnCollisionHandlingOverride = ESpawnActorCollisionHandlingMethod::AdjustIfPossibleButAlwaysSpawn;
@@ -232,18 +187,12 @@ namespace UE::PoseSearch
 		return PreviewActor;
 	}
 
-	UWorld* FDatabaseViewModel::GetWorld() const
+	UWorld* FDatabaseViewModel::GetWorld()
 	{
 		check(PreviewScenePtr.IsValid());
 		return PreviewScenePtr.Pin()->GetWorld();
 	}
 
-	UObject* FDatabaseViewModel::GetPlaybackContext() const
-	{
-		return GetWorld();
-	}
-
-
 	void FDatabaseViewModel::OnPreviewActorClassChanged()
 	{
 		// todo: implement
@@ -252,7 +201,7 @@ namespace UE::PoseSearch
 	void FDatabaseViewModel::Tick(float DeltaSeconds)
 	{
 		PlayTime += DeltaSeconds * DeltaTimeMultiplier;
-		PlayTime = FMath::Clamp(PlayTime, 0.0f, MaxPreviewPlayLength);
+		PlayTime = FMath::Clamp(PlayTime, MinPreviewPlayLength, MaxPreviewPlayLength);
 
 		if (FAsyncPoseSearchDatabasesManagement::RequestAsyncBuildIndex(PoseSearchDatabase, ERequestAsyncBuildFlag::ContinueRequest))
 		{
@@ -280,7 +229,7 @@ namespace UE::PoseSearch
 
 			float CurrentTime = 0.f;
 			const FPoseSearchIndexAsset& IndexAsset = SearchIndex.Assets[PreviewActor.IndexAssetIndex];
-			float CurrentPlayTime = PlayTime + IndexAsset.SamplingInterval.Min;
+			float CurrentPlayTime = PlayTime + IndexAsset.SamplingInterval.Min + PreviewActor.PlayTimeOffset;
 			FAnimationRuntime::AdvanceTime(false, CurrentPlayTime, CurrentTime, IndexAsset.SamplingInterval.Max);
 
 			// time to pose index
@@ -304,21 +253,35 @@ namespace UE::PoseSearch
 			{
 				RootMotion = MirrorRootMotion(RootMotion, AnimInstance->GetMirrorDataTable());
 			}
-			PreviewActor.Actor->SetActorTransform(RootMotion);
 
+			if (PreviewActor.PlayTimeOffset != 0.f)
+			{
+				FTransform OriginRootMotion = PreviewActor.Sampler.ExtractRootTransform(PreviewActor.PlayTimeOffset);
+				if (AnimInstance->GetMirrorDataTable())
+				{
+					OriginRootMotion = MirrorRootMotion(OriginRootMotion, AnimInstance->GetMirrorDataTable());
+				}
+				RootMotion.SetToRelativeTransform(OriginRootMotion);
+			}
+
+			PreviewActor.Actor->SetActorTransform(RootMotion);
 		}
 	}
 
 	void FDatabaseViewModel::RemovePreviewActors()
 	{
+		PlayTime = 0.f;
+		DeltaTimeMultiplier = 1.f;
+		MaxPreviewPlayLength = 0.f;
+		MinPreviewPlayLength = 0.f;
+		bIsEditorSelection = true;
+
 		for (auto PreviewActor : PreviewActors)
 		{
 			// @todo: PreviewActor.Actor is a TWeakObjectPtr so it can be null.
 			PreviewActor.Actor->Destroy();
 		}
 		PreviewActors.Reset();
-
-		MaxPreviewPlayLength = 0.0f;
 	}
 
 	FTransform FDatabaseViewModel::MirrorRootMotion(
@@ -354,7 +317,7 @@ namespace UE::PoseSearch
 		if (PreviewMode != AnimationPreviewMode)
 		{
 			AnimationPreviewMode = PreviewMode;
-			RespawnPreviewActors();
+			RemovePreviewActors();
 		}
 	}
 
@@ -419,18 +382,94 @@ namespace UE::PoseSearch
 		return false;
 	}
 
-	void FDatabaseViewModel::SetSelectedNode(const TSharedPtr<FDatabaseAssetTreeNode>& InSelectedNode)
+	int32 FDatabaseViewModel::SetSelectedNode(int32 PoseIdx, bool bClearSelection)
 	{
-		SelectedNodes.Reset(1);
-		SelectedNodes.Push(InSelectedNode);
-		ResetPreviewActors();
+		int32 SelectedSourceAssetIdx = -1;
+
+		if (bClearSelection)
+		{
+			RemovePreviewActors();
+		}
+
+		bIsEditorSelection = false;
+
+		if (FAsyncPoseSearchDatabasesManagement::RequestAsyncBuildIndex(PoseSearchDatabase, ERequestAsyncBuildFlag::ContinueRequest))
+		{
+			FBoneContainer BoneContainer;
+			BoneContainer.InitializeTo(PoseSearchDatabase->Schema->BoneIndicesWithParents, FCurveEvaluationOption(false), *PoseSearchDatabase->Schema->Skeleton);
+
+			const FPoseSearchIndex& SearchIndex = PoseSearchDatabase->GetSearchIndex();
+			if (SearchIndex.PoseMetadata.IsValidIndex(PoseIdx))
+			{
+				const uint32 IndexAssetIndex = SearchIndex.PoseMetadata[PoseIdx].GetAssetIndex();
+				if (SearchIndex.Assets.IsValidIndex(IndexAssetIndex))
+				{
+					const FPoseSearchIndexAsset& IndexAsset = SearchIndex.Assets[IndexAssetIndex];
+					// clamping just to make sure the time interval is correct
+					const float PlayTimeOffset = FMath::Clamp(PoseSearchDatabase->GetAssetTime(PoseIdx), IndexAsset.SamplingInterval.Min, IndexAsset.SamplingInterval.Max);
+					FDatabasePreviewActor PreviewActor = SpawnPreviewActor(IndexAssetIndex, BoneContainer, PlayTimeOffset);
+					if (PreviewActor.IsValid())
+					{
+						MaxPreviewPlayLength = FMath::Max(MaxPreviewPlayLength, IndexAsset.SamplingInterval.Max - PlayTimeOffset);
+						MinPreviewPlayLength = FMath::Min(MinPreviewPlayLength, IndexAsset.SamplingInterval.Min - PlayTimeOffset);
+						PreviewActors.Add(PreviewActor);
+						SelectedSourceAssetIdx = IndexAsset.SourceAssetIdx;
+					}
+				}
+			}
+
+			ParallelFor(PreviewActors.Num(), [this](int32 PreviewActorIndex) { PreviewActors[PreviewActorIndex].Process(); }, ParallelForFlags);
+
+			UpdatePreviewActors();
+
+			SetPlayTime(0.f, false);
+		}
+
 		ProcessSelectedActor(nullptr);
+
+		return SelectedSourceAssetIdx;
 	}
 
 	void FDatabaseViewModel::SetSelectedNodes(const TArrayView<TSharedPtr<FDatabaseAssetTreeNode>>& InSelectedNodes)
 	{
-		SelectedNodes = InSelectedNodes;
-		ResetPreviewActors();
+		bIsEditorSelection = true;
+
+		RemovePreviewActors();
+
+		if (FAsyncPoseSearchDatabasesManagement::RequestAsyncBuildIndex(PoseSearchDatabase, ERequestAsyncBuildFlag::ContinueRequest))
+		{
+			FBoneContainer BoneContainer;
+			BoneContainer.InitializeTo(PoseSearchDatabase->Schema->BoneIndicesWithParents, FCurveEvaluationOption(false), *PoseSearchDatabase->Schema->Skeleton);
+
+			TMap<int32, int32> AssociatedAssetIndices;
+			for (int32 i = 0; i < InSelectedNodes.Num(); ++i)
+			{
+				AssociatedAssetIndices.FindOrAdd(InSelectedNodes[i]->SourceAssetIdx) = i;
+			}
+
+			const FPoseSearchIndex& SearchIndex = PoseSearchDatabase->GetSearchIndex();
+			for (int32 IndexAssetIndex = 0; IndexAssetIndex < SearchIndex.Assets.Num(); ++IndexAssetIndex)
+			{
+				const FPoseSearchIndexAsset& IndexAsset = SearchIndex.Assets[IndexAssetIndex];
+				if (AnimationPreviewMode == EAnimationPreviewMode::OriginalAndMirrored || !IndexAsset.bMirrored)
+				{
+					if (const int32* SelectedNodesIndex = AssociatedAssetIndices.Find(IndexAsset.SourceAssetIdx))
+					{
+						FDatabasePreviewActor PreviewActor = SpawnPreviewActor(IndexAssetIndex, BoneContainer, 0.f);
+						if (PreviewActor.IsValid())
+						{
+							MaxPreviewPlayLength = FMath::Max(MaxPreviewPlayLength, IndexAsset.SamplingInterval.Max - IndexAsset.SamplingInterval.Min);
+							PreviewActors.Add(PreviewActor);
+						}
+					}
+				}
+			}
+
+			ParallelFor(PreviewActors.Num(), [this](int32 PreviewActorIndex) { PreviewActors[PreviewActorIndex].Process(); }, ParallelForFlags);
+
+			UpdatePreviewActors();
+		}
+
 		ProcessSelectedActor(nullptr);
 	}
 
@@ -456,14 +495,15 @@ namespace UE::PoseSearch
 		return nullptr;
 	}
 
-	float FDatabaseViewModel::GetMaxPreviewPlayLength() const
+	TRange<double> FDatabaseViewModel::GetPreviewPlayRange() const
 	{
-		return MaxPreviewPlayLength;
+		constexpr double ViewRangeSlack = 0.2;
+		return TRange<double>(MinPreviewPlayLength - ViewRangeSlack, MaxPreviewPlayLength + ViewRangeSlack);
 	}
 
 	float FDatabaseViewModel::GetPlayTime() const
 	{
-		const float ClampedPlayTime = FMath::Clamp(PlayTime, 0.0f, MaxPreviewPlayLength);
+		const float ClampedPlayTime = FMath::Clamp(PlayTime, MinPreviewPlayLength, MaxPreviewPlayLength);
 		return ClampedPlayTime;
 	}
 
@@ -472,4 +512,39 @@ namespace UE::PoseSearch
 		PlayTime = NewPlayTime;
 		DeltaTimeMultiplier = bInTickPlayTime ? DeltaTimeMultiplier : 0.0f;
 	}
+
+	bool FDatabaseViewModel::GetAnimationTime(int32 SourceAssetIdx, float& CurrentPlayTime, FVector& BlendParameters) const
+	{
+		if (FAsyncPoseSearchDatabasesManagement::RequestAsyncBuildIndex(PoseSearchDatabase, ERequestAsyncBuildFlag::ContinueRequest))
+		{
+			const FPoseSearchIndex& SearchIndex = PoseSearchDatabase->GetSearchIndex();
+			for (const FDatabasePreviewActor& PreviewActor : GetPreviewActors())
+			{
+				if (PreviewActor.IndexAssetIndex >= 0 && PreviewActor.IndexAssetIndex < SearchIndex.Assets.Num() && PreviewActor.Sampler.IsInitialized())
+				{
+					const FPoseSearchIndexAsset& IndexAsset = SearchIndex.Assets[PreviewActor.IndexAssetIndex];
+					if (IndexAsset.SourceAssetIdx == SourceAssetIdx)
+					{
+						CurrentPlayTime = PlayTime + IndexAsset.SamplingInterval.Min + PreviewActor.PlayTimeOffset;
+						BlendParameters = IndexAsset.BlendParameters;
+						return true;
+					}
+				}
+			}
+
+			for (const FPoseSearchIndexAsset& IndexAsset : SearchIndex.Assets)
+			{
+				if (IndexAsset.SourceAssetIdx == SourceAssetIdx)
+				{
+					CurrentPlayTime = PlayTime + IndexAsset.SamplingInterval.Min;
+					BlendParameters = IndexAsset.BlendParameters;
+					return true;
+				}
+			}
+		}
+
+		CurrentPlayTime = 0.f;
+		BlendParameters = FVector::ZeroVector;
+		return false;
+	}
 }
diff --git a/Engine/Plugins/Experimental/Animation/PoseSearch/Source/Editor/Private/PoseSearchDebuggerDatabaseColumns.h b/Engine/Plugins/Experimental/Animation/PoseSearch/Source/Editor/Private/PoseSearchDebuggerDatabaseColumns.h
--- a/Engine/Plugins/Experimental/Animation/PoseSearch/Source/Editor/Private/PoseSearchDebuggerDatabaseColumns.h
+++ b/Engine/Plugins/Experimental/Animation/PoseSearch/Source/Editor/Private/PoseSearchDebuggerDatabaseColumns.h
@@ -158,8 +158,7 @@ struct FDatabaseName : IColumn
 									FDatabaseEditor* DatabaseEditor = static_cast<FDatabaseEditor*>(Editor);
 									
 									// Open asset paused and at specific time as seen on the pose search debugger.
-									DatabaseEditor->SetSelectedAsset(RowData->DbAssetIdx);
-									DatabaseEditor->GetViewModel()->SetPlayTime(RowData->AssetTime, false);
+									DatabaseEditor->SetSelectedPoseIdx(RowData->PoseIdx);
 								}
 							}
 						}
diff --git a/Engine/Plugins/Experimental/Animation/PoseSearch/Source/Editor/Private/SPoseSearchDatabaseViewport.cpp b/Engine/Plugins/Experimental/Animation/PoseSearch/Source/Editor/Private/SPoseSearchDatabaseViewport.cpp
--- a/Engine/Plugins/Experimental/Animation/PoseSearch/Source/Editor/Private/SPoseSearchDatabaseViewport.cpp
+++ b/Engine/Plugins/Experimental/Animation/PoseSearch/Source/Editor/Private/SPoseSearchDatabaseViewport.cpp
@@ -141,7 +141,8 @@ namespace UE::PoseSearch
 
 	void SDatabasePreview::Construct(const FArguments& InArgs, const FDatabasePreviewRequiredArgs& InRequiredArgs)
 	{
-		SliderScrubTimeAttribute = InArgs._SliderScrubTime;
+		SliderColor = InArgs._SliderColor;
+		SliderScrubTime = InArgs._SliderScrubTime;
 		SliderViewRange = InArgs._SliderViewRange;
 		OnSliderScrubPositionChanged = InArgs._OnSliderScrubPositionChanged;
 
@@ -208,18 +209,29 @@ namespace UE::PoseSearch
 				[
 					SNew(SSimpleTimeSlider)
 					.ClampRangeHighlightSize(0.15f)
-					.ClampRangeHighlightColor(FLinearColor::Red.CopyWithNewOpacity(0.5f))
-					.ScrubPosition_Lambda([this]() { return SliderScrubTimeAttribute.Get(); })
-					.ViewRange_Lambda([this]() { return SliderViewRange.Get(); })
-					.ClampRange_Lambda([this]() { return SliderViewRange.Get(); })
-					.OnScrubPositionChanged_Lambda(
-						[this](double NewScrubTime, bool bIsScrubbing)
-					{
-						if (bIsScrubbing)
+					.ClampRangeHighlightColor_Lambda([this]()
+						{
+							return SliderColor.Get();
+						})
+					.ScrubPosition_Lambda([this]()
+						{
+							return SliderScrubTime.Get();
+						})
+					.ViewRange_Lambda([this]()
+						{
+							return SliderViewRange.Get();
+						})
+					.ClampRange_Lambda([this]()
+						{
+							return SliderViewRange.Get();
+						})
+					.OnScrubPositionChanged_Lambda([this](double NewScrubTime, bool bIsScrubbing)
 						{
-							OnSliderScrubPositionChanged.ExecuteIfBound(NewScrubTime, bIsScrubbing);
-						}
-					})
+							if (bIsScrubbing)
+							{
+								OnSliderScrubPositionChanged.ExecuteIfBound(NewScrubTime, bIsScrubbing);
+							}
+						})
 				]
 			]
 		];
diff --git a/Engine/Plugins/Experimental/Animation/PoseSearch/Source/Editor/Public/PoseSearchDatabaseAssetTree.h b/Engine/Plugins/Experimental/Animation/PoseSearch/Source/Editor/Public/PoseSearchDatabaseAssetTree.h
--- a/Engine/Plugins/Experimental/Animation/PoseSearch/Source/Editor/Public/PoseSearchDatabaseAssetTree.h
+++ b/Engine/Plugins/Experimental/Animation/PoseSearch/Source/Editor/Public/PoseSearchDatabaseAssetTree.h
@@ -39,7 +39,7 @@ namespace UE::PoseSearch
 		void RefreshTreeView(bool bIsInitialSetup = false, bool bRecoverSelection = false);
 		void FinalizeTreeChanges(bool bRecoverSelection = false);
 
-		TWeakPtr<FDatabaseAssetTreeNode> SetSelectedItem(int32 SourceAssetIdx);
+		void SetSelectedItem(int32 SourceAssetIdx, bool bClearSelection);
 		FDetailColumnSizeData& GetColumnSizeData() { return ColumnSizeData; }
 		
 	protected:
diff --git a/Engine/Plugins/Experimental/Animation/PoseSearch/Source/Editor/Public/PoseSearchDatabaseEditor.h b/Engine/Plugins/Experimental/Animation/PoseSearch/Source/Editor/Public/PoseSearchDatabaseEditor.h
--- a/Engine/Plugins/Experimental/Animation/PoseSearch/Source/Editor/Public/PoseSearchDatabaseEditor.h
+++ b/Engine/Plugins/Experimental/Animation/PoseSearch/Source/Editor/Public/PoseSearchDatabaseEditor.h
@@ -40,7 +40,7 @@ namespace UE::PoseSearch
 		FDatabaseViewModel* GetViewModel() const { return ViewModel.Get(); }
 		TSharedPtr<FDatabaseViewModel> GetViewModelSharedPtr() const { return ViewModel; }
 		
-		void SetSelectedAsset(int32 SourceAssetIdx);
+		void SetSelectedPoseIdx(int32 PoseIdx);
 		
 		void BuildSearchIndex();
 
diff --git a/Engine/Plugins/Experimental/Animation/PoseSearch/Source/Editor/Public/PoseSearchDatabaseViewModel.h b/Engine/Plugins/Experimental/Animation/PoseSearch/Source/Editor/Public/PoseSearchDatabaseViewModel.h
--- a/Engine/Plugins/Experimental/Animation/PoseSearch/Source/Editor/Public/PoseSearchDatabaseViewModel.h
+++ b/Engine/Plugins/Experimental/Animation/PoseSearch/Source/Editor/Public/PoseSearchDatabaseViewModel.h
@@ -45,6 +45,7 @@ namespace UE::PoseSearch
 		TWeakObjectPtr<AActor> Actor = nullptr;
 		int32 IndexAssetIndex = INDEX_NONE;
 		int32 CurrentPoseIndex = INDEX_NONE;
+		float PlayTimeOffset = 0.f;
 		FAnimationAssetSampler Sampler;
 
 		bool IsValid() const;
@@ -67,8 +68,6 @@ namespace UE::PoseSearch
 		void Initialize(UPoseSearchDatabase* InPoseSearchDatabase, const TSharedRef<FDatabasePreviewScene>& InPreviewScene);
 
 		void RemovePreviewActors();
-		void ResetPreviewActors();
-		void RespawnPreviewActors();
 		void BuildSearchIndex();
 
 		void PreviewBackwardEnd();
@@ -79,7 +78,7 @@ namespace UE::PoseSearch
 		void PreviewForwardStep();
 		void PreviewForwardEnd();
 
-		UPoseSearchDatabase* GetPoseSearchDatabase() const { return PoseSearchDatabase; }
+		UPoseSearchDatabase* GetPoseSearchDatabase() { return PoseSearchDatabase; }
 		void OnPreviewActorClassChanged();
 
 		void Tick(float DeltaSeconds);
@@ -105,19 +104,22 @@ namespace UE::PoseSearch
 		void SetIsEnabled(int32 AnimationAssetIndex, bool bEnabled);
 		bool IsEnabled(int32 AnimationAssetIndex) const;
 
-		void SetSelectedNode(const TSharedPtr<FDatabaseAssetTreeNode>& InSelectedNode);
+		int32 SetSelectedNode(int32 PoseIdx, bool bClearSelection);
 		void SetSelectedNodes(const TArrayView<TSharedPtr<FDatabaseAssetTreeNode>>& InSelectedNodes);
 		void ProcessSelectedActor(AActor* Actor);
 		
 		const FPoseSearchIndexAsset* GetSelectedActorIndexAsset() const;
 
-		float GetMaxPreviewPlayLength() const;
-		float GetPlayTime() const;
+		TRange<double> GetPreviewPlayRange() const;
+
 		void SetPlayTime(float NewPlayTime, bool bInTickPlayTime);
+		float GetPlayTime() const;
+		bool IsEditorSelection() const { return bIsEditorSelection; }
+		bool GetAnimationTime(int32 SourceAssetIdx, float& CurrentPlayTime, FVector& BlendParameters) const;
 
 	private:
-		float PlayTime = 0.0f;
-		float DeltaTimeMultiplier = 1.0f;
+		float PlayTime = 0.f;
+		float DeltaTimeMultiplier = 1.f;
 
 		/** Scene asset being viewed and edited by this view model. */
 		TObjectPtr<UPoseSearchDatabase> PoseSearchDatabase;
@@ -129,7 +131,9 @@ namespace UE::PoseSearch
 		TArray<FDatabasePreviewActor> PreviewActors;
 		
 		/** From zero to the play length of the longest preview */
-		float MaxPreviewPlayLength = 0.0f;
+		float MaxPreviewPlayLength = 0.f;
+		float MinPreviewPlayLength = 0.f;
+		bool bIsEditorSelection = true;
 
 		/** What features to show in the viewport */
 		EFeaturesDrawMode PoseFeaturesDrawMode = EFeaturesDrawMode::All;
@@ -140,15 +144,11 @@ namespace UE::PoseSearch
 		/** Is animation debug draw enabled */
 		bool DisplayRootMotionSpeed = false;
 		
-		TArray<TSharedPtr<FDatabaseAssetTreeNode>> SelectedNodes;
-
 		int32 SelectedActorIndexAssetIndex = INDEX_NONE;
 
-		UWorld* GetWorld() const;
-
-		UObject* GetPlaybackContext() const;
+		UWorld* GetWorld();
 
-		FDatabasePreviewActor SpawnPreviewActor(int32 IndexAssetIndex, const FBoneContainer& BoneContainer);
+		FDatabasePreviewActor SpawnPreviewActor(int32 IndexAssetIndex, const FBoneContainer& BoneContainer, float PlayTimeOffset);
 
 		void UpdatePreviewActors(bool bInTickPlayTime = false);
 
diff --git a/Engine/Plugins/Experimental/Animation/PoseSearch/Source/Editor/Public/SPoseSearchDatabaseViewport.h b/Engine/Plugins/Experimental/Animation/PoseSearch/Source/Editor/Public/SPoseSearchDatabaseViewport.h
--- a/Engine/Plugins/Experimental/Animation/PoseSearch/Source/Editor/Public/SPoseSearchDatabaseViewport.h
+++ b/Engine/Plugins/Experimental/Animation/PoseSearch/Source/Editor/Public/SPoseSearchDatabaseViewport.h
@@ -72,6 +72,7 @@ namespace UE::PoseSearch
 		DECLARE_DELEGATE(FOnButtonClickedEvent)
 
 		SLATE_BEGIN_ARGS(SDatabasePreview) {}
+			SLATE_ATTRIBUTE(FLinearColor, SliderColor);
 			SLATE_ATTRIBUTE(double, SliderScrubTime);
 			SLATE_ATTRIBUTE(TRange<double>, SliderViewRange);
 			SLATE_EVENT(FOnScrubPositionChanged, OnSliderScrubPositionChanged);
@@ -87,7 +88,8 @@ namespace UE::PoseSearch
 		void Construct(const FArguments& InArgs, const FDatabasePreviewRequiredArgs& InRequiredArgs);
 
 	protected:
-		TAttribute<double> SliderScrubTimeAttribute;
+		TAttribute<FLinearColor> SliderColor;
+		TAttribute<double> SliderScrubTime;
 		TAttribute<TRange<double>> SliderViewRange = TRange<double>(0.0, 1.0);
 		FOnScrubPositionChanged OnSliderScrubPositionChanged;
 		FOnButtonClickedEvent OnBackwardEnd;
diff --git a/Engine/Plugins/Experimental/Animation/PoseSearch/Source/Runtime/Private/PoseSearchDatabase.cpp b/Engine/Plugins/Experimental/Animation/PoseSearch/Source/Runtime/Private/PoseSearchDatabase.cpp
--- a/Engine/Plugins/Experimental/Animation/PoseSearch/Source/Runtime/Private/PoseSearchDatabase.cpp
+++ b/Engine/Plugins/Experimental/Animation/PoseSearch/Source/Runtime/Private/PoseSearchDatabase.cpp
@@ -636,26 +636,6 @@ bool UPoseSearchDatabase::IsCachedCookedPlatformDataLoaded(const ITargetPlatform
 	check(IsInGameThread());
 	return FAsyncPoseSearchDatabasesManagement::RequestAsyncBuildIndex(this, ERequestAsyncBuildFlag::ContinueRequest);
 }
-
-// returns all the FPoseSearchIndexAsset associated with the AnimationAssetIndex
-bool UPoseSearchDatabase::GetPoseSearchIndexAssets(int32 AnimationAssetIndex, TArray<const FPoseSearchIndexAsset*>& OutSearchIndexAssets) const
-{
-	using namespace UE::PoseSearch;
-	OutSearchIndexAssets.Reset();
-	if (FAsyncPoseSearchDatabasesManagement::RequestAsyncBuildIndex(this, ERequestAsyncBuildFlag::ContinueRequest))
-	{
-		const FPoseSearchIndex& SearchIndex = GetSearchIndex();
-		for (const FPoseSearchIndexAsset& SearchIndexAsset : SearchIndex.Assets)
-		{
-			if (SearchIndexAsset.SourceAssetIdx == AnimationAssetIndex)
-			{
-				OutSearchIndexAssets.Add(&SearchIndexAsset);
-			}
-		}
-		return true;
-	}
-	return false;
-}
 #endif // WITH_EDITOR
 
 void UPoseSearchDatabase::PostSaveRoot(FObjectPostSaveRootContext ObjectSaveContext)
diff --git a/Engine/Plugins/Experimental/Animation/PoseSearch/Source/Runtime/Public/PoseSearch/PoseSearchDatabase.h b/Engine/Plugins/Experimental/Animation/PoseSearch/Source/Runtime/Public/PoseSearch/PoseSearchDatabase.h
--- a/Engine/Plugins/Experimental/Animation/PoseSearch/Source/Runtime/Public/PoseSearch/PoseSearchDatabase.h
+++ b/Engine/Plugins/Experimental/Animation/PoseSearch/Source/Runtime/Public/PoseSearch/PoseSearchDatabase.h
@@ -335,9 +335,6 @@ public:
 	virtual void BeginCacheForCookedPlatformData(const ITargetPlatform* TargetPlatform) override;
 	virtual bool IsCachedCookedPlatformDataLoaded(const ITargetPlatform* TargetPlatform) override;
 
-	// returns all the FPoseSearchIndexAsset associated with the AnimationAssetIndex
-	bool GetPoseSearchIndexAssets(int32 AnimationAssetIndex, TArray<const FPoseSearchIndexAsset*>& OutSearchIndexAssets) const;
-
 private:
 	DECLARE_MULTICAST_DELEGATE(FOnDerivedDataRebuildMulticaster);
 	FOnDerivedDataRebuildMulticaster OnDerivedDataRebuild;
