motion matching - upgrading FBoneContainer::InitializeTo to avoid deprecation error

diff --git a/Engine/Plugins/Experimental/Animation/PoseSearch/Source/Editor/Private/PoseSearchDatabaseViewModel.cpp b/Engine/Plugins/Experimental/Animation/PoseSearch/Source/Editor/Private/PoseSearchDatabaseViewModel.cpp
--- a/Engine/Plugins/Experimental/Animation/PoseSearch/Source/Editor/Private/PoseSearchDatabaseViewModel.cpp
+++ b/Engine/Plugins/Experimental/Animation/PoseSearch/Source/Editor/Private/PoseSearchDatabaseViewModel.cpp
@@ -396,7 +396,7 @@ namespace UE::PoseSearch
 		if (FAsyncPoseSearchDatabasesManagement::RequestAsyncBuildIndex(PoseSearchDatabase, ERequestAsyncBuildFlag::ContinueRequest))
 		{
 			FBoneContainer BoneContainer;
-			BoneContainer.InitializeTo(PoseSearchDatabase->Schema->BoneIndicesWithParents, FCurveEvaluationOption(false), *PoseSearchDatabase->Schema->Skeleton);
+			BoneContainer.InitializeTo(PoseSearchDatabase->Schema->BoneIndicesWithParents, UE::Anim::FCurveFilterSettings(UE::Anim::ECurveFilterMode::DisallowAll), *PoseSearchDatabase->Schema->Skeleton);
 
 			const FPoseSearchIndex& SearchIndex = PoseSearchDatabase->GetSearchIndex();
 			if (SearchIndex.PoseMetadata.IsValidIndex(PoseIdx))
@@ -439,7 +439,7 @@ namespace UE::PoseSearch
 		if (FAsyncPoseSearchDatabasesManagement::RequestAsyncBuildIndex(PoseSearchDatabase, ERequestAsyncBuildFlag::ContinueRequest))
 		{
 			FBoneContainer BoneContainer;
-			BoneContainer.InitializeTo(PoseSearchDatabase->Schema->BoneIndicesWithParents, FCurveEvaluationOption(false), *PoseSearchDatabase->Schema->Skeleton);
+			BoneContainer.InitializeTo(PoseSearchDatabase->Schema->BoneIndicesWithParents, UE::Anim::FCurveFilterSettings(UE::Anim::ECurveFilterMode::DisallowAll), *PoseSearchDatabase->Schema->Skeleton);
 
 			TMap<int32, int32> AssociatedAssetIndices;
 			for (int32 i = 0; i < InSelectedNodes.Num(); ++i)
