motion matching - fix for Assert encountered when MotionMatching packaged build is launched

diff --git a/Engine/Plugins/Experimental/Animation/PoseSearch/Source/Runtime/Private/PoseSearchFeatureChannel_Pose.cpp b/Engine/Plugins/Experimental/Animation/PoseSearch/Source/Runtime/Private/PoseSearchFeatureChannel_Pose.cpp
--- a/Engine/Plugins/Experimental/Animation/PoseSearch/Source/Runtime/Private/PoseSearchFeatureChannel_Pose.cpp
+++ b/Engine/Plugins/Experimental/Animation/PoseSearch/Source/Runtime/Private/PoseSearchFeatureChannel_Pose.cpp
@@ -13,12 +13,10 @@
 
 UPoseSearchFeatureChannel_Pose::UPoseSearchFeatureChannel_Pose()
 {
-#if WITH_EDITOR
 	// defaulting UPoseSearchFeatureChannel_Pose for a meaningful locomotion setup
 	Weight = 1.f;
 	SampledBones.Add(FPoseSearchBone({ {"foot_l"}, int32(EPoseSearchBoneFlags::Position | EPoseSearchBoneFlags::Velocity), 1.f, FLinearColor::Green }));
 	SampledBones.Add(FPoseSearchBone({ {"foot_r"}, int32(EPoseSearchBoneFlags::Position | EPoseSearchBoneFlags::Velocity), 1.f, FLinearColor::Green}));
-#endif
 }
 
 void UPoseSearchFeatureChannel_Pose::Finalize(UPoseSearchSchema* Schema)
diff --git a/Engine/Plugins/Experimental/Animation/PoseSearch/Source/Runtime/Private/PoseSearchFeatureChannel_Trajectory.cpp b/Engine/Plugins/Experimental/Animation/PoseSearch/Source/Runtime/Private/PoseSearchFeatureChannel_Trajectory.cpp
--- a/Engine/Plugins/Experimental/Animation/PoseSearch/Source/Runtime/Private/PoseSearchFeatureChannel_Trajectory.cpp
+++ b/Engine/Plugins/Experimental/Animation/PoseSearch/Source/Runtime/Private/PoseSearchFeatureChannel_Trajectory.cpp
@@ -14,14 +14,12 @@
 
 UPoseSearchFeatureChannel_Trajectory::UPoseSearchFeatureChannel_Trajectory()
 {
-#if WITH_EDITOR
 	// defaulting UPoseSearchFeatureChannel_Trajectory for a meaningful locomotion setup
 	Weight = 7.f;
 	Samples.Add(FPoseSearchTrajectorySample({ -0.4f, int32(EPoseSearchTrajectoryFlags::PositionXY), 0.4f, FLinearColor::Red }));
 	Samples.Add(FPoseSearchTrajectorySample({ 0.f, int32(EPoseSearchTrajectoryFlags::VelocityXY | EPoseSearchTrajectoryFlags::FacingDirectionXY), 2.f, FLinearColor::Blue }));
 	Samples.Add(FPoseSearchTrajectorySample({ 0.35f, int32(EPoseSearchTrajectoryFlags::PositionXY | EPoseSearchTrajectoryFlags::FacingDirectionXY), 0.7f, FLinearColor::Blue }));
 	Samples.Add(FPoseSearchTrajectorySample({ 0.7f, int32(EPoseSearchTrajectoryFlags::VelocityXY | EPoseSearchTrajectoryFlags::PositionXY | EPoseSearchTrajectoryFlags::FacingDirectionXY), 0.5f, FLinearColor::Blue }));
-#endif
 }
 
 void UPoseSearchFeatureChannel_Trajectory::Finalize(UPoseSearchSchema* Schema)
