motion matchig - tuning database editor layout (usability improvements when docked in half monitor size)

diff --git a/Engine/Plugins/Experimental/Animation/PoseSearch/Source/Editor/Private/PoseSearchDatabaseEditor.cpp b/Engine/Plugins/Experimental/Animation/PoseSearch/Source/Editor/Private/PoseSearchDatabaseEditor.cpp
--- a/Engine/Plugins/Experimental/Animation/PoseSearch/Source/Editor/Private/PoseSearchDatabaseEditor.cpp
+++ b/Engine/Plugins/Experimental/Animation/PoseSearch/Source/Editor/Private/PoseSearchDatabaseEditor.cpp
@@ -214,58 +214,47 @@ namespace UE::PoseSearch
 		
 		// Define Editor Layout
 		const TSharedRef<FTabManager::FLayout> StandaloneDefaultLayout =
-		FTabManager::NewLayout("Standalone_PoseSearchDatabaseEditor_Layout_v0.08")
+		FTabManager::NewLayout("Standalone_PoseSearchDatabaseEditor_Layout_v0.12")
 			->AddArea
 			(
-				FTabManager::NewPrimaryArea()->SetOrientation(Orient_Vertical)
+				FTabManager::NewPrimaryArea()->SetOrientation(Orient_Horizontal)
 				->Split
 				(
-					FTabManager::NewSplitter()->SetOrientation(Orient_Vertical)->SetSizeCoefficient(0.9f)
+					FTabManager::NewSplitter()->SetOrientation(Orient_Vertical)
+					->SetSizeCoefficient(0.4f)
 					->Split
 					(
-						FTabManager::NewSplitter()->SetOrientation(Orient_Horizontal)->SetSizeCoefficient(0.9f)
-						->Split
-						(
-							FTabManager::NewSplitter()->SetOrientation(Orient_Vertical)->SetSizeCoefficient(0.9f)
-							->Split
-							(
-								FTabManager::NewStack()
-								->SetSizeCoefficient(0.5f)
-								->AddTab(FDatabaseEditorTabs::AssetTreeViewID, ETabState::OpenedTab)
-								->SetHideTabWell(false)
-							)
-							->Split
-							(
-								FTabManager::NewStack()
-								->SetSizeCoefficient(0.5f)
-								->AddTab(FDatabaseEditorTabs::AssetDetailsID, ETabState::OpenedTab)
-								->SetHideTabWell(false)
-							)
-						)
-						->Split
-						(
-							FTabManager::NewStack()
-							->SetSizeCoefficient(0.4f)
-							->AddTab(FDatabaseEditorTabs::ViewportID, ETabState::OpenedTab)
-							->SetHideTabWell(true)
-						)
-						->Split
-						(
-							FTabManager::NewSplitter()->SetOrientation(Orient_Vertical)
-							->Split
-							(
-								FTabManager::NewStack()
-								->SetSizeCoefficient(0.3f)
-								->AddTab(FDatabaseEditorTabs::SelectionDetailsID, ETabState::OpenedTab)
-								->AddTab(FDatabaseEditorTabs::PreviewSettingsID, ETabState::OpenedTab)
-							)
-							->Split
-							(
-								FTabManager::NewStack()
-								->SetSizeCoefficient(0.5f)
-								->AddTab(FDatabaseEditorTabs::StatisticsOverview, ETabState::OpenedTab)
-							)
-						)
+						FTabManager::NewStack()
+						->SetSizeCoefficient(0.6f)
+						->AddTab(FDatabaseEditorTabs::AssetTreeViewID, ETabState::OpenedTab)
+						->SetHideTabWell(false)
+					)
+					->Split
+					(
+						FTabManager::NewStack()
+						->SetSizeCoefficient(0.4f)
+						->AddTab(FDatabaseEditorTabs::AssetDetailsID, ETabState::OpenedTab)
+						->SetHideTabWell(false)
+					)
+				)
+				->Split
+				(
+					FTabManager::NewSplitter()->SetOrientation(Orient_Vertical)
+					->SetSizeCoefficient(0.6f)
+					->Split
+					(
+						FTabManager::NewStack()
+						->SetSizeCoefficient(0.6f)
+						->AddTab(FDatabaseEditorTabs::ViewportID, ETabState::OpenedTab)
+						->SetHideTabWell(true)
+					)
+					->Split
+					(
+						FTabManager::NewStack()
+						->SetSizeCoefficient(0.4f)
+						->AddTab(FDatabaseEditorTabs::StatisticsOverview, ETabState::OpenedTab)
+						->AddTab(FDatabaseEditorTabs::PreviewSettingsID, ETabState::OpenedTab)
+						->AddTab(FDatabaseEditorTabs::SelectionDetailsID, ETabState::OpenedTab)
 					)
 				)
 			);
