motion matching - removing pose search debugger AssetPlayRate, PlaySelectedAsset, DrawSkeleton(s)

diff --git a/Engine/Plugins/Experimental/Animation/PoseSearch/Source/Editor/Private/PoseSearchDebuggerReflection.h b/Engine/Plugins/Experimental/Animation/PoseSearch/Source/Editor/Private/PoseSearchDebuggerReflection.h
--- a/Engine/Plugins/Experimental/Animation/PoseSearch/Source/Editor/Private/PoseSearchDebuggerReflection.h
+++ b/Engine/Plugins/Experimental/Animation/PoseSearch/Source/Editor/Private/PoseSearchDebuggerReflection.h
@@ -5,18 +5,6 @@
 #include "CoreMinimal.h"
 #include "PoseSearchDebuggerReflection.generated.h"
 
-/**
- * Used by the reflection UObject to encompass draw options for the query and database selections
- */
-USTRUCT()
-struct FPoseSearchDebuggerFeatureDrawOptions
-{
-	GENERATED_BODY()
-
-	UPROPERTY(EditAnywhere, Category="Draw Options")
-    bool bDisable = false;
-};
-
 /**
  * Reflection UObject being observed in the details view panel of the debugger
  */
@@ -26,17 +14,10 @@ class POSESEARCHEDITOR_API UPoseSearchDebuggerReflection : public UObject
 	GENERATED_BODY()
 
 public:
-
-	UPROPERTY(VisibleAnywhere, Category="Motion Matching State", Meta=(DisplayName="Current Database"))
-	FString CurrentDatabaseName = "";
-
 	/** Time since last PoseSearch */
 	UPROPERTY(VisibleAnywhere, Category="Motion Matching State")
 	float ElapsedPoseSearchTime = 0.0f;
 
-	UPROPERTY(VisibleAnywhere, Category = "Motion Matching State")
-	FString AssetPlayerAssetName = "";
-
 	UPROPERTY(VisibleAnywhere, Category = "Motion Matching State")
 	float AssetPlayerTime = 0.0f;
 
@@ -55,18 +36,6 @@ public:
 	UPROPERTY(VisibleAnywhere, Category = "Motion Matching State")
 	float AnimAngularVelocity;
 
-	UPROPERTY(EditAnywhere, Category="Draw Options", Meta=(DisplayName="Query"))
-	FPoseSearchDebuggerFeatureDrawOptions QueryDrawOptions;
-
-	UPROPERTY(EditAnywhere, Category="Draw Options", Meta=(DisplayName="Selected Pose"))
-	FPoseSearchDebuggerFeatureDrawOptions SelectedPoseDrawOptions;
-
-	UPROPERTY(EditAnywhere, Category = "Draw Options")
-	bool bDrawActiveSkeleton = false;
-	
-	UPROPERTY(EditAnywhere, Category = "Draw Options")
-	bool bDrawSelectedSkeleton = false;
-
     UPROPERTY(VisibleAnywhere, Category="Pose Vectors")
 	TArray<float> QueryPoseVector;
     	
diff --git a/Engine/Plugins/Experimental/Animation/PoseSearch/Source/Editor/Private/PoseSearchDebuggerView.cpp b/Engine/Plugins/Experimental/Animation/PoseSearch/Source/Editor/Private/PoseSearchDebuggerView.cpp
--- a/Engine/Plugins/Experimental/Animation/PoseSearch/Source/Editor/Private/PoseSearchDebuggerView.cpp
+++ b/Engine/Plugins/Experimental/Animation/PoseSearch/Source/Editor/Private/PoseSearchDebuggerView.cpp
@@ -95,15 +95,8 @@ void SDebuggerDetailsView::UpdateReflection(const FTraceMotionMatchingStateMessa
 		const FPoseSearchIndex& CurrentSearchIndex = CurrentDatabase->GetSearchIndex();
 		int32 CurrentDbPoseIdx = State.GetCurrentDatabasePoseIndex();
 
-		Reflection->CurrentDatabaseName = CurrentDatabase->GetName();
 		Reflection->ElapsedPoseSearchTime = State.ElapsedPoseSearchTime;
 
-		Reflection->AssetPlayerAssetName = "None";
-		if (const FPoseSearchIndexAsset* IndexAsset = CurrentSearchIndex.GetAssetForPoseSafe(CurrentDbPoseIdx))
-		{
-			Reflection->AssetPlayerAssetName = CurrentDatabase->GetSourceAssetName(*IndexAsset);
-		}
-
 		Reflection->AssetPlayerTime = State.AssetPlayerTime;
 		Reflection->LastDeltaTime = State.DeltaTime;
 		Reflection->SimLinearVelocity = State.SimLinearVelocity;
@@ -161,7 +154,7 @@ void SDebuggerView::Construct(const FArguments& InArgs, uint64 InAnimInstanceId)
 		SAssignNew(DebuggerView, SVerticalBox)
 		+ SVerticalBox::Slot()
 		.FillHeight(1.0f)
-		.VAlign(VAlign_Center)
+		.VAlign(VAlign_Fill)
 		.HAlign(HAlign_Fill)
 		[
 			SAssignNew(Switcher, SWidgetSwitcher)
@@ -171,7 +164,7 @@ void SDebuggerView::Construct(const FArguments& InArgs, uint64 InAnimInstanceId)
 			+ SWidgetSwitcher::Slot()
 			.Padding(40.0f)
 			.HAlign(HAlign_Fill)
-            .VAlign(VAlign_Center)
+            .VAlign(VAlign_Fill)
 			[
 				SAssignNew(SelectionView, SVerticalBox)
 			]
@@ -272,8 +265,6 @@ void SDebuggerView::Tick(const FGeometry& AllottedGeometry, const double InCurre
 		bUpdated = true;
 	}
 
-	Model->UpdateAsset();
-	
 	// Draw visualization every tick
 	DrawVisualization();
 }
@@ -354,47 +345,40 @@ void SDebuggerView::DrawFeatures(
 	const USkinnedMeshComponent* Mesh
 ) const
 {
-	const TObjectPtr<UPoseSearchDebuggerReflection> Reflection = DetailsView->GetReflection();
-
+	// @todo: have a check box like "Channels Breakdown" to disable drawing the query vector
+	// @todo: draw the query with greater thickness?
 	// Draw query vector
-	if (!Reflection || !Reflection->QueryDrawOptions.bDisable)
+	const UPoseSearchDatabase* CurrentDatabase = ViewModel.Get()->GetCurrentDatabase();
+	if (CurrentDatabase)
 	{
-		const UPoseSearchDatabase* CurrentDatabase = ViewModel.Get()->GetCurrentDatabase();
-		if (CurrentDatabase)
+		for (const FTraceMotionMatchingStateDatabaseEntry& DbEntry : State.DatabaseEntries)
 		{
-			for (const FTraceMotionMatchingStateDatabaseEntry& DbEntry : State.DatabaseEntries)
+			const UPoseSearchDatabase* Database = FTraceMotionMatchingState::GetObjectFromId<UPoseSearchDatabase>(DbEntry.DatabaseId);
+			if (Database && Database == CurrentDatabase && DbEntry.QueryVector.Num() == Database->Schema->SchemaCardinality &&
+				FAsyncPoseSearchDatabasesManagement::RequestAsyncBuildIndex(CurrentDatabase, ERequestAsyncBuildFlag::ContinueRequest))
 			{
-				const UPoseSearchDatabase* Database = FTraceMotionMatchingState::GetObjectFromId<UPoseSearchDatabase>(DbEntry.DatabaseId);
-				if (Database && Database == CurrentDatabase && DbEntry.QueryVector.Num() == Database->Schema->SchemaCardinality && 
-					FAsyncPoseSearchDatabasesManagement::RequestAsyncBuildIndex(CurrentDatabase, ERequestAsyncBuildFlag::ContinueRequest))
-				{
-					FDebugDrawParams DrawParams(&DebuggerWorld, Mesh, CurrentDatabase, EDebugDrawFlags::DrawQuery);
-					DrawParams.DrawFeatureVector(DbEntry.QueryVector);
-					break;
-				}
+				FDebugDrawParams DrawParams(&DebuggerWorld, Mesh, CurrentDatabase, EDebugDrawFlags::DrawQuery);
+				DrawParams.DrawFeatureVector(DbEntry.QueryVector);
+				break;
 			}
 		}
 	}
 
 	// Draw selected poses
-	if (!Reflection || !Reflection->QueryDrawOptions.bDisable)
-	{
-		const TSharedPtr<SListView<TSharedRef<FDebuggerDatabaseRowData>>>& DatabaseRows = DatabaseView->GetDatabaseRows();
-		TArray<TSharedRef<FDebuggerDatabaseRowData>> SelectedRows = DatabaseRows->GetSelectedItems();
+	const TSharedPtr<SListView<TSharedRef<FDebuggerDatabaseRowData>>>& DatabaseRows = DatabaseView->GetDatabaseRows();
+	TArray<TSharedRef<FDebuggerDatabaseRowData>> SelectedRows = DatabaseRows->GetSelectedItems();
 	
-		// Draw any selected database vectors
-		for (const TSharedRef<FDebuggerDatabaseRowData>& Row : SelectedRows)
+	// Draw any selected database vectors
+	for (const TSharedRef<FDebuggerDatabaseRowData>& Row : SelectedRows)
+	{
+		const UPoseSearchDatabase* RowDatabase = Row->SourceDatabase.Get();
+		if (FAsyncPoseSearchDatabasesManagement::RequestAsyncBuildIndex(RowDatabase, ERequestAsyncBuildFlag::ContinueRequest))
 		{
-			const UPoseSearchDatabase* RowDatabase = Row->SourceDatabase.Get();
-			if (FAsyncPoseSearchDatabasesManagement::RequestAsyncBuildIndex(RowDatabase, ERequestAsyncBuildFlag::ContinueRequest))
-			{
-				FDebugDrawParams DrawParams(&DebuggerWorld, Mesh, RowDatabase);
-				DrawParams.DrawFeatureVector(Row->PoseIdx);
-			}
+			FDebugDrawParams DrawParams(&DebuggerWorld, Mesh, RowDatabase);
+			DrawParams.DrawFeatureVector(Row->PoseIdx);
 		}
 	}
 
-
 	// Draw active pose
 	{
 		TArray<TSharedRef<FDebuggerDatabaseRowData>> ActiveRows = DatabaseView->GetActiveRow()->GetSelectedItems();
@@ -433,23 +417,6 @@ void SDebuggerView::DrawFeatures(
 		}
 	}
 
-
-	// Draw skeleton
-	{
-		FSkeletonDrawParams SkeletonDrawParams;
-		if (Reflection && Reflection->bDrawSelectedSkeleton)
-		{
-			SkeletonDrawParams.Flags |= ESkeletonDrawFlags::SelectedPose;
-		}
-		if (Reflection && Reflection->bDrawActiveSkeleton)
-		{
-			SkeletonDrawParams.Flags |= ESkeletonDrawFlags::ActivePose;
-		}
-
-		SkeletonDrawParams.Flags |= ESkeletonDrawFlags::Asset;
-
-		ViewModel.Get()->OnDraw(SkeletonDrawParams);
-	}
 }
 
 int32 SDebuggerView::SelectView() const
@@ -504,9 +471,6 @@ void SDebuggerView::OnPoseSelectionChanged(const UPoseSearchDatabase* Database,
 	else
 	{
 		Model->ShowSelectedSkeleton(Database, DbPoseIdx, Time);
-
-		// Stop asset player when switching selections
-		Model->StopSelection();
 	}
 }
 
@@ -518,27 +482,6 @@ FReply SDebuggerView::OnUpdateNodeSelection(int32 InSelectedNodeId)
 	return FReply::Handled();
 }
 
-FReply SDebuggerView::TogglePlaySelectedAssets() const
-{
-	const TSharedPtr<SListView<TSharedRef<FDebuggerDatabaseRowData>>>& DatabaseRows = DatabaseView->GetDatabaseRows();
-	TArray<TSharedRef<FDebuggerDatabaseRowData>> Selected = DatabaseRows->GetSelectedItems();
-	const bool bPlaying = ViewModel.Get()->IsPlayingSelections();
-	if (!bPlaying)
-	{
-		if (!Selected.IsEmpty())
-		{
-			// @TODO: Make functional with multiple poses being selected
-			ViewModel.Get()->PlaySelection(Selected[0]->PoseIdx, Selected[0]->AssetTime);
-		}
-	}
-	else
-	{
-		ViewModel.Get()->StopSelection();
-	}
-
-	return FReply::Handled();
-}
-
 TSharedRef<SWidget> SDebuggerView::GenerateNoDataMessageView()
 {
 	TSharedRef<SWidget> ReturnButtonView = GenerateReturnButtonView();
@@ -603,113 +546,42 @@ TSharedRef<SHorizontalBox> SDebuggerView::GenerateReturnButtonView()
 					.Justification(ETextJustify::Center)
 				]
 			]
-		];
-}
-
-TSharedRef<SWidget> SDebuggerView::GenerateNodeDebuggerView()
-{
-	TSharedRef<SHorizontalBox> ReturnButtonView = GenerateReturnButtonView();
-	ReturnButtonView->AddSlot()
-	.VAlign(VAlign_Center)
-	.HAlign(HAlign_Fill)
-	.Padding(32, 5, 0, 0)
-	.AutoWidth()
-	[
-		SNew(SButton)
-		.VAlign(VAlign_Center)
-		.HAlign(HAlign_Fill)
-		.ButtonStyle(FAppStyle::Get(), "Button")
-		.ContentPadding( FMargin(5, 0) )
-		.OnClicked(this, &SDebuggerView::TogglePlaySelectedAssets)
+		]
+		+SHorizontalBox::Slot()
+		.VAlign(VAlign_Top)
+		.HAlign(HAlign_Left)
+		.Padding(64, 5, 0, 0)
+		.AutoWidth()
 		[
 			SNew(SHorizontalBox)
-			// Icon
+
 			+ SHorizontalBox::Slot()
 			.AutoWidth()
+			.Padding(0, 5, 0, 0)
 			[
-				SNew(SImage)
-				.Image_Lambda([this]
+				SNew(SCheckBox)
+				.IsChecked_Lambda([this]
 				{
-					const bool bPlayingSelections = ViewModel.Get()->IsPlayingSelections();
-					return FSlateIcon("FAppStyle", bPlayingSelections ? "PlayWorld.StopPlaySession.Small" : "PlayWorld.PlayInViewport.Small").GetSmallIcon();
+					return ViewModel.Get()->IsVerbose() ? ECheckBoxState::Checked : ECheckBoxState::Unchecked; 
 				})
+				.OnCheckStateChanged_Lambda([this](ECheckBoxState State)
+				{
+					ViewModel.Get()->SetVerbose(State == ECheckBoxState::Checked);
+					UpdateViews();
+				})
+				[
+					SNew(STextBlock)
+					.Text(LOCTEXT("PoseSearchDebuggerShowVerbose", "Channels Breakdown"))
+				]
 			]
-			// Text
-			+ SHorizontalBox::Slot()
-			.Padding(FMargin(8, 0, 0, 0))
-			.AutoWidth()
-			[
-				SNew(STextBlock)
-				.Text_Lambda([this] { return ViewModel.Get()->IsPlayingSelections() ? FText::FromString("Stop Selected Asset") : FText::FromString("Play Selected Asset"); })
-				.Justification(ETextJustify::Center)
-			]
-		]
-	];
-	
-	ReturnButtonView->AddSlot()
-	.VAlign(VAlign_Center)
-	.HAlign(HAlign_Left)
-	.Padding(64, 5, 0, 0)
-	.AutoWidth()
-	[
-		SNew(SHorizontalBox)
-		
-		+ SHorizontalBox::Slot()
-		.AutoWidth()
-		.Padding(0, 5, 0, 0)
-		[
-			SNew(STextBlock)
-			.Text(FText::FromString("Asset Play Rate: "))
-		]
-		+ SHorizontalBox::Slot()
-		.AutoWidth()
-		.Padding(8, 0, 0, 0)
-		[
-			SNew(SNumericEntryBox<float>)
-			.MinValue(0.0f)
-			.MaxValue(5.0f)
-			.MinSliderValue(0.0f)
-			.MaxSliderValue(5.0f)
-			.Delta(0.01f)
-			.AllowSpin(true)
-			// Lambda to accomodate the TOptional this requires (for now)
-			.Value_Lambda([this] { return ViewModel.Get()->GetPlayRate(); })
-			.OnValueChanged(ViewModel.Get().ToSharedRef(), &FDebuggerViewModel::ChangePlayRate)	
-		]
-	];
-
-	ReturnButtonView->AddSlot()
-	.VAlign(VAlign_Center)
-	.HAlign(HAlign_Left)
-	.Padding(64, 5, 0, 0)
-	.AutoWidth()
-	[
-		SNew(SHorizontalBox)
-
-		+ SHorizontalBox::Slot()
-		.AutoWidth()
-		.Padding(0, 5, 0, 0)
-		[
-			SNew(SCheckBox)
-			.IsChecked_Lambda([this]
-			{
-				return ViewModel.Get()->IsVerbose() ? ECheckBoxState::Checked : ECheckBoxState::Unchecked; 
-			})
-			.OnCheckStateChanged_Lambda([this](ECheckBoxState State)
-			{
-				ViewModel.Get()->SetVerbose(State == ECheckBoxState::Checked);
-				UpdateViews();
-			})
-			[
-				SNew(STextBlock)
-				.Text(LOCTEXT("PoseSearchDebuggerShowVerbose", "Channels Breakdown"))
-			]
-		]
-	];
+		];
+}
 
+TSharedRef<SWidget> SDebuggerView::GenerateNodeDebuggerView()
+{
 	return 
 		SNew(SSplitter)
-		.Orientation(Orient_Horizontal)
+		.Orientation(Orient_Vertical)
 		.ResizeMode(ESplitterResizeMode::Fill)
 	
 		// Database view
@@ -721,7 +593,7 @@ TSharedRef<SWidget> SDebuggerView::GenerateNodeDebuggerView()
 			+ SVerticalBox::Slot()
 			.AutoHeight()
 			[
-				ReturnButtonView
+				GenerateReturnButtonView()
 			]
 			
 			+ SVerticalBox::Slot()
diff --git a/Engine/Plugins/Experimental/Animation/PoseSearch/Source/Editor/Private/PoseSearchDebuggerView.h b/Engine/Plugins/Experimental/Animation/PoseSearch/Source/Editor/Private/PoseSearchDebuggerView.h
--- a/Engine/Plugins/Experimental/Animation/PoseSearch/Source/Editor/Private/PoseSearchDebuggerView.h
+++ b/Engine/Plugins/Experimental/Animation/PoseSearch/Source/Editor/Private/PoseSearchDebuggerView.h
@@ -97,9 +97,6 @@ private:
 
 	void OnPoseSelectionChanged(const UPoseSearchDatabase* Database, int32 PoseIdx, float Time);
 
-	/** Button interaction to toggle play / stop of the asset */
-	FReply TogglePlaySelectedAssets() const;
-
 	/** Generates the message view relaying that there is no data */
 	TSharedRef<SWidget> GenerateNoDataMessageView();
 
diff --git a/Engine/Plugins/Experimental/Animation/PoseSearch/Source/Editor/Private/PoseSearchDebuggerViewModel.cpp b/Engine/Plugins/Experimental/Animation/PoseSearch/Source/Editor/Private/PoseSearchDebuggerViewModel.cpp
--- a/Engine/Plugins/Experimental/Animation/PoseSearch/Source/Editor/Private/PoseSearchDebuggerViewModel.cpp
+++ b/Engine/Plugins/Experimental/Animation/PoseSearch/Source/Editor/Private/PoseSearchDebuggerViewModel.cpp
@@ -200,73 +200,6 @@ void FDebuggerViewModel::UpdatePoseSearchContext(UPoseSearchMeshComponent::FUpda
 	}
 }
 
-void FDebuggerViewModel::OnDraw(FSkeletonDrawParams& DrawParams)
-{
-	const UPoseSearchDatabase* CurrentDatabase = GetCurrentDatabase();
-	if (!CurrentDatabase)
-	{
-		return;
-	}
-
-	// Returns if it is to be drawn this frame
-	auto SetDrawSkeleton = [this](UPoseSearchMeshComponent* InComponent, bool bDraw)
-	{
-		if (InComponent && InComponent->RequiredBones.IsValid())
-		{
-			const bool bIsDrawingSkeleton = InComponent->ShouldDrawDebugSkeleton();
-			if (bIsDrawingSkeleton != bDraw)
-			{
-				InComponent->SetDrawDebugSkeleton(bDraw);
-			}
-			InComponent->MarkRenderStateDirty();
-		}
-	};
-	const bool bDrawActivePose = EnumHasAnyFlags(DrawParams.Flags, ESkeletonDrawFlags::ActivePose);
-	SetDrawSkeleton(Skeletons[ActivePose].Component.Get(), bDrawActivePose);
-	// If flag is set and we are currently in a valid drawing state
-	const bool bDrawSelectedPose = EnumHasAnyFlags(DrawParams.Flags, ESkeletonDrawFlags::SelectedPose) && bSelecting;
-	SetDrawSkeleton(Skeletons[SelectedPose].Component.Get(), bDrawSelectedPose);
-
-	FillCompactPoseAndComponentRefRotations();
-
-	UPoseSearchMeshComponent::FUpdateContext UpdateContext;
-
-	UpdateContext.MirrorDataTable = CurrentDatabase->Schema->MirrorDataTable;
-	UpdateContext.CompactPoseMirrorBones = &CompactPoseMirrorBones;
-	UpdateContext.ComponentSpaceRefRotations = &ComponentSpaceRefRotations;
-
-	if (bDrawSelectedPose)
-	{
-		UPoseSearchMeshComponent* Component = Skeletons[SelectedPose].Component.Get();
-		if (Component && Component->RequiredBones.IsValid())
-		{
-			UpdatePoseSearchContext(UpdateContext, Skeletons[SelectedPose]);
-
-			if (UpdateContext.SequenceBase || UpdateContext.BlendSpace)
-			{
-				Component->UpdatePose(UpdateContext);
-			}
-		}
-	}
-
-	const bool bDrawAsset = EnumHasAnyFlags(DrawParams.Flags, ESkeletonDrawFlags::Asset);
-	if (bDrawAsset && AssetData.bActive)
-	{
-		UPoseSearchMeshComponent* Component = Skeletons[Asset].Component.Get();
-		if (Component && Component->RequiredBones.IsValid())
-		{
-			SetDrawSkeleton(Component, true);
-
-			UpdatePoseSearchContext(UpdateContext, Skeletons[SelectedPose]);
-
-			if (UpdateContext.SequenceBase || UpdateContext.BlendSpace)
-			{
-				Component->UpdatePose(UpdateContext);
-			}
-		}
-	}
-}
-
 void FDebuggerViewModel::UpdateFromTimeline()
 {
 	NodeIds.Empty();
@@ -346,96 +279,6 @@ void FDebuggerViewModel::UpdateFromTimeline()
 	});
 }
 
-void FDebuggerViewModel::UpdateAsset()
-{
-	// @todo: expose those parameters
-	static float MAX_DISTANCE_RANGE = 200.f;
-	static float MAX_TIME_RANGE = 2.f;
-
-	const UPoseSearchDatabase* Database = GetCurrentDatabase();
-	if (!Database || !IsPlayingSelections())
-	{
-		return;
-	}
-
-	FSkeleton& AssetSkeleton = Skeletons[Asset];
-	UPoseSearchMeshComponent* Component = AssetSkeleton.Component.Get();
-
-	auto RestartAsset = [&]()
-	{
-		Component->ResetToStart();
-		AssetData.AccumulatedTime = 0.0;
-		AssetSkeleton.Time = AssetData.StartTime;
-	};
-
-	const FInstancedStruct* AnimationAssetStruct = AssetSkeleton.GetAnimationAsset();
-	if (!AnimationAssetStruct)
-	{
-		checkNoEntry();
-		return;
-	}
-
-	const FPoseSearchDatabaseAnimationAssetBase* DatabaseAsset = AnimationAssetStruct->GetPtr<FPoseSearchDatabaseAnimationAssetBase>();
-	if (!DatabaseAsset)
-	{
-		checkNoEntry();
-		return;
-	}
-
-	const UAnimationAsset* AnimAsset = DatabaseAsset->GetAnimationAsset();
-	const bool bAssetLooping = DatabaseAsset->IsLooping();
-
-	const float DT = static_cast<float>(FApp::GetDeltaTime()) * AssetPlayRate;
-	const float PlayLength = AnimAsset->GetPlayLength();
-	const bool bExceededDistanceHorizon = Component->LastRootMotionDelta.GetTranslation().Size() > MAX_DISTANCE_RANGE;
-	const bool bExceededTimeHorizon = (AssetSkeleton.Time - AssetData.StartTime) > MAX_TIME_RANGE;
-	const bool bExceededHorizon = bExceededDistanceHorizon && bExceededTimeHorizon;
-	if (bAssetLooping)
-	{
-		if (bExceededHorizon)
-		{
-			// Delay before restarting the asset to give the user some idea of where it would land
-			if (AssetData.AccumulatedTime > AssetData.StopDuration)
-			{
-				RestartAsset();
-			}
-			else
-			{
-				AssetData.AccumulatedTime += DT;
-			}
-			return;
-		}
-
-		AssetSkeleton.Time += DT;
-		AssetData.AccumulatedTime += DT;
-	}
-	else
-	{
-		// Used to cap the asset, but avoid modding when updating the pose
-		static constexpr float LengthOffset = 0.001f;
-		const bool bFinishedAsset = AssetSkeleton.Time >= PlayLength - LengthOffset;
-
-		// Asset player reached end of clip or reached distance horizon of trajectory vector
-		if (bFinishedAsset || bExceededHorizon)
-		{
-			// Delay before restarting the asset to give the user some idea of where it would land
-			if (AssetData.AccumulatedTime > AssetData.StopDuration)
-			{
-				RestartAsset();
-			}
-			else
-			{
-				AssetData.AccumulatedTime += DT;
-			}
-		}
-		else
-		{
-			// If we haven't finished, update the play time capped by the anim asset (not looping)
-			AssetSkeleton.Time += DT;
-		}
-	}
-}
-
 const USkinnedMeshComponent* FDebuggerViewModel::GetMeshComponent() const
 {
 	if (Skeletons.Num() > FDebuggerViewModel::Asset)
@@ -473,45 +316,6 @@ void FDebuggerViewModel::FillCompactPoseAndComponentRefRotations()
 	}
 }
 
-void FDebuggerViewModel::PlaySelection(int32 PoseIdx, float Time)
-{
-	UPoseSearchMeshComponent* Component = Skeletons[Asset].Component.Get();
-	if (!Component)
-	{
-		return;
-	}
-
-	const UPoseSearchDatabase* Database = GetCurrentDatabase();
-	if (!FAsyncPoseSearchDatabasesManagement::RequestAsyncBuildIndex(Database, ERequestAsyncBuildFlag::ContinueRequest))
-	{
-		return;
-	}
-	
-	const FPoseSearchIndexAsset& IndexAsset = Database->GetSearchIndex().GetAssetForPose(PoseIdx);
-	Component->ResetToStart();
-	
-	Skeletons[Asset].AssetIdx = IndexAsset.SourceAssetIdx;
-	Skeletons[Asset].Time = Time;
-	Skeletons[Asset].bMirrored = IndexAsset.bMirrored;
-	Skeletons[Asset].BlendParameters = IndexAsset.BlendParameters;
-
-	AssetData.StartTime = Time;
-	AssetData.AccumulatedTime = 0.0f;
-	AssetData.bActive = true;
-}
-
-void FDebuggerViewModel::StopSelection()
-{
-	UPoseSearchMeshComponent* Component = Skeletons[Asset].Component.Get();
-	if (!Component)
-	{
-		return;
-	}
-
-	AssetData = {};
-	// @TODO: Make more functionality rely on checking if it should draw the asset
-	Component->SetDrawDebugSkeleton(false);
-}
 void FDebuggerViewModel::OnWorldCleanup(UWorld* InWorld, bool bSessionEnded, bool bCleanupResources)
 {
 	bSkeletonsInitialized = false;
diff --git a/Engine/Plugins/Experimental/Animation/PoseSearch/Source/Editor/Private/PoseSearchDebuggerViewModel.h b/Engine/Plugins/Experimental/Animation/PoseSearch/Source/Editor/Private/PoseSearchDebuggerViewModel.h
--- a/Engine/Plugins/Experimental/Animation/PoseSearch/Source/Editor/Private/PoseSearchDebuggerViewModel.h
+++ b/Engine/Plugins/Experimental/Animation/PoseSearch/Source/Editor/Private/PoseSearchDebuggerViewModel.h
@@ -17,21 +17,6 @@ namespace UE::PoseSearch
 
 struct FTraceMotionMatchingStateMessage;
 
-/** Draw flags for the view's debug draw */
-enum class ESkeletonDrawFlags : uint32
-{
-	None			= 0,
-	ActivePose		= 1 << 0,
-	SelectedPose	= 1 << 1,
-	Asset			= 1 << 2,
-};
-ENUM_CLASS_FLAGS(ESkeletonDrawFlags)
-
-struct FSkeletonDrawParams
-{
-	ESkeletonDrawFlags Flags = ESkeletonDrawFlags::None;
-};
-
 class FDebuggerViewModel : public TSharedFromThis<FDebuggerViewModel>
 {
 public:
@@ -54,30 +39,12 @@ public:
 	/** Updates active motion matching state based on node selection */
 	void OnUpdateNodeSelection(int32 InNodeId);
 
-	/** Updates internal Skeletal Mesh Component depending on input */
-	void OnDraw(FSkeletonDrawParams& DrawParams);
-	
 	/** Sets the selected pose skeleton*/
 	void ShowSelectedSkeleton(const UPoseSearchDatabase* Database, int32 DbPoseIdx, float Time);
 	
 	/** Clears the selected pose skeleton */
 	void ClearSelectedSkeleton();
 	
-	/** Plays the selected row upon button press */
-	void PlaySelection(int32 PoseIdx, float Time);
-
-	/** Stops the playing selection upon button press */
-	void StopSelection();
-	
-	/** If there is an asset selection playing */
-	bool IsPlayingSelections() const { return AssetData.bActive; }
-	
-	/** Play Rate of the asset selection player */
-	void ChangePlayRate(float PlayRate) { AssetPlayRate = PlayRate; }
-	
-	/** Current play rate of the asset selection player */
-	float GetPlayRate() const { return AssetPlayRate; }
-	
 	void SetVerbose(bool bVerbose) { bIsVerbose = bVerbose; }
 
 	bool IsVerbose() const { return bIsVerbose; }
@@ -85,9 +52,6 @@ public:
 	/** Callback to reset debug skeletons for the active world */
 	void OnWorldCleanup(UWorld* InWorld, bool bSessionEnded, bool bCleanupResources);
 
-	/** Updates the current playing asset */
-	void UpdateAsset();
-
 	const USkinnedMeshComponent* GetMeshComponent() const;
 
 private:
@@ -176,26 +140,6 @@ private:
 	/** If we currently have a selection active in the view */
 	bool bSelecting = false;
 	
-	/** Data of the active playing asset */
-	struct FAsset
-	{
-		/** How long to stop upon reaching target */
-		static constexpr float StopDuration = 2.0f;
-
-		/** Time since the start of play */
-		float AccumulatedTime = 0.0f;
-
-		/** Start time of the asset */
-		float StartTime = 0.0f;
-
-		/** If we are currently playing a asset */
-		bool bActive = false;
-
-	} AssetData;
-	
-	/** Current play rate of the asset selection player */
-	float AssetPlayRate = 1.0f;
-
 	bool bIsVerbose = false;
 
 	/** Limits some public API */
