motion matching
- database editor animation quantization fix
- pose search debugger draw query checkbox

diff --git a/Engine/Plugins/Experimental/Animation/PoseSearch/Source/Editor/Private/PoseSearchDatabasePreviewScene.cpp b/Engine/Plugins/Experimental/Animation/PoseSearch/Source/Editor/Private/PoseSearchDatabasePreviewScene.cpp
--- a/Engine/Plugins/Experimental/Animation/PoseSearch/Source/Editor/Private/PoseSearchDatabasePreviewScene.cpp
+++ b/Engine/Plugins/Experimental/Animation/PoseSearch/Source/Editor/Private/PoseSearchDatabasePreviewScene.cpp
@@ -67,7 +67,7 @@ namespace UE::PoseSearch
 				{
 					if (UDebugSkelMeshComponent* Mesh = PreviewActor.GetDebugSkelMeshComponent())
 					{
-						UE::PoseSearch::FDebugDrawParams DrawParams(GetWorld(), Mesh, Database);
+						UE::PoseSearch::FDebugDrawParams DrawParams(GetWorld(), Mesh, &PreviewActor.QuantizedTimeRootMotion, Database);
 						DrawParams.DrawFeatureVector(PreviewActor.CurrentPoseIndex);
 					}
 				}
diff --git a/Engine/Plugins/Experimental/Animation/PoseSearch/Source/Editor/Private/PoseSearchDatabaseViewModel.cpp b/Engine/Plugins/Experimental/Animation/PoseSearch/Source/Editor/Private/PoseSearchDatabaseViewModel.cpp
--- a/Engine/Plugins/Experimental/Animation/PoseSearch/Source/Editor/Private/PoseSearchDatabaseViewModel.cpp
+++ b/Engine/Plugins/Experimental/Animation/PoseSearch/Source/Editor/Private/PoseSearchDatabaseViewModel.cpp
@@ -232,26 +232,19 @@ namespace UE::PoseSearch
 			float CurrentPlayTime = PlayTime + IndexAsset.SamplingInterval.Min + PreviewActor.PlayTimeOffset;
 			FAnimationRuntime::AdvanceTime(false, CurrentPlayTime, CurrentTime, IndexAsset.SamplingInterval.Max);
 
-			// time to pose index
-			PreviewActor.CurrentPoseIndex = PoseSearchDatabase->GetPoseIndexFromTime(CurrentTime, IndexAsset);
-
-			// pose index to quantized time
-			if (PreviewActor.CurrentPoseIndex >= 0)
-			{
-				CurrentTime = PoseSearchDatabase->GetAssetTime(PreviewActor.CurrentPoseIndex);
-			}
-
-			// quantized time to scaled quantized time
-			const float CurrentScaledTime = PreviewActor.Sampler.GetScaledTime(CurrentTime);
-
-			AnimInstance->SetPosition(CurrentScaledTime);
+			AnimInstance->SetPosition(PreviewActor.Sampler.GetScaledTime(CurrentTime));
 			AnimInstance->SetPlayRate(0.f);
 			AnimInstance->SetBlendSpacePosition(IndexAsset.BlendParameters);
 
+			// time to pose index
+			PreviewActor.CurrentPoseIndex = PoseSearchDatabase->GetPoseIndexFromTime(CurrentTime, IndexAsset);
+			const float QuantizedTime = PreviewActor.CurrentPoseIndex >= 0 ? PoseSearchDatabase->GetAssetTime(PreviewActor.CurrentPoseIndex) : CurrentTime;
+			PreviewActor.QuantizedTimeRootMotion = PreviewActor.Sampler.ExtractRootTransform(QuantizedTime);
 			FTransform RootMotion = PreviewActor.Sampler.ExtractRootTransform(CurrentTime);
 			if (AnimInstance->GetMirrorDataTable())
 			{
 				RootMotion = MirrorRootMotion(RootMotion, AnimInstance->GetMirrorDataTable());
+				PreviewActor.QuantizedTimeRootMotion = MirrorRootMotion(PreviewActor.QuantizedTimeRootMotion, AnimInstance->GetMirrorDataTable());
 			}
 
 			if (PreviewActor.PlayTimeOffset != 0.f)
diff --git a/Engine/Plugins/Experimental/Animation/PoseSearch/Source/Editor/Private/PoseSearchDebuggerView.cpp b/Engine/Plugins/Experimental/Animation/PoseSearch/Source/Editor/Private/PoseSearchDebuggerView.cpp
--- a/Engine/Plugins/Experimental/Animation/PoseSearch/Source/Editor/Private/PoseSearchDebuggerView.cpp
+++ b/Engine/Plugins/Experimental/Animation/PoseSearch/Source/Editor/Private/PoseSearchDebuggerView.cpp
@@ -345,21 +345,21 @@ void SDebuggerView::DrawFeatures(
 	const USkinnedMeshComponent* Mesh
 ) const
 {
-	// @todo: have a check box like "Channels Breakdown" to disable drawing the query vector
-	// @todo: draw the query with greater thickness?
 	// Draw query vector
-	const UPoseSearchDatabase* CurrentDatabase = ViewModel.Get()->GetCurrentDatabase();
-	if (CurrentDatabase)
+	if (ViewModel.Get()->GetDrawQuery())
 	{
-		for (const FTraceMotionMatchingStateDatabaseEntry& DbEntry : State.DatabaseEntries)
+		if (const UPoseSearchDatabase* CurrentDatabase = ViewModel.Get()->GetCurrentDatabase())
 		{
-			const UPoseSearchDatabase* Database = FTraceMotionMatchingState::GetObjectFromId<UPoseSearchDatabase>(DbEntry.DatabaseId);
-			if (Database && Database == CurrentDatabase && DbEntry.QueryVector.Num() == Database->Schema->SchemaCardinality &&
-				FAsyncPoseSearchDatabasesManagement::RequestAsyncBuildIndex(CurrentDatabase, ERequestAsyncBuildFlag::ContinueRequest))
+			for (const FTraceMotionMatchingStateDatabaseEntry& DbEntry : State.DatabaseEntries)
 			{
-				FDebugDrawParams DrawParams(&DebuggerWorld, Mesh, CurrentDatabase, EDebugDrawFlags::DrawQuery);
-				DrawParams.DrawFeatureVector(DbEntry.QueryVector);
-				break;
+				const UPoseSearchDatabase* Database = FTraceMotionMatchingState::GetObjectFromId<UPoseSearchDatabase>(DbEntry.DatabaseId);
+				if (Database && Database == CurrentDatabase && DbEntry.QueryVector.Num() == Database->Schema->SchemaCardinality &&
+					FAsyncPoseSearchDatabasesManagement::RequestAsyncBuildIndex(CurrentDatabase, ERequestAsyncBuildFlag::ContinueRequest))
+				{
+					FDebugDrawParams DrawParams(&DebuggerWorld, Mesh, nullptr, CurrentDatabase, EDebugDrawFlags::DrawQuery);
+					DrawParams.DrawFeatureVector(DbEntry.QueryVector);
+					break;
+				}
 			}
 		}
 	}
@@ -374,7 +374,7 @@ void SDebuggerView::DrawFeatures(
 		const UPoseSearchDatabase* RowDatabase = Row->SourceDatabase.Get();
 		if (FAsyncPoseSearchDatabasesManagement::RequestAsyncBuildIndex(RowDatabase, ERequestAsyncBuildFlag::ContinueRequest))
 		{
-			FDebugDrawParams DrawParams(&DebuggerWorld, Mesh, RowDatabase);
+			FDebugDrawParams DrawParams(&DebuggerWorld, Mesh, nullptr, RowDatabase);
 			DrawParams.DrawFeatureVector(Row->PoseIdx);
 		}
 	}
@@ -392,7 +392,7 @@ void SDebuggerView::DrawFeatures(
 			if (Database && FAsyncPoseSearchDatabasesManagement::RequestAsyncBuildIndex(Database, ERequestAsyncBuildFlag::ContinueRequest))
 			{
 				// Use the motion-matching state's pose idx, as the active row may be update-throttled at this point
-				FDebugDrawParams DrawParams(&DebuggerWorld, Mesh, Database);
+				FDebugDrawParams DrawParams(&DebuggerWorld, Mesh, nullptr, Database);
 				DrawParams.DrawFeatureVector(ActiveRows[0]->PoseIdx);
 			}
 		}
@@ -411,7 +411,7 @@ void SDebuggerView::DrawFeatures(
 			const UPoseSearchDatabase* Database = ContinuingRows[0]->SourceDatabase.Get();
 			if (Database && FAsyncPoseSearchDatabasesManagement::RequestAsyncBuildIndex(Database, ERequestAsyncBuildFlag::ContinueRequest))
 			{
-				FDebugDrawParams DrawParams(&DebuggerWorld, Mesh, Database);
+				FDebugDrawParams DrawParams(&DebuggerWorld, Mesh, nullptr, Database);
 				DrawParams.DrawFeatureVector(ContinuingRows[0]->PoseIdx);
 			}
 		}
@@ -511,8 +511,7 @@ TSharedRef<SWidget> SDebuggerView::GenerateNoDataMessageView()
 
 TSharedRef<SHorizontalBox> SDebuggerView::GenerateReturnButtonView()
 {
-	return
-		SNew(SHorizontalBox)
+	return SNew(SHorizontalBox)
 
 		+ SHorizontalBox::Slot()
 		.Padding(10, 5, 0, 0)
@@ -562,7 +561,34 @@ TSharedRef<SHorizontalBox> SDebuggerView::GenerateReturnButtonView()
 				SNew(SCheckBox)
 				.IsChecked_Lambda([this]
 				{
-					return ViewModel.Get()->IsVerbose() ? ECheckBoxState::Checked : ECheckBoxState::Unchecked; 
+					return ViewModel.Get()->GetDrawQuery() ? ECheckBoxState::Checked : ECheckBoxState::Unchecked; 
+				})
+				.OnCheckStateChanged_Lambda([this](ECheckBoxState State)
+				{
+					ViewModel.Get()->SetDrawQuery(State == ECheckBoxState::Checked);
+				})
+				[
+					SNew(STextBlock)
+					.Text(LOCTEXT("PoseSearchDebuggerDrawQuery", "Draw Query"))
+				]
+			]
+		]
+		+ SHorizontalBox::Slot()
+		.VAlign(VAlign_Top)
+		.HAlign(HAlign_Left)
+		.Padding(64, 5, 0, 0)
+		.AutoWidth()
+		[
+			SNew(SHorizontalBox)
+
+			+ SHorizontalBox::Slot()
+			.AutoWidth()
+			.Padding(0, 5, 0, 0)
+			[
+				SNew(SCheckBox)
+				.IsChecked_Lambda([this]
+				{
+					return ViewModel.Get()->IsVerbose() ? ECheckBoxState::Checked : ECheckBoxState::Unchecked;
 				})
 				.OnCheckStateChanged_Lambda([this](ECheckBoxState State)
 				{
diff --git a/Engine/Plugins/Experimental/Animation/PoseSearch/Source/Editor/Private/PoseSearchDebuggerViewModel.h b/Engine/Plugins/Experimental/Animation/PoseSearch/Source/Editor/Private/PoseSearchDebuggerViewModel.h
--- a/Engine/Plugins/Experimental/Animation/PoseSearch/Source/Editor/Private/PoseSearchDebuggerViewModel.h
+++ b/Engine/Plugins/Experimental/Animation/PoseSearch/Source/Editor/Private/PoseSearchDebuggerViewModel.h
@@ -49,6 +49,12 @@ public:
 
 	bool IsVerbose() const { return bIsVerbose; }
 
+	void SetDrawQuery(bool bInDrawQuery) { bDrawQuery = bInDrawQuery; }
+
+	bool GetDrawQuery() const { return bDrawQuery; }
+
+
+
 	/** Callback to reset debug skeletons for the active world */
 	void OnWorldCleanup(UWorld* InWorld, bool bSessionEnded, bool bCleanupResources);
 
@@ -142,6 +148,8 @@ private:
 	
 	bool bIsVerbose = false;
 
+	bool bDrawQuery = true;	
+
 	/** Limits some public API */
 	friend class FDebugger;
 };
diff --git a/Engine/Plugins/Experimental/Animation/PoseSearch/Source/Editor/Public/PoseSearchDatabaseViewModel.h b/Engine/Plugins/Experimental/Animation/PoseSearch/Source/Editor/Public/PoseSearchDatabaseViewModel.h
--- a/Engine/Plugins/Experimental/Animation/PoseSearch/Source/Editor/Public/PoseSearchDatabaseViewModel.h
+++ b/Engine/Plugins/Experimental/Animation/PoseSearch/Source/Editor/Public/PoseSearchDatabaseViewModel.h
@@ -47,6 +47,7 @@ namespace UE::PoseSearch
 		int32 CurrentPoseIndex = INDEX_NONE;
 		float PlayTimeOffset = 0.f;
 		FAnimationAssetSampler Sampler;
+		FTransform QuantizedTimeRootMotion = FTransform::Identity;
 
 		bool IsValid() const;
 		void Process();
diff --git a/Engine/Plugins/Experimental/Animation/PoseSearch/Source/Runtime/Private/PoseSearchContext.cpp b/Engine/Plugins/Experimental/Animation/PoseSearch/Source/Runtime/Private/PoseSearchContext.cpp
--- a/Engine/Plugins/Experimental/Animation/PoseSearch/Source/Runtime/Private/PoseSearchContext.cpp
+++ b/Engine/Plugins/Experimental/Animation/PoseSearch/Source/Runtime/Private/PoseSearchContext.cpp
@@ -19,15 +19,17 @@ FDebugDrawParams::FDebugDrawParams(FAnimInstanceProxy* InAnimInstanceProxy, cons
 : AnimInstanceProxy(InAnimInstanceProxy)
 , World(nullptr)
 , Mesh(nullptr)
+, RootMotionTransform(nullptr)
 , Database(InDatabase)
 , Flags(InFlags)
 {
 }
 
-FDebugDrawParams::FDebugDrawParams(const UWorld* InWorld, const USkinnedMeshComponent* InMesh, const UPoseSearchDatabase* InDatabase, EDebugDrawFlags InFlags)
+FDebugDrawParams::FDebugDrawParams(const UWorld* InWorld, const USkinnedMeshComponent* InMesh, const FTransform* InRootMotionTransform, const UPoseSearchDatabase* InDatabase, EDebugDrawFlags InFlags)
 : AnimInstanceProxy(nullptr)
 , World(InWorld)
 , Mesh(InMesh)
+, RootMotionTransform(InRootMotionTransform)
 , Database(InDatabase)
 , Flags(InFlags)
 {
@@ -83,7 +85,7 @@ FVector FDebugDrawParams::ExtractPosition(TConstArrayView<float> PoseVector, flo
 			return ExtractPosition(PoseVector, FoundPosition);
 		}
 
-		if (Mesh && SchemaBoneIdx >= 0)
+		if (Mesh && SchemaBoneIdx > RootSchemaBoneIdx)
 		{
 			return Mesh->GetSocketTransform(Schema->BoneReferences[SchemaBoneIdx].BoneName).GetTranslation();
 		}
@@ -93,6 +95,11 @@ FVector FDebugDrawParams::ExtractPosition(TConstArrayView<float> PoseVector, flo
 
 const FTransform& FDebugDrawParams::GetRootTransform() const
 {
+	if (RootMotionTransform)
+	{
+		return *RootMotionTransform;
+	}
+
 	if (AnimInstanceProxy)
 	{
 		return AnimInstanceProxy->GetComponentTransform();
diff --git a/Engine/Plugins/Experimental/Animation/PoseSearch/Source/Runtime/Public/PoseSearch/PoseSearchContext.h b/Engine/Plugins/Experimental/Animation/PoseSearch/Source/Runtime/Public/PoseSearch/PoseSearchContext.h
--- a/Engine/Plugins/Experimental/Animation/PoseSearch/Source/Runtime/Public/PoseSearch/PoseSearchContext.h
+++ b/Engine/Plugins/Experimental/Animation/PoseSearch/Source/Runtime/Public/PoseSearch/PoseSearchContext.h
@@ -106,7 +106,7 @@ private:
 struct POSESEARCH_API FDebugDrawParams
 {
 	FDebugDrawParams(FAnimInstanceProxy* InAnimInstanceProxy, const UPoseSearchDatabase* InDatabase, EDebugDrawFlags InFlags = EDebugDrawFlags::None);
-	FDebugDrawParams(const UWorld* InWorld, const USkinnedMeshComponent* InMesh, const UPoseSearchDatabase* InDatabase, EDebugDrawFlags InFlags = EDebugDrawFlags::None);
+	FDebugDrawParams(const UWorld* InWorld, const USkinnedMeshComponent* InMesh, const FTransform* InRootMotionTransform, const UPoseSearchDatabase* InDatabase, EDebugDrawFlags InFlags = EDebugDrawFlags::None);
 
 	const FPoseSearchIndex* GetSearchIndex() const;
 	const UPoseSearchSchema* GetSchema() const;
@@ -129,6 +129,7 @@ private:
 	FAnimInstanceProxy* AnimInstanceProxy = nullptr;
 	const UWorld* World = nullptr;
 	const USkinnedMeshComponent* Mesh = nullptr;
+	const FTransform* RootMotionTransform = nullptr;
 	const UPoseSearchDatabase* Database = nullptr;
 	EDebugDrawFlags Flags = EDebugDrawFlags::None;
 };
