motion matching
- adding posesearch debugger horizontal scrollbar
- broken down pose search debugger FFrame into FFrame, FTime, FPercentage

diff --git a/Engine/Plugins/Experimental/Animation/PoseSearch/Source/Editor/Private/PoseSearchDebuggerDatabaseColumns.h b/Engine/Plugins/Experimental/Animation/PoseSearch/Source/Editor/Private/PoseSearchDebuggerDatabaseColumns.h
--- a/Engine/Plugins/Experimental/Animation/PoseSearch/Source/Editor/Private/PoseSearchDebuggerDatabaseColumns.h
+++ b/Engine/Plugins/Experimental/Animation/PoseSearch/Source/Editor/Private/PoseSearchDebuggerDatabaseColumns.h
@@ -245,6 +245,26 @@ struct FFrame : ITextColumn
 		return LOCTEXT("ColumnLabelFrame", "Frame");
 	}
 		
+	virtual FSortPredicate GetSortPredicate() const override
+	{
+		return [](const FRowDataRef& Row0, const FRowDataRef& Row1) -> bool { return Row0->AnimFrame < Row1->AnimFrame; };
+	}
+
+	virtual FText GetRowText(const FRowDataRef& Row) const override
+	{
+		return FText::AsNumber(Row->AnimFrame, &FNumberFormattingOptions::DefaultNoGrouping());
+	}
+};
+
+struct FTime : ITextColumn
+{
+	using ITextColumn::ITextColumn;
+
+	virtual FText GetLabel() const override
+	{
+		return LOCTEXT("ColumnLabelTime", "Time");
+	}
+
 	virtual FSortPredicate GetSortPredicate() const override
 	{
 		return [](const FRowDataRef& Row0, const FRowDataRef& Row1) -> bool { return Row0->AssetTime < Row1->AssetTime; };
@@ -252,22 +272,27 @@ struct FFrame : ITextColumn
 
 	virtual FText GetRowText(const FRowDataRef& Row) const override
 	{
-		FNumberFormattingOptions TimeFormattingOptions = FNumberFormattingOptions().SetUseGrouping(false).SetMaximumFractionalDigits(2);
-		FNumberFormattingOptions PercentageFormattingOptions = FNumberFormattingOptions().SetMaximumFractionalDigits(2);
-		
-		if (GetDefault<UPersonaOptions>()->bTimelineDisplayPercentage)
-		{
-			return FText::Format(FText::FromString("{0} ({1}) ({2})"),
-				FText::AsNumber(Row->AnimFrame, &FNumberFormattingOptions::DefaultNoGrouping()),
-				FText::AsNumber(Row->AssetTime, &TimeFormattingOptions),
-				FText::AsPercent(Row->AnimPercentage, &PercentageFormattingOptions));
-		}
-		else
-		{
-			return FText::Format(FText::FromString("{0} ({1})"),
-				FText::AsNumber(Row->AnimFrame, &FNumberFormattingOptions::DefaultNoGrouping()),
-				FText::AsNumber(Row->AssetTime, &TimeFormattingOptions));
-		}
+		return FText::AsNumber(Row->AssetTime, &FNumberFormattingOptions().SetUseGrouping(false).SetMaximumFractionalDigits(2));
+	}
+};
+
+struct FPercentage : ITextColumn
+{
+	using ITextColumn::ITextColumn;
+
+	virtual FText GetLabel() const override
+	{
+		return LOCTEXT("ColumnLabelPercentage", "Percentage");
+	}
+
+	virtual FSortPredicate GetSortPredicate() const override
+	{
+		return [](const FRowDataRef& Row0, const FRowDataRef& Row1) -> bool { return Row0->AnimPercentage < Row1->AnimPercentage; };
+	}
+
+	virtual FText GetRowText(const FRowDataRef& Row) const override
+	{
+		return FText::AsPercent(Row->AnimPercentage, &FNumberFormattingOptions().SetMaximumFractionalDigits(2));
 	}
 };
 
@@ -462,25 +487,25 @@ struct FPoseCandidateFlags : ITextColumn
 			if (EnumHasAnyFlags(Row->PoseCandidateFlags, EPoseCandidateFlags::DiscardedBy_PoseJumpThresholdTime))
 			{
 				AddDelimiter(TextBuilder, bNeedDelimiter);
-				TextBuilder.AppendLine(LOCTEXT("DiscardedBy_PoseJumpThresholdTime", "PoseJumpThresholdTime"));
+				TextBuilder.AppendLine(LOCTEXT("DiscardedBy_PoseJumpThresholdTime", "JumpThresh"));
 			}
 				
 			if (EnumHasAnyFlags(Row->PoseCandidateFlags, EPoseCandidateFlags::DiscardedBy_PoseReselectHistory))
 			{
 				AddDelimiter(TextBuilder, bNeedDelimiter);
-				TextBuilder.AppendLine(LOCTEXT("DiscardedBy_PoseReselectHistory", "PoseReselectHistory"));
+				TextBuilder.AppendLine(LOCTEXT("DiscardedBy_PoseReselectHistory", "History"));
 			}
 
 			if (EnumHasAnyFlags(Row->PoseCandidateFlags, EPoseCandidateFlags::DiscardedBy_BlockTransition))
 			{
 				AddDelimiter(TextBuilder, bNeedDelimiter);
-				TextBuilder.AppendLine(LOCTEXT("DiscardedBy_BlockTransition", "BlockTransition"));
+				TextBuilder.AppendLine(LOCTEXT("DiscardedBy_BlockTransition", "BlockTrans"));
 			}
 
 			if (EnumHasAnyFlags(Row->PoseCandidateFlags, EPoseCandidateFlags::DiscardedBy_PoseFilter))
 			{
 				AddDelimiter(TextBuilder, bNeedDelimiter);
-				TextBuilder.AppendLine(LOCTEXT("DiscardedBy_PoseFilter", "PoseFilter"));
+				TextBuilder.AppendLine(LOCTEXT("DiscardedBy_PoseFilter", "Filter"));
 			}
 
 			return TextBuilder.ToText();
diff --git a/Engine/Plugins/Experimental/Animation/PoseSearch/Source/Editor/Private/PoseSearchDebuggerDatabaseRow.h b/Engine/Plugins/Experimental/Animation/PoseSearch/Source/Editor/Private/PoseSearchDebuggerDatabaseRow.h
--- a/Engine/Plugins/Experimental/Animation/PoseSearch/Source/Editor/Private/PoseSearchDebuggerDatabaseRow.h
+++ b/Engine/Plugins/Experimental/Animation/PoseSearch/Source/Editor/Private/PoseSearchDebuggerDatabaseRow.h
@@ -48,7 +48,6 @@ class SDebuggerDatabaseRow : public SMultiColumnTableRow<TSharedRef<FDebuggerDat
 		// Get column
 		const TSharedRef<DebuggerDatabaseColumns::IColumn>& Column = (*ColumnMap.Get())[InColumnName];
 		
-		static FSlateFontInfo NormalFont = FAppStyle::Get().GetFontStyle("DetailsView.CategoryTextStyle");
 		const TSharedRef<SWidget> Widget = Column->GenerateWidget(Row.ToSharedRef());
 		
 		return
diff --git a/Engine/Plugins/Experimental/Animation/PoseSearch/Source/Editor/Private/PoseSearchDebuggerDatabaseView.cpp b/Engine/Plugins/Experimental/Animation/PoseSearch/Source/Editor/Private/PoseSearchDebuggerDatabaseView.cpp
--- a/Engine/Plugins/Experimental/Animation/PoseSearch/Source/Editor/Private/PoseSearchDebuggerDatabaseView.cpp
+++ b/Engine/Plugins/Experimental/Animation/PoseSearch/Source/Editor/Private/PoseSearchDebuggerDatabaseView.cpp
@@ -17,6 +17,7 @@
 #include "Trace/PoseSearchTraceProvider.h"
 #include "Widgets/Input/SCheckBox.h"
 #include "Widgets/Layout/SScrollBar.h"
+#include "Widgets/Layout/SScrollBox.h"
 #include "Widgets/Views/SListView.h"
 
 #define LOCTEXT_NAMESPACE "PoseSearchDebugger"
@@ -352,6 +353,8 @@ void SDebuggerDatabaseView::Update(const FTraceMotionMatchingStateMessage& State
 
 		AddColumn(MakeShared<FCostModifier>(ColumnIdx++));
 		AddColumn(MakeShared<FFrame>(ColumnIdx++));
+		AddColumn(MakeShared<FTime>(ColumnIdx++));
+		AddColumn(MakeShared<FPercentage>(ColumnIdx++));
 		AddColumn(MakeShared<FMirrored>(ColumnIdx++));
 		AddColumn(MakeShared<FLooping>(ColumnIdx++));
 		AddColumn(MakeShared<FPoseIdx>(ColumnIdx++));
@@ -393,11 +396,10 @@ void SDebuggerDatabaseView::Update(const FTraceMotionMatchingStateMessage& State
 					.HAlignCell(HAlign_Fill);
 
 				FilteredDatabaseView.HeaderRow->AddColumn(ColumnArgs);
+				ContinuingPoseView.HeaderRow->AddColumn(ColumnArgs);
 
 				// Every time the active column is changed, update the database column
 				ActiveView.HeaderRow->AddColumn(ColumnArgs.OnWidthChanged(this, &SDebuggerDatabaseView::OnColumnWidthChanged, Column.ColumnId));
-
-				ContinuingPoseView.HeaderRow->AddColumn(ColumnArgs.OnWidthChanged(this, &SDebuggerDatabaseView::OnColumnWidthChanged, Column.ColumnId));
 			}
 		}
 	}
@@ -654,230 +656,237 @@ void SDebuggerDatabaseView::Construct(const FArguments& InArgs)
 
 	ChildSlot
 	[
-		SNew(SVerticalBox)
-		+ SVerticalBox::Slot()
-		// Side and top margins, ignore bottom handled by the color border below
-		.Padding(0.0f, 5.0f, 0.0f, 0.0f)
-		.AutoHeight()
+		SNew(SScrollBox)
+		.Orientation(Orient_Horizontal)
+		.ScrollBarAlwaysVisible(true)
+		+ SScrollBox::Slot()
+		.FillSize(1.f)
 		[
-			// Active Row text tab
 			SNew(SVerticalBox)
 			+ SVerticalBox::Slot()
+			// Side and top margins, ignore bottom handled by the color border below
+			.Padding(0.0f, 5.0f, 0.0f, 0.0f)
 			.AutoHeight()
-			.Padding(0.0f)
 			[
-				SNew(SHorizontalBox)
-				+ SHorizontalBox::Slot()
-				.HAlign(HAlign_Center)
-				.VAlign(VAlign_Fill)
+				// Active Row text tab
+				SNew(SVerticalBox)
+				+ SVerticalBox::Slot()
+				.AutoHeight()
 				.Padding(0.0f)
-				.AutoWidth()
 				[
-					SNew(SBorder)
-					.BorderImage(FAppStyle::GetBrush("DetailsView.CategoryTop"))
-					.Padding(FMargin(30.0f, 3.0f, 30.0f, 0.0f))
+					SNew(SHorizontalBox)
+					+ SHorizontalBox::Slot()
 					.HAlign(HAlign_Center)
 					.VAlign(VAlign_Fill)
+					.Padding(0.0f)
+					.AutoWidth()
 					[
-						SNew(STextBlock)
-						.Text(FText::FromString("Active Pose"))	
+						SNew(SBorder)
+						.BorderImage(FAppStyle::GetBrush("DetailsView.CategoryTop"))
+						.Padding(FMargin(30.0f, 3.0f, 30.0f, 0.0f))
+						.HAlign(HAlign_Center)
+						.VAlign(VAlign_Fill)
+						[
+							SNew(STextBlock)
+							.Text(FText::FromString("Active Pose"))	
+						]
 					]
 				]
-			]
 
-			// Active row list view with scroll bar
-			+ SVerticalBox::Slot()
+				// Active row list view with scroll bar
+				+ SVerticalBox::Slot()
 			
-			.AutoHeight()
-			[
-				SNew(SHorizontalBox)
-				
-				+ SHorizontalBox::Slot()
-				.HAlign(HAlign_Fill)
-				.VAlign(VAlign_Fill)
-				.Padding(0.0f)
+				.AutoHeight()
 				[
-					SNew(SBorder)
-					
-					.BorderImage(FAppStyle::GetBrush("NoBorder"))
+					SNew(SHorizontalBox)
+				
+					+ SHorizontalBox::Slot()
+					.HAlign(HAlign_Fill)
+					.VAlign(VAlign_Fill)
 					.Padding(0.0f)
 					[
-						ActiveView.ListView.ToSharedRef()
+						SNew(SBorder)
+					
+						.BorderImage(FAppStyle::GetBrush("NoBorder"))
+						.Padding(0.0f)
+						[
+							ActiveView.ListView.ToSharedRef()
+						]
 					]
-				]
 
-				+ SHorizontalBox::Slot()
-				.AutoWidth()
-				[
-					ActiveView.ScrollBar.ToSharedRef()
-				]
-			]	
-		]
+					+ SHorizontalBox::Slot()
+					.AutoWidth()
+					[
+						ActiveView.ScrollBar.ToSharedRef()
+					]
+				]	
+			]
 
-		+ SVerticalBox::Slot()
-		// Side and top margins, ignore bottom handled by the color border below
-		.Padding(0.0f, 5.0f, 0.0f, 0.0f)
-		.AutoHeight()
-		[
-			// ContinuingPose Row text tab
-			SNew(SVerticalBox)
 			+ SVerticalBox::Slot()
+			// Side and top margins, ignore bottom handled by the color border below
+			.Padding(0.0f, 5.0f, 0.0f, 0.0f)
 			.AutoHeight()
-			.Padding(0.0f)
 			[
-				SNew(SHorizontalBox)
-				+ SHorizontalBox::Slot()
-				.HAlign(HAlign_Center)
-				.VAlign(VAlign_Fill)
+				// ContinuingPose Row text tab
+				SNew(SVerticalBox)
+				+ SVerticalBox::Slot()
+				.AutoHeight()
 				.Padding(0.0f)
-				.AutoWidth()
 				[
-					SNew(SBorder)
-					.BorderImage(FAppStyle::GetBrush("DetailsView.CategoryTop"))
-					.Padding(FMargin(30.0f, 3.0f, 30.0f, 0.0f))
+					SNew(SHorizontalBox)
+					+ SHorizontalBox::Slot()
 					.HAlign(HAlign_Center)
 					.VAlign(VAlign_Fill)
+					.Padding(0.0f)
+					.AutoWidth()
 					[
-						SNew(STextBlock)
-						.Text(FText::FromString("Continuing Pose"))	
+						SNew(SBorder)
+						.BorderImage(FAppStyle::GetBrush("DetailsView.CategoryTop"))
+						.Padding(FMargin(30.0f, 3.0f, 30.0f, 0.0f))
+						.HAlign(HAlign_Center)
+						.VAlign(VAlign_Fill)
+						[
+							SNew(STextBlock)
+							.Text(FText::FromString("Continuing Pose"))	
+						]
 					]
 				]
-			]
 
-			// ContinuingPose row list view with scroll bar
-			+ SVerticalBox::Slot()
+				// ContinuingPose row list view with scroll bar
+				+ SVerticalBox::Slot()
 			
-			.AutoHeight()
-			[
-				SNew(SHorizontalBox)
-				
-				+ SHorizontalBox::Slot()
-				.HAlign(HAlign_Fill)
-				.VAlign(VAlign_Fill)
-				.Padding(0.0f)
+				.AutoHeight()
 				[
-					SNew(SBorder)
-					
-					.BorderImage(FAppStyle::GetBrush("NoBorder"))
+					SNew(SHorizontalBox)
+				
+					+ SHorizontalBox::Slot()
+					.HAlign(HAlign_Fill)
+					.VAlign(VAlign_Fill)
 					.Padding(0.0f)
 					[
-						ContinuingPoseView.ListView.ToSharedRef()
+						SNew(SBorder)
+					
+						.BorderImage(FAppStyle::GetBrush("NoBorder"))
+						.Padding(0.0f)
+						[
+							ContinuingPoseView.ListView.ToSharedRef()
+						]
 					]
-				]
 
-				+ SHorizontalBox::Slot()
-				.AutoWidth()
-				[
-					ContinuingPoseView.ScrollBar.ToSharedRef()
-				]
-			]	
-		]
+					+ SHorizontalBox::Slot()
+					.AutoWidth()
+					[
+						ContinuingPoseView.ScrollBar.ToSharedRef()
+					]
+				]	
+			]
 		
-		+ SVerticalBox::Slot()
-		.Padding(0.0f, 0.0f, 0.0f, 5.0f)
-		[
-			// Database view text tab
-			SNew(SVerticalBox)
 			+ SVerticalBox::Slot()
+			.Padding(0.0f, 0.0f, 0.0f, 5.0f)
 			[
-				SNew(SHorizontalBox)
-				+ SHorizontalBox::Slot()
-				.HAlign(HAlign_Center)
-				.VAlign(VAlign_Fill)
-				.Padding(0.0f)
+				// Database view text tab
+				SNew(SVerticalBox)
+				+ SVerticalBox::Slot()
 				[
-					SNew(SBorder)
-					.BorderImage(FAppStyle::GetBrush("DetailsView.CategoryTop"))
-					.Padding(FMargin(30.0f, 3.0f, 30.0f, 0.0f))
+					SNew(SHorizontalBox)
+					+ SHorizontalBox::Slot()
 					.HAlign(HAlign_Center)
 					.VAlign(VAlign_Fill)
+					.Padding(0.0f)
 					[
-						SNew(STextBlock)
-						.Text(FText::FromString("Pose Candidates"))
+						SNew(SBorder)
+						.BorderImage(FAppStyle::GetBrush("DetailsView.CategoryTop"))
+						.Padding(FMargin(30.0f, 3.0f, 30.0f, 0.0f))
+						.HAlign(HAlign_Center)
+						.VAlign(VAlign_Fill)
+						[
+							SNew(STextBlock)
+							.Text(FText::FromString("Pose Candidates"))
+						]
 					]
-				]
-				.AutoWidth()
+					.AutoWidth()
 				
-				+ SHorizontalBox::Slot()
-				.HAlign(HAlign_Fill)
-				[
-					SNew(SBorder)
-					.BorderImage(&FilteredDatabaseView.RowStyle.EvenRowBackgroundBrush)
+					+ SHorizontalBox::Slot()
+					.HAlign(HAlign_Fill)
+					[
+						SNew(SBorder)
+						.BorderImage(&FilteredDatabaseView.RowStyle.EvenRowBackgroundBrush)
+					]
 				]
-			]
-			.AutoHeight()
+				.AutoHeight()
 
-			// Gray line below the tab 
-			+ SVerticalBox::Slot()
-			.AutoHeight()
-			.Padding(0.0f)
-			[
-				SNew(SBorder)
-				.BorderImage(FAppStyle::GetBrush("DetailsView.CategoryTop"))
-				.Padding(FMargin(0.0f, 3.0f, 0.0f, 3.0f))
-				.HAlign(HAlign_Fill)
-				.VAlign(VAlign_Fill)
-			]
-
-			+ SVerticalBox::Slot()
-			.Padding(0.0f, 0.0f, 0.0f, 5.0f)
-			.AutoHeight()
-			[
-				SNew(SHorizontalBox)
-				+ SHorizontalBox::Slot()
-				.Padding(10, 5, 10, 5)
+				// Gray line below the tab 
+				+ SVerticalBox::Slot()
+				.AutoHeight()
+				.Padding(0.0f)
 				[
-					SAssignNew(FilterBox, SSearchBox)
-					.OnTextChanged(this, &SDebuggerDatabaseView::OnFilterTextChanged)
+					SNew(SBorder)
+					.BorderImage(FAppStyle::GetBrush("DetailsView.CategoryTop"))
+					.Padding(FMargin(0.0f, 3.0f, 0.0f, 3.0f))
+					.HAlign(HAlign_Fill)
+					.VAlign(VAlign_Fill)
 				]
-				+ SHorizontalBox::Slot()
-				.AutoWidth()
-				.Padding(10, 5, 10, 5)
+
+				+ SVerticalBox::Slot()
+				.Padding(0.0f, 0.0f, 0.0f, 5.0f)
+				.AutoHeight()
 				[
-					SNew(SCheckBox)
-					.OnCheckStateChanged_Lambda([this](ECheckBoxState State)
-					{
-						   SDebuggerDatabaseView::OnHideInvalidPosesCheckboxChanged(State);
-					})
+					SNew(SHorizontalBox)
+					+ SHorizontalBox::Slot()
+					.Padding(10, 5, 10, 5)
 					[
-						SNew(STextBlock)
-						.Text(LOCTEXT("PoseSearchDebuggerHideInvalidPosesFlag", "Hide Invalid Poses"))
+						SAssignNew(FilterBox, SSearchBox)
+						.OnTextChanged(this, &SDebuggerDatabaseView::OnFilterTextChanged)
 					]
-				]
-				+ SHorizontalBox::Slot()
-				.AutoWidth()
-				.Padding(10, 5, 10, 5)
-				[
-					SNew(SCheckBox)
-					.OnCheckStateChanged_Lambda([this](ECheckBoxState State)
+					+ SHorizontalBox::Slot()
+					.AutoWidth()
+					.Padding(10, 5, 10, 5)
+					[
+						SNew(SCheckBox)
+						.OnCheckStateChanged_Lambda([this](ECheckBoxState State)
 						{
-							SDebuggerDatabaseView::OnUseRegexCheckboxChanged(State);
+							   SDebuggerDatabaseView::OnHideInvalidPosesCheckboxChanged(State);
 						})
-				[
-					SNew(STextBlock)
-					.Text(LOCTEXT("PoseSearchDebuggerUseRegexFlag", "Use Regex"))
-				]
+						[
+							SNew(STextBlock)
+							.Text(LOCTEXT("PoseSearchDebuggerHideInvalidPosesFlag", "Hide Invalid Poses"))
+						]
+					]
+					+ SHorizontalBox::Slot()
+					.AutoWidth()
+					.Padding(10, 5, 10, 5)
+					[
+						SNew(SCheckBox)
+						.OnCheckStateChanged_Lambda([this](ECheckBoxState State)
+							{
+								SDebuggerDatabaseView::OnUseRegexCheckboxChanged(State);
+							})
+					[
+						SNew(STextBlock)
+						.Text(LOCTEXT("PoseSearchDebuggerUseRegexFlag", "Use Regex"))
+					]
+					]
 				]
-			]
 		
-			+ SVerticalBox::Slot()
-			[
-				SNew(SHorizontalBox)
-				+ SHorizontalBox::Slot()
-				.Padding(0.0f)
+				+ SVerticalBox::Slot()
 				[
-					SNew(SBorder)
-					.BorderImage(FAppStyle::GetBrush("NoBorder"))
+					SNew(SHorizontalBox)
+					+ SHorizontalBox::Slot()
 					.Padding(0.0f)
 					[
-						FilteredDatabaseView.ListView.ToSharedRef()
+						SNew(SBorder)
+						.BorderImage(FAppStyle::GetBrush("NoBorder"))
+						.Padding(0.0f)
+						[
+							FilteredDatabaseView.ListView.ToSharedRef()
+						]
 					]
-				]
 				
-				+ SHorizontalBox::Slot()
-				.AutoWidth()
-				[
-					FilteredDatabaseView.ScrollBar.ToSharedRef()
+					+ SHorizontalBox::Slot()
+					.AutoWidth()
+					[
+						FilteredDatabaseView.ScrollBar.ToSharedRef()
+					]
 				]
 			]
 		]
