motion matching - fix for database debug lines drawing when !FDatabaseViewModel::bIsEditorSelection

diff --git a/Engine/Plugins/Experimental/Animation/PoseSearch/Source/Editor/Private/PoseSearchDatabaseViewModel.cpp b/Engine/Plugins/Experimental/Animation/PoseSearch/Source/Editor/Private/PoseSearchDatabaseViewModel.cpp
--- a/Engine/Plugins/Experimental/Animation/PoseSearch/Source/Editor/Private/PoseSearchDatabaseViewModel.cpp
+++ b/Engine/Plugins/Experimental/Animation/PoseSearch/Source/Editor/Private/PoseSearchDatabaseViewModel.cpp
@@ -255,6 +255,7 @@ namespace UE::PoseSearch
 					OriginRootMotion = MirrorRootMotion(OriginRootMotion, AnimInstance->GetMirrorDataTable());
 				}
 				RootMotion.SetToRelativeTransform(OriginRootMotion);
+				PreviewActor.QuantizedTimeRootMotion.SetToRelativeTransform(OriginRootMotion);
 			}
 
 			PreviewActor.Actor->SetActorTransform(RootMotion);
