motion matching
- avoiding PopulateNonSelectableIdx with duplicates poses
- consolidating pose search debugger duplicate poses reporting by merging flags

diff --git a/Engine/Plugins/Experimental/Animation/PoseSearch/Source/Editor/Private/PoseSearchDebuggerDatabaseColumns.h b/Engine/Plugins/Experimental/Animation/PoseSearch/Source/Editor/Private/PoseSearchDebuggerDatabaseColumns.h
--- a/Engine/Plugins/Experimental/Animation/PoseSearch/Source/Editor/Private/PoseSearchDebuggerDatabaseColumns.h
+++ b/Engine/Plugins/Experimental/Animation/PoseSearch/Source/Editor/Private/PoseSearchDebuggerDatabaseColumns.h
@@ -69,12 +69,17 @@ struct ITextColumn : IColumn
 		return SNew(STextBlock)
             .Font(RowFont)
 			.Text_Lambda([this, RowData]() -> FText { return GetRowText(RowData); })
+			.ToolTipText_Lambda([this, RowData]() -> FText { return GetRowToolTipText(RowData); })
             .Justification(ETextJustify::Center)
 			.ColorAndOpacity_Lambda([this, RowData] { return GetColorAndOpacity(RowData); });
 	}
 
 protected:
 	virtual FText GetRowText(const FRowDataRef& Row) const = 0;
+	virtual FText GetRowToolTipText(const FRowDataRef& Row) const
+	{
+		return FText::GetEmpty();
+	}
 	virtual FSlateColor GetColorAndOpacity(const FRowDataRef& Row) const
 	{
 		return FSlateColor(FLinearColor::White);
@@ -472,40 +477,59 @@ struct FPoseCandidateFlags : ITextColumn
 	{
 		if (EnumHasAnyFlags(Row->PoseCandidateFlags, EPoseCandidateFlags::AnyDiscardedMask))
 		{
-			auto AddDelimiter = [](FTextBuilder& TextBuilder, bool& bNeedDelimiter)
+			FString Sring;
+
+			if (EnumHasAnyFlags(Row->PoseCandidateFlags, EPoseCandidateFlags::DiscardedBy_PoseJumpThresholdTime))
 			{
-				if (bNeedDelimiter)
-				{
-					TextBuilder.AppendLine(LOCTEXT("DiscardedBy_Delimiter", " | "));
-				}
-				bNeedDelimiter = true;
-			};
+				Sring.Append("J ");
+			}
+
+			if (EnumHasAnyFlags(Row->PoseCandidateFlags, EPoseCandidateFlags::DiscardedBy_PoseReselectHistory))
+			{
+				Sring.Append("H ");
+			}
+			
+			if (EnumHasAnyFlags(Row->PoseCandidateFlags, EPoseCandidateFlags::DiscardedBy_BlockTransition))
+			{
+				Sring.Append("B ");
+			}
+			
+			if (EnumHasAnyFlags(Row->PoseCandidateFlags, EPoseCandidateFlags::DiscardedBy_PoseFilter))
+			{
+				Sring.Append("F ");
+			}
+
+			return FText::FromString(Sring);
+		}
 
-			bool bNeedDelimiter = false;
+		return FText::GetEmpty();
+	}
 
+	virtual FText GetRowToolTipText(const FRowDataRef& Row) const override
+	{
+		if (EnumHasAnyFlags(Row->PoseCandidateFlags, EPoseCandidateFlags::AnyDiscardedMask))
+		{
 			FTextBuilder TextBuilder;
+			TextBuilder.AppendLine(LOCTEXT("DiscardedBy_Reason_Tooltip", "Pose discarded because of:"));
+
 			if (EnumHasAnyFlags(Row->PoseCandidateFlags, EPoseCandidateFlags::DiscardedBy_PoseJumpThresholdTime))
 			{
-				AddDelimiter(TextBuilder, bNeedDelimiter);
-				TextBuilder.AppendLine(LOCTEXT("DiscardedBy_PoseJumpThresholdTime", "JumpThresh"));
+				TextBuilder.AppendLine(LOCTEXT("DiscardedBy_PoseJumpThresholdTime_Tooltip", "(J) Pose Jump Threshold Time"));
 			}
-				
+
 			if (EnumHasAnyFlags(Row->PoseCandidateFlags, EPoseCandidateFlags::DiscardedBy_PoseReselectHistory))
 			{
-				AddDelimiter(TextBuilder, bNeedDelimiter);
-				TextBuilder.AppendLine(LOCTEXT("DiscardedBy_PoseReselectHistory", "History"));
+				TextBuilder.AppendLine(LOCTEXT("DiscardedBy_PoseReselectHistory_Tooltip", "(H) Pose Reselect History"));
 			}
 
 			if (EnumHasAnyFlags(Row->PoseCandidateFlags, EPoseCandidateFlags::DiscardedBy_BlockTransition))
 			{
-				AddDelimiter(TextBuilder, bNeedDelimiter);
-				TextBuilder.AppendLine(LOCTEXT("DiscardedBy_BlockTransition", "BlockTrans"));
+				TextBuilder.AppendLine(LOCTEXT("DiscardedBy_BlockTransition_Tooltip", "(B) Block Transition"));
 			}
 
 			if (EnumHasAnyFlags(Row->PoseCandidateFlags, EPoseCandidateFlags::DiscardedBy_PoseFilter))
 			{
-				AddDelimiter(TextBuilder, bNeedDelimiter);
-				TextBuilder.AppendLine(LOCTEXT("DiscardedBy_PoseFilter", "Filter"));
+				TextBuilder.AppendLine(LOCTEXT("DiscardedBy_PoseFilter_Tooltip", "(F) Filter"));
 			}
 
 			return TextBuilder.ToText();
diff --git a/Engine/Plugins/Experimental/Animation/PoseSearch/Source/Runtime/Private/PoseSearchDatabase.cpp b/Engine/Plugins/Experimental/Animation/PoseSearch/Source/Runtime/Private/PoseSearchDatabase.cpp
--- a/Engine/Plugins/Experimental/Animation/PoseSearch/Source/Runtime/Private/PoseSearchDatabase.cpp
+++ b/Engine/Plugins/Experimental/Animation/PoseSearch/Source/Runtime/Private/PoseSearchDatabase.cpp
@@ -26,7 +26,7 @@ DEFINE_STAT(STAT_PoseSearchPCAKNN);
 namespace UE::PoseSearch
 {
 
-typedef TArray<size_t, TInlineAllocator<128>> FNonSelectableIdx;
+typedef TArray<size_t, TInlineAllocator<256>> FNonSelectableIdx;
 static void PopulateNonSelectableIdx(FNonSelectableIdx& NonSelectableIdx, FSearchContext& SearchContext, const UPoseSearchDatabase* Database
 #if UE_POSE_SEARCH_TRACE_ENABLED
 	, TConstArrayView<float> QueryValues
@@ -38,6 +38,7 @@ static void PopulateNonSelectableIdx(FNonSelectableIdx& NonSelectableIdx, FSearc
 	const FPoseSearchIndex& SearchIndex = Database->GetSearchIndex();
 #endif
 
+	NonSelectableIdx.Reset();
 	const FPoseSearchIndexAsset* CurrentIndexAsset = SearchContext.GetCurrentResult().GetSearchIndexAsset();
 	if (CurrentIndexAsset && SearchContext.IsCurrentResultFromDatabase(Database) && SearchContext.GetPoseJumpThresholdTime() > 0.f)
 	{
@@ -64,7 +65,7 @@ static void PopulateNonSelectableIdx(FNonSelectableIdx& NonSelectableIdx, FSearc
 
 			if (bIsPoseInRange)
 			{
-				NonSelectableIdx.Add(PoseIdx);
+				NonSelectableIdx.AddUnique(PoseIdx);
 
 #if UE_POSE_SEARCH_TRACE_ENABLED
 				const TArray<float> PoseValues = SearchIndex.GetPoseValuesSafe(PoseIdx);
@@ -98,7 +99,7 @@ static void PopulateNonSelectableIdx(FNonSelectableIdx& NonSelectableIdx, FSearc
 
 			if (bIsPoseInRange)
 			{
-				NonSelectableIdx.Add(PoseIdx);
+				NonSelectableIdx.AddUnique(PoseIdx);
 
 #if UE_POSE_SEARCH_TRACE_ENABLED
 				const FPoseSearchCost PoseCost = SearchIndex.ComparePoses(PoseIdx, SearchContext.GetQueryMirrorRequest(), 0.f, Database->Schema->MirrorMismatchCostBias, SearchIndex.GetPoseValuesSafe(PoseIdx), QueryValues);
@@ -120,7 +121,7 @@ static void PopulateNonSelectableIdx(FNonSelectableIdx& NonSelectableIdx, FSearc
 			const FHistoricalPoseIndex& HistoricalPoseIndex = It.Key();
 			if (HistoricalPoseIndex.DatabaseKey == DatabaseKey)
 			{
-				NonSelectableIdx.Add(HistoricalPoseIndex.PoseIndex);
+				NonSelectableIdx.AddUnique(HistoricalPoseIndex.PoseIndex);
 
 #if UE_POSE_SEARCH_TRACE_ENABLED
 				check(HistoricalPoseIndex.PoseIndex >= 0);
diff --git a/Engine/Plugins/Experimental/Animation/PoseSearch/Source/Runtime/Public/PoseSearch/PoseSearchContext.h b/Engine/Plugins/Experimental/Animation/PoseSearch/Source/Runtime/Public/PoseSearch/PoseSearchContext.h
--- a/Engine/Plugins/Experimental/Animation/PoseSearch/Source/Runtime/Public/PoseSearch/PoseSearchContext.h
+++ b/Engine/Plugins/Experimental/Animation/PoseSearch/Source/Runtime/Public/PoseSearch/PoseSearchContext.h
@@ -202,47 +202,68 @@ private:
 #if UE_POSE_SEARCH_TRACE_ENABLED
 
 public:
-	struct FPoseCandidate
+	struct FPoseCandidateBase
 	{
 		FPoseSearchCost Cost;
 		int32 PoseIdx = 0;
 		const UPoseSearchDatabase* Database = nullptr;
-		EPoseCandidateFlags PoseCandidateFlags = EPoseCandidateFlags::None;
 
-		bool operator<(const FPoseCandidate& Other) const { return Other.Cost < Cost; } // Reverse compare because BestCandidates is a max heap
+		bool operator<(const FPoseCandidateBase& Other) const { return Other.Cost < Cost; } // Reverse compare because BestCandidates is a max heap
 		bool operator==(const FSearchResult& SearchResult) const { return (PoseIdx == SearchResult.PoseIdx) && (Database == SearchResult.Database.Get()); }
 	};
 
-	struct FBestPoseCandidates : private TArray<FPoseCandidate>
+	struct FPoseCandidate : public FPoseCandidateBase
 	{
-		typedef TArray<FPoseCandidate> Super;
-		using Super::IsEmpty;
-
-		int32 MaxPoseCandidates = 100;
+		EPoseCandidateFlags PoseCandidateFlags = EPoseCandidateFlags::None;
+	};
 
+	struct FBestPoseCandidates
+	{
 		void Add(const FPoseSearchCost& Cost, int32 PoseIdx, const UPoseSearchDatabase* Database, EPoseCandidateFlags PoseCandidateFlags)
 		{
-			if (Num() < MaxPoseCandidates || Cost < HeapTop().Cost)
+			if (EPoseCandidateFlags* PoseIdxPoseCandidateFlags = PoseIdxToFlags.Find(PoseIdx))
+			{
+				*PoseIdxPoseCandidateFlags |= PoseCandidateFlags;
+			}
+			else if (PoseCandidateHeap.Num() < MaxPoseCandidates || Cost < PoseCandidateHeap.HeapTop().Cost)
 			{
-				while (Num() >= MaxPoseCandidates)
+				while (PoseCandidateHeap.Num() >= MaxPoseCandidates)
 				{
-					ElementType Unused;
+					FPoseCandidate Unused;
 					Pop(Unused);
 				}
 
-				FSearchContext::FPoseCandidate PoseCandidate;
+				FPoseCandidate PoseCandidate;
 				PoseCandidate.Cost = Cost;
 				PoseCandidate.PoseIdx = PoseIdx;
 				PoseCandidate.Database = Database;
-				PoseCandidate.PoseCandidateFlags = PoseCandidateFlags;
-				HeapPush(PoseCandidate);
+				
+				PoseCandidateHeap.HeapPush(PoseCandidate);
+
+				PoseIdxToFlags.Add(PoseIdx, PoseCandidateFlags);
 			}
 		}
 
 		void Pop(FPoseCandidate& OutItem)
 		{
-			HeapPop(OutItem, false);
+			PoseCandidateHeap.HeapPop(OutItem, false);
+			OutItem.PoseCandidateFlags = PoseIdxToFlags.FindAndRemoveChecked(OutItem.PoseIdx);
 		}
+
+		bool IsEmpty() const
+		{
+			return PoseCandidateHeap.IsEmpty();
+		}
+
+		void SetMaxPoseCandidates(int32 Value)
+		{
+			MaxPoseCandidates = Value;
+		}
+
+	private:
+		TArray<FPoseCandidateBase> PoseCandidateHeap;
+		TMap<int32, EPoseCandidateFlags> PoseIdxToFlags;
+		int32 MaxPoseCandidates = 100;
 	};
 	
 	FBestPoseCandidates BestCandidates;
