motion matching - UE_POSE_SEARCH_TRACE_ENABLED code cleanup, QOL improvements

diff --git a/Engine/Plugins/Experimental/Animation/MotionTrajectory/Source/MotionTrajectory/Private/CharacterTrajectoryComponent.cpp b/Engine/Plugins/Experimental/Animation/MotionTrajectory/Source/MotionTrajectory/Private/CharacterTrajectoryComponent.cpp
--- a/Engine/Plugins/Experimental/Animation/MotionTrajectory/Source/MotionTrajectory/Private/CharacterTrajectoryComponent.cpp
+++ b/Engine/Plugins/Experimental/Animation/MotionTrajectory/Source/MotionTrajectory/Private/CharacterTrajectoryComponent.cpp
@@ -137,11 +137,20 @@ void UCharacterTrajectoryComponent::OnMovementUpdated(float DeltaSeconds, FVecto
 
 FPoseSearchQueryTrajectory UCharacterTrajectoryComponent::GetCharacterRelativeTrajectory() const
 {
+	check(GetOwner());
+
 	FPoseSearchQueryTrajectory CharacterRelativeTrajectory = Trajectory;
-	const FTransform OwnerTransformWS = GetOwner()->GetActorTransform();
-	const FTransform ComponentTransformWS = SkelMeshComponent->GetComponentTransform();
-	const FTransform ReferenceChangeTransform = OwnerTransformWS.GetRelativeTransform(ComponentTransformWS);
-	CharacterRelativeTrajectory.TransformReferenceFrame(ReferenceChangeTransform);
+	if (!SkelMeshComponent)
+	{
+		UE_LOG(LogMotionTrajectory, Error, TEXT("UCharacterTrajectoryComponent requires SkelMeshComponent to be initialized"));
+	}
+	else
+	{
+		const FTransform OwnerTransformWS = GetOwner()->GetActorTransform();
+		const FTransform ComponentTransformWS = SkelMeshComponent->GetComponentTransform();
+		const FTransform ReferenceChangeTransform = OwnerTransformWS.GetRelativeTransform(ComponentTransformWS);
+		CharacterRelativeTrajectory.TransformReferenceFrame(ReferenceChangeTransform);
+	}
 	return CharacterRelativeTrajectory;
 }
 
diff --git a/Engine/Plugins/Experimental/Animation/MotionTrajectory/Source/MotionTrajectory/Public/CharacterTrajectoryComponent.h b/Engine/Plugins/Experimental/Animation/MotionTrajectory/Source/MotionTrajectory/Public/CharacterTrajectoryComponent.h
--- a/Engine/Plugins/Experimental/Animation/MotionTrajectory/Source/MotionTrajectory/Public/CharacterTrajectoryComponent.h
+++ b/Engine/Plugins/Experimental/Animation/MotionTrajectory/Source/MotionTrajectory/Public/CharacterTrajectoryComponent.h
@@ -30,7 +30,7 @@ protected:
 	UFUNCTION()
 	void OnMovementUpdated(float DeltaSeconds, FVector OldLocation, FVector OldVelocity);
 
-	UFUNCTION(BlueprintCallable, Category="Trajectory")
+	UFUNCTION(BlueprintCallable, Category = "Trajectory", meta = (BlueprintThreadSafe))
 	FPoseSearchQueryTrajectory GetCharacterRelativeTrajectory() const;
 
 	void UpdateHistory(float DeltaSeconds, const FTransform& DeltaTransformCS);
diff --git a/Engine/Plugins/Experimental/Animation/PoseSearch/Source/Runtime/Private/AnimNode_MotionMatching.cpp b/Engine/Plugins/Experimental/Animation/PoseSearch/Source/Runtime/Private/AnimNode_MotionMatching.cpp
--- a/Engine/Plugins/Experimental/Animation/PoseSearch/Source/Runtime/Private/AnimNode_MotionMatching.cpp
+++ b/Engine/Plugins/Experimental/Animation/PoseSearch/Source/Runtime/Private/AnimNode_MotionMatching.cpp
@@ -48,15 +48,11 @@ void FAnimNode_MotionMatching::Evaluate_AnyThread(FPoseContext& Output)
 	MotionMatchingState.RootMotionTransformDelta = FTransform::Identity;
 
 	const UE::Anim::IAnimRootMotionProvider* RootMotionProvider = UE::Anim::IAnimRootMotionProvider::Get();
-
-	if (ensureMsgf(RootMotionProvider, TEXT("Could not get Root Motion Provider.")))
+	if (RootMotionProvider && RootMotionProvider->HasRootMotion(Output.CustomAttributes))
 	{
-		if (RootMotionProvider->HasRootMotion(Output.CustomAttributes))
-		{
-			RootMotionProvider->ExtractRootMotion(Output.CustomAttributes, MotionMatchingState.RootMotionTransformDelta);
-		}
+		RootMotionProvider->ExtractRootMotion(Output.CustomAttributes, MotionMatchingState.RootMotionTransformDelta);
 	}
-#endif
+#endif // UE_POSE_SEARCH_TRACE_ENABLED
 }
 
 void FAnimNode_MotionMatching::UpdateAssetPlayer(const FAnimationUpdateContext& Context)
diff --git a/Engine/Plugins/Experimental/Animation/PoseSearch/Source/Runtime/Private/PoseSearchLibrary.cpp b/Engine/Plugins/Experimental/Animation/PoseSearch/Source/Runtime/Private/PoseSearchLibrary.cpp
--- a/Engine/Plugins/Experimental/Animation/PoseSearch/Source/Runtime/Private/PoseSearchLibrary.cpp
+++ b/Engine/Plugins/Experimental/Animation/PoseSearch/Source/Runtime/Private/PoseSearchLibrary.cpp
@@ -137,6 +137,7 @@ void FMotionMatchingState::UpdateWantedPlayRate(const UE::PoseSearch::FSearchCon
 	}
 }
 
+#if UE_POSE_SEARCH_TRACE_ENABLED
 void UPoseSearchLibrary::TraceMotionMatchingState(
 	const UPoseSearchDatabase* Database,
 	UE::PoseSearch::FSearchContext& SearchContext,
@@ -149,7 +150,6 @@ void UPoseSearchLibrary::TraceMotionMatchingState(
 	float DeltaTime,
 	bool bSearch)
 {
-#if UE_POSE_SEARCH_TRACE_ENABLED
 	using namespace UE::PoseSearch;
 	
 	auto AddUniqueDatabase = [](TArray<FTraceMotionMatchingStateDatabaseEntry>& DatabaseEntries, const UPoseSearchDatabase* Database, UE::PoseSearch::FSearchContext& SearchContext) -> int32
@@ -247,8 +247,8 @@ void UPoseSearchLibrary::TraceMotionMatchingState(
 	TraceState.DeltaTime = DeltaTime;
 
 	TraceState.Output(AnimInstance, NodeId);
-#endif
 }
+#endif // UE_POSE_SEARCH_TRACE_ENABLED
 
 void UPoseSearchLibrary::UpdateMotionMatchingState(
 	const FAnimationUpdateContext& Context,
diff --git a/Engine/Plugins/Experimental/Animation/PoseSearch/Source/Runtime/Public/PoseSearch/PoseSearchLibrary.h b/Engine/Plugins/Experimental/Animation/PoseSearch/Source/Runtime/Public/PoseSearch/PoseSearchLibrary.h
--- a/Engine/Plugins/Experimental/Animation/PoseSearch/Source/Runtime/Public/PoseSearch/PoseSearchLibrary.h
+++ b/Engine/Plugins/Experimental/Animation/PoseSearch/Source/Runtime/Public/PoseSearch/PoseSearchLibrary.h
@@ -8,6 +8,7 @@
 #include "Animation/AnimNodeReference.h"
 #include "Kismet/BlueprintFunctionLibrary.h"
 #include "Kismet/BlueprintFunctionLibrary.h"
+#include "PoseSearch/PoseSearchDefines.h"
 #include "PoseSearch/PoseSearchHistory.h"
 #include "PoseSearch/PoseSearchResult.h"
 #include "SequenceEvaluatorLibrary.h"
@@ -95,9 +96,10 @@ struct POSESEARCH_API FMotionMatchingState
 	UPROPERTY(VisibleAnywhere, BlueprintReadOnly, Category=State)
 	bool bJumpedToPose = false;
 
-	// Root motion delta for currently playing animation. Only required when UE_POSE_SEARCH_TRACE_ENABLED is active.
-	UPROPERTY(Transient)
+#if UE_POSE_SEARCH_TRACE_ENABLED
+	// Root motion delta for currently playing animation.
 	FTransform RootMotionTransformDelta = FTransform::Identity;
+#endif //UE_POSE_SEARCH_TRACE_ENABLED
 
 	// @todo: add ContinuingPoseCost / BruteForcePoseCost / PoseCost graphs to rewind debugger 
 //#if WITH_EDITORONLY_DATA
@@ -114,6 +116,8 @@ UCLASS()
 class POSESEARCH_API UPoseSearchLibrary : public UBlueprintFunctionLibrary
 {
 	GENERATED_BODY()
+
+#if UE_POSE_SEARCH_TRACE_ENABLED
 	static void TraceMotionMatchingState(
 		const UPoseSearchDatabase* Database,
 		UE::PoseSearch::FSearchContext& SearchContext,
@@ -125,6 +129,7 @@ class POSESEARCH_API UPoseSearchLibrary : public UBlueprintFunctionLibrary
 		int32 NodeId,
 		float DeltaTime,
 		bool bSearch);
+#endif // UE_POSE_SEARCH_TRACE_ENABLED
 
 public:
 	/**
