motion matching - fix for Previewing and scrubbing the animation timeline in a Motion Database causes jitter on the Sequences

diff --git a/Engine/Plugins/Experimental/Animation/PoseSearch/Source/Editor/Private/PoseSearchDatabaseViewModel.cpp b/Engine/Plugins/Experimental/Animation/PoseSearch/Source/Editor/Private/PoseSearchDatabaseViewModel.cpp
--- a/Engine/Plugins/Experimental/Animation/PoseSearch/Source/Editor/Private/PoseSearchDatabaseViewModel.cpp
+++ b/Engine/Plugins/Experimental/Animation/PoseSearch/Source/Editor/Private/PoseSearchDatabaseViewModel.cpp
@@ -200,12 +200,16 @@ namespace UE::PoseSearch
 
 	void FDatabaseViewModel::Tick(float DeltaSeconds)
 	{
-		PlayTime += DeltaSeconds * DeltaTimeMultiplier;
-		PlayTime = FMath::Clamp(PlayTime, MinPreviewPlayLength, MaxPreviewPlayLength);
-
-		if (FAsyncPoseSearchDatabasesManagement::RequestAsyncBuildIndex(PoseSearchDatabase, ERequestAsyncBuildFlag::ContinueRequest))
+		const float DeltaPlayTime = DeltaSeconds * DeltaTimeMultiplier;
+		if (!FMath::IsNearlyZero(DeltaPlayTime))
 		{
-			UpdatePreviewActors(true);
+			PlayTime += DeltaPlayTime;
+			PlayTime = FMath::Clamp(PlayTime, MinPreviewPlayLength, MaxPreviewPlayLength);
+
+			if (FAsyncPoseSearchDatabasesManagement::RequestAsyncBuildIndex(PoseSearchDatabase, ERequestAsyncBuildFlag::ContinueRequest))
+			{
+				UpdatePreviewActors(true);
+			}
 		}
 	}
 
@@ -497,14 +501,22 @@ namespace UE::PoseSearch
 
 	float FDatabaseViewModel::GetPlayTime() const
 	{
-		const float ClampedPlayTime = FMath::Clamp(PlayTime, MinPreviewPlayLength, MaxPreviewPlayLength);
-		return ClampedPlayTime;
+		return PlayTime;
 	}
 
 	void FDatabaseViewModel::SetPlayTime(float NewPlayTime, bool bInTickPlayTime)
 	{
-		PlayTime = NewPlayTime;
+		NewPlayTime = FMath::Clamp(NewPlayTime, MinPreviewPlayLength, MaxPreviewPlayLength);
 		DeltaTimeMultiplier = bInTickPlayTime ? DeltaTimeMultiplier : 0.0f;
+
+		if (!FMath::IsNearlyEqual(PlayTime, NewPlayTime))
+		{
+			PlayTime = NewPlayTime;
+			if (FAsyncPoseSearchDatabasesManagement::RequestAsyncBuildIndex(PoseSearchDatabase, ERequestAsyncBuildFlag::ContinueRequest))
+			{
+				UpdatePreviewActors(true);
+			}
+		}
 	}
 
 	bool FDatabaseViewModel::GetAnimationTime(int32 SourceAssetIdx, float& CurrentPlayTime, FVector& BlendParameters) const
