motion matching
- adding bShouldSearch to FAnimNode_MotionMatching: the node will perform a search only if the continuing pose is invalid. This is useful if you want to stagger searches of different nodes for performance reasons
- reducing FSearchContext::GetCachedQuery allocations / copies

diff --git a/Engine/Plugins/Experimental/Animation/PoseSearch/Source/Runtime/Private/AnimNode_MotionMatching.cpp b/Engine/Plugins/Experimental/Animation/PoseSearch/Source/Runtime/Private/AnimNode_MotionMatching.cpp
--- a/Engine/Plugins/Experimental/Animation/PoseSearch/Source/Runtime/Private/AnimNode_MotionMatching.cpp
+++ b/Engine/Plugins/Experimental/Animation/PoseSearch/Source/Runtime/Private/AnimNode_MotionMatching.cpp
@@ -81,8 +81,7 @@ void FAnimNode_MotionMatching::UpdateAssetPlayer(const FAnimationUpdateContext&
 			const UPoseSearchDatabase* CurrentResultDatabase = MotionMatchingState.CurrentSearchResult.Database.Get();
 			const FPoseSearchIndex& SearchIndex = CurrentResultDatabase->GetSearchIndex();
 			if (!SearchIndex.IsValidPoseIndex(MotionMatchingState.CurrentSearchResult.PrevPoseIdx) ||
-				!SearchIndex.IsValidPoseIndex(MotionMatchingState.CurrentSearchResult.NextPoseIdx) ||
-				CurrentResultDatabase->Schema != MotionMatchingState.CurrentSearchResult.ComposedQuery.GetSchema())
+				!SearchIndex.IsValidPoseIndex(MotionMatchingState.CurrentSearchResult.NextPoseIdx))
 			{
 				// MotionMatchingState is out of sync with CurrentResultDatabase: we need to reset the MM state. This could happen if PIE is paused, and we edit the database,
 				// so FAnimNode_MotionMatching::UpdateAssetPlayer is never called and FAsyncPoseSearchDatabasesManagement::RequestAsyncBuildIndex never returns false here
@@ -117,7 +116,12 @@ void FAnimNode_MotionMatching::UpdateAssetPlayer(const FAnimationUpdateContext&
 		Trajectory,
 		Settings,
 		MotionMatchingState,
-		bForceInterrupt | bForceInterruptNextUpdate
+		bForceInterrupt | bForceInterruptNextUpdate,
+		bShouldSearch
+		#if ENABLE_ANIM_DEBUG
+		, CVarAnimNodeMotionMatchingDrawQuery.GetValueOnAnyThread() > 0
+		, CVarAnimNodeMotionMatchingDrawCurResult.GetValueOnAnyThread() > 0
+		#endif // WITH_EDITORONLY_DATA && ENABLE_ANIM_DEBUG
 	);
 
 	// If a new pose is requested, blend into the new asset via BlendStackNode
@@ -139,70 +143,6 @@ void FAnimNode_MotionMatching::UpdateAssetPlayer(const FAnimationUpdateContext&
 
 	Source.Update(Context);
 
-
-#if WITH_EDITORONLY_DATA && ENABLE_ANIM_DEBUG
-	const bool bDebugDrawQuery = CVarAnimNodeMotionMatchingDrawQuery.GetValueOnAnyThread() > 0;
-	const bool bDebugDrawCurResult = CVarAnimNodeMotionMatchingDrawCurResult.GetValueOnAnyThread() > 0;
-	if (bDebugDrawQuery || bDebugDrawCurResult)
-	{
-		const UE::PoseSearch::FSearchResult& CurResult = MotionMatchingState.CurrentSearchResult;
-
-#if WITH_EDITOR
-		// in case we're still indexing MotionMatchingState.CurrentSearchResult.Database we Reset the MotionMatchingState
-		if (!FAsyncPoseSearchDatabasesManagement::RequestAsyncBuildIndex(CurResult.Database.Get(), ERequestAsyncBuildFlag::ContinueRequest))
-		{
-		}
-		else
-#endif // WITH_EDITOR
-		{
-			if (bDebugDrawCurResult)
-			{
-				UE::PoseSearch::FDebugDrawParams DrawParams(Context.AnimInstanceProxy, CurResult.Database.Get());
-				DrawParams.DrawFeatureVector(CurResult.PoseIdx);
-			}
-
-			if (bDebugDrawQuery)
-			{
-				UE::PoseSearch::FDebugDrawParams DrawParams(Context.AnimInstanceProxy, CurResult.Database.Get(), EDebugDrawFlags::DrawQuery);
-				DrawParams.DrawFeatureVector(CurResult.ComposedQuery.GetValues());
-			}
-
-			// @todo: add ContinuingPoseCost / BruteForcePoseCost / PoseCost graphs to rewind debugger 
-			//if (DrawParams.Database->PoseSearchMode == EPoseSearchMode::PCAKDTree_Compare)
-			//{
-			//	FDebugFloatHistory& C = MotionMatchingState.SearchCostHistoryContinuing;
-			//	FDebugFloatHistory& B = MotionMatchingState.SearchCostHistoryBruteForce;
-			//	FDebugFloatHistory& K = MotionMatchingState.SearchCostHistoryKDTree;
-
-			//	C.AddSample(CurResult.ContinuingPoseCost.IsValid() ? CurResult.ContinuingPoseCost.GetTotalCost() : C.MaxValue);
-			//	B.AddSample(CurResult.BruteForcePoseCost.IsValid() ? CurResult.BruteForcePoseCost.GetTotalCost() : B.MaxValue);
-			//	K.AddSample(CurResult.PoseCost.IsValid() ? CurResult.PoseCost.GetTotalCost() : K.MaxValue);
-
-			//	// making SearchCostHistoryKDTree and SearchCostHistoryBruteForce min max consistent
-			//	const float MinValue = FMath::Min(C.MinValue, FMath::Min(B.MinValue, K.MinValue));
-			//	const float MaxValue = FMath::Max(C.MaxValue, FMath::Max(B.MaxValue, K.MaxValue));
-
-			//	C.MinValue = MinValue;
-			//	C.MaxValue = MaxValue;
-
-			//	B.MinValue = MinValue;
-			//	B.MaxValue = MaxValue;
-
-			//	K.MinValue = MinValue;
-			//	K.MaxValue = MaxValue;
-
-			//	const FVector2D DrawSize(150.f, 100.f);
-			//	const FTransform OffsetTransform(FRotator(0.f, 0.f, 0.f), FVector(-50.f, -75.f, 100.f));
-			//	const FTransform DrawTransform = OffsetTransform * DrawParams.GetRootTransform();
-
-			//	DrawDebugFloatHistory(*DrawParams.World, K, OffsetTransform * DrawParams.GetRootTransform(), DrawSize, FColor(255, 192, 203, 160)); // pink
-			//	DrawDebugFloatHistory(*DrawParams.World, B, OffsetTransform * DrawParams.GetRootTransform(), DrawSize, FColor(0, 0, 255, 160)); // blue
-			//	DrawDebugFloatHistory(*DrawParams.World, C, OffsetTransform * DrawParams.GetRootTransform(), DrawSize, FColor(160, 160, 160, 160)); // gray
-			//}
-		}
-	} 
-#endif
-
 	bForceInterruptNextUpdate = false;
 }
 
diff --git a/Engine/Plugins/Experimental/Animation/PoseSearch/Source/Runtime/Private/PoseSearchContext.cpp b/Engine/Plugins/Experimental/Animation/PoseSearch/Source/Runtime/Private/PoseSearchContext.cpp
--- a/Engine/Plugins/Experimental/Animation/PoseSearch/Source/Runtime/Private/PoseSearchContext.cpp
+++ b/Engine/Plugins/Experimental/Animation/PoseSearch/Source/Runtime/Private/PoseSearchContext.cpp
@@ -451,8 +451,6 @@ void FSearchContext::ResetCurrentBestCost()
 
 void FSearchContext::UpdateCurrentBestCost(const FPoseSearchCost& PoseSearchCost)
 {
-	check(PoseSearchCost.IsValid());
-
 	if (PoseSearchCost.GetTotalCost() < CurrentBestTotalCost)
 	{
 		CurrentBestTotalCost = PoseSearchCost.GetTotalCost();
@@ -473,20 +471,18 @@ const FPoseSearchFeatureVectorBuilder* FSearchContext::GetCachedQuery(const UPos
 	return nullptr;
 }
 
-void FSearchContext::GetOrBuildQuery(const UPoseSearchSchema* Schema, FPoseSearchFeatureVectorBuilder& FeatureVectorBuilder)
+const FPoseSearchFeatureVectorBuilder& FSearchContext::GetOrBuildQuery(const UPoseSearchSchema* Schema)
 {
 	check(Schema && Schema->IsValid());
 	const FPoseSearchFeatureVectorBuilder* CachedFeatureVectorBuilder = GetCachedQuery(Schema);
 	if (CachedFeatureVectorBuilder)
 	{
-		FeatureVectorBuilder = *CachedFeatureVectorBuilder;
-	}
-	else
-	{
-		FPoseSearchFeatureVectorBuilder& NewCachedQuery = CachedQueries[CachedQueries.AddDefaulted()];
-		Schema->BuildQuery(*this, NewCachedQuery);
-		FeatureVectorBuilder = NewCachedQuery;
+		return *CachedFeatureVectorBuilder;
 	}
+	
+	FPoseSearchFeatureVectorBuilder& NewCachedQuery = CachedQueries[CachedQueries.AddDefaulted()];
+	Schema->BuildQuery(*this, NewCachedQuery);
+	return NewCachedQuery;
 }
 
 bool FSearchContext::IsCurrentResultFromDatabase(const UPoseSearchDatabase* Database) const
diff --git a/Engine/Plugins/Experimental/Animation/PoseSearch/Source/Runtime/Private/PoseSearchDatabase.cpp b/Engine/Plugins/Experimental/Animation/PoseSearch/Source/Runtime/Private/PoseSearchDatabase.cpp
--- a/Engine/Plugins/Experimental/Animation/PoseSearch/Source/Runtime/Private/PoseSearchDatabase.cpp
+++ b/Engine/Plugins/Experimental/Animation/PoseSearch/Source/Runtime/Private/PoseSearchDatabase.cpp
@@ -7,6 +7,7 @@
 #include "Animation/BlendSpace.h"
 #include "Animation/BlendSpace1D.h"
 #include "InstancedStruct.h"
+#include "PoseSearch/PoseSearchAnimNotifies.h"
 #include "PoseSearch/PoseSearchContext.h"
 #include "PoseSearch/PoseSearchDefines.h"
 #include "PoseSearch/PoseSearchDerivedData.h"
@@ -16,7 +17,6 @@
 #include "PoseSearchEigenHelper.h"
 #include "UObject/ObjectSaveContext.h"
 
-
 DECLARE_STATS_GROUP(TEXT("PoseSearch"), STATGROUP_PoseSearch, STATCAT_Advanced);
 DECLARE_CYCLE_STAT_EXTERN(TEXT("Search Brute Force"), STAT_PoseSearchBruteForce, STATGROUP_PoseSearch, );
 DECLARE_CYCLE_STAT_EXTERN(TEXT("Search PCA/KNN"), STAT_PoseSearchPCAKNN, STATGROUP_PoseSearch, );
@@ -724,6 +724,59 @@ UE::PoseSearch::FSearchResult UPoseSearchDatabase::Search(UE::PoseSearch::FSearc
 	return Result;
 }
 
+FPoseSearchCost UPoseSearchDatabase::SearchContinuingPose(UE::PoseSearch::FSearchContext& SearchContext) const
+{
+	using namespace UE::PoseSearch;
+
+	check(SearchContext.GetCurrentResult().Database.Get() == this);
+
+	FPoseSearchCost ContinuingPoseCost;
+
+#if WITH_EDITOR
+	if (!FAsyncPoseSearchDatabasesManagement::RequestAsyncBuildIndex(this, ERequestAsyncBuildFlag::ContinueRequest))
+	{
+		return ContinuingPoseCost;
+	}
+#endif
+
+	// extracting notifies from the database animation asset at time SampleTime to search for UAnimNotifyState_PoseSearchOverrideContinuingPoseCostBias eventually overriding the schema ContinuingPoseCostBias
+	const FPoseSearchIndex& SearchIndex = GetSearchIndex();
+	const int32 PoseIdx = SearchContext.GetCurrentResult().PoseIdx;
+	const FPoseSearchIndexAsset& SearchIndexAsset = SearchIndex.GetAssetForPose(PoseIdx);
+	const FPoseSearchDatabaseAnimationAssetBase* DatabaseAnimationAssetBase = GetAnimationAssetStruct(SearchIndexAsset).GetPtr<FPoseSearchDatabaseAnimationAssetBase>();
+	check(DatabaseAnimationAssetBase);
+	const FAnimationAssetSampler SequenceBaseSampler(DatabaseAnimationAssetBase->GetAnimationAsset(), SearchIndexAsset.BlendParameters);
+	const float SampleTime = GetAssetTime(PoseIdx);
+
+	// @todo: change ExtractPoseSearchNotifyStates api to avoid NotifyStates allocation
+	TArray<UAnimNotifyState_PoseSearchBase*> NotifyStates;
+	SequenceBaseSampler.ExtractPoseSearchNotifyStates(SampleTime, NotifyStates);
+
+	float ContinuingPoseCostBias = Schema->ContinuingPoseCostBias;
+	for (const UAnimNotifyState_PoseSearchBase* PoseSearchNotify : NotifyStates)
+	{
+		if (const UAnimNotifyState_PoseSearchOverrideContinuingPoseCostBias* ContinuingPoseCostBiasNotify = Cast<const UAnimNotifyState_PoseSearchOverrideContinuingPoseCostBias>(PoseSearchNotify))
+		{
+			ContinuingPoseCostBias = ContinuingPoseCostBiasNotify->CostAddend;
+			break;
+		}
+	}
+
+	// since any PoseCost calculated here is at least SearchIndex.MinCostAddend + ContinuingPoseCostBias,
+	// there's no point in performing the search if CurrentBestTotalCost is already better than that
+	if (!GetSkipSearchIfPossible() || SearchContext.GetCurrentBestTotalCost() > SearchIndex.MinCostAddend + ContinuingPoseCostBias)
+	{
+		const int32 NumDimensions = Schema->SchemaCardinality;
+		TArrayView<float> ReconstructedPoseValuesBuffer((float*)FMemory_Alloca(NumDimensions * sizeof(float)), NumDimensions);
+		const TConstArrayView<float> PoseValues = SearchIndex.Values.IsEmpty() ? SearchIndex.GetReconstructedPoseValues(PoseIdx, ReconstructedPoseValuesBuffer) : SearchIndex.GetPoseValues(PoseIdx);
+
+		ContinuingPoseCost = SearchIndex.ComparePoses(SearchContext.GetCurrentResult().PoseIdx, SearchContext.GetQueryMirrorRequest(),
+			ContinuingPoseCostBias, Schema->MirrorMismatchCostBias, PoseValues, SearchContext.GetOrBuildQuery(Schema).GetValues());
+	}
+
+	return ContinuingPoseCost;
+}
+
 UE::PoseSearch::FSearchResult UPoseSearchDatabase::SearchPCAKDTree(UE::PoseSearch::FSearchContext& SearchContext) const
 {
 	QUICK_SCOPE_CYCLE_COUNTER(STAT_PoseSearch_PCA_KNN);
@@ -785,8 +838,7 @@ UE::PoseSearch::FSearchResult UPoseSearchDatabase::SearchPCAKDTree(UE::PoseSearc
 	// there's no point in performing the search if CurrentBestTotalCost is already better than that
 	if (!GetSkipSearchIfPossible() || SearchContext.GetCurrentBestTotalCost() > SearchIndex.MinCostAddend)
 	{
-		SearchContext.GetOrBuildQuery(Schema, Result.ComposedQuery);
-		TConstArrayView<float> QueryValues = Result.ComposedQuery.GetValues();
+		TConstArrayView<float> QueryValues = SearchContext.GetOrBuildQuery(Schema).GetValues();
 
 		FNonSelectableIdx NonSelectableIdx;
 		PopulateNonSelectableIdx(NonSelectableIdx, SearchContext, this
@@ -840,8 +892,7 @@ UE::PoseSearch::FSearchResult UPoseSearchDatabase::SearchPCAKDTree(UE::PoseSearc
 	{
 #if UE_POSE_SEARCH_TRACE_ENABLED
 		// calling just for reporting non selectable poses
-		SearchContext.GetOrBuildQuery(Schema, Result.ComposedQuery);
-		TConstArrayView<float> QueryValues = Result.ComposedQuery.GetValues();
+		TConstArrayView<float> QueryValues = SearchContext.GetOrBuildQuery(Schema).GetValues();
 		FNonSelectableIdx NonSelectableIdx;
 		PopulateNonSelectableIdx(NonSelectableIdx, SearchContext, this, QueryValues);
 #endif
@@ -872,8 +923,7 @@ UE::PoseSearch::FSearchResult UPoseSearchDatabase::SearchBruteForce(UE::PoseSear
 	// there's no point in performing the search if CurrentBestTotalCost is already better than that
 	if (!GetSkipSearchIfPossible() || SearchContext.GetCurrentBestTotalCost() > SearchIndex.MinCostAddend)
 	{
-		SearchContext.GetOrBuildQuery(Schema, Result.ComposedQuery);
-		TConstArrayView<float> QueryValues = Result.ComposedQuery.GetValues();
+		TConstArrayView<float> QueryValues = SearchContext.GetOrBuildQuery(Schema).GetValues();
 
 		FNonSelectableIdx NonSelectableIdx;
 		PopulateNonSelectableIdx(NonSelectableIdx, SearchContext, this
@@ -915,8 +965,7 @@ UE::PoseSearch::FSearchResult UPoseSearchDatabase::SearchBruteForce(UE::PoseSear
 	{
 #if UE_POSE_SEARCH_TRACE_ENABLED
 		// calling just for reporting non selectable poses
-		SearchContext.GetOrBuildQuery(Schema, Result.ComposedQuery);
-		TConstArrayView<float> QueryValues = Result.ComposedQuery.GetValues();
+		TConstArrayView<float> QueryValues = SearchContext.GetOrBuildQuery(Schema).GetValues();
 		FNonSelectableIdx NonSelectableIdx;
 		PopulateNonSelectableIdx(NonSelectableIdx, SearchContext, this, QueryValues);
 #endif
diff --git a/Engine/Plugins/Experimental/Animation/PoseSearch/Source/Runtime/Private/PoseSearchLibrary.cpp b/Engine/Plugins/Experimental/Animation/PoseSearch/Source/Runtime/Private/PoseSearchLibrary.cpp
--- a/Engine/Plugins/Experimental/Animation/PoseSearch/Source/Runtime/Private/PoseSearchLibrary.cpp
+++ b/Engine/Plugins/Experimental/Animation/PoseSearch/Source/Runtime/Private/PoseSearchLibrary.cpp
@@ -13,7 +13,6 @@
 #include "InstancedStruct.h"
 #include "PoseSearch/AnimNode_MotionMatching.h"
 #include "PoseSearch/AnimNode_PoseSearchHistoryCollector.h"
-#include "PoseSearch/PoseSearchAnimNotifies.h"
 #include "PoseSearch/PoseSearchDatabase.h"
 #include "PoseSearch/PoseSearchDerivedData.h"
 #include "PoseSearch/PoseSearchSchema.h"
@@ -170,10 +169,7 @@ void UPoseSearchLibrary::TraceMotionMatchingState(
 			DbEntryIdx = DatabaseEntries.Add({ DatabaseId });
 
 			// if throttling is on, the continuing pose can be valid, but no actual search occurred, so the query will not be cached, and we need to build it
-			FPoseSearchFeatureVectorBuilder FeatureVectorBuilder;
-			FeatureVectorBuilder.Init(Database->Schema);
-			SearchContext.GetOrBuildQuery(Database->Schema, FeatureVectorBuilder);
-			DatabaseEntries[DbEntryIdx].QueryVector = FeatureVectorBuilder.GetValues();
+			DatabaseEntries[DbEntryIdx].QueryVector = SearchContext.GetOrBuildQuery(Database->Schema).GetValues();
 		}
 
 		return DbEntryIdx;
@@ -256,7 +252,10 @@ void UPoseSearchLibrary::UpdateMotionMatchingState(
 	const FPoseSearchQueryTrajectory& Trajectory,
 	const FMotionMatchingSettings& Settings,
 	FMotionMatchingState& InOutMotionMatchingState,
-	bool bForceInterrupt)
+	bool bForceInterrupt,
+	bool bShouldSearch,
+	bool bDebugDrawQuery,
+	bool bDebugDrawCurResult)
 {
 	QUICK_SCOPE_CYCLE_COUNTER(STAT_PoseSearch_Update);
 
@@ -295,56 +294,17 @@ void UPoseSearchLibrary::UpdateMotionMatchingState(
 		InOutMotionMatchingState.CurrentSearchResult, Settings.PoseJumpThresholdTime, bForceInterrupt, InOutMotionMatchingState.CanAdvance(DeltaTime));
 
 	// If we can't advance or enough time has elapsed since the last pose jump then search
-	const bool bSearch = !SearchContext.CanAdvance() || (InOutMotionMatchingState.ElapsedPoseSearchTime >= Settings.SearchThrottleTime);
+	const bool bSearch = !SearchContext.CanAdvance() || (bShouldSearch && (InOutMotionMatchingState.ElapsedPoseSearchTime >= Settings.SearchThrottleTime));
 	if (bSearch)
 	{
 		InOutMotionMatchingState.ElapsedPoseSearchTime = 0.f;
 
 		// Evaluate continuing pose
 		FPoseSearchCost ContinuingPoseCost;
-		FPoseSearchFeatureVectorBuilder ContinuingPoseComposedQuery;
-		if (!SearchContext.IsForceInterrupt() && SearchContext.CanAdvance() && SearchContext.GetCurrentResult().Database.IsValid()
-#if WITH_EDITOR
-			&& FAsyncPoseSearchDatabasesManagement::RequestAsyncBuildIndex(SearchContext.GetCurrentResult().Database.Get(), ERequestAsyncBuildFlag::ContinueRequest)
-#endif
-		)
+		if (!SearchContext.IsForceInterrupt() && SearchContext.CanAdvance() && SearchContext.GetCurrentResult().Database.IsValid())
 		{
-			const UPoseSearchDatabase* ContinuingPoseDatabase = SearchContext.GetCurrentResult().Database.Get();
-			if (ensure(ContinuingPoseDatabase->Schema))
-			{	
-				SearchContext.GetOrBuildQuery(ContinuingPoseDatabase->Schema, ContinuingPoseComposedQuery);
-
-				const FPoseSearchIndex& SearchIndex = ContinuingPoseDatabase->GetSearchIndex();
-				const int32 PoseIdx = SearchContext.GetCurrentResult().PoseIdx;
-				const int32 NumDimensions = ContinuingPoseDatabase->Schema->SchemaCardinality;
-				TArrayView<float> ReconstructedPoseValuesBuffer((float*)FMemory_Alloca(NumDimensions * sizeof(float)), NumDimensions);
-				const TConstArrayView<float> PoseValues = SearchIndex.Values.IsEmpty() ? SearchIndex.GetReconstructedPoseValues(PoseIdx, ReconstructedPoseValuesBuffer) : SearchIndex.GetPoseValues(PoseIdx);
-
-				// extracting notifies from the database animation asset at time SampleTime to search for UAnimNotifyState_PoseSearchOverrideContinuingPoseCostBias eventually overriding the schema ContinuingPoseCostBias
-				const FPoseSearchIndexAsset& SearchIndexAsset = SearchIndex.GetAssetForPose(PoseIdx);
-				const FPoseSearchDatabaseAnimationAssetBase* DatabaseAnimationAssetBase = ContinuingPoseDatabase->GetAnimationAssetStruct(SearchIndexAsset).GetPtr<FPoseSearchDatabaseAnimationAssetBase>();
-				check(DatabaseAnimationAssetBase);
-				const FAnimationAssetSampler SequenceBaseSampler(DatabaseAnimationAssetBase->GetAnimationAsset(), SearchIndexAsset.BlendParameters);
-				const float SampleTime = ContinuingPoseDatabase->GetAssetTime(PoseIdx);
-
-				// @todo: change ExtractPoseSearchNotifyStates api to avoid NotifyStates allocation
-				TArray<UAnimNotifyState_PoseSearchBase*> NotifyStates;
-				SequenceBaseSampler.ExtractPoseSearchNotifyStates(SampleTime, NotifyStates);
-
-				float ContinuingPoseCostBias = ContinuingPoseDatabase->Schema->ContinuingPoseCostBias;
-				for (const UAnimNotifyState_PoseSearchBase* PoseSearchNotify : NotifyStates)
-				{
-					if (const UAnimNotifyState_PoseSearchOverrideContinuingPoseCostBias* ContinuingPoseCostBiasNotify = Cast<const UAnimNotifyState_PoseSearchOverrideContinuingPoseCostBias>(PoseSearchNotify))
-					{
-						ContinuingPoseCostBias = ContinuingPoseCostBiasNotify->CostAddend;
-						break;
-					}
-				}
-
-				ContinuingPoseCost = ContinuingPoseDatabase->GetSearchIndex().ComparePoses(SearchContext.GetCurrentResult().PoseIdx, SearchContext.GetQueryMirrorRequest(),
-					ContinuingPoseCostBias, ContinuingPoseDatabase->Schema->MirrorMismatchCostBias, PoseValues, ContinuingPoseComposedQuery.GetValues());
-				SearchContext.UpdateCurrentBestCost(ContinuingPoseCost);
-			}
+			ContinuingPoseCost = SearchContext.GetCurrentResult().Database->SearchContinuingPose(SearchContext);
+			SearchContext.UpdateCurrentBestCost(ContinuingPoseCost);
 		}
 
 		FSearchResult SearchResult;
@@ -375,7 +335,6 @@ void UPoseSearchLibrary::UpdateMotionMatchingState(
 #endif
 			InOutMotionMatchingState.CurrentSearchResult.PoseCost = ContinuingPoseCost;
 			InOutMotionMatchingState.CurrentSearchResult.ContinuingPoseCost = ContinuingPoseCost;
-			InOutMotionMatchingState.CurrentSearchResult.ComposedQuery = ContinuingPoseComposedQuery;
 		}
 
 		InOutMotionMatchingState.UpdateWantedPlayRate(SearchContext, Settings);
@@ -396,6 +355,35 @@ void UPoseSearchLibrary::UpdateMotionMatchingState(
 			InOutMotionMatchingState.RootMotionTransformDelta, Context.AnimInstanceProxy->GetAnimInstanceObject(), Context.GetCurrentNodeId(), DeltaTime, bSearch);
 	}
 #endif
+
+#if WITH_EDITORONLY_DATA && ENABLE_ANIM_DEBUG
+	const FSearchResult& CurResult = InOutMotionMatchingState.CurrentSearchResult;
+	if ((bDebugDrawQuery || bDebugDrawCurResult) && CurResult.Database.IsValid())
+	{
+		const UPoseSearchDatabase* CurResultDatabase = CurResult.Database.Get();
+
+#if WITH_EDITOR
+		// in case we're still indexing MotionMatchingState.CurrentSearchResult.Database we Reset the MotionMatchingState
+		if (!FAsyncPoseSearchDatabasesManagement::RequestAsyncBuildIndex(CurResultDatabase, ERequestAsyncBuildFlag::ContinueRequest))
+		{
+		}
+		else
+#endif // WITH_EDITOR
+		{
+			if (bDebugDrawCurResult)
+			{
+				UE::PoseSearch::FDebugDrawParams DrawParams(Context.AnimInstanceProxy, CurResultDatabase);
+				DrawParams.DrawFeatureVector(CurResult.PoseIdx);
+			}
+
+			if (bDebugDrawQuery)
+			{
+				UE::PoseSearch::FDebugDrawParams DrawParams(Context.AnimInstanceProxy, CurResultDatabase, EDebugDrawFlags::DrawQuery);
+				DrawParams.DrawFeatureVector(SearchContext.GetOrBuildQuery(CurResultDatabase->Schema).GetValues());
+			}
+		}
+	}
+#endif
 }
 
 void UPoseSearchLibrary::MotionMatch(
@@ -556,7 +544,7 @@ void UPoseSearchLibrary::MotionMatch(
 				if (CVarAnimMotionMatchDrawQueryEnable.GetValueOnAnyThread())
 				{
 					UE::PoseSearch::FDebugDrawParams DrawParams(&UAnimInstanceProxyProvider::GetAnimInstanceProxy(AnimInstance), SearchResult.Database.Get(), EDebugDrawFlags::DrawQuery);
-					DrawParams.DrawFeatureVector(SearchResult.ComposedQuery.GetValues());
+					DrawParams.DrawFeatureVector(SearchContext.GetOrBuildQuery(SearchResult.Database->Schema).GetValues());
 				}
 			}
 #endif // ENABLE_DRAW_DEBUG && ENABLE_ANIM_DEBUG
diff --git a/Engine/Plugins/Experimental/Animation/PoseSearch/Source/Runtime/Private/PoseSearchResult.cpp b/Engine/Plugins/Experimental/Animation/PoseSearch/Source/Runtime/Private/PoseSearchResult.cpp
--- a/Engine/Plugins/Experimental/Animation/PoseSearch/Source/Runtime/Private/PoseSearchResult.cpp
+++ b/Engine/Plugins/Experimental/Animation/PoseSearch/Source/Runtime/Private/PoseSearchResult.cpp
@@ -83,7 +83,6 @@ void FSearchResult::Reset()
 {
 	PoseIdx = INDEX_NONE;
 	Database = nullptr;
-	ComposedQuery.Reset();
 	AssetTime = 0.0f;
 }
 
diff --git a/Engine/Plugins/Experimental/Animation/PoseSearch/Source/Runtime/Public/PoseSearch/AnimNode_MotionMatching.h b/Engine/Plugins/Experimental/Animation/PoseSearch/Source/Runtime/Public/PoseSearch/AnimNode_MotionMatching.h
--- a/Engine/Plugins/Experimental/Animation/PoseSearch/Source/Runtime/Public/PoseSearch/AnimNode_MotionMatching.h
+++ b/Engine/Plugins/Experimental/Animation/PoseSearch/Source/Runtime/Public/PoseSearch/AnimNode_MotionMatching.h
@@ -40,6 +40,10 @@ private:
 	UPROPERTY(EditAnywhere, Category = Settings, meta = (PinHiddenByDefault))
 	bool bForceInterrupt = false;
 
+	// If set to false, the motion matching node will perform a search only if the continuing pose is invalid. This is useful if you want to stagger searches of different nodes for performance reasons
+	UPROPERTY(EditAnywhere, Category = Settings, meta = (PinHiddenByDefault))
+	bool bShouldSearch = true;
+
 public:
 
 	// FAnimNode_Base interface
diff --git a/Engine/Plugins/Experimental/Animation/PoseSearch/Source/Runtime/Public/PoseSearch/PoseSearchContext.h b/Engine/Plugins/Experimental/Animation/PoseSearch/Source/Runtime/Public/PoseSearch/PoseSearchContext.h
--- a/Engine/Plugins/Experimental/Animation/PoseSearch/Source/Runtime/Public/PoseSearch/PoseSearchContext.h
+++ b/Engine/Plugins/Experimental/Animation/PoseSearch/Source/Runtime/Public/PoseSearch/PoseSearchContext.h
@@ -152,7 +152,7 @@ struct POSESEARCH_API FSearchContext
 	void UpdateCurrentBestCost(const FPoseSearchCost& PoseSearchCost);
 	float GetCurrentBestTotalCost() const { return CurrentBestTotalCost; }
 
-	void GetOrBuildQuery(const UPoseSearchSchema* Schema, FPoseSearchFeatureVectorBuilder& FeatureVectorBuilder);
+	const FPoseSearchFeatureVectorBuilder& GetOrBuildQuery(const UPoseSearchSchema* Schema);
 	const FPoseSearchFeatureVectorBuilder* GetCachedQuery(const UPoseSearchSchema* Schema) const;
 
 	bool IsCurrentResultFromDatabase(const UPoseSearchDatabase* Database) const;
diff --git a/Engine/Plugins/Experimental/Animation/PoseSearch/Source/Runtime/Public/PoseSearch/PoseSearchDatabase.h b/Engine/Plugins/Experimental/Animation/PoseSearch/Source/Runtime/Public/PoseSearch/PoseSearchDatabase.h
--- a/Engine/Plugins/Experimental/Animation/PoseSearch/Source/Runtime/Public/PoseSearch/PoseSearchDatabase.h
+++ b/Engine/Plugins/Experimental/Animation/PoseSearch/Source/Runtime/Public/PoseSearch/PoseSearchDatabase.h
@@ -4,6 +4,7 @@
 
 #include "Engine/DataAsset.h"
 #include "PoseSearch/PoseSearchAssetSampler.h"
+#include "PoseSearch/PoseSearchCost.h"
 #include "PoseSearch/PoseSearchIndex.h"
 #include "PoseSearch/PoseSearchResult.h"
 #include "PoseSearchDatabase.generated.h"
@@ -330,6 +331,7 @@ public:
 	// End UObject
 	
 	UE::PoseSearch::FSearchResult Search(UE::PoseSearch::FSearchContext& SearchContext) const;
+	FPoseSearchCost SearchContinuingPose(UE::PoseSearch::FSearchContext& SearchContext) const;
 
 #if WITH_EDITOR
 	virtual void BeginCacheForCookedPlatformData(const ITargetPlatform* TargetPlatform) override;
diff --git a/Engine/Plugins/Experimental/Animation/PoseSearch/Source/Runtime/Public/PoseSearch/PoseSearchLibrary.h b/Engine/Plugins/Experimental/Animation/PoseSearch/Source/Runtime/Public/PoseSearch/PoseSearchLibrary.h
--- a/Engine/Plugins/Experimental/Animation/PoseSearch/Source/Runtime/Public/PoseSearch/PoseSearchLibrary.h
+++ b/Engine/Plugins/Experimental/Animation/PoseSearch/Source/Runtime/Public/PoseSearch/PoseSearchLibrary.h
@@ -141,6 +141,9 @@ public:
 	* @param Settings						Input motion matching algorithm configuration settings
 	* @param InOutMotionMatchingState		Input/Output encapsulated motion matching algorithm and state
 	* @param bForceInterrupt				Input force interrupt request (if true the continuing pose will be invalidated)
+	* @param bShouldSearch					Input if false search will happen only if there's no valid continuing pose
+	* @param bDebugDrawQuery				Input draw the composed query if valid
+	* @param bDebugDrawCurResult			Input draw the current result if valid
 	*/
 	static void UpdateMotionMatchingState(
 		const FAnimationUpdateContext& Context,
@@ -148,7 +151,10 @@ public:
 		const FPoseSearchQueryTrajectory& Trajectory,
 		const FMotionMatchingSettings& Settings,
 		FMotionMatchingState& InOutMotionMatchingState,
-		bool bForceInterrupt);
+		bool bForceInterrupt = false,
+		bool bShouldSearch = true,
+		bool bDebugDrawQuery = false,
+		bool bDebugDrawCurResult = false);
 
 	/**
 	* Implementation of the core motion matching algorithm
diff --git a/Engine/Plugins/Experimental/Animation/PoseSearch/Source/Runtime/Public/PoseSearch/PoseSearchResult.h b/Engine/Plugins/Experimental/Animation/PoseSearch/Source/Runtime/Public/PoseSearch/PoseSearchResult.h
--- a/Engine/Plugins/Experimental/Animation/PoseSearch/Source/Runtime/Public/PoseSearch/PoseSearchResult.h
+++ b/Engine/Plugins/Experimental/Animation/PoseSearch/Source/Runtime/Public/PoseSearch/PoseSearchResult.h
@@ -49,7 +49,6 @@ struct FSearchResult
 	float LerpValue = 0.f;
 
 	TWeakObjectPtr<const UPoseSearchDatabase> Database;
-	FPoseSearchFeatureVectorBuilder ComposedQuery;
 
 	// cost of the current pose with the query from database in the result, if possible
 	FPoseSearchCost ContinuingPoseCost; 
