motion matching - Update Motion Matching assets to use Asset Definitions

diff --git a/Engine/Plugins/Experimental/Animation/PoseSearch/Source/Editor/PoseSearchEditor.Build.cs b/Engine/Plugins/Experimental/Animation/PoseSearch/Source/Editor/PoseSearchEditor.Build.cs
--- a/Engine/Plugins/Experimental/Animation/PoseSearch/Source/Editor/PoseSearchEditor.Build.cs
+++ b/Engine/Plugins/Experimental/Animation/PoseSearch/Source/Editor/PoseSearchEditor.Build.cs
@@ -41,6 +41,7 @@ public class PoseSearchEditor : ModuleRules
 		PrivateDependencyModuleNames.AddRange(
 			new string[]
 			{
+				"AssetDefinition",
 				"BlueprintGraph",
 				"PropertyEditor",
 				"SlateCore",
diff --git a/Engine/Plugins/Experimental/Animation/PoseSearch/Source/Editor/Private/PoseSearchAssetDefinitions.cpp b/Engine/Plugins/Experimental/Animation/PoseSearch/Source/Editor/Private/PoseSearchAssetDefinitions.cpp
new file mode 100644
--- /dev/null
+++ b/Engine/Plugins/Experimental/Animation/PoseSearch/Source/Editor/Private/PoseSearchAssetDefinitions.cpp
@@ -0,0 +1,45 @@
+// Copyright Epic Games, Inc. All Rights Reserved.
+
+#include "PoseSearchAssetDefinitions.h"
+#include "PoseSearchDatabaseEditor.h"
+#include "ThumbnailRendering/SceneThumbnailInfo.h"
+
+namespace UE::PoseSearch
+{
+
+FLinearColor GetAssetColor()
+{
+	static const FLinearColor AssetColor(FColor(29, 96, 125));
+	return AssetColor;
+}
+
+TConstArrayView<FAssetCategoryPath> GetAssetCategories()
+{
+	static const auto Categories = { FAssetCategoryPath(EAssetCategoryPaths::Animation, NSLOCTEXT("PoseSearchAssetDefinition", "PoseSearchAssetDefinitionMenu", "Motion Matching")) };
+	return Categories;
+}
+
+UThumbnailInfo* LoadThumbnailInfo(const FAssetData & InAssetData)
+{
+	return UE::Editor::FindOrCreateThumbnailInfo(InAssetData.GetAsset(), USceneThumbnailInfo::StaticClass());
+}
+
+} // namespace UE::PoseSearch
+
+EAssetCommandResult UAssetDefinition_PoseSearchDatabase::OpenAssets(const FAssetOpenArgs& OpenArgs) const
+{
+	using namespace UE::PoseSearch;
+
+	TArray<UPoseSearchDatabase*> Objects = OpenArgs.LoadObjects<UPoseSearchDatabase>();
+	const EToolkitMode::Type Mode = OpenArgs.ToolkitHost.IsValid() ? EToolkitMode::WorldCentric : EToolkitMode::Standalone;
+	for (UPoseSearchDatabase* Database : Objects)
+	{
+		if (Database)
+		{
+			TSharedRef<FDatabaseEditor> NewEditor(new FDatabaseEditor());
+			NewEditor->InitAssetEditor(Mode, OpenArgs.ToolkitHost, Database);
+		}
+	}
+	
+	return EAssetCommandResult::Handled;
+}
diff --git a/Engine/Plugins/Experimental/Animation/PoseSearch/Source/Editor/Private/PoseSearchAssetDefinitions.h b/Engine/Plugins/Experimental/Animation/PoseSearch/Source/Editor/Private/PoseSearchAssetDefinitions.h
new file mode 100644
--- /dev/null
+++ b/Engine/Plugins/Experimental/Animation/PoseSearch/Source/Editor/Private/PoseSearchAssetDefinitions.h
@@ -0,0 +1,61 @@
+// Copyright Epic Games, Inc. All Rights Reserved.
+
+#pragma once
+
+#include "AssetDefinitionDefault.h"
+#include "PoseSearch/PoseSearchDatabase.h"
+#include "PoseSearch/PoseSearchNormalizationSet.h"
+#include "PoseSearch/PoseSearchSchema.h"
+#include "PoseSearchAssetDefinitions.generated.h"
+
+namespace UE::PoseSearch
+{
+	FLinearColor GetAssetColor();
+	TConstArrayView<FAssetCategoryPath> GetAssetCategories();
+	UThumbnailInfo* LoadThumbnailInfo(const FAssetData& InAssetData);
+} // namespace UE::PoseSearch
+
+UCLASS()
+class UAssetDefinition_PoseSearchDatabase : public UAssetDefinitionDefault
+{
+	GENERATED_BODY()
+
+public:
+	virtual FLinearColor GetAssetColor() const override { return UE::PoseSearch::GetAssetColor(); }
+	virtual TConstArrayView<FAssetCategoryPath> GetAssetCategories() const override { return UE::PoseSearch::GetAssetCategories(); }
+	virtual UThumbnailInfo* LoadThumbnailInfo(const FAssetData& InAssetData) const override { return UE::PoseSearch::LoadThumbnailInfo(InAssetData); }
+
+	virtual FText GetAssetDisplayName() const override { return NSLOCTEXT("PoseSearchAssetDefinition", "DisplayName_UPoseSearchDatabase", "Motion Database"); }
+	virtual TSoftClassPtr<UObject> GetAssetClass() const override { return UPoseSearchDatabase::StaticClass(); }
+	virtual EAssetCommandResult OpenAssets(const FAssetOpenArgs& OpenArgs) const override;
+};
+
+UCLASS()
+class UAssetDefinition_PoseSearchSchema : public UAssetDefinitionDefault
+{
+	GENERATED_BODY()
+
+public:
+	virtual FLinearColor GetAssetColor() const override { return UE::PoseSearch::GetAssetColor(); }
+	virtual TConstArrayView<FAssetCategoryPath> GetAssetCategories() const override { return UE::PoseSearch::GetAssetCategories(); }
+	virtual UThumbnailInfo* LoadThumbnailInfo(const FAssetData& InAssetData) const override { return UE::PoseSearch::LoadThumbnailInfo(InAssetData); }
+
+	virtual FText GetAssetDisplayName() const override { return NSLOCTEXT("PoseSearchAssetDefinition", "DisplayName_UPoseSearchSchema", "Motion Database Config"); }
+	virtual TSoftClassPtr<UObject> GetAssetClass() const override { return UPoseSearchSchema::StaticClass(); }
+};
+
+UCLASS()
+class UAssetDefinition_PoseSearchNormalizationSet : public UAssetDefinitionDefault
+{
+	GENERATED_BODY()
+
+public:
+	virtual FLinearColor GetAssetColor() const override { return UE::PoseSearch::GetAssetColor(); }
+	virtual TConstArrayView<FAssetCategoryPath> GetAssetCategories() const override { return UE::PoseSearch::GetAssetCategories(); }
+	virtual UThumbnailInfo* LoadThumbnailInfo(const FAssetData& InAssetData) const override { return UE::PoseSearch::LoadThumbnailInfo(InAssetData); }
+
+	virtual FText GetAssetDisplayName() const override { return NSLOCTEXT("PoseSearchAssetDefinition", "DisplayName_UPoseSearchNormalizationSet", "Normalization Set"); }
+	virtual TSoftClassPtr<UObject> GetAssetClass() const override { return UPoseSearchNormalizationSet::StaticClass(); }
+};
+
+
diff --git a/Engine/Plugins/Experimental/Animation/PoseSearch/Source/Editor/Private/PoseSearchEditor.cpp b/Engine/Plugins/Experimental/Animation/PoseSearch/Source/Editor/Private/PoseSearchEditor.cpp
--- a/Engine/Plugins/Experimental/Animation/PoseSearch/Source/Editor/Private/PoseSearchEditor.cpp
+++ b/Engine/Plugins/Experimental/Animation/PoseSearch/Source/Editor/Private/PoseSearchEditor.cpp
@@ -3,7 +3,6 @@
 #include "PoseSearchEditor.h"
 #include "PoseSearchCustomization.h"
 #include "PoseSearchDebugger.h"
-#include "PoseSearchTypeActions.h"
 #include "PoseSearchDatabaseEdMode.h"
 #include "PoseSearchDatabaseEditorCommands.h"
 
@@ -53,10 +52,6 @@ private:
 	/** List of registered class that we must unregister when the module shuts down */
 	TSet<FName> RegisteredClassNames;
 	TSet<FName> RegisteredPropertyTypes;
-
-	TSharedPtr<IAssetTypeActions> PoseSearchDatabaseActions;
-	TSharedPtr<IAssetTypeActions> PoseSearchSchemaActions;
-	TSharedPtr<IAssetTypeActions> PoseSearchNormalizationSetActions;
 };
 
 void FEditorModule::StartupModule()
@@ -75,19 +70,6 @@ void FEditorModule::StartupModule()
 		
 		// Register Ed Mode used by pose search database
 		FEditorModeRegistry::Get().RegisterMode<FDatabaseEdMode>(FDatabaseEdMode::EdModeId, LOCTEXT("PoseSearchDatabaseEdModeName", "PoseSearchDatabase"));
-
-		// Register UPoseSearchDatabase Type Actions 
-		IAssetTools& AssetTools = FModuleManager::LoadModuleChecked<FAssetToolsModule>("AssetTools").Get();
-		PoseSearchDatabaseActions = MakeShared<FDatabaseTypeActions>();
-		AssetTools.RegisterAssetTypeActions(PoseSearchDatabaseActions.ToSharedRef());
-
-		// Register UPoseSearchSchema Type Actions 
-		PoseSearchSchemaActions = MakeShared<FSchemaTypeActions>();
-		AssetTools.RegisterAssetTypeActions(PoseSearchSchemaActions.ToSharedRef());
-
-		// Register UPoseSearchSchema Type Actions 
-		PoseSearchNormalizationSetActions = MakeShared<FNormalizationSetTypeActions>();
-		AssetTools.RegisterAssetTypeActions(PoseSearchNormalizationSetActions.ToSharedRef());
 	}
 
 	RegisterPropertyTypeCustomizations();
@@ -105,14 +87,6 @@ void FEditorModule::ShutdownModule()
 	}
 	ConsoleCommands.Empty();
 
-	if (FModuleManager::Get().IsModuleLoaded("AssetTools"))
-	{
-		IAssetTools& AssetTools = FModuleManager::GetModuleChecked<FAssetToolsModule>("AssetTools").Get();
-		AssetTools.UnregisterAssetTypeActions(PoseSearchDatabaseActions.ToSharedRef());
-		AssetTools.UnregisterAssetTypeActions(PoseSearchSchemaActions.ToSharedRef());
-		AssetTools.UnregisterAssetTypeActions(PoseSearchNormalizationSetActions.ToSharedRef());
-	}
-
 	// Unregister Ed Mode
 	FEditorModeRegistry::Get().UnregisterMode(FDatabaseEdMode::EdModeId);
 	
diff --git a/Engine/Plugins/Experimental/Animation/PoseSearch/Source/Editor/Private/PoseSearchTypeActions.cpp b/Engine/Plugins/Experimental/Animation/PoseSearch/Source/Editor/Private/PoseSearchTypeActions.cpp
deleted file mode 100644
diff --git a/Engine/Plugins/Experimental/Animation/PoseSearch/Source/Editor/Public/PoseSearchTypeActions.h b/Engine/Plugins/Experimental/Animation/PoseSearch/Source/Editor/Public/PoseSearchTypeActions.h
deleted file mode 100644
