motion matching - show pose search debugger query in database editor

diff --git a/Engine/Plugins/Experimental/Animation/PoseSearch/Source/Editor/Private/PoseSearchDatabaseEditor.cpp b/Engine/Plugins/Experimental/Animation/PoseSearch/Source/Editor/Private/PoseSearchDatabaseEditor.cpp
--- a/Engine/Plugins/Experimental/Animation/PoseSearch/Source/Editor/Private/PoseSearchDatabaseEditor.cpp
+++ b/Engine/Plugins/Experimental/Animation/PoseSearch/Source/Editor/Private/PoseSearchDatabaseEditor.cpp
@@ -65,6 +65,22 @@ namespace UE::PoseSearch
 		}
 	}
 
+	void FDatabaseEditor::SetQueryVector(TConstArrayView<float> InQueryVector)
+	{
+		if (ViewModel.IsValid())
+		{
+			ViewModel->SetQueryVector(InQueryVector);
+		}
+	}
+
+	void FDatabaseEditor::SetDrawQueryVector(bool bValue)
+	{
+		if (ViewModel.IsValid())
+		{
+			ViewModel->SetDrawQueryVector(bValue);
+		}
+	}
+
 	void FDatabaseEditor::BuildSearchIndex()
 	{
 		ViewModel->BuildSearchIndex();
diff --git a/Engine/Plugins/Experimental/Animation/PoseSearch/Source/Editor/Private/PoseSearchDatabasePreviewScene.cpp b/Engine/Plugins/Experimental/Animation/PoseSearch/Source/Editor/Private/PoseSearchDatabasePreviewScene.cpp
--- a/Engine/Plugins/Experimental/Animation/PoseSearch/Source/Editor/Private/PoseSearchDatabasePreviewScene.cpp
+++ b/Engine/Plugins/Experimental/Animation/PoseSearch/Source/Editor/Private/PoseSearchDatabasePreviewScene.cpp
@@ -61,6 +61,7 @@ namespace UE::PoseSearch
 		if (ViewModel->IsPoseFeaturesDrawMode(EFeaturesDrawMode::All | EFeaturesDrawMode::Detailed) && !ViewModel->GetPreviewActors().IsEmpty() &&
 			FAsyncPoseSearchDatabasesManagement::RequestAsyncBuildIndex(Database, ERequestAsyncBuildFlag::ContinueRequest))
 		{
+			bool bDrawQueryVector = ViewModel->ShouldDrawQueryVector();
 			for (FDatabasePreviewActor& PreviewActor : ViewModel->GetPreviewActors())
 			{
 				if (Database->GetSearchIndex().IsValidPoseIndex(PreviewActor.CurrentPoseIndex))
@@ -69,6 +70,12 @@ namespace UE::PoseSearch
 					{
 						UE::PoseSearch::FDebugDrawParams DrawParams(GetWorld(), Mesh, &PreviewActor.QuantizedTimeRootMotion, Database);
 						DrawParams.DrawFeatureVector(PreviewActor.CurrentPoseIndex);
+
+						if (bDrawQueryVector)
+						{
+							DrawParams.DrawFeatureVector(ViewModel->GetQueryVector());
+							bDrawQueryVector = false;
+						}
 					}
 				}
 			}
diff --git a/Engine/Plugins/Experimental/Animation/PoseSearch/Source/Editor/Private/PoseSearchDatabaseViewModel.cpp b/Engine/Plugins/Experimental/Animation/PoseSearch/Source/Editor/Private/PoseSearchDatabaseViewModel.cpp
--- a/Engine/Plugins/Experimental/Animation/PoseSearch/Source/Editor/Private/PoseSearchDatabaseViewModel.cpp
+++ b/Engine/Plugins/Experimental/Animation/PoseSearch/Source/Editor/Private/PoseSearchDatabaseViewModel.cpp
@@ -273,6 +273,7 @@ namespace UE::PoseSearch
 		MaxPreviewPlayLength = 0.f;
 		MinPreviewPlayLength = 0.f;
 		bIsEditorSelection = true;
+		bDrawQueryVector = false;
 
 		for (auto PreviewActor : PreviewActors)
 		{
@@ -430,8 +431,6 @@ namespace UE::PoseSearch
 
 	void FDatabaseViewModel::SetSelectedNodes(const TArrayView<TSharedPtr<FDatabaseAssetTreeNode>>& InSelectedNodes)
 	{
-		bIsEditorSelection = true;
-
 		RemovePreviewActors();
 
 		if (FAsyncPoseSearchDatabasesManagement::RequestAsyncBuildIndex(PoseSearchDatabase, ERequestAsyncBuildFlag::ContinueRequest))
diff --git a/Engine/Plugins/Experimental/Animation/PoseSearch/Source/Editor/Private/PoseSearchDebuggerDatabaseColumns.h b/Engine/Plugins/Experimental/Animation/PoseSearch/Source/Editor/Private/PoseSearchDebuggerDatabaseColumns.h
--- a/Engine/Plugins/Experimental/Animation/PoseSearch/Source/Editor/Private/PoseSearchDebuggerDatabaseColumns.h
+++ b/Engine/Plugins/Experimental/Animation/PoseSearch/Source/Editor/Private/PoseSearchDebuggerDatabaseColumns.h
@@ -117,36 +117,36 @@ struct FDatabaseName : IColumn
 
 	virtual FSortPredicate GetSortPredicate() const override
 	{
-		return [](const FRowDataRef& Row0, const FRowDataRef& Row1) -> bool { return Row0->DatabaseName < Row1->DatabaseName; };
+		return [](const FRowDataRef& Row0, const FRowDataRef& Row1) -> bool { return Row0->SharedData->DatabaseName < Row1->SharedData->DatabaseName; };
 	}
 
 	virtual TSharedRef<SWidget> GenerateWidget(const FRowDataRef& RowData) const override
 	{
 		return SNew(SHyperlink)
-			.Text_Lambda([RowData]() -> FText { return FText::FromString(RowData->DatabaseName); })
+			.Text_Lambda([RowData]() -> FText { return FText::FromString(RowData->SharedData->DatabaseName); })
 			.TextStyle(&FCoreStyle::Get().GetWidgetStyle<FTextBlockStyle>("SmallText"))
 			.ToolTipText_Lambda([RowData]() -> FText
 				{
 					return FText::Format(
 						LOCTEXT("DatabaseHyperlinkTooltipFormat", "Open database '{0}'"),
-						FText::FromString(RowData->DatabasePath));
+						FText::FromString(RowData->SharedData->DatabasePath));
 				})
 			.OnNavigate_Lambda([RowData]()
 				{
 					UObject* Asset = nullptr;
 
 					// Load asset
-					if (UPackage* Package = LoadPackage(NULL, *RowData->DatabasePath, LOAD_NoRedirects))
+					if (UPackage* Package = LoadPackage(NULL, *RowData->SharedData->DatabasePath, LOAD_NoRedirects))
 					{
 						Package->FullyLoad();
 
-						const FString AssetName = FPaths::GetBaseFilename(RowData->DatabasePath);
+						const FString AssetName = FPaths::GetBaseFilename(RowData->SharedData->DatabasePath);
 						Asset = FindObject<UObject>(Package, *AssetName);
 					}
 					else
 					{
 						// Fallback for unsaved assets
-						Asset = FindObject<UObject>(nullptr, *RowData->DatabasePath);
+						Asset = FindObject<UObject>(nullptr, *RowData->SharedData->DatabasePath);
 					}
 
 					// Open editor
@@ -164,6 +164,7 @@ struct FDatabaseName : IColumn
 									
 									// Open asset paused and at specific time as seen on the pose search debugger.
 									DatabaseEditor->SetSelectedPoseIdx(RowData->PoseIdx);
+									DatabaseEditor->SetQueryVector(RowData->SharedData->QueryVector);
 								}
 							}
 						}
diff --git a/Engine/Plugins/Experimental/Animation/PoseSearch/Source/Editor/Private/PoseSearchDebuggerDatabaseRowData.h b/Engine/Plugins/Experimental/Animation/PoseSearch/Source/Editor/Private/PoseSearchDebuggerDatabaseRowData.h
--- a/Engine/Plugins/Experimental/Animation/PoseSearch/Source/Editor/Private/PoseSearchDebuggerDatabaseRowData.h
+++ b/Engine/Plugins/Experimental/Animation/PoseSearch/Source/Editor/Private/PoseSearchDebuggerDatabaseRowData.h
@@ -9,16 +9,22 @@
 namespace UE::PoseSearch
 {
 
-class FDebuggerDatabaseRowData : public TSharedFromThis<FDebuggerDatabaseRowData>
+class FDebuggerDatabaseSharedData : public TSharedFromThis<FDebuggerDatabaseSharedData>
 {
 public:
-	FDebuggerDatabaseRowData() = default;
-	
-	int32 PoseIdx = 0;
 	TWeakObjectPtr<const UPoseSearchDatabase> SourceDatabase = nullptr;
-	EPoseCandidateFlags PoseCandidateFlags = EPoseCandidateFlags::None;
 	FString DatabaseName = "";
 	FString DatabasePath = "";
+	TArray<float> QueryVector;
+};
+
+class FDebuggerDatabaseRowData : public TSharedFromThis<FDebuggerDatabaseRowData>
+{
+public:
+	explicit FDebuggerDatabaseRowData(TSharedRef<FDebuggerDatabaseSharedData> InSharedData) : SharedData(InSharedData) {}
+
+	int32 PoseIdx = 0;
+	EPoseCandidateFlags PoseCandidateFlags = EPoseCandidateFlags::None;
 	FString AssetName = "";
 	FString AssetPath = "";
 	int32 DbAssetIdx = 0;
@@ -33,6 +39,8 @@ public:
 	TArray<float> CostBreakdowns;
 	TArray<FLinearColor> CostBreakdownsColors;
 	TArray<float> CostVector;
+
+	TSharedRef<FDebuggerDatabaseSharedData> SharedData;
 };
 
 } // namespace UE::PoseSearch
diff --git a/Engine/Plugins/Experimental/Animation/PoseSearch/Source/Editor/Private/PoseSearchDebuggerDatabaseView.cpp b/Engine/Plugins/Experimental/Animation/PoseSearch/Source/Editor/Private/PoseSearchDebuggerDatabaseView.cpp
--- a/Engine/Plugins/Experimental/Animation/PoseSearch/Source/Editor/Private/PoseSearchDebuggerDatabaseView.cpp
+++ b/Engine/Plugins/Experimental/Animation/PoseSearch/Source/Editor/Private/PoseSearchDebuggerDatabaseView.cpp
@@ -64,7 +64,7 @@ void SDebuggerDatabaseView::Update(const FTraceMotionMatchingStateMessage& State
 					for (const UPoseSearchFeatureChannel* Channel : LabelToChannel.Channels)
 					{
 						// checking if the row is associated to the Channel
-						if (UnfilteredDatabaseRow.SourceDatabase->Schema == Channel->GetSchema())
+						if (UnfilteredDatabaseRow.SharedData->SourceDatabase->Schema == Channel->GetSchema())
 						{
 							CostBreakdown += ArraySum(UnfilteredDatabaseRow.CostVector, Channel->GetChannelDataOffset(), Channel->GetChannelCardinality());
 						}
@@ -175,18 +175,22 @@ void SDebuggerDatabaseView::Update(const FTraceMotionMatchingStateMessage& State
 		if (FAsyncPoseSearchDatabasesManagement::RequestAsyncBuildIndex(Database, ERequestAsyncBuildFlag::ContinueRequest))
 		{
 			const FPoseSearchIndex& SearchIndex = Database->GetSearchIndex();
+
+			TSharedRef<FDebuggerDatabaseSharedData> SharedData = MakeShared<FDebuggerDatabaseSharedData>();
+			SharedData->SourceDatabase = Database;
+			SharedData->DatabaseName = Database->GetName();
+			SharedData->DatabasePath = Database->GetPathName();
+			SharedData->QueryVector = DbEntry.QueryVector;
+
 			for (const FTraceMotionMatchingStatePoseEntry& PoseEntry : DbEntry.PoseEntries)
 			{
 				if (const FPoseSearchIndexAsset* SearchIndexAsset = SearchIndex.GetAssetForPoseSafe(PoseEntry.DbPoseIdx))
 				{
-					TSharedRef<FDebuggerDatabaseRowData>& Row = UnfilteredDatabaseRows.Add_GetRef(MakeShared<FDebuggerDatabaseRowData>());
+					TSharedRef<FDebuggerDatabaseRowData>& Row = UnfilteredDatabaseRows.Add_GetRef(MakeShared<FDebuggerDatabaseRowData>(SharedData));
 
 					const float Time = Database->GetAssetTime(PoseEntry.DbPoseIdx);
 
 					Row->PoseIdx = PoseEntry.DbPoseIdx;
-					Row->SourceDatabase = Database;
-					Row->DatabaseName = Database->GetName();
-					Row->DatabasePath = Database->GetPathName();
 					Row->PoseCandidateFlags = PoseEntry.PoseCandidateFlags;
 					Row->DbAssetIdx = SearchIndexAsset->SourceAssetIdx;
 					Row->AssetTime = Time;
@@ -482,7 +486,7 @@ void SDebuggerDatabaseView::OnDatabaseRowSelectionChanged(TSharedPtr<FDebuggerDa
 {
 	if (Row.IsValid())
 	{
-		OnPoseSelectionChanged.ExecuteIfBound(Row->SourceDatabase.Get(), Row->PoseIdx, Row->AssetTime);
+		OnPoseSelectionChanged.ExecuteIfBound(Row->SharedData->SourceDatabase.Get(), Row->PoseIdx, Row->AssetTime);
 	}
 }
 
diff --git a/Engine/Plugins/Experimental/Animation/PoseSearch/Source/Editor/Private/PoseSearchDebuggerView.cpp b/Engine/Plugins/Experimental/Animation/PoseSearch/Source/Editor/Private/PoseSearchDebuggerView.cpp
--- a/Engine/Plugins/Experimental/Animation/PoseSearch/Source/Editor/Private/PoseSearchDebuggerView.cpp
+++ b/Engine/Plugins/Experimental/Animation/PoseSearch/Source/Editor/Private/PoseSearchDebuggerView.cpp
@@ -2,15 +2,18 @@
 
 #include "PoseSearchDebuggerView.h"
 #include "Modules/ModuleManager.h"
+#include "Editor.h"
 #include "PoseSearch/PoseSearchDatabase.h"
 #include "PoseSearch/PoseSearchDerivedData.h"
 #include "PoseSearch/PoseSearchSchema.h"
+#include "PoseSearchDatabaseEditor.h"
 #include "PoseSearchDebugger.h"
 #include "PoseSearchDebuggerDatabaseRowData.h"
 #include "PoseSearchDebuggerDatabaseView.h"
 #include "PoseSearchDebuggerReflection.h"
 #include "PoseSearchDebuggerViewModel.h"
 #include "PropertyEditorModule.h"
+#include "Subsystems/AssetEditorSubsystem.h"
 #include "Trace/PoseSearchTraceProvider.h"
 #include "Widgets/Input/SButton.h"
 #include "Widgets/Input/SNumericEntryBox.h"
@@ -127,9 +130,9 @@ void SDebuggerDetailsView::UpdateReflection(const FTraceMotionMatchingStateMessa
 		if (!SelectedRows.IsEmpty())
 		{
 			const TSharedRef<FDebuggerDatabaseRowData>& Selected = SelectedRows[0];
-			if (FAsyncPoseSearchDatabasesManagement::RequestAsyncBuildIndex(Selected->SourceDatabase.Get(), ERequestAsyncBuildFlag::ContinueRequest))
+			if (FAsyncPoseSearchDatabasesManagement::RequestAsyncBuildIndex(Selected->SharedData->SourceDatabase.Get(), ERequestAsyncBuildFlag::ContinueRequest))
 			{
-				const FPoseSearchIndex& SelectedSearchIndex = Selected->SourceDatabase->GetSearchIndex();
+				const FPoseSearchIndex& SelectedSearchIndex = Selected->SharedData->SourceDatabase->GetSearchIndex();
 				Reflection->SelectedPoseVector = SelectedSearchIndex.GetPoseValuesSafe(Selected->PoseIdx);
 			}
 			Reflection->CostVector = Selected->CostVector;
@@ -267,6 +270,27 @@ void SDebuggerView::Tick(const FGeometry& AllottedGeometry, const double InCurre
 
 	// Draw visualization every tick
 	DrawVisualization();
+	
+	// synchronizing the model DrawQuery state with all the open PoseSearchDatabaseEditor(s)
+	const bool bDrawQuery = Model->GetDrawQuery();
+	if (UAssetEditorSubsystem* AssetEditorSS = GEditor->GetEditorSubsystem<UAssetEditorSubsystem>())
+	{
+		TArray<UObject*> EditedAssets = AssetEditorSS->GetAllEditedAssets();
+		for (UObject* EditedAsset : EditedAssets)
+		{
+			if (Cast<UPoseSearchDatabase>(EditedAsset))
+			{
+				if (IAssetEditorInstance* Editor = AssetEditorSS->FindEditorForAsset(EditedAsset, false))
+				{
+					if (Editor->GetEditorName() == FName("PoseSearchDatabaseEditor"))
+					{
+						FDatabaseEditor* DatabaseEditor = static_cast<FDatabaseEditor*>(Editor);
+						DatabaseEditor->SetDrawQueryVector(bDrawQuery);
+					}
+				}
+			}
+		}
+	}
 }
 
 bool SDebuggerView::UpdateNodeSelection()
@@ -371,7 +395,7 @@ void SDebuggerView::DrawFeatures(
 	// Draw any selected database vectors
 	for (const TSharedRef<FDebuggerDatabaseRowData>& Row : SelectedRows)
 	{
-		const UPoseSearchDatabase* RowDatabase = Row->SourceDatabase.Get();
+		const UPoseSearchDatabase* RowDatabase = Row->SharedData->SourceDatabase.Get();
 		if (FAsyncPoseSearchDatabasesManagement::RequestAsyncBuildIndex(RowDatabase, ERequestAsyncBuildFlag::ContinueRequest))
 		{
 			FDebugDrawParams DrawParams(&DebuggerWorld, Mesh, nullptr, RowDatabase);
@@ -388,7 +412,7 @@ void SDebuggerView::DrawFeatures(
 
 		if (!ActiveRows.IsEmpty())
 		{
-			const UPoseSearchDatabase* Database = ActiveRows[0]->SourceDatabase.Get();
+			const UPoseSearchDatabase* Database = ActiveRows[0]->SharedData->SourceDatabase.Get();
 			if (Database && FAsyncPoseSearchDatabasesManagement::RequestAsyncBuildIndex(Database, ERequestAsyncBuildFlag::ContinueRequest))
 			{
 				// Use the motion-matching state's pose idx, as the active row may be update-throttled at this point
@@ -408,7 +432,7 @@ void SDebuggerView::DrawFeatures(
 
 		if (!ContinuingRows.IsEmpty())
 		{
-			const UPoseSearchDatabase* Database = ContinuingRows[0]->SourceDatabase.Get();
+			const UPoseSearchDatabase* Database = ContinuingRows[0]->SharedData->SourceDatabase.Get();
 			if (Database && FAsyncPoseSearchDatabasesManagement::RequestAsyncBuildIndex(Database, ERequestAsyncBuildFlag::ContinueRequest))
 			{
 				FDebugDrawParams DrawParams(&DebuggerWorld, Mesh, nullptr, Database);
diff --git a/Engine/Plugins/Experimental/Animation/PoseSearch/Source/Editor/Public/PoseSearchDatabaseEditor.h b/Engine/Plugins/Experimental/Animation/PoseSearch/Source/Editor/Public/PoseSearchDatabaseEditor.h
--- a/Engine/Plugins/Experimental/Animation/PoseSearch/Source/Editor/Public/PoseSearchDatabaseEditor.h
+++ b/Engine/Plugins/Experimental/Animation/PoseSearch/Source/Editor/Public/PoseSearchDatabaseEditor.h
@@ -41,6 +41,8 @@ namespace UE::PoseSearch
 		TSharedPtr<FDatabaseViewModel> GetViewModelSharedPtr() const { return ViewModel; }
 		
 		void SetSelectedPoseIdx(int32 PoseIdx);
+		void SetQueryVector(TConstArrayView<float> InQueryVector);
+		void SetDrawQueryVector(bool bValue);
 		
 		void BuildSearchIndex();
 
diff --git a/Engine/Plugins/Experimental/Animation/PoseSearch/Source/Editor/Public/PoseSearchDatabaseViewModel.h b/Engine/Plugins/Experimental/Animation/PoseSearch/Source/Editor/Public/PoseSearchDatabaseViewModel.h
--- a/Engine/Plugins/Experimental/Animation/PoseSearch/Source/Editor/Public/PoseSearchDatabaseViewModel.h
+++ b/Engine/Plugins/Experimental/Animation/PoseSearch/Source/Editor/Public/PoseSearchDatabaseViewModel.h
@@ -109,6 +109,11 @@ namespace UE::PoseSearch
 		void SetSelectedNodes(const TArrayView<TSharedPtr<FDatabaseAssetTreeNode>>& InSelectedNodes);
 		void ProcessSelectedActor(AActor* Actor);
 		
+		void SetQueryVector(TConstArrayView<float> InQueryVector) { QueryVector = InQueryVector; }
+		TConstArrayView<float> GetQueryVector() const { return QueryVector; }
+		void SetDrawQueryVector(bool bValue) { bDrawQueryVector = bValue; }
+		bool ShouldDrawQueryVector() const { return bDrawQueryVector && !bIsEditorSelection; }
+
 		const FPoseSearchIndexAsset* GetSelectedActorIndexAsset() const;
 
 		TRange<double> GetPreviewPlayRange() const;
@@ -118,6 +123,7 @@ namespace UE::PoseSearch
 		bool IsEditorSelection() const { return bIsEditorSelection; }
 		bool GetAnimationTime(int32 SourceAssetIdx, float& CurrentPlayTime, FVector& BlendParameters) const;
 
+
 	private:
 		float PlayTime = 0.f;
 		float DeltaTimeMultiplier = 1.f;
@@ -134,7 +140,10 @@ namespace UE::PoseSearch
 		/** From zero to the play length of the longest preview */
 		float MaxPreviewPlayLength = 0.f;
 		float MinPreviewPlayLength = 0.f;
+
 		bool bIsEditorSelection = true;
+		bool bDrawQueryVector = false;
+		TArray<float> QueryVector;
 
 		/** What features to show in the viewport */
 		EFeaturesDrawMode PoseFeaturesDrawMode = EFeaturesDrawMode::All;
