motion matching - removing database editor build index button (unnecessary after MM DDC lazy indexing implementation)

diff --git a/Engine/Plugins/Experimental/Animation/PoseSearch/Source/Editor/Private/PoseSearchDatabaseEditor.cpp b/Engine/Plugins/Experimental/Animation/PoseSearch/Source/Editor/Private/PoseSearchDatabaseEditor.cpp
--- a/Engine/Plugins/Experimental/Animation/PoseSearch/Source/Editor/Private/PoseSearchDatabaseEditor.cpp
+++ b/Engine/Plugins/Experimental/Animation/PoseSearch/Source/Editor/Private/PoseSearchDatabaseEditor.cpp
@@ -81,11 +81,6 @@ namespace UE::PoseSearch
 		}
 	}
 
-	void FDatabaseEditor::BuildSearchIndex()
-	{
-		ViewModel->BuildSearchIndex();
-	}
-
 	void FDatabaseEditor::PreviewBackwardEnd()
 	{
 		ViewModel->PreviewBackwardEnd();
@@ -126,9 +121,6 @@ namespace UE::PoseSearch
 		const TSharedPtr<IToolkitHost>& InitToolkitHost,
 		UPoseSearchDatabase* DatabaseAsset)
 	{
-		// Bind Commands
-		BindCommands();
-
 		// Create Preview Scene
 		if (!PreviewScene.IsValid())
 		{
@@ -288,44 +280,9 @@ namespace UE::PoseSearch
 			DatabaseAsset,
 			bIsToolbarFocusableParam);
 
-		ExtendToolbar();
-
 		RegenerateMenusAndToolbars();
 	}
 
-	void FDatabaseEditor::BindCommands()
-	{
-		const FDatabaseEditorCommands& Commands = FDatabaseEditorCommands::Get();
-
-		ToolkitCommands->MapAction(
-			Commands.BuildSearchIndex,
-			FExecuteAction::CreateSP(this, &FDatabaseEditor::BuildSearchIndex),
-			EUIActionRepeatMode::RepeatDisabled);
-	}
-
-	void FDatabaseEditor::ExtendToolbar()
-	{
-		TSharedPtr<FExtender> ToolbarExtender = MakeShareable(new FExtender);
-
-		AddToolbarExtender(ToolbarExtender);
-
-		ToolbarExtender->AddToolBarExtension(
-			"Asset",
-			EExtensionHook::After,
-			GetToolkitCommands(),
-			FToolBarExtensionDelegate::CreateSP(this, &FDatabaseEditor::FillToolbar));
-	}
-
-	void FDatabaseEditor::FillToolbar(FToolBarBuilder& ToolbarBuilder)
-	{
-		ToolbarBuilder.AddToolBarButton(
-			FDatabaseEditorCommands::Get().BuildSearchIndex,
-			NAME_None,
-			TAttribute<FText>(),
-			TAttribute<FText>(),
-			FSlateIcon());
-	}
-
 	void FDatabaseEditor::RegisterTabSpawners(const TSharedRef<FTabManager>& InTabManager)
 	{
 		WorkspaceMenuCategory = InTabManager->AddLocalWorkspaceMenuCategory(
diff --git a/Engine/Plugins/Experimental/Animation/PoseSearch/Source/Editor/Private/PoseSearchDatabaseEditorCommands.cpp b/Engine/Plugins/Experimental/Animation/PoseSearch/Source/Editor/Private/PoseSearchDatabaseEditorCommands.cpp
--- a/Engine/Plugins/Experimental/Animation/PoseSearch/Source/Editor/Private/PoseSearchDatabaseEditorCommands.cpp
+++ b/Engine/Plugins/Experimental/Animation/PoseSearch/Source/Editor/Private/PoseSearchDatabaseEditorCommands.cpp
@@ -8,8 +8,6 @@ namespace UE::PoseSearch
 {
 	void FDatabaseEditorCommands::RegisterCommands()
 	{
-		UI_COMMAND(BuildSearchIndex, "Build Index", "Build Index", EUserInterfaceActionType::Button, FInputChord());
-
 		UI_COMMAND(
 			ShowPoseFeaturesNone, "None", "Hide all features",
 			EUserInterfaceActionType::RadioButton, FInputChord());
@@ -27,7 +25,9 @@ namespace UE::PoseSearch
 			ShowAnimationOriginalAndMirrored, "Original and mirrored", "Show original and mirrored animations",
 			EUserInterfaceActionType::RadioButton, FInputChord());
 
-		UI_COMMAND(ShowDisplayRootMotionSpeed, "Display Root Motion Speed", "Show root motion speed for the selected item", EUserInterfaceActionType::ToggleButton, FInputChord());
+		UI_COMMAND(
+			ShowDisplayRootMotionSpeed, "Display Root Motion Speed", "Show root motion speed for the selected item",
+			EUserInterfaceActionType::ToggleButton, FInputChord());
 	}
 }
 
diff --git a/Engine/Plugins/Experimental/Animation/PoseSearch/Source/Editor/Public/PoseSearchDatabaseEditor.h b/Engine/Plugins/Experimental/Animation/PoseSearch/Source/Editor/Public/PoseSearchDatabaseEditor.h
--- a/Engine/Plugins/Experimental/Animation/PoseSearch/Source/Editor/Public/PoseSearchDatabaseEditor.h
+++ b/Engine/Plugins/Experimental/Animation/PoseSearch/Source/Editor/Public/PoseSearchDatabaseEditor.h
@@ -44,8 +44,6 @@ namespace UE::PoseSearch
 		void SetQueryVector(TConstArrayView<float> InQueryVector);
 		void SetDrawQueryVector(bool bValue);
 		
-		void BuildSearchIndex();
-
 		void PreviewBackwardEnd();
 		void PreviewBackwardStep();
 		void PreviewBackward();
@@ -63,10 +61,6 @@ namespace UE::PoseSearch
 		TSharedRef<SDockTab> SpawnTab_SelectionDetails(const FSpawnTabArgs& Args);
 		TSharedRef<SDockTab> SpawnTab_StatisticsOverview(const FSpawnTabArgs& Args) const;
 		
-		void BindCommands();
-		void ExtendToolbar();
-		void FillToolbar(FToolBarBuilder& ToolbarBuilder);
-
 		void OnFinishedChangingSelectionProperties(const FPropertyChangedEvent& PropertyChangedEvent);
 		void OnAssetTreeSelectionChanged(
 			const TArrayView<TSharedPtr<FDatabaseAssetTreeNode>>& SelectedItems,
diff --git a/Engine/Plugins/Experimental/Animation/PoseSearch/Source/Editor/Public/PoseSearchDatabaseEditorCommands.h b/Engine/Plugins/Experimental/Animation/PoseSearch/Source/Editor/Public/PoseSearchDatabaseEditorCommands.h
--- a/Engine/Plugins/Experimental/Animation/PoseSearch/Source/Editor/Public/PoseSearchDatabaseEditorCommands.h
+++ b/Engine/Plugins/Experimental/Animation/PoseSearch/Source/Editor/Public/PoseSearchDatabaseEditorCommands.h
@@ -24,8 +24,6 @@ namespace UE::PoseSearch
 		virtual void RegisterCommands() override;
 
 	public:
-		TSharedPtr<FUICommandInfo> BuildSearchIndex;
-
 		TSharedPtr<FUICommandInfo> ShowPoseFeaturesNone;
 		TSharedPtr<FUICommandInfo> ShowPoseFeaturesAll;
 		TSharedPtr<FUICommandInfo> ShowPoseFeaturesDetailed;
