motion matching - MM node / function input trajectory in skeletal mesh space

diff --git a/Engine/Plugins/Experimental/Animation/MotionTrajectory/Source/MotionTrajectory/Private/CharacterTrajectoryComponent.cpp b/Engine/Plugins/Experimental/Animation/MotionTrajectory/Source/MotionTrajectory/Private/CharacterTrajectoryComponent.cpp
--- a/Engine/Plugins/Experimental/Animation/MotionTrajectory/Source/MotionTrajectory/Private/CharacterTrajectoryComponent.cpp
+++ b/Engine/Plugins/Experimental/Animation/MotionTrajectory/Source/MotionTrajectory/Private/CharacterTrajectoryComponent.cpp
@@ -135,25 +135,6 @@ void UCharacterTrajectoryComponent::OnMovementUpdated(float DeltaSeconds, FVecto
 #endif // ENABLE_ANIM_DEBUG
 }
 
-FPoseSearchQueryTrajectory UCharacterTrajectoryComponent::GetCharacterRelativeTrajectory() const
-{
-	check(GetOwner());
-
-	FPoseSearchQueryTrajectory CharacterRelativeTrajectory = Trajectory;
-	if (!SkelMeshComponent)
-	{
-		UE_LOG(LogMotionTrajectory, Error, TEXT("UCharacterTrajectoryComponent requires SkelMeshComponent to be initialized"));
-	}
-	else
-	{
-		const FTransform OwnerTransformWS = GetOwner()->GetActorTransform();
-		const FTransform ComponentTransformWS = SkelMeshComponent->GetComponentTransform();
-		const FTransform ReferenceChangeTransform = OwnerTransformWS.GetRelativeTransform(ComponentTransformWS);
-		CharacterRelativeTrajectory.TransformReferenceFrame(ReferenceChangeTransform);
-	}
-	return CharacterRelativeTrajectory;
-}
-
 void UpdateHistorySample(FPoseSearchQueryTrajectorySample& Sample, float DeltaSeconds, const FTransform& DeltaTransformCS)
 {
 	Sample.Facing = DeltaTransformCS.InverseTransformRotation(Sample.Facing);
diff --git a/Engine/Plugins/Experimental/Animation/MotionTrajectory/Source/MotionTrajectory/Public/CharacterTrajectoryComponent.h b/Engine/Plugins/Experimental/Animation/MotionTrajectory/Source/MotionTrajectory/Public/CharacterTrajectoryComponent.h
--- a/Engine/Plugins/Experimental/Animation/MotionTrajectory/Source/MotionTrajectory/Public/CharacterTrajectoryComponent.h
+++ b/Engine/Plugins/Experimental/Animation/MotionTrajectory/Source/MotionTrajectory/Public/CharacterTrajectoryComponent.h
@@ -30,9 +30,6 @@ protected:
 	UFUNCTION()
 	void OnMovementUpdated(float DeltaSeconds, FVector OldLocation, FVector OldVelocity);
 
-	UFUNCTION(BlueprintCallable, Category = "Trajectory", meta = (BlueprintThreadSafe))
-	FPoseSearchQueryTrajectory GetCharacterRelativeTrajectory() const;
-
 	void UpdateHistory(float DeltaSeconds, const FTransform& DeltaTransformCS);
 	void UpdatePrediction(const FVector& VelocityCS, const FVector& AccelerationCS, const FRotator& ControllerRotationRate);
 
diff --git a/Engine/Plugins/Experimental/Animation/PoseSearch/Source/Runtime/Private/PoseSearchContext.cpp b/Engine/Plugins/Experimental/Animation/PoseSearch/Source/Runtime/Private/PoseSearchContext.cpp
--- a/Engine/Plugins/Experimental/Animation/PoseSearch/Source/Runtime/Private/PoseSearchContext.cpp
+++ b/Engine/Plugins/Experimental/Animation/PoseSearch/Source/Runtime/Private/PoseSearchContext.cpp
@@ -350,22 +350,27 @@ FVector FSearchContext::GetSampleVelocity(float SampleTimeOffset, const UPoseSea
 	return LinearVelocity;
 }
 
-FTransform FSearchContext::GetTransform(float SampleTime, const UPoseSearchSchema* Schema, int8 SchemaBoneIdx, bool bUseHistoryRoot)
+FTransform FSearchContext::GetRootAtTime(float Time, bool bUseHistoryRoot, bool bExtrapolate) const
 {
-	// collecting the RootTransform from the IPoseHistory
 	FTransform RootTransform = FTransform::Identity;
 	if (bUseHistoryRoot)
 	{
 		check(History);
-		History->GetRootTransformAtTime(SampleTime, RootTransform);
+		History->GetRootTransformAtTime(Time, RootTransform);
 	}
 	else
 	{
 		check(Trajectory);
-		const FPoseSearchQueryTrajectorySample TrajectorySample = Trajectory->GetSampleAtTime(SampleTime);
-		RootTransform = TrajectorySample.GetTransform();
+		RootTransform = Trajectory->GetSampleAtTime(Time, bExtrapolate).GetTransform();
 	}
 
+	return RootTransform;
+}
+
+FTransform FSearchContext::GetTransform(float SampleTime, const UPoseSearchSchema* Schema, int8 SchemaBoneIdx, bool bUseHistoryRoot)
+{
+	const FTransform RootTransform = GetRootAtTime(SampleTime, bUseHistoryRoot);
+
 	const FBoneIndexType BoneIndexType = Schema->GetBoneIndexType(SchemaBoneIdx);
 	if (BoneIndexType != RootBoneIndexType)
 	{
diff --git a/Engine/Plugins/Experimental/Animation/PoseSearch/Source/Runtime/Private/PoseSearchFeatureChannel_FilterCrashingLegs.cpp b/Engine/Plugins/Experimental/Animation/PoseSearch/Source/Runtime/Private/PoseSearchFeatureChannel_FilterCrashingLegs.cpp
--- a/Engine/Plugins/Experimental/Animation/PoseSearch/Source/Runtime/Private/PoseSearchFeatureChannel_FilterCrashingLegs.cpp
+++ b/Engine/Plugins/Experimental/Animation/PoseSearch/Source/Runtime/Private/PoseSearchFeatureChannel_FilterCrashingLegs.cpp
@@ -56,14 +56,14 @@ void UPoseSearchFeatureChannel_FilterCrashingLegs::BuildQuery(UE::PoseSearch::FS
 	const bool bIsCurrentResultValid = SearchContext.GetCurrentResult().IsValid() && SearchContext.GetCurrentResult().Database->Schema == InOutQuery.GetSchema();
 	const bool bSkip = InputQueryPose != EInputQueryPose::UseCharacterPose && bIsCurrentResultValid;
 	
-	if (bSkip || !SearchContext.GetHistory())
+	if (bSkip || !SearchContext.IsHistoryValid())
 	{
 		if (bIsCurrentResultValid)
 		{
 			const float LerpValue = InputQueryPose == EInputQueryPose::UseInterpolatedContinuingPose ? SearchContext.GetCurrentResult().LerpValue : 0.f;
 			FFeatureVectorHelper::EncodeFloat(InOutQuery.EditValues(), ChannelDataOffset, SearchContext.GetCurrentResultPrevPoseVector(), SearchContext.GetCurrentResultPoseVector(), SearchContext.GetCurrentResultNextPoseVector(), LerpValue);
 		}
-		// else leave the InOutQuery set to zero since the SearchContext.GetHistory() is invalid and it'll fail if we continue
+		// else leave the InOutQuery set to zero since the SearchContext.History is invalid and it'll fail if we continue
 	}
 	else
 	{
diff --git a/Engine/Plugins/Experimental/Animation/PoseSearch/Source/Runtime/Private/PoseSearchFeatureChannel_Heading.cpp b/Engine/Plugins/Experimental/Animation/PoseSearch/Source/Runtime/Private/PoseSearchFeatureChannel_Heading.cpp
--- a/Engine/Plugins/Experimental/Animation/PoseSearch/Source/Runtime/Private/PoseSearchFeatureChannel_Heading.cpp
+++ b/Engine/Plugins/Experimental/Animation/PoseSearch/Source/Runtime/Private/PoseSearchFeatureChannel_Heading.cpp
@@ -47,7 +47,7 @@ void UPoseSearchFeatureChannel_Heading::BuildQuery(UE::PoseSearch::FSearchContex
 	const bool bIsCurrentResultValid = SearchContext.GetCurrentResult().IsValid() && SearchContext.GetCurrentResult().Database->Schema == InOutQuery.GetSchema();
 	const bool bSkip = InputQueryPose != EInputQueryPose::UseCharacterPose && bIsCurrentResultValid;
 	const bool bIsRootBone = InOutQuery.GetSchema()->IsRootBone(SchemaBoneIdx);
-	if (bSkip || (!SearchContext.GetHistory() && !bIsRootBone))
+	if (bSkip || (!SearchContext.IsHistoryValid() && !bIsRootBone))
 	{
 		if (bIsCurrentResultValid)
 		{
@@ -55,7 +55,7 @@ void UPoseSearchFeatureChannel_Heading::BuildQuery(UE::PoseSearch::FSearchContex
 			// @todo: we should normalize if LerpValue != 0
 			FFeatureVectorHelper::EncodeVector(InOutQuery.EditValues(), ChannelDataOffset, SearchContext.GetCurrentResultPrevPoseVector(), SearchContext.GetCurrentResultPoseVector(), SearchContext.GetCurrentResultNextPoseVector(), LerpValue, true, ComponentStripping);
 		}
-		// else leave the InOutQuery set to zero since the SearchContext.GetHistory() is invalid and it'll fail if we continue
+		// else leave the InOutQuery set to zero since the SearchContext.History is invalid and it'll fail if we continue
 	}
 	else
 	{
diff --git a/Engine/Plugins/Experimental/Animation/PoseSearch/Source/Runtime/Private/PoseSearchFeatureChannel_Phase.cpp b/Engine/Plugins/Experimental/Animation/PoseSearch/Source/Runtime/Private/PoseSearchFeatureChannel_Phase.cpp
--- a/Engine/Plugins/Experimental/Animation/PoseSearch/Source/Runtime/Private/PoseSearchFeatureChannel_Phase.cpp
+++ b/Engine/Plugins/Experimental/Animation/PoseSearch/Source/Runtime/Private/PoseSearchFeatureChannel_Phase.cpp
@@ -302,14 +302,14 @@ void UPoseSearchFeatureChannel_Phase::BuildQuery(UE::PoseSearch::FSearchContext&
 
 	const bool bIsCurrentResultValid = SearchContext.GetCurrentResult().IsValid() && SearchContext.GetCurrentResult().Database->Schema == InOutQuery.GetSchema();
 	const bool bSkip = InputQueryPose != EInputQueryPose::UseCharacterPose && bIsCurrentResultValid;
-	if (bSkip || !SearchContext.GetHistory())
+	if (bSkip || !SearchContext.IsHistoryValid())
 	{
 		if (bIsCurrentResultValid)
 		{
 			const float LerpValue = InputQueryPose == EInputQueryPose::UseInterpolatedContinuingPose ? SearchContext.GetCurrentResult().LerpValue : 0.f;
 			FFeatureVectorHelper::EncodeVector2D(InOutQuery.EditValues(), ChannelDataOffset, SearchContext.GetCurrentResultPrevPoseVector(), SearchContext.GetCurrentResultPoseVector(), SearchContext.GetCurrentResultNextPoseVector(), LerpValue);
 		}
-		// else leave the InOutQuery set to zero since the SearchContext.GetHistory() is invalid and it'll fail if we continue
+		// else leave the InOutQuery set to zero since the SearchContext.History is invalid and it'll fail if we continue
 	}
 	else
 	{
diff --git a/Engine/Plugins/Experimental/Animation/PoseSearch/Source/Runtime/Private/PoseSearchFeatureChannel_Position.cpp b/Engine/Plugins/Experimental/Animation/PoseSearch/Source/Runtime/Private/PoseSearchFeatureChannel_Position.cpp
--- a/Engine/Plugins/Experimental/Animation/PoseSearch/Source/Runtime/Private/PoseSearchFeatureChannel_Position.cpp
+++ b/Engine/Plugins/Experimental/Animation/PoseSearch/Source/Runtime/Private/PoseSearchFeatureChannel_Position.cpp
@@ -63,14 +63,14 @@ void UPoseSearchFeatureChannel_Position::BuildQuery(UE::PoseSearch::FSearchConte
 	const bool bIsCurrentResultValid = SearchContext.GetCurrentResult().IsValid() && SearchContext.GetCurrentResult().Database->Schema == InOutQuery.GetSchema();
 	const bool bSkip = InputQueryPose != EInputQueryPose::UseCharacterPose && bIsCurrentResultValid;
 	const bool bIsRootBone = InOutQuery.GetSchema()->IsRootBone(SchemaBoneIdx);
-	if (bSkip || (!SearchContext.GetHistory() && !bIsRootBone))
+	if (bSkip || (!SearchContext.IsHistoryValid() && !bIsRootBone))
 	{
 		if (bIsCurrentResultValid)
 		{
 			const float LerpValue = InputQueryPose == EInputQueryPose::UseInterpolatedContinuingPose ? SearchContext.GetCurrentResult().LerpValue : 0.f;
 			FFeatureVectorHelper::EncodeVector(InOutQuery.EditValues(), ChannelDataOffset, SearchContext.GetCurrentResultPrevPoseVector(), SearchContext.GetCurrentResultPoseVector(), SearchContext.GetCurrentResultNextPoseVector(), LerpValue, false, ComponentStripping);
 		}
-		// else leave the InOutQuery set to zero since the SearchContext.GetHistory() is invalid and it'll fail if we continue
+		// else leave the InOutQuery set to zero since the SearchContext.History is invalid and it'll fail if we continue
 	}
 	else
 	{
diff --git a/Engine/Plugins/Experimental/Animation/PoseSearch/Source/Runtime/Private/PoseSearchFeatureChannel_Velocity.cpp b/Engine/Plugins/Experimental/Animation/PoseSearch/Source/Runtime/Private/PoseSearchFeatureChannel_Velocity.cpp
--- a/Engine/Plugins/Experimental/Animation/PoseSearch/Source/Runtime/Private/PoseSearchFeatureChannel_Velocity.cpp
+++ b/Engine/Plugins/Experimental/Animation/PoseSearch/Source/Runtime/Private/PoseSearchFeatureChannel_Velocity.cpp
@@ -33,7 +33,7 @@ void UPoseSearchFeatureChannel_Velocity::BuildQuery(UE::PoseSearch::FSearchConte
 	const bool bIsCurrentResultValid = SearchContext.GetCurrentResult().IsValid() && SearchContext.GetCurrentResult().Database->Schema == InOutQuery.GetSchema();
 	const bool bSkip = InputQueryPose != EInputQueryPose::UseCharacterPose && bIsCurrentResultValid;
 	const bool bIsRootBone = InOutQuery.GetSchema()->IsRootBone(SchemaBoneIdx);
-	if (bSkip || (!SearchContext.GetHistory() && !bIsRootBone))
+	if (bSkip || (!SearchContext.IsHistoryValid() && !bIsRootBone))
 	{
 		if (bIsCurrentResultValid)
 		{
@@ -41,7 +41,7 @@ void UPoseSearchFeatureChannel_Velocity::BuildQuery(UE::PoseSearch::FSearchConte
 			// @todo: we should normalize if EPoseSearchVelocityFlags::Normalize && LerpValue != 0
 			FFeatureVectorHelper::EncodeVector(InOutQuery.EditValues(), ChannelDataOffset, SearchContext.GetCurrentResultPrevPoseVector(), SearchContext.GetCurrentResultPoseVector(), SearchContext.GetCurrentResultNextPoseVector(), LerpValue, false, ComponentStripping);
 		}
-		// else leave the InOutQuery set to zero since the SearchContext.GetHistory() is invalid and it'll fail if we continue
+		// else leave the InOutQuery set to zero since the SearchContext.History is invalid and it'll fail if we continue
 	}
 	else
 	{
diff --git a/Engine/Plugins/Experimental/Animation/PoseSearch/Source/Runtime/Private/PoseSearchLibrary.cpp b/Engine/Plugins/Experimental/Animation/PoseSearch/Source/Runtime/Private/PoseSearchLibrary.cpp
--- a/Engine/Plugins/Experimental/Animation/PoseSearch/Source/Runtime/Private/PoseSearchLibrary.cpp
+++ b/Engine/Plugins/Experimental/Animation/PoseSearch/Source/Runtime/Private/PoseSearchLibrary.cpp
@@ -221,13 +221,12 @@ void UPoseSearchLibrary::TraceMotionMatchingState(
 		TraceState.CurrentPoseEntryIdx = PoseEntryIdx;
 	}
 
-	if (DeltaTime > SMALL_NUMBER && SearchContext.GetTrajectory())
+	if (DeltaTime > SMALL_NUMBER && SearchContext.IsTrajectoryValid())
 	{
 		// simulation
-		const FPoseSearchQueryTrajectorySample PrevSample = SearchContext.GetTrajectory()->GetSampleAtTime(-DeltaTime);
-		const FPoseSearchQueryTrajectorySample CurrSample = SearchContext.GetTrajectory()->GetSampleAtTime(0.f);
-
-		const FTransform SimDelta = CurrSample.GetTransform().GetRelativeTransform(PrevSample.GetTransform());
+		const FTransform PrevRoot = SearchContext.GetRootAtTime(-DeltaTime);
+		const FTransform CurrRoot = SearchContext.GetRootAtTime(0.f);
+		const FTransform SimDelta = CurrRoot.GetRelativeTransform(PrevRoot);
 
 		TraceState.SimLinearVelocity = SimDelta.GetTranslation().Size() / DeltaTime;
 		TraceState.SimAngularVelocity = FMath::RadiansToDegrees(SimDelta.GetRotation().GetAngle()) / DeltaTime;
@@ -290,7 +289,8 @@ void UPoseSearchLibrary::UpdateMotionMatchingState(
 		QueryMirrorRequest = CurrentIndexAsset->bMirrored ? EPoseSearchBooleanRequest::TrueValue : EPoseSearchBooleanRequest::FalseValue;
 	}
 
-	FSearchContext SearchContext(&Trajectory, History, 0.f, &InOutMotionMatchingState.PoseIndicesHistory, QueryMirrorRequest, 
+	const FPoseSearchQueryTrajectory TrajectoryCS = GetCharacterRelativeTrajectory(Trajectory, Context.AnimInstanceProxy->GetActorTransform(), Context.AnimInstanceProxy->GetComponentTransform());
+	FSearchContext SearchContext(&TrajectoryCS, History, 0.f, &InOutMotionMatchingState.PoseIndicesHistory, QueryMirrorRequest, 
 		InOutMotionMatchingState.CurrentSearchResult, Settings.PoseJumpThresholdTime, bForceInterrupt, InOutMotionMatchingState.CanAdvance(DeltaTime));
 
 	// If we can't advance or enough time has elapsed since the last pose jump then search
@@ -386,6 +386,15 @@ void UPoseSearchLibrary::UpdateMotionMatchingState(
 #endif
 }
 
+// transforms Trajectory from the SkeletalMeshComponent relative space into Character relative space
+FPoseSearchQueryTrajectory UPoseSearchLibrary::GetCharacterRelativeTrajectory(const FPoseSearchQueryTrajectory& Trajectory, const FTransform& OwnerTransformWS, const FTransform& ComponentTransformWS)
+{
+	const FTransform ReferenceChangeTransform = OwnerTransformWS.GetRelativeTransform(ComponentTransformWS);
+	FPoseSearchQueryTrajectory TrajectoryCS = Trajectory;
+	TrajectoryCS.TransformReferenceFrame(ReferenceChangeTransform);
+	return TrajectoryCS;
+}
+
 void UPoseSearchLibrary::MotionMatch(
 	UAnimInstance* AnimInstance,
 	const UPoseSearchDatabase* Database,
@@ -428,6 +437,8 @@ void UPoseSearchLibrary::MotionMatch(
 
 	if (Database)
 	{
+		const FPoseSearchQueryTrajectory TrajectoryCS = GetCharacterRelativeTrajectory(Trajectory, AnimInstance->GetOwningActor()->GetActorTransform(), AnimInstance->GetOwningComponent()->GetComponentTransform());
+
 		// ExtendedPoseHistory will hold future poses to match AssetSamplerBase (at FutureAnimationStartTime) TimeToFutureAnimationStart seconds in the future
 		FExtendedPoseHistory ExtendedPoseHistory;
 		if (AnimInstance)
@@ -496,7 +507,7 @@ void UPoseSearchLibrary::MotionMatch(
 					FCSPose<FCompactPose> ComponentSpacePose;
 					ComponentSpacePose.InitPose(Pose);
 
-					const FPoseSearchQueryTrajectorySample TrajectorySample = Trajectory.GetSampleAtTime(ExtractionTime);
+					const FPoseSearchQueryTrajectorySample TrajectorySample = TrajectoryCS.GetSampleAtTime(ExtractionTime);
 					const FTransform& ComponentTransform = AnimInstance->GetOwningComponent()->GetComponentTransform();
 					const FTransform FutureComponentTransform = TrajectorySample.GetTransform() * ComponentTransform;
 
@@ -513,7 +524,7 @@ void UPoseSearchLibrary::MotionMatch(
 		}
 
 		// @todo: finish set up SearchContext by exposing or calculating additional members
-		FSearchContext SearchContext(&Trajectory, ExtendedPoseHistory.IsInitialized() ? &ExtendedPoseHistory : nullptr, TimeToFutureAnimationStart);
+		FSearchContext SearchContext(&TrajectoryCS, ExtendedPoseHistory.IsInitialized() ? &ExtendedPoseHistory : nullptr, TimeToFutureAnimationStart);
 
 		FSearchResult SearchResult = Database->Search(SearchContext);
 		if (SearchResult.IsValid())
diff --git a/Engine/Plugins/Experimental/Animation/PoseSearch/Source/Runtime/Public/PoseSearch/PoseSearchContext.h b/Engine/Plugins/Experimental/Animation/PoseSearch/Source/Runtime/Public/PoseSearch/PoseSearchContext.h
--- a/Engine/Plugins/Experimental/Animation/PoseSearch/Source/Runtime/Public/PoseSearch/PoseSearchContext.h
+++ b/Engine/Plugins/Experimental/Animation/PoseSearch/Source/Runtime/Public/PoseSearch/PoseSearchContext.h
@@ -165,11 +165,12 @@ struct POSESEARCH_API FSearchContext
 	float GetPoseJumpThresholdTime() const { return PoseJumpThresholdTime; }
 	EPoseSearchBooleanRequest GetQueryMirrorRequest() const { return QueryMirrorRequest; }
 	const FPoseIndicesHistory* GetPoseIndicesHistory() const { return PoseIndicesHistory; }
-	const IPoseHistory* GetHistory() const { return History; }
+	bool IsHistoryValid() const { return History != nullptr; }
 	float GetDesiredPermutationTimeOffset() const { return DesiredPermutationTimeOffset; }
-	const FPoseSearchQueryTrajectory* GetTrajectory() const { return Trajectory; }
+	bool IsTrajectoryValid() const { return Trajectory != nullptr; }
 	bool CanAdvance() const { return bCanAdvance; }
 	bool IsForceInterrupt() const { return bForceInterrupt; }
+	FTransform GetRootAtTime(float Time, bool bUseHistoryRoot = false, bool bExtrapolate = true) const;
 
 private:
 	FTransform GetTransform(float SampleTime, const UPoseSearchSchema* Schema, int8 SchemaBoneIdx = RootSchemaBoneIdx, bool bUseHistoryRoot = false);
diff --git a/Engine/Plugins/Experimental/Animation/PoseSearch/Source/Runtime/Public/PoseSearch/PoseSearchLibrary.h b/Engine/Plugins/Experimental/Animation/PoseSearch/Source/Runtime/Public/PoseSearch/PoseSearchLibrary.h
--- a/Engine/Plugins/Experimental/Animation/PoseSearch/Source/Runtime/Public/PoseSearch/PoseSearchLibrary.h
+++ b/Engine/Plugins/Experimental/Animation/PoseSearch/Source/Runtime/Public/PoseSearch/PoseSearchLibrary.h
@@ -131,13 +131,18 @@ class POSESEARCH_API UPoseSearchLibrary : public UBlueprintFunctionLibrary
 		bool bSearch);
 #endif // UE_POSE_SEARCH_TRACE_ENABLED
 
+	static FPoseSearchQueryTrajectory GetCharacterRelativeTrajectory(
+		const FPoseSearchQueryTrajectory& Trajectory,
+		const FTransform& OwnerTransformWS,
+		const FTransform& ComponentTransformWS);
+
 public:
 	/**
 	* Implementation of the core motion matching algorithm
 	*
 	* @param Context						Input animation update context providing access to the proxy and delta time
 	* @param Databases						Input array of databases to search
-	* @param Trajectory						Input motion trajectory samples for pose search queries
+	* @param Trajectory						Input motion trajectory samples for pose search queries. Expected to be in the space of the SkeletalMeshComponent. This is provided with the CharacterMovementTrajectory Component output.
 	* @param Settings						Input motion matching algorithm configuration settings
 	* @param InOutMotionMatchingState		Input/Output encapsulated motion matching algorithm and state
 	* @param bForceInterrupt				Input force interrupt request (if true the continuing pose will be invalidated)
@@ -161,7 +166,7 @@ public:
 	*
 	* @param AnimInstance					Input animation instance
 	* @param Database						Input database to search
-	* @param Trajectory						Input motion trajectory samples for pose search queries
+	* @param Trajectory						Input motion trajectory samples for pose search queries. Expected to be in the space of the SkeletalMeshComponent. This is provided with the CharacterMovementTrajectory Component output.
 	* @param PoseHistoryName				Input tag of the associated PoseSearchHistoryCollector node in the anim graph
 	* @param SelectedAnimation				Output selected animation from the searchable asset
 	* @param SelectedTime					Output selected animation time
