motion matching - deprecating MM node Settings in favor of pinnable properties that can be exposed one by one

diff --git a/Engine/Plugins/Experimental/Animation/PoseSearch/Source/Runtime/Private/AnimNode_MotionMatching.cpp b/Engine/Plugins/Experimental/Animation/PoseSearch/Source/Runtime/Private/AnimNode_MotionMatching.cpp
--- a/Engine/Plugins/Experimental/Animation/PoseSearch/Source/Runtime/Private/AnimNode_MotionMatching.cpp
+++ b/Engine/Plugins/Experimental/Animation/PoseSearch/Source/Runtime/Private/AnimNode_MotionMatching.cpp
@@ -114,7 +114,12 @@ void FAnimNode_MotionMatching::UpdateAssetPlayer(const FAnimationUpdateContext&
 		Context,
 		DatabasesToSearch,
 		Trajectory,
-		Settings,
+		BlendTime,
+		MaxActiveBlends,
+		PoseJumpThresholdTime,
+		PoseReselectHistory,
+		SearchThrottleTime,
+		PlayRate,
 		MotionMatchingState,
 		bForceInterrupt | bForceInterruptNextUpdate,
 		bShouldSearch
@@ -135,7 +140,7 @@ void FAnimNode_MotionMatching::UpdateAssetPlayer(const FAnimationUpdateContext&
 			{
 				BlendStackNode.BlendTo(DatabaseAsset->GetAnimationAsset(), MotionMatchingState.CurrentSearchResult.AssetTime,
 					DatabaseAsset->IsLooping(), SearchIndexAsset->bMirrored, CurrentResultDatabase->Schema->MirrorDataTable.Get(),
-					Settings.MaxActiveBlends, Settings.BlendTime, Settings.BlendProfile, Settings.BlendOption, SearchIndexAsset->BlendParameters, MotionMatchingState.WantedPlayRate);
+					MaxActiveBlends, BlendTime, BlendProfile, BlendOption, SearchIndexAsset->BlendParameters, MotionMatchingState.WantedPlayRate);
 			}
 		}
 	}
diff --git a/Engine/Plugins/Experimental/Animation/PoseSearch/Source/Runtime/Private/PoseSearchLibrary.cpp b/Engine/Plugins/Experimental/Animation/PoseSearch/Source/Runtime/Private/PoseSearchLibrary.cpp
--- a/Engine/Plugins/Experimental/Animation/PoseSearch/Source/Runtime/Private/PoseSearchLibrary.cpp
+++ b/Engine/Plugins/Experimental/Animation/PoseSearch/Source/Runtime/Private/PoseSearchLibrary.cpp
@@ -95,12 +95,12 @@ static void RequestInertialBlend(const FAnimationUpdateContext& Context, float B
 	}
 }
 
-void FMotionMatchingState::JumpToPose(const FAnimationUpdateContext& Context, const FMotionMatchingSettings& Settings, const UE::PoseSearch::FSearchResult& Result)
+void FMotionMatchingState::JumpToPose(const FAnimationUpdateContext& Context, const UE::PoseSearch::FSearchResult& Result, int32 MaxActiveBlends, float BlendTime)
 {
 	// requesting inertial blending only if blendstack is disabled
-	if (Settings.MaxActiveBlends <= 0)
+	if (MaxActiveBlends <= 0)
 	{
-		RequestInertialBlend(Context, Settings.BlendTime);
+		RequestInertialBlend(Context, BlendTime);
 	}
 
 	// Remember which pose and sequence we're playing from the database
@@ -109,11 +109,11 @@ void FMotionMatchingState::JumpToPose(const FAnimationUpdateContext& Context, co
 	bJumpedToPose = true;
 }
 
-void FMotionMatchingState::UpdateWantedPlayRate(const UE::PoseSearch::FSearchContext& SearchContext, const FMotionMatchingSettings& Settings)
+void FMotionMatchingState::UpdateWantedPlayRate(const UE::PoseSearch::FSearchContext& SearchContext, const FFloatInterval& PlayRate)
 {
 	if (CurrentSearchResult.IsValid())
 	{
-		if (!FMath::IsNearlyEqual(Settings.PlayRate.Min, 1.f, UE_KINDA_SMALL_NUMBER) || !FMath::IsNearlyEqual(Settings.PlayRate.Max, 1.f, UE_KINDA_SMALL_NUMBER))
+		if (!FMath::IsNearlyEqual(PlayRate.Min, 1.f, UE_KINDA_SMALL_NUMBER) || !FMath::IsNearlyEqual(PlayRate.Max, 1.f, UE_KINDA_SMALL_NUMBER))
 		{
 			if (const FPoseSearchFeatureVectorBuilder* PoseSearchFeatureVectorBuilder = SearchContext.GetCachedQuery(CurrentSearchResult.Database->Schema))
 			{
@@ -122,8 +122,8 @@ void FMotionMatchingState::UpdateWantedPlayRate(const UE::PoseSearch::FSearchCon
 					TConstArrayView<float> QueryData = PoseSearchFeatureVectorBuilder->GetValues();
 					TConstArrayView<float> ResultData = CurrentSearchResult.Database->GetSearchIndex().GetPoseValues(CurrentSearchResult.PoseIdx);
 					const float EstimatedSpeedRatio = TrajectoryChannel->GetEstimatedSpeedRatio(QueryData, ResultData);
-					check(Settings.PlayRate.Min <= Settings.PlayRate.Max);
-					WantedPlayRate = FMath::Clamp(EstimatedSpeedRatio, Settings.PlayRate.Min, Settings.PlayRate.Max);
+					check(PlayRate.Min <= PlayRate.Max);
+					WantedPlayRate = FMath::Clamp(EstimatedSpeedRatio, PlayRate.Min, PlayRate.Max);
 				}
 				else
 				{
@@ -249,7 +249,12 @@ void UPoseSearchLibrary::UpdateMotionMatchingState(
 	const FAnimationUpdateContext& Context,
 	const TArray<TObjectPtr<const UPoseSearchDatabase>>& Databases,
 	const FPoseSearchQueryTrajectory& Trajectory,
-	const FMotionMatchingSettings& Settings,
+	float BlendTime,
+	int32 MaxActiveBlends,
+	float PoseJumpThresholdTime,
+	float PoseReselectHistory,
+	float SearchThrottleTime,
+	const FFloatInterval& PlayRate,
 	FMotionMatchingState& InOutMotionMatchingState,
 	bool bForceInterrupt,
 	bool bShouldSearch,
@@ -291,10 +296,10 @@ void UPoseSearchLibrary::UpdateMotionMatchingState(
 
 	const FPoseSearchQueryTrajectory TrajectoryCS = GetCharacterRelativeTrajectory(Trajectory, Context.AnimInstanceProxy->GetActorTransform(), Context.AnimInstanceProxy->GetComponentTransform());
 	FSearchContext SearchContext(&TrajectoryCS, History, 0.f, &InOutMotionMatchingState.PoseIndicesHistory, QueryMirrorRequest, 
-		InOutMotionMatchingState.CurrentSearchResult, Settings.PoseJumpThresholdTime, bForceInterrupt, InOutMotionMatchingState.CanAdvance(DeltaTime));
+		InOutMotionMatchingState.CurrentSearchResult, PoseJumpThresholdTime, bForceInterrupt, InOutMotionMatchingState.CanAdvance(DeltaTime));
 
 	// If we can't advance or enough time has elapsed since the last pose jump then search
-	const bool bSearch = !SearchContext.CanAdvance() || (bShouldSearch && (InOutMotionMatchingState.ElapsedPoseSearchTime >= Settings.SearchThrottleTime));
+	const bool bSearch = !SearchContext.CanAdvance() || (bShouldSearch && (InOutMotionMatchingState.ElapsedPoseSearchTime >= SearchThrottleTime));
 	if (bSearch)
 	{
 		InOutMotionMatchingState.ElapsedPoseSearchTime = 0.f;
@@ -325,7 +330,7 @@ void UPoseSearchLibrary::UpdateMotionMatchingState(
 		if (SearchResult.PoseCost.GetTotalCost() < ContinuingPoseCost.GetTotalCost())
 		{
 			SearchResult.ContinuingPoseCost = ContinuingPoseCost;
-			InOutMotionMatchingState.JumpToPose(Context, Settings, SearchResult);
+			InOutMotionMatchingState.JumpToPose(Context, SearchResult, MaxActiveBlends, BlendTime);
 		}
 		else
 		{
@@ -337,14 +342,14 @@ void UPoseSearchLibrary::UpdateMotionMatchingState(
 			InOutMotionMatchingState.CurrentSearchResult.ContinuingPoseCost = ContinuingPoseCost;
 		}
 
-		InOutMotionMatchingState.UpdateWantedPlayRate(SearchContext, Settings);
+		InOutMotionMatchingState.UpdateWantedPlayRate(SearchContext, PlayRate);
 	}
 	else
 	{
 		InOutMotionMatchingState.ElapsedPoseSearchTime += DeltaTime;
 	}
 
-	InOutMotionMatchingState.PoseIndicesHistory.Update(InOutMotionMatchingState.CurrentSearchResult, DeltaTime, Settings.PoseReselectHistory);
+	InOutMotionMatchingState.PoseIndicesHistory.Update(InOutMotionMatchingState.CurrentSearchResult, DeltaTime, PoseReselectHistory);
 
 #if UE_POSE_SEARCH_TRACE_ENABLED
 	// Record debugger details
diff --git a/Engine/Plugins/Experimental/Animation/PoseSearch/Source/Runtime/Public/PoseSearch/AnimNode_MotionMatching.h b/Engine/Plugins/Experimental/Animation/PoseSearch/Source/Runtime/Public/PoseSearch/AnimNode_MotionMatching.h
--- a/Engine/Plugins/Experimental/Animation/PoseSearch/Source/Runtime/Public/PoseSearch/AnimNode_MotionMatching.h
+++ b/Engine/Plugins/Experimental/Animation/PoseSearch/Source/Runtime/Public/PoseSearch/AnimNode_MotionMatching.h
@@ -25,12 +25,41 @@ private:
 	TObjectPtr<const UPoseSearchDatabase> Database = nullptr;
 
 	// Motion Trajectory samples for pose search queries in Motion Matching.These are expected to be in the space of the SkeletalMeshComponent.This is provided with the CharacterMovementTrajectory Component output.
-	UPROPERTY(EditAnywhere, Category=Settings, meta=(PinShownByDefault))
+	UPROPERTY(EditAnywhere, Category = Settings, meta = (PinShownByDefault))
 	FPoseSearchQueryTrajectory Trajectory;
 
 	// Settings for the core motion matching node.
-	UPROPERTY(EditAnywhere, Category=Settings, meta=(PinHiddenByDefault))
-	FMotionMatchingSettings Settings;
+	// Time in seconds to blend out to the new pose. Uses either inertial blending, requiring an Inertialization node after this node, or the internal blend stack, if MaxActiveBlends is greater than zero.
+	UPROPERTY(EditAnywhere, Category = Settings, meta = (PinHiddenByDefault, ClampMin="0"))
+	float BlendTime = 0.2f;
+
+	// Number of max active animation segments being blended together in the blend stack. If MaxActiveBlends is zero then the blend stack is disabled.
+	UPROPERTY(EditAnywhere, Category = Settings, meta = (PinHiddenByDefault, ClampMin="0"))
+	int32 MaxActiveBlends = 4;
+
+	// Set Blend Profiles (editable in the skeleton) to determine how the blending is distributed among your character's bones. It could be used to differentiate between upper body and lower body to blend timing.
+	UPROPERTY(EditAnywhere, Category = Settings, meta = (PinHiddenByDefault, UseAsBlendProfile = true))
+	TObjectPtr<UBlendProfile> BlendProfile;
+
+	// How the blend is applied over time to the bones. Common selections are linear, ease in, ease out, and ease in and out.
+	UPROPERTY(EditAnywhere, Category = Settings, meta = (PinHiddenByDefault))
+	EAlphaBlendOption BlendOption = EAlphaBlendOption::Linear;
+
+	// Don't jump to poses of the same segment that are less than this many seconds away.
+	UPROPERTY(EditAnywhere, Category = Settings, meta = (PinHiddenByDefault, ClampMin="0"))
+	float PoseJumpThresholdTime = 0.f;
+
+	// Prevent re-selection of poses that have been selected previously within this much time (in seconds) in the past. This is across all animation segments that have been selected within this time range.
+	UPROPERTY(EditAnywhere, Category = Settings, meta = (PinHiddenByDefault, ClampMin = "0"))
+	float PoseReselectHistory = 0.3f;
+
+	// Minimum amount of time to wait between searching for a new pose segment. It allows users to define how often the system searches, default for locomotion is searching every update, but you may only want to search once for other situations, like jump.
+	UPROPERTY(EditAnywhere, Category = Settings, meta = (PinHiddenByDefault, ClampMin="0"))
+	float SearchThrottleTime = 0.f;
+
+	// Effective range of play rate that can be applied to the animations to account for discrepancies in estimated velocity between the movement modeland the animation.
+	UPROPERTY(EditAnywhere, Category = Settings, meta = (PinHiddenByDefault, ClampMin = "0.5", ClampMax = "2.0", UIMin = "0.5", UIMax = "2.0"))
+	FFloatInterval PlayRate = FFloatInterval(1.f, 1.f);
 
 	// Reset the motion matching selection state if it has become relevant to the graph after not being updated on previous frames.
 	UPROPERTY(EditAnywhere, Category = Settings, meta = (PinHiddenByDefault))
@@ -102,7 +131,7 @@ private:
 
 #if WITH_EDITORONLY_DATA
 	// If true, "Relevant anim" nodes that look for the highest weighted animation in a state will ignore this node
-	UPROPERTY(EditAnywhere, Category=Relevancy, meta=(FoldProperty, PinHiddenByDefault))
+	UPROPERTY(EditAnywhere, Category = Relevancy, meta = (FoldProperty, PinHiddenByDefault))
 	bool bIgnoreForRelevancyTest = false;
 #endif // WITH_EDITORONLY_DATA
 };
\ No newline at end of file
diff --git a/Engine/Plugins/Experimental/Animation/PoseSearch/Source/Runtime/Public/PoseSearch/PoseSearchLibrary.h b/Engine/Plugins/Experimental/Animation/PoseSearch/Source/Runtime/Public/PoseSearch/PoseSearchLibrary.h
--- a/Engine/Plugins/Experimental/Animation/PoseSearch/Source/Runtime/Public/PoseSearch/PoseSearchLibrary.h
+++ b/Engine/Plugins/Experimental/Animation/PoseSearch/Source/Runtime/Public/PoseSearch/PoseSearchLibrary.h
@@ -23,44 +23,6 @@ namespace UE::PoseSearch
 struct FAnimationUpdateContext;
 struct FPoseSearchQueryTrajectory;
 
-USTRUCT(BlueprintType, Category = "Animation|Pose Search")
-struct POSESEARCH_API FMotionMatchingSettings
-{
-	GENERATED_BODY()
-
-	// Time in seconds to blend out to the new pose. Uses either inertial blending, requiring an Inertialization node after this node, or the internal blend stack, if MaxActiveBlends is greater than zero.
-	UPROPERTY(EditAnywhere, BlueprintReadWrite, Category=Settings, meta=(ClampMin="0"))
-	float BlendTime = 0.2f;
-
-	// Number of max active animation segments being blended together in the blend stack. If MaxActiveBlends is zero then the blend stack is disabled.
-	UPROPERTY(EditAnywhere, BlueprintReadWrite, Category=Settings, meta=(ClampMin="0"))
-	int32 MaxActiveBlends = 4;
-
-	// Set Blend Profiles (editable in the skeleton) to determine how the blending is distributed among your character's bones. It could be used to differentiate between upper body and lower body to blend timing.
-	UPROPERTY(EditAnywhere, BlueprintReadWrite, Category=Settings, meta=(UseAsBlendProfile = true))
-	TObjectPtr<UBlendProfile> BlendProfile;
-
-	// How the blend is applied over time to the bones. Common selections are linear, ease in, ease out, and ease in and out.
-	UPROPERTY(EditAnywhere, BlueprintReadWrite, Category = Settings)
-	EAlphaBlendOption BlendOption = EAlphaBlendOption::Linear;
-
-	// Don't jump to poses of the same segment that are less than this many seconds away.
-	UPROPERTY(EditAnywhere, BlueprintReadWrite, Category=Settings, meta=(ClampMin="0"))
-	float PoseJumpThresholdTime = 0.f;
-
-	// Prevent re-selection of poses that have been selected previously within this much time (in seconds) in the past. This is across all animation segments that have been selected within this time range.
-	UPROPERTY(EditAnywhere, BlueprintReadWrite, Category = Settings, meta = (ClampMin = "0"))
-	float PoseReselectHistory = 0.3f;
-
-	// Minimum amount of time to wait between searching for a new pose segment. It allows users to define how often the system searches, default for locomotion is searching every update, but you may only want to search once for other situations, like jump.
-	UPROPERTY(EditAnywhere, BlueprintReadWrite, Category=Settings, meta=(ClampMin="0"))
-	float SearchThrottleTime = 0.f;
-
-	// Effective range of play rate that can be applied to the animations to account for discrepancies in estimated velocity between the movement modeland the animation.
-	UPROPERTY(EditAnywhere, BlueprintReadWrite, Category = Settings, meta = (ClampMin = "0.5", ClampMax = "2.0", UIMin = "0.5", UIMax = "2.0"))
-	FFloatInterval PlayRate = FFloatInterval(1.f, 1.f);
-};
-
 USTRUCT(BlueprintType, Category="Animation|Pose Search")
 struct POSESEARCH_API FMotionMatchingState
 {
@@ -78,9 +40,9 @@ struct POSESEARCH_API FMotionMatchingState
 	void AdjustAssetTime(float AssetTime);
 
 	// Internally stores the 'jump' to a new pose/sequence index and asset time for evaluation
-	void JumpToPose(const FAnimationUpdateContext& Context, const FMotionMatchingSettings& Settings, const UE::PoseSearch::FSearchResult& Result);
+	void JumpToPose(const FAnimationUpdateContext& Context, const UE::PoseSearch::FSearchResult& Result, int32 MaxActiveBlends, float BlendTime);
 
-	void UpdateWantedPlayRate(const UE::PoseSearch::FSearchContext& SearchContext, const FMotionMatchingSettings& Settings);
+	void UpdateWantedPlayRate(const UE::PoseSearch::FSearchContext& SearchContext, const FFloatInterval& PlayRate);
 
 	UE::PoseSearch::FSearchResult CurrentSearchResult;
 
@@ -143,7 +105,12 @@ public:
 	* @param Context						Input animation update context providing access to the proxy and delta time
 	* @param Databases						Input array of databases to search
 	* @param Trajectory						Input motion trajectory samples for pose search queries. Expected to be in the space of the SkeletalMeshComponent. This is provided with the CharacterMovementTrajectory Component output.
-	* @param Settings						Input motion matching algorithm configuration settings
+	* @param BlendTime						Input time in seconds to blend out to the new pose. Uses either inertial blending, requiring an Inertialization node after this node, or the internal blend stack, if MaxActiveBlends is greater than zero.
+	* @param MaxActiveBlends				Input number of max active animation segments being blended together in the blend stack. If MaxActiveBlends is zero then the blend stack is disabled.
+	* @param PoseJumpThresholdTime			Input don't jump to poses of the same segment that are less than this many seconds away.
+	* @param PoseReselectHistory			Input prevent re-selection of poses that have been selected previously within this much time (in seconds) in the past. This is across all animation segments that have been selected within this time range.
+	* @param SearchThrottleTime				Input minimum amount of time to wait between searching for a new pose segment. It allows users to define how often the system searches, default for locomotion is searching every update, but you may only want to search once for other situations, like jump.
+	* @param PlayRate						Input effective range of play rate that can be applied to the animations to account for discrepancies in estimated velocity between the movement modeland the animation.
 	* @param InOutMotionMatchingState		Input/Output encapsulated motion matching algorithm and state
 	* @param bForceInterrupt				Input force interrupt request (if true the continuing pose will be invalidated)
 	* @param bShouldSearch					Input if false search will happen only if there's no valid continuing pose
@@ -154,7 +121,12 @@ public:
 		const FAnimationUpdateContext& Context,
 		const TArray<TObjectPtr<const UPoseSearchDatabase>>& Databases,
 		const FPoseSearchQueryTrajectory& Trajectory,
-		const FMotionMatchingSettings& Settings,
+		float BlendTime,
+		int32 MaxActiveBlends,
+		float PoseJumpThresholdTime,
+		float PoseReselectHistory,
+		float SearchThrottleTime,
+		const FFloatInterval& PlayRate,
 		FMotionMatchingState& InOutMotionMatchingState,
 		bool bForceInterrupt = false,
 		bool bShouldSearch = true,
