motion matching - trajectory query scaling

diff --git a/Engine/Plugins/Experimental/Animation/PoseSearch/Source/Runtime/Private/AnimNode_MotionMatching.cpp b/Engine/Plugins/Experimental/Animation/PoseSearch/Source/Runtime/Private/AnimNode_MotionMatching.cpp
--- a/Engine/Plugins/Experimental/Animation/PoseSearch/Source/Runtime/Private/AnimNode_MotionMatching.cpp
+++ b/Engine/Plugins/Experimental/Animation/PoseSearch/Source/Runtime/Private/AnimNode_MotionMatching.cpp
@@ -114,6 +114,7 @@ void FAnimNode_MotionMatching::UpdateAssetPlayer(const FAnimationUpdateContext&
 		Context,
 		DatabasesToSearch,
 		Trajectory,
+		TrajectorySpeedMultiplier,
 		BlendTime,
 		MaxActiveBlends,
 		PoseJumpThresholdTime,
diff --git a/Engine/Plugins/Experimental/Animation/PoseSearch/Source/Runtime/Private/PoseSearchLibrary.cpp b/Engine/Plugins/Experimental/Animation/PoseSearch/Source/Runtime/Private/PoseSearchLibrary.cpp
--- a/Engine/Plugins/Experimental/Animation/PoseSearch/Source/Runtime/Private/PoseSearchLibrary.cpp
+++ b/Engine/Plugins/Experimental/Animation/PoseSearch/Source/Runtime/Private/PoseSearchLibrary.cpp
@@ -109,7 +109,7 @@ void FMotionMatchingState::JumpToPose(const FAnimationUpdateContext& Context, co
 	bJumpedToPose = true;
 }
 
-void FMotionMatchingState::UpdateWantedPlayRate(const UE::PoseSearch::FSearchContext& SearchContext, const FFloatInterval& PlayRate)
+void FMotionMatchingState::UpdateWantedPlayRate(const UE::PoseSearch::FSearchContext& SearchContext, const FFloatInterval& PlayRate, float TrajectorySpeedMultiplier)
 {
 	if (CurrentSearchResult.IsValid())
 	{
@@ -133,6 +133,10 @@ void FMotionMatchingState::UpdateWantedPlayRate(const UE::PoseSearch::FSearchCon
 				}
 			}
 		}
+		else if (!FMath::IsNearlyZero(TrajectorySpeedMultiplier))
+		{
+			WantedPlayRate = 1.f / TrajectorySpeedMultiplier;
+		}
 	}
 }
 
@@ -249,6 +253,7 @@ void UPoseSearchLibrary::UpdateMotionMatchingState(
 	const FAnimationUpdateContext& Context,
 	const TArray<TObjectPtr<const UPoseSearchDatabase>>& Databases,
 	const FPoseSearchQueryTrajectory& Trajectory,
+	float TrajectorySpeedMultiplier,
 	float BlendTime,
 	int32 MaxActiveBlends,
 	float PoseJumpThresholdTime,
@@ -294,7 +299,7 @@ void UPoseSearchLibrary::UpdateMotionMatchingState(
 		QueryMirrorRequest = CurrentIndexAsset->bMirrored ? EPoseSearchBooleanRequest::TrueValue : EPoseSearchBooleanRequest::FalseValue;
 	}
 
-	const FPoseSearchQueryTrajectory TrajectoryCS = GetCharacterRelativeTrajectory(Trajectory, Context.AnimInstanceProxy->GetActorTransform(), Context.AnimInstanceProxy->GetComponentTransform());
+	const FPoseSearchQueryTrajectory TrajectoryCS = ProcessTrajectory(Trajectory, Context.AnimInstanceProxy->GetActorTransform(), Context.AnimInstanceProxy->GetComponentTransform(), TrajectorySpeedMultiplier);
 	FSearchContext SearchContext(&TrajectoryCS, History, 0.f, &InOutMotionMatchingState.PoseIndicesHistory, QueryMirrorRequest, 
 		InOutMotionMatchingState.CurrentSearchResult, PoseJumpThresholdTime, bForceInterrupt, InOutMotionMatchingState.CanAdvance(DeltaTime));
 
@@ -341,14 +346,14 @@ void UPoseSearchLibrary::UpdateMotionMatchingState(
 			InOutMotionMatchingState.CurrentSearchResult.PoseCost = ContinuingPoseCost;
 			InOutMotionMatchingState.CurrentSearchResult.ContinuingPoseCost = ContinuingPoseCost;
 		}
-
-		InOutMotionMatchingState.UpdateWantedPlayRate(SearchContext, PlayRate);
 	}
 	else
 	{
 		InOutMotionMatchingState.ElapsedPoseSearchTime += DeltaTime;
 	}
 
+	InOutMotionMatchingState.UpdateWantedPlayRate(SearchContext, PlayRate, TrajectorySpeedMultiplier);
+
 	InOutMotionMatchingState.PoseIndicesHistory.Update(InOutMotionMatchingState.CurrentSearchResult, DeltaTime, PoseReselectHistory);
 
 #if UE_POSE_SEARCH_TRACE_ENABLED
@@ -391,12 +396,20 @@ void UPoseSearchLibrary::UpdateMotionMatchingState(
 #endif
 }
 
-// transforms Trajectory from the SkeletalMeshComponent relative space into Character relative space
-FPoseSearchQueryTrajectory UPoseSearchLibrary::GetCharacterRelativeTrajectory(const FPoseSearchQueryTrajectory& Trajectory, const FTransform& OwnerTransformWS, const FTransform& ComponentTransformWS)
+// transforms Trajectory from the SkeletalMeshComponent relative space into Character relative space, and scale it by TrajectorySpeedMultiplier
+FPoseSearchQueryTrajectory UPoseSearchLibrary::ProcessTrajectory(const FPoseSearchQueryTrajectory& Trajectory, const FTransform& OwnerTransformWS, const FTransform& ComponentTransformWS, float TrajectorySpeedMultiplier)
 {
 	const FTransform ReferenceChangeTransform = OwnerTransformWS.GetRelativeTransform(ComponentTransformWS);
 	FPoseSearchQueryTrajectory TrajectoryCS = Trajectory;
 	TrajectoryCS.TransformReferenceFrame(ReferenceChangeTransform);
+
+	if (!FMath::IsNearlyEqual(TrajectorySpeedMultiplier, 1.f) && !FMath::IsNearlyZero(TrajectorySpeedMultiplier))
+	{
+		for (FPoseSearchQueryTrajectorySample& Sample : TrajectoryCS.Samples)
+		{
+			Sample.AccumulatedSeconds /= TrajectorySpeedMultiplier;
+		}
+	}
 	return TrajectoryCS;
 }
 
@@ -404,6 +417,7 @@ void UPoseSearchLibrary::MotionMatch(
 	UAnimInstance* AnimInstance,
 	const UPoseSearchDatabase* Database,
 	const FPoseSearchQueryTrajectory Trajectory,
+	float TrajectorySpeedMultiplier,
 	const FName PoseHistoryName,
 	UAnimationAsset*& SelectedAnimation,
 	float& SelectedTime,
@@ -442,7 +456,7 @@ void UPoseSearchLibrary::MotionMatch(
 
 	if (Database)
 	{
-		const FPoseSearchQueryTrajectory TrajectoryCS = GetCharacterRelativeTrajectory(Trajectory, AnimInstance->GetOwningActor()->GetActorTransform(), AnimInstance->GetOwningComponent()->GetComponentTransform());
+		const FPoseSearchQueryTrajectory TrajectoryCS = ProcessTrajectory(Trajectory, AnimInstance->GetOwningActor()->GetActorTransform(), AnimInstance->GetOwningComponent()->GetComponentTransform(), TrajectorySpeedMultiplier);
 
 		// ExtendedPoseHistory will hold future poses to match AssetSamplerBase (at FutureAnimationStartTime) TimeToFutureAnimationStart seconds in the future
 		FExtendedPoseHistory ExtendedPoseHistory;
diff --git a/Engine/Plugins/Experimental/Animation/PoseSearch/Source/Runtime/Public/PoseSearch/AnimNode_MotionMatching.h b/Engine/Plugins/Experimental/Animation/PoseSearch/Source/Runtime/Public/PoseSearch/AnimNode_MotionMatching.h
--- a/Engine/Plugins/Experimental/Animation/PoseSearch/Source/Runtime/Public/PoseSearch/AnimNode_MotionMatching.h
+++ b/Engine/Plugins/Experimental/Animation/PoseSearch/Source/Runtime/Public/PoseSearch/AnimNode_MotionMatching.h
@@ -28,7 +28,10 @@ private:
 	UPROPERTY(EditAnywhere, Category = Settings, meta = (PinShownByDefault))
 	FPoseSearchQueryTrajectory Trajectory;
 
-	// Settings for the core motion matching node.
+	// Input Trajectory velocity will be multiplied by TrajectorySpeedMultiplier: values below 1 will result in selecting animation slower than requested from the original Trajectory
+	UPROPERTY(EditAnywhere, Category = Settings, meta = (PinHiddenByDefault, ClampMin="0"))
+	float TrajectorySpeedMultiplier = 1.f;
+
 	// Time in seconds to blend out to the new pose. Uses either inertial blending, requiring an Inertialization node after this node, or the internal blend stack, if MaxActiveBlends is greater than zero.
 	UPROPERTY(EditAnywhere, Category = Settings, meta = (PinHiddenByDefault, ClampMin="0"))
 	float BlendTime = 0.2f;
@@ -57,8 +60,8 @@ private:
 	UPROPERTY(EditAnywhere, Category = Settings, meta = (PinHiddenByDefault, ClampMin="0"))
 	float SearchThrottleTime = 0.f;
 
-	// Effective range of play rate that can be applied to the animations to account for discrepancies in estimated velocity between the movement modeland the animation.
-	UPROPERTY(EditAnywhere, Category = Settings, meta = (PinHiddenByDefault, ClampMin = "0.5", ClampMax = "2.0", UIMin = "0.5", UIMax = "2.0"))
+	// Effective range of play rate that can be applied to the animations to account for discrepancies in estimated velocity between the movement model and the animation.
+	UPROPERTY(EditAnywhere, Category = Settings, meta = (PinHiddenByDefault, ClampMin = "0.2", ClampMax = "3.0", UIMin = "0.2", UIMax = "3.0"))
 	FFloatInterval PlayRate = FFloatInterval(1.f, 1.f);
 
 	// Reset the motion matching selection state if it has become relevant to the graph after not being updated on previous frames.
diff --git a/Engine/Plugins/Experimental/Animation/PoseSearch/Source/Runtime/Public/PoseSearch/PoseSearchLibrary.h b/Engine/Plugins/Experimental/Animation/PoseSearch/Source/Runtime/Public/PoseSearch/PoseSearchLibrary.h
--- a/Engine/Plugins/Experimental/Animation/PoseSearch/Source/Runtime/Public/PoseSearch/PoseSearchLibrary.h
+++ b/Engine/Plugins/Experimental/Animation/PoseSearch/Source/Runtime/Public/PoseSearch/PoseSearchLibrary.h
@@ -42,7 +42,7 @@ struct POSESEARCH_API FMotionMatchingState
 	// Internally stores the 'jump' to a new pose/sequence index and asset time for evaluation
 	void JumpToPose(const FAnimationUpdateContext& Context, const UE::PoseSearch::FSearchResult& Result, int32 MaxActiveBlends, float BlendTime);
 
-	void UpdateWantedPlayRate(const UE::PoseSearch::FSearchContext& SearchContext, const FFloatInterval& PlayRate);
+	void UpdateWantedPlayRate(const UE::PoseSearch::FSearchContext& SearchContext, const FFloatInterval& PlayRate, float TrajectorySpeedMultiplier);
 
 	UE::PoseSearch::FSearchResult CurrentSearchResult;
 
@@ -93,10 +93,11 @@ class POSESEARCH_API UPoseSearchLibrary : public UBlueprintFunctionLibrary
 		bool bSearch);
 #endif // UE_POSE_SEARCH_TRACE_ENABLED
 
-	static FPoseSearchQueryTrajectory GetCharacterRelativeTrajectory(
+	static FPoseSearchQueryTrajectory ProcessTrajectory(
 		const FPoseSearchQueryTrajectory& Trajectory,
 		const FTransform& OwnerTransformWS,
-		const FTransform& ComponentTransformWS);
+		const FTransform& ComponentTransformWS,
+		float TrajectorySpeedMultiplier);
 
 public:
 	/**
@@ -121,6 +122,7 @@ public:
 		const FAnimationUpdateContext& Context,
 		const TArray<TObjectPtr<const UPoseSearchDatabase>>& Databases,
 		const FPoseSearchQueryTrajectory& Trajectory,
+		float TrajectorySpeedMultiplier,
 		float BlendTime,
 		int32 MaxActiveBlends,
 		float PoseJumpThresholdTime,
@@ -139,8 +141,9 @@ public:
 	* @param AnimInstance					Input animation instance
 	* @param Database						Input database to search
 	* @param Trajectory						Input motion trajectory samples for pose search queries. Expected to be in the space of the SkeletalMeshComponent. This is provided with the CharacterMovementTrajectory Component output.
+	* @param TrajectorySpeedMultiplier		Input Trajectory velocity will be multiplied by TrajectorySpeedMultiplier: values below 1 will result in selecting animation slower than requested from the original Trajectory
 	* @param PoseHistoryName				Input tag of the associated PoseSearchHistoryCollector node in the anim graph
-	* @param SelectedAnimation				Output selected animation from the searchable asset
+	* @param SelectedAnimation				Output selected animation from the Database asset
 	* @param SelectedTime					Output selected animation time
 	* @param bLoop							Output selected animation looping state
 	* @param bIsMirrored					Output selected animation mirror state
@@ -156,6 +159,7 @@ public:
 		UAnimInstance* AnimInstance,
 		const UPoseSearchDatabase* Database,
 		const FPoseSearchQueryTrajectory Trajectory,
+		float TrajectorySpeedMultiplier,
 		const FName PoseHistoryName,
 		UAnimationAsset*& SelectedAnimation,
 		float& SelectedTime,
