motion matching - fix for "warning V595: The 'AnimInstance' pointer was utilized before it was verified against nullptr"

diff --git a/Engine/Plugins/Experimental/Animation/PoseSearch/Source/Runtime/Private/PoseSearchLibrary.cpp b/Engine/Plugins/Experimental/Animation/PoseSearch/Source/Runtime/Private/PoseSearchLibrary.cpp
--- a/Engine/Plugins/Experimental/Animation/PoseSearch/Source/Runtime/Private/PoseSearchLibrary.cpp
+++ b/Engine/Plugins/Experimental/Animation/PoseSearch/Source/Runtime/Private/PoseSearchLibrary.cpp
@@ -434,10 +434,13 @@ void UPoseSearchLibrary::MotionMatch(
 	class UAnimInstanceProxyProvider : public UAnimInstance
 	{
 	public:
-		static FAnimInstanceProxy& GetAnimInstanceProxy(UAnimInstance* AnimInstance)
+		static FAnimInstanceProxy* GetAnimInstanceProxy(UAnimInstance* AnimInstance)
 		{
-			check(AnimInstance);
-			return static_cast<UAnimInstanceProxyProvider*>(AnimInstance)->GetProxyOnAnyThread<FAnimInstanceProxy>();
+			if (AnimInstance)
+			{
+				return &static_cast<UAnimInstanceProxyProvider*>(AnimInstance)->GetProxyOnAnyThread<FAnimInstanceProxy>();
+			}
+			return nullptr;
 		}
 	};
 #endif //ENABLE_DRAW_DEBUG && ENABLE_ANIM_DEBUG
@@ -536,7 +539,10 @@ void UPoseSearchLibrary::MotionMatch(
 #if ENABLE_DRAW_DEBUG && ENABLE_ANIM_DEBUG
 				if (CVarAnimMotionMatchDrawHistoryEnable.GetValueOnAnyThread())
 				{
-					ExtendedPoseHistory.DebugDraw(UAnimInstanceProxyProvider::GetAnimInstanceProxy(AnimInstance));
+					if (FAnimInstanceProxy* AnimInstanceProxy = UAnimInstanceProxyProvider::GetAnimInstanceProxy(AnimInstance))
+					{
+						ExtendedPoseHistory.DebugDraw(*AnimInstanceProxy);
+					}
 				}
 #endif // ENABLE_DRAW_DEBUG && ENABLE_ANIM_DEBUG
 			}
@@ -567,13 +573,13 @@ void UPoseSearchLibrary::MotionMatch(
 			{
 				if (CVarAnimMotionMatchDrawMatchEnable.GetValueOnAnyThread())
 				{
-					UE::PoseSearch::FDebugDrawParams DrawParams(&UAnimInstanceProxyProvider::GetAnimInstanceProxy(AnimInstance), SearchResult.Database.Get());
+					UE::PoseSearch::FDebugDrawParams DrawParams(UAnimInstanceProxyProvider::GetAnimInstanceProxy(AnimInstance), SearchResult.Database.Get());
 					DrawParams.DrawFeatureVector(SearchResult.PoseIdx);
 				}
 
 				if (CVarAnimMotionMatchDrawQueryEnable.GetValueOnAnyThread())
 				{
-					UE::PoseSearch::FDebugDrawParams DrawParams(&UAnimInstanceProxyProvider::GetAnimInstanceProxy(AnimInstance), SearchResult.Database.Get(), EDebugDrawFlags::DrawQuery);
+					UE::PoseSearch::FDebugDrawParams DrawParams(UAnimInstanceProxyProvider::GetAnimInstanceProxy(AnimInstance), SearchResult.Database.Get(), EDebugDrawFlags::DrawQuery);
 					DrawParams.DrawFeatureVector(SearchContext.GetOrBuildQuery(SearchResult.Database->Schema).GetValues());
 				}
 			}
