motion matching
- precalculating FAnimationAssetSampler::CachedPlayLength in FAnimationAssetSampler::Init since required also in ExtractPoseSearchNotifyStates
- adding "Quantize Animation To PoseData" for the MM database editor
- added FPoseSearchDatabaseBlendSpace::bUseSingleSample to sample BS with a single blend value in MM database

diff --git a/Engine/Plugins/Experimental/Animation/PoseSearch/Source/Editor/Private/PoseSearchDatabaseEditorCommands.cpp b/Engine/Plugins/Experimental/Animation/PoseSearch/Source/Editor/Private/PoseSearchDatabaseEditorCommands.cpp
--- a/Engine/Plugins/Experimental/Animation/PoseSearch/Source/Editor/Private/PoseSearchDatabaseEditorCommands.cpp
+++ b/Engine/Plugins/Experimental/Animation/PoseSearch/Source/Editor/Private/PoseSearchDatabaseEditorCommands.cpp
@@ -8,26 +8,13 @@ namespace UE::PoseSearch
 {
 	void FDatabaseEditorCommands::RegisterCommands()
 	{
-		UI_COMMAND(
-			ShowPoseFeaturesNone, "None", "Hide all features",
-			EUserInterfaceActionType::RadioButton, FInputChord());
-		UI_COMMAND(
-			ShowPoseFeaturesAll, "All Features", "Show all features",
-			EUserInterfaceActionType::RadioButton, FInputChord());
-		UI_COMMAND(
-			ShowPoseFeaturesDetailed, "Detailed Features", "Show all features with mode details",
-			EUserInterfaceActionType::RadioButton, FInputChord());
-
-		UI_COMMAND(
-			ShowAnimationOriginalOnly, "Original Only", "Show only original animations",
-			EUserInterfaceActionType::RadioButton, FInputChord());
-		UI_COMMAND(
-			ShowAnimationOriginalAndMirrored, "Original and mirrored", "Show original and mirrored animations",
-			EUserInterfaceActionType::RadioButton, FInputChord());
-
-		UI_COMMAND(
-			ShowDisplayRootMotionSpeed, "Display Root Motion Speed", "Show root motion speed for the selected item",
-			EUserInterfaceActionType::ToggleButton, FInputChord());
+		UI_COMMAND(ShowPoseFeaturesNone, "None", "Hide all features", EUserInterfaceActionType::RadioButton, FInputChord());
+		UI_COMMAND(ShowPoseFeaturesAll, "All Features", "Show all features", EUserInterfaceActionType::RadioButton, FInputChord());
+		UI_COMMAND(ShowPoseFeaturesDetailed, "Detailed Features", "Show all features with mode details", EUserInterfaceActionType::RadioButton, FInputChord());
+		UI_COMMAND(ShowAnimationOriginalOnly, "Original Only", "Show only original animations", EUserInterfaceActionType::RadioButton, FInputChord());
+		UI_COMMAND(ShowAnimationOriginalAndMirrored, "Original and mirrored", "Show original and mirrored animations", EUserInterfaceActionType::RadioButton, FInputChord());
+		UI_COMMAND(ShowDisplayRootMotionSpeed, "Display Root Motion Speed", "Show root motion speed for the selected item", EUserInterfaceActionType::ToggleButton, FInputChord());
+		UI_COMMAND(ShowQuantizeAnimationToPoseData, "Quantize Animation To Pose Data", "Animations will not playing continuously, only using the associated pose quantized time", EUserInterfaceActionType::ToggleButton, FInputChord());
 	}
 }
 
diff --git a/Engine/Plugins/Experimental/Animation/PoseSearch/Source/Editor/Private/PoseSearchDatabaseViewModel.cpp b/Engine/Plugins/Experimental/Animation/PoseSearch/Source/Editor/Private/PoseSearchDatabaseViewModel.cpp
--- a/Engine/Plugins/Experimental/Animation/PoseSearch/Source/Editor/Private/PoseSearchDatabaseViewModel.cpp
+++ b/Engine/Plugins/Experimental/Animation/PoseSearch/Source/Editor/Private/PoseSearchDatabaseViewModel.cpp
@@ -236,13 +236,22 @@ namespace UE::PoseSearch
 			float CurrentPlayTime = PlayTime + IndexAsset.SamplingInterval.Min + PreviewActor.PlayTimeOffset;
 			FAnimationRuntime::AdvanceTime(false, CurrentPlayTime, CurrentTime, IndexAsset.SamplingInterval.Max);
 
-			AnimInstance->SetPosition(PreviewActor.Sampler.GetScaledTime(CurrentTime));
+			// time to pose index
+			PreviewActor.CurrentPoseIndex = PoseSearchDatabase->GetPoseIndexFromTime(CurrentTime, IndexAsset);
+
+			const float TimeToAssetTimeMultiplier = PreviewActor.Sampler.GetTimeToAssetTimeMultiplier();
+			check(!FMath::IsNearlyZero(TimeToAssetTimeMultiplier));
+			const float QuantizedTime = PreviewActor.CurrentPoseIndex >= 0 ? PoseSearchDatabase->GetAssetTime(PreviewActor.CurrentPoseIndex) / TimeToAssetTimeMultiplier : CurrentTime;
+
+			if (bQuantizeAnimationToPoseData)
+			{
+				CurrentTime = QuantizedTime;
+			}
+
+			AnimInstance->SetPosition(CurrentTime * TimeToAssetTimeMultiplier);
 			AnimInstance->SetPlayRate(0.f);
 			AnimInstance->SetBlendSpacePosition(IndexAsset.BlendParameters);
 
-			// time to pose index
-			PreviewActor.CurrentPoseIndex = PoseSearchDatabase->GetPoseIndexFromTime(CurrentTime, IndexAsset);
-			const float QuantizedTime = PreviewActor.CurrentPoseIndex >= 0 ? PoseSearchDatabase->GetAssetTime(PreviewActor.CurrentPoseIndex) : CurrentTime;
 			PreviewActor.QuantizedTimeRootMotion = PreviewActor.Sampler.ExtractRootTransform(QuantizedTime);
 			FTransform RootMotion = PreviewActor.Sampler.ExtractRootTransform(CurrentTime);
 			if (AnimInstance->GetMirrorDataTable())
@@ -327,7 +336,12 @@ namespace UE::PoseSearch
 
 	void FDatabaseViewModel::OnToggleDisplayRootMotionSpeed()
 	{
-		DisplayRootMotionSpeed = !DisplayRootMotionSpeed;
+		bDisplayRootMotionSpeed = !bDisplayRootMotionSpeed;
+	}
+
+	void FDatabaseViewModel::OnToggleQuantizeAnimationToPoseData()
+	{
+		bQuantizeAnimationToPoseData = !bQuantizeAnimationToPoseData;
 	}
 
 	void FDatabaseViewModel::AddSequenceToDatabase(UAnimSequence* AnimSequence)
diff --git a/Engine/Plugins/Experimental/Animation/PoseSearch/Source/Editor/Private/SPoseSearchDatabaseViewport.cpp b/Engine/Plugins/Experimental/Animation/PoseSearch/Source/Editor/Private/SPoseSearchDatabaseViewport.cpp
--- a/Engine/Plugins/Experimental/Animation/PoseSearch/Source/Editor/Private/SPoseSearchDatabaseViewport.cpp
+++ b/Engine/Plugins/Experimental/Animation/PoseSearch/Source/Editor/Private/SPoseSearchDatabaseViewport.cpp
@@ -102,6 +102,13 @@ namespace UE::PoseSearch
 			FExecuteAction::CreateSP(ViewModelRef, &FDatabaseViewModel::OnToggleDisplayRootMotionSpeed),
 			FCanExecuteAction(),
 			FIsActionChecked::CreateSP(ViewModelRef, &FDatabaseViewModel::IsDisplayRootMotionSpeedChecked));
+
+		CommandList->MapAction(
+			Commands.ShowQuantizeAnimationToPoseData,
+			FExecuteAction::CreateSP(ViewModelRef, &FDatabaseViewModel::OnToggleQuantizeAnimationToPoseData),
+			FCanExecuteAction(),
+			FIsActionChecked::CreateSP(ViewModelRef, &FDatabaseViewModel::IsQuantizeAnimationToPoseDataChecked));
+		
 	}
 
 	TSharedRef<FEditorViewportClient> SDatabaseViewport::MakeEditorViewportClient()
diff --git a/Engine/Plugins/Experimental/Animation/PoseSearch/Source/Editor/Private/SPoseSearchDatabaseViewportToolbar.cpp b/Engine/Plugins/Experimental/Animation/PoseSearch/Source/Editor/Private/SPoseSearchDatabaseViewportToolbar.cpp
--- a/Engine/Plugins/Experimental/Animation/PoseSearch/Source/Editor/Private/SPoseSearchDatabaseViewportToolbar.cpp
+++ b/Engine/Plugins/Experimental/Animation/PoseSearch/Source/Editor/Private/SPoseSearchDatabaseViewportToolbar.cpp
@@ -73,6 +73,7 @@ namespace UE::PoseSearch
 				SubMenuBuilder.BeginSection("Debug", LOCTEXT("ShowMenu_DebugLabel", "Debug"));
 				{
 					SubMenuBuilder.AddMenuEntry(Commands.ShowDisplayRootMotionSpeed);
+					SubMenuBuilder.AddMenuEntry(Commands.ShowQuantizeAnimationToPoseData);
 				}
 				SubMenuBuilder.EndSection();
 			})
diff --git a/Engine/Plugins/Experimental/Animation/PoseSearch/Source/Editor/Public/PoseSearchDatabaseEditorCommands.h b/Engine/Plugins/Experimental/Animation/PoseSearch/Source/Editor/Public/PoseSearchDatabaseEditorCommands.h
--- a/Engine/Plugins/Experimental/Animation/PoseSearch/Source/Editor/Public/PoseSearchDatabaseEditorCommands.h
+++ b/Engine/Plugins/Experimental/Animation/PoseSearch/Source/Editor/Public/PoseSearchDatabaseEditorCommands.h
@@ -27,11 +27,10 @@ namespace UE::PoseSearch
 		TSharedPtr<FUICommandInfo> ShowPoseFeaturesNone;
 		TSharedPtr<FUICommandInfo> ShowPoseFeaturesAll;
 		TSharedPtr<FUICommandInfo> ShowPoseFeaturesDetailed;
-
 		TSharedPtr<FUICommandInfo> ShowAnimationOriginalOnly;
 		TSharedPtr<FUICommandInfo> ShowAnimationOriginalAndMirrored;
-		
 		TSharedPtr<FUICommandInfo> ShowDisplayRootMotionSpeed;
+		TSharedPtr<FUICommandInfo> ShowQuantizeAnimationToPoseData;
 	};
 }
 
diff --git a/Engine/Plugins/Experimental/Animation/PoseSearch/Source/Editor/Public/PoseSearchDatabaseViewModel.h b/Engine/Plugins/Experimental/Animation/PoseSearch/Source/Editor/Public/PoseSearchDatabaseViewModel.h
--- a/Engine/Plugins/Experimental/Animation/PoseSearch/Source/Editor/Public/PoseSearchDatabaseViewModel.h
+++ b/Engine/Plugins/Experimental/Animation/PoseSearch/Source/Editor/Public/PoseSearchDatabaseViewModel.h
@@ -94,7 +94,10 @@ namespace UE::PoseSearch
 		bool IsAnimationPreviewMode(EAnimationPreviewMode PreviewMode) const;
 
 		void OnToggleDisplayRootMotionSpeed();
-		bool IsDisplayRootMotionSpeedChecked() const { return DisplayRootMotionSpeed; };
+		bool IsDisplayRootMotionSpeedChecked() const { return bDisplayRootMotionSpeed; };
+
+		void OnToggleQuantizeAnimationToPoseData();
+		bool IsQuantizeAnimationToPoseDataChecked() const { return bQuantizeAnimationToPoseData; };
 
 		void AddSequenceToDatabase(UAnimSequence* AnimSequence);
 		void AddBlendSpaceToDatabase(UBlendSpace* BlendSpace);
@@ -152,8 +155,10 @@ namespace UE::PoseSearch
 		EAnimationPreviewMode AnimationPreviewMode = EAnimationPreviewMode::OriginalAndMirrored;
 
 		/** Is animation debug draw enabled */
-		bool DisplayRootMotionSpeed = false;
-		
+		bool bDisplayRootMotionSpeed = false;
+
+		bool bQuantizeAnimationToPoseData = false;
+
 		int32 SelectedActorIndexAssetIndex = INDEX_NONE;
 
 		UWorld* GetWorld();
diff --git a/Engine/Plugins/Experimental/Animation/PoseSearch/Source/Runtime/Private/PoseSearchAssetSampler.cpp b/Engine/Plugins/Experimental/Animation/PoseSearch/Source/Runtime/Private/PoseSearchAssetSampler.cpp
--- a/Engine/Plugins/Experimental/Animation/PoseSearch/Source/Runtime/Private/PoseSearchAssetSampler.cpp
+++ b/Engine/Plugins/Experimental/Animation/PoseSearch/Source/Runtime/Private/PoseSearchAssetSampler.cpp
@@ -119,7 +119,6 @@ static FTransform ExtractBlendSpaceRootMotion(float StartTime, float DeltaTime,
 		do
 		{
 			// Disable looping here. Advance to desired position, or beginning / end of animation 
-			check(CachedPlayLength >= 0.f);
 			const ETypeAdvanceAnim AdvanceType = FAnimationRuntime::AdvanceTime(false, DesiredDeltaMove, CurrentPosition, CachedPlayLength);
 
 			// Verify position assumptions
@@ -161,7 +160,6 @@ static void ProcessRootTransform(const UBlendSpace* BlendSpace, const FVector& B
 	int32 RootTransformSamplingRate, bool bIsLoopable, TArray<FTransform>& AccumulatedRootTransform)
 {
 	// Pre-compute root motion
-	check(CachedPlayLength >= 0.f);
 	int32 NumRootSamples = FMath::Max(CachedPlayLength * RootTransformSamplingRate + 1, 1);
 	AccumulatedRootTransform.SetNumUninitialized(NumRootSamples);
 
@@ -258,6 +256,19 @@ void FAnimationAssetSampler::Init(TWeakObjectPtr<const UAnimationAsset> InAnimat
 	BlendParameters = InBlendParameters;
 	BoneContainer = InBoneContainer;
 	RootTransformSamplingRate = InRootTransformSamplingRate;
+
+	if (const UAnimationAsset* AnimAsset = AnimationAsset.Get())
+	{
+		if (const UBlendSpace* BlendSpace = Cast<UBlendSpace>(AnimAsset))
+		{
+			FMemMark Mark(FMemStack::Get());
+			ProcessPlayLength(BlendSpace, BlendParameters, CachedPlayLength);
+		}
+		else
+		{
+			CachedPlayLength = AnimAsset->GetPlayLength();
+		}
+	}
 }
 
 bool FAnimationAssetSampler::IsInitialized() const
@@ -270,26 +281,20 @@ const UAnimationAsset* FAnimationAssetSampler::GetAsset() const
 	return AnimationAsset.Get();
 }
 
-float FAnimationAssetSampler::GetScaledTime(float Time) const
+float FAnimationAssetSampler::GetTimeToAssetTimeMultiplier() const
 {
-	// Asset player time for blend spaces is normalized [0, 1] so we convert it to scaled time by dividing it by CachedPlayLength
-	if (Cast<UBlendSpace>(AnimationAsset.Get()))
+	// Asset player time for blend spaces is normalized [0, 1] so we convert the sampling / animation time to asset time by dividing it by CachedPlayLength
+	if (CachedPlayLength > UE_KINDA_SMALL_NUMBER && Cast<UBlendSpace>(AnimationAsset.Get()))
 	{
-		check(CachedPlayLength >= 0.f);
-		return CachedPlayLength > UE_KINDA_SMALL_NUMBER ? Time / CachedPlayLength : 0.f;
+		return 1.f / CachedPlayLength;
 	}
 
-	return Time;
+	return 1.f;
 }
 
 float FAnimationAssetSampler::GetPlayLength() const
 {
-	if (Cast<UBlendSpace>(AnimationAsset.Get()))
-	{
-		check(CachedPlayLength >= 0.f);
-		return CachedPlayLength;
-	}
-	return AnimationAsset->GetPlayLength();
+	return CachedPlayLength;
 }
 
 bool FAnimationAssetSampler::IsLoopable() const
@@ -312,7 +317,6 @@ FTransform FAnimationAssetSampler::GetTotalRootTransform() const
 {
 	if (Cast<UBlendSpace>(AnimationAsset.Get()))
 	{
-		check(CachedPlayLength >= 0.f);
 		const FTransform InitialRootTransform = ExtractBlendSpaceRootTrackTransform(0.f, AccumulatedRootTransform, RootTransformSamplingRate);
 		const FTransform LastRootTransform = ExtractBlendSpaceRootTrackTransform(CachedPlayLength, AccumulatedRootTransform, RootTransformSamplingRate);
 		const FTransform TotalRootTransform = LastRootTransform.GetRelativeTransform(InitialRootTransform);
@@ -349,7 +353,6 @@ void FAnimationAssetSampler::ExtractPose(const FAnimExtractContext& ExtractionCt
 		int32 TriangulationIndex = 0;
 		if (BlendSpace->GetSamplesFromBlendInput(BlendParameters, BlendSamples, TriangulationIndex, true))
 		{
-			check(CachedPlayLength >= 0.f);
 			for (int32 BlendSampleIdex = 0; BlendSampleIdex < BlendSamples.Num(); BlendSampleIdex++)
 			{
 				float Scale = BlendSamples[BlendSampleIdex].Animation->GetPlayLength() / CachedPlayLength;
@@ -399,7 +402,6 @@ FTransform FAnimationAssetSampler::ExtractRootTransform(float Time) const
 		}
 		else
 		{
-			check(CachedPlayLength >= 0.f);
 			const float ClampedTime = FMath::Clamp(Time, 0.0f, CachedPlayLength);
 			const float ExtrapolationTime = Time - ClampedTime;
 
@@ -533,8 +535,6 @@ void FAnimationAssetSampler::Process()
 	if (const UBlendSpace* BlendSpace = Cast<UBlendSpace>(AnimationAsset.Get()))
 	{
 		FMemMark Mark(FMemStack::Get());
-
-		ProcessPlayLength(BlendSpace, BlendParameters, CachedPlayLength);
 		ProcessRootTransform(BlendSpace, BlendParameters, CachedPlayLength, BoneContainer, RootTransformSamplingRate, IsLoopable(), AccumulatedRootTransform);
 	}
 }
@@ -556,7 +556,6 @@ void FAnimationAssetSampler::ExtractPoseSearchNotifyStates(float Time, TArray<UA
 				const int32 HighestWeightIndex = GetHighestWeightSample(BlendSamples);
 
 				// getting pose search notifies in an interval of size ExtractionInterval, centered on Time
-				check(CachedPlayLength >= 0.f);
 				SampleTime = Time * (BlendSamples[HighestWeightIndex].Animation->GetPlayLength() / CachedPlayLength);
 
 				// Get notifies for highest weighted
diff --git a/Engine/Plugins/Experimental/Animation/PoseSearch/Source/Runtime/Private/PoseSearchDatabase.cpp b/Engine/Plugins/Experimental/Animation/PoseSearch/Source/Runtime/Private/PoseSearchDatabase.cpp
--- a/Engine/Plugins/Experimental/Animation/PoseSearch/Source/Runtime/Private/PoseSearchDatabase.cpp
+++ b/Engine/Plugins/Experimental/Animation/PoseSearch/Source/Runtime/Private/PoseSearchDatabase.cpp
@@ -310,8 +310,21 @@ void FPoseSearchDatabaseBlendSpace::GetBlendSpaceParameterSampleRanges(int32& Ho
 {
 	check(BlendSpace);
 
-	HorizontalBlendNum = bUseGridForSampling ? BlendSpace->GetBlendParameter(0).GridNum + 1 : FMath::Max(NumberOfHorizontalSamples, 1);
-	VerticalBlendNum = BlendSpace->IsA<UBlendSpace1D>() ? 1 : bUseGridForSampling ? BlendSpace->GetBlendParameter(1).GridNum + 1 : FMath::Max(NumberOfVerticalSamples, 1);
+	if (bUseSingleSample)
+	{
+		HorizontalBlendNum = 1;
+		VerticalBlendNum = 1;
+	}
+	else if (bUseGridForSampling)
+	{
+		HorizontalBlendNum = BlendSpace->GetBlendParameter(0).GridNum + 1;
+		VerticalBlendNum = BlendSpace->IsA<UBlendSpace1D>() ? 1 : BlendSpace->GetBlendParameter(1).GridNum + 1;
+	}
+	else
+	{
+		HorizontalBlendNum = FMath::Max(NumberOfHorizontalSamples, 1);
+		VerticalBlendNum = BlendSpace->IsA<UBlendSpace1D>() ? 1 : FMath::Max(NumberOfVerticalSamples, 1);
+	}
 
 	check(HorizontalBlendNum >= 1 && VerticalBlendNum >= 1);
 }
@@ -320,6 +333,12 @@ FVector FPoseSearchDatabaseBlendSpace::BlendParameterForSampleRanges(int32 Horiz
 {
 	check(BlendSpace);
 
+	if (bUseSingleSample)
+	{
+		check(HorizontalBlendIndex == 0 && VerticalBlendIndex == 0);
+		return FVector(BlendParamX, BlendParamY, 0.f);
+	}
+	
 	const bool bWrapInputOnHorizontalAxis = BlendSpace->GetBlendParameter(0).bWrapInput;
 	const bool bWrapInputOnVerticalAxis = BlendSpace->GetBlendParameter(1).bWrapInput;
 
@@ -351,7 +370,7 @@ FVector FPoseSearchDatabaseBlendSpace::BlendParameterForSampleRanges(int32 Horiz
 			VerticalBlendMin + (VerticalBlendMax - VerticalBlendMin) * 
 			((float)VerticalBlendIndex) / (VerticalBlendNum - 1) : 
 		VerticalBlendMin,
-		0.0f);
+		0.f);
 }
 
 //////////////////////////////////////////////////////////////////////////
diff --git a/Engine/Plugins/Experimental/Animation/PoseSearch/Source/Runtime/Public/PoseSearch/PoseSearchAssetSampler.h b/Engine/Plugins/Experimental/Animation/PoseSearch/Source/Runtime/Public/PoseSearch/PoseSearchAssetSampler.h
--- a/Engine/Plugins/Experimental/Animation/PoseSearch/Source/Runtime/Public/PoseSearch/PoseSearchAssetSampler.h
+++ b/Engine/Plugins/Experimental/Animation/PoseSearch/Source/Runtime/Public/PoseSearch/PoseSearchAssetSampler.h
@@ -22,7 +22,7 @@ struct POSESEARCH_API FAnimationAssetSampler
 
 	bool IsInitialized() const;
 	float GetPlayLength() const;
-	float GetScaledTime(float Time) const;
+	float GetTimeToAssetTimeMultiplier() const;
 	bool IsLoopable() const;
 
 	// Gets the final root transformation at the end of the asset's playback time
@@ -50,7 +50,7 @@ protected:
 	FVector BlendParameters = FVector::ZeroVector;
 	FBoneContainer BoneContainer;
 	int32 RootTransformSamplingRate = 30;
-	float CachedPlayLength = -1.f;
+	float CachedPlayLength = 0.f;
 	TArray<FTransform> AccumulatedRootTransform;
 
 	const float ExtrapolationSampleTime = 1.f / 30.f;
diff --git a/Engine/Plugins/Experimental/Animation/PoseSearch/Source/Runtime/Public/PoseSearch/PoseSearchDatabase.h b/Engine/Plugins/Experimental/Animation/PoseSearch/Source/Runtime/Public/PoseSearch/PoseSearchDatabase.h
--- a/Engine/Plugins/Experimental/Animation/PoseSearch/Source/Runtime/Public/PoseSearch/PoseSearchDatabase.h
+++ b/Engine/Plugins/Experimental/Animation/PoseSearch/Source/Runtime/Public/PoseSearch/PoseSearchDatabase.h
@@ -132,27 +132,39 @@ struct POSESEARCH_API FPoseSearchDatabaseBlendSpace : public FPoseSearchDatabase
 	UPROPERTY(EditAnywhere, Category = "BlendSpace", meta = (DisplayPriority = 0))
 	TObjectPtr<UBlendSpace> BlendSpace;
 
-	// This allows users to enable or exclude animations from this database. Useful for debugging.
-	UPROPERTY(EditAnywhere, Category = "BlendSpace", meta = (DisplayPriority = 5))
-	bool bEnabled = true;
-
 	// This allows users to set if this animation is original only (no mirrored data), original and mirrored, or only the mirrored version of this animation.
 	// It requires the mirror table to be set up in the config file.
 	UPROPERTY(EditAnywhere, Category = "BlendSpace", meta = (DisplayPriority = 1))
 	EPoseSearchMirrorOption MirrorOption = EPoseSearchMirrorOption::UnmirroredOnly;
 
+	// If true this BlendSpace will output a single segment in the database.
+	UPROPERTY(EditAnywhere, Category = "BlendSpace", meta = (DisplayPriority = 2))
+	bool bUseSingleSample = false;
+
 	// When turned on, this will use the set grid samples of the blend space asset for sampling. This will override the Number of Horizontal/Vertical Samples.
-	UPROPERTY(EditAnywhere, Category = "BlendSpace", meta = (DisplayPriority = 4))
+	UPROPERTY(EditAnywhere, Category = "BlendSpace", meta = (EditCondition = "!bUseSingleSample", EditConditionHides, DisplayPriority = 3))
 	bool bUseGridForSampling = false;
 
 	// Sets the number of horizontal samples in the blend space to pull the animation data coverage from. The larger the samples the more the data, but also the more memory and performance it takes.
-	UPROPERTY(EditAnywhere, Category = "BlendSpace", meta = (EditCondition = "!bUseGridForSampling", EditConditionHides, ClampMin = "1", UIMin = "1", UIMax = "25", DisplayPriority = 2))
+	UPROPERTY(EditAnywhere, Category = "BlendSpace", meta = (EditCondition = "!bUseSingleSample && !bUseGridForSampling", EditConditionHides, ClampMin = "1", UIMin = "1", UIMax = "25", DisplayPriority = 4))
 	int32 NumberOfHorizontalSamples = 9;
-
+	
 	// Sets the number of vertical samples in the blend space to pull the animation data coverage from.The larger the samples the more the data, but also the more memory and performance it takes.
-	UPROPERTY(EditAnywhere, Category = "BlendSpace", meta = (EditCondition = "!bUseGridForSampling", EditConditionHides, ClampMin = "1", UIMin = "1", UIMax = "25", DisplayPriority = 3))
+	UPROPERTY(EditAnywhere, Category = "BlendSpace", meta = (EditCondition = "!bUseSingleSample && !bUseGridForSampling", EditConditionHides, ClampMin = "1", UIMin = "1", UIMax = "25", DisplayPriority = 5))
 	int32 NumberOfVerticalSamples = 2;
 
+	// BlendParams used to sample this BlendSpace
+	UPROPERTY(EditAnywhere, Category = "BlendSpace", meta = (EditCondition = "bUseSingleSample", EditConditionHides, DisplayPriority = 6))
+	float BlendParamX = 0.f;
+
+	// BlendParams used to sample this BlendSpace
+	UPROPERTY(EditAnywhere, Category = "BlendSpace", meta = (EditCondition = "bUseSingleSample", EditConditionHides, DisplayPriority = 7))
+	float BlendParamY = 0.f;
+
+	// This allows users to enable or exclude animations from this database. Useful for debugging.
+	UPROPERTY(EditAnywhere, Category = "BlendSpace", meta = (DisplayPriority = 8))
+	bool bEnabled = true;
+
 	UAnimationAsset* GetAnimationAsset() const override;
 	UClass* GetAnimationAssetStaticClass() const override;
 	bool IsLooping() const override;
