motion matching - Data Details Tab

diff --git a/Engine/Plugins/Experimental/Animation/PoseSearch/Source/Editor/Private/PoseSearchDatabaseDataDetails.cpp b/Engine/Plugins/Experimental/Animation/PoseSearch/Source/Editor/Private/PoseSearchDatabaseDataDetails.cpp
new file mode 100644
--- /dev/null
+++ b/Engine/Plugins/Experimental/Animation/PoseSearch/Source/Editor/Private/PoseSearchDatabaseDataDetails.cpp
@@ -0,0 +1,315 @@
+// Copyright Epic Games, Inc. All Rights Reserved.
+
+#include "PoseSearchDatabaseDataDetails.h"
+#include "PoseSearch/PoseSearchDatabase.h"
+#include "PoseSearch/PoseSearchDefines.h"
+#include "PoseSearch/PoseSearchDerivedData.h"
+#include "PoseSearch/PoseSearchSchema.h"
+#include "PoseSearchDatabaseViewModel.h"
+#include "Widgets/Layout/SScrollBox.h"
+
+#define LOCTEXT_NAMESPACE "PoseSearchDatabaseDataDetails"
+
+namespace UE::PoseSearch
+{
+
+class FChannelItem
+{
+public:
+	FChannelItem(const UPoseSearchFeatureChannel* InChannel, int32 InChannelComponentIdx = -1)
+	: ChannelComponentIdx(InChannelComponentIdx)
+	, Channel(InChannel)
+	{
+	}
+
+	const FString GetLabel() const
+	{
+		if (ChannelComponentIdx >= 0)
+		{
+			switch (ChannelComponentIdx)
+			{
+			case 0: return FString("x");
+			case 1: return FString("y");
+			case 2: return FString("z");
+			case 3: return FString("w");
+			default: return FString::Printf(TEXT("%d"), ChannelComponentIdx);
+			}
+		}
+
+		if (Channel != nullptr)
+		{
+			return Channel->GetLabel();
+		}
+
+		return FString();
+	}
+
+	int32 GetDataOffset() const
+	{
+		int32 DataOffset = 0;
+
+		if (Channel != nullptr)
+		{
+			DataOffset = Channel->GetChannelDataOffset();
+			
+			if (ChannelComponentIdx > 0)
+			{
+				DataOffset += ChannelComponentIdx;
+			}
+		}
+
+		return DataOffset;
+	}
+
+	int32 GetCardinality() const
+	{
+		int32 Cardinality = 0;
+
+		if (ChannelComponentIdx >= 0)
+		{
+			Cardinality = 1;
+		}
+		else if (Channel != nullptr)
+		{
+			Cardinality = Channel->GetChannelCardinality();
+		}
+
+		return Cardinality;
+	}
+
+	TArray<FChannelItemPtr>& GetChannelItems()
+	{
+		return ChannelItems;
+	}
+
+private:
+	int32 ChannelComponentIdx = -1;
+	TWeakObjectPtr<const UPoseSearchFeatureChannel> Channel;
+	TArray<FChannelItemPtr> ChannelItems;
+};
+
+class SDatabaseDataDetailsTableRow : public SMultiColumnTableRow<FChannelItemPtr>
+{
+	FChannelItemPtr ChannelItem;
+	TWeakPtr<FDatabaseViewModel> EditorViewModel;
+
+public:
+	void Construct(const FTableRowArgs& InArgs, const TSharedRef<STableViewBase>& InOwnerTableView, FChannelItemPtr InChannelItem, TSharedRef<FDatabaseViewModel> InEditorViewModel)
+	{
+		EditorViewModel = InEditorViewModel;
+		ChannelItem = InChannelItem;
+		SMultiColumnTableRow<FChannelItemPtr>::Construct(InArgs, InOwnerTableView);
+	}
+
+	virtual TSharedRef<SWidget> GenerateWidgetForColumn(const FName& ColumnName) override
+	{
+		if (ColumnName == FName("ChannelName"))
+		{
+			// Rows in a TreeView need an expander button and some indentation
+			return SNew(SHorizontalBox)
+				+ SHorizontalBox::Slot()
+				.AutoWidth()
+				.HAlign(HAlign_Right)
+				.VAlign(VAlign_Fill)
+				[
+					SNew(SExpanderArrow, SharedThis(this))
+					.StyleSet(ExpanderStyleSet)
+				]
+				+ SHorizontalBox::Slot()
+				.FillWidth(1)
+				.Padding(FMargin(4.0f, 0.0f))
+				.VAlign(VAlign_Center)
+				[
+					SNew(STextBlock)
+					.Text(FText::FromString(ChannelItem->GetLabel()))
+				];
+		}
+		
+		if (ColumnName == FName("DataOffset"))
+		{
+			return SNew(STextBlock)
+			.Text(FText::FromString(FString::Printf(TEXT("%d"), ChannelItem->GetDataOffset())));
+		}
+
+		if (ColumnName == FName("Query"))
+		{
+			return SNew(STextBlock)
+			.Text_Lambda([this, ColumnName]() -> FText
+			{
+				TStringBuilder<256> StringBuilder;
+				if (const TSharedPtr<FDatabaseViewModel> ViewModel = EditorViewModel.Pin())
+				{
+					const int32 DataOffset = ChannelItem->GetDataOffset();
+					const int32 Cardinality = ChannelItem->GetCardinality();
+
+					const TConstArrayView<float> QueryValues = ViewModel->GetQueryVector();
+					if (DataOffset + Cardinality <= QueryValues.Num())
+					{
+						if (Cardinality > 1)
+						{
+							// using only one decimal to keep the string compact
+							for (int32 i = 0; i < Cardinality; ++i)
+							{
+								if (i != 0)
+								{
+									StringBuilder.Append(TEXT(", "));
+								}
+								const float Value = QueryValues[i + DataOffset];
+								StringBuilder.Appendf(TEXT("%.1f"), Value);
+							}
+						}
+						else
+						{
+							// using all the float digits 
+							const float Value = QueryValues[DataOffset];
+							StringBuilder.Appendf(TEXT("%f"), Value);
+						}
+					}
+				}
+				return FText::FromString(StringBuilder.ToString());
+			});
+		}
+		
+		return SNew(STextBlock)
+		.Text_Lambda([this, ColumnName]() -> FText
+		{
+			TStringBuilder<256> StringBuilder;
+			if (const TSharedPtr<FDatabaseViewModel> ViewModel = EditorViewModel.Pin())
+			{
+				const UPoseSearchDatabase* PoseSearchDatabase = ViewModel->GetPoseSearchDatabase();
+				if (FAsyncPoseSearchDatabasesManagement::RequestAsyncBuildIndex(PoseSearchDatabase, ERequestAsyncBuildFlag::ContinueRequest))
+				{
+					if (const FDatabasePreviewActor* FoundPreviewActor = ViewModel->GetPreviewActors().FindByPredicate(
+						[ColumnName](const FDatabasePreviewActor& PreviewActor)
+						{ return *PreviewActor.Actor->GetName() == ColumnName; }))
+					{
+						const int32 PoseIdx = FoundPreviewActor->CurrentPoseIndex;
+						const int32 NumDimensions = PoseSearchDatabase->Schema->SchemaCardinality;
+						const FPoseSearchIndex& SearchIndex = PoseSearchDatabase->GetSearchIndex();
+						TArrayView<float> ReconstructedPoseValuesBuffer((float*)FMemory_Alloca(NumDimensions * sizeof(float)), NumDimensions);
+						const TConstArrayView<float> PoseValues = SearchIndex.Values.IsEmpty() ? SearchIndex.GetReconstructedPoseValues(PoseIdx, ReconstructedPoseValuesBuffer) : SearchIndex.GetPoseValues(PoseIdx);
+
+						const int32 DataOffset = ChannelItem->GetDataOffset();
+						const int32 Cardinality = ChannelItem->GetCardinality();
+
+						if (Cardinality > 1)
+						{
+							// using only one decimal to keep the string compact
+							for (int32 i = 0; i < Cardinality; ++i)
+							{
+								if (i != 0)
+								{
+									StringBuilder.Append(TEXT(", "));
+								}
+								const float Value = PoseValues[i + DataOffset];
+								StringBuilder.Appendf(TEXT("%.1f"), Value);
+							}
+						}
+						else
+						{
+							// using all the float digits 
+							const float Value = PoseValues[DataOffset];
+							StringBuilder.Appendf(TEXT("%f"), Value);
+						}
+					}
+				}
+			}
+
+			return FText::FromString(StringBuilder.ToString());
+		});
+	}
+};
+
+void SDatabaseDataDetails::Construct(const FArguments& Args, TSharedRef<FDatabaseViewModel> InEditorViewModel)
+{
+	EditorViewModel = InEditorViewModel;
+}
+
+void SDatabaseDataDetails::Reconstruct()
+{
+	ChannelItems.Reset();
+
+	if (const TSharedPtr<FDatabaseViewModel> ViewModel = EditorViewModel.Pin())
+	{
+		RebuildChannelItemsTreeRecursively(ChannelItems, ViewModel->GetPoseSearchDatabase()->Schema->GetChannels());
+
+		// generating the header
+		TSharedPtr<SHeaderRow> HeaderRow = SNew(SHeaderRow); //.CanSelectGeneratedColumn(true).Visibility(EVisibility::Visible);
+		HeaderRow->AddColumn(
+			SHeaderRow::Column(TEXT("ChannelName"))
+				.DefaultLabel(LOCTEXT("ChannelName_Header", "Channel Name"))
+				.ToolTipText(LOCTEXT("ChannelName_ToolTip", "Channel Name")));
+		HeaderRow->AddColumn(
+			SHeaderRow::Column(TEXT("DataOffset"))
+			.DefaultLabel(LOCTEXT("DataOffset_Header", "Data Offset"))
+			.ToolTipText(LOCTEXT("DataOffset_ToolTip", "Offset from the beginning of the features data")));
+
+		if (ViewModel->ShouldDrawQueryVector())
+		{
+			HeaderRow->AddColumn(
+				SHeaderRow::Column(TEXT("Query"))
+				.DefaultLabel(LOCTEXT("Query_Header", "Query"))
+				.ToolTipText(LOCTEXT("Query_ToolTip", "Query Values")));
+		}
+
+		for (const FDatabasePreviewActor& PreviewActor : ViewModel->GetPreviewActors())
+		{
+			HeaderRow->AddColumn(
+				SHeaderRow::Column(*PreviewActor.Actor->GetName())
+				.DefaultLabel(FText::FromString(PreviewActor.Sampler.GetAsset()->GetName())));
+		}
+
+		ChannelItemsTreeView = SNew(SChannelItemsTreeView)
+			.TreeItemsSource(&ChannelItems)
+			.HeaderRow(HeaderRow)
+			.OnGenerateRow_Lambda([this](FChannelItemPtr ChannelItem, const TSharedRef<STableViewBase>& OwnerTable)
+			{
+				return SNew(SDatabaseDataDetailsTableRow, OwnerTable, ChannelItem, EditorViewModel.Pin().ToSharedRef());
+			})
+			.OnGetChildren_Lambda([](FChannelItemPtr ChannelItem, TArray<FChannelItemPtr>& OutChildren)
+			{
+				OutChildren.Append(ChannelItem->GetChannelItems());
+			});
+
+		ChildSlot
+		[
+			// @todo: add a SScrollBox to handle the complexity of multiple actors
+			//SNew(SScrollBox)
+			//+ SScrollBox::Slot()
+			//[
+				ChannelItemsTreeView.ToSharedRef()
+			//]
+		];
+
+	//ChannelItemsTreeView->RequestTreeRefresh();
+	}
+}
+
+void SDatabaseDataDetails::RebuildChannelItemsTreeRecursively(TArray<FChannelItemPtr>& ChannelItems, TConstArrayView<TObjectPtr<UPoseSearchFeatureChannel>> Channels)
+{
+	for (const TObjectPtr<UPoseSearchFeatureChannel>& ChannelPtr : Channels)
+	{
+		if (ChannelPtr)
+		{
+			TSharedRef<FChannelItem> ChannelItem = MakeShareable(new FChannelItem(ChannelPtr));
+			ChannelItems.Add(ChannelItem);
+
+			if (ChannelPtr->GetSubChannels().IsEmpty())
+			{
+				for (int32 ChannelComponentIdx = 0; ChannelComponentIdx < ChannelPtr->GetChannelCardinality(); ++ChannelComponentIdx)
+				{
+					TSharedRef<FChannelItem> SubChannelItem = MakeShareable(new FChannelItem(ChannelPtr, ChannelComponentIdx));
+					ChannelItem->GetChannelItems().Add(SubChannelItem);
+				}
+			}
+			else
+			{
+				RebuildChannelItemsTreeRecursively(ChannelItem->GetChannelItems(), ChannelPtr->GetSubChannels());
+			}
+		}
+	}
+}
+
+} // namespace UE::PoseSearch
+
+#undef LOCTEXT_NAMESPACE
diff --git a/Engine/Plugins/Experimental/Animation/PoseSearch/Source/Editor/Private/PoseSearchDatabaseEditor.cpp b/Engine/Plugins/Experimental/Animation/PoseSearch/Source/Editor/Private/PoseSearchDatabaseEditor.cpp
--- a/Engine/Plugins/Experimental/Animation/PoseSearch/Source/Editor/Private/PoseSearchDatabaseEditor.cpp
+++ b/Engine/Plugins/Experimental/Animation/PoseSearch/Source/Editor/Private/PoseSearchDatabaseEditor.cpp
@@ -1,25 +1,26 @@
 // Copyright Epic Games, Inc. All Rights Reserved.
 
 #include "PoseSearchDatabaseEditor.h"
+#include "AdvancedPreviewSceneModule.h"
+#include "Framework/Application/SlateApplication.h"
+#include "GameFramework/WorldSettings.h"
+#include "IStructureDetailsView.h"
+#include "InstancedStruct.h"
+#include "Modules/ModuleManager.h"
+#include "PoseSearch/PoseSearchDerivedData.h"
 #include "PoseSearchDatabaseAssetTree.h"
-#include "SPoseSearchDatabaseViewport.h"
-#include "PoseSearchDatabasePreviewScene.h"
+#include "PoseSearchDatabaseDataDetails.h"
 #include "PoseSearchDatabaseEditorCommands.h"
-#include "PoseSearchDatabaseViewModel.h"
 #include "PoseSearchDatabaseEditorReflection.h"
+#include "PoseSearchDatabasePreviewScene.h"
+#include "PoseSearchDatabaseViewModel.h"
 #include "PoseSearchEditor.h"
-#include "PoseSearch/PoseSearchDerivedData.h"
-#include "Framework/Application/SlateApplication.h"
-#include "GameFramework/WorldSettings.h"
-#include "AdvancedPreviewSceneModule.h"
-#include "InstancedStruct.h"
-#include "IStructureDetailsView.h"
 #include "PropertyEditorModule.h"
-#include "Modules/ModuleManager.h"
-#include "Widgets/Input/SButton.h"
+#include "SPoseSearchDatabaseViewport.h"
+#include "Styling/AppStyle.h"
 #include "Widgets/Docking/SDockTab.h"
+#include "Widgets/Input/SButton.h"
 #include "Widgets/Layout/SScrollBox.h"
-#include "Styling/AppStyle.h"
 
 #define LOCTEXT_NAMESPACE "PoseSearchDatabaseEditor"
 
@@ -36,6 +37,7 @@ namespace UE::PoseSearch
 		static const FName AssetTreeViewID;
 		static const FName SelectionDetailsID;
 		static const FName StatisticsOverview;
+		static const FName DataDetailsID;
 	};
 
 	const FName FDatabaseEditorTabs::AssetDetailsID(TEXT("PoseSearchDatabaseEditorAssetDetailsTabID"));
@@ -44,7 +46,8 @@ namespace UE::PoseSearch
 	const FName FDatabaseEditorTabs::AssetTreeViewID(TEXT("PoseSearchDatabaseEditorAssetTreeViewTabID"));
 	const FName FDatabaseEditorTabs::SelectionDetailsID(TEXT("PoseSearchDatabaseEditorSelectionDetailsID"));
 	const FName FDatabaseEditorTabs::StatisticsOverview(TEXT("PoseSearchDatabaseEditorStatisticsOverviewID"));
-	
+	const FName FDatabaseEditorTabs::DataDetailsID(TEXT("PoseSearchDatabaseEditorDataDetailsTabID"));
+
 	const UPoseSearchDatabase* FDatabaseEditor::GetPoseSearchDatabase() const
 	{
 		return ViewModel.IsValid() ? ViewModel->GetPoseSearchDatabase() : nullptr;
@@ -55,24 +58,16 @@ namespace UE::PoseSearch
 		return ViewModel.IsValid() ? ViewModel->GetPoseSearchDatabase() : nullptr;
 	}
 
-	void FDatabaseEditor::SetSelectedPoseIdx(int32 PoseIdx)
+	void FDatabaseEditor::SetSelectedPoseIdx(int32 PoseIdx, bool bDrawQuery, TConstArrayView<float> InQueryVector)
 	{
 		if (ViewModel.IsValid())
 		{
 			const bool bClearSelection = !FSlateApplication::Get().GetModifierKeys().IsControlDown() || ViewModel->IsEditorSelection();
-			const int32 SourceAssetIdx = ViewModel->SetSelectedNode(PoseIdx, bClearSelection);
+			const int32 SourceAssetIdx = ViewModel->SetSelectedNode(PoseIdx, bClearSelection, bDrawQuery, InQueryVector);
 			AssetTreeWidget->SetSelectedItem(SourceAssetIdx, bClearSelection);
 		}
 	}
 
-	void FDatabaseEditor::SetQueryVector(TConstArrayView<float> InQueryVector)
-	{
-		if (ViewModel.IsValid())
-		{
-			ViewModel->SetQueryVector(InQueryVector);
-		}
-	}
-
 	void FDatabaseEditor::SetDrawQueryVector(bool bValue)
 	{
 		if (ViewModel.IsValid())
@@ -138,7 +133,15 @@ namespace UE::PoseSearch
 
 		// Create view model
 		ViewModel = MakeShared<FDatabaseViewModel>();
-		ViewModel->Initialize(DatabaseAsset, PreviewScene.ToSharedRef());
+
+		// Create Data Details widget
+		DataDetails = SNew(SDatabaseDataDetails, ViewModel.ToSharedRef());
+
+		// Initialize view model
+		ViewModel->Initialize(DatabaseAsset, PreviewScene.ToSharedRef(), DataDetails.ToSharedRef());
+
+		// Initialize DataDetails
+		DataDetails->Reconstruct();
 
 		// Create viewport widget
 		{
@@ -222,7 +225,7 @@ namespace UE::PoseSearch
 		
 		// Define Editor Layout
 		const TSharedRef<FTabManager::FLayout> StandaloneDefaultLayout =
-		FTabManager::NewLayout("Standalone_PoseSearchDatabaseEditor_Layout_v0.12")
+		FTabManager::NewLayout("Standalone_PoseSearchDatabaseEditor_Layout_v0.13")
 			->AddArea
 			(
 				FTabManager::NewPrimaryArea()->SetOrientation(Orient_Horizontal)
@@ -262,6 +265,7 @@ namespace UE::PoseSearch
 						->SetSizeCoefficient(0.4f)
 						->AddTab(FDatabaseEditorTabs::StatisticsOverview, ETabState::OpenedTab)
 						->AddTab(FDatabaseEditorTabs::PreviewSettingsID, ETabState::OpenedTab)
+						->AddTab(FDatabaseEditorTabs::DataDetailsID, ETabState::OpenedTab)
 						->AddTab(FDatabaseEditorTabs::SelectionDetailsID, ETabState::OpenedTab)
 					)
 				)
@@ -287,7 +291,7 @@ namespace UE::PoseSearch
 	{
 		WorkspaceMenuCategory = InTabManager->AddLocalWorkspaceMenuCategory(
 			LOCTEXT("WorkspaceMenu_PoseSearchDbEditor", "Pose Search Database Editor"));
-		auto WorkspaceMenuCategoryRef = WorkspaceMenuCategory.ToSharedRef();
+		TSharedRef<FWorkspaceItem> WorkspaceMenuCategoryRef = WorkspaceMenuCategory.ToSharedRef();
 
 		FAssetEditorToolkit::RegisterTabSpawners(InTabManager);
 
@@ -327,11 +331,18 @@ namespace UE::PoseSearch
 			.SetIcon(FSlateIcon(FAppStyle::GetAppStyleSetName(), "GraphEditor.EventGraph_16x"));
 
 		InTabManager->RegisterTabSpawner(
-		FDatabaseEditorTabs::StatisticsOverview,
-		FOnSpawnTab::CreateSP(this, &FDatabaseEditor::SpawnTab_StatisticsOverview))
-		.SetDisplayName(LOCTEXT("StatisticsOverviewTab", "Statistics Overview"))
-		.SetGroup(WorkspaceMenuCategoryRef)
-		.SetIcon(FSlateIcon(FAppStyle::GetAppStyleSetName(), "GraphEditor.EventGraph_16x"));
+			FDatabaseEditorTabs::StatisticsOverview,
+			FOnSpawnTab::CreateSP(this, &FDatabaseEditor::SpawnTab_StatisticsOverview))
+			.SetDisplayName(LOCTEXT("StatisticsOverviewTab", "Statistics Overview"))
+			.SetGroup(WorkspaceMenuCategoryRef)
+			.SetIcon(FSlateIcon(FAppStyle::GetAppStyleSetName(), "GraphEditor.EventGraph_16x"));
+
+		InTabManager->RegisterTabSpawner(
+			FDatabaseEditorTabs::DataDetailsID,
+			FOnSpawnTab::CreateSP(this, &FDatabaseEditor::SpawnTab_DataDetails))
+			.SetDisplayName(LOCTEXT("DataDetailsTab", "Data Details"))
+			.SetGroup(WorkspaceMenuCategoryRef)
+			.SetIcon(FSlateIcon(FAppStyle::GetAppStyleSetName(), "GraphEditor.EventGraph_16x"));
 	}
 
 	void FDatabaseEditor::UnregisterTabSpawners(const TSharedRef<FTabManager>& InTabManager)
@@ -344,6 +355,7 @@ namespace UE::PoseSearch
 		InTabManager->UnregisterTabSpawner(FDatabaseEditorTabs::AssetTreeViewID);
 		InTabManager->UnregisterTabSpawner(FDatabaseEditorTabs::SelectionDetailsID);
 		InTabManager->UnregisterTabSpawner(FDatabaseEditorTabs::StatisticsOverview);
+		InTabManager->UnregisterTabSpawner(FDatabaseEditorTabs::DataDetailsID);
 	}
 
 	FName FDatabaseEditor::GetToolkitFName() const
@@ -459,6 +471,17 @@ namespace UE::PoseSearch
 		];
 	}
 
+	TSharedRef<SDockTab> FDatabaseEditor::SpawnTab_DataDetails(const FSpawnTabArgs& Args) const
+	{
+		check(Args.GetTabId() == FDatabaseEditorTabs::DataDetailsID);
+		
+		return SNew(SDockTab)
+			.Label(LOCTEXT("DataDetails_Title", "Data Details"))
+			[
+				DataDetails.ToSharedRef()
+			];
+	}
+
 	void FDatabaseEditor::OnFinishedChangingSelectionProperties(const FPropertyChangedEvent& PropertyChangedEvent)
 	{
 	}
@@ -468,7 +491,7 @@ namespace UE::PoseSearch
 		ESelectInfo::Type SelectionType)
 	{
 		// Reset selected objects on all selection widgets
-		for (auto& SelectionWidgetPair : SelectionWidgets)
+		for (TPair<const UScriptStruct*, FSelectionWidget>& SelectionWidgetPair : SelectionWidgets)
 		{
 			SelectionWidgetPair.Value.SelectedReflections.Reset();
 		}
@@ -542,7 +565,7 @@ namespace UE::PoseSearch
 			}
 		}
 
-		for (auto& SelectionWidgetPair : SelectionWidgets)
+		for (TPair<const UScriptStruct*, FSelectionWidget>& SelectionWidgetPair : SelectionWidgets)
 		{
 			FSelectionWidget& SelectionWidget = SelectionWidgetPair.Value;
 
diff --git a/Engine/Plugins/Experimental/Animation/PoseSearch/Source/Editor/Private/PoseSearchDatabaseViewModel.cpp b/Engine/Plugins/Experimental/Animation/PoseSearch/Source/Editor/Private/PoseSearchDatabaseViewModel.cpp
--- a/Engine/Plugins/Experimental/Animation/PoseSearch/Source/Editor/Private/PoseSearchDatabaseViewModel.cpp
+++ b/Engine/Plugins/Experimental/Animation/PoseSearch/Source/Editor/Private/PoseSearchDatabaseViewModel.cpp
@@ -1,28 +1,28 @@
 // Copyright Epic Games, Inc. All Rights Reserved.
 
 #include "PoseSearchDatabaseViewModel.h"
-#include "PoseSearchDatabasePreviewScene.h"
-#include "PoseSearchDatabaseAssetTreeNode.h"
-#include "PoseSearchEditor.h"
-#include "PoseSearch/PoseSearchDatabase.h"
-#include "PoseSearch/PoseSearchDefines.h"
-#include "PoseSearch/PoseSearchDerivedData.h"
-#include "PoseSearch/PoseSearchSchema.h"
-#include "Modules/ModuleManager.h"
 #include "AnimPreviewInstance.h"
 #include "Animation/AnimComposite.h"
 #include "Animation/AnimMontage.h"
 #include "Animation/AnimSequence.h"
+#include "Animation/BlendSpace.h"
 #include "Animation/DebugSkelMeshComponent.h"
 #include "Animation/MirrorDataTable.h"
-#include "Animation/BlendSpace.h"
-#include "EngineUtils.h"
-#include "InstancedStruct.h"
-#include "PropertyEditorModule.h"
-#include "Modules/ModuleManager.h"
 #include "Camera/CameraComponent.h"
 #include "Components/CapsuleComponent.h"
+#include "EngineUtils.h"
 #include "GameFramework/PlayerController.h"
+#include "InstancedStruct.h"
+#include "Modules/ModuleManager.h"
+#include "PoseSearch/PoseSearchDatabase.h"
+#include "PoseSearch/PoseSearchDefines.h"
+#include "PoseSearch/PoseSearchDerivedData.h"
+#include "PoseSearch/PoseSearchSchema.h"
+#include "PoseSearchDatabaseAssetTreeNode.h"
+#include "PoseSearchDatabaseDataDetails.h"
+#include "PoseSearchDatabasePreviewScene.h"
+#include "PoseSearchEditor.h"
+#include "PropertyEditorModule.h"
 
 namespace UE::PoseSearch
 {
@@ -73,10 +73,11 @@ namespace UE::PoseSearch
 		Collector.AddReferencedObject(PoseSearchDatabase);
 	}
 
-	void FDatabaseViewModel::Initialize(UPoseSearchDatabase* InPoseSearchDatabase, const TSharedRef<FDatabasePreviewScene>& InPreviewScene)
+	void FDatabaseViewModel::Initialize(UPoseSearchDatabase* InPoseSearchDatabase, const TSharedRef<FDatabasePreviewScene>& InPreviewScene, const TSharedRef<SDatabaseDataDetails>& InDatabaseDataDetails)
 	{
 		PoseSearchDatabase = InPoseSearchDatabase;
 		PreviewScenePtr = InPreviewScene;
+		DatabaseDataDetails = InDatabaseDataDetails;
 
 		RemovePreviewActors();
 
@@ -235,7 +236,7 @@ namespace UE::PoseSearch
 			const FPoseSearchIndexAsset& IndexAsset = SearchIndex.Assets[PreviewActor.IndexAssetIndex];
 			float CurrentPlayTime = PlayTime + IndexAsset.SamplingInterval.Min + PreviewActor.PlayTimeOffset;
 			FAnimationRuntime::AdvanceTime(false, CurrentPlayTime, CurrentTime, IndexAsset.SamplingInterval.Max);
-
+			 
 			// time to pose index
 			PreviewActor.CurrentPoseIndex = PoseSearchDatabase->GetPoseIndexFromTime(CurrentTime, IndexAsset);
 
@@ -245,7 +246,7 @@ namespace UE::PoseSearch
 
 			if (bQuantizeAnimationToPoseData)
 			{
-				CurrentTime = QuantizedTime;
+				CurrentTime = QuantizedTime; 
 			}
 
 			AnimInstance->SetPosition(CurrentTime * TimeToAssetTimeMultiplier);
@@ -284,7 +285,7 @@ namespace UE::PoseSearch
 		bIsEditorSelection = true;
 		bDrawQueryVector = false;
 
-		for (auto PreviewActor : PreviewActors)
+		for (FDatabasePreviewActor& PreviewActor : PreviewActors)
 		{
 			// @todo: PreviewActor.Actor is a TWeakObjectPtr so it can be null.
 			PreviewActor.Actor->Destroy();
@@ -395,7 +396,7 @@ namespace UE::PoseSearch
 		return false;
 	}
 
-	int32 FDatabaseViewModel::SetSelectedNode(int32 PoseIdx, bool bClearSelection)
+	int32 FDatabaseViewModel::SetSelectedNode(int32 PoseIdx, bool bClearSelection, bool bDrawQuery, TConstArrayView<float> InQueryVector)
 	{
 		int32 SelectedSourceAssetIdx = -1;
 
@@ -405,6 +406,8 @@ namespace UE::PoseSearch
 		}
 
 		bIsEditorSelection = false;
+		bDrawQueryVector = bDrawQuery;
+		QueryVector = InQueryVector;
 
 		if (FAsyncPoseSearchDatabasesManagement::RequestAsyncBuildIndex(PoseSearchDatabase, ERequestAsyncBuildFlag::ContinueRequest))
 		{
@@ -433,6 +436,8 @@ namespace UE::PoseSearch
 
 			ParallelFor(PreviewActors.Num(), [this](int32 PreviewActorIndex) { PreviewActors[PreviewActorIndex].Process(); }, ParallelForFlags);
 
+			DatabaseDataDetails.Pin()->Reconstruct();
+
 			UpdatePreviewActors();
 
 			SetPlayTime(0.f, false);
@@ -478,6 +483,7 @@ namespace UE::PoseSearch
 
 			ParallelFor(PreviewActors.Num(), [this](int32 PreviewActorIndex) { PreviewActors[PreviewActorIndex].Process(); }, ParallelForFlags);
 
+			DatabaseDataDetails.Pin()->Reconstruct();
 			UpdatePreviewActors();
 		}
 
@@ -493,6 +499,15 @@ namespace UE::PoseSearch
 		}
 	}
 
+	void FDatabaseViewModel::SetDrawQueryVector(bool bValue)
+	{
+		if (bDrawQueryVector != bValue)
+		{
+			bDrawQueryVector = bValue;
+			DatabaseDataDetails.Pin()->Reconstruct();
+		}
+	}
+
 	const FPoseSearchIndexAsset* FDatabaseViewModel::GetSelectedActorIndexAsset() const
 	{
 		if (FAsyncPoseSearchDatabasesManagement::RequestAsyncBuildIndex(PoseSearchDatabase, ERequestAsyncBuildFlag::ContinueRequest))
diff --git a/Engine/Plugins/Experimental/Animation/PoseSearch/Source/Editor/Private/PoseSearchDebuggerDatabaseColumns.h b/Engine/Plugins/Experimental/Animation/PoseSearch/Source/Editor/Private/PoseSearchDebuggerDatabaseColumns.h
--- a/Engine/Plugins/Experimental/Animation/PoseSearch/Source/Editor/Private/PoseSearchDebuggerDatabaseColumns.h
+++ b/Engine/Plugins/Experimental/Animation/PoseSearch/Source/Editor/Private/PoseSearchDebuggerDatabaseColumns.h
@@ -3,13 +3,14 @@
 #pragma once
 
 #include "AnimPreviewInstance.h"
+#include "Animation/DebugSkelMeshComponent.h"
 #include "Editor.h"
 #include "IAnimationEditor.h"
 #include "IPersonaToolkit.h"
 #include "PoseSearchDatabaseEditor.h"
 #include "PoseSearchDatabaseViewModel.h"
 #include "PoseSearchDebuggerDatabaseRowData.h"
-#include "Animation/DebugSkelMeshComponent.h"
+#include "PoseSearchDebuggerViewModel.h"
 #include "Preferences/PersonaOptions.h"
 #include "Styling/AppStyle.h"
 #include "Subsystems/AssetEditorSubsystem.h"
@@ -108,7 +109,13 @@ struct FPoseIdx : ITextColumn
 
 struct FDatabaseName : IColumn
 {
-	using IColumn::IColumn;
+	TSharedPtr<FDebuggerViewModel> DebuggerViewModel;
+
+	FDatabaseName(int32 InSortIndex, TSharedPtr<FDebuggerViewModel> InDebuggerViewModel)
+	: IColumn(InSortIndex)
+	, DebuggerViewModel(InDebuggerViewModel)
+	{
+	}
 
 	virtual FText GetLabel() const override
 	{
@@ -123,53 +130,55 @@ struct FDatabaseName : IColumn
 	virtual TSharedRef<SWidget> GenerateWidget(const FRowDataRef& RowData) const override
 	{
 		return SNew(SHyperlink)
-			.Text_Lambda([RowData]() -> FText { return FText::FromString(RowData->SharedData->DatabaseName); })
-			.TextStyle(&FCoreStyle::Get().GetWidgetStyle<FTextBlockStyle>("SmallText"))
-			.ToolTipText_Lambda([RowData]() -> FText
-				{
-					return FText::Format(
-						LOCTEXT("DatabaseHyperlinkTooltipFormat", "Open database '{0}'"),
-						FText::FromString(RowData->SharedData->DatabasePath));
-				})
-			.OnNavigate_Lambda([RowData]()
-				{
-					UObject* Asset = nullptr;
+		.Text_Lambda([RowData]() -> FText
+		{
+			return FText::FromString(RowData->SharedData->DatabaseName);
+		})
+		.TextStyle(&FCoreStyle::Get().GetWidgetStyle<FTextBlockStyle>("SmallText"))
+		.ToolTipText_Lambda([RowData]() -> FText
+		{
+			return FText::Format(
+				LOCTEXT("DatabaseHyperlinkTooltipFormat", "Open database '{0}'"),
+				FText::FromString(RowData->SharedData->DatabasePath));
+		})
+		.OnNavigate_Lambda([this, RowData]()
+		{
+			UObject* Asset = nullptr;
 
-					// Load asset
-					if (UPackage* Package = LoadPackage(NULL, *RowData->SharedData->DatabasePath, LOAD_NoRedirects))
-					{
-						Package->FullyLoad();
+			// Load asset
+			if (UPackage* Package = LoadPackage(NULL, *RowData->SharedData->DatabasePath, LOAD_NoRedirects))
+			{
+				Package->FullyLoad();
 
-						const FString AssetName = FPaths::GetBaseFilename(RowData->SharedData->DatabasePath);
-						Asset = FindObject<UObject>(Package, *AssetName);
-					}
-					else
-					{
-						// Fallback for unsaved assets
-						Asset = FindObject<UObject>(nullptr, *RowData->SharedData->DatabasePath);
-					}
+				const FString AssetName = FPaths::GetBaseFilename(RowData->SharedData->DatabasePath);
+				Asset = FindObject<UObject>(Package, *AssetName);
+			}
+			else
+			{
+				// Fallback for unsaved assets
+				Asset = FindObject<UObject>(nullptr, *RowData->SharedData->DatabasePath);
+			}
 
-					// Open editor
-					if (Asset != nullptr)
+			// Open editor
+			if (Asset != nullptr)
+			{
+				if (UAssetEditorSubsystem* AssetEditorSS = GEditor->GetEditorSubsystem<UAssetEditorSubsystem>())
+				{
+					AssetEditorSS->OpenEditorForAsset(Asset);
+								
+					if (IAssetEditorInstance* Editor = AssetEditorSS->FindEditorForAsset(Asset, true))
 					{
-						if (UAssetEditorSubsystem* AssetEditorSS = GEditor->GetEditorSubsystem<UAssetEditorSubsystem>())
+						if (Editor->GetEditorName() == FName("PoseSearchDatabaseEditor"))
 						{
-							AssetEditorSS->OpenEditorForAsset(Asset);
-								
-							if (IAssetEditorInstance* Editor = AssetEditorSS->FindEditorForAsset(Asset, true))
-							{
-								if (Editor->GetEditorName() == FName("PoseSearchDatabaseEditor"))
-								{
-									FDatabaseEditor* DatabaseEditor = static_cast<FDatabaseEditor*>(Editor);
+							FDatabaseEditor* DatabaseEditor = static_cast<FDatabaseEditor*>(Editor);
 									
-									// Open asset paused and at specific time as seen on the pose search debugger.
-									DatabaseEditor->SetSelectedPoseIdx(RowData->PoseIdx);
-									DatabaseEditor->SetQueryVector(RowData->SharedData->QueryVector);
-								}
-							}
+							// Open asset paused and at specific time as seen on the pose search debugger.
+							DatabaseEditor->SetSelectedPoseIdx(RowData->PoseIdx, DebuggerViewModel->GetDrawQuery(), RowData->SharedData->QueryVector);
 						}
 					}
-				});
+				}
+			}
+		});
 	}
 };
 
diff --git a/Engine/Plugins/Experimental/Animation/PoseSearch/Source/Editor/Private/PoseSearchDebuggerDatabaseView.cpp b/Engine/Plugins/Experimental/Animation/PoseSearch/Source/Editor/Private/PoseSearchDebuggerDatabaseView.cpp
--- a/Engine/Plugins/Experimental/Animation/PoseSearch/Source/Editor/Private/PoseSearchDebuggerDatabaseView.cpp
+++ b/Engine/Plugins/Experimental/Animation/PoseSearch/Source/Editor/Private/PoseSearchDebuggerDatabaseView.cpp
@@ -162,10 +162,12 @@ void SDebuggerDatabaseView::Update(const FTraceMotionMatchingStateMessage& State
 	using namespace DebuggerDatabaseColumns;
 
 	bool bIsVerbose = false;
-	auto DebuggerView = ParentDebuggerViewPtr.Pin();
-	if (DebuggerView.IsValid())
+
+	TSharedPtr<FDebuggerViewModel> ViewModel;
+	if (TSharedPtr<SDebuggerView> DebuggerView = ParentDebuggerViewPtr.Pin())
 	{
-		bIsVerbose = DebuggerView->GetViewModel()->IsVerbose();
+		ViewModel = DebuggerView->GetViewModel();
+		bIsVerbose = ViewModel->IsVerbose();
 	}
 
 	UnfilteredDatabaseRows.Reset();
@@ -257,7 +259,7 @@ void SDebuggerDatabaseView::Update(const FTraceMotionMatchingStateMessage& State
 			}
 		};
 
-		for (auto& UnfilteredRow : UnfilteredDatabaseRows)
+		for (TSharedRef<FDebuggerDatabaseRowData>& UnfilteredRow : UnfilteredDatabaseRows)
 		{
 			ArrayMinMax(UnfilteredRow->CostBreakdowns, MinCostBreakdowns, MaxCostBreakdowns, UE_MAX_FLT);
 		}
@@ -297,7 +299,7 @@ void SDebuggerDatabaseView::Update(const FTraceMotionMatchingStateMessage& State
 
 		TArray<float> CostBreakdownsColorBlend;
 		CostBreakdownsColorBlend.Init(0, CostBreakdownsCardinality);
-		for (auto& UnfilteredRow : UnfilteredDatabaseRows)
+		for (TSharedRef<FDebuggerDatabaseRowData>& UnfilteredRow : UnfilteredDatabaseRows)
 		{
 			if (EnumHasAnyFlags(UnfilteredRow->PoseCandidateFlags, EPoseCandidateFlags::AnyValidMask))
 			{
@@ -312,7 +314,7 @@ void SDebuggerDatabaseView::Update(const FTraceMotionMatchingStateMessage& State
 
 		float MinCost = UE_MAX_FLT;
 		float MaxCost = -UE_MAX_FLT;
-		for (auto& UnfilteredRow : UnfilteredDatabaseRows)
+		for (TSharedRef<FDebuggerDatabaseRowData>& UnfilteredRow : UnfilteredDatabaseRows)
 		{
 			const float Cost = UnfilteredRow->PoseCost.GetTotalCost();
 			MinCost = FMath::Min(MinCost, Cost);
@@ -320,7 +322,7 @@ void SDebuggerDatabaseView::Update(const FTraceMotionMatchingStateMessage& State
 		}
 
 		const float DeltaCost = MaxCost - MinCost;
-		for (auto& UnfilteredRow : UnfilteredDatabaseRows)
+		for (TSharedRef<FDebuggerDatabaseRowData>& UnfilteredRow : UnfilteredDatabaseRows)
 		{
 			if (EnumHasAnyFlags(UnfilteredRow->PoseCandidateFlags, EPoseCandidateFlags::AnyValidMask))
 			{
@@ -343,10 +345,10 @@ void SDebuggerDatabaseView::Update(const FTraceMotionMatchingStateMessage& State
 
 		// Construct all column types
 		int32 ColumnIdx = 0;
-		AddColumn(MakeShared<FDatabaseName>(ColumnIdx++));
+		AddColumn(MakeShared<FDatabaseName>(ColumnIdx++, ViewModel));
 		AddColumn(MakeShared<FAssetName>(ColumnIdx++));
 
-		auto CostColumn = MakeShared<FCost>(ColumnIdx++);
+		TSharedRef<FCost> CostColumn = MakeShared<FCost>(ColumnIdx++);
 		AddColumn(CostColumn);
 
 		int32 LabelIdx = 0;
@@ -522,7 +524,7 @@ void SDebuggerDatabaseView::PopulateViewRows()
 	const bool bHasNameFilter = !Tokens.IsEmpty();
 	FRegexPattern Pattern(FilterString);
 
-	for (const auto& UnfilteredRow : UnfilteredDatabaseRows)
+	for (const TSharedRef<FDebuggerDatabaseRowData>& UnfilteredRow : UnfilteredDatabaseRows)
 	{
 		bool bTryAddToFilteredDatabaseViewRows = true;
 		if (EnumHasAnyFlags(UnfilteredRow->PoseCandidateFlags, EPoseCandidateFlags::Valid_ContinuingPose))
diff --git a/Engine/Plugins/Experimental/Animation/PoseSearch/Source/Editor/Public/PoseSearchDatabaseDataDetails.h b/Engine/Plugins/Experimental/Animation/PoseSearch/Source/Editor/Public/PoseSearchDatabaseDataDetails.h
new file mode 100644
--- /dev/null
+++ b/Engine/Plugins/Experimental/Animation/PoseSearch/Source/Editor/Public/PoseSearchDatabaseDataDetails.h
@@ -0,0 +1,32 @@
+// Copyright Epic Games, Inc. All Rights Reserved.
+
+#pragma once
+
+#include "CoreMinimal.h"
+#include "Widgets/Views/STreeView.h"
+
+class UPoseSearchFeatureChannel;
+
+namespace UE::PoseSearch
+{
+	class FDatabaseViewModel;
+	typedef TSharedPtr<class FChannelItem> FChannelItemPtr;
+	typedef STreeView<FChannelItemPtr> SChannelItemsTreeView;
+
+	class SDatabaseDataDetails : public SCompoundWidget
+	{
+	public:
+		SLATE_BEGIN_ARGS( SDatabaseDataDetails ) {}
+		SLATE_END_ARGS()
+
+		void Construct(const FArguments& Args, TSharedRef<FDatabaseViewModel> InEditorViewModel);
+		void Reconstruct();
+
+	private:
+		static void RebuildChannelItemsTreeRecursively(TArray<FChannelItemPtr>& ChannelItems, TConstArrayView<TObjectPtr<UPoseSearchFeatureChannel>> Channels);
+
+		TSharedPtr<SChannelItemsTreeView> ChannelItemsTreeView;
+		TArray<FChannelItemPtr> ChannelItems;
+		TWeakPtr<FDatabaseViewModel> EditorViewModel;
+	};
+} // namespace UE::PoseSearch
\ No newline at end of file
diff --git a/Engine/Plugins/Experimental/Animation/PoseSearch/Source/Editor/Public/PoseSearchDatabaseEditor.h b/Engine/Plugins/Experimental/Animation/PoseSearch/Source/Editor/Public/PoseSearchDatabaseEditor.h
--- a/Engine/Plugins/Experimental/Animation/PoseSearch/Source/Editor/Public/PoseSearchDatabaseEditor.h
+++ b/Engine/Plugins/Experimental/Animation/PoseSearch/Source/Editor/Public/PoseSearchDatabaseEditor.h
@@ -16,6 +16,7 @@ namespace UE::PoseSearch
 	class FDatabasePreviewScene;
 	class FDatabaseViewModel;
 	class SDatabaseAssetTree;
+	class SDatabaseDataDetails;
 	class SDatabasePreview;
 
 	class FDatabaseEditor : public FAssetEditorToolkit, public FNotifyHook
@@ -40,8 +41,7 @@ namespace UE::PoseSearch
 		FDatabaseViewModel* GetViewModel() const { return ViewModel.Get(); }
 		TSharedPtr<FDatabaseViewModel> GetViewModelSharedPtr() const { return ViewModel; }
 		
-		void SetSelectedPoseIdx(int32 PoseIdx);
-		void SetQueryVector(TConstArrayView<float> InQueryVector);
+		void SetSelectedPoseIdx(int32 PoseIdx, bool bDrawQuery, TConstArrayView<float> InQueryVector);
 		void SetDrawQueryVector(bool bValue);
 		
 		void PreviewBackwardEnd();
@@ -60,6 +60,7 @@ namespace UE::PoseSearch
 		TSharedRef<SDockTab> SpawnTab_AssetTreeView(const FSpawnTabArgs& Args);
 		TSharedRef<SDockTab> SpawnTab_SelectionDetails(const FSpawnTabArgs& Args);
 		TSharedRef<SDockTab> SpawnTab_StatisticsOverview(const FSpawnTabArgs& Args) const;
+		TSharedRef<SDockTab> SpawnTab_DataDetails(const FSpawnTabArgs& Args) const;
 		
 		void OnFinishedChangingSelectionProperties(const FPropertyChangedEvent& PropertyChangedEvent);
 		void OnAssetTreeSelectionChanged(
@@ -79,6 +80,8 @@ namespace UE::PoseSearch
 
 		TSharedPtr<SDatabaseAssetTree> AssetTreeWidget;
 
+		TSharedPtr<SDatabaseDataDetails> DataDetails;
+		
 		TSharedPtr<IDetailsView> EditingAssetWidget;
 
 		TSharedPtr<SVerticalBox> DetailsContainer;
diff --git a/Engine/Plugins/Experimental/Animation/PoseSearch/Source/Editor/Public/PoseSearchDatabaseViewModel.h b/Engine/Plugins/Experimental/Animation/PoseSearch/Source/Editor/Public/PoseSearchDatabaseViewModel.h
--- a/Engine/Plugins/Experimental/Animation/PoseSearch/Source/Editor/Public/PoseSearchDatabaseViewModel.h
+++ b/Engine/Plugins/Experimental/Animation/PoseSearch/Source/Editor/Public/PoseSearchDatabaseViewModel.h
@@ -23,6 +23,7 @@ class UMirrorDataTable;
 namespace UE::PoseSearch
 {
 	class FDatabaseAssetTreeNode;
+	class SDatabaseDataDetails;
 
 	enum class EFeaturesDrawMode : uint8
 	{
@@ -66,7 +67,7 @@ namespace UE::PoseSearch
 		virtual void AddReferencedObjects(FReferenceCollector& Collector) override;
 		virtual FString GetReferencerName() const override { return TEXT("FPoseSearchDatabaseViewModel"); }
 
-		void Initialize(UPoseSearchDatabase* InPoseSearchDatabase, const TSharedRef<FDatabasePreviewScene>& InPreviewScene);
+		void Initialize(UPoseSearchDatabase* InPoseSearchDatabase, const TSharedRef<FDatabasePreviewScene>& InPreviewScene, const TSharedRef<SDatabaseDataDetails>& InDatabaseDataDetails);
 
 		void RemovePreviewActors();
 		void BuildSearchIndex();
@@ -108,13 +109,12 @@ namespace UE::PoseSearch
 		void SetIsEnabled(int32 AnimationAssetIndex, bool bEnabled);
 		bool IsEnabled(int32 AnimationAssetIndex) const;
 
-		int32 SetSelectedNode(int32 PoseIdx, bool bClearSelection);
+		int32 SetSelectedNode(int32 PoseIdx, bool bClearSelection, bool bDrawQuery, TConstArrayView<float> InQueryVector);
 		void SetSelectedNodes(const TArrayView<TSharedPtr<FDatabaseAssetTreeNode>>& InSelectedNodes);
 		void ProcessSelectedActor(AActor* Actor);
 		
-		void SetQueryVector(TConstArrayView<float> InQueryVector) { QueryVector = InQueryVector; }
 		TConstArrayView<float> GetQueryVector() const { return QueryVector; }
-		void SetDrawQueryVector(bool bValue) { bDrawQueryVector = bValue; }
+		void SetDrawQueryVector(bool bValue);
 		bool ShouldDrawQueryVector() const { return bDrawQueryVector && !bIsEditorSelection; }
 
 		const FPoseSearchIndexAsset* GetSelectedActorIndexAsset() const;
@@ -137,6 +137,9 @@ namespace UE::PoseSearch
 		/** Weak pointer to the PreviewScene */
 		TWeakPtr<FDatabasePreviewScene> PreviewScenePtr;
 
+		/** Weak pointer to the SDatabaseDataDetails */
+		TWeakPtr<SDatabaseDataDetails> DatabaseDataDetails;
+
 		/** Actors to be displayed in the preview viewport */
 		TArray<FDatabasePreviewActor> PreviewActors;
 		
