motion matching - fix for sampled asset first frame having inaccurate velocities and positions

diff --git a/Engine/Plugins/Experimental/Animation/PoseSearch/Source/Runtime/Private/PoseSearchAssetIndexer.cpp b/Engine/Plugins/Experimental/Animation/PoseSearch/Source/Runtime/Private/PoseSearchAssetIndexer.cpp
--- a/Engine/Plugins/Experimental/Animation/PoseSearch/Source/Runtime/Private/PoseSearchAssetIndexer.cpp
+++ b/Engine/Plugins/Experimental/Animation/PoseSearch/Source/Runtime/Private/PoseSearchAssetIndexer.cpp
@@ -287,27 +287,17 @@ FAssetIndexer::CachedEntry& FAssetIndexer::GetEntry(float SampleTime)
 
 		const FAssetIndexer::FSampleInfo Sample = GetSampleInfo(SampleTime);
 		float CurrentTime = Sample.ClipTime;
-		float PreviousTime = CurrentTime - SamplingContext->FiniteDelta;
 
 		const bool bLoopable = IndexingContext.AssetSampler->IsLoopable();
 		const float PlayLength = IndexingContext.AssetSampler->GetPlayLength();
+
 		if (!bLoopable)
 		{
-			// if not loopable we clamp the pose at time zero or PlayLength
-			if (PreviousTime < 0.f)
-			{
-				PreviousTime = 0.f;
-				CurrentTime = FMath::Min(SamplingContext->FiniteDelta, PlayLength);
-			}
-			else if (CurrentTime > PlayLength)
-			{
-				CurrentTime = PlayLength;
-				PreviousTime = FMath::Max(PlayLength - SamplingContext->FiniteDelta, 0.f);
-			}
+			CurrentTime = FMath::Clamp(CurrentTime, 0.f, PlayLength);
 		}
 
 		FDeltaTimeRecord DeltaTimeRecord;
-		DeltaTimeRecord.Set(PreviousTime, CurrentTime - PreviousTime);
+		DeltaTimeRecord.Set(CurrentTime, 0.f);
 		// no need to extract root motion here, since we use the precalculated Sample.RootTransform as root transform for the Entry
 		FAnimExtractContext ExtractionCtx(static_cast<double>(CurrentTime), false, DeltaTimeRecord, bLoopable);
 
@@ -322,7 +312,7 @@ FAssetIndexer::CachedEntry& FAssetIndexer::GetEntry(float SampleTime)
 		IndexingContext.AssetSampler->ExtractPose(ExtractionCtx, AnimPoseData);
 		Pose[FCompactPoseBoneIndex(RootBoneIndexType)].SetIdentity();
 
-		if (IndexingContext.bMirrored)
+		if (IndexingContext.bMirrored && IndexingContext.Schema->MirrorDataTable)
 		{
 			FAnimationRuntime::MirrorPose(
 				AnimPoseData.GetPose(),
