motion matching
- SDatabaseDataDetails expansion handling between tree reconstruction
- SDatabaseDataDetails doesn't print data values for expanded rows
- removed no longer necessary properties in UPoseSearchDebuggerReflection
- increased pose search debugger capture candidates to 200
- FPoseSearchFeatureVectorBuilder cleanup
- fix for UPoseSearchSchemaFactory::FactoryCreateNew transaction undos

diff --git a/Engine/Plugins/Experimental/Animation/PoseSearch/Source/Editor/Private/PoseSearchDatabaseDataDetails.cpp b/Engine/Plugins/Experimental/Animation/PoseSearch/Source/Editor/Private/PoseSearchDatabaseDataDetails.cpp
--- a/Engine/Plugins/Experimental/Animation/PoseSearch/Source/Editor/Private/PoseSearchDatabaseDataDetails.cpp
+++ b/Engine/Plugins/Experimental/Animation/PoseSearch/Source/Editor/Private/PoseSearchDatabaseDataDetails.cpp
@@ -16,13 +16,45 @@ namespace UE::PoseSearch
 class FChannelItem
 {
 public:
-	FChannelItem(const UPoseSearchFeatureChannel* InChannel, int32 InChannelComponentIdx = -1)
-	: ChannelComponentIdx(InChannelComponentIdx)
-	, Channel(InChannel)
+	FChannelItem(const UPoseSearchFeatureChannel* Channel, int32 ChannelComponentIdx = -1)
+	: Label(ComputeLabel(Channel, ChannelComponentIdx))
+	, DataOffset(ComputeDataOffset(Channel, ChannelComponentIdx))
+	, Cardinality(ComputeCardinality(Channel, ChannelComponentIdx))
 	{
 	}
 
-	const FString GetLabel() const
+	TArray<FChannelItemPtr>& GetChannelItems()
+	{
+		return ChannelItems;
+	}
+
+	FString GetLabel() const
+	{
+		return Label;
+	}
+
+	int32 GetDataOffset() const
+	{
+		return DataOffset;
+	}
+
+	int32 GetCardinality() const
+	{
+		return Cardinality;
+	}
+
+	bool IsExpanded() const
+	{
+		return bExpanded;
+	}
+
+	void SetExpanded(bool bValue)
+	{
+		bExpanded = bValue;
+	}
+
+private:
+	static FString ComputeLabel(const UPoseSearchFeatureChannel* Channel, int32 ChannelComponentIdx)
 	{
 		if (ChannelComponentIdx >= 0)
 		{
@@ -44,14 +76,14 @@ public:
 		return FString();
 	}
 
-	int32 GetDataOffset() const
+	static int32 ComputeDataOffset(const UPoseSearchFeatureChannel* Channel, int32 ChannelComponentIdx)
 	{
 		int32 DataOffset = 0;
 
 		if (Channel != nullptr)
 		{
 			DataOffset = Channel->GetChannelDataOffset();
-			
+
 			if (ChannelComponentIdx > 0)
 			{
 				DataOffset += ChannelComponentIdx;
@@ -61,7 +93,7 @@ public:
 		return DataOffset;
 	}
 
-	int32 GetCardinality() const
+	static int32 ComputeCardinality(const UPoseSearchFeatureChannel* Channel, int32 ChannelComponentIdx)
 	{
 		int32 Cardinality = 0;
 
@@ -77,14 +109,10 @@ public:
 		return Cardinality;
 	}
 
-	TArray<FChannelItemPtr>& GetChannelItems()
-	{
-		return ChannelItems;
-	}
-
-private:
-	int32 ChannelComponentIdx = -1;
-	TWeakObjectPtr<const UPoseSearchFeatureChannel> Channel;
+	const FString Label;
+	const int32 DataOffset = 0;
+	const int32 Cardinality = 0;
+	bool bExpanded = false;
 	TArray<FChannelItemPtr> ChannelItems;
 };
 
@@ -174,47 +202,50 @@ public:
 		.Text_Lambda([this, ColumnName]() -> FText
 		{
 			TStringBuilder<256> StringBuilder;
-			if (const TSharedPtr<FDatabaseViewModel> ViewModel = EditorViewModel.Pin())
+			
+			if (!ChannelItem->IsExpanded())
 			{
-				const UPoseSearchDatabase* PoseSearchDatabase = ViewModel->GetPoseSearchDatabase();
-				if (FAsyncPoseSearchDatabasesManagement::RequestAsyncBuildIndex(PoseSearchDatabase, ERequestAsyncBuildFlag::ContinueRequest))
+				if (const TSharedPtr<FDatabaseViewModel> ViewModel = EditorViewModel.Pin())
 				{
-					if (const FDatabasePreviewActor* FoundPreviewActor = ViewModel->GetPreviewActors().FindByPredicate(
-						[ColumnName](const FDatabasePreviewActor& PreviewActor)
-						{ return *PreviewActor.Actor->GetName() == ColumnName; }))
+					const UPoseSearchDatabase* PoseSearchDatabase = ViewModel->GetPoseSearchDatabase();
+					if (FAsyncPoseSearchDatabasesManagement::RequestAsyncBuildIndex(PoseSearchDatabase, ERequestAsyncBuildFlag::ContinueRequest))
 					{
-						const int32 PoseIdx = FoundPreviewActor->CurrentPoseIndex;
-						const int32 NumDimensions = PoseSearchDatabase->Schema->SchemaCardinality;
-						const FPoseSearchIndex& SearchIndex = PoseSearchDatabase->GetSearchIndex();
-						TArrayView<float> ReconstructedPoseValuesBuffer((float*)FMemory_Alloca(NumDimensions * sizeof(float)), NumDimensions);
-						const TConstArrayView<float> PoseValues = SearchIndex.Values.IsEmpty() ? SearchIndex.GetReconstructedPoseValues(PoseIdx, ReconstructedPoseValuesBuffer) : SearchIndex.GetPoseValues(PoseIdx);
+						if (const FDatabasePreviewActor* FoundPreviewActor = ViewModel->GetPreviewActors().FindByPredicate(
+							[ColumnName](const FDatabasePreviewActor& PreviewActor)
+							{ return *PreviewActor.Actor->GetName() == ColumnName; }))
+						{
+							const int32 PoseIdx = FoundPreviewActor->CurrentPoseIndex;
+							const int32 NumDimensions = PoseSearchDatabase->Schema->SchemaCardinality;
+							const FPoseSearchIndex& SearchIndex = PoseSearchDatabase->GetSearchIndex();
+							TArrayView<float> ReconstructedPoseValuesBuffer((float*)FMemory_Alloca(NumDimensions * sizeof(float)), NumDimensions);
+							const TConstArrayView<float> PoseValues = SearchIndex.Values.IsEmpty() ? SearchIndex.GetReconstructedPoseValues(PoseIdx, ReconstructedPoseValuesBuffer) : SearchIndex.GetPoseValues(PoseIdx);
 
-						const int32 DataOffset = ChannelItem->GetDataOffset();
-						const int32 Cardinality = ChannelItem->GetCardinality();
+							const int32 DataOffset = ChannelItem->GetDataOffset();
+							const int32 Cardinality = ChannelItem->GetCardinality();
 
-						if (Cardinality > 1)
-						{
-							// using only one decimal to keep the string compact
-							for (int32 i = 0; i < Cardinality; ++i)
+							if (Cardinality > 1)
 							{
-								if (i != 0)
+								// using only one decimal to keep the string compact
+								for (int32 i = 0; i < Cardinality; ++i)
 								{
-									StringBuilder.Append(TEXT(", "));
+									if (i != 0)
+									{
+										StringBuilder.Append(TEXT(", "));
+									}
+									const float Value = PoseValues[i + DataOffset];
+									StringBuilder.Appendf(TEXT("%.1f"), Value);
 								}
-								const float Value = PoseValues[i + DataOffset];
-								StringBuilder.Appendf(TEXT("%.1f"), Value);
 							}
-						}
-						else
-						{
-							// using all the float digits 
-							const float Value = PoseValues[DataOffset];
-							StringBuilder.Appendf(TEXT("%f"), Value);
+							else
+							{
+								// using all the float digits 
+								const float Value = PoseValues[DataOffset];
+								StringBuilder.Appendf(TEXT("%f"), Value);
+							}
 						}
 					}
 				}
 			}
-
 			return FText::FromString(StringBuilder.ToString());
 		});
 	}
@@ -225,8 +256,36 @@ void SDatabaseDataDetails::Construct(const FArguments& Args, TSharedRef<FDatabas
 	EditorViewModel = InEditorViewModel;
 }
 
+void SDatabaseDataDetails::TrackExpandedItems(const TArray<FChannelItemPtr>& ChannelItems, TMap<const FString, bool>& ExpandedItems)
+{
+	for (const FChannelItemPtr& ChannelItem : ChannelItems)
+	{
+		TrackExpandedItems(ChannelItem->GetChannelItems(), ExpandedItems);
+		ExpandedItems.FindOrAdd(ChannelItem->GetLabel()) = ChannelItem->IsExpanded();
+	}
+}
+
+void SDatabaseDataDetails::SetExpandedItems(TArray<FChannelItemPtr>& ChannelItems, const TMap<const FString, bool>& ExpandedItems, SChannelItemsTreeView* ChannelItemsTreeView)
+{
+	for (const FChannelItemPtr& ChannelItem : ChannelItems)
+	{
+		SetExpandedItems(ChannelItem->GetChannelItems(), ExpandedItems, ChannelItemsTreeView);
+		
+		bool bIsExpanded = false;
+		if (const bool* IsExpandedPtr = ExpandedItems.Find(ChannelItem->GetLabel()))
+		{
+			bIsExpanded = *IsExpandedPtr;
+		}
+
+		ChannelItemsTreeView->SetItemExpansion(ChannelItem, bIsExpanded);
+	}
+};
+
 void SDatabaseDataDetails::Reconstruct()
 {
+	TMap<const FString, bool> ExpandedItems;
+	TrackExpandedItems(ChannelItems, ExpandedItems);
+
 	ChannelItems.Reset();
 
 	if (const TSharedPtr<FDatabaseViewModel> ViewModel = EditorViewModel.Pin())
@@ -238,11 +297,14 @@ void SDatabaseDataDetails::Reconstruct()
 		HeaderRow->AddColumn(
 			SHeaderRow::Column(TEXT("ChannelName"))
 				.DefaultLabel(LOCTEXT("ChannelName_Header", "Channel Name"))
-				.ToolTipText(LOCTEXT("ChannelName_ToolTip", "Channel Name")));
+				.ToolTipText(LOCTEXT("ChannelName_ToolTip", "Channel Name"))
+				.FillWidth(0.2f));
+
 		HeaderRow->AddColumn(
 			SHeaderRow::Column(TEXT("DataOffset"))
 			.DefaultLabel(LOCTEXT("DataOffset_Header", "Data Offset"))
-			.ToolTipText(LOCTEXT("DataOffset_ToolTip", "Offset from the beginning of the features data")));
+			.ToolTipText(LOCTEXT("DataOffset_ToolTip", "Offset from the beginning of the features data"))
+			.FillWidth(0.1f));
 
 		if (ViewModel->ShouldDrawQueryVector())
 		{
@@ -269,6 +331,10 @@ void SDatabaseDataDetails::Reconstruct()
 			.OnGetChildren_Lambda([](FChannelItemPtr ChannelItem, TArray<FChannelItemPtr>& OutChildren)
 			{
 				OutChildren.Append(ChannelItem->GetChannelItems());
+			})
+			.OnExpansionChanged_Lambda([](FChannelItemPtr ChannelItem, bool bExpanded)
+			{
+				ChannelItem->SetExpanded(bExpanded);
 			});
 
 		ChildSlot
@@ -281,7 +347,7 @@ void SDatabaseDataDetails::Reconstruct()
 			//]
 		];
 
-	//ChannelItemsTreeView->RequestTreeRefresh();
+		SetExpandedItems(ChannelItems, ExpandedItems, ChannelItemsTreeView.Get());
 	}
 }
 
diff --git a/Engine/Plugins/Experimental/Animation/PoseSearch/Source/Editor/Private/PoseSearchDebuggerReflection.h b/Engine/Plugins/Experimental/Animation/PoseSearch/Source/Editor/Private/PoseSearchDebuggerReflection.h
--- a/Engine/Plugins/Experimental/Animation/PoseSearch/Source/Editor/Private/PoseSearchDebuggerReflection.h
+++ b/Engine/Plugins/Experimental/Animation/PoseSearch/Source/Editor/Private/PoseSearchDebuggerReflection.h
@@ -35,16 +35,4 @@ public:
 
 	UPROPERTY(VisibleAnywhere, Category = "Motion Matching State")
 	float AnimAngularVelocity;
-
-    UPROPERTY(VisibleAnywhere, Category="Pose Vectors")
-	TArray<float> QueryPoseVector;
-    	
-    UPROPERTY(VisibleAnywhere, Category="Pose Vectors")
-	TArray<float> ActivePoseVector;
-
-	UPROPERTY(VisibleAnywhere, Category="Pose Vectors")
-	TArray<float> SelectedPoseVector;
-
-	UPROPERTY(VisibleAnywhere, Category="Pose Vectors")
-	TArray<float> CostVector;
 };
diff --git a/Engine/Plugins/Experimental/Animation/PoseSearch/Source/Editor/Private/PoseSearchDebuggerView.cpp b/Engine/Plugins/Experimental/Animation/PoseSearch/Source/Editor/Private/PoseSearchDebuggerView.cpp
--- a/Engine/Plugins/Experimental/Animation/PoseSearch/Source/Editor/Private/PoseSearchDebuggerView.cpp
+++ b/Engine/Plugins/Experimental/Animation/PoseSearch/Source/Editor/Private/PoseSearchDebuggerView.cpp
@@ -106,37 +106,6 @@ void SDebuggerDetailsView::UpdateReflection(const FTraceMotionMatchingStateMessa
 		Reflection->SimAngularVelocity = State.SimAngularVelocity;
 		Reflection->AnimLinearVelocity = State.AnimLinearVelocity;
 		Reflection->AnimAngularVelocity = State.AnimAngularVelocity;
-
-		// Query pose
-		Reflection->QueryPoseVector.Reset();
-		for (const FTraceMotionMatchingStateDatabaseEntry& DbEntry : State.DatabaseEntries)
-		{
-			const UPoseSearchDatabase* Database = FTraceMotionMatchingState::GetObjectFromId<UPoseSearchDatabase>(DbEntry.DatabaseId);
-			if (Database && Database == CurrentDatabase)
-			{
-				Reflection->QueryPoseVector = DbEntry.QueryVector;
-				break;
-			}
-		}
-
-		// Active pose
-		Reflection->ActivePoseVector = CurrentSearchIndex.GetPoseValuesSafe(CurrentDbPoseIdx);
-	}
-
-	auto DebuggerView = ParentDebuggerViewPtr.Pin();
-	if (DebuggerView.IsValid())
-	{
-		TArray<TSharedRef<FDebuggerDatabaseRowData>> SelectedRows = DebuggerView->GetSelectedDatabaseRows();
-		if (!SelectedRows.IsEmpty())
-		{
-			const TSharedRef<FDebuggerDatabaseRowData>& Selected = SelectedRows[0];
-			if (FAsyncPoseSearchDatabasesManagement::RequestAsyncBuildIndex(Selected->SharedData->SourceDatabase.Get(), ERequestAsyncBuildFlag::ContinueRequest))
-			{
-				const FPoseSearchIndex& SelectedSearchIndex = Selected->SharedData->SourceDatabase->GetSearchIndex();
-				Reflection->SelectedPoseVector = SelectedSearchIndex.GetPoseValuesSafe(Selected->PoseIdx);
-			}
-			Reflection->CostVector = Selected->CostVector;
-		}
 	}
 }
 
@@ -377,8 +346,9 @@ void SDebuggerView::DrawFeatures(
 			for (const FTraceMotionMatchingStateDatabaseEntry& DbEntry : State.DatabaseEntries)
 			{
 				const UPoseSearchDatabase* Database = FTraceMotionMatchingState::GetObjectFromId<UPoseSearchDatabase>(DbEntry.DatabaseId);
-				if (Database && Database == CurrentDatabase && DbEntry.QueryVector.Num() == Database->Schema->SchemaCardinality &&
-					FAsyncPoseSearchDatabasesManagement::RequestAsyncBuildIndex(CurrentDatabase, ERequestAsyncBuildFlag::ContinueRequest))
+				if (Database && Database == CurrentDatabase && 
+					FAsyncPoseSearchDatabasesManagement::RequestAsyncBuildIndex(CurrentDatabase, ERequestAsyncBuildFlag::ContinueRequest) &&
+					DbEntry.QueryVector.Num() == Database->Schema->SchemaCardinality)
 				{
 					FDebugDrawParams DrawParams(&DebuggerWorld, Mesh, nullptr, CurrentDatabase, EDebugDrawFlags::DrawQuery);
 					DrawParams.DrawFeatureVector(DbEntry.QueryVector);
@@ -636,7 +606,7 @@ TSharedRef<SWidget> SDebuggerView::GenerateNodeDebuggerView()
 	
 		// Database view
 		+ SSplitter::Slot()
-		.Value(0.65f)
+		.Value(0.8f)
 		[
 			SNew(SVerticalBox)
 			
@@ -656,7 +626,7 @@ TSharedRef<SWidget> SDebuggerView::GenerateNodeDebuggerView()
 
 		// Details panel view
 		+ SSplitter::Slot()
-		.Value(0.35f)
+		.Value(0.2f)
 		[
 			SAssignNew(DetailsView, SDebuggerDetailsView)
 			.Parent(SharedThis(this))
diff --git a/Engine/Plugins/Experimental/Animation/PoseSearch/Source/Editor/Private/PoseSearchSchemaFactory.cpp b/Engine/Plugins/Experimental/Animation/PoseSearch/Source/Editor/Private/PoseSearchSchemaFactory.cpp
--- a/Engine/Plugins/Experimental/Animation/PoseSearch/Source/Editor/Private/PoseSearchSchemaFactory.cpp
+++ b/Engine/Plugins/Experimental/Animation/PoseSearch/Source/Editor/Private/PoseSearchSchemaFactory.cpp
@@ -20,8 +20,8 @@ UObject* UPoseSearchSchemaFactory::FactoryCreateNew(UClass* Class, UObject* InPa
 	UPoseSearchSchema* Schema = NewObject<UPoseSearchSchema>(InParent, Class, Name, Flags);
 
 	// defaulting UPoseSearchSchema for a meaningful locomotion setup
-	Schema->Channels.Add(NewObject<UPoseSearchFeatureChannel_Trajectory>(Schema, NAME_None));
-	Schema->Channels.Add(NewObject<UPoseSearchFeatureChannel_Pose>(Schema, NAME_None));
+	Schema->AddChannel(NewObject<UPoseSearchFeatureChannel_Trajectory>(Schema, NAME_None, RF_Transactional));
+	Schema->AddChannel(NewObject<UPoseSearchFeatureChannel_Pose>(Schema, NAME_None, RF_Transactional));
 
 	return Schema;
 }
diff --git a/Engine/Plugins/Experimental/Animation/PoseSearch/Source/Editor/Public/PoseSearchDatabaseDataDetails.h b/Engine/Plugins/Experimental/Animation/PoseSearch/Source/Editor/Public/PoseSearchDatabaseDataDetails.h
--- a/Engine/Plugins/Experimental/Animation/PoseSearch/Source/Editor/Public/PoseSearchDatabaseDataDetails.h
+++ b/Engine/Plugins/Experimental/Animation/PoseSearch/Source/Editor/Public/PoseSearchDatabaseDataDetails.h
@@ -24,6 +24,8 @@ namespace UE::PoseSearch
 
 	private:
 		static void RebuildChannelItemsTreeRecursively(TArray<FChannelItemPtr>& ChannelItems, TConstArrayView<TObjectPtr<UPoseSearchFeatureChannel>> Channels);
+		static void TrackExpandedItems(const TArray<FChannelItemPtr>& ChannelItems, TMap<const FString, bool>& ExpandedItems);
+		static void SetExpandedItems(TArray<FChannelItemPtr>& ChannelItems, const TMap<const FString, bool>& ExpandedItems, SChannelItemsTreeView* ChannelItemsTreeView);
 
 		TSharedPtr<SChannelItemsTreeView> ChannelItemsTreeView;
 		TArray<FChannelItemPtr> ChannelItems;
diff --git a/Engine/Plugins/Experimental/Animation/PoseSearch/Source/Runtime/Private/PoseSearchContext.cpp b/Engine/Plugins/Experimental/Animation/PoseSearch/Source/Runtime/Private/PoseSearchContext.cpp
--- a/Engine/Plugins/Experimental/Animation/PoseSearch/Source/Runtime/Private/PoseSearchContext.cpp
+++ b/Engine/Plugins/Experimental/Animation/PoseSearch/Source/Runtime/Private/PoseSearchContext.cpp
@@ -462,9 +462,9 @@ void FSearchContext::UpdateCurrentBestCost(const FPoseSearchCost& PoseSearchCost
 	};
 }
 
-const FPoseSearchFeatureVectorBuilder* FSearchContext::GetCachedQuery(const UPoseSearchSchema* Schema) const
+const FFeatureVectorBuilder* FSearchContext::GetCachedQuery(const UPoseSearchSchema* Schema) const
 {
-	const FPoseSearchFeatureVectorBuilder* CachedQuery = CachedQueries.FindByPredicate([Schema](const FPoseSearchFeatureVectorBuilder& CachedQuery)
+	const FFeatureVectorBuilder* CachedQuery = CachedQueries.FindByPredicate([Schema](const FFeatureVectorBuilder& CachedQuery)
 	{
 		return CachedQuery.GetSchema() == Schema;
 	});
@@ -476,16 +476,16 @@ const FPoseSearchFeatureVectorBuilder* FSearchContext::GetCachedQuery(const UPos
 	return nullptr;
 }
 
-const FPoseSearchFeatureVectorBuilder& FSearchContext::GetOrBuildQuery(const UPoseSearchSchema* Schema)
+const FFeatureVectorBuilder& FSearchContext::GetOrBuildQuery(const UPoseSearchSchema* Schema)
 {
 	check(Schema && Schema->IsValid());
-	const FPoseSearchFeatureVectorBuilder* CachedFeatureVectorBuilder = GetCachedQuery(Schema);
+	const FFeatureVectorBuilder* CachedFeatureVectorBuilder = GetCachedQuery(Schema);
 	if (CachedFeatureVectorBuilder)
 	{
 		return *CachedFeatureVectorBuilder;
 	}
 	
-	FPoseSearchFeatureVectorBuilder& NewCachedQuery = CachedQueries[CachedQueries.AddDefaulted()];
+	FFeatureVectorBuilder& NewCachedQuery = CachedQueries[CachedQueries.AddDefaulted()];
 	Schema->BuildQuery(*this, NewCachedQuery);
 	return NewCachedQuery;
 }
diff --git a/Engine/Plugins/Experimental/Animation/PoseSearch/Source/Runtime/Private/PoseSearchFeatureChannel_FilterCrashingLegs.cpp b/Engine/Plugins/Experimental/Animation/PoseSearch/Source/Runtime/Private/PoseSearchFeatureChannel_FilterCrashingLegs.cpp
--- a/Engine/Plugins/Experimental/Animation/PoseSearch/Source/Runtime/Private/PoseSearchFeatureChannel_FilterCrashingLegs.cpp
+++ b/Engine/Plugins/Experimental/Animation/PoseSearch/Source/Runtime/Private/PoseSearchFeatureChannel_FilterCrashingLegs.cpp
@@ -49,7 +49,7 @@ void UPoseSearchFeatureChannel_FilterCrashingLegs::AddDependentChannels(UPoseSea
 		UPoseSearchFeatureChannel_Position::FindOrAddToSchema(Schema, 0.f, RightFoot.BoneName);
 	}
 }
-void UPoseSearchFeatureChannel_FilterCrashingLegs::BuildQuery(UE::PoseSearch::FSearchContext& SearchContext, FPoseSearchFeatureVectorBuilder& InOutQuery) const
+void UPoseSearchFeatureChannel_FilterCrashingLegs::BuildQuery(UE::PoseSearch::FSearchContext& SearchContext, UE::PoseSearch::FFeatureVectorBuilder& InOutQuery) const
 {
 	using namespace UE::PoseSearch;
 
diff --git a/Engine/Plugins/Experimental/Animation/PoseSearch/Source/Runtime/Private/PoseSearchFeatureChannel_FilterCrashingLegs.h b/Engine/Plugins/Experimental/Animation/PoseSearch/Source/Runtime/Private/PoseSearchFeatureChannel_FilterCrashingLegs.h
--- a/Engine/Plugins/Experimental/Animation/PoseSearch/Source/Runtime/Private/PoseSearchFeatureChannel_FilterCrashingLegs.h
+++ b/Engine/Plugins/Experimental/Animation/PoseSearch/Source/Runtime/Private/PoseSearchFeatureChannel_FilterCrashingLegs.h
@@ -54,7 +54,7 @@ public:
 
 	// UPoseSearchFeatureChannel interface
 	virtual void Finalize(UPoseSearchSchema* Schema) override;
-	virtual void BuildQuery(UE::PoseSearch::FSearchContext& SearchContext, FPoseSearchFeatureVectorBuilder& InOutQuery) const override;
+	virtual void BuildQuery(UE::PoseSearch::FSearchContext& SearchContext, UE::PoseSearch::FFeatureVectorBuilder& InOutQuery) const override;
 
 	virtual void AddDependentChannels(UPoseSearchSchema* Schema) const override;
 
diff --git a/Engine/Plugins/Experimental/Animation/PoseSearch/Source/Runtime/Private/PoseSearchFeatureChannel_Group.cpp b/Engine/Plugins/Experimental/Animation/PoseSearch/Source/Runtime/Private/PoseSearchFeatureChannel_Group.cpp
--- a/Engine/Plugins/Experimental/Animation/PoseSearch/Source/Runtime/Private/PoseSearchFeatureChannel_Group.cpp
+++ b/Engine/Plugins/Experimental/Animation/PoseSearch/Source/Runtime/Private/PoseSearchFeatureChannel_Group.cpp
@@ -51,7 +51,7 @@ void UPoseSearchFeatureChannel_GroupBase::IndexAsset(UE::PoseSearch::FAssetIndex
 }
 #endif // WITH_EDITOR
 
-void UPoseSearchFeatureChannel_GroupBase::BuildQuery(UE::PoseSearch::FSearchContext& SearchContext, FPoseSearchFeatureVectorBuilder& InOutQuery) const
+void UPoseSearchFeatureChannel_GroupBase::BuildQuery(UE::PoseSearch::FSearchContext& SearchContext, UE::PoseSearch::FFeatureVectorBuilder& InOutQuery) const
 {
 	for (const TObjectPtr<UPoseSearchFeatureChannel>& SubChannelPtr : GetSubChannels())
 	{
diff --git a/Engine/Plugins/Experimental/Animation/PoseSearch/Source/Runtime/Private/PoseSearchFeatureChannel_Group.h b/Engine/Plugins/Experimental/Animation/PoseSearch/Source/Runtime/Private/PoseSearchFeatureChannel_Group.h
--- a/Engine/Plugins/Experimental/Animation/PoseSearch/Source/Runtime/Private/PoseSearchFeatureChannel_Group.h
+++ b/Engine/Plugins/Experimental/Animation/PoseSearch/Source/Runtime/Private/PoseSearchFeatureChannel_Group.h
@@ -15,7 +15,7 @@ class POSESEARCH_API UPoseSearchFeatureChannel_GroupBase : public UPoseSearchFea
 public:
 	// UPoseSearchFeatureChannel interface
 	virtual void Finalize(UPoseSearchSchema* Schema) override;
-	virtual void BuildQuery(UE::PoseSearch::FSearchContext& SearchContext, FPoseSearchFeatureVectorBuilder& InOutQuery) const override;
+	virtual void BuildQuery(UE::PoseSearch::FSearchContext& SearchContext, UE::PoseSearch::FFeatureVectorBuilder& InOutQuery) const override;
 	virtual void AddDependentChannels(UPoseSearchSchema* Schema) const override; 
 
 #if ENABLE_DRAW_DEBUG
diff --git a/Engine/Plugins/Experimental/Animation/PoseSearch/Source/Runtime/Private/PoseSearchFeatureChannel_Heading.cpp b/Engine/Plugins/Experimental/Animation/PoseSearch/Source/Runtime/Private/PoseSearchFeatureChannel_Heading.cpp
--- a/Engine/Plugins/Experimental/Animation/PoseSearch/Source/Runtime/Private/PoseSearchFeatureChannel_Heading.cpp
+++ b/Engine/Plugins/Experimental/Animation/PoseSearch/Source/Runtime/Private/PoseSearchFeatureChannel_Heading.cpp
@@ -40,7 +40,7 @@ FVector UPoseSearchFeatureChannel_Heading::GetAxis(const FQuat& Rotation) const
 	return FVector::XAxisVector;
 }
 
-void UPoseSearchFeatureChannel_Heading::BuildQuery(UE::PoseSearch::FSearchContext& SearchContext, FPoseSearchFeatureVectorBuilder& InOutQuery) const
+void UPoseSearchFeatureChannel_Heading::BuildQuery(UE::PoseSearch::FSearchContext& SearchContext, UE::PoseSearch::FFeatureVectorBuilder& InOutQuery) const
 {
 	using namespace UE::PoseSearch;
 
diff --git a/Engine/Plugins/Experimental/Animation/PoseSearch/Source/Runtime/Private/PoseSearchFeatureChannel_Heading.h b/Engine/Plugins/Experimental/Animation/PoseSearch/Source/Runtime/Private/PoseSearchFeatureChannel_Heading.h
--- a/Engine/Plugins/Experimental/Animation/PoseSearch/Source/Runtime/Private/PoseSearchFeatureChannel_Heading.h
+++ b/Engine/Plugins/Experimental/Animation/PoseSearch/Source/Runtime/Private/PoseSearchFeatureChannel_Heading.h
@@ -49,7 +49,7 @@ public:
 
 	// UPoseSearchFeatureChannel interface
 	virtual void Finalize(UPoseSearchSchema* Schema) override;
-	virtual void BuildQuery(UE::PoseSearch::FSearchContext& SearchContext, FPoseSearchFeatureVectorBuilder& InOutQuery) const override;
+	virtual void BuildQuery(UE::PoseSearch::FSearchContext& SearchContext, UE::PoseSearch::FFeatureVectorBuilder& InOutQuery) const override;
 
 	virtual void AddDependentChannels(UPoseSearchSchema* Schema) const override;
 
diff --git a/Engine/Plugins/Experimental/Animation/PoseSearch/Source/Runtime/Private/PoseSearchFeatureChannel_PermutationTime.cpp b/Engine/Plugins/Experimental/Animation/PoseSearch/Source/Runtime/Private/PoseSearchFeatureChannel_PermutationTime.cpp
--- a/Engine/Plugins/Experimental/Animation/PoseSearch/Source/Runtime/Private/PoseSearchFeatureChannel_PermutationTime.cpp
+++ b/Engine/Plugins/Experimental/Animation/PoseSearch/Source/Runtime/Private/PoseSearchFeatureChannel_PermutationTime.cpp
@@ -11,8 +11,7 @@ void UPoseSearchFeatureChannel_PermutationTime::FindOrAddToSchema(UPoseSearchSch
 	{
 		UPoseSearchFeatureChannel_PermutationTime* PermutationTime = NewObject<UPoseSearchFeatureChannel_PermutationTime>(Schema, NAME_None, RF_Transient);
 		PermutationTime->Weight = 0.f;
-		PermutationTime->Finalize(Schema);
-		Schema->FinalizedChannels.Add(PermutationTime);
+		Schema->AddTemporaryChannel(PermutationTime);
 	}
 }
 
@@ -23,7 +22,7 @@ void UPoseSearchFeatureChannel_PermutationTime::Finalize(UPoseSearchSchema* Sche
 	Schema->SchemaCardinality += ChannelCardinality;
 }
 
-void UPoseSearchFeatureChannel_PermutationTime::BuildQuery(UE::PoseSearch::FSearchContext& SearchContext, FPoseSearchFeatureVectorBuilder& InOutQuery) const
+void UPoseSearchFeatureChannel_PermutationTime::BuildQuery(UE::PoseSearch::FSearchContext& SearchContext, UE::PoseSearch::FFeatureVectorBuilder& InOutQuery) const
 {
 	using namespace UE::PoseSearch;
 	FFeatureVectorHelper::EncodeFloat(InOutQuery.EditValues(), ChannelDataOffset, SearchContext.GetDesiredPermutationTimeOffset());
diff --git a/Engine/Plugins/Experimental/Animation/PoseSearch/Source/Runtime/Private/PoseSearchFeatureChannel_PermutationTime.h b/Engine/Plugins/Experimental/Animation/PoseSearch/Source/Runtime/Private/PoseSearchFeatureChannel_PermutationTime.h
--- a/Engine/Plugins/Experimental/Animation/PoseSearch/Source/Runtime/Private/PoseSearchFeatureChannel_PermutationTime.h
+++ b/Engine/Plugins/Experimental/Animation/PoseSearch/Source/Runtime/Private/PoseSearchFeatureChannel_PermutationTime.h
@@ -16,7 +16,7 @@ public:
 
 	// UPoseSearchFeatureChannel interface
 	virtual void Finalize(UPoseSearchSchema* Schema) override;
-	virtual void BuildQuery(UE::PoseSearch::FSearchContext& SearchContext, FPoseSearchFeatureVectorBuilder& InOutQuery) const override;
+	virtual void BuildQuery(UE::PoseSearch::FSearchContext& SearchContext, UE::PoseSearch::FFeatureVectorBuilder& InOutQuery) const override;
 
 #if WITH_EDITOR
 	virtual void FillWeights(TArray<float>& Weights) const override;
diff --git a/Engine/Plugins/Experimental/Animation/PoseSearch/Source/Runtime/Private/PoseSearchFeatureChannel_Phase.cpp b/Engine/Plugins/Experimental/Animation/PoseSearch/Source/Runtime/Private/PoseSearchFeatureChannel_Phase.cpp
--- a/Engine/Plugins/Experimental/Animation/PoseSearch/Source/Runtime/Private/PoseSearchFeatureChannel_Phase.cpp
+++ b/Engine/Plugins/Experimental/Animation/PoseSearch/Source/Runtime/Private/PoseSearchFeatureChannel_Phase.cpp
@@ -296,7 +296,7 @@ void UPoseSearchFeatureChannel_Phase::AddDependentChannels(UPoseSearchSchema* Sc
 	}
 }
 
-void UPoseSearchFeatureChannel_Phase::BuildQuery(UE::PoseSearch::FSearchContext& SearchContext, FPoseSearchFeatureVectorBuilder& InOutQuery) const
+void UPoseSearchFeatureChannel_Phase::BuildQuery(UE::PoseSearch::FSearchContext& SearchContext, UE::PoseSearch::FFeatureVectorBuilder& InOutQuery) const
 {
 	using namespace UE::PoseSearch;
 
diff --git a/Engine/Plugins/Experimental/Animation/PoseSearch/Source/Runtime/Private/PoseSearchFeatureChannel_Phase.h b/Engine/Plugins/Experimental/Animation/PoseSearch/Source/Runtime/Private/PoseSearchFeatureChannel_Phase.h
--- a/Engine/Plugins/Experimental/Animation/PoseSearch/Source/Runtime/Private/PoseSearchFeatureChannel_Phase.h
+++ b/Engine/Plugins/Experimental/Animation/PoseSearch/Source/Runtime/Private/PoseSearchFeatureChannel_Phase.h
@@ -29,7 +29,7 @@ public:
 
 	// UPoseSearchFeatureChannel interface
 	virtual void Finalize(UPoseSearchSchema* Schema) override;
-	virtual void BuildQuery(UE::PoseSearch::FSearchContext& SearchContext, FPoseSearchFeatureVectorBuilder& InOutQuery) const override;
+	virtual void BuildQuery(UE::PoseSearch::FSearchContext& SearchContext, UE::PoseSearch::FFeatureVectorBuilder& InOutQuery) const override;
 
 	virtual void AddDependentChannels(UPoseSearchSchema* Schema) const override;
 
diff --git a/Engine/Plugins/Experimental/Animation/PoseSearch/Source/Runtime/Private/PoseSearchFeatureChannel_Position.cpp b/Engine/Plugins/Experimental/Animation/PoseSearch/Source/Runtime/Private/PoseSearchFeatureChannel_Position.cpp
--- a/Engine/Plugins/Experimental/Animation/PoseSearch/Source/Runtime/Private/PoseSearchFeatureChannel_Position.cpp
+++ b/Engine/Plugins/Experimental/Animation/PoseSearch/Source/Runtime/Private/PoseSearchFeatureChannel_Position.cpp
@@ -28,8 +28,7 @@ void UPoseSearchFeatureChannel_Position::FindOrAddToSchema(UPoseSearchSchema* Sc
 		// @todo: perhaps add a tunable color for injected channels
 		Position->DebugColor = FLinearColor::Gray;
 		Position->PermutationTimeType = PermutationTimeType;
-		Position->Finalize(Schema);
-		Schema->FinalizedChannels.Add(Position);
+		Schema->AddTemporaryChannel(Position);
 	}
 }
 
@@ -55,7 +54,7 @@ void UPoseSearchFeatureChannel_Position::AddDependentChannels(UPoseSearchSchema*
 	}
 }
 
-void UPoseSearchFeatureChannel_Position::BuildQuery(UE::PoseSearch::FSearchContext& SearchContext, FPoseSearchFeatureVectorBuilder& InOutQuery) const
+void UPoseSearchFeatureChannel_Position::BuildQuery(UE::PoseSearch::FSearchContext& SearchContext, UE::PoseSearch::FFeatureVectorBuilder& InOutQuery) const
 {
 	using namespace UE::PoseSearch;
 
diff --git a/Engine/Plugins/Experimental/Animation/PoseSearch/Source/Runtime/Private/PoseSearchFeatureChannel_Position.h b/Engine/Plugins/Experimental/Animation/PoseSearch/Source/Runtime/Private/PoseSearchFeatureChannel_Position.h
--- a/Engine/Plugins/Experimental/Animation/PoseSearch/Source/Runtime/Private/PoseSearchFeatureChannel_Position.h
+++ b/Engine/Plugins/Experimental/Animation/PoseSearch/Source/Runtime/Private/PoseSearchFeatureChannel_Position.h
@@ -44,7 +44,7 @@ public:
 
 	// UPoseSearchFeatureChannel interface
 	virtual void Finalize(UPoseSearchSchema* Schema) override;
-	virtual void BuildQuery(UE::PoseSearch::FSearchContext& SearchContext, FPoseSearchFeatureVectorBuilder& InOutQuery) const override;
+	virtual void BuildQuery(UE::PoseSearch::FSearchContext& SearchContext, UE::PoseSearch::FFeatureVectorBuilder& InOutQuery) const override;
 
 	virtual EPermutationTimeType GetPermutationTimeType() const override { return PermutationTimeType; }
 	virtual void AddDependentChannels(UPoseSearchSchema* Schema) const override;
diff --git a/Engine/Plugins/Experimental/Animation/PoseSearch/Source/Runtime/Private/PoseSearchFeatureChannel_Velocity.cpp b/Engine/Plugins/Experimental/Animation/PoseSearch/Source/Runtime/Private/PoseSearchFeatureChannel_Velocity.cpp
--- a/Engine/Plugins/Experimental/Animation/PoseSearch/Source/Runtime/Private/PoseSearchFeatureChannel_Velocity.cpp
+++ b/Engine/Plugins/Experimental/Animation/PoseSearch/Source/Runtime/Private/PoseSearchFeatureChannel_Velocity.cpp
@@ -26,7 +26,7 @@ void UPoseSearchFeatureChannel_Velocity::AddDependentChannels(UPoseSearchSchema*
 	}
 }
 
-void UPoseSearchFeatureChannel_Velocity::BuildQuery(UE::PoseSearch::FSearchContext& SearchContext, FPoseSearchFeatureVectorBuilder& InOutQuery) const
+void UPoseSearchFeatureChannel_Velocity::BuildQuery(UE::PoseSearch::FSearchContext& SearchContext, UE::PoseSearch::FFeatureVectorBuilder& InOutQuery) const
 {
 	using namespace UE::PoseSearch;
 
diff --git a/Engine/Plugins/Experimental/Animation/PoseSearch/Source/Runtime/Private/PoseSearchFeatureChannel_Velocity.h b/Engine/Plugins/Experimental/Animation/PoseSearch/Source/Runtime/Private/PoseSearchFeatureChannel_Velocity.h
--- a/Engine/Plugins/Experimental/Animation/PoseSearch/Source/Runtime/Private/PoseSearchFeatureChannel_Velocity.h
+++ b/Engine/Plugins/Experimental/Animation/PoseSearch/Source/Runtime/Private/PoseSearchFeatureChannel_Velocity.h
@@ -43,7 +43,7 @@ public:
 
 	// UPoseSearchFeatureChannel interface
 	virtual void Finalize(UPoseSearchSchema* Schema) override;
-	virtual void BuildQuery(UE::PoseSearch::FSearchContext& SearchContext, FPoseSearchFeatureVectorBuilder& InOutQuery) const override;
+	virtual void BuildQuery(UE::PoseSearch::FSearchContext& SearchContext, UE::PoseSearch::FFeatureVectorBuilder& InOutQuery) const override;
 
 	virtual void AddDependentChannels(UPoseSearchSchema* Schema) const override;
 
diff --git a/Engine/Plugins/Experimental/Animation/PoseSearch/Source/Runtime/Private/PoseSearchLibrary.cpp b/Engine/Plugins/Experimental/Animation/PoseSearch/Source/Runtime/Private/PoseSearchLibrary.cpp
--- a/Engine/Plugins/Experimental/Animation/PoseSearch/Source/Runtime/Private/PoseSearchLibrary.cpp
+++ b/Engine/Plugins/Experimental/Animation/PoseSearch/Source/Runtime/Private/PoseSearchLibrary.cpp
@@ -115,7 +115,7 @@ void FMotionMatchingState::UpdateWantedPlayRate(const UE::PoseSearch::FSearchCon
 	{
 		if (!FMath::IsNearlyEqual(PlayRate.Min, 1.f, UE_KINDA_SMALL_NUMBER) || !FMath::IsNearlyEqual(PlayRate.Max, 1.f, UE_KINDA_SMALL_NUMBER))
 		{
-			if (const FPoseSearchFeatureVectorBuilder* PoseSearchFeatureVectorBuilder = SearchContext.GetCachedQuery(CurrentSearchResult.Database->Schema))
+			if (const UE::PoseSearch::FFeatureVectorBuilder* PoseSearchFeatureVectorBuilder = SearchContext.GetCachedQuery(CurrentSearchResult.Database->Schema))
 			{
 				if (const UPoseSearchFeatureChannel_Trajectory* TrajectoryChannel = CurrentSearchResult.Database->Schema->FindFirstChannelOfType<UPoseSearchFeatureChannel_Trajectory>())
 				{
diff --git a/Engine/Plugins/Experimental/Animation/PoseSearch/Source/Runtime/Private/PoseSearchResult.cpp b/Engine/Plugins/Experimental/Animation/PoseSearch/Source/Runtime/Private/PoseSearchResult.cpp
--- a/Engine/Plugins/Experimental/Animation/PoseSearch/Source/Runtime/Private/PoseSearchResult.cpp
+++ b/Engine/Plugins/Experimental/Animation/PoseSearch/Source/Runtime/Private/PoseSearchResult.cpp
@@ -6,9 +6,12 @@
 #include "PoseSearch/PoseSearchDatabase.h"
 #include "PoseSearch/PoseSearchSchema.h"
 
+namespace UE::PoseSearch
+{
+
 //////////////////////////////////////////////////////////////////////////
-// FPoseSearchFeatureVectorBuilder
-void FPoseSearchFeatureVectorBuilder::Init(const UPoseSearchSchema* InSchema)
+// FFeatureVectorBuilder
+void FFeatureVectorBuilder::Init(const UPoseSearchSchema* InSchema)
 {
 	check(InSchema && InSchema->IsValid());
 	Schema = InSchema;
@@ -16,14 +19,12 @@ void FPoseSearchFeatureVectorBuilder::Init(const UPoseSearchSchema* InSchema)
 	Values.SetNumZeroed(Schema->SchemaCardinality);
 }
 
-void FPoseSearchFeatureVectorBuilder::Reset()
+void FFeatureVectorBuilder::Reset()
 {
 	Schema = nullptr;
 	Values.Reset();
 }
 
-namespace UE::PoseSearch
-{
 //////////////////////////////////////////////////////////////////////////
 // FSearchResult
 void FSearchResult::Update(float NewAssetTime)
diff --git a/Engine/Plugins/Experimental/Animation/PoseSearch/Source/Runtime/Private/PoseSearchSchema.cpp b/Engine/Plugins/Experimental/Animation/PoseSearch/Source/Runtime/Private/PoseSearchSchema.cpp
--- a/Engine/Plugins/Experimental/Animation/PoseSearch/Source/Runtime/Private/PoseSearchSchema.cpp
+++ b/Engine/Plugins/Experimental/Animation/PoseSearch/Source/Runtime/Private/PoseSearchSchema.cpp
@@ -12,7 +12,18 @@ bool UPoseSearchSchema::IsValid() const
 	return Skeleton != nullptr;
 }
 
-void UPoseSearchSchema::BuildQuery(UE::PoseSearch::FSearchContext& SearchContext, FPoseSearchFeatureVectorBuilder& InOutQuery) const
+void UPoseSearchSchema::AddChannel(UPoseSearchFeatureChannel* Channel)
+{
+	Channels.Add(Channel);
+}
+
+void UPoseSearchSchema::AddTemporaryChannel(UPoseSearchFeatureChannel* TemporaryChannel)
+{
+	TemporaryChannel->Finalize(this);
+	FinalizedChannels.Add(TemporaryChannel);
+}
+
+void UPoseSearchSchema::BuildQuery(UE::PoseSearch::FSearchContext& SearchContext, UE::PoseSearch::FFeatureVectorBuilder& InOutQuery) const
 {
 	QUICK_SCOPE_CYCLE_COUNTER(STAT_PoseSearch_BuildQuery);
 
diff --git a/Engine/Plugins/Experimental/Animation/PoseSearch/Source/Runtime/Public/PoseSearch/PoseSearchContext.h b/Engine/Plugins/Experimental/Animation/PoseSearch/Source/Runtime/Public/PoseSearch/PoseSearchContext.h
--- a/Engine/Plugins/Experimental/Animation/PoseSearch/Source/Runtime/Public/PoseSearch/PoseSearchContext.h
+++ b/Engine/Plugins/Experimental/Animation/PoseSearch/Source/Runtime/Public/PoseSearch/PoseSearchContext.h
@@ -152,8 +152,8 @@ struct POSESEARCH_API FSearchContext
 	void UpdateCurrentBestCost(const FPoseSearchCost& PoseSearchCost);
 	float GetCurrentBestTotalCost() const { return CurrentBestTotalCost; }
 
-	const FPoseSearchFeatureVectorBuilder& GetOrBuildQuery(const UPoseSearchSchema* Schema);
-	const FPoseSearchFeatureVectorBuilder* GetCachedQuery(const UPoseSearchSchema* Schema) const;
+	const FFeatureVectorBuilder& GetOrBuildQuery(const UPoseSearchSchema* Schema);
+	const FFeatureVectorBuilder* GetCachedQuery(const UPoseSearchSchema* Schema) const;
 
 	bool IsCurrentResultFromDatabase(const UPoseSearchDatabase* Database) const;
 
@@ -196,7 +196,7 @@ private:
 
 	// transforms cached in component space
 	FCachedTransforms<FTransform> CachedTransforms;
-	TArray<FPoseSearchFeatureVectorBuilder, TInlineAllocator<8>> CachedQueries;
+	TArray<FFeatureVectorBuilder, TInlineAllocator<8>> CachedQueries;
 
 	float CurrentBestTotalCost = MAX_flt;
 	
@@ -264,7 +264,7 @@ public:
 	private:
 		TArray<FPoseCandidateBase> PoseCandidateHeap;
 		TMap<int32, EPoseCandidateFlags> PoseIdxToFlags;
-		int32 MaxPoseCandidates = 100;
+		int32 MaxPoseCandidates = 200;
 	};
 	
 	FBestPoseCandidates BestCandidates;
diff --git a/Engine/Plugins/Experimental/Animation/PoseSearch/Source/Runtime/Public/PoseSearch/PoseSearchFeatureChannel.h b/Engine/Plugins/Experimental/Animation/PoseSearch/Source/Runtime/Public/PoseSearch/PoseSearchFeatureChannel.h
--- a/Engine/Plugins/Experimental/Animation/PoseSearch/Source/Runtime/Public/PoseSearch/PoseSearchFeatureChannel.h
+++ b/Engine/Plugins/Experimental/Animation/PoseSearch/Source/Runtime/Public/PoseSearch/PoseSearchFeatureChannel.h
@@ -9,7 +9,6 @@
 
 class UPoseSearchSchema;
 struct FPoseSearchPoseMetadata;
-struct FPoseSearchFeatureVectorBuilder;
 
 UENUM()
 enum class EComponentStrippingVector : uint8
@@ -60,6 +59,7 @@ namespace UE::PoseSearch
 {
 
 struct FDebugDrawParams;
+struct FFeatureVectorBuilder;
 struct FSearchContext;
 
 #if WITH_EDITOR
@@ -113,7 +113,7 @@ public:
 	virtual void Finalize(UPoseSearchSchema* Schema) PURE_VIRTUAL(UPoseSearchFeatureChannel::Finalize, );
 	
 	// Called at runtime to add this channel's data to the query pose vector
-	virtual void BuildQuery(UE::PoseSearch::FSearchContext& SearchContext, FPoseSearchFeatureVectorBuilder& InOutQuery) const PURE_VIRTUAL(UPoseSearchFeatureChannel::BuildQuery, );
+	virtual void BuildQuery(UE::PoseSearch::FSearchContext& SearchContext, UE::PoseSearch::FFeatureVectorBuilder& InOutQuery) const PURE_VIRTUAL(UPoseSearchFeatureChannel::BuildQuery, );
 
 	// UPoseSearchFeatureChannels can hold sub channels
 	virtual TArrayView<TObjectPtr<UPoseSearchFeatureChannel>> GetSubChannels() { return TArrayView<TObjectPtr<UPoseSearchFeatureChannel>>(); }
diff --git a/Engine/Plugins/Experimental/Animation/PoseSearch/Source/Runtime/Public/PoseSearch/PoseSearchResult.h b/Engine/Plugins/Experimental/Animation/PoseSearch/Source/Runtime/Public/PoseSearch/PoseSearchResult.h
--- a/Engine/Plugins/Experimental/Animation/PoseSearch/Source/Runtime/Public/PoseSearch/PoseSearchResult.h
+++ b/Engine/Plugins/Experimental/Animation/PoseSearch/Source/Runtime/Public/PoseSearch/PoseSearchResult.h
@@ -3,21 +3,20 @@
 #pragma once
 
 #include "PoseSearch/PoseSearchCost.h"
-#include "PoseSearchResult.generated.h"
 
 struct FPoseSearchIndexAsset;
 class UPoseSearchDatabase;
 class UPoseSearchSchema;
 
+namespace UE::PoseSearch
+{
+
 /**
 * float buffer of features according to a UPoseSearchSchema layout.
 * FFeatureVectorBuilder is used to build search queries at runtime and for adding samples during search index construction.
 */
-USTRUCT()
-struct POSESEARCH_API FPoseSearchFeatureVectorBuilder
+struct FFeatureVectorBuilder
 {
-	GENERATED_BODY()
-
 public:
 	void Init(const UPoseSearchSchema* Schema);
 	void Reset();
@@ -28,13 +27,9 @@ public:
 	TConstArrayView<float> GetValues() const { return Values; }
 
 private:
-	UPROPERTY(Transient)
 	TWeakObjectPtr<const UPoseSearchSchema> Schema = nullptr;
 	TArray<float> Values;
 };
-
-namespace UE::PoseSearch
-{
 	
 struct FSearchResult
 {
diff --git a/Engine/Plugins/Experimental/Animation/PoseSearch/Source/Runtime/Public/PoseSearch/PoseSearchSchema.h b/Engine/Plugins/Experimental/Animation/PoseSearch/Source/Runtime/Public/PoseSearch/PoseSearchSchema.h
--- a/Engine/Plugins/Experimental/Animation/PoseSearch/Source/Runtime/Public/PoseSearch/PoseSearchSchema.h
+++ b/Engine/Plugins/Experimental/Animation/PoseSearch/Source/Runtime/Public/PoseSearch/PoseSearchSchema.h
@@ -44,14 +44,16 @@ public:
 	UPROPERTY(EditAnywhere, Category = "Schema", meta = (DisplayPriority = 3, ClampMin = "1", ClampMax = "240"))
 	int32 SampleRate = 30;
 
+private:
 	// Channels itemize the cost breakdown of the config in simpler parts such as position or velocity of a bones, or phase of limbs. The total cost of a query against an indexed database pose will be the sum of the combined channel costs
-	UPROPERTY(EditAnywhere, Instanced, BlueprintReadWrite, Category = "Schema")
+	UPROPERTY(EditAnywhere, Instanced, Category = "Schema")
 	TArray<TObjectPtr<UPoseSearchFeatureChannel>> Channels;
 
 	// FinalizedChannels gets populated with UPoseSearchFeatureChannel(s) from Channels and additional injected ones during the Finalize.
 	UPROPERTY(Transient)
 	TArray<TObjectPtr<UPoseSearchFeatureChannel>> FinalizedChannels;
 
+public:
 	// Setting up and assigning a mirror data table will allow all your assets in your database to access the mirrored version of the data. This is required for mirroring to work with Motion Matching.
 	UPROPERTY(EditAnywhere, Category = "Schema", meta = (DisplayPriority = 1))
 	TObjectPtr<UMirrorDataTable> MirrorDataTable;
@@ -115,6 +117,9 @@ public:
 
 	TConstArrayView<TObjectPtr<UPoseSearchFeatureChannel>> GetChannels() const { return FinalizedChannels; }
 
+	void AddChannel(UPoseSearchFeatureChannel* Channel);
+	void AddTemporaryChannel(UPoseSearchFeatureChannel* DependentChannel);
+
 	template <typename FindPredicateType>
 	const UPoseSearchFeatureChannel* FindChannel(FindPredicateType FindPredicate) const
 	{
@@ -140,7 +145,7 @@ public:
 	virtual void PostEditChangeProperty(struct FPropertyChangedEvent& PropertyChangedEvent) override;
 #endif
 
-	void BuildQuery(UE::PoseSearch::FSearchContext& SearchContext, FPoseSearchFeatureVectorBuilder& InOutQuery) const;
+	void BuildQuery(UE::PoseSearch::FSearchContext& SearchContext, UE::PoseSearch::FFeatureVectorBuilder& InOutQuery) const;
 
 	FBoneIndexType GetBoneIndexType(int8 SchemaBoneIdx) const;
 
@@ -169,4 +174,4 @@ private:
 	}
 
 	void Finalize();
-};
+};
\ No newline at end of file
