motion matching - propagating bExtrapolate correctly to pose history in FSearchContext::GetRootAtTime

diff --git a/Engine/Plugins/Experimental/Animation/PoseSearch/Source/Runtime/Private/PoseSearchContext.cpp b/Engine/Plugins/Experimental/Animation/PoseSearch/Source/Runtime/Private/PoseSearchContext.cpp
--- a/Engine/Plugins/Experimental/Animation/PoseSearch/Source/Runtime/Private/PoseSearchContext.cpp
+++ b/Engine/Plugins/Experimental/Animation/PoseSearch/Source/Runtime/Private/PoseSearchContext.cpp
@@ -356,7 +356,7 @@ FTransform FSearchContext::GetRootAtTime(float Time, bool bUseHistoryRoot, bool
 	if (bUseHistoryRoot)
 	{
 		check(History);
-		History->GetRootTransformAtTime(Time, RootTransform);
+		History->GetRootTransformAtTime(Time, RootTransform, bExtrapolate);
 	}
 	else
 	{
