motion matching
- converting pose search debugger from "IRewindDebuggerViewCreator / IRewindDebuggerView" to "IRewindDebuggerTrackCreator / FRewindDebuggerTrack"
- exposing search costs in rewing debugger pose search track

diff --git a/Engine/Plugins/Experimental/Animation/PoseSearch/Source/Editor/Private/PoseSearchDebugger.cpp b/Engine/Plugins/Experimental/Animation/PoseSearch/Source/Editor/Private/PoseSearchDebugger.cpp
--- a/Engine/Plugins/Experimental/Animation/PoseSearch/Source/Editor/Private/PoseSearchDebugger.cpp
+++ b/Engine/Plugins/Experimental/Animation/PoseSearch/Source/Editor/Private/PoseSearchDebugger.cpp
@@ -8,6 +8,7 @@
 #include "PoseSearchDebuggerViewModel.h"
 #include "Styling/SlateIconFinder.h"
 #include "Trace/PoseSearchTraceProvider.h"
+#include "Widgets/SOverlay.h"
 
 #define LOCTEXT_NAMESPACE "PoseSearchDebugger"
 
@@ -110,13 +111,48 @@ TSharedPtr<SDebuggerView> FDebugger::GenerateInstance(uint64 InAnimInstanceId)
 	return DebuggerView;
 }
 
-
-FText FDebuggerViewCreator::GetTitle() const
+FDebuggerTrack::FDebuggerTrack(uint64 InObjectId)
+: ObjectId(InObjectId)
 {
-	return LOCTEXT("PoseSearchDebuggerTabTitle", "Pose Search");
+	BestCostData = MakeShared<SCurveTimelineView::FTimelineCurveData>();
+	BruteForceCostData = MakeShared<SCurveTimelineView::FTimelineCurveData>();
+
+	BestCostView = SNew(SCurveTimelineView)
+		.CurveColor(FLinearColor::White)
+		.ViewRange_Lambda([]()
+		{
+			return IRewindDebugger::Instance()->GetCurrentViewRange();
+		})
+		.RenderFill(false)
+		.CurveData_Lambda([this]()
+		{
+			return BestCostData;
+		});
+	
+	BruteForceCostView = SNew(SCurveTimelineView)
+		.CurveColor(FLinearColor::Red)
+		.ViewRange_Lambda([]()
+		{
+			return IRewindDebugger::Instance()->GetCurrentViewRange();
+		})
+		.RenderFill(false)
+		.CurveData_Lambda([this]()
+		{
+			return BruteForceCostData;
+		});
+
+	OverlayView = SNew(SOverlay)
+		+ SOverlay::Slot()
+		[
+			BruteForceCostView.ToSharedRef()
+		]
+		+ SOverlay::Slot()
+		[
+			BestCostView.ToSharedRef()
+		];
 }
 
-FSlateIcon FDebuggerViewCreator::GetIcon() const
+FSlateIcon FDebuggerTrack::GetIconInternal()
 {
 #if WITH_EDITOR
 	return FSlateIconFinder::FindIconForClass(UAnimInstance::StaticClass());
@@ -125,18 +161,110 @@ FSlateIcon FDebuggerViewCreator::GetIcon() const
 #endif
 }
 
-FName FDebuggerViewCreator::GetTargetTypeName() const
+bool FDebuggerTrack::UpdateInternal()
+{
+	IRewindDebugger* RewindDebugger = IRewindDebugger::Instance();
+	
+	const TraceServices::IAnalysisSession* AnalysisSession = RewindDebugger->GetAnalysisSession();
+	check(AnalysisSession);
+	if (const FTraceProvider* PoseSearchProvider = AnalysisSession->ReadProvider<FTraceProvider>(FTraceProvider::ProviderName))
+	{
+		TraceServices::FAnalysisSessionReadScope SessionReadScope(*AnalysisSession);
+
+		TArray<SCurveTimelineView::FTimelineCurveData::CurvePoint>& BestCostPoints = BestCostData->Points;
+		BestCostPoints.Reset();
+
+		TArray<SCurveTimelineView::FTimelineCurveData::CurvePoint>& BruteForceCostPoints = BruteForceCostData->Points;
+		BruteForceCostPoints.Reset();
+
+		// convert time range to from rewind debugger times to profiler times
+		TRange<double> TraceTimeRange = RewindDebugger->GetCurrentTraceRange();
+		double StartTime = TraceTimeRange.GetLowerBoundValue();
+		double EndTime = TraceTimeRange.GetUpperBoundValue();
+
+		PoseSearchProvider->EnumerateMotionMatchingStateTimelines(ObjectId, [StartTime, EndTime, &BestCostPoints, &BruteForceCostPoints](const FTraceProvider::FMotionMatchingStateTimeline& InTimeline)
+		{
+			// this isn't very efficient, and it gets called every frame.  will need optimizing
+			InTimeline.EnumerateEvents(StartTime, EndTime, [StartTime, EndTime, &BestCostPoints, &BruteForceCostPoints](double InStartTime, double InEndTime, uint32 InDepth, const FTraceMotionMatchingStateMessage& InMessage)
+			{
+				if (InEndTime > StartTime && InStartTime < EndTime)
+				{
+					BestCostPoints.Add({ InMessage.RecordingTime, InMessage.SearchBestCost });
+					BruteForceCostPoints.Add({ InMessage.RecordingTime, InMessage.SearchBruteForceCost });
+				}
+				return TraceServices::EEventEnumerate::Continue;
+			});
+		});
+
+		float MinValue = UE_MAX_FLT;
+		float MaxValue = -UE_MAX_FLT;
+		for (const SCurveTimelineView::FTimelineCurveData::CurvePoint& CurvePoint : BestCostPoints)
+		{
+			MinValue = FMath::Min(MinValue, CurvePoint.Value);
+			MaxValue = FMath::Max(MaxValue, CurvePoint.Value);
+		}
+		for (const SCurveTimelineView::FTimelineCurveData::CurvePoint& CurvePoint : BruteForceCostPoints)
+		{
+			MinValue = FMath::Min(MinValue, CurvePoint.Value);
+			MaxValue = FMath::Max(MaxValue, CurvePoint.Value);
+		}
+
+		BestCostView->SetFixedRange(MinValue, MaxValue);
+		BruteForceCostView->SetFixedRange(MinValue, MaxValue);
+	}
+
+	if (TSharedPtr<IRewindDebuggerView> PinnedView = View.Pin())
+	{
+		PinnedView->SetTimeMarker(RewindDebugger->CurrentTraceTime());
+	}
+
+	return false;
+}
+
+FName FDebuggerTrack::GetNameInternal() const
+{
+	static const FName Name("PoseSearchDebugger");
+	return Name;
+}
+
+FText FDebuggerTrack::GetDisplayNameInternal() const
+{
+	return LOCTEXT("PoseSearchDebuggerTabTitle", "Pose Search");
+}
+
+TSharedPtr<SWidget> FDebuggerTrack::GetTimelineViewInternal()
+{
+	return OverlayView;
+}
+
+TSharedPtr<SWidget> FDebuggerTrack::GetDetailsViewInternal()
+{
+	TSharedPtr<IRewindDebuggerView> RewindDebuggerView = FDebugger::Get()->GenerateInstance(ObjectId);
+	View = RewindDebuggerView;
+	return RewindDebuggerView;
+}
+
+// FDebuggerTrackCreator
+///////////////////////////////////////////////////
+
+FName FDebuggerTrackCreator::GetTargetTypeNameInternal() const
 {
 	static FName TargetTypeName = "AnimInstance";
 	return TargetTypeName;
 }
 
-TSharedPtr<IRewindDebuggerView> FDebuggerViewCreator::CreateDebugView(uint64 ObjectId, double CurrentTime, const TraceServices::IAnalysisSession& InAnalysisSession) const
+FName FDebuggerTrackCreator::GetNameInternal() const
 {
-	return FDebugger::Get()->GenerateInstance(ObjectId);
+	static const FName Name("PoseSearchDebugger");
+	return Name;
+}
+
+TSharedPtr<RewindDebugger::FRewindDebuggerTrack> FDebuggerTrackCreator::CreateTrackInternal(uint64 ObjectId) const
+{
+	return MakeShared<FDebuggerTrack>(ObjectId);
 }
 
-bool FDebuggerViewCreator::HasDebugInfo(uint64 AnimInstanceId) const
+bool FDebuggerTrackCreator::HasDebugInfoInternal(uint64 ObjectId) const
 {
 	// Get provider and validate
 	const TraceServices::IAnalysisSession* Session = IRewindDebugger::Instance()->GetAnalysisSession();
@@ -149,22 +277,17 @@ bool FDebuggerViewCreator::HasDebugInfo(uint64 AnimInstanceId) const
 	{
 		return false;
 	}
-	
+
 	bool bHasData = false;
-	
-	PoseSearchProvider->EnumerateMotionMatchingStateTimelines(AnimInstanceId, [&bHasData](const FTraceProvider::FMotionMatchingStateTimeline& InTimeline)
+
+	PoseSearchProvider->EnumerateMotionMatchingStateTimelines(ObjectId, [&bHasData](const FTraceProvider::FMotionMatchingStateTimeline& InTimeline)
 	{
 		bHasData = true;
 	});
-	
+
 	return bHasData;
 }
 
-FName FDebuggerViewCreator::GetName() const
-{
-	static const FName Name("PoseSearchDebugger");
-	return Name;
-}
 
 } // namespace UE::PoseSearch
 
diff --git a/Engine/Plugins/Experimental/Animation/PoseSearch/Source/Editor/Private/PoseSearchDebugger.h b/Engine/Plugins/Experimental/Animation/PoseSearch/Source/Editor/Private/PoseSearchDebugger.h
--- a/Engine/Plugins/Experimental/Animation/PoseSearch/Source/Editor/Private/PoseSearchDebugger.h
+++ b/Engine/Plugins/Experimental/Animation/PoseSearch/Source/Editor/Private/PoseSearchDebugger.h
@@ -4,7 +4,11 @@
 
 #include "CoreMinimal.h"
 #include "IRewindDebuggerExtension.h"
-#include "IRewindDebuggerViewCreator.h"
+#include "IRewindDebuggerTrackCreator.h"
+#include "RewindDebuggerTrack.h"
+#include "SCurveTimelineView.h"
+
+class SOverlay;
 
 namespace TraceServices { class IAnalysisSession; }
 
@@ -62,19 +66,40 @@ private:
  * Creates the slate widgets associated with the PoseSearch debugger
  * when prompted by the Rewind Debugger
  */
-class FDebuggerViewCreator : public IRewindDebuggerViewCreator
+class FDebuggerTrack : public RewindDebugger::FRewindDebuggerTrack
 {
 public:
-	virtual ~FDebuggerViewCreator() = default;
-	virtual FName GetName() const override;
-	virtual FText GetTitle() const override;
-	virtual FSlateIcon GetIcon() const override;
-	virtual FName GetTargetTypeName() const override;
-	
-	/** Creates the PoseSearch Slate view for the provided AnimInstance */
-	virtual TSharedPtr<IRewindDebuggerView> CreateDebugView(uint64 ObjectId, double CurrentTime, const TraceServices::IAnalysisSession& InAnalysisSession) const override;
-	virtual bool HasDebugInfo(uint64 ObjectId) const override;
+	FDebuggerTrack(uint64 InObjectId);
+
+private:
+	virtual FSlateIcon GetIconInternal() override;
+	virtual TSharedPtr<SWidget> GetTimelineViewInternal() override;
+	virtual TSharedPtr<SWidget> GetDetailsViewInternal() override;
+	virtual FName GetNameInternal() const override;
+	virtual FText GetDisplayNameInternal() const override;
+	virtual uint64 GetObjectIdInternal() const override { return ObjectId; }
+	virtual bool UpdateInternal() override;
+
+	TSharedPtr<SCurveTimelineView> BestCostView;
+	TSharedPtr<SCurveTimelineView::FTimelineCurveData> BestCostData;
+
+	TSharedPtr<SCurveTimelineView> BruteForceCostView;
+	TSharedPtr<SCurveTimelineView::FTimelineCurveData> BruteForceCostData;
+
+	TSharedPtr<SOverlay> OverlayView;
+
+	TWeakPtr<IRewindDebuggerView> View;
+	FSlateIcon Icon;
+	uint64 ObjectId;
 };
 
+class FDebuggerTrackCreator : public RewindDebugger::IRewindDebuggerTrackCreator
+{
+private:
+	virtual FName GetTargetTypeNameInternal() const override;
+	virtual FName GetNameInternal() const override;
+	virtual TSharedPtr<RewindDebugger::FRewindDebuggerTrack> CreateTrackInternal(uint64 ObjectId) const override;
+	virtual bool HasDebugInfoInternal(uint64 ObjectId) const override;
+};
 
 } // namespace UE::PoseSearch
diff --git a/Engine/Plugins/Experimental/Animation/PoseSearch/Source/Editor/Private/PoseSearchEditor.cpp b/Engine/Plugins/Experimental/Animation/PoseSearch/Source/Editor/Private/PoseSearchEditor.cpp
--- a/Engine/Plugins/Experimental/Animation/PoseSearch/Source/Editor/Private/PoseSearchEditor.cpp
+++ b/Engine/Plugins/Experimental/Animation/PoseSearch/Source/Editor/Private/PoseSearchEditor.cpp
@@ -36,8 +36,8 @@ public:
 private:
 	TArray<class IConsoleObject*> ConsoleCommands;
 	
-	/** Creates the view for the Rewind Debugger */
-	TSharedPtr<FDebuggerViewCreator> DebuggerViewCreator;
+	/** Creates the track for the Rewind Debugger */
+	TSharedPtr<FDebuggerTrackCreator> DebuggerTrackCreator;
 	/** Enables dedicated PoseSearch trace module */
 	TSharedPtr<FTraceModule> TraceModule;
 	
@@ -63,9 +63,9 @@ void FEditorModule::StartupModule()
 	{
 		FDebugger::Initialize();
 		TraceModule = MakeShared<FTraceModule>();
-		DebuggerViewCreator = MakeShared<FDebuggerViewCreator>();
+		DebuggerTrackCreator = MakeShared<FDebuggerTrackCreator>();
 
-		IModularFeatures::Get().RegisterModularFeature("RewindDebuggerViewCreator", DebuggerViewCreator.Get());
+		IModularFeatures::Get().RegisterModularFeature(FDebuggerTrackCreator::ModularFeatureName, DebuggerTrackCreator.Get());
 		IModularFeatures::Get().RegisterModularFeature(TraceServices::ModuleFeatureName, TraceModule.Get());
 		
 		// Register Ed Mode used by pose search database
diff --git a/Engine/Plugins/Experimental/Animation/PoseSearch/Source/Runtime/Private/PoseSearchLibrary.cpp b/Engine/Plugins/Experimental/Animation/PoseSearch/Source/Runtime/Private/PoseSearchLibrary.cpp
--- a/Engine/Plugins/Experimental/Animation/PoseSearch/Source/Runtime/Private/PoseSearchLibrary.cpp
+++ b/Engine/Plugins/Experimental/Animation/PoseSearch/Source/Runtime/Private/PoseSearchLibrary.cpp
@@ -151,7 +151,10 @@ void UPoseSearchLibrary::TraceMotionMatchingState(
 	const UObject* AnimInstance,
 	int32 NodeId,
 	float DeltaTime,
-	bool bSearch)
+	bool bSearch,
+	float RecordingTime,
+	float SearchBestCost,
+	float SearchBruteForceCost)
 {
 	using namespace UE::PoseSearch;
 	
@@ -245,6 +248,10 @@ void UPoseSearchLibrary::TraceMotionMatchingState(
 	TraceState.AssetPlayerTime = CurrentResult.AssetTime;
 	TraceState.DeltaTime = DeltaTime;
 
+	TraceState.RecordingTime = RecordingTime;
+	TraceState.SearchBestCost = SearchBestCost;
+	TraceState.SearchBruteForceCost = SearchBruteForceCost;
+
 	TraceState.Output(AnimInstance, NodeId);
 }
 #endif // UE_POSE_SEARCH_TRACE_ENABLED
@@ -360,9 +367,21 @@ void UPoseSearchLibrary::UpdateMotionMatchingState(
 	// Record debugger details
 	if (IsTracing(Context))
 	{
+		const UAnimInstance* AnimInstance = Cast<const UAnimInstance>(Context.AnimInstanceProxy->GetAnimInstanceObject());
+
+		const float SearchBestCost = InOutMotionMatchingState.CurrentSearchResult.PoseCost.GetTotalCost();
+		float SearchBruteForceCost = SearchBestCost;
+#if WITH_EDITORONLY_DATA
+		if (Databases[0]->PoseSearchMode == EPoseSearchMode::PCAKDTree_Compare)
+		{
+			SearchBruteForceCost = InOutMotionMatchingState.CurrentSearchResult.BruteForcePoseCost.GetTotalCost();
+		}
+#endif // WITH_EDITORONLY_DATA
+
 		// @todo: Update this to account for multiple databases.
 		TraceMotionMatchingState(Databases[0], SearchContext, InOutMotionMatchingState.CurrentSearchResult, LastResult, InOutMotionMatchingState.ElapsedPoseSearchTime,
-			InOutMotionMatchingState.RootMotionTransformDelta, Context.AnimInstanceProxy->GetAnimInstanceObject(), Context.GetCurrentNodeId(), DeltaTime, bSearch);
+			InOutMotionMatchingState.RootMotionTransformDelta, Context.AnimInstanceProxy->GetAnimInstanceObject(), Context.GetCurrentNodeId(), DeltaTime, bSearch,
+			AnimInstance ? FObjectTrace::GetWorldElapsedTime(AnimInstance->GetWorld()) : 0.f, SearchBestCost, SearchBruteForceCost);
 	}
 #endif
 
@@ -586,7 +605,16 @@ void UPoseSearchLibrary::MotionMatch(
 #endif // ENABLE_DRAW_DEBUG && ENABLE_ANIM_DEBUG
 
 #if UE_POSE_SEARCH_TRACE_ENABLED
-			TraceMotionMatchingState(Database, SearchContext, SearchResult, FSearchResult(), 0.f, FTransform::Identity, AnimInstance, DebugSessionUniqueIdentifier, AnimInstance->GetDeltaSeconds(), true);
+			const float SearchBestCost = SearchResult.PoseCost.GetTotalCost();
+			float SearchBruteForceCost = SearchBestCost;
+#if WITH_EDITORONLY_DATA
+			if (Database->PoseSearchMode == EPoseSearchMode::PCAKDTree_Compare)
+			{
+				SearchBruteForceCost = SearchResult.BruteForcePoseCost.GetTotalCost();
+			}
+#endif // WITH_EDITORONLY_DATA
+			TraceMotionMatchingState(Database, SearchContext, SearchResult, FSearchResult(), 0.f, FTransform::Identity, AnimInstance, DebugSessionUniqueIdentifier,
+				AnimInstance->GetDeltaSeconds(), true, FObjectTrace::GetWorldElapsedTime(AnimInstance->GetWorld()), SearchBestCost, SearchBruteForceCost);
 #endif // UE_POSE_SEARCH_TRACE_ENABLED
 		}
 	}
diff --git a/Engine/Plugins/Experimental/Animation/PoseSearch/Source/Runtime/Private/Trace/PoseSearchTraceLogger.cpp b/Engine/Plugins/Experimental/Animation/PoseSearch/Source/Runtime/Private/Trace/PoseSearchTraceLogger.cpp
--- a/Engine/Plugins/Experimental/Animation/PoseSearch/Source/Runtime/Private/Trace/PoseSearchTraceLogger.cpp
+++ b/Engine/Plugins/Experimental/Animation/PoseSearch/Source/Runtime/Private/Trace/PoseSearchTraceLogger.cpp
@@ -78,6 +78,9 @@ FArchive& operator<<(FArchive& Ar, FTraceMotionMatchingState& State)
 	Ar << State.SimAngularVelocity;
 	Ar << State.AnimLinearVelocity;
 	Ar << State.AnimAngularVelocity;
+	Ar << State.RecordingTime;
+	Ar << State.SearchBestCost;
+	Ar << State.SearchBruteForceCost;
 	Ar << State.DatabaseEntries;
 	Ar << State.CurrentDbEntryIdx;
 	Ar << State.CurrentPoseEntryIdx;
diff --git a/Engine/Plugins/Experimental/Animation/PoseSearch/Source/Runtime/Private/Trace/PoseSearchTraceLogger.h b/Engine/Plugins/Experimental/Animation/PoseSearch/Source/Runtime/Private/Trace/PoseSearchTraceLogger.h
--- a/Engine/Plugins/Experimental/Animation/PoseSearch/Source/Runtime/Private/Trace/PoseSearchTraceLogger.h
+++ b/Engine/Plugins/Experimental/Animation/PoseSearch/Source/Runtime/Private/Trace/PoseSearchTraceLogger.h
@@ -66,17 +66,21 @@ struct POSESEARCH_API FTraceMotionMatchingState
 {
 	/** ObjectId of active searchable asset */
 	uint64 SearchableAssetId = 0;
-	
+
 	/** Amount of time since the last pose switch */
-	float ElapsedPoseSearchTime = 0.0f;
-
-	float AssetPlayerTime = 0.0f;
-	float DeltaTime = 0.0f;
-	float SimLinearVelocity = 0.0f;
-	float SimAngularVelocity = 0.0f;
-	float AnimLinearVelocity = 0.0f;
-	float AnimAngularVelocity = 0.0f;
+	float ElapsedPoseSearchTime = 0.f;
+
+	float AssetPlayerTime = 0.f;
+	float DeltaTime = 0.f;
+	float SimLinearVelocity = 0.f;
+	float SimAngularVelocity = 0.f;
+	float AnimLinearVelocity = 0.f;
+	float AnimAngularVelocity = 0.f;
 	
+	float RecordingTime = 0.f;
+	float SearchBestCost = 0.f;
+	float SearchBruteForceCost = 0.f;
+
 	TArray<FTraceMotionMatchingStateDatabaseEntry> DatabaseEntries;
 
 	/** Index of the current database in DatabaseEntries */
diff --git a/Engine/Plugins/Experimental/Animation/PoseSearch/Source/Runtime/Public/PoseSearch/PoseSearchLibrary.h b/Engine/Plugins/Experimental/Animation/PoseSearch/Source/Runtime/Public/PoseSearch/PoseSearchLibrary.h
--- a/Engine/Plugins/Experimental/Animation/PoseSearch/Source/Runtime/Public/PoseSearch/PoseSearchLibrary.h
+++ b/Engine/Plugins/Experimental/Animation/PoseSearch/Source/Runtime/Public/PoseSearch/PoseSearchLibrary.h
@@ -90,7 +90,10 @@ class POSESEARCH_API UPoseSearchLibrary : public UBlueprintFunctionLibrary
 		const UObject* AnimInstance,
 		int32 NodeId,
 		float DeltaTime,
-		bool bSearch);
+		bool bSearch,
+		float RecordingTime,
+		float SearchBestCost,
+		float SearchBruteForceCost);
 #endif // UE_POSE_SEARCH_TRACE_ENABLED
 
 	static FPoseSearchQueryTrajectory ProcessTrajectory(
