Disable physics in Pose Search preview scene.

diff --git a/Engine/Plugins/Experimental/Animation/PoseSearch/Source/Editor/Private/PoseSearchDatabaseEditor.cpp b/Engine/Plugins/Experimental/Animation/PoseSearch/Source/Editor/Private/PoseSearchDatabaseEditor.cpp
--- a/Engine/Plugins/Experimental/Animation/PoseSearch/Source/Editor/Private/PoseSearchDatabaseEditor.cpp
+++ b/Engine/Plugins/Experimental/Animation/PoseSearch/Source/Editor/Private/PoseSearchDatabaseEditor.cpp
@@ -123,7 +123,7 @@ namespace UE::PoseSearch
 				new FDatabasePreviewScene(
 					FPreviewScene::ConstructionValues()
 					.AllowAudioPlayback(true)
-					.ShouldSimulatePhysics(true)
+					.ShouldSimulatePhysics(false)
 					.ForceUseMovementComponentInNonGameWorld(true),
 					StaticCastSharedRef<FDatabaseEditor>(AsShared())));
 
