motion matching - drawing bones in database editor

diff --git a/Engine/Plugins/Experimental/Animation/PoseSearch/Source/Editor/Private/PoseSearchDatabaseEdMode.cpp b/Engine/Plugins/Experimental/Animation/PoseSearch/Source/Editor/Private/PoseSearchDatabaseEdMode.cpp
--- a/Engine/Plugins/Experimental/Animation/PoseSearch/Source/Editor/Private/PoseSearchDatabaseEdMode.cpp
+++ b/Engine/Plugins/Experimental/Animation/PoseSearch/Source/Editor/Private/PoseSearchDatabaseEdMode.cpp
@@ -1,11 +1,13 @@
 // Copyright Epic Games, Inc. All Rights Reserved.
 
 #include "PoseSearchDatabaseEdMode.h"
-#include "PoseSearchDatabaseViewportClient.h"
+#include "Animation/DebugSkelMeshComponent.h"
+#include "EngineUtils.h"
 #include "PoseSearchDatabaseEditor.h"
 #include "PoseSearchDatabaseViewModel.h"
-
-#include "EngineUtils.h"
+#include "PoseSearchDatabaseViewportClient.h"
+#include "Preferences/PersonaOptions.h"
+#include "SkeletalDebugRendering.h"
 
 namespace UE::PoseSearch
 {
@@ -44,6 +46,54 @@ namespace UE::PoseSearch
 
 	void FDatabaseEdMode::Render(const FSceneView* View, FViewport* Viewport, FPrimitiveDrawInterface* PDI)
 	{
+		if (ViewModel->IsShowBones())
+		{
+			const UPersonaOptions* PersonaOptions = GetDefault<UPersonaOptions>();
+
+			TArray<FTransform> WorldTransforms;
+			TArray<FLinearColor> BoneColors;
+			for (FDatabasePreviewActor& PreviewActor : ViewModel->GetPreviewActors())
+			{
+				UDebugSkelMeshComponent* MeshComponent = PreviewActor.GetDebugSkelMeshComponent();
+				if (MeshComponent && MeshComponent->GetSkeletalMeshAsset() && MeshComponent->GetNumDrawTransform() > 0 && MeshComponent->SkeletonDrawMode != ESkeletonDrawMode::Hidden)
+				{
+					WorldTransforms.SetNumUninitialized(MeshComponent->GetNumDrawTransform());
+					BoneColors.SetNumUninitialized(MeshComponent->GetNumDrawTransform());
+
+					// factor skeleton draw mode into color selection
+					const FLinearColor BoneColor = MeshComponent->SkeletonDrawMode == ESkeletonDrawMode::GreyedOut ? PersonaOptions->DisabledBoneColor : PersonaOptions->DefaultBoneColor;
+					const FLinearColor VirtualBoneColor = MeshComponent->SkeletonDrawMode == ESkeletonDrawMode::GreyedOut ? PersonaOptions->DisabledBoneColor : PersonaOptions->VirtualBoneColor;
+
+					const TArray<FBoneIndexType>& DrawBoneIndices = MeshComponent->GetDrawBoneIndices();
+					for (int32 Index = 0; Index < DrawBoneIndices.Num(); ++Index)
+					{
+						const int32 BoneIndex = DrawBoneIndices[Index];
+						WorldTransforms[BoneIndex] = MeshComponent->GetDrawTransform(BoneIndex) * MeshComponent->GetComponentTransform();
+						BoneColors[BoneIndex] = BoneColor;
+					}
+
+					// color virtual bones
+					for (int16 VirtualBoneIndex : MeshComponent->GetReferenceSkeleton().GetRequiredVirtualBones())
+					{
+						BoneColors[VirtualBoneIndex] = VirtualBoneColor;
+					}
+
+					FSkelDebugDrawConfig DrawConfig;
+					DrawConfig.BoneDrawMode = (EBoneDrawMode::Type)PersonaOptions->DefaultBoneDrawSelection;
+					DrawConfig.BoneDrawSize = 0.2f;
+					DrawConfig.bAddHitProxy = false;
+					DrawConfig.bForceDraw = false;
+					DrawConfig.DefaultBoneColor = PersonaOptions->DefaultBoneColor;
+					DrawConfig.AffectedBoneColor = PersonaOptions->AffectedBoneColor;
+					DrawConfig.SelectedBoneColor = PersonaOptions->SelectedBoneColor;
+					DrawConfig.ParentOfSelectedBoneColor = PersonaOptions->ParentOfSelectedBoneColor;
+
+					SkeletalDebugRendering::DrawBones(PDI, MeshComponent->GetComponentLocation(), DrawBoneIndices, MeshComponent->GetReferenceSkeleton(),
+						WorldTransforms, MeshComponent->BonesOfInterest, BoneColors, TArray<TRefCountPtr<HHitProxy>>(), DrawConfig);
+				}
+			}
+		}
+
 		FEdMode::Render(View, Viewport, PDI);
 	}
 
diff --git a/Engine/Plugins/Experimental/Animation/PoseSearch/Source/Editor/Private/PoseSearchDatabaseEditorCommands.cpp b/Engine/Plugins/Experimental/Animation/PoseSearch/Source/Editor/Private/PoseSearchDatabaseEditorCommands.cpp
--- a/Engine/Plugins/Experimental/Animation/PoseSearch/Source/Editor/Private/PoseSearchDatabaseEditorCommands.cpp
+++ b/Engine/Plugins/Experimental/Animation/PoseSearch/Source/Editor/Private/PoseSearchDatabaseEditorCommands.cpp
@@ -15,6 +15,7 @@ namespace UE::PoseSearch
 		UI_COMMAND(ShowAnimationOriginalAndMirrored, "Original and mirrored", "Show original and mirrored animations", EUserInterfaceActionType::RadioButton, FInputChord());
 		UI_COMMAND(ShowDisplayRootMotionSpeed, "Display Root Motion Speed", "Show root motion speed for the selected item", EUserInterfaceActionType::ToggleButton, FInputChord());
 		UI_COMMAND(ShowQuantizeAnimationToPoseData, "Quantize Animation To Pose Data", "Animations will not playing continuously, only using the associated pose quantized time", EUserInterfaceActionType::ToggleButton, FInputChord());
+		UI_COMMAND(ShowBones, "Show All Bones", "Debug Draw All Bones", EUserInterfaceActionType::ToggleButton, FInputChord());
 	}
 }
 
diff --git a/Engine/Plugins/Experimental/Animation/PoseSearch/Source/Editor/Private/PoseSearchDatabaseViewModel.cpp b/Engine/Plugins/Experimental/Animation/PoseSearch/Source/Editor/Private/PoseSearchDatabaseViewModel.cpp
--- a/Engine/Plugins/Experimental/Animation/PoseSearch/Source/Editor/Private/PoseSearchDatabaseViewModel.cpp
+++ b/Engine/Plugins/Experimental/Animation/PoseSearch/Source/Editor/Private/PoseSearchDatabaseViewModel.cpp
@@ -335,16 +335,6 @@ namespace UE::PoseSearch
 		return EnumHasAnyFlags(AnimationPreviewMode, PreviewMode);
 	}
 
-	void FDatabaseViewModel::OnToggleDisplayRootMotionSpeed()
-	{
-		bDisplayRootMotionSpeed = !bDisplayRootMotionSpeed;
-	}
-
-	void FDatabaseViewModel::OnToggleQuantizeAnimationToPoseData()
-	{
-		bQuantizeAnimationToPoseData = !bQuantizeAnimationToPoseData;
-	}
-
 	void FDatabaseViewModel::AddSequenceToDatabase(UAnimSequence* AnimSequence)
 	{
 		FPoseSearchDatabaseSequence NewAsset;
diff --git a/Engine/Plugins/Experimental/Animation/PoseSearch/Source/Editor/Private/SPoseSearchDatabaseViewport.cpp b/Engine/Plugins/Experimental/Animation/PoseSearch/Source/Editor/Private/SPoseSearchDatabaseViewport.cpp
--- a/Engine/Plugins/Experimental/Animation/PoseSearch/Source/Editor/Private/SPoseSearchDatabaseViewport.cpp
+++ b/Engine/Plugins/Experimental/Animation/PoseSearch/Source/Editor/Private/SPoseSearchDatabaseViewport.cpp
@@ -108,6 +108,12 @@ namespace UE::PoseSearch
 			FExecuteAction::CreateSP(ViewModelRef, &FDatabaseViewModel::OnToggleQuantizeAnimationToPoseData),
 			FCanExecuteAction(),
 			FIsActionChecked::CreateSP(ViewModelRef, &FDatabaseViewModel::IsQuantizeAnimationToPoseDataChecked));
+
+		CommandList->MapAction(
+			Commands.ShowBones,
+			FExecuteAction::CreateSP(ViewModelRef, &FDatabaseViewModel::OnToggleShowBones),
+			FCanExecuteAction(),
+			FIsActionChecked::CreateSP(ViewModelRef, &FDatabaseViewModel::IsShowBones));
 		
 	}
 
diff --git a/Engine/Plugins/Experimental/Animation/PoseSearch/Source/Editor/Private/SPoseSearchDatabaseViewportToolbar.cpp b/Engine/Plugins/Experimental/Animation/PoseSearch/Source/Editor/Private/SPoseSearchDatabaseViewportToolbar.cpp
--- a/Engine/Plugins/Experimental/Animation/PoseSearch/Source/Editor/Private/SPoseSearchDatabaseViewportToolbar.cpp
+++ b/Engine/Plugins/Experimental/Animation/PoseSearch/Source/Editor/Private/SPoseSearchDatabaseViewportToolbar.cpp
@@ -74,6 +74,7 @@ namespace UE::PoseSearch
 				{
 					SubMenuBuilder.AddMenuEntry(Commands.ShowDisplayRootMotionSpeed);
 					SubMenuBuilder.AddMenuEntry(Commands.ShowQuantizeAnimationToPoseData);
+					SubMenuBuilder.AddMenuEntry(Commands.ShowBones);
 				}
 				SubMenuBuilder.EndSection();
 			})
diff --git a/Engine/Plugins/Experimental/Animation/PoseSearch/Source/Editor/Public/PoseSearchDatabaseEditorCommands.h b/Engine/Plugins/Experimental/Animation/PoseSearch/Source/Editor/Public/PoseSearchDatabaseEditorCommands.h
--- a/Engine/Plugins/Experimental/Animation/PoseSearch/Source/Editor/Public/PoseSearchDatabaseEditorCommands.h
+++ b/Engine/Plugins/Experimental/Animation/PoseSearch/Source/Editor/Public/PoseSearchDatabaseEditorCommands.h
@@ -31,6 +31,7 @@ namespace UE::PoseSearch
 		TSharedPtr<FUICommandInfo> ShowAnimationOriginalAndMirrored;
 		TSharedPtr<FUICommandInfo> ShowDisplayRootMotionSpeed;
 		TSharedPtr<FUICommandInfo> ShowQuantizeAnimationToPoseData;
+		TSharedPtr<FUICommandInfo> ShowBones;
 	};
 }
 
diff --git a/Engine/Plugins/Experimental/Animation/PoseSearch/Source/Editor/Public/PoseSearchDatabaseViewModel.h b/Engine/Plugins/Experimental/Animation/PoseSearch/Source/Editor/Public/PoseSearchDatabaseViewModel.h
--- a/Engine/Plugins/Experimental/Animation/PoseSearch/Source/Editor/Public/PoseSearchDatabaseViewModel.h
+++ b/Engine/Plugins/Experimental/Animation/PoseSearch/Source/Editor/Public/PoseSearchDatabaseViewModel.h
@@ -94,12 +94,15 @@ namespace UE::PoseSearch
 		void OnSetAnimationPreviewMode(EAnimationPreviewMode PreviewMode);
 		bool IsAnimationPreviewMode(EAnimationPreviewMode PreviewMode) const;
 
-		void OnToggleDisplayRootMotionSpeed();
+		void OnToggleDisplayRootMotionSpeed() { bDisplayRootMotionSpeed = !bDisplayRootMotionSpeed; }
 		bool IsDisplayRootMotionSpeedChecked() const { return bDisplayRootMotionSpeed; };
 
-		void OnToggleQuantizeAnimationToPoseData();
+		void OnToggleQuantizeAnimationToPoseData() { bQuantizeAnimationToPoseData = !bQuantizeAnimationToPoseData; }
 		bool IsQuantizeAnimationToPoseDataChecked() const { return bQuantizeAnimationToPoseData; };
 
+		void OnToggleShowBones() { bShowBones = !bShowBones; }
+		bool IsShowBones() const { return bShowBones; };
+
 		void AddSequenceToDatabase(UAnimSequence* AnimSequence);
 		void AddBlendSpaceToDatabase(UBlendSpace* BlendSpace);
 		void AddAnimCompositeToDatabase(UAnimComposite* AnimComposite);
@@ -162,6 +165,8 @@ namespace UE::PoseSearch
 
 		bool bQuantizeAnimationToPoseData = false;
 
+		bool bShowBones = false;
+
 		int32 SelectedActorIndexAssetIndex = INDEX_NONE;
 
 		UWorld* GetWorld();
