motion matching - invalidating database indexing if required MirrorDataTable is not set

diff --git a/Engine/Plugins/Experimental/Animation/PoseSearch/Source/Runtime/Private/PoseSearchDerivedData.cpp b/Engine/Plugins/Experimental/Animation/PoseSearch/Source/Runtime/Private/PoseSearchDerivedData.cpp
--- a/Engine/Plugins/Experimental/Animation/PoseSearch/Source/Runtime/Private/PoseSearchDerivedData.cpp
+++ b/Engine/Plugins/Experimental/Animation/PoseSearch/Source/Runtime/Private/PoseSearchDerivedData.cpp
@@ -882,6 +882,24 @@ void FPoseSearchDatabaseAsyncCacheTask::OnGetComplete(UE::DerivedData::FCacheGet
 						return;
 					}
 
+					// validating that the missing MirrorDataTable is not necessary 
+					if (!IndexBaseDatabase->Schema->MirrorDataTable)
+					{
+						for (int32 AnimationAssetIndex = 0; AnimationAssetIndex < IndexBaseDatabase->AnimationAssets.Num(); ++AnimationAssetIndex)
+						{
+							if (const FPoseSearchDatabaseAnimationAssetBase* DatabaseAsset = IndexBaseDatabase->GetAnimationAssetBase(AnimationAssetIndex))
+							{
+								if (DatabaseAsset->GetMirrorOption() == EPoseSearchMirrorOption::MirroredOnly || DatabaseAsset->GetMirrorOption() == EPoseSearchMirrorOption::UnmirroredAndMirrored)
+								{
+									// want to sample a mirrored asset
+									UE_LOG(LogPoseSearch, Log, TEXT("%s - %s BuildIndex Failed because '%s' requires a MirrorDataTable to sample mirrored animation assets"), *LexToString(FullIndexKey.Hash), *Database->GetName(), *IndexBaseDatabase->Schema->GetName());
+									SearchIndex.Reset();
+									return;
+								}
+							}
+						}
+					}
+
 					if (Owner.IsCanceled())
 					{
 						UE_LOG(LogPoseSearch, Log, TEXT("%s - %s BuildIndex Cancelled"), *LexToString(FullIndexKey.Hash), *Database->GetName());
