motion matching
- DDC threading fixes when PIE is running
- SDatabaseDataDetails::Reconstruct when DDC is indexing

diff --git a/Engine/Plugins/Experimental/Animation/PoseSearch/Source/Editor/Private/PoseSearchDatabaseDataDetails.cpp b/Engine/Plugins/Experimental/Animation/PoseSearch/Source/Editor/Private/PoseSearchDatabaseDataDetails.cpp
--- a/Engine/Plugins/Experimental/Animation/PoseSearch/Source/Editor/Private/PoseSearchDatabaseDataDetails.cpp
+++ b/Engine/Plugins/Experimental/Animation/PoseSearch/Source/Editor/Private/PoseSearchDatabaseDataDetails.cpp
@@ -283,72 +283,84 @@ void SDatabaseDataDetails::SetExpandedItems(TArray<FChannelItemPtr>& ChannelItem
 
 void SDatabaseDataDetails::Reconstruct()
 {
+	const TSharedPtr<FDatabaseViewModel> ViewModel = EditorViewModel.Pin();
+	
+	if (!ViewModel)
+	{
+		ChannelItems.Reset();
+		return;
+	}
+	
+	const UPoseSearchDatabase* PoseSearchDatabase = ViewModel->GetPoseSearchDatabase();
+	if (!FAsyncPoseSearchDatabasesManagement::RequestAsyncBuildIndex(PoseSearchDatabase, ERequestAsyncBuildFlag::ContinueRequest))
+	{
+		ChannelItems.Reset();
+		return;
+	}
+
 	TMap<const FString, bool> ExpandedItems;
 	TrackExpandedItems(ChannelItems, ExpandedItems);
 
 	ChannelItems.Reset();
 
-	if (const TSharedPtr<FDatabaseViewModel> ViewModel = EditorViewModel.Pin())
-	{
-		RebuildChannelItemsTreeRecursively(ChannelItems, ViewModel->GetPoseSearchDatabase()->Schema->GetChannels());
+	RebuildChannelItemsTreeRecursively(ChannelItems, PoseSearchDatabase->Schema->GetChannels());
 
-		// generating the header
-		TSharedPtr<SHeaderRow> HeaderRow = SNew(SHeaderRow); //.CanSelectGeneratedColumn(true).Visibility(EVisibility::Visible);
+	// generating the header
+	TSharedPtr<SHeaderRow> HeaderRow = SNew(SHeaderRow); //.CanSelectGeneratedColumn(true).Visibility(EVisibility::Visible);
+	HeaderRow->AddColumn(
+		SHeaderRow::Column(TEXT("ChannelName"))
+			.DefaultLabel(LOCTEXT("ChannelName_Header", "Channel Name"))
+			.ToolTipText(LOCTEXT("ChannelName_ToolTip", "Channel Name"))
+			.FillWidth(0.2f));
+
+	HeaderRow->AddColumn(
+		SHeaderRow::Column(TEXT("DataOffset"))
+		.DefaultLabel(LOCTEXT("DataOffset_Header", "Data Offset"))
+		.ToolTipText(LOCTEXT("DataOffset_ToolTip", "Offset from the beginning of the features data"))
+		.FillWidth(0.1f));
+
+	if (ViewModel->ShouldDrawQueryVector())
+	{
 		HeaderRow->AddColumn(
-			SHeaderRow::Column(TEXT("ChannelName"))
-				.DefaultLabel(LOCTEXT("ChannelName_Header", "Channel Name"))
-				.ToolTipText(LOCTEXT("ChannelName_ToolTip", "Channel Name"))
-				.FillWidth(0.2f));
+			SHeaderRow::Column(TEXT("Query"))
+			.DefaultLabel(LOCTEXT("Query_Header", "Query"))
+			.ToolTipText(LOCTEXT("Query_ToolTip", "Query Values")));
+	}
 
+	for (const FDatabasePreviewActor& PreviewActor : ViewModel->GetPreviewActors())
+	{
 		HeaderRow->AddColumn(
-			SHeaderRow::Column(TEXT("DataOffset"))
-			.DefaultLabel(LOCTEXT("DataOffset_Header", "Data Offset"))
-			.ToolTipText(LOCTEXT("DataOffset_ToolTip", "Offset from the beginning of the features data"))
-			.FillWidth(0.1f));
+			SHeaderRow::Column(*PreviewActor.Actor->GetName())
+			.DefaultLabel(FText::FromString(PreviewActor.Sampler.GetAsset()->GetName())));
+	}
 
-		if (ViewModel->ShouldDrawQueryVector())
+	ChannelItemsTreeView = SNew(SChannelItemsTreeView)
+		.TreeItemsSource(&ChannelItems)
+		.HeaderRow(HeaderRow)
+		.OnGenerateRow_Lambda([this](FChannelItemPtr ChannelItem, const TSharedRef<STableViewBase>& OwnerTable)
 		{
-			HeaderRow->AddColumn(
-				SHeaderRow::Column(TEXT("Query"))
-				.DefaultLabel(LOCTEXT("Query_Header", "Query"))
-				.ToolTipText(LOCTEXT("Query_ToolTip", "Query Values")));
-		}
-
-		for (const FDatabasePreviewActor& PreviewActor : ViewModel->GetPreviewActors())
+			return SNew(SDatabaseDataDetailsTableRow, OwnerTable, ChannelItem, EditorViewModel.Pin().ToSharedRef());
+		})
+		.OnGetChildren_Lambda([](FChannelItemPtr ChannelItem, TArray<FChannelItemPtr>& OutChildren)
 		{
-			HeaderRow->AddColumn(
-				SHeaderRow::Column(*PreviewActor.Actor->GetName())
-				.DefaultLabel(FText::FromString(PreviewActor.Sampler.GetAsset()->GetName())));
-		}
-
-		ChannelItemsTreeView = SNew(SChannelItemsTreeView)
-			.TreeItemsSource(&ChannelItems)
-			.HeaderRow(HeaderRow)
-			.OnGenerateRow_Lambda([this](FChannelItemPtr ChannelItem, const TSharedRef<STableViewBase>& OwnerTable)
-			{
-				return SNew(SDatabaseDataDetailsTableRow, OwnerTable, ChannelItem, EditorViewModel.Pin().ToSharedRef());
-			})
-			.OnGetChildren_Lambda([](FChannelItemPtr ChannelItem, TArray<FChannelItemPtr>& OutChildren)
-			{
-				OutChildren.Append(ChannelItem->GetChannelItems());
-			})
-			.OnExpansionChanged_Lambda([](FChannelItemPtr ChannelItem, bool bExpanded)
-			{
-				ChannelItem->SetExpanded(bExpanded);
-			});
+			OutChildren.Append(ChannelItem->GetChannelItems());
+		})
+		.OnExpansionChanged_Lambda([](FChannelItemPtr ChannelItem, bool bExpanded)
+		{
+			ChannelItem->SetExpanded(bExpanded);
+		});
 
-		ChildSlot
-		[
-			// @todo: add a SScrollBox to handle the complexity of multiple actors
-			//SNew(SScrollBox)
-			//+ SScrollBox::Slot()
-			//[
-				ChannelItemsTreeView.ToSharedRef()
-			//]
-		];
-
-		SetExpandedItems(ChannelItems, ExpandedItems, ChannelItemsTreeView.Get());
-	}
+	ChildSlot
+	[
+		// @todo: add a SScrollBox to handle the complexity of multiple actors
+		//SNew(SScrollBox)
+		//+ SScrollBox::Slot()
+		//[
+			ChannelItemsTreeView.ToSharedRef()
+		//]
+	];
+
+	SetExpandedItems(ChannelItems, ExpandedItems, ChannelItemsTreeView.Get());
 }
 
 void SDatabaseDataDetails::RebuildChannelItemsTreeRecursively(TArray<FChannelItemPtr>& ChannelItems, TConstArrayView<TObjectPtr<UPoseSearchFeatureChannel>> Channels)
diff --git a/Engine/Plugins/Experimental/Animation/PoseSearch/Source/Runtime/Private/PoseSearchDerivedData.cpp b/Engine/Plugins/Experimental/Animation/PoseSearch/Source/Runtime/Private/PoseSearchDerivedData.cpp
--- a/Engine/Plugins/Experimental/Animation/PoseSearch/Source/Runtime/Private/PoseSearchDerivedData.cpp
+++ b/Engine/Plugins/Experimental/Animation/PoseSearch/Source/Runtime/Private/PoseSearchDerivedData.cpp
@@ -637,7 +637,15 @@ FPoseSearchDatabaseAsyncCacheTask::FPoseSearchDatabaseAsyncCacheTask(UPoseSearch
 	, Owner(UE::DerivedData::EPriority::Normal)
 	, DerivedDataKey(FIoHash::Zero)
 {
-	StartNewRequestIfNeeded(OuterMutex);
+	if (IsInGameThread())
+	{
+		// it is safe to compose DDC key only on the game thread, since assets can modified in this thread execution
+		StartNewRequestIfNeeded(OuterMutex);
+	}
+	else
+	{
+		UE_LOG(LogPoseSearch, Log, TEXT("Delaying DDC until on the game thread    - %s"), *Database->GetName());
+	}
 }
 
 FPoseSearchDatabaseAsyncCacheTask::~FPoseSearchDatabaseAsyncCacheTask()
@@ -653,6 +661,8 @@ void FPoseSearchDatabaseAsyncCacheTask::StartNewRequestIfNeeded(FCriticalSection
 {
 	using namespace UE::DerivedData;
 
+	check(IsInGameThread());
+
 	FScopeLock Lock(&OuterMutex);
 
 	// making sure there are no active requests
@@ -665,7 +675,7 @@ void FPoseSearchDatabaseAsyncCacheTask::StartNewRequestIfNeeded(FCriticalSection
 		DerivedDataKey = FIoHash::Zero;
 		SetState(EState::Notstarted);
 	
-		UE_LOG(LogPoseSearch, Log, TEXT("Delaying DDC until dependents post load  - %s"), *LexToString(DerivedDataKey), *Database->GetName());
+		UE_LOG(LogPoseSearch, Log, TEXT("Delaying DDC until dependents post load  - %s"), *Database->GetName());
 	}
 	else
 	{
@@ -701,6 +711,8 @@ void FPoseSearchDatabaseAsyncCacheTask::StartNewRequestIfNeeded(FCriticalSection
 // it cancels and waits for the task to be done and reset the local SearchIndex. SetState to Cancelled
 void FPoseSearchDatabaseAsyncCacheTask::Cancel(FCriticalSection& OuterMutex)
 {
+	check(IsInGameThread());
+
 	FScopeLock Lock(&OuterMutex);
 
 	Owner.Cancel();
@@ -1033,6 +1045,8 @@ FAsyncPoseSearchDatabasesManagement::~FAsyncPoseSearchDatabasesManagement()
 // we're listening to OnObjectModified to cancel any pending Task indexing databases depending from Object to avoid multi threading issues
 void FAsyncPoseSearchDatabasesManagement::OnObjectModified(UObject* Object)
 {
+	check(IsInGameThread());
+
 	FScopeLock Lock(&Mutex);
 
 	// iterating backwards because of the possible RemoveAtSwap
@@ -1047,6 +1061,8 @@ void FAsyncPoseSearchDatabasesManagement::OnObjectModified(UObject* Object)
 
 void FAsyncPoseSearchDatabasesManagement::OnPackageReloaded(const EPackageReloadPhase InPackageReloadPhase, FPackageReloadedEvent* InPackageReloadedEvent)
 {
+	check(IsInGameThread());
+
 	if (InPackageReloadPhase == EPackageReloadPhase::PostPackageFixup && InPackageReloadedEvent)
 	{
 		FScopeLock Lock(&Mutex);
