motion matching - adding units to UPoseSearchDebuggerReflection properties

diff --git a/Engine/Plugins/Experimental/Animation/PoseSearch/Source/Editor/Private/PoseSearchDebuggerReflection.h b/Engine/Plugins/Experimental/Animation/PoseSearch/Source/Editor/Private/PoseSearchDebuggerReflection.h
--- a/Engine/Plugins/Experimental/Animation/PoseSearch/Source/Editor/Private/PoseSearchDebuggerReflection.h
+++ b/Engine/Plugins/Experimental/Animation/PoseSearch/Source/Editor/Private/PoseSearchDebuggerReflection.h
@@ -15,24 +15,24 @@ class POSESEARCHEDITOR_API UPoseSearchDebuggerReflection : public UObject
 
 public:
 	/** Time since last PoseSearch */
-	UPROPERTY(VisibleAnywhere, Category="Motion Matching State")
+	UPROPERTY(VisibleAnywhere, Category="Motion Matching State", meta = (ForceUnits = "s"))
 	float ElapsedPoseSearchTime = 0.0f;
 
-	UPROPERTY(VisibleAnywhere, Category = "Motion Matching State")
+	UPROPERTY(VisibleAnywhere, Category = "Motion Matching State", meta = (ForceUnits = "s"))
 	float AssetPlayerTime = 0.0f;
 
-	UPROPERTY(VisibleAnywhere, Category = "Motion Matching State")
+	UPROPERTY(VisibleAnywhere, Category = "Motion Matching State", meta = (ForceUnits = "s"))
 	float LastDeltaTime = 0.0f;
 
-	UPROPERTY(VisibleAnywhere, Category = "Motion Matching State")
+	UPROPERTY(VisibleAnywhere, Category = "Motion Matching State", meta = (ForceUnits = "cm/s"))
 	float SimLinearVelocity;
 
-	UPROPERTY(VisibleAnywhere, Category = "Motion Matching State")
+	UPROPERTY(VisibleAnywhere, Category = "Motion Matching State", meta = (ForceUnits = "deg/s"))
 	float SimAngularVelocity;
 
-	UPROPERTY(VisibleAnywhere, Category = "Motion Matching State")
+	UPROPERTY(VisibleAnywhere, Category = "Motion Matching State", meta = (ForceUnits = "cm/s"))
 	float AnimLinearVelocity;
 
-	UPROPERTY(VisibleAnywhere, Category = "Motion Matching State")
+	UPROPERTY(VisibleAnywhere, Category = "Motion Matching State", meta = (ForceUnits = "deg/s"))
 	float AnimAngularVelocity;
 };
