motion matching - adding custom tooltip to PoseSearch rewind debugger track

diff --git a/Engine/Plugins/Experimental/Animation/PoseSearch/Source/Editor/Private/PoseSearchDebugger.cpp b/Engine/Plugins/Experimental/Animation/PoseSearch/Source/Editor/Private/PoseSearchDebugger.cpp
--- a/Engine/Plugins/Experimental/Animation/PoseSearch/Source/Editor/Private/PoseSearchDebugger.cpp
+++ b/Engine/Plugins/Experimental/Animation/PoseSearch/Source/Editor/Private/PoseSearchDebugger.cpp
@@ -6,15 +6,240 @@
 #include "IRewindDebugger.h"
 #include "PoseSearchDebuggerView.h"
 #include "PoseSearchDebuggerViewModel.h"
+#include "SSimpleTimeSlider.h"
 #include "Styling/SlateIconFinder.h"
 #include "Trace/PoseSearchTraceProvider.h"
+#include "Widgets/SBoxPanel.h"
 #include "Widgets/SOverlay.h"
+#include "Widgets/SToolTip.h"
+#include "Widgets/Text/STextBlock.h"
 
 #define LOCTEXT_NAMESPACE "PoseSearchDebugger"
 
 namespace UE::PoseSearch
 {
 
+typedef SCurveTimelineView::FTimelineCurveData::CurvePoint FCurvePoint;
+class SCostCurveTimelineView : public SCurveTimelineView
+{
+public:
+	TRange<double> GetViewRange() const { return ViewRange.Get(); }
+	virtual FReply OnMouseMove(const FGeometry& MyGeometry, const FPointerEvent& MouseEvent) override { return FReply::Unhandled(); }
+};
+
+class SCostTimelineView : public SOverlay
+{
+public:
+	SLATE_BEGIN_ARGS(SCostTimelineView) {}
+	SLATE_END_ARGS()
+
+	void Construct(const FArguments& InArgs);
+	void UpdateInternal(uint64 ObjectId);
+
+protected:
+	virtual FReply OnMouseMove(const FGeometry& MyGeometry, const FPointerEvent& MouseEvent) override;
+
+	TSharedPtr<SCostCurveTimelineView> BestCostView;
+	TSharedPtr<SCostCurveTimelineView::FTimelineCurveData> BestCostData;
+
+	TSharedPtr<SCostCurveTimelineView> BruteForceCostView;
+	TSharedPtr<SCostCurveTimelineView::FTimelineCurveData> BruteForceCostData;
+
+	TSharedPtr<SToolTip> CostToolTip;
+
+	FText ToolTipTime;
+	FText ToolTipCost;
+	FText ToolTipCostBruteForce;
+};
+
+void SCostTimelineView::Construct(const FArguments& InArgs)
+{
+	BestCostData = MakeShared<SCurveTimelineView::FTimelineCurveData>();
+	BruteForceCostData = MakeShared<SCurveTimelineView::FTimelineCurveData>();
+
+	BestCostView = SNew(SCostCurveTimelineView)
+	.CurveColor(FLinearColor::White)
+	.ViewRange_Lambda([]()
+	{
+		return IRewindDebugger::Instance()->GetCurrentViewRange();
+	})
+	.RenderFill(false)
+	.CurveData_Lambda([this]()
+	{
+		return BestCostData;
+	});
+
+	BruteForceCostView = SNew(SCostCurveTimelineView)
+	.CurveColor(FLinearColor::Red)
+	.ViewRange_Lambda([]()
+	{
+		return IRewindDebugger::Instance()->GetCurrentViewRange();
+	})
+	.RenderFill(false)
+	.CurveData_Lambda([this]()
+	{
+		return BruteForceCostData;
+	});
+		
+	AddSlot()
+	[
+		BruteForceCostView.ToSharedRef()
+	];
+	AddSlot()
+	[
+		BestCostView.ToSharedRef()
+	];
+}
+
+void SCostTimelineView::UpdateInternal(uint64 ObjectId)
+{
+	IRewindDebugger* RewindDebugger = IRewindDebugger::Instance();
+
+	const TraceServices::IAnalysisSession* AnalysisSession = RewindDebugger->GetAnalysisSession();
+	check(AnalysisSession);
+	if (const FTraceProvider* PoseSearchProvider = AnalysisSession->ReadProvider<FTraceProvider>(FTraceProvider::ProviderName))
+	{
+		TraceServices::FAnalysisSessionReadScope SessionReadScope(*AnalysisSession);
+
+		TArray<FCurvePoint>& BestCostPoints = BestCostData->Points;
+		BestCostPoints.Reset();
+
+		TArray<FCurvePoint>& BruteForceCostPoints = BruteForceCostData->Points;
+		BruteForceCostPoints.Reset();
+
+		// convert time range to from rewind debugger times to profiler times
+		TRange<double> TraceTimeRange = RewindDebugger->GetCurrentTraceRange();
+		double StartTime = TraceTimeRange.GetLowerBoundValue();
+		double EndTime = TraceTimeRange.GetUpperBoundValue();
+
+		PoseSearchProvider->EnumerateMotionMatchingStateTimelines(ObjectId, [StartTime, EndTime, &BestCostPoints, &BruteForceCostPoints](const FTraceProvider::FMotionMatchingStateTimeline& InTimeline)
+		{
+			// this isn't very efficient, and it gets called every frame.  will need optimizing
+			InTimeline.EnumerateEvents(StartTime, EndTime, [StartTime, EndTime, &BestCostPoints, &BruteForceCostPoints](double InStartTime, double InEndTime, uint32 InDepth, const FTraceMotionMatchingStateMessage& InMessage)
+			{
+				if (InEndTime > StartTime && InStartTime < EndTime)
+				{
+					BestCostPoints.Add({ InMessage.RecordingTime, InMessage.SearchBestCost });
+					BruteForceCostPoints.Add({ InMessage.RecordingTime, InMessage.SearchBruteForceCost });
+				}
+				return TraceServices::EEventEnumerate::Continue;
+			});
+		});
+
+		float MinValue = UE_MAX_FLT;
+		float MaxValue = -UE_MAX_FLT;
+		for (const FCurvePoint& CurvePoint : BestCostPoints)
+		{
+			MinValue = FMath::Min(MinValue, CurvePoint.Value);
+			MaxValue = FMath::Max(MaxValue, CurvePoint.Value);
+		}
+		for (const FCurvePoint& CurvePoint : BruteForceCostPoints)
+		{
+			MinValue = FMath::Min(MinValue, CurvePoint.Value);
+			MaxValue = FMath::Max(MaxValue, CurvePoint.Value);
+		}
+
+		BestCostView->SetFixedRange(MinValue, MaxValue);
+		BruteForceCostView->SetFixedRange(MinValue, MaxValue);
+	}
+}
+
+FReply SCostTimelineView::OnMouseMove(const FGeometry& MyGeometry, const FPointerEvent& MouseEvent)
+{
+	if (MyGeometry.IsUnderLocation(MouseEvent.GetScreenSpacePosition()))
+	{
+		// Mouse position in widget space
+		const FVector2D HitPosition = MyGeometry.AbsoluteToLocal(MouseEvent.GetScreenSpacePosition());
+
+		// Range helper struct
+		const SSimpleTimeSlider::FScrubRangeToScreen RangeToScreen(BestCostView->GetViewRange(), MyGeometry.GetLocalSize());
+
+		// Mouse position from widget space to curve input space
+		const double TargetTime = RangeToScreen.LocalXToInput(HitPosition.X);
+
+		// Get curve value at given time
+		const TArray<FCurvePoint>& CurvePoints = BestCostData->Points;
+		const int NumPoints = CurvePoints.Num();
+
+		if (NumPoints > 0)
+		{
+			for (int i = 1; i < NumPoints; ++i)
+			{
+				const FCurvePoint& Point1 = CurvePoints[i - 1];
+				const FCurvePoint& Point2 = CurvePoints[i];
+
+				// Find points that contain mouse hit-point time
+				if (Point1.Time >= TargetTime && TargetTime <= Point2.Time)
+				{
+					// Choose point with the smallest delta
+					const float Delta1 = abs(TargetTime - Point1.Time);
+					const float Delta2 = abs(TargetTime - Point2.Time);
+
+					// Get closest point index
+					const int TargetPointIndex = Delta1 < Delta2 ? i - 1 : i;
+
+					const float Time = CurvePoints[TargetPointIndex].Time;
+					const float BestCost = CurvePoints[TargetPointIndex].Value;
+					const float BruteForceCost = BruteForceCostData->Points[TargetPointIndex].Value;
+
+					// Tooltip text formatting
+					FNumberFormattingOptions FormattingOptions;
+					FormattingOptions.MaximumFractionalDigits = 3;
+
+					ToolTipTime = FText::Format(LOCTEXT("CostTimelineViewToolTip_TimeFormat", "Search Time: {0}"), FText::AsNumber(Time, &FormattingOptions));
+					ToolTipCost = FText::Format(LOCTEXT("CostTimelineViewToolTip_CostFormat", "Search Cost: {0}"), FText::AsNumber(BestCost, &FormattingOptions));
+
+					if (FMath::IsNearlyEqual(BestCost, BruteForceCost))
+					{
+						ToolTipCostBruteForce = FText::GetEmpty();
+					}
+					else
+					{
+						ToolTipCostBruteForce = FText::Format(LOCTEXT("CostTimelineViewToolTip_CostBruteForceFormat", "Search BruteForce Cost: {0}"), FText::AsNumber(BruteForceCost, &FormattingOptions));
+					}
+
+					// Update tooltip info
+					if (!CostToolTip.IsValid())
+					{
+						SetToolTip(
+							SAssignNew(CostToolTip, SToolTip)
+							.BorderImage(FCoreStyle::Get().GetBrush("ToolTip.BrightBackground"))
+							[
+								SNew(SVerticalBox)
+								+ SVerticalBox::Slot()
+								[
+									SNew(STextBlock)
+									.Text_Lambda([this]() { return ToolTipTime; })
+									.Font(FCoreStyle::Get().GetFontStyle("ToolTip.LargerFont"))
+									.ColorAndOpacity(FLinearColor::Black)
+								]
+								+ SVerticalBox::Slot()
+								[
+									SNew(STextBlock)
+									.Text_Lambda([this]() { return ToolTipCost; })
+									.Font(FCoreStyle::Get().GetFontStyle("ToolTip.LargerFont"))
+									.ColorAndOpacity(FLinearColor::Black)
+								]
+								+ SVerticalBox::Slot()
+								[
+									SNew(STextBlock)
+									.Visibility_Lambda([this]() { return ToolTipCostBruteForce.IsEmpty() ? EVisibility::Collapsed : EVisibility::Visible; })
+									.Text_Lambda([this]() { return ToolTipCostBruteForce; })
+									.Font(FCoreStyle::Get().GetFontStyle("ToolTip.LargerFont"))
+									.ColorAndOpacity(FLinearColor::Black)
+								]
+							]);
+					}
+
+					break;
+				}
+			}
+		}
+	}
+
+	return FReply::Unhandled();
+}
+
 FDebugger* FDebugger::Debugger;
 void FDebugger::Initialize()
 {
@@ -46,7 +271,6 @@ bool FDebugger::IsPIESimulating()
 bool FDebugger::IsRecording()
 {
 	return Debugger->RewindDebugger->IsRecording();
-
 }
 
 double FDebugger::GetRecordingDuration()
@@ -114,42 +338,7 @@ TSharedPtr<SDebuggerView> FDebugger::GenerateInstance(uint64 InAnimInstanceId)
 FDebuggerTrack::FDebuggerTrack(uint64 InObjectId)
 : ObjectId(InObjectId)
 {
-	BestCostData = MakeShared<SCurveTimelineView::FTimelineCurveData>();
-	BruteForceCostData = MakeShared<SCurveTimelineView::FTimelineCurveData>();
-
-	BestCostView = SNew(SCurveTimelineView)
-		.CurveColor(FLinearColor::White)
-		.ViewRange_Lambda([]()
-		{
-			return IRewindDebugger::Instance()->GetCurrentViewRange();
-		})
-		.RenderFill(false)
-		.CurveData_Lambda([this]()
-		{
-			return BestCostData;
-		});
-	
-	BruteForceCostView = SNew(SCurveTimelineView)
-		.CurveColor(FLinearColor::Red)
-		.ViewRange_Lambda([]()
-		{
-			return IRewindDebugger::Instance()->GetCurrentViewRange();
-		})
-		.RenderFill(false)
-		.CurveData_Lambda([this]()
-		{
-			return BruteForceCostData;
-		});
-
-	OverlayView = SNew(SOverlay)
-		+ SOverlay::Slot()
-		[
-			BruteForceCostView.ToSharedRef()
-		]
-		+ SOverlay::Slot()
-		[
-			BestCostView.ToSharedRef()
-		];
+	CostTimelineView = SNew(SCostTimelineView);
 }
 
 FSlateIcon FDebuggerTrack::GetIconInternal()
@@ -163,59 +352,11 @@ FSlateIcon FDebuggerTrack::GetIconInternal()
 
 bool FDebuggerTrack::UpdateInternal()
 {
-	IRewindDebugger* RewindDebugger = IRewindDebugger::Instance();
-	
-	const TraceServices::IAnalysisSession* AnalysisSession = RewindDebugger->GetAnalysisSession();
-	check(AnalysisSession);
-	if (const FTraceProvider* PoseSearchProvider = AnalysisSession->ReadProvider<FTraceProvider>(FTraceProvider::ProviderName))
-	{
-		TraceServices::FAnalysisSessionReadScope SessionReadScope(*AnalysisSession);
-
-		TArray<SCurveTimelineView::FTimelineCurveData::CurvePoint>& BestCostPoints = BestCostData->Points;
-		BestCostPoints.Reset();
-
-		TArray<SCurveTimelineView::FTimelineCurveData::CurvePoint>& BruteForceCostPoints = BruteForceCostData->Points;
-		BruteForceCostPoints.Reset();
-
-		// convert time range to from rewind debugger times to profiler times
-		TRange<double> TraceTimeRange = RewindDebugger->GetCurrentTraceRange();
-		double StartTime = TraceTimeRange.GetLowerBoundValue();
-		double EndTime = TraceTimeRange.GetUpperBoundValue();
-
-		PoseSearchProvider->EnumerateMotionMatchingStateTimelines(ObjectId, [StartTime, EndTime, &BestCostPoints, &BruteForceCostPoints](const FTraceProvider::FMotionMatchingStateTimeline& InTimeline)
-		{
-			// this isn't very efficient, and it gets called every frame.  will need optimizing
-			InTimeline.EnumerateEvents(StartTime, EndTime, [StartTime, EndTime, &BestCostPoints, &BruteForceCostPoints](double InStartTime, double InEndTime, uint32 InDepth, const FTraceMotionMatchingStateMessage& InMessage)
-			{
-				if (InEndTime > StartTime && InStartTime < EndTime)
-				{
-					BestCostPoints.Add({ InMessage.RecordingTime, InMessage.SearchBestCost });
-					BruteForceCostPoints.Add({ InMessage.RecordingTime, InMessage.SearchBruteForceCost });
-				}
-				return TraceServices::EEventEnumerate::Continue;
-			});
-		});
-
-		float MinValue = UE_MAX_FLT;
-		float MaxValue = -UE_MAX_FLT;
-		for (const SCurveTimelineView::FTimelineCurveData::CurvePoint& CurvePoint : BestCostPoints)
-		{
-			MinValue = FMath::Min(MinValue, CurvePoint.Value);
-			MaxValue = FMath::Max(MaxValue, CurvePoint.Value);
-		}
-		for (const SCurveTimelineView::FTimelineCurveData::CurvePoint& CurvePoint : BruteForceCostPoints)
-		{
-			MinValue = FMath::Min(MinValue, CurvePoint.Value);
-			MaxValue = FMath::Max(MaxValue, CurvePoint.Value);
-		}
-
-		BestCostView->SetFixedRange(MinValue, MaxValue);
-		BruteForceCostView->SetFixedRange(MinValue, MaxValue);
-	}
+	CostTimelineView->UpdateInternal(ObjectId);
 
 	if (TSharedPtr<IRewindDebuggerView> PinnedView = View.Pin())
 	{
-		PinnedView->SetTimeMarker(RewindDebugger->CurrentTraceTime());
+		PinnedView->SetTimeMarker(IRewindDebugger::Instance()->CurrentTraceTime());
 	}
 
 	return false;
@@ -234,7 +375,7 @@ FText FDebuggerTrack::GetDisplayNameInternal() const
 
 TSharedPtr<SWidget> FDebuggerTrack::GetTimelineViewInternal()
 {
-	return OverlayView;
+	return CostTimelineView;
 }
 
 TSharedPtr<SWidget> FDebuggerTrack::GetDetailsViewInternal()
@@ -288,7 +429,6 @@ bool FDebuggerTrackCreator::HasDebugInfoInternal(uint64 ObjectId) const
 	return bHasData;
 }
 
-
 } // namespace UE::PoseSearch
 
 #undef LOCTEXT_NAMESPACE
diff --git a/Engine/Plugins/Experimental/Animation/PoseSearch/Source/Editor/Private/PoseSearchDebugger.h b/Engine/Plugins/Experimental/Animation/PoseSearch/Source/Editor/Private/PoseSearchDebugger.h
--- a/Engine/Plugins/Experimental/Animation/PoseSearch/Source/Editor/Private/PoseSearchDebugger.h
+++ b/Engine/Plugins/Experimental/Animation/PoseSearch/Source/Editor/Private/PoseSearchDebugger.h
@@ -8,13 +8,12 @@
 #include "RewindDebuggerTrack.h"
 #include "SCurveTimelineView.h"
 
-class SOverlay;
-
 namespace TraceServices { class IAnalysisSession; }
 
 namespace UE::PoseSearch
 {
 
+class SCostTimelineView;
 class SDebuggerView;
 class FDebuggerViewModel;
 
@@ -80,13 +79,7 @@ private:
 	virtual uint64 GetObjectIdInternal() const override { return ObjectId; }
 	virtual bool UpdateInternal() override;
 
-	TSharedPtr<SCurveTimelineView> BestCostView;
-	TSharedPtr<SCurveTimelineView::FTimelineCurveData> BestCostData;
-
-	TSharedPtr<SCurveTimelineView> BruteForceCostView;
-	TSharedPtr<SCurveTimelineView::FTimelineCurveData> BruteForceCostData;
-
-	TSharedPtr<SOverlay> OverlayView;
+	TSharedPtr<SCostTimelineView> CostTimelineView;
 
 	TWeakPtr<IRewindDebuggerView> View;
 	FSlateIcon Icon;
