motion matching - adding pose search debugger column tooltips

diff --git a/Engine/Plugins/Experimental/Animation/PoseSearch/Source/Editor/Private/PoseSearchDebuggerDatabaseColumns.h b/Engine/Plugins/Experimental/Animation/PoseSearch/Source/Editor/Private/PoseSearchDebuggerDatabaseColumns.h
--- a/Engine/Plugins/Experimental/Animation/PoseSearch/Source/Editor/Private/PoseSearchDebuggerDatabaseColumns.h
+++ b/Engine/Plugins/Experimental/Animation/PoseSearch/Source/Editor/Private/PoseSearchDebuggerDatabaseColumns.h
@@ -44,7 +44,8 @@ struct IColumn : TSharedFromThis<IColumn>
 	bool bEnabled = false;
 
 	virtual FText GetLabel() const = 0;
-
+	virtual FText GetLabelTooltip() const { return GetLabel(); }
+	
 	using FRowDataRef = TSharedRef<FDebuggerDatabaseRowData>;
 	using FSortPredicate = TFunction<bool(const FRowDataRef&, const FRowDataRef&)>;
 
@@ -96,6 +97,11 @@ struct FPoseIdx : ITextColumn
 		return LOCTEXT("ColumnLabelPoseIndex", "Index");
 	}
 
+	virtual FText GetLabelTooltip() const override
+	{
+		return LOCTEXT("ColumnLabelTooltipPoseIndex", "Index of the Pose in the Database");
+	}
+
 	virtual FSortPredicate GetSortPredicate() const override
 	{
 		return [](const FRowDataRef& Row0, const FRowDataRef& Row1) -> bool { return Row0->PoseIdx < Row1->PoseIdx; };
@@ -122,6 +128,11 @@ struct FDatabaseName : IColumn
 		return LOCTEXT("ColumnLabelDatabaseName", "Database");
 	}
 
+	virtual FText GetLabelTooltip() const override
+	{
+		return LOCTEXT("ColumnLabelTooltipDatabaseName", "Database Name");
+	}
+
 	virtual FSortPredicate GetSortPredicate() const override
 	{
 		return [](const FRowDataRef& Row0, const FRowDataRef& Row1) -> bool { return Row0->SharedData->DatabaseName < Row1->SharedData->DatabaseName; };
@@ -190,7 +201,12 @@ struct FAssetName : IColumn
 	{
 		return LOCTEXT("ColumnLabelAssetName", "Asset");
 	}
-		
+	
+	virtual FText GetLabelTooltip() const override
+	{
+		return LOCTEXT("ColumnLabelTooltipAssetName", "Animation Asset Name");
+	}
+
 	virtual FSortPredicate GetSortPredicate() const override
 	{
 		return [](const FRowDataRef& Row0, const FRowDataRef& Row1) -> bool { return Row0->AssetName < Row1->AssetName; };
@@ -259,7 +275,12 @@ struct FFrame : ITextColumn
 	{
 		return LOCTEXT("ColumnLabelFrame", "Frame");
 	}
-		
+	
+	virtual FText GetLabelTooltip() const override
+	{
+		return LOCTEXT("ColumnLabelTooltipFrame", "Frame number from the start of the Animation Asset");
+	}
+
 	virtual FSortPredicate GetSortPredicate() const override
 	{
 		return [](const FRowDataRef& Row0, const FRowDataRef& Row1) -> bool { return Row0->AnimFrame < Row1->AnimFrame; };
@@ -280,6 +301,11 @@ struct FTime : ITextColumn
 		return LOCTEXT("ColumnLabelTime", "Time");
 	}
 
+	virtual FText GetLabelTooltip() const override
+	{
+		return LOCTEXT("ColumnLabelTooltipTime", "Time in seconds from the start of the Animation Asset");
+	}
+
 	virtual FSortPredicate GetSortPredicate() const override
 	{
 		return [](const FRowDataRef& Row0, const FRowDataRef& Row1) -> bool { return Row0->AssetTime < Row1->AssetTime; };
@@ -300,6 +326,11 @@ struct FPercentage : ITextColumn
 		return LOCTEXT("ColumnLabelPercentage", "Percentage");
 	}
 
+	virtual FText GetLabelTooltip() const override
+	{
+		return LOCTEXT("ColumnLabelTooltipPercentage", "Time in percentage from the start of the Animation Asset");
+	}
+
 	virtual FSortPredicate GetSortPredicate() const override
 	{
 		return [](const FRowDataRef& Row0, const FRowDataRef& Row1) -> bool { return Row0->AnimPercentage < Row1->AnimPercentage; };
@@ -319,7 +350,12 @@ struct FCost : ITextColumn
 	{
 		return LOCTEXT("ColumnLabelCost", "Cost");
 	}
-		
+	
+	virtual FText GetLabelTooltip() const override
+	{
+		return LOCTEXT("ColumnLabelTooltipCost", "Total Cost of the associated Pose");
+	}
+
 	virtual FSortPredicate GetSortPredicate() const override
 	{
 		return [](const FRowDataRef& Row0, const FRowDataRef& Row1) -> bool { return Row0->PoseCost < Row1->PoseCost; };
@@ -357,6 +393,11 @@ struct FChannelBreakdownCostColumn : ITextColumn
 		return Label;
 	}
 
+	virtual FText GetLabelTooltip() const override
+	{
+		return FText::Format(LOCTEXT("ColumnLabelTooltipChannelBreakdownCost", "Breakdown Cost for the Channel '{0}'"), Label);
+	}
+
 	virtual FSortPredicate GetSortPredicate() const override
 	{
 		return [this](const FRowDataRef& Row0, const FRowDataRef& Row1) -> bool
@@ -393,6 +434,11 @@ struct FCostModifier : ITextColumn
 		return LOCTEXT("ColumnLabelCostModifier", "Bias");
 	}
 
+	virtual FText GetLabelTooltip() const override
+	{
+		return LOCTEXT("ColumnLabelTooltipCostModifier", "Total Cost for all the Bias contributions");
+	}
+
 	virtual FSortPredicate GetSortPredicate() const override
 	{
 		return [](const FRowDataRef& Row0, const FRowDataRef& Row1) -> bool { return Row0->PoseCost.GetCostAddend() < Row1->PoseCost.GetCostAddend(); };
@@ -413,6 +459,11 @@ struct FMirrored : ITextColumn
 		return LOCTEXT("ColumnLabelMirrored", "Mirror");
 	}
 
+	virtual FText GetLabelTooltip() const override
+	{
+		return LOCTEXT("ColumnLabelTooltipMirrored", "Mirror state of the associated Pose");
+	}
+
 	virtual FSortPredicate GetSortPredicate() const override
 	{
 		return [](const FRowDataRef& Row0, const FRowDataRef& Row1) -> bool { return Row0->bMirrored < Row1->bMirrored; };
@@ -433,6 +484,11 @@ struct FLooping : ITextColumn
 		return LOCTEXT("ColumnLabelLooping", "Loop");
 	}
 
+	virtual FText GetLabelTooltip() const override
+	{
+		return LOCTEXT("ColumnLabelTooltipLooping", "Loop state of the associated Pose");
+	}
+
 	virtual FSortPredicate GetSortPredicate() const override
 	{
 		return [](const FRowDataRef& Row0, const FRowDataRef& Row1) -> bool { return Row0->bLooping < Row1->bLooping; };
@@ -453,6 +509,11 @@ struct FBlendParameters : ITextColumn
 		return LOCTEXT("ColumnLabelBlendParams", "Blend Params");
 	}
 
+	virtual FText GetLabelTooltip() const override
+	{
+		return LOCTEXT("ColumnLabelBlendTooltipParams", "Blend Params used to sample the associated BlendSpace asset");
+	}
+
 	virtual FSortPredicate GetSortPredicate() const override
 	{
 		return [](const FRowDataRef& Row0, const FRowDataRef& Row1) -> bool
@@ -478,6 +539,11 @@ struct FPoseCandidateFlags : ITextColumn
 		return LOCTEXT("ColumnLabelPoseCandidateFlags", "Flags");
 	}
 
+	virtual FText GetLabelTooltip() const override
+	{
+		return LOCTEXT("ColumnLabelTooltipPoseCandidateFlags", "Flags indicating why a Pose has been discarded");
+	}
+
 	virtual FSortPredicate GetSortPredicate() const override
 	{
 		return [](const FRowDataRef& Row0, const FRowDataRef& Row1) -> bool { return Row0->PoseCandidateFlags < Row1->PoseCandidateFlags; };
diff --git a/Engine/Plugins/Experimental/Animation/PoseSearch/Source/Editor/Private/PoseSearchDebuggerDatabaseView.cpp b/Engine/Plugins/Experimental/Animation/PoseSearch/Source/Editor/Private/PoseSearchDebuggerDatabaseView.cpp
--- a/Engine/Plugins/Experimental/Animation/PoseSearch/Source/Editor/Private/PoseSearchDebuggerDatabaseView.cpp
+++ b/Engine/Plugins/Experimental/Animation/PoseSearch/Source/Editor/Private/PoseSearchDebuggerDatabaseView.cpp
@@ -393,6 +393,7 @@ void SDebuggerDatabaseView::Update(const FTraceMotionMatchingStateMessage& State
 				SHeaderRow::FColumn::FArguments ColumnArgs = SHeaderRow::FColumn::FArguments()
 					.ColumnId(Column.ColumnId)
 					.DefaultLabel(Column.GetLabel())
+					.DefaultTooltip(Column.GetLabelTooltip())
 					.SortMode(this, &SDebuggerDatabaseView::GetColumnSortMode, Column.ColumnId)
 					.OnSort(this, &SDebuggerDatabaseView::OnColumnSortModeChanged)
 					.FillWidth(this, &SDebuggerDatabaseView::GetColumnWidth, Column.ColumnId)
@@ -431,14 +432,11 @@ EColumnSortMode::Type SDebuggerDatabaseView::GetColumnSortMode(const FName Colum
 
 float SDebuggerDatabaseView::GetColumnWidth(const FName ColumnId) const
 {
-	check(Columns.Find(ColumnId));
-
 	return Columns[ColumnId]->Width;
 }
 
 void SDebuggerDatabaseView::OnColumnSortModeChanged(const EColumnSortPriority::Type SortPriority, const FName & ColumnId, const EColumnSortMode::Type InSortMode)
 {
-	check(Columns.Find(ColumnId));
 	SortColumn = ColumnId;
 	SortMode = InSortMode;
 	SortDatabaseRows();
@@ -447,8 +445,6 @@ void SDebuggerDatabaseView::OnColumnSortModeChanged(const EColumnSortPriority::T
 
 void SDebuggerDatabaseView::OnColumnWidthChanged(const float NewWidth, FName ColumnId) const
 {
-	check(Columns.Find(ColumnId));
-	
 	Columns[ColumnId]->Width = NewWidth;
 }
 
