motion matching - heading channel fix (FSearchContext::GetSampleRotation support for SampleTimeOffset != 0)

diff --git a/Engine/Plugins/Experimental/Animation/PoseSearch/Source/Runtime/Private/PoseSearchContext.cpp b/Engine/Plugins/Experimental/Animation/PoseSearch/Source/Runtime/Private/PoseSearchContext.cpp
--- a/Engine/Plugins/Experimental/Animation/PoseSearch/Source/Runtime/Private/PoseSearchContext.cpp
+++ b/Engine/Plugins/Experimental/Animation/PoseSearch/Source/Runtime/Private/PoseSearchContext.cpp
@@ -300,24 +300,10 @@ FQuat FSearchContext::GetSampleRotation(float SampleTimeOffset, const UPoseSearc
 	float PermutationOriginTimeOffset = 0.f;
 	UPoseSearchFeatureChannel::GetPermutationTimeOffsets(PermutationTimeType, DesiredPermutationTimeOffset, PermutationSampleTimeOffset, PermutationOriginTimeOffset);
 
-	// @todo: add support for SchemaSampleBoneIdx
-	if (SchemaOriginBoneIdx != RootSchemaBoneIdx)
-	{
-		UE_LOG(LogPoseSearch, Error, TEXT("FSearchContext::GetSampleRotation: support for non root origin bones not implemented (bone: '%s', schema: '%s'"),
-			*Schema->BoneReferences[SchemaOriginBoneIdx].BoneName.ToString(), *GetNameSafe(Schema));
-	}
-
 	const float SampleTime = SampleTimeOffset + PermutationSampleTimeOffset;
 	const float OriginTime = PermutationOriginTimeOffset;
 
-	// @todo: add support for OriginTime != 0 (like in GetSamplePosition and GetSampleVelocity)
-	if (OriginTime != 0.f)
-	{
-		UE_LOG(LogPoseSearch, Error, TEXT("FSearchContext::GetSampleRotation: support for OriginTime != 0 not implemented (bone: '%s', schema: '%s'"),
-			SchemaOriginBoneIdx >= 0 ? *Schema->BoneReferences[SchemaOriginBoneIdx].BoneName.ToString() : TEXT("RootBone"), *GetNameSafe(Schema));
-	}
-
-	return GetComponentSpaceTransform(SampleTime, Schema, SchemaSampleBoneIdx).GetRotation();
+	return GetSampleRotationInternal(SampleTime, OriginTime, Schema, SchemaSampleBoneIdx, SchemaOriginBoneIdx, bUseHistoryRoot);
 }
 
 FVector FSearchContext::GetSamplePosition(float SampleTimeOffset, const UPoseSearchSchema* Schema, int8 SchemaSampleBoneIdx, int8 SchemaOriginBoneIdx, bool bUseHistoryRoot, EPermutationTimeType PermutationTimeType)
@@ -381,6 +367,7 @@ FTransform FSearchContext::GetTransform(float SampleTime, const UPoseSearchSchem
 	return RootTransform;
 }
 
+// returns the component space transform of the bone SchemaBoneIdx at time SampleTime
 FTransform FSearchContext::GetComponentSpaceTransform(float SampleTime, const UPoseSearchSchema* Schema, int8 SchemaBoneIdx)
 {
 	check(Schema);
@@ -444,6 +431,32 @@ FVector FSearchContext::GetSamplePositionInternal(float SampleTime, float Origin
 	return RootBoneTransform.InverseTransformVector(DeltaBoneTranslation);
 }
 
+FQuat FSearchContext::GetSampleRotationInternal(float SampleTime, float OriginTime, const UPoseSearchSchema* Schema, int8 SchemaSampleBoneIdx, int8 SchemaOriginBoneIdx, bool bUseHistoryRoot)
+{
+	if (SampleTime == OriginTime)
+	{
+		if (Schema->IsRootBone(SchemaOriginBoneIdx))
+		{
+			return GetComponentSpaceTransform(SampleTime, Schema, SchemaSampleBoneIdx).GetRotation();
+		}
+
+		const FTransform SampleBoneTransform = GetComponentSpaceTransform(SampleTime, Schema, SchemaSampleBoneIdx);
+		const FTransform OriginBoneTransform = GetComponentSpaceTransform(OriginTime, Schema, SchemaOriginBoneIdx);
+		return OriginBoneTransform.InverseTransformRotation(SampleBoneTransform.GetRotation());
+	}
+
+	const FTransform RootBoneTransform = GetTransform(OriginTime, Schema, RootSchemaBoneIdx, bUseHistoryRoot);
+	const FTransform SampleBoneTransform = GetTransform(SampleTime, Schema, SchemaSampleBoneIdx, bUseHistoryRoot);
+	if (Schema->IsRootBone(SchemaOriginBoneIdx))
+	{
+		return RootBoneTransform.InverseTransformRotation(SampleBoneTransform.GetRotation());
+	}
+
+	const FTransform OriginBoneTransform = GetTransform(OriginTime, Schema, SchemaOriginBoneIdx, bUseHistoryRoot);
+	const FQuat DeltaBoneRotation = OriginBoneTransform.InverseTransformRotation(SampleBoneTransform.GetRotation());
+	return RootBoneTransform.InverseTransformRotation(DeltaBoneRotation);
+}
+
 void FSearchContext::ClearCachedEntries()
 {
 	CachedTransforms.Reset();
diff --git a/Engine/Plugins/Experimental/Animation/PoseSearch/Source/Runtime/Public/PoseSearch/PoseSearchContext.h b/Engine/Plugins/Experimental/Animation/PoseSearch/Source/Runtime/Public/PoseSearch/PoseSearchContext.h
--- a/Engine/Plugins/Experimental/Animation/PoseSearch/Source/Runtime/Public/PoseSearch/PoseSearchContext.h
+++ b/Engine/Plugins/Experimental/Animation/PoseSearch/Source/Runtime/Public/PoseSearch/PoseSearchContext.h
@@ -176,6 +176,7 @@ private:
 	FTransform GetTransform(float SampleTime, const UPoseSearchSchema* Schema, int8 SchemaBoneIdx = RootSchemaBoneIdx, bool bUseHistoryRoot = false);
 	FTransform GetComponentSpaceTransform(float SampleTime, const UPoseSearchSchema* Schema, int8 SchemaSampleBoneIdx = RootSchemaBoneIdx);
 	FVector GetSamplePositionInternal(float SampleTime, float OriginTime, const UPoseSearchSchema* Schema, int8 SchemaSampleBoneIdx = RootSchemaBoneIdx, int8 SchemaOriginBoneIdx = RootSchemaBoneIdx, bool bUseHistoryRoot = false);
+	FQuat GetSampleRotationInternal(float SampleTime, float OriginTime, const UPoseSearchSchema* Schema, int8 SchemaSampleBoneIdx = RootSchemaBoneIdx, int8 SchemaOriginBoneIdx = RootSchemaBoneIdx, bool bUseHistoryRoot = false);
 
 	const FPoseSearchQueryTrajectory* Trajectory = nullptr;
 	const IPoseHistory* History = nullptr;
