motion matching - matching from mesh component space to root bone space

diff --git a/Engine/Plugins/Experimental/Animation/PoseSearch/Source/Editor/Private/PoseSearchDatabasePreviewScene.cpp b/Engine/Plugins/Experimental/Animation/PoseSearch/Source/Editor/Private/PoseSearchDatabasePreviewScene.cpp
--- a/Engine/Plugins/Experimental/Animation/PoseSearch/Source/Editor/Private/PoseSearchDatabasePreviewScene.cpp
+++ b/Engine/Plugins/Experimental/Animation/PoseSearch/Source/Editor/Private/PoseSearchDatabasePreviewScene.cpp
@@ -68,7 +68,7 @@ namespace UE::PoseSearch
 				{
 					if (UDebugSkelMeshComponent* Mesh = PreviewActor.GetDebugSkelMeshComponent())
 					{
-						UE::PoseSearch::FDebugDrawParams DrawParams(GetWorld(), Mesh, &PreviewActor.QuantizedTimeRootMotion, Database);
+						UE::PoseSearch::FDebugDrawParams DrawParams(GetWorld(), Mesh, PreviewActor.QuantizedTimeRootMotion, Database);
 						DrawParams.DrawFeatureVector(PreviewActor.CurrentPoseIndex);
 
 						if (bDrawQueryVector)
diff --git a/Engine/Plugins/Experimental/Animation/PoseSearch/Source/Editor/Private/PoseSearchDebuggerView.cpp b/Engine/Plugins/Experimental/Animation/PoseSearch/Source/Editor/Private/PoseSearchDebuggerView.cpp
--- a/Engine/Plugins/Experimental/Animation/PoseSearch/Source/Editor/Private/PoseSearchDebuggerView.cpp
+++ b/Engine/Plugins/Experimental/Animation/PoseSearch/Source/Editor/Private/PoseSearchDebuggerView.cpp
@@ -315,14 +315,12 @@ void SDebuggerView::UpdateViews() const
 
 void SDebuggerView::DrawVisualization() const
 {
-	const UWorld* DebuggerWorld = FDebugger::GetWorld();
-	check(DebuggerWorld);
-
-	const FTraceMotionMatchingStateMessage* State = ViewModel.Get()->GetMotionMatchingState();
-	const FTransform* Transform = ViewModel.Get()->GetRootTransform();
-	if (State && Transform)
+	if (const FTraceMotionMatchingStateMessage* State = ViewModel.Get()->GetMotionMatchingState())
 	{
-		DrawFeatures(*DebuggerWorld, *State, *Transform, ViewModel.Get()->GetMeshComponent());
+		const UWorld* DebuggerWorld = FDebugger::GetWorld();
+		check(DebuggerWorld);
+
+		DrawFeatures(*DebuggerWorld, *State, ViewModel.Get()->GetRootBoneTransform(), ViewModel.Get()->GetMeshComponent());
 	}
 }
 
@@ -334,10 +332,16 @@ TArray<TSharedRef<FDebuggerDatabaseRowData>> SDebuggerView::GetSelectedDatabaseR
 void SDebuggerView::DrawFeatures(
 	const UWorld& DebuggerWorld,
 	const FTraceMotionMatchingStateMessage& State,
-	const FTransform& Transform,
+	const FTransform& RootBoneWorldTransform,
 	const USkinnedMeshComponent* Mesh
 ) const
 {
+	// Draw world space trajectory
+	if (ViewModel.Get()->GetDrawTrajectory())
+	{
+		State.Trajectory.DebugDrawTrajectory(&DebuggerWorld);
+	}
+
 	// Draw query vector
 	if (ViewModel.Get()->GetDrawQuery())
 	{
@@ -350,7 +354,7 @@ void SDebuggerView::DrawFeatures(
 					FAsyncPoseSearchDatabasesManagement::RequestAsyncBuildIndex(CurrentDatabase, ERequestAsyncBuildFlag::ContinueRequest) &&
 					DbEntry.QueryVector.Num() == Database->Schema->SchemaCardinality)
 				{
-					FDebugDrawParams DrawParams(&DebuggerWorld, Mesh, nullptr, CurrentDatabase, EDebugDrawFlags::DrawQuery);
+					FDebugDrawParams DrawParams(&DebuggerWorld, Mesh, RootBoneWorldTransform, CurrentDatabase, EDebugDrawFlags::DrawQuery);
 					DrawParams.DrawFeatureVector(DbEntry.QueryVector);
 					break;
 				}
@@ -368,7 +372,7 @@ void SDebuggerView::DrawFeatures(
 		const UPoseSearchDatabase* RowDatabase = Row->SharedData->SourceDatabase.Get();
 		if (FAsyncPoseSearchDatabasesManagement::RequestAsyncBuildIndex(RowDatabase, ERequestAsyncBuildFlag::ContinueRequest))
 		{
-			FDebugDrawParams DrawParams(&DebuggerWorld, Mesh, nullptr, RowDatabase);
+			FDebugDrawParams DrawParams(&DebuggerWorld, Mesh, RootBoneWorldTransform, RowDatabase);
 			DrawParams.DrawFeatureVector(Row->PoseIdx);
 		}
 	}
@@ -386,7 +390,7 @@ void SDebuggerView::DrawFeatures(
 			if (Database && FAsyncPoseSearchDatabasesManagement::RequestAsyncBuildIndex(Database, ERequestAsyncBuildFlag::ContinueRequest))
 			{
 				// Use the motion-matching state's pose idx, as the active row may be update-throttled at this point
-				FDebugDrawParams DrawParams(&DebuggerWorld, Mesh, nullptr, Database);
+				FDebugDrawParams DrawParams(&DebuggerWorld, Mesh, RootBoneWorldTransform, Database);
 				DrawParams.DrawFeatureVector(ActiveRows[0]->PoseIdx);
 			}
 		}
@@ -405,12 +409,11 @@ void SDebuggerView::DrawFeatures(
 			const UPoseSearchDatabase* Database = ContinuingRows[0]->SharedData->SourceDatabase.Get();
 			if (Database && FAsyncPoseSearchDatabasesManagement::RequestAsyncBuildIndex(Database, ERequestAsyncBuildFlag::ContinueRequest))
 			{
-				FDebugDrawParams DrawParams(&DebuggerWorld, Mesh, nullptr, Database);
+				FDebugDrawParams DrawParams(&DebuggerWorld, Mesh, RootBoneWorldTransform, Database);
 				DrawParams.DrawFeatureVector(ContinuingRows[0]->PoseIdx);
 			}
 		}
 	}
-
 }
 
 int32 SDebuggerView::SelectView() const
@@ -567,6 +570,33 @@ TSharedRef<SHorizontalBox> SDebuggerView::GenerateReturnButtonView()
 				]
 			]
 		]
+		+SHorizontalBox::Slot()
+		.VAlign(VAlign_Top)
+		.HAlign(HAlign_Left)
+		.Padding(64, 5, 0, 0)
+		.AutoWidth()
+		[
+			SNew(SHorizontalBox)
+
+			+ SHorizontalBox::Slot()
+			.AutoWidth()
+			.Padding(0, 5, 0, 0)
+			[
+				SNew(SCheckBox)
+				.IsChecked_Lambda([this]
+				{
+					return ViewModel.Get()->GetDrawTrajectory() ? ECheckBoxState::Checked : ECheckBoxState::Unchecked; 
+				})
+				.OnCheckStateChanged_Lambda([this](ECheckBoxState State)
+				{
+					ViewModel.Get()->SetDrawTrajectory(State == ECheckBoxState::Checked);
+				})
+				[
+					SNew(STextBlock)
+					.Text(LOCTEXT("PoseSearchDebuggerDrawTrajectory", "Draw Trajectory"))
+				]
+			]
+		]
 		+ SHorizontalBox::Slot()
 		.VAlign(VAlign_Top)
 		.HAlign(HAlign_Left)
diff --git a/Engine/Plugins/Experimental/Animation/PoseSearch/Source/Editor/Private/PoseSearchDebuggerView.h b/Engine/Plugins/Experimental/Animation/PoseSearch/Source/Editor/Private/PoseSearchDebuggerView.h
--- a/Engine/Plugins/Experimental/Animation/PoseSearch/Source/Editor/Private/PoseSearchDebuggerView.h
+++ b/Engine/Plugins/Experimental/Animation/PoseSearch/Source/Editor/Private/PoseSearchDebuggerView.h
@@ -79,7 +79,7 @@ public:
 
 private:
 	/** Called each frame to draw features of the query vector & database selections */
-	void DrawFeatures(const UWorld& DebuggerWorld, const FTraceMotionMatchingStateMessage& State, const FTransform& Transform, const USkinnedMeshComponent* Mesh) const;
+	void DrawFeatures(const UWorld& DebuggerWorld, const FTraceMotionMatchingStateMessage& State, const FTransform& RootBoneWorldTransform, const USkinnedMeshComponent* Mesh) const;
 	
 	/** Check if a node selection was made, true if a node is selected */
 	bool UpdateNodeSelection();
diff --git a/Engine/Plugins/Experimental/Animation/PoseSearch/Source/Editor/Private/PoseSearchDebuggerViewModel.cpp b/Engine/Plugins/Experimental/Animation/PoseSearch/Source/Editor/Private/PoseSearchDebuggerViewModel.cpp
--- a/Engine/Plugins/Experimental/Animation/PoseSearch/Source/Editor/Private/PoseSearchDebuggerViewModel.cpp
+++ b/Engine/Plugins/Experimental/Animation/PoseSearch/Source/Editor/Private/PoseSearchDebuggerViewModel.cpp
@@ -87,9 +87,9 @@ int32 FDebuggerViewModel::GetNodesNum() const
 	return MotionMatchingStates.Num();
 }
 
-const FTransform* FDebuggerViewModel::GetRootTransform() const
+const FTransform& FDebuggerViewModel::GetRootBoneTransform() const
 {
-	return RootTransform;
+	return RootBoneWorldTransform;
 }
 
 bool FDebuggerViewModel::HasSearchableAssetChanged() const
@@ -245,15 +245,13 @@ void FDebuggerViewModel::UpdateFromTimeline()
 	{
 		TimelineData.EnumerateEvents(Frame.StartTime, Frame.EndTime, [&](double InStartTime, double InEndTime, uint32 InDepth, const FSkeletalMeshPoseMessage& PoseMessage) -> TraceServices::EEventEnumerate
 		{
-			// Update root transform
-			RootTransform = &PoseMessage.ComponentToWorld;
 			const FSkeletalMeshInfo* SkeletalMeshInfo = AnimationProvider->FindSkeletalMeshInfo(PoseMessage.MeshId);
 			const FObjectInfo* SkeletalMeshObjectInfo = GameplayProvider->FindObjectInfo(PoseMessage.MeshId);
 			if (!SkeletalMeshInfo || !SkeletalMeshObjectInfo)
 			{
-
 				return TraceServices::EEventEnumerate::Stop;
 			}
+
 			UPoseSearchMeshComponent* ActiveComponent = Skeletons[ActivePose].Component.Get();
 			UPoseSearchMeshComponent* SelectedComponent = Skeletons[SelectedPose].Component.Get();
 			UPoseSearchMeshComponent* AssetComponent = Skeletons[Asset].Component.Get();
@@ -266,7 +264,20 @@ void FDebuggerViewModel::UpdateFromTimeline()
 			}
 			FTransform ComponentWorldTransform;
 			// Active skeleton is simply the traced bone transforms
-			AnimationProvider->GetSkeletalMeshComponentSpacePose(PoseMessage, *SkeletalMeshInfo, ComponentWorldTransform, ActiveComponent->GetEditableComponentSpaceTransforms());
+			TArray<FTransform>& ComponentSpaceTransforms = ActiveComponent->GetEditableComponentSpaceTransforms();
+			AnimationProvider->GetSkeletalMeshComponentSpacePose(PoseMessage, *SkeletalMeshInfo, ComponentWorldTransform, ComponentSpaceTransforms);
+
+			check(ComponentWorldTransform.Equals(PoseMessage.ComponentToWorld));
+
+			if (!ComponentSpaceTransforms.IsEmpty())
+			{
+				RootBoneWorldTransform = ComponentSpaceTransforms[RootBoneIndexType] * ComponentWorldTransform;
+			}
+			else
+			{
+				RootBoneWorldTransform = ComponentWorldTransform;
+			}
+
 			ActiveComponent->Initialize(ComponentWorldTransform);
 			ActiveComponent->SetDebugDrawColor(FLinearColor::Green);
 			SelectedComponent->SetDebugDrawColor(FLinearColor::Blue);
diff --git a/Engine/Plugins/Experimental/Animation/PoseSearch/Source/Editor/Private/PoseSearchDebuggerViewModel.h b/Engine/Plugins/Experimental/Animation/PoseSearch/Source/Editor/Private/PoseSearchDebuggerViewModel.h
--- a/Engine/Plugins/Experimental/Animation/PoseSearch/Source/Editor/Private/PoseSearchDebuggerViewModel.h
+++ b/Engine/Plugins/Experimental/Animation/PoseSearch/Source/Editor/Private/PoseSearchDebuggerViewModel.h
@@ -28,7 +28,7 @@ public:
 	const UPoseSearchDatabase* GetCurrentDatabase() const;
 	const TArray<int32>* GetNodeIds() const;
 	int32 GetNodesNum() const;
-	const FTransform* GetRootTransform() const;
+	const FTransform& GetRootBoneTransform() const;
 
 	/** Checks if Update must be called */
 	bool HasSearchableAssetChanged() const;
@@ -53,6 +53,9 @@ public:
 
 	bool GetDrawQuery() const { return bDrawQuery; }
 
+	void SetDrawTrajectory(bool bInDrawTrajectory) { bDrawTrajectory = bInDrawTrajectory; }
+
+	bool GetDrawTrajectory() const { return bDrawTrajectory; }
 
 
 	/** Callback to reset debug skeletons for the active world */
@@ -111,8 +114,8 @@ private:
 	/** Current Skeletal Mesh Component Id for the AnimInstance */
 	uint64 SkeletalMeshComponentId = 0;
 
-	/** Currently active root transform on the skeletal mesh */
-	const FTransform* RootTransform = nullptr;
+	/** Currently active root bone transform */
+	FTransform RootBoneWorldTransform = FTransform::Identity;
 
 	/** Pointer to the active rewind debugger in the scene */
 	TAttribute<const IRewindDebugger*> RewindDebugger;
@@ -150,6 +153,8 @@ private:
 
 	bool bDrawQuery = true;	
 
+	bool bDrawTrajectory = false;
+	
 	/** Limits some public API */
 	friend class FDebugger;
 };
diff --git a/Engine/Plugins/Experimental/Animation/PoseSearch/Source/Runtime/Private/AnimNode_MotionMatching.cpp b/Engine/Plugins/Experimental/Animation/PoseSearch/Source/Runtime/Private/AnimNode_MotionMatching.cpp
--- a/Engine/Plugins/Experimental/Animation/PoseSearch/Source/Runtime/Private/AnimNode_MotionMatching.cpp
+++ b/Engine/Plugins/Experimental/Animation/PoseSearch/Source/Runtime/Private/AnimNode_MotionMatching.cpp
@@ -38,6 +38,8 @@ void FAnimNode_MotionMatching::Initialize_AnyThread(const FAnimationInitializeCo
 
 	Source.SetLinkNode(&BlendStackNode);
 	Source.Initialize(Context);
+
+	MotionMatchingState.UpdateRootBoneControl(Context.AnimInstanceProxy, 0.f);
 }
 
 void FAnimNode_MotionMatching::Evaluate_AnyThread(FPoseContext& Output)
@@ -46,6 +48,15 @@ void FAnimNode_MotionMatching::Evaluate_AnyThread(FPoseContext& Output)
 
 	Source.Evaluate(Output);
 
+	// applying MotionMatchingState.RootBoneDeltaYaw
+	if (!FMath::IsNearlyZero(MotionMatchingState.GetRootBoneDeltaYaw()))
+	{
+		const FQuat RootBoneDelta(FRotator(0.f, MotionMatchingState.GetRootBoneDeltaYaw(), 0.f));
+		FCompactPoseBoneIndex RootBoneIndex(RootBoneIndexType);
+		Output.Pose[RootBoneIndex].SetRotation(Output.Pose[RootBoneIndex].GetRotation() * RootBoneDelta);
+		Output.Pose[RootBoneIndex].NormalizeRotation();
+	}
+	
 #if UE_POSE_SEARCH_TRACE_ENABLED
 	MotionMatchingState.RootMotionTransformDelta = FTransform::Identity;
 
@@ -138,6 +149,8 @@ void FAnimNode_MotionMatching::UpdateAssetPlayer(const FAnimationUpdateContext&
 		SearchThrottleTime,
 		PlayRate,
 		MotionMatchingState,
+		YawFromAnimation,
+		YawFromAnimationTrajectoryBlendTime,
 		bForceInterruptNextUpdate,
 		bShouldSearch
 		#if ENABLE_ANIM_DEBUG
diff --git a/Engine/Plugins/Experimental/Animation/PoseSearch/Source/Runtime/Private/PoseSearchContext.cpp b/Engine/Plugins/Experimental/Animation/PoseSearch/Source/Runtime/Private/PoseSearchContext.cpp
--- a/Engine/Plugins/Experimental/Animation/PoseSearch/Source/Runtime/Private/PoseSearchContext.cpp
+++ b/Engine/Plugins/Experimental/Animation/PoseSearch/Source/Runtime/Private/PoseSearchContext.cpp
@@ -15,21 +15,21 @@ namespace UE::PoseSearch
 #if ENABLE_DRAW_DEBUG
 //////////////////////////////////////////////////////////////////////////
 // FDebugDrawParams
-FDebugDrawParams::FDebugDrawParams(FAnimInstanceProxy* InAnimInstanceProxy, const UPoseSearchDatabase* InDatabase, EDebugDrawFlags InFlags)
+FDebugDrawParams::FDebugDrawParams(FAnimInstanceProxy* InAnimInstanceProxy, const FTransform& InRootMotionTransform, const UPoseSearchDatabase* InDatabase, EDebugDrawFlags InFlags)
 : AnimInstanceProxy(InAnimInstanceProxy)
 , World(nullptr)
 , Mesh(nullptr)
-, RootMotionTransform(nullptr)
+, RootMotionTransform(&InRootMotionTransform)
 , Database(InDatabase)
 , Flags(InFlags)
 {
 }
 
-FDebugDrawParams::FDebugDrawParams(const UWorld* InWorld, const USkinnedMeshComponent* InMesh, const FTransform* InRootMotionTransform, const UPoseSearchDatabase* InDatabase, EDebugDrawFlags InFlags)
+FDebugDrawParams::FDebugDrawParams(const UWorld* InWorld, const USkinnedMeshComponent* InMesh, const FTransform& InRootMotionTransform, const UPoseSearchDatabase* InDatabase, EDebugDrawFlags InFlags)
 : AnimInstanceProxy(nullptr)
 , World(InWorld)
 , Mesh(InMesh)
-, RootMotionTransform(InRootMotionTransform)
+, RootMotionTransform(&InRootMotionTransform)
 , Database(InDatabase)
 , Flags(InFlags)
 {
@@ -95,22 +95,8 @@ FVector FDebugDrawParams::ExtractPosition(TConstArrayView<float> PoseVector, flo
 
 const FTransform& FDebugDrawParams::GetRootTransform() const
 {
-	if (RootMotionTransform)
-	{
-		return *RootMotionTransform;
-	}
-
-	if (AnimInstanceProxy)
-	{
-		return AnimInstanceProxy->GetComponentTransform();
-	}
-	
-	if (Mesh)
-	{
-		return Mesh->GetComponentTransform();
-	}
-
-	return FTransform::Identity;
+	check(RootMotionTransform);
+	return *RootMotionTransform;
 }
 
 void FDebugDrawParams::DrawLine(const FVector& LineStart, const FVector& LineEnd, const FColor& Color, float Thickness) const
diff --git a/Engine/Plugins/Experimental/Animation/PoseSearch/Source/Runtime/Private/PoseSearchHistory.cpp b/Engine/Plugins/Experimental/Animation/PoseSearch/Source/Runtime/Private/PoseSearchHistory.cpp
--- a/Engine/Plugins/Experimental/Animation/PoseSearch/Source/Runtime/Private/PoseSearchHistory.cpp
+++ b/Engine/Plugins/Experimental/Animation/PoseSearch/Source/Runtime/Private/PoseSearchHistory.cpp
@@ -66,8 +66,12 @@ FORCEINLINE auto LowerBound(IteratorType First, IteratorType Last, const ValueTy
 // FPoseHistoryEntry
 void FPoseHistoryEntry::Update(float InTime, FCSPose<FCompactPose>& ComponentSpacePose, const FTransform& ComponentTransform, const FBoneToTransformMap& BoneToTransformMap)
 {
+	// @todo: optimize this math by initializing FCSPose<FCompactPose> root to identity
+	const FTransform ComponentToRootTransform = ComponentSpacePose.GetComponentSpaceTransform(FCompactPoseBoneIndex(RootBoneIndexType));
+	const FTransform ComponentToRootTransformInv = ComponentToRootTransform.Inverse();
+
 	Time = InTime;
-	RootTransform = ComponentTransform;
+	RootTransform = ComponentToRootTransform * ComponentTransform;
 
 	const FBoneContainer& BoneContainer = ComponentSpacePose.GetPose().GetBoneContainer();
 	const USkeleton* SkeletonAsset = BoneContainer.GetSkeletonAsset();
@@ -83,7 +87,8 @@ void FPoseHistoryEntry::Update(float InTime, FCSPose<FCompactPose>& ComponentSpa
 		for (FSkeletonPoseBoneIndex SkeletonBoneIdx(0); SkeletonBoneIdx != NumSkeletonBones; ++SkeletonBoneIdx)
 		{
 			const FCompactPoseBoneIndex CompactBoneIdx = BoneContainer.GetCompactPoseIndexFromSkeletonPoseIndex(SkeletonBoneIdx);
-			ComponentSpaceTransforms[SkeletonBoneIdx.GetInt()] = CompactBoneIdx.IsValid() ? ComponentSpacePose.GetComponentSpaceTransform(CompactBoneIdx) : RefBonePose[SkeletonBoneIdx.GetInt()];
+			ComponentSpaceTransforms[SkeletonBoneIdx.GetInt()] = 
+				(CompactBoneIdx.IsValid() ? ComponentSpacePose.GetComponentSpaceTransform(CompactBoneIdx) : RefBonePose[SkeletonBoneIdx.GetInt()]) * ComponentToRootTransformInv;
 		}
 	}
 	else
@@ -93,7 +98,8 @@ void FPoseHistoryEntry::Update(float InTime, FCSPose<FCompactPose>& ComponentSpa
 		{
 			const FSkeletonPoseBoneIndex SkeletonBoneIdx(BoneToTransformPair.Key);
 			const FCompactPoseBoneIndex CompactBoneIdx = BoneContainer.GetCompactPoseIndexFromSkeletonPoseIndex(SkeletonBoneIdx);
-			ComponentSpaceTransforms[BoneToTransformPair.Value] = CompactBoneIdx.IsValid() ? ComponentSpacePose.GetComponentSpaceTransform(CompactBoneIdx) : RefBonePose[SkeletonBoneIdx.GetInt()];
+			ComponentSpaceTransforms[BoneToTransformPair.Value] = 
+				(CompactBoneIdx.IsValid() ? ComponentSpacePose.GetComponentSpaceTransform(CompactBoneIdx) : RefBonePose[SkeletonBoneIdx.GetInt()]) * ComponentToRootTransformInv;
 		}
 	}
 }
diff --git a/Engine/Plugins/Experimental/Animation/PoseSearch/Source/Runtime/Private/PoseSearchLibrary.cpp b/Engine/Plugins/Experimental/Animation/PoseSearch/Source/Runtime/Private/PoseSearchLibrary.cpp
--- a/Engine/Plugins/Experimental/Animation/PoseSearch/Source/Runtime/Private/PoseSearchLibrary.cpp
+++ b/Engine/Plugins/Experimental/Animation/PoseSearch/Source/Runtime/Private/PoseSearchLibrary.cpp
@@ -10,6 +10,7 @@
 #include "Animation/AnimSequence.h"
 #include "Animation/AnimSubsystem_Tag.h"
 #include "Animation/BlendSpace.h"
+#include "Animation/BuiltInAttributeTypes.h"
 #include "InstancedStruct.h"
 #include "PoseSearch/AnimNode_MotionMatching.h"
 #include "PoseSearch/AnimNode_PoseSearchHistoryCollector.h"
@@ -37,8 +38,14 @@ void FMotionMatchingState::Reset()
 	CurrentSearchResult.Reset();
 	// Set the elapsed time to INFINITY to trigger a search right away
 	ElapsedPoseSearchTime = INFINITY;
-	PoseIndicesHistory.Reset();
 	WantedPlayRate = 1.f;
+	bJumpedToPose = false;
+	RootBoneDeltaYaw = 0.f;
+	RootBoneWorldYaw = 0.f;
+#if UE_POSE_SEARCH_TRACE_ENABLED
+	RootMotionTransformDelta = FTransform::Identity;
+#endif // UE_POSE_SEARCH_TRACE_ENABLED
+	PoseIndicesHistory.Reset();
 }
 
 void FMotionMatchingState::AdjustAssetTime(float AssetTime)
@@ -140,9 +147,52 @@ void FMotionMatchingState::UpdateWantedPlayRate(const UE::PoseSearch::FSearchCon
 	}
 }
 
+void FMotionMatchingState::UpdateRootBoneControl(const FAnimInstanceProxy* AnimInstanceProxy, float RootBoneYawFromAnimation)
+{
+	static const FName RootMotionAttributeName = "RootMotionDelta";
+	static const UE::Anim::FAttributeId RootMotionAttributeId = { RootMotionAttributeName , FCompactPoseBoneIndex(0) };
+
+	const FRotator ComponentWorldRotator(AnimInstanceProxy->GetComponentTransform().GetRotation());
+	if (FMath::IsNearlyZero(RootBoneYawFromAnimation))
+	{
+		RootBoneWorldYaw = ComponentWorldRotator.Yaw;
+		RootBoneWorldYaw = 0.f;
+	}
+	else
+	{
+		// integrating MotionMatchingState.RootBoneTargetYaw with the previous frame RootMotionDelta
+		const USkeletalMeshComponent* Mesh = AnimInstanceProxy->GetSkelMeshComponent();
+		check(Mesh);
+		
+		if (const FTransformAnimationAttribute* RootMotionAttribute = Mesh->GetCustomAttributes().Find<FTransformAnimationAttribute>(RootMotionAttributeId))
+		{
+			const FRotator RootMotionRotatorDelta(RootMotionAttribute->Value.GetRotation());
+			const float RootBoneAnimationDelta = RootMotionRotatorDelta.Yaw;
+			const float RootBoneToComponentDelta = FRotator::NormalizeAxis(ComponentWorldRotator.Yaw - RootBoneWorldYaw);
+
+			// @todo: RootBoneYawFromAnimation should be a speed (influenced by dt)
+			// lerping the animation delta with the capsule delta
+			const float RootBoneDelta = FMath::Lerp(RootBoneToComponentDelta, RootBoneAnimationDelta, RootBoneYawFromAnimation);
+
+			RootBoneWorldYaw = FRotator::NormalizeAxis(RootBoneWorldYaw + RootBoneDelta);
+		}
+		else
+		{
+			UE_LOG(LogPoseSearch, Warning, TEXT("FMotionMatchingState::UpdateRootBoneControl - Couldn't find FTransformAnimationAttribute. Root bone yaw control will not be performed"));
+
+			RootBoneWorldYaw = ComponentWorldRotator.Yaw;
+			RootBoneWorldYaw = 0.f;
+		}
+
+		// @todo: handle the case when the character is on top of a rotating platform
+		RootBoneDeltaYaw = FRotator::NormalizeAxis(RootBoneWorldYaw - ComponentWorldRotator.Yaw);
+	}
+}
+
 #if UE_POSE_SEARCH_TRACE_ENABLED
 void UPoseSearchLibrary::TraceMotionMatchingState(
 	const UPoseSearchDatabase* Database,
+	const FPoseSearchQueryTrajectory& Trajectory,
 	UE::PoseSearch::FSearchContext& SearchContext,
 	const UE::PoseSearch::FSearchResult& CurrentResult,
 	const UE::PoseSearch::FSearchResult& LastResult,
@@ -252,6 +302,8 @@ void UPoseSearchLibrary::TraceMotionMatchingState(
 	TraceState.SearchBestCost = SearchBestCost;
 	TraceState.SearchBruteForceCost = SearchBruteForceCost;
 
+	TraceState.Trajectory = Trajectory;
+
 	TraceState.Output(AnimInstance, NodeId);
 }
 #endif // UE_POSE_SEARCH_TRACE_ENABLED
@@ -268,6 +320,8 @@ void UPoseSearchLibrary::UpdateMotionMatchingState(
 	float SearchThrottleTime,
 	const FFloatInterval& PlayRate,
 	FMotionMatchingState& InOutMotionMatchingState,
+	float RootBoneYawFromAnimation,
+	float RootBoneDeltaYawBlendTime,
 	bool bForceInterrupt,
 	bool bShouldSearch,
 	bool bDebugDrawQuery,
@@ -285,6 +339,8 @@ void UPoseSearchLibrary::UpdateMotionMatchingState(
 		return;
 	}
 
+	InOutMotionMatchingState.UpdateRootBoneControl(Context.AnimInstanceProxy, RootBoneYawFromAnimation);
+
 	const float DeltaTime = Context.GetDeltaTime();
 
 	InOutMotionMatchingState.bJumpedToPose = false;
@@ -306,8 +362,8 @@ void UPoseSearchLibrary::UpdateMotionMatchingState(
 		QueryMirrorRequest = CurrentIndexAsset->bMirrored ? EPoseSearchBooleanRequest::TrueValue : EPoseSearchBooleanRequest::FalseValue;
 	}
 
-	const FPoseSearchQueryTrajectory TrajectoryCS = ProcessTrajectory(Trajectory, Context.AnimInstanceProxy->GetActorTransform(), Context.AnimInstanceProxy->GetComponentTransform(), TrajectorySpeedMultiplier);
-	FSearchContext SearchContext(&TrajectoryCS, History, 0.f, &InOutMotionMatchingState.PoseIndicesHistory, QueryMirrorRequest, 
+	const FPoseSearchQueryTrajectory TrajectoryRootSpace = ProcessTrajectory(Trajectory, Context.AnimInstanceProxy->GetComponentTransform(), InOutMotionMatchingState.GetRootBoneDeltaYaw(), RootBoneDeltaYawBlendTime, TrajectorySpeedMultiplier);
+	FSearchContext SearchContext(&TrajectoryRootSpace, History, 0.f, &InOutMotionMatchingState.PoseIndicesHistory, QueryMirrorRequest,
 		InOutMotionMatchingState.CurrentSearchResult, PoseJumpThresholdTime, bForceInterrupt, InOutMotionMatchingState.CanAdvance(DeltaTime));
 
 	// If we can't advance or enough time has elapsed since the last pose jump then search
@@ -379,7 +435,7 @@ void UPoseSearchLibrary::UpdateMotionMatchingState(
 #endif // WITH_EDITORONLY_DATA
 
 		// @todo: Update this to account for multiple databases.
-		TraceMotionMatchingState(Databases[0], SearchContext, InOutMotionMatchingState.CurrentSearchResult, LastResult, InOutMotionMatchingState.ElapsedPoseSearchTime,
+		TraceMotionMatchingState(Databases[0], Trajectory, SearchContext, InOutMotionMatchingState.CurrentSearchResult, LastResult, InOutMotionMatchingState.ElapsedPoseSearchTime,
 			InOutMotionMatchingState.RootMotionTransformDelta, Context.AnimInstanceProxy->GetAnimInstanceObject(), Context.GetCurrentNodeId(), DeltaTime, bSearch,
 			AnimInstance ? FObjectTrace::GetWorldElapsedTime(AnimInstance->GetWorld()) : 0.f, SearchBestCost, SearchBruteForceCost);
 	}
@@ -399,15 +455,19 @@ void UPoseSearchLibrary::UpdateMotionMatchingState(
 		else
 #endif // WITH_EDITOR
 		{
+			const FRotator DeltaRotation(0.f, InOutMotionMatchingState.GetRootBoneDeltaYaw(), 0.f);
+			const FTransform DeltaTransform(DeltaRotation);
+			const FTransform RootBoneTransform = DeltaTransform * Context.AnimInstanceProxy->GetComponentTransform();
+
 			if (bDebugDrawCurResult)
 			{
-				UE::PoseSearch::FDebugDrawParams DrawParams(Context.AnimInstanceProxy, CurResultDatabase);
+				UE::PoseSearch::FDebugDrawParams DrawParams(Context.AnimInstanceProxy, RootBoneTransform, CurResultDatabase);
 				DrawParams.DrawFeatureVector(CurResult.PoseIdx);
 			}
 
 			if (bDebugDrawQuery)
 			{
-				UE::PoseSearch::FDebugDrawParams DrawParams(Context.AnimInstanceProxy, CurResultDatabase, EDebugDrawFlags::DrawQuery);
+				UE::PoseSearch::FDebugDrawParams DrawParams(Context.AnimInstanceProxy, RootBoneTransform, CurResultDatabase, EDebugDrawFlags::DrawQuery);
 				DrawParams.DrawFeatureVector(SearchContext.GetOrBuildQuery(CurResultDatabase->Schema).GetValues());
 			}
 		}
@@ -415,26 +475,26 @@ void UPoseSearchLibrary::UpdateMotionMatchingState(
 #endif
 }
 
-// transforms Trajectory from world space to mesh component space, and scale it by TrajectorySpeedMultiplier
-FPoseSearchQueryTrajectory UPoseSearchLibrary::ProcessTrajectory(const FPoseSearchQueryTrajectory& Trajectory, const FTransform& OwnerTransformWS, const FTransform& ComponentTransformWS, float TrajectorySpeedMultiplier)
+// transforms Trajectory from world space to root bone space, and scale it by TrajectorySpeedMultiplier
+FPoseSearchQueryTrajectory UPoseSearchLibrary::ProcessTrajectory(const FPoseSearchQueryTrajectory& Trajectory, const FTransform& ComponentWorldTransform, float RootBoneDeltaYaw, float RootBoneDeltaYawBlendTime, float TrajectorySpeedMultiplier)
 {
-	FPoseSearchQueryTrajectory TrajectoryCS = Trajectory;
-	const FTransform ToComponentTransform = ComponentTransformWS.Inverse();
-	for (FPoseSearchQueryTrajectorySample& Sample : TrajectoryCS.Samples)
-	{
-		Sample.Position = ToComponentTransform.TransformPosition(Sample.Position);
-		Sample.Facing = ToComponentTransform.TransformRotation(Sample.Facing);
-	}
+	const float TrajectorySpeedMultiplierInv = FMath::IsNearlyZero(TrajectorySpeedMultiplier) ? 1.f : 1.f / TrajectorySpeedMultiplier;
 
-	if (!FMath::IsNearlyEqual(TrajectorySpeedMultiplier, 1.f) && !FMath::IsNearlyZero(TrajectorySpeedMultiplier))
+	FPoseSearchQueryTrajectory TrajectoryRootSpace = Trajectory;
+	const FTransform ToRootSpace = ComponentWorldTransform.Inverse();
+	for (FPoseSearchQueryTrajectorySample& Sample : TrajectoryRootSpace.Samples)
 	{
-		for (FPoseSearchQueryTrajectorySample& Sample : TrajectoryCS.Samples)
-		{
-			Sample.AccumulatedSeconds /= TrajectorySpeedMultiplier;
-		}
+		Sample.AccumulatedSeconds *= TrajectorySpeedMultiplierInv;
+
+		Sample.Position = ToRootSpace.TransformPosition(Sample.Position);
+
+		const float BlendParam = RootBoneDeltaYawBlendTime < UE_KINDA_SMALL_NUMBER ? 1 : FMath::Clamp(1.f - (Sample.AccumulatedSeconds - RootBoneDeltaYawBlendTime) / RootBoneDeltaYawBlendTime, 0.f, 1.f);
+		const FQuat RootBoneDelta(FRotator(0.f, RootBoneDeltaYaw * BlendParam, 0.f));
+		const FQuat ToRootSpaceRotation = ToRootSpace.GetRotation() * RootBoneDelta;
+		Sample.Facing = ToRootSpaceRotation * Sample.Facing;
 	}
 
-	return TrajectoryCS;
+	return TrajectoryRootSpace;
 }
 
 void UPoseSearchLibrary::MotionMatch(
@@ -483,7 +543,7 @@ void UPoseSearchLibrary::MotionMatch(
 
 	if (Database && AnimInstance)
 	{
-		const FPoseSearchQueryTrajectory TrajectoryCS = ProcessTrajectory(Trajectory, AnimInstance->GetOwningActor()->GetActorTransform(), AnimInstance->GetOwningComponent()->GetComponentTransform(), TrajectorySpeedMultiplier);
+		const FPoseSearchQueryTrajectory TrajectoryRootSpace = ProcessTrajectory(Trajectory, AnimInstance->GetOwningComponent()->GetComponentTransform(), 0.f, 0.f, TrajectorySpeedMultiplier);
 
 		// ExtendedPoseHistory will hold future poses to match AssetSamplerBase (at FutureAnimationStartTime) TimeToFutureAnimationStart seconds in the future
 		FExtendedPoseHistory ExtendedPoseHistory;
@@ -551,7 +611,7 @@ void UPoseSearchLibrary::MotionMatch(
 				FCSPose<FCompactPose> ComponentSpacePose;
 				ComponentSpacePose.InitPose(Pose);
 
-				const FPoseSearchQueryTrajectorySample TrajectorySample = TrajectoryCS.GetSampleAtTime(ExtractionTime);
+				const FPoseSearchQueryTrajectorySample TrajectorySample = TrajectoryRootSpace.GetSampleAtTime(ExtractionTime);
 				const FTransform& ComponentTransform = AnimInstance->GetOwningComponent()->GetComponentTransform();
 				const FTransform FutureComponentTransform = TrajectorySample.GetTransform() * ComponentTransform;
 
@@ -570,7 +630,7 @@ void UPoseSearchLibrary::MotionMatch(
 		}
 
 		// @todo: finish set up SearchContext by exposing or calculating additional members
-		FSearchContext SearchContext(&TrajectoryCS, ExtendedPoseHistory.IsInitialized() ? &ExtendedPoseHistory : nullptr, TimeToFutureAnimationStart);
+		FSearchContext SearchContext(&TrajectoryRootSpace, ExtendedPoseHistory.IsInitialized() ? &ExtendedPoseHistory : nullptr, TimeToFutureAnimationStart);
 
 		FSearchResult SearchResult = Database->Search(SearchContext);
 		if (SearchResult.IsValid())
@@ -590,15 +650,16 @@ void UPoseSearchLibrary::MotionMatch(
 #if ENABLE_DRAW_DEBUG && ENABLE_ANIM_DEBUG
 		if (SearchResult.IsValid())
 		{
+			FAnimInstanceProxy* AnimInstanceProxy = UAnimInstanceProxyProvider::GetAnimInstanceProxy(AnimInstance);
 			if (CVarAnimMotionMatchDrawMatchEnable.GetValueOnAnyThread())
 			{
-				UE::PoseSearch::FDebugDrawParams DrawParams(UAnimInstanceProxyProvider::GetAnimInstanceProxy(AnimInstance), SearchResult.Database.Get());
+				UE::PoseSearch::FDebugDrawParams DrawParams(AnimInstanceProxy, AnimInstanceProxy->GetComponentTransform(), SearchResult.Database.Get());
 				DrawParams.DrawFeatureVector(SearchResult.PoseIdx);
 			}
 
 			if (CVarAnimMotionMatchDrawQueryEnable.GetValueOnAnyThread())
 			{
-				UE::PoseSearch::FDebugDrawParams DrawParams(UAnimInstanceProxyProvider::GetAnimInstanceProxy(AnimInstance), SearchResult.Database.Get(), EDebugDrawFlags::DrawQuery);
+				UE::PoseSearch::FDebugDrawParams DrawParams(AnimInstanceProxy, AnimInstanceProxy->GetComponentTransform(), SearchResult.Database.Get(), EDebugDrawFlags::DrawQuery);
 				DrawParams.DrawFeatureVector(SearchContext.GetOrBuildQuery(SearchResult.Database->Schema).GetValues());
 			}
 		}
@@ -613,7 +674,7 @@ void UPoseSearchLibrary::MotionMatch(
 			SearchBruteForceCost = SearchResult.BruteForcePoseCost.GetTotalCost();
 		}
 #endif // WITH_EDITORONLY_DATA
-		TraceMotionMatchingState(Database, SearchContext, SearchResult, FSearchResult(), 0.f, FTransform::Identity, AnimInstance, DebugSessionUniqueIdentifier,
+		TraceMotionMatchingState(Database, Trajectory, SearchContext, SearchResult, FSearchResult(), 0.f, FTransform::Identity, AnimInstance, DebugSessionUniqueIdentifier,
 			AnimInstance->GetDeltaSeconds(), true, FObjectTrace::GetWorldElapsedTime(AnimInstance->GetWorld()), SearchBestCost, SearchBruteForceCost);
 #endif // UE_POSE_SEARCH_TRACE_ENABLED
 	}
diff --git a/Engine/Plugins/Experimental/Animation/PoseSearch/Source/Runtime/Private/PoseSearchTrajectoryTypes.cpp b/Engine/Plugins/Experimental/Animation/PoseSearch/Source/Runtime/Private/PoseSearchTrajectoryTypes.cpp
--- a/Engine/Plugins/Experimental/Animation/PoseSearch/Source/Runtime/Private/PoseSearchTrajectoryTypes.cpp
+++ b/Engine/Plugins/Experimental/Animation/PoseSearch/Source/Runtime/Private/PoseSearchTrajectoryTypes.cpp
@@ -79,4 +79,18 @@ void FPoseSearchQueryTrajectory::DebugDrawTrajectory(const UWorld* World) const
 		}
 	}
 }
-#endif // ENABLE_ANIM_DEBUG
\ No newline at end of file
+#endif // ENABLE_ANIM_DEBUG
+
+FArchive& operator<<(FArchive& Ar, FPoseSearchQueryTrajectorySample& TrajectorySample)
+{
+	Ar << TrajectorySample.Facing;
+	Ar << TrajectorySample.Position;
+	Ar << TrajectorySample.AccumulatedSeconds;
+	return Ar;
+}
+
+FArchive& operator<<(FArchive& Ar, FPoseSearchQueryTrajectory& Trajectory)
+{
+	Ar << Trajectory.Samples;
+	return Ar;
+}
\ No newline at end of file
diff --git a/Engine/Plugins/Experimental/Animation/PoseSearch/Source/Runtime/Private/Trace/PoseSearchTraceLogger.cpp b/Engine/Plugins/Experimental/Animation/PoseSearch/Source/Runtime/Private/Trace/PoseSearchTraceLogger.cpp
--- a/Engine/Plugins/Experimental/Animation/PoseSearch/Source/Runtime/Private/Trace/PoseSearchTraceLogger.cpp
+++ b/Engine/Plugins/Experimental/Animation/PoseSearch/Source/Runtime/Private/Trace/PoseSearchTraceLogger.cpp
@@ -82,6 +82,7 @@ FArchive& operator<<(FArchive& Ar, FTraceMotionMatchingState& State)
 	Ar << State.SearchBestCost;
 	Ar << State.SearchBruteForceCost;
 	Ar << State.DatabaseEntries;
+	Ar << State.Trajectory;
 	Ar << State.CurrentDbEntryIdx;
 	Ar << State.CurrentPoseEntryIdx;
 	return Ar;
diff --git a/Engine/Plugins/Experimental/Animation/PoseSearch/Source/Runtime/Private/Trace/PoseSearchTraceLogger.h b/Engine/Plugins/Experimental/Animation/PoseSearch/Source/Runtime/Private/Trace/PoseSearchTraceLogger.h
--- a/Engine/Plugins/Experimental/Animation/PoseSearch/Source/Runtime/Private/Trace/PoseSearchTraceLogger.h
+++ b/Engine/Plugins/Experimental/Animation/PoseSearch/Source/Runtime/Private/Trace/PoseSearchTraceLogger.h
@@ -3,6 +3,7 @@
 #pragma once
 
 #include "PoseSearch/PoseSearchContext.h"
+#include "PoseSearch/PoseSearchTrajectoryTypes.h"
 
 UE_TRACE_CHANNEL_EXTERN(PoseSearchChannel, POSESEARCH_API);
 
@@ -83,6 +84,8 @@ struct POSESEARCH_API FTraceMotionMatchingState
 
 	TArray<FTraceMotionMatchingStateDatabaseEntry> DatabaseEntries;
 
+	FPoseSearchQueryTrajectory Trajectory;
+
 	/** Index of the current database in DatabaseEntries */
 	int32 CurrentDbEntryIdx = INDEX_NONE;
 
diff --git a/Engine/Plugins/Experimental/Animation/PoseSearch/Source/Runtime/Public/PoseSearch/AnimNode_MotionMatching.h b/Engine/Plugins/Experimental/Animation/PoseSearch/Source/Runtime/Public/PoseSearch/AnimNode_MotionMatching.h
--- a/Engine/Plugins/Experimental/Animation/PoseSearch/Source/Runtime/Public/PoseSearch/AnimNode_MotionMatching.h
+++ b/Engine/Plugins/Experimental/Animation/PoseSearch/Source/Runtime/Public/PoseSearch/AnimNode_MotionMatching.h
@@ -16,7 +16,38 @@
 {
 	GENERATED_BODY()
 
+public:
+	// Search InDatabase instead of the Database property on this node. Use bForceInterruptIfNew to ignore the continuing pose if InDatabase is new.
+	void SetDatabaseToSearch(UPoseSearchDatabase* InDatabase, bool bForceInterruptIfNew);
+
+	// Search InDatabases instead of the Database property on the node. Use bForceInterruptIfNew to ignore the continuing pose if InDatabases is new.
+	void SetDatabasesToSearch(const TArray<UPoseSearchDatabase*>& InDatabases, bool bForceInterruptIfNew);
+
+	// Reset the effects of SetDatabaseToSearch/SetDatabasesToSearch and use the Database property on this node.
+	void ResetDatabasesToSearch(bool bInForceInterrupt);
+
+	// Ignore the continuing pose on the next update and force a search.
+	void ForceInterruptNextUpdate();
+
 private:
+	// FAnimNode_Base interface
+	// @todo: implement CacheBones_AnyThread to rebind the schema bones
+	virtual void Initialize_AnyThread(const FAnimationInitializeContext& Context) override;
+	virtual void Evaluate_AnyThread(FPoseContext& Output) override;
+	virtual void GatherDebugData(FNodeDebugData& DebugData) override;
+	// End of FAnimNode_Base interface
+
+	// FAnimNode_AssetPlayerBase interface
+	virtual float GetAccumulatedTime() const override;
+	virtual UAnimationAsset* GetAnimAsset() const override;
+	virtual void UpdateAssetPlayer(const FAnimationUpdateContext& Context) override;
+	virtual float GetCurrentAssetLength() const override;
+	virtual float GetCurrentAssetTime() const override;
+	virtual float GetCurrentAssetTimePlayRateAdjusted() const override;
+	virtual bool GetIgnoreForRelevancyTest() const override;
+	virtual bool SetIgnoreForRelevancyTest(bool bInIgnoreForRelevancyTest) override;
+	// End of FAnimNode_AssetPlayerBase interface
+
 	UPROPERTY()
 	FPoseLink Source;
 
@@ -72,28 +103,13 @@
 	UPROPERTY(EditAnywhere, Category = Settings, meta = (PinHiddenByDefault))
 	bool bShouldSearch = true;
 
-public:
+	// blend time over which the YawFromAnimation is distributed across the trajectory samples (negative values implies YawFromAnimation is constant over the entire trajectory, so the trajectory will not try to recover towards the capsule orientation)
+	UPROPERTY(EditAnywhere, Category = RootMotion, meta = (PinHiddenByDefault))
+	float YawFromAnimationTrajectoryBlendTime = 0.1f;
 
-	// FAnimNode_Base interface
-	// @todo: implement CacheBones_AnyThread to rebind the schema bones
-	virtual void Initialize_AnyThread(const FAnimationInitializeContext& Context) override;
-	virtual void Evaluate_AnyThread(FPoseContext& Output) override;
-	virtual void GatherDebugData(FNodeDebugData& DebugData) override;
-	// End of FAnimNode_Base interface
+	// percentage of delta yaw from the animation controlling the root bone (0% the capsule is authoritative over the root bone orientation, 100% the orientation will be fully controlled by animation)
+	UPROPERTY(EditAnywhere, Category = RootMotion, meta = (PinHiddenByDefault))
+	float YawFromAnimation = 0.f;
 
-	// Search InDatabase instead of the Database property on this node. Use bForceInterruptIfNew to ignore the continuing pose if InDatabase is new.
-	void SetDatabaseToSearch(UPoseSearchDatabase* InDatabase, bool bForceInterruptIfNew);
-
-	// Search InDatabases instead of the Database property on the node. Use bForceInterruptIfNew to ignore the continuing pose if InDatabases is new.
-	void SetDatabasesToSearch(const TArray<UPoseSearchDatabase*>& InDatabases, bool bForceInterruptIfNew);
-
-	// Reset the effects of SetDatabaseToSearch/SetDatabasesToSearch and use the Database property on this node.
-	void ResetDatabasesToSearch(bool bInForceInterrupt);
-
-	// Ignore the continuing pose on the next update and force a search.
-	void ForceInterruptNextUpdate();
-
-private:
-
 	FAnimNode_BlendStack_Standalone BlendStackNode;
 
@@ -113,21 +129,6 @@
 	// True if the Database property on this node has been overridden by SetDatabaseToSearch/SetDatabasesToSearch.
 	bool bOverrideDatabaseInput = false;
 
-	// FAnimNode_AssetPlayerBase
-protected:
-	// FAnimNode_AssetPlayerBase interface
-	virtual float GetAccumulatedTime() const override;
-	virtual UAnimationAsset* GetAnimAsset() const override;
-	virtual void UpdateAssetPlayer(const FAnimationUpdateContext& Context) override;
-	virtual float GetCurrentAssetLength() const override;
-	virtual float GetCurrentAssetTime() const override;
-	virtual float GetCurrentAssetTimePlayRateAdjusted() const override;
-	virtual bool GetIgnoreForRelevancyTest() const override;
-	virtual bool SetIgnoreForRelevancyTest(bool bInIgnoreForRelevancyTest) override;
-	// End of FAnimNode_AssetPlayerBase interface
-
-private:
-
 #if WITH_EDITORONLY_DATA
 	// If true, "Relevant anim" nodes that look for the highest weighted animation in a state will ignore this node
 	UPROPERTY(EditAnywhere, Category = Relevancy, meta = (FoldProperty, PinHiddenByDefault))
diff --git a/Engine/Plugins/Experimental/Animation/PoseSearch/Source/Runtime/Public/PoseSearch/PoseSearchContext.h b/Engine/Plugins/Experimental/Animation/PoseSearch/Source/Runtime/Public/PoseSearch/PoseSearchContext.h
--- a/Engine/Plugins/Experimental/Animation/PoseSearch/Source/Runtime/Public/PoseSearch/PoseSearchContext.h
+++ b/Engine/Plugins/Experimental/Animation/PoseSearch/Source/Runtime/Public/PoseSearch/PoseSearchContext.h
@@ -105,8 +105,8 @@ private:
 #if ENABLE_DRAW_DEBUG
 struct POSESEARCH_API FDebugDrawParams
 {
-	FDebugDrawParams(FAnimInstanceProxy* InAnimInstanceProxy, const UPoseSearchDatabase* InDatabase, EDebugDrawFlags InFlags = EDebugDrawFlags::None);
-	FDebugDrawParams(const UWorld* InWorld, const USkinnedMeshComponent* InMesh, const FTransform* InRootMotionTransform, const UPoseSearchDatabase* InDatabase, EDebugDrawFlags InFlags = EDebugDrawFlags::None);
+	FDebugDrawParams(FAnimInstanceProxy* InAnimInstanceProxy, const FTransform& InRootMotionTransform, const UPoseSearchDatabase* InDatabase, EDebugDrawFlags InFlags = EDebugDrawFlags::None);
+	FDebugDrawParams(const UWorld* InWorld, const USkinnedMeshComponent* InMesh, const FTransform& InRootMotionTransform, const UPoseSearchDatabase* InDatabase, EDebugDrawFlags InFlags = EDebugDrawFlags::None);
 
 	const FPoseSearchIndex* GetSearchIndex() const;
 	const UPoseSearchSchema* GetSchema() const;
diff --git a/Engine/Plugins/Experimental/Animation/PoseSearch/Source/Runtime/Public/PoseSearch/PoseSearchLibrary.h b/Engine/Plugins/Experimental/Animation/PoseSearch/Source/Runtime/Public/PoseSearch/PoseSearchLibrary.h
--- a/Engine/Plugins/Experimental/Animation/PoseSearch/Source/Runtime/Public/PoseSearch/PoseSearchLibrary.h
+++ b/Engine/Plugins/Experimental/Animation/PoseSearch/Source/Runtime/Public/PoseSearch/PoseSearchLibrary.h
@@ -44,6 +44,10 @@ struct POSESEARCH_API FMotionMatchingState
 
 	void UpdateWantedPlayRate(const UE::PoseSearch::FSearchContext& SearchContext, const FFloatInterval& PlayRate, float TrajectorySpeedMultiplier);
 
+	void UpdateRootBoneControl(const FAnimInstanceProxy* AnimInstanceProxy, float RootBoneYawFromAnimation);
+
+	float GetRootBoneDeltaYaw() const { return RootBoneDeltaYaw; }
+
 	UE::PoseSearch::FSearchResult CurrentSearchResult;
 
 	// Time since the last pose jump
@@ -55,23 +59,19 @@ struct POSESEARCH_API FMotionMatchingState
 	float WantedPlayRate = 1.f;
 
 	// true if a new animation has been selected
-	UPROPERTY(VisibleAnywhere, BlueprintReadOnly, Category=State)
+	UPROPERTY(Transient, VisibleAnywhere, BlueprintReadOnly, Category=State)
 	bool bJumpedToPose = false;
 
 #if UE_POSE_SEARCH_TRACE_ENABLED
-	// Root motion delta for currently playing animation.
+	// Root motion delta for currently playing animation (or animation tree if from the blend stack)
 	FTransform RootMotionTransformDelta = FTransform::Identity;
 #endif //UE_POSE_SEARCH_TRACE_ENABLED
 
-	// @todo: add ContinuingPoseCost / BruteForcePoseCost / PoseCost graphs to rewind debugger 
-//#if WITH_EDITORONLY_DATA
-//	enum { SearchCostHistoryNumSamples = 200 };
-//	FDebugFloatHistory SearchCostHistoryContinuing = FDebugFloatHistory(SearchCostHistoryNumSamples, 0, 0, true);
-//	FDebugFloatHistory SearchCostHistoryBruteForce = FDebugFloatHistory(SearchCostHistoryNumSamples, 0, 0, true);
-//	FDebugFloatHistory SearchCostHistoryKDTree = FDebugFloatHistory(SearchCostHistoryNumSamples, 0, 0, true);
-//#endif
-
 	UE::PoseSearch::FPoseIndicesHistory PoseIndicesHistory;
+
+private:
+	float RootBoneDeltaYaw = 0.f;
+	float RootBoneWorldYaw = 0.f;
 };
 
 UCLASS()
@@ -82,6 +82,7 @@ class POSESEARCH_API UPoseSearchLibrary : public UBlueprintFunctionLibrary
 #if UE_POSE_SEARCH_TRACE_ENABLED
 	static void TraceMotionMatchingState(
 		const UPoseSearchDatabase* Database,
+		const FPoseSearchQueryTrajectory& Trajectory,
 		UE::PoseSearch::FSearchContext& SearchContext,
 		const UE::PoseSearch::FSearchResult& CurrentResult,
 		const UE::PoseSearch::FSearchResult& LastResult,
@@ -98,8 +99,9 @@ class POSESEARCH_API UPoseSearchLibrary : public UBlueprintFunctionLibrary
 
 	static FPoseSearchQueryTrajectory ProcessTrajectory(
 		const FPoseSearchQueryTrajectory& Trajectory,
-		const FTransform& OwnerTransformWS,
-		const FTransform& ComponentTransformWS,
+		const FTransform& ComponentWorldTransform,
+		float RootBoneDeltaYaw,
+		float RootBoneDeltaYawBlendTime,
 		float TrajectorySpeedMultiplier);
 
 public:
@@ -133,6 +135,8 @@ public:
 		float SearchThrottleTime,
 		const FFloatInterval& PlayRate,
 		FMotionMatchingState& InOutMotionMatchingState,
+		float RootBoneYawFromAnimation,
+		float RootBoneDeltaYawBlendTime,
 		bool bForceInterrupt = false,
 		bool bShouldSearch = true,
 		bool bDebugDrawQuery = false,
diff --git a/Engine/Plugins/Experimental/Animation/PoseSearch/Source/Runtime/Public/PoseSearch/PoseSearchTrajectoryTypes.h b/Engine/Plugins/Experimental/Animation/PoseSearch/Source/Runtime/Public/PoseSearch/PoseSearchTrajectoryTypes.h
--- a/Engine/Plugins/Experimental/Animation/PoseSearch/Source/Runtime/Public/PoseSearch/PoseSearchTrajectoryTypes.h
+++ b/Engine/Plugins/Experimental/Animation/PoseSearch/Source/Runtime/Public/PoseSearch/PoseSearchTrajectoryTypes.h
@@ -25,6 +25,7 @@ struct POSESEARCH_API FPoseSearchQueryTrajectorySample
 	void SetTransform(const FTransform& Transform);
 	FTransform GetTransform() const { return FTransform(Facing, Position); }
 };
+POSESEARCH_API FArchive& operator<<(FArchive& Ar, FPoseSearchQueryTrajectorySample& TrajectorySample);
 
 USTRUCT(BlueprintType, Category = "Motion Trajectory")
 struct POSESEARCH_API FPoseSearchQueryTrajectory
@@ -41,3 +42,4 @@ struct POSESEARCH_API FPoseSearchQueryTrajectory
 	void DebugDrawTrajectory(const UWorld* World) const;
 #endif // ENABLE_ANIM_DEBUG
 };
+POSESEARCH_API FArchive& operator<<(FArchive& Ar, FPoseSearchQueryTrajectory& Trajectory);
