Fix shipping editor compile
#rnx

diff --git a/Engine/Plugins/Experimental/Animation/PoseSearch/Source/Editor/Private/PoseSearchDebuggerView.cpp b/Engine/Plugins/Experimental/Animation/PoseSearch/Source/Editor/Private/PoseSearchDebuggerView.cpp
--- a/Engine/Plugins/Experimental/Animation/PoseSearch/Source/Editor/Private/PoseSearchDebuggerView.cpp
+++ b/Engine/Plugins/Experimental/Animation/PoseSearch/Source/Editor/Private/PoseSearchDebuggerView.cpp
@@ -337,10 +337,12 @@ void SDebuggerView::DrawFeatures(
 ) const
 {
 	// Draw world space trajectory
+#if ENABLE_ANIM_DEBUG
 	if (ViewModel.Get()->GetDrawTrajectory())
 	{
 		State.Trajectory.DebugDrawTrajectory(&DebuggerWorld);
 	}
+#endif
 
 	// Draw query vector
 	if (ViewModel.Get()->GetDrawQuery())
