motion matching
- pose search debugger show all poses
- removed unnecessary FSearchResult::ContinuingPoseCost

diff --git a/Engine/Plugins/Experimental/Animation/PoseSearch/Source/Editor/Private/PoseSearchDebuggerDatabaseColumns.h b/Engine/Plugins/Experimental/Animation/PoseSearch/Source/Editor/Private/PoseSearchDebuggerDatabaseColumns.h
--- a/Engine/Plugins/Experimental/Animation/PoseSearch/Source/Editor/Private/PoseSearchDebuggerDatabaseColumns.h
+++ b/Engine/Plugins/Experimental/Animation/PoseSearch/Source/Editor/Private/PoseSearchDebuggerDatabaseColumns.h
@@ -575,6 +575,11 @@ struct FPoseCandidateFlags : ITextColumn
 				Sring.Append("F ");
 			}
 
+			if (EnumHasAnyFlags(Row->PoseCandidateFlags, EPoseCandidateFlags::DiscardedBy_Search))
+			{
+				Sring.Append("S ");
+			}
+
 			return FText::FromString(Sring);
 		}
 
@@ -608,6 +613,11 @@ struct FPoseCandidateFlags : ITextColumn
 				TextBuilder.AppendLine(LOCTEXT("DiscardedBy_PoseFilter_Tooltip", "(F) Filter"));
 			}
 
+			if (EnumHasAnyFlags(Row->PoseCandidateFlags, EPoseCandidateFlags::DiscardedBy_Search))
+			{
+				TextBuilder.AppendLine(LOCTEXT("DiscardedBy_Search_Tooltip", "(S) Search"));
+			}
+
 			return TextBuilder.ToText();
 		}
 
diff --git a/Engine/Plugins/Experimental/Animation/PoseSearch/Source/Editor/Private/PoseSearchDebuggerDatabaseView.cpp b/Engine/Plugins/Experimental/Animation/PoseSearch/Source/Editor/Private/PoseSearchDebuggerDatabaseView.cpp
--- a/Engine/Plugins/Experimental/Animation/PoseSearch/Source/Editor/Private/PoseSearchDebuggerDatabaseView.cpp
+++ b/Engine/Plugins/Experimental/Animation/PoseSearch/Source/Editor/Private/PoseSearchDebuggerDatabaseView.cpp
@@ -25,135 +25,195 @@
 namespace UE::PoseSearch
 {
 
-void SDebuggerDatabaseView::Update(const FTraceMotionMatchingStateMessage& State)
+class SCostBreakDownData
 {
-	class SCostBreakDownData
+public:
+	SCostBreakDownData(const TArray<FTraceMotionMatchingStateDatabaseEntry>& DatabaseEntries, bool bIsVerbose)
 	{
-	public:
-		SCostBreakDownData(const TArray<FTraceMotionMatchingStateDatabaseEntry>& DatabaseEntries, bool bIsVerbose)
+		// processing all the DatabaseEntries to collect the LabelToChannels
+		for (const FTraceMotionMatchingStateDatabaseEntry& DbEntry : DatabaseEntries)
 		{
-			// processing all the DatabaseEntries to collect the LabelToChannels
-			for (const FTraceMotionMatchingStateDatabaseEntry& DbEntry : DatabaseEntries)
+			const UPoseSearchDatabase* Database = FTraceMotionMatchingState::GetObjectFromId<UPoseSearchDatabase>(DbEntry.DatabaseId);
+			if (FAsyncPoseSearchDatabasesManagement::RequestAsyncBuildIndex(Database, ERequestAsyncBuildFlag::ContinueRequest))
 			{
-				const UPoseSearchDatabase* Database = FTraceMotionMatchingState::GetObjectFromId<UPoseSearchDatabase>(DbEntry.DatabaseId);
-				if (FAsyncPoseSearchDatabasesManagement::RequestAsyncBuildIndex(Database, ERequestAsyncBuildFlag::ContinueRequest))
+				for (const TObjectPtr<UPoseSearchFeatureChannel>& ChannelPtr : Database->Schema->GetChannels())
 				{
-					for (const TObjectPtr<UPoseSearchFeatureChannel>& ChannelPtr : Database->Schema->GetChannels())
-					{
-						AnalyzeChannelRecursively(ChannelPtr.Get(), bIsVerbose);
-					}
+					AnalyzeChannelRecursively(ChannelPtr.Get(), bIsVerbose);
 				}
 			}
 		}
+	}
 
-		void ProcessData(TArray<TSharedRef<FDebuggerDatabaseRowData>>& InOutUnfilteredDatabaseRows) const
+	void ProcessData(TArray<TSharedRef<FDebuggerDatabaseRowData>>& InOutUnfilteredDatabaseRows) const
+	{
+		for (TSharedRef<FDebuggerDatabaseRowData>& UnfilteredDatabaseRowRef : InOutUnfilteredDatabaseRows)
 		{
-			for (TSharedRef<FDebuggerDatabaseRowData>& UnfilteredDatabaseRowRef : InOutUnfilteredDatabaseRows)
+			FDebuggerDatabaseRowData& UnfilteredDatabaseRow = UnfilteredDatabaseRowRef.Get();
+			UnfilteredDatabaseRow.CostBreakdowns.AddDefaulted(LabelToChannels.Num());
+
+			for (int32 LabelToChannelIndex = 0; LabelToChannelIndex < LabelToChannels.Num(); ++LabelToChannelIndex)
 			{
-				FDebuggerDatabaseRowData& UnfilteredDatabaseRow = UnfilteredDatabaseRowRef.Get();
-				UnfilteredDatabaseRow.CostBreakdowns.AddDefaulted(LabelToChannels.Num());
+				const FLabelToChannels& LabelToChannel = LabelToChannels[LabelToChannelIndex];
 
-				for (int32 LabelToChannelIndex = 0; LabelToChannelIndex < LabelToChannels.Num(); ++LabelToChannelIndex)
+				// there should only be at most one channel per schema with the unique label,
+				// but we'll keep this generic allowing multiple channels from the same schema having the same label.
+				// the cost will be the sum of all the channels cost
+				float CostBreakdown = 0.f;
+				for (const UPoseSearchFeatureChannel* Channel : LabelToChannel.Channels)
 				{
-					const FLabelToChannels& LabelToChannel = LabelToChannels[LabelToChannelIndex];
-
-					// there should only be at most one channel per schema with the unique label,
-					// but we'll keep this generic allowing multiple channels from the same schema having the same label.
-					// the cost will be the sum of all the channels cost
-					float CostBreakdown = 0.f;
-					for (const UPoseSearchFeatureChannel* Channel : LabelToChannel.Channels)
+					// checking if the row is associated to the Channel
+					if (UnfilteredDatabaseRow.SharedData->SourceDatabase->Schema == Channel->GetSchema())
 					{
-						// checking if the row is associated to the Channel
-						if (UnfilteredDatabaseRow.SharedData->SourceDatabase->Schema == Channel->GetSchema())
-						{
-							CostBreakdown += ArraySum(UnfilteredDatabaseRow.CostVector, Channel->GetChannelDataOffset(), Channel->GetChannelCardinality());
-						}
+						CostBreakdown += ArraySum(UnfilteredDatabaseRow.CostVector, Channel->GetChannelDataOffset(), Channel->GetChannelCardinality());
 					}
-					UnfilteredDatabaseRow.CostBreakdowns[LabelToChannelIndex] = CostBreakdown;
 				}
+				UnfilteredDatabaseRow.CostBreakdowns[LabelToChannelIndex] = CostBreakdown;
 			}
 		}
+	}
 
-		bool AreLabelsEqualTo(const TArray<FText>& OtherLabels) const
+	bool AreLabelsEqualTo(const TArray<FText>& OtherLabels) const
+	{
+		if (LabelToChannels.Num() != OtherLabels.Num())
 		{
-			if (LabelToChannels.Num() != OtherLabels.Num())
+			return false;
+		}
+
+		for (int32 i = 0; i < LabelToChannels.Num(); ++i)
+		{
+			if (!LabelToChannels[i].Label.EqualTo(OtherLabels[i]))
 			{
 				return false;
 			}
+		}
 
-			for (int32 i = 0; i < LabelToChannels.Num(); ++i)
-			{
-				if (!LabelToChannels[i].Label.EqualTo(OtherLabels[i]))
-				{
-					return false;
-				}
-			}
+		return true;
+	}
 
-			return true;
+	const TArray<FText> GetLabels() const
+	{
+		TArray<FText> Labels;
+		for (const FLabelToChannels& LabelToChannel : LabelToChannels)
+		{
+			Labels.Add(LabelToChannel.Label);
 		}
+		return Labels;
+	}
 
-		const TArray<FText> GetLabels() const
+private:
+	void AnalyzeChannelRecursively(const UPoseSearchFeatureChannel* Channel, bool bIsVerbose)
+	{
+		const FText Label = FText::FromString(Channel->GetLabel());
+
+		bool bLabelFound = false;
+		for (int32 i = 0; i < LabelToChannels.Num(); ++i)
 		{
-			TArray<FText> Labels;
-			for (const FLabelToChannels& LabelToChannel : LabelToChannels)
+			if (LabelToChannels[i].Label.EqualTo(Label))
 			{
-				Labels.Add(LabelToChannel.Label);
+				LabelToChannels[i].Channels.AddUnique(Channel);
+				bLabelFound = true;
 			}
-			return Labels;
 		}
-
-	private:
-		void AnalyzeChannelRecursively(const UPoseSearchFeatureChannel* Channel, bool bIsVerbose)
+		if (!bLabelFound)
 		{
-			const FText Label = FText::FromString(Channel->GetLabel());
+			LabelToChannels.AddDefaulted();
+			LabelToChannels.Last().Label = Label;
+			LabelToChannels.Last().Channels.Add(Channel);
+		}
 
-			bool bLabelFound = false;
-			for (int32 i = 0; i < LabelToChannels.Num(); ++i)
+		if (bIsVerbose)
+		{
+			for (const TObjectPtr<UPoseSearchFeatureChannel>& SubChannelPtr : Channel->GetSubChannels())
 			{
-				if (LabelToChannels[i].Label.EqualTo(Label))
+				if (const UPoseSearchFeatureChannel* SubChannel = SubChannelPtr.Get())
 				{
-					LabelToChannels[i].Channels.AddUnique(Channel);
-					bLabelFound = true;
+					AnalyzeChannelRecursively(SubChannel, bIsVerbose);
 				}
 			}
-			if (!bLabelFound)
-			{
-				LabelToChannels.AddDefaulted();
-				LabelToChannels.Last().Label = Label;
-				LabelToChannels.Last().Channels.Add(Channel);
-			}
+		}
+	}
 
-			if (bIsVerbose)
-			{
-				for (const TObjectPtr<UPoseSearchFeatureChannel>& SubChannelPtr : Channel->GetSubChannels())
-				{
-					if (const UPoseSearchFeatureChannel* SubChannel = SubChannelPtr.Get())
-					{
-						AnalyzeChannelRecursively(SubChannel, bIsVerbose);
-					}
-				}
-			}
+	static float ArraySum(TConstArrayView<float> View, int32 StartIndex, int32 Offset)
+	{
+		float Sum = 0.f;
+		const int32 EndIndex = StartIndex + Offset;
+		for (int i = StartIndex; i < EndIndex; ++i)
+		{
+			Sum += View[i];
 		}
+		return Sum;
+	}
+
+	struct FLabelToChannels
+	{
+		FText Label;
+		TArray<const UPoseSearchFeatureChannel*> Channels; // NoTe: channels can be from different schemas
+	};
+	TArray<FLabelToChannels> LabelToChannels;
+};
 
-		static float ArraySum(TConstArrayView<float> View, int32 StartIndex, int32 Offset)
+static void AddUnfilteredDatabaseRow(const UPoseSearchDatabase* Database, 
+	TArray<TSharedRef<FDebuggerDatabaseRowData>>& UnfilteredDatabaseRows, TSharedRef<FDebuggerDatabaseSharedData> SharedData,
+	int32 DbPoseIdx, EPoseCandidateFlags PoseCandidateFlags, const FPoseSearchCost& Cost = FPoseSearchCost())
+{
+	const FPoseSearchIndex& SearchIndex = Database->GetSearchIndex();
+	if (const FPoseSearchIndexAsset* SearchIndexAsset = SearchIndex.GetAssetForPoseSafe(DbPoseIdx))
+	{
+		TSharedRef<FDebuggerDatabaseRowData>& Row = UnfilteredDatabaseRows.Add_GetRef(MakeShared<FDebuggerDatabaseRowData>(SharedData));
+
+		const float Time = Database->GetAssetTime(DbPoseIdx);
+
+		Row->PoseIdx = DbPoseIdx;
+		Row->PoseCandidateFlags = PoseCandidateFlags;
+		Row->DbAssetIdx = SearchIndexAsset->SourceAssetIdx;
+		Row->AssetTime = Time;
+		Row->bMirrored = SearchIndexAsset->bMirrored;
+
+		Row->CostVector.SetNum(Database->Schema->SchemaCardinality);
+		const TArray<float> PoseValues = SearchIndex.GetPoseValuesSafe(DbPoseIdx);
+
+		// in case we modify the schema while PIE is paused and displaying the Pose Search Editor, we could end up with a stale State with a SharedData->QueryVector saved with the previous schema
+		// so the cardinality of SharedData->QueryVector and PoseValues don't match. In that case we just use PoseValues as query to have all costs set to zero
+		const bool bIsQueryVectorValid = SharedData->QueryVector.Num() == PoseValues.Num();
+		const TArray<float>& QueryVector = bIsQueryVectorValid ? SharedData->QueryVector : PoseValues;
+
+		CompareFeatureVectors(PoseValues, QueryVector, SearchIndex.WeightsSqrt, Row->CostVector);
+
+		if (Cost.IsValid())
 		{
-			float Sum = 0.f;
-			const int32 EndIndex = StartIndex + Offset;
-			for (int i = StartIndex; i < EndIndex; ++i)
-			{
-				Sum += View[i];
-			}
-			return Sum;
+			Row->PoseCost = Cost;
+		}
+		else
+		{
+			// @todo: perhaps reuse CompareFeatureVectors cost calculation
+			// @todo: initialize EPoseSearchBooleanRequest::Indifferent properly
+			Row->PoseCost = SearchIndex.ComparePoses(DbPoseIdx, EPoseSearchBooleanRequest::Indifferent, 0.f, Database->Schema->MirrorMismatchCostBias, PoseValues, QueryVector);
 		}
 
-		struct FLabelToChannels
+		const FInstancedStruct& DatabaseAssetStruct = Database->GetAnimationAssetStruct(*SearchIndexAsset);
+		if (const FPoseSearchDatabaseAnimationAssetBase* DatabaseAsset = DatabaseAssetStruct.GetPtr<FPoseSearchDatabaseAnimationAssetBase>())
 		{
-			FText Label;
-			TArray<const UPoseSearchFeatureChannel*> Channels; // NoTe: channels can be from different schemas
-		};
-		TArray<FLabelToChannels> LabelToChannels;
-	};
+			Row->AssetName = DatabaseAsset->GetName();
+			Row->AssetPath = DatabaseAsset->GetAnimationAsset() ? DatabaseAsset->GetAnimationAsset()->GetPathName() : "";
+			Row->bLooping = DatabaseAsset->IsLooping();
+			Row->BlendParameters = SearchIndexAsset->BlendParameters;
+			Row->AnimFrame = 0;
+			Row->AnimPercentage = 0.0f;
+
+			if (const FPoseSearchDatabaseAnimationAssetBase* DatabaseAnimationAssetBase = DatabaseAssetStruct.GetPtr<FPoseSearchDatabaseAnimationAssetBase>())
+			{
+				if (const UAnimSequenceBase* SequenceBase = Cast<UAnimSequenceBase>(DatabaseAnimationAssetBase->GetAnimationAsset()))
+				{
+					Row->AnimFrame = SequenceBase->GetFrameAtTime(Time);
+					Row->AnimPercentage = Time / SequenceBase->GetPlayLength();
+				}
+			}
+		}
+	}
+}
 
+void SDebuggerDatabaseView::Update(const FTraceMotionMatchingStateMessage& State)
+{
 	// row cost color palette
 	static const FLinearColor DiscardedRowColor(0.314f, 0.314f, 0.314f); // darker gray
 	static const FLinearColor BestScoreRowColor = FLinearColor::Green;
@@ -171,13 +231,12 @@ void SDebuggerDatabaseView::Update(const FTraceMotionMatchingStateMessage& State
 	}
 
 	UnfilteredDatabaseRows.Reset();
+
 	for (const FTraceMotionMatchingStateDatabaseEntry& DbEntry : State.DatabaseEntries)
 	{
 		const UPoseSearchDatabase* Database = FTraceMotionMatchingState::GetObjectFromId<UPoseSearchDatabase>(DbEntry.DatabaseId);
 		if (FAsyncPoseSearchDatabasesManagement::RequestAsyncBuildIndex(Database, ERequestAsyncBuildFlag::ContinueRequest))
 		{
-			const FPoseSearchIndex& SearchIndex = Database->GetSearchIndex();
-
 			TSharedRef<FDebuggerDatabaseSharedData> SharedData = MakeShared<FDebuggerDatabaseSharedData>();
 			SharedData->SourceDatabase = Database;
 			SharedData->DatabaseName = Database->GetName();
@@ -186,45 +245,23 @@ void SDebuggerDatabaseView::Update(const FTraceMotionMatchingStateMessage& State
 
 			for (const FTraceMotionMatchingStatePoseEntry& PoseEntry : DbEntry.PoseEntries)
 			{
-				if (const FPoseSearchIndexAsset* SearchIndexAsset = SearchIndex.GetAssetForPoseSafe(PoseEntry.DbPoseIdx))
-				{
-					TSharedRef<FDebuggerDatabaseRowData>& Row = UnfilteredDatabaseRows.Add_GetRef(MakeShared<FDebuggerDatabaseRowData>(SharedData));
-
-					const float Time = Database->GetAssetTime(PoseEntry.DbPoseIdx);
-
-					Row->PoseIdx = PoseEntry.DbPoseIdx;
-					Row->PoseCandidateFlags = PoseEntry.PoseCandidateFlags;
-					Row->DbAssetIdx = SearchIndexAsset->SourceAssetIdx;
-					Row->AssetTime = Time;
-					Row->bMirrored = SearchIndexAsset->bMirrored;
-					Row->PoseCost = PoseEntry.Cost;
-
-					Row->CostVector.SetNum(Database->Schema->SchemaCardinality);
-					const TArray<float> PoseValues = SearchIndex.GetPoseValuesSafe(PoseEntry.DbPoseIdx);
+				AddUnfilteredDatabaseRow(Database, UnfilteredDatabaseRows, SharedData, PoseEntry.DbPoseIdx, PoseEntry.PoseCandidateFlags, PoseEntry.Cost);
+			}
 
-					// in case we modify the schema while PIE is paused and displaying the Pose Search Editor, we could end up with a stale State with a DbEntry.QueryVector saved with the previous schema
-					// so the cardinality of DbEntry.QueryVector and PoseValues don't match. In that case we just use PoseValues as query to have all costs set to zero
-					const bool bIsQueryVectorValid = DbEntry.QueryVector.Num() == PoseValues.Num();
-					CompareFeatureVectors(PoseValues, bIsQueryVectorValid ? DbEntry.QueryVector : PoseValues, SearchIndex.WeightsSqrt, Row->CostVector);
+			if (bShowAllPoses)
+			{
+				TSet<int32> PoseEntriesIdx;
+				for (const FTraceMotionMatchingStatePoseEntry& PoseEntry : DbEntry.PoseEntries)
+				{
+					PoseEntriesIdx.Add(PoseEntry.DbPoseIdx);
+				}
 
-					const FInstancedStruct& DatabaseAssetStruct = Database->GetAnimationAssetStruct(*SearchIndexAsset);
-					if (const FPoseSearchDatabaseAnimationAssetBase* DatabaseAsset = DatabaseAssetStruct.GetPtr<FPoseSearchDatabaseAnimationAssetBase>())
+				const FPoseSearchIndex& SearchIndex = Database->GetSearchIndex();
+				for (int32 DbPoseIdx = 0; DbPoseIdx < SearchIndex.GetNumPoses(); ++DbPoseIdx)
+				{
+					if (!PoseEntriesIdx.Find(DbPoseIdx))
 					{
-						Row->AssetName = DatabaseAsset->GetName();
-						Row->AssetPath = DatabaseAsset->GetAnimationAsset() ? DatabaseAsset->GetAnimationAsset()->GetPathName() : "";
-						Row->bLooping = DatabaseAsset->IsLooping();
-						Row->BlendParameters = SearchIndexAsset->BlendParameters;
-						Row->AnimFrame = 0;
-						Row->AnimPercentage = 0.0f;
-
-						if (const FPoseSearchDatabaseAnimationAssetBase* DatabaseAnimationAssetBase = DatabaseAssetStruct.GetPtr<FPoseSearchDatabaseAnimationAssetBase>())
-						{
-							if (const UAnimSequenceBase* SequenceBase = Cast<UAnimSequenceBase>(DatabaseAnimationAssetBase->GetAnimationAsset()))
-							{
-								Row->AnimFrame = SequenceBase->GetFrameAtTime(Time);
-								Row->AnimPercentage = Time / SequenceBase->GetPlayLength();
-							}
-						}
+						AddUnfilteredDatabaseRow(Database, UnfilteredDatabaseRows, SharedData, DbPoseIdx, EPoseCandidateFlags::DiscardedBy_Search);
 					}
 				}
 			}
@@ -261,7 +298,10 @@ void SDebuggerDatabaseView::Update(const FTraceMotionMatchingStateMessage& State
 
 		for (TSharedRef<FDebuggerDatabaseRowData>& UnfilteredRow : UnfilteredDatabaseRows)
 		{
-			ArrayMinMax(UnfilteredRow->CostBreakdowns, MinCostBreakdowns, MaxCostBreakdowns, UE_MAX_FLT);
+			if (EnumHasAnyFlags(UnfilteredRow->PoseCandidateFlags, EPoseCandidateFlags::AnyValidMask))
+			{
+				ArrayMinMax(UnfilteredRow->CostBreakdowns, MinCostBreakdowns, MaxCostBreakdowns, UE_MAX_FLT);
+			}
 		}
 		
 		auto ArraySafeNormalize = [](TConstArrayView<float> View, TConstArrayView<float> Min, TConstArrayView<float> Max, TArrayView<float> NormalizedView)
@@ -316,9 +356,12 @@ void SDebuggerDatabaseView::Update(const FTraceMotionMatchingStateMessage& State
 		float MaxCost = -UE_MAX_FLT;
 		for (TSharedRef<FDebuggerDatabaseRowData>& UnfilteredRow : UnfilteredDatabaseRows)
 		{
-			const float Cost = UnfilteredRow->PoseCost.GetTotalCost();
-			MinCost = FMath::Min(MinCost, Cost);
-			MaxCost = FMath::Max(MaxCost, Cost);
+			if (EnumHasAnyFlags(UnfilteredRow->PoseCandidateFlags, EPoseCandidateFlags::AnyValidMask))
+			{
+				const float Cost = UnfilteredRow->PoseCost.GetTotalCost();
+				MinCost = FMath::Min(MinCost, Cost);
+				MaxCost = FMath::Max(MaxCost, Cost);
+			}
 		}
 
 		const float DeltaCost = MaxCost - MinCost;
@@ -454,6 +497,30 @@ void SDebuggerDatabaseView::OnFilterTextChanged(const FText& SearchText)
 	PopulateViewRows();
 }
 
+void SDebuggerDatabaseView::OnShowAllPosesCheckboxChanged(ECheckBoxState State)
+{
+	if (State == ECheckBoxState::Checked)
+	{
+		bShowAllPoses = true;
+	}
+	else
+	{
+		bShowAllPoses = false;
+	}
+
+	if (TSharedPtr<SDebuggerView> DebuggerView = ParentDebuggerViewPtr.Pin())
+	{
+		if (TSharedPtr<FDebuggerViewModel> ViewModel = DebuggerView->GetViewModel())
+		{
+			const FTraceMotionMatchingStateMessage* MotionMatchingState = ViewModel.Get()->GetMotionMatchingState();
+			if (MotionMatchingState)
+			{
+				Update(*MotionMatchingState);
+			}
+		}
+	}
+}
+
 void SDebuggerDatabaseView::OnHideInvalidPosesCheckboxChanged(ECheckBoxState State)
 {
 	if (State == ECheckBoxState::Checked) 
@@ -843,6 +910,20 @@ void SDebuggerDatabaseView::Construct(const FArguments& InArgs)
 					+ SHorizontalBox::Slot()
 					.AutoWidth()
 					.Padding(10, 5, 10, 5)
+					[
+						SNew(SCheckBox)
+						.OnCheckStateChanged_Lambda([this](ECheckBoxState State)
+						{
+							   SDebuggerDatabaseView::OnShowAllPosesCheckboxChanged(State);
+						})
+						[
+							SNew(STextBlock)
+							.Text(LOCTEXT("PoseSearchDebuggerShowAllPosesFlag", "Show All Poses"))
+						]
+					]
+					+ SHorizontalBox::Slot()
+					.AutoWidth()
+					.Padding(10, 5, 10, 5)
 					[
 						SNew(SCheckBox)
 						.OnCheckStateChanged_Lambda([this](ECheckBoxState State)
diff --git a/Engine/Plugins/Experimental/Animation/PoseSearch/Source/Editor/Private/PoseSearchDebuggerDatabaseView.h b/Engine/Plugins/Experimental/Animation/PoseSearch/Source/Editor/Private/PoseSearchDebuggerDatabaseView.h
--- a/Engine/Plugins/Experimental/Animation/PoseSearch/Source/Editor/Private/PoseSearchDebuggerDatabaseView.h
+++ b/Engine/Plugins/Experimental/Animation/PoseSearch/Source/Editor/Private/PoseSearchDebuggerDatabaseView.h
@@ -66,6 +66,9 @@ private:
 	/** Called when the text in the filter box is modified to update the filtering */
 	void OnFilterTextChanged(const FText& SearchText);
 
+	/** Called when the ShowAllPoses Checkbox is Toggled */
+	void OnShowAllPosesCheckboxChanged(ECheckBoxState State);
+	
 	/** Called when the HideInvalidPoses Checkbox is Toggled */
 	void OnHideInvalidPosesCheckboxChanged(ECheckBoxState State);
 
@@ -139,6 +142,9 @@ private:
 	/** Text used to filter DatabaseView */
 	FText FilterText;
 
+	/* Bool used to show all poses from display */
+	bool bShowAllPoses = false;
+	
 	/* Bool used to hide invalid poses from display */
 	bool bHideInvalidPoses = false;
 
diff --git a/Engine/Plugins/Experimental/Animation/PoseSearch/Source/Editor/Private/PoseSearchDebuggerView.cpp b/Engine/Plugins/Experimental/Animation/PoseSearch/Source/Editor/Private/PoseSearchDebuggerView.cpp
--- a/Engine/Plugins/Experimental/Animation/PoseSearch/Source/Editor/Private/PoseSearchDebuggerView.cpp
+++ b/Engine/Plugins/Experimental/Animation/PoseSearch/Source/Editor/Private/PoseSearchDebuggerView.cpp
@@ -92,14 +92,7 @@ void SDebuggerDetailsView::UpdateReflection(const FTraceMotionMatchingStateMessa
 {
 	check(Reflection);
 
-	const UPoseSearchDatabase* CurrentDatabase = State.GetCurrentDatabase();
-	if (FAsyncPoseSearchDatabasesManagement::RequestAsyncBuildIndex(CurrentDatabase, ERequestAsyncBuildFlag::ContinueRequest))
-	{
-		const FPoseSearchIndex& CurrentSearchIndex = CurrentDatabase->GetSearchIndex();
-		int32 CurrentDbPoseIdx = State.GetCurrentDatabasePoseIndex();
-
 		Reflection->ElapsedPoseSearchTime = State.ElapsedPoseSearchTime;
-
 		Reflection->AssetPlayerTime = State.AssetPlayerTime;
 		Reflection->LastDeltaTime = State.DeltaTime;
 		Reflection->SimLinearVelocity = State.SimLinearVelocity;
@@ -107,7 +100,6 @@ void SDebuggerDetailsView::UpdateReflection(const FTraceMotionMatchingStateMessa
 		Reflection->AnimLinearVelocity = State.AnimLinearVelocity;
 		Reflection->AnimAngularVelocity = State.AnimAngularVelocity;
 	}
-}
 
 void SDebuggerView::Construct(const FArguments& InArgs, uint64 InAnimInstanceId)
 {
@@ -200,8 +192,6 @@ void SDebuggerView::Tick(const FGeometry& AllottedGeometry, const double InCurre
 
 	TSharedPtr<FDebuggerViewModel> Model = ViewModel.Get();
 
-	bool bNeedUpdate = Model->HasSearchableAssetChanged();
-
 	// We haven't reached the update point yet
 	if (CurrentConsecutiveFrames < ConsecutiveFramesUpdateThreshold)
 	{
@@ -222,12 +212,6 @@ void SDebuggerView::Tick(const FGeometry& AllottedGeometry, const double InCurre
 		// Haven't updated since passing through frame gate, update once
 		else if (!bUpdated)
 		{
-			bNeedUpdate = true;
-		}
-	}
-
-	if (bNeedUpdate)
-	{
 		Model->OnUpdate();
 		if (UpdateNodeSelection())
 		{
@@ -236,6 +220,7 @@ void SDebuggerView::Tick(const FGeometry& AllottedGeometry, const double InCurre
 		}
 		bUpdated = true;
 	}
+	}
 
 	// Draw visualization every tick
 	DrawVisualization();
diff --git a/Engine/Plugins/Experimental/Animation/PoseSearch/Source/Editor/Private/PoseSearchDebuggerViewModel.cpp b/Engine/Plugins/Experimental/Animation/PoseSearch/Source/Editor/Private/PoseSearchDebuggerViewModel.cpp
--- a/Engine/Plugins/Experimental/Animation/PoseSearch/Source/Editor/Private/PoseSearchDebuggerViewModel.cpp
+++ b/Engine/Plugins/Experimental/Animation/PoseSearch/Source/Editor/Private/PoseSearchDebuggerViewModel.cpp
@@ -92,12 +92,6 @@ const FTransform& FDebuggerViewModel::GetRootBoneTransform() const
 	return RootBoneWorldTransform;
 }
 
-bool FDebuggerViewModel::HasSearchableAssetChanged() const
-{
-	uint64 NewSearchableAssetId = ActiveMotionMatchingState ? ActiveMotionMatchingState->SearchableAssetId : 0;
-	return NewSearchableAssetId != SearchableAssetId;
-}
-
 void FDebuggerViewModel::OnUpdate()
 {
 	if (!bSkeletonsInitialized)
@@ -157,13 +151,6 @@ void FDebuggerViewModel::OnUpdateNodeSelection(int32 InNodeId)
 			}
 		}
 	}
-
-	uint64 NewSearchableAssetId = ActiveMotionMatchingState ? ActiveMotionMatchingState->SearchableAssetId : 0;
-	if (NewSearchableAssetId != SearchableAssetId)
-	{
-		ClearSelectedSkeleton();
-		SearchableAssetId = NewSearchableAssetId;
-	}
 }
 
 void FDebuggerViewModel::UpdatePoseSearchContext(UPoseSearchMeshComponent::FUpdateContext& InOutContext, const FSkeleton& Skeleton) const
diff --git a/Engine/Plugins/Experimental/Animation/PoseSearch/Source/Editor/Private/PoseSearchDebuggerViewModel.h b/Engine/Plugins/Experimental/Animation/PoseSearch/Source/Editor/Private/PoseSearchDebuggerViewModel.h
--- a/Engine/Plugins/Experimental/Animation/PoseSearch/Source/Editor/Private/PoseSearchDebuggerViewModel.h
+++ b/Engine/Plugins/Experimental/Animation/PoseSearch/Source/Editor/Private/PoseSearchDebuggerViewModel.h
@@ -30,9 +30,6 @@ public:
 	int32 GetNodesNum() const;
 	const FTransform& GetRootBoneTransform() const;
 
-	/** Checks if Update must be called */
-	bool HasSearchableAssetChanged() const;
-
 	/** Update motion matching states for frame */
 	void OnUpdate();
 	
@@ -46,18 +43,15 @@ public:
 	void ClearSelectedSkeleton();
 	
 	void SetVerbose(bool bVerbose) { bIsVerbose = bVerbose; }
-
 	bool IsVerbose() const { return bIsVerbose; }
 
 	void SetDrawQuery(bool bInDrawQuery) { bDrawQuery = bInDrawQuery; }
-
 	bool GetDrawQuery() const { return bDrawQuery; }
 
 	void SetDrawTrajectory(bool bInDrawTrajectory) { bDrawTrajectory = bInDrawTrajectory; }
 
 	bool GetDrawTrajectory() const { return bDrawTrajectory; }
 
-
 	/** Callback to reset debug skeletons for the active world */
 	void OnWorldCleanup(UWorld* InWorld, bool bSessionEnded, bool bCleanupResources);
 
@@ -108,9 +102,6 @@ private:
 	/** Currently active motion matching state based on node selection in the view */
 	const FTraceMotionMatchingStateMessage* ActiveMotionMatchingState = nullptr;
 
-	/** Active motion matching state's searchable asset */
-	uint64 SearchableAssetId = 0;
-
 	/** Current Skeletal Mesh Component Id for the AnimInstance */
 	uint64 SkeletalMeshComponentId = 0;
 
diff --git a/Engine/Plugins/Experimental/Animation/PoseSearch/Source/Runtime/Private/PoseSearchDatabase.cpp b/Engine/Plugins/Experimental/Animation/PoseSearch/Source/Runtime/Private/PoseSearchDatabase.cpp
--- a/Engine/Plugins/Experimental/Animation/PoseSearch/Source/Runtime/Private/PoseSearchDatabase.cpp
+++ b/Engine/Plugins/Experimental/Animation/PoseSearch/Source/Runtime/Private/PoseSearchDatabase.cpp
@@ -789,8 +789,13 @@ FPoseSearchCost UPoseSearchDatabase::SearchContinuingPose(UE::PoseSearch::FSearc
 		TArrayView<float> ReconstructedPoseValuesBuffer((float*)FMemory_Alloca(NumDimensions * sizeof(float)), NumDimensions);
 		const TConstArrayView<float> PoseValues = SearchIndex.Values.IsEmpty() ? SearchIndex.GetReconstructedPoseValues(PoseIdx, ReconstructedPoseValuesBuffer) : SearchIndex.GetPoseValues(PoseIdx);
 
-		ContinuingPoseCost = SearchIndex.ComparePoses(SearchContext.GetCurrentResult().PoseIdx, SearchContext.GetQueryMirrorRequest(),
+		const int32 ContinuingPoseIdx = SearchContext.GetCurrentResult().PoseIdx;
+		ContinuingPoseCost = SearchIndex.ComparePoses(ContinuingPoseIdx, SearchContext.GetQueryMirrorRequest(),
 			ContinuingPoseCostBias, Schema->MirrorMismatchCostBias, PoseValues, SearchContext.GetOrBuildQuery(Schema).GetValues());
+
+#if UE_POSE_SEARCH_TRACE_ENABLED
+		SearchContext.BestCandidates.Add(ContinuingPoseCost, ContinuingPoseIdx, this, EPoseCandidateFlags::Valid_ContinuingPose);
+#endif
 	}
 
 	return ContinuingPoseCost;
diff --git a/Engine/Plugins/Experimental/Animation/PoseSearch/Source/Runtime/Private/PoseSearchLibrary.cpp b/Engine/Plugins/Experimental/Animation/PoseSearch/Source/Runtime/Private/PoseSearchLibrary.cpp
--- a/Engine/Plugins/Experimental/Animation/PoseSearch/Source/Runtime/Private/PoseSearchLibrary.cpp
+++ b/Engine/Plugins/Experimental/Animation/PoseSearch/Source/Runtime/Private/PoseSearchLibrary.cpp
@@ -191,7 +191,6 @@ void FMotionMatchingState::UpdateRootBoneControl(const FAnimInstanceProxy* AnimI
 
 #if UE_POSE_SEARCH_TRACE_ENABLED
 void UPoseSearchLibrary::TraceMotionMatchingState(
-	const UPoseSearchDatabase* Database,
 	const FPoseSearchQueryTrajectory& Trajectory,
 	UE::PoseSearch::FSearchContext& SearchContext,
 	const UE::PoseSearch::FSearchResult& CurrentResult,
@@ -232,6 +231,7 @@ void UPoseSearchLibrary::TraceMotionMatchingState(
 		return DbEntryIdx;
 	};
 
+	const int32 CurrentPoseIdx = bSearch && CurrentResult.PoseCost.IsValid() ? CurrentResult.PoseIdx : -1;
 	FTraceMotionMatchingState TraceState;
 	while (!SearchContext.BestCandidates.IsEmpty())
 	{
@@ -245,37 +245,17 @@ void UPoseSearchLibrary::TraceMotionMatchingState(
 		PoseEntry.DbPoseIdx = PoseCandidate.PoseIdx;
 		PoseEntry.Cost = PoseCandidate.Cost;
 		PoseEntry.PoseCandidateFlags = PoseCandidate.PoseCandidateFlags;
-
-		DbEntry.PoseEntries.Add(PoseEntry);
-	}
-
-	if (bSearch && CurrentResult.ContinuingPoseCost.IsValid())
-	{
-		check(LastResult.IsValid());
-
-		const int32 DbEntryIdx = AddUniqueDatabase(TraceState.DatabaseEntries, LastResult.Database.Get(), SearchContext);
-		FTraceMotionMatchingStateDatabaseEntry& DbEntry = TraceState.DatabaseEntries[DbEntryIdx];
-
-		const int32 LastResultPoseEntryIdx = DbEntry.PoseEntries.Add({ LastResult.PoseIdx });
-		FTraceMotionMatchingStatePoseEntry& PoseEntry = DbEntry.PoseEntries[LastResultPoseEntryIdx];
-
-		PoseEntry.Cost = CurrentResult.ContinuingPoseCost;
-		PoseEntry.PoseCandidateFlags = EPoseCandidateFlags::Valid_ContinuingPose;
-	}
-
-	if (bSearch && CurrentResult.PoseCost.IsValid())
-	{
-		const int32 DbEntryIdx = AddUniqueDatabase(TraceState.DatabaseEntries, CurrentResult.Database.Get(), SearchContext);
-		FTraceMotionMatchingStateDatabaseEntry& DbEntry = TraceState.DatabaseEntries[DbEntryIdx];
-
-		const int32 PoseEntryIdx = DbEntry.PoseEntries.Add({ CurrentResult.PoseIdx });
-		FTraceMotionMatchingStatePoseEntry& PoseEntry = DbEntry.PoseEntries[PoseEntryIdx];
-
-		PoseEntry.Cost = CurrentResult.PoseCost;
-		PoseEntry.PoseCandidateFlags = EPoseCandidateFlags::Valid_CurrentPose;
-
-		TraceState.CurrentDbEntryIdx = DbEntryIdx;
-		TraceState.CurrentPoseEntryIdx = PoseEntryIdx;
+		if (CurrentPoseIdx == PoseCandidate.PoseIdx && CurrentResult.Database.Get() == PoseCandidate.Database)
+		{
+			EnumAddFlags(PoseEntry.PoseCandidateFlags, EPoseCandidateFlags::Valid_CurrentPose);
+			
+			TraceState.CurrentDbEntryIdx = DbEntryIdx;
+			TraceState.CurrentPoseEntryIdx = DbEntry.PoseEntries.Add(PoseEntry);
+		}
+		else
+		{
+			DbEntry.PoseEntries.Add(PoseEntry);
+		}
 	}
 
 	if (DeltaTime > SMALL_NUMBER && SearchContext.IsTrajectoryValid())
@@ -293,7 +273,6 @@ void UPoseSearchLibrary::TraceMotionMatchingState(
 		TraceState.AnimAngularVelocity = FMath::RadiansToDegrees(RootMotionTransformDelta.GetRotation().GetAngle()) / DeltaTime;
 	}
 
-	TraceState.SearchableAssetId = FTraceMotionMatchingState::GetIdFromObject(Database);
 	TraceState.ElapsedPoseSearchTime = ElapsedPoseSearchTime;
 	TraceState.AssetPlayerTime = CurrentResult.AssetTime;
 	TraceState.DeltaTime = DeltaTime;
@@ -373,14 +352,14 @@ void UPoseSearchLibrary::UpdateMotionMatchingState(
 		InOutMotionMatchingState.ElapsedPoseSearchTime = 0.f;
 
 		// Evaluate continuing pose
-		FPoseSearchCost ContinuingPoseCost;
+		FSearchResult SearchResult;
 		if (!SearchContext.IsForceInterrupt() && SearchContext.CanAdvance() && SearchContext.GetCurrentResult().Database.IsValid())
 		{
-			ContinuingPoseCost = SearchContext.GetCurrentResult().Database->SearchContinuingPose(SearchContext);
-			SearchContext.UpdateCurrentBestCost(ContinuingPoseCost);
+			SearchResult.PoseCost = SearchContext.GetCurrentResult().Database->SearchContinuingPose(SearchContext);
+			SearchContext.UpdateCurrentBestCost(SearchResult.PoseCost);
 		}
 
-		FSearchResult SearchResult;
+		bool bJumpToPose = false;
 		for (TObjectPtr<const UPoseSearchDatabase> Database : Databases)
 		{
 			if (ensure(Database))
@@ -388,26 +367,31 @@ void UPoseSearchLibrary::UpdateMotionMatchingState(
 				FSearchResult NewSearchResult = Database->Search(SearchContext);
 				if (NewSearchResult.PoseCost.GetTotalCost() < SearchResult.PoseCost.GetTotalCost())
 				{
+					bJumpToPose = true;
 					SearchResult = NewSearchResult;
 					SearchContext.UpdateCurrentBestCost(SearchResult.PoseCost);
 				}
 			}
+		}
 
+#if WITH_EDITORONLY_DATA
+		if (!SearchResult.BruteForcePoseCost.IsValid())
+		{
+			SearchResult.BruteForcePoseCost = SearchResult.PoseCost;
 		}
+#endif // WITH_EDITORONLY_DATA
 
-		if (SearchResult.PoseCost.GetTotalCost() < ContinuingPoseCost.GetTotalCost())
+		if (bJumpToPose)
 		{
-			SearchResult.ContinuingPoseCost = ContinuingPoseCost;
 			InOutMotionMatchingState.JumpToPose(Context, SearchResult, MaxActiveBlends, BlendTime);
 		}
 		else
 		{
 			// copying few properties of SearchResult into CurrentSearchResult to facilitate debug drawing
-#if WITH_EDITOR
-			InOutMotionMatchingState.CurrentSearchResult.BruteForcePoseCost = ContinuingPoseCost;
-#endif
-			InOutMotionMatchingState.CurrentSearchResult.PoseCost = ContinuingPoseCost;
-			InOutMotionMatchingState.CurrentSearchResult.ContinuingPoseCost = ContinuingPoseCost;
+#if WITH_EDITORONLY_DATA
+			InOutMotionMatchingState.CurrentSearchResult.BruteForcePoseCost = SearchResult.BruteForcePoseCost;
+#endif // WITH_EDITORONLY_DATA
+			InOutMotionMatchingState.CurrentSearchResult.PoseCost = SearchResult.PoseCost;
 		}
 	}
 	else
@@ -428,14 +412,10 @@ void UPoseSearchLibrary::UpdateMotionMatchingState(
 		const float SearchBestCost = InOutMotionMatchingState.CurrentSearchResult.PoseCost.GetTotalCost();
 		float SearchBruteForceCost = SearchBestCost;
 #if WITH_EDITORONLY_DATA
-		if (Databases[0]->PoseSearchMode == EPoseSearchMode::PCAKDTree_Compare)
-		{
-			SearchBruteForceCost = InOutMotionMatchingState.CurrentSearchResult.BruteForcePoseCost.GetTotalCost();
-		}
+		SearchBruteForceCost = InOutMotionMatchingState.CurrentSearchResult.BruteForcePoseCost.GetTotalCost();
 #endif // WITH_EDITORONLY_DATA
 
-		// @todo: Update this to account for multiple databases.
-		TraceMotionMatchingState(Databases[0], Trajectory, SearchContext, InOutMotionMatchingState.CurrentSearchResult, LastResult, InOutMotionMatchingState.ElapsedPoseSearchTime,
+		TraceMotionMatchingState(Trajectory, SearchContext, InOutMotionMatchingState.CurrentSearchResult, LastResult, InOutMotionMatchingState.ElapsedPoseSearchTime,
 			InOutMotionMatchingState.RootMotionTransformDelta, Context.AnimInstanceProxy->GetAnimInstanceObject(), Context.GetCurrentNodeId(), DeltaTime, bSearch,
 			AnimInstance ? FObjectTrace::GetWorldElapsedTime(AnimInstance->GetWorld()) : 0.f, SearchBestCost, SearchBruteForceCost);
 	}
@@ -669,12 +649,9 @@ void UPoseSearchLibrary::MotionMatch(
 		const float SearchBestCost = SearchResult.PoseCost.GetTotalCost();
 		float SearchBruteForceCost = SearchBestCost;
 #if WITH_EDITORONLY_DATA
-		if (Database->PoseSearchMode == EPoseSearchMode::PCAKDTree_Compare)
-		{
-			SearchBruteForceCost = SearchResult.BruteForcePoseCost.GetTotalCost();
-		}
+		SearchBruteForceCost = SearchResult.BruteForcePoseCost.GetTotalCost();
 #endif // WITH_EDITORONLY_DATA
-		TraceMotionMatchingState(Database, Trajectory, SearchContext, SearchResult, FSearchResult(), 0.f, FTransform::Identity, AnimInstance, DebugSessionUniqueIdentifier,
+		TraceMotionMatchingState(Trajectory, SearchContext, SearchResult, FSearchResult(), 0.f, FTransform::Identity, AnimInstance, DebugSessionUniqueIdentifier,
 			AnimInstance->GetDeltaSeconds(), true, FObjectTrace::GetWorldElapsedTime(AnimInstance->GetWorld()), SearchBestCost, SearchBruteForceCost);
 #endif // UE_POSE_SEARCH_TRACE_ENABLED
 	}
diff --git a/Engine/Plugins/Experimental/Animation/PoseSearch/Source/Runtime/Private/Trace/PoseSearchTraceLogger.cpp b/Engine/Plugins/Experimental/Animation/PoseSearch/Source/Runtime/Private/Trace/PoseSearchTraceLogger.cpp
--- a/Engine/Plugins/Experimental/Animation/PoseSearch/Source/Runtime/Private/Trace/PoseSearchTraceLogger.cpp
+++ b/Engine/Plugins/Experimental/Animation/PoseSearch/Source/Runtime/Private/Trace/PoseSearchTraceLogger.cpp
@@ -70,7 +70,6 @@ FArchive& operator<<(FArchive& Ar, FTraceMotionMatchingStateDatabaseEntry& Entry
 
 FArchive& operator<<(FArchive& Ar, FTraceMotionMatchingState& State)
 {
-	Ar << State.SearchableAssetId;
 	Ar << State.ElapsedPoseSearchTime;
 	Ar << State.AssetPlayerTime;
 	Ar << State.DeltaTime;
diff --git a/Engine/Plugins/Experimental/Animation/PoseSearch/Source/Runtime/Private/Trace/PoseSearchTraceLogger.h b/Engine/Plugins/Experimental/Animation/PoseSearch/Source/Runtime/Private/Trace/PoseSearchTraceLogger.h
--- a/Engine/Plugins/Experimental/Animation/PoseSearch/Source/Runtime/Private/Trace/PoseSearchTraceLogger.h
+++ b/Engine/Plugins/Experimental/Animation/PoseSearch/Source/Runtime/Private/Trace/PoseSearchTraceLogger.h
@@ -65,9 +65,6 @@ POSESEARCH_API FArchive& operator<<(FArchive& Ar, FTraceMotionMatchingStateDatab
  */
 struct POSESEARCH_API FTraceMotionMatchingState
 {
-	/** ObjectId of active searchable asset */
-	uint64 SearchableAssetId = 0;
-
 	/** Amount of time since the last pose switch */
 	float ElapsedPoseSearchTime = 0.f;
 
diff --git a/Engine/Plugins/Experimental/Animation/PoseSearch/Source/Runtime/Public/PoseSearch/PoseSearchContext.h b/Engine/Plugins/Experimental/Animation/PoseSearch/Source/Runtime/Public/PoseSearch/PoseSearchContext.h
--- a/Engine/Plugins/Experimental/Animation/PoseSearch/Source/Runtime/Public/PoseSearch/PoseSearchContext.h
+++ b/Engine/Plugins/Experimental/Animation/PoseSearch/Source/Runtime/Public/PoseSearch/PoseSearchContext.h
@@ -41,8 +41,9 @@ enum class EPoseCandidateFlags : uint8
 	DiscardedBy_PoseReselectHistory = 1 << 4,
 	DiscardedBy_BlockTransition = 1 << 5,
 	DiscardedBy_PoseFilter = 1 << 6,
+	DiscardedBy_Search = 1 << 7,
 
-	AnyDiscardedMask = DiscardedBy_PoseJumpThresholdTime | DiscardedBy_PoseReselectHistory | DiscardedBy_BlockTransition | DiscardedBy_PoseFilter,
+	AnyDiscardedMask = DiscardedBy_PoseJumpThresholdTime | DiscardedBy_PoseReselectHistory | DiscardedBy_BlockTransition | DiscardedBy_PoseFilter | DiscardedBy_Search,
 };
 ENUM_CLASS_FLAGS(EPoseCandidateFlags);
 
@@ -229,10 +230,29 @@ public:
 			}
 			else if (PoseCandidateHeap.Num() < MaxPoseCandidates || Cost < PoseCandidateHeap.HeapTop().Cost)
 			{
+				bool bPoppedContinuingPoseCandidate = false;
+				FPoseCandidate ContinuingPoseCandidate;
 				while (PoseCandidateHeap.Num() >= MaxPoseCandidates)
 				{
-					FPoseCandidate Unused;
-					Pop(Unused);
+					FPoseCandidate PoppedPoseCandidate;
+					Pop(PoppedPoseCandidate);
+
+					if (EnumHasAnyFlags(PoppedPoseCandidate.PoseCandidateFlags, EPoseCandidateFlags::Valid_ContinuingPose))
+					{
+						// we can only have one continuing pose candidate
+						check(!bPoppedContinuingPoseCandidate);
+						ContinuingPoseCandidate = PoppedPoseCandidate;
+						bPoppedContinuingPoseCandidate = true;
+					}					
+				}
+
+				if (bPoppedContinuingPoseCandidate)
+				{
+					// if we popped the continuing pose candidate, we make some space fir it and push it back
+					FPoseCandidate PoppedPoseCandidate;
+					Pop(PoppedPoseCandidate);
+					PoseCandidateHeap.HeapPush(ContinuingPoseCandidate);
+					PoseIdxToFlags.Add(ContinuingPoseCandidate.PoseIdx, PoppedPoseCandidate.PoseCandidateFlags);
 				}
 
 				FPoseCandidate PoseCandidate;
@@ -241,7 +261,6 @@ public:
 				PoseCandidate.Database = Database;
 				
 				PoseCandidateHeap.HeapPush(PoseCandidate);
-
 				PoseIdxToFlags.Add(PoseIdx, PoseCandidateFlags);
 			}
 		}
diff --git a/Engine/Plugins/Experimental/Animation/PoseSearch/Source/Runtime/Public/PoseSearch/PoseSearchLibrary.h b/Engine/Plugins/Experimental/Animation/PoseSearch/Source/Runtime/Public/PoseSearch/PoseSearchLibrary.h
--- a/Engine/Plugins/Experimental/Animation/PoseSearch/Source/Runtime/Public/PoseSearch/PoseSearchLibrary.h
+++ b/Engine/Plugins/Experimental/Animation/PoseSearch/Source/Runtime/Public/PoseSearch/PoseSearchLibrary.h
@@ -81,7 +81,6 @@ class POSESEARCH_API UPoseSearchLibrary : public UBlueprintFunctionLibrary
 
 #if UE_POSE_SEARCH_TRACE_ENABLED
 	static void TraceMotionMatchingState(
-		const UPoseSearchDatabase* Database,
 		const FPoseSearchQueryTrajectory& Trajectory,
 		UE::PoseSearch::FSearchContext& SearchContext,
 		const UE::PoseSearch::FSearchResult& CurrentResult,
diff --git a/Engine/Plugins/Experimental/Animation/PoseSearch/Source/Runtime/Public/PoseSearch/PoseSearchResult.h b/Engine/Plugins/Experimental/Animation/PoseSearch/Source/Runtime/Public/PoseSearch/PoseSearchResult.h
--- a/Engine/Plugins/Experimental/Animation/PoseSearch/Source/Runtime/Public/PoseSearch/PoseSearchResult.h
+++ b/Engine/Plugins/Experimental/Animation/PoseSearch/Source/Runtime/Public/PoseSearch/PoseSearchResult.h
@@ -45,9 +45,6 @@ struct FSearchResult
 
 	TWeakObjectPtr<const UPoseSearchDatabase> Database;
 
-	// cost of the current pose with the query from database in the result, if possible
-	FPoseSearchCost ContinuingPoseCost; 
-
 	float AssetTime = 0.0f;
 
 #if WITH_EDITORONLY_DATA
