motion matching
- separate support for root bone blending in blendstack
- MM support for rotating the root motion delta transform by controlled offset between mesh and pose

diff --git a/Engine/Plugins/Experimental/Animation/PoseSearch/Source/Runtime/Private/AnimNode_BlendStack.cpp b/Engine/Plugins/Experimental/Animation/PoseSearch/Source/Runtime/Private/AnimNode_BlendStack.cpp
--- a/Engine/Plugins/Experimental/Animation/PoseSearch/Source/Runtime/Private/AnimNode_BlendStack.cpp
+++ b/Engine/Plugins/Experimental/Animation/PoseSearch/Source/Runtime/Private/AnimNode_BlendStack.cpp
@@ -18,7 +18,7 @@ TAutoConsoleVariable<int32> CVarAnimBlendStackPruningEnable(TEXT("a.AnimNode.Ble
 
 /////////////////////////////////////////////////////
 // FPoseSearchAnimPlayer
-void FPoseSearchAnimPlayer::Initialize(UAnimationAsset* AnimationAsset, float AccumulatedTime, bool bLoop, bool bMirrored, UMirrorDataTable* MirrorDataTable, float BlendTime, const UBlendProfile* BlendProfile, EAlphaBlendOption InBlendOption, FVector BlendParameters, float PlayRate)
+void FPoseSearchAnimPlayer::Initialize(UAnimationAsset* AnimationAsset, float AccumulatedTime, bool bLoop, bool bMirrored, UMirrorDataTable* MirrorDataTable, float BlendTime, float RootBoneBlendTime, const UBlendProfile* BlendProfile, EAlphaBlendOption InBlendOption, FVector BlendParameters, float PlayRate)
 {
 	check(AnimationAsset);
 
@@ -27,18 +27,46 @@ void FPoseSearchAnimPlayer::Initialize(UAnimationAsset* AnimationAsset, float Ac
 		UE_LOG(LogPoseSearch, Error, TEXT("FPoseSearchAnimPlayer failed to Initialize for %s. Mirroring will not work becasue MirrorDataTable is missing"), *GetNameSafe(AnimationAsset));
 	}
 
-	if (BlendProfile != nullptr)
+	const FReferenceSkeleton& RefSkeleton = AnimationAsset->GetSkeleton()->GetReferenceSkeleton();
+	const bool bApplyDifferentRootBoneBlendTime = RootBoneBlendTime >= 0.f && !FMath::IsNearlyEqual(RootBoneBlendTime, BlendTime);
+	const int32 NumSkeletonBones = RefSkeleton.GetNum();
+	if (NumSkeletonBones <= 0)
 	{
-		const USkeleton* SkeletonAsset = BlendProfile->OwningSkeleton;
-		check(SkeletonAsset);
+		UE_LOG(LogPoseSearch, Error, TEXT("FPoseSearchAnimPlayer failed to Initialize for %s. Skeleton has no bones?!"), *GetNameSafe(AnimationAsset));
+	}
+	else if (BlendTime > UE_KINDA_SMALL_NUMBER)
+	{
+		// handling BlendTime > 0 and RootBoneBlendTime >= 0
+		if (BlendProfile != nullptr)
+		{
+			check(BlendProfile->OwningSkeleton && NumSkeletonBones == BlendProfile->OwningSkeleton->GetReferenceSkeleton().GetNum());
 
-		const FReferenceSkeleton& RefSkeleton = SkeletonAsset->GetReferenceSkeleton();
-		const int32 NumSkeletonBones = RefSkeleton.GetNum();
-		TotalBlendInTimePerBone.Init(BlendTime, NumSkeletonBones);
+			TotalBlendInTimePerBone.Init(BlendTime, NumSkeletonBones);
+
+			BlendProfile->FillSkeletonBoneDurationsArray(TotalBlendInTimePerBone, BlendTime);
+
+			if (bApplyDifferentRootBoneBlendTime)
+			{
+				TotalBlendInTimePerBone[RootBoneIndexType] *= RootBoneBlendTime / BlendTime;
+			}
 
-		BlendProfile->FillSkeletonBoneDurationsArray(TotalBlendInTimePerBone, BlendTime);
-		BlendTime = *Algo::MaxElement(TotalBlendInTimePerBone);
+			BlendTime = *Algo::MaxElement(TotalBlendInTimePerBone);
+		}
+		else if (bApplyDifferentRootBoneBlendTime)
+		{
+			TotalBlendInTimePerBone.Init(BlendTime, NumSkeletonBones);
+			TotalBlendInTimePerBone[RootBoneIndexType] *= RootBoneBlendTime / BlendTime;
+			BlendTime = FMath::Max(BlendTime, RootBoneBlendTime);
+		}
 	}
+	else if (bApplyDifferentRootBoneBlendTime)
+	{
+		// handling BlendTime ~= 0 and RootBoneBlendTime >= 0
+		TotalBlendInTimePerBone.Init(BlendTime, NumSkeletonBones);
+		TotalBlendInTimePerBone[RootBoneIndexType] = RootBoneBlendTime;
+		BlendTime = FMath::Max(BlendTime, RootBoneBlendTime);
+	}
+
 	BlendOption = InBlendOption;
 
 	TotalBlendInTime = BlendTime;
@@ -382,11 +410,11 @@ float FAnimNode_BlendStack_Standalone::GetAccumulatedTime() const
 	return AnimPlayers.IsEmpty() ? 0.f : AnimPlayers.First().GetAccumulatedTime();
 }
 
-void FAnimNode_BlendStack_Standalone::BlendTo(UAnimationAsset* AnimationAsset, float AccumulatedTime, bool bLoop, bool bMirrored, UMirrorDataTable* MirrorDataTable, int32 MaxActiveBlends, float BlendTime, const UBlendProfile* BlendProfile, EAlphaBlendOption BlendOption, FVector BlendParameters, float PlayRate)
+void FAnimNode_BlendStack_Standalone::BlendTo(UAnimationAsset* AnimationAsset, float AccumulatedTime, bool bLoop, bool bMirrored, UMirrorDataTable* MirrorDataTable, int32 MaxActiveBlends, float BlendTime, float RootBoneBlendTime, const UBlendProfile* BlendProfile, EAlphaBlendOption BlendOption, FVector BlendParameters, float PlayRate)
 {
 	RequestedMaxActiveBlends = MaxActiveBlends;
 	AnimPlayers.PushFirst(FPoseSearchAnimPlayer());
-	AnimPlayers.First().Initialize(AnimationAsset, AccumulatedTime, bLoop, bMirrored, MirrorDataTable, BlendTime, BlendProfile, BlendOption, BlendParameters, PlayRate);
+	AnimPlayers.First().Initialize(AnimationAsset, AccumulatedTime, bLoop, bMirrored, MirrorDataTable, BlendTime, RootBoneBlendTime, BlendProfile, BlendOption, BlendParameters, PlayRate);
 }
 
 void FAnimNode_BlendStack_Standalone::UpdatePlayRate(float PlayRate)
@@ -457,7 +485,7 @@ void FAnimNode_BlendStack::UpdateAssetPlayer(const FAnimationUpdateContext& Cont
 
 		if (bExecuteBlendTo)
 		{
-			BlendTo(AnimationAsset, AnimationTime, bLoop, bMirrored, MirrorDataTable.Get(), MaxActiveBlends, BlendTime, BlendProfile, BlendOption, BlendParameters, WantedPlayRate);
+			BlendTo(AnimationAsset, AnimationTime, bLoop, bMirrored, MirrorDataTable.Get(), MaxActiveBlends, BlendTime, RootBoneBlendTime, BlendProfile, BlendOption, BlendParameters, WantedPlayRate);
 		}
 	}
 	
diff --git a/Engine/Plugins/Experimental/Animation/PoseSearch/Source/Runtime/Private/AnimNode_MotionMatching.cpp b/Engine/Plugins/Experimental/Animation/PoseSearch/Source/Runtime/Private/AnimNode_MotionMatching.cpp
--- a/Engine/Plugins/Experimental/Animation/PoseSearch/Source/Runtime/Private/AnimNode_MotionMatching.cpp
+++ b/Engine/Plugins/Experimental/Animation/PoseSearch/Source/Runtime/Private/AnimNode_MotionMatching.cpp
@@ -48,23 +48,34 @@ void FAnimNode_MotionMatching::Evaluate_AnyThread(FPoseContext& Output)
 
 	Source.Evaluate(Output);
 
-	// applying MotionMatchingState.RootBoneDeltaYaw
+	// applying MotionMatchingState.RootBoneDeltaYaw to the root bone and the root motion delta transform
 	if (!FMath::IsNearlyZero(MotionMatchingState.GetRootBoneDeltaYaw()))
 	{
 		const FQuat RootBoneDelta(FRotator(0.f, MotionMatchingState.GetRootBoneDeltaYaw(), 0.f));
 		FCompactPoseBoneIndex RootBoneIndex(RootBoneIndexType);
 		Output.Pose[RootBoneIndex].SetRotation(Output.Pose[RootBoneIndex].GetRotation() * RootBoneDelta);
 		Output.Pose[RootBoneIndex].NormalizeRotation();
+
+		const UE::Anim::IAnimRootMotionProvider* RootMotionProvider = UE::Anim::IAnimRootMotionProvider::Get();
+		if (RootMotionProvider && RootMotionProvider->HasRootMotion(Output.CustomAttributes))
+		{
+			FTransform RootMotionTransformDelta;
+			RootMotionProvider->ExtractRootMotion(Output.CustomAttributes, RootMotionTransformDelta);
+			RootMotionTransformDelta.SetTranslation(RootBoneDelta.RotateVector(RootMotionTransformDelta.GetTranslation()));
+			RootMotionProvider->OverrideRootMotion(RootMotionTransformDelta, Output.CustomAttributes);
+		}
 	}
 	
 #if UE_POSE_SEARCH_TRACE_ENABLED
-	MotionMatchingState.RootMotionTransformDelta = FTransform::Identity;
-
 	const UE::Anim::IAnimRootMotionProvider* RootMotionProvider = UE::Anim::IAnimRootMotionProvider::Get();
 	if (RootMotionProvider && RootMotionProvider->HasRootMotion(Output.CustomAttributes))
 	{
 		RootMotionProvider->ExtractRootMotion(Output.CustomAttributes, MotionMatchingState.RootMotionTransformDelta);
 	}
+	else
+	{
+		MotionMatchingState.RootMotionTransformDelta = FTransform::Identity;
+	}
 #endif // UE_POSE_SEARCH_TRACE_ENABLED
 }
 
@@ -168,9 +179,11 @@ void FAnimNode_MotionMatching::UpdateAssetPlayer(const FAnimationUpdateContext&
 		{
 			if (const FPoseSearchDatabaseAnimationAssetBase* DatabaseAsset = CurrentResultDatabase->GetAnimationAssetBase(*SearchIndexAsset))
 			{
+				// root bone blending needs to be immediate if MM node controls the offset between mesh component and root bone
+				const float RootBoneBlendTime = FMath::IsNearlyZero(YawFromAnimation) ? BlendTime : 0.f;
 				BlendStackNode.BlendTo(DatabaseAsset->GetAnimationAsset(), MotionMatchingState.CurrentSearchResult.AssetTime,
 					DatabaseAsset->IsLooping(), SearchIndexAsset->bMirrored, CurrentResultDatabase->Schema->MirrorDataTable.Get(),
-					MaxActiveBlends, BlendTime, BlendProfile, BlendOption, SearchIndexAsset->BlendParameters, MotionMatchingState.WantedPlayRate);
+					MaxActiveBlends, BlendTime, RootBoneBlendTime, BlendProfile, BlendOption, SearchIndexAsset->BlendParameters, MotionMatchingState.WantedPlayRate);
 			}
 		}
 	}
diff --git a/Engine/Plugins/Experimental/Animation/PoseSearch/Source/Runtime/Public/PoseSearch/AnimNode_BlendStack.h b/Engine/Plugins/Experimental/Animation/PoseSearch/Source/Runtime/Public/PoseSearch/AnimNode_BlendStack.h
--- a/Engine/Plugins/Experimental/Animation/PoseSearch/Source/Runtime/Public/PoseSearch/AnimNode_BlendStack.h
+++ b/Engine/Plugins/Experimental/Animation/PoseSearch/Source/Runtime/Public/PoseSearch/AnimNode_BlendStack.h
@@ -15,7 +15,7 @@ struct FPoseSearchAnimPlayer
 {
 	GENERATED_BODY()
 	
-	void Initialize(UAnimationAsset* AnimationAsset, float AccumulatedTime, bool bLoop, bool bMirrored, UMirrorDataTable* MirrorDataTable, float BlendTime, const UBlendProfile* BlendProfile, EAlphaBlendOption InBlendOption, FVector BlendParameters, float PlayRate);
+	void Initialize(UAnimationAsset* AnimationAsset, float AccumulatedTime, bool bLoop, bool bMirrored, UMirrorDataTable* MirrorDataTable, float BlendTime, float RootBoneBlendTime, const UBlendProfile* BlendProfile, EAlphaBlendOption InBlendOption, FVector BlendParameters, float PlayRate);
 	void UpdatePlayRate(float PlayRate);
 	void Evaluate_AnyThread(FPoseContext& Output);
 	void Update_AnyThread(const FAnimationUpdateContext& Context, float BlendWeight);
@@ -76,7 +76,7 @@ struct POSESEARCH_API FAnimNode_BlendStack_Standalone : public FAnimNode_AssetPl
 	virtual void GatherDebugData(FNodeDebugData& DebugData) override;
 	// End of FAnimNode_Base interface
 
-	void BlendTo(UAnimationAsset* AnimationAsset, float AccumulatedTime = 0.f, bool bLoop = false, bool bMirrored = false, UMirrorDataTable* MirrorDataTable = nullptr, int32 MaxActiveBlends = 3, float BlendTime = 0.3f, const UBlendProfile* BlendProfile = nullptr, EAlphaBlendOption BlendOption = EAlphaBlendOption::Linear, FVector BlendParameters = FVector::Zero(), float PlayRate = 1.f);
+	void BlendTo(UAnimationAsset* AnimationAsset, float AccumulatedTime = 0.f, bool bLoop = false, bool bMirrored = false, UMirrorDataTable* MirrorDataTable = nullptr, int32 MaxActiveBlends = 3, float BlendTime = 0.2f, float RootBoneBlendTime = -1.f, const UBlendProfile* BlendProfile = nullptr, EAlphaBlendOption BlendOption = EAlphaBlendOption::Linear, FVector BlendParameters = FVector::Zero(), float PlayRate = 1.f);
 	void UpdatePlayRate(float PlayRate);
 
 	// FAnimNode_AssetPlayerBase interface
@@ -115,6 +115,10 @@ struct POSESEARCH_API FAnimNode_BlendStack : public FAnimNode_BlendStack_Standal
 	UPROPERTY(EditAnywhere, BlueprintReadWrite, Category = Settings, meta = (PinHiddenByDefault, ClampMin = "0"))
 	float BlendTime = 0.2f;
 
+	// Time in seconds to blend out to the new pose root bone. Negative values would imply RootBoneBlendTime is equal to BlendTime
+	UPROPERTY(EditAnywhere, BlueprintReadWrite, Category = Settings, meta = (PinHiddenByDefault))
+	float RootBoneBlendTime = -1.f;
+
 	// if AnimationTime and MaxAnimationDeltaTime are positive and the currently playing animation total time differs more than MaxAnimationDeltaTime from AnimationTime
 	// (animation desynchronized from the requested time) the blend stack will force a blend into the same animation
 	UPROPERTY(EditAnywhere, BlueprintReadWrite, Category = Settings, meta = (PinHiddenByDefault))
