motion matching - adding ReasonForNoActivePose, ReasonForNoContinuingPose, ReasonForNoCandidates

diff --git a/Engine/Plugins/Experimental/Animation/PoseSearch/Source/Editor/Private/PoseSearchDebuggerDatabaseView.cpp b/Engine/Plugins/Experimental/Animation/PoseSearch/Source/Editor/Private/PoseSearchDebuggerDatabaseView.cpp
--- a/Engine/Plugins/Experimental/Animation/PoseSearch/Source/Editor/Private/PoseSearchDebuggerDatabaseView.cpp
+++ b/Engine/Plugins/Experimental/Animation/PoseSearch/Source/Editor/Private/PoseSearchDebuggerDatabaseView.cpp
@@ -632,6 +632,41 @@ void SDebuggerDatabaseView::PopulateViewRows()
 	ActiveView.ListView->RequestListRefresh();
 	ContinuingPoseView.ListView->RequestListRefresh();
 	FilteredDatabaseView.ListView->RequestListRefresh();
+
+	if (ActiveView.Rows.Num() > 0)
+	{
+		ReasonForNoActivePose = FText::GetEmpty();
+	}
+	else
+	{
+		ReasonForNoActivePose = LOCTEXT("ReasonForNoActivePose", "Database search didn't find any candidates, or the search has not been performed");
+	}
+
+	if (ContinuingPoseView.Rows.Num() > 0)
+	{
+		ReasonForNoContinuingPose = FText::GetEmpty();
+	}
+	else
+	{
+		ReasonForNoContinuingPose = LOCTEXT("ReasonForNoContinuingPose", "Invalid continuing pose");
+	}
+
+	if (FilteredDatabaseView.Rows.Num() > 0)
+	{
+		ReasonForNoCandidates = FText::GetEmpty();
+	}
+	else if (UnfilteredDatabaseRows.Num() == 0)
+	{
+		ReasonForNoCandidates = LOCTEXT("ReasonForNoCandidates_NoSearch", "Database search didn't find any candidates, or the search has not been performed");
+	}
+	else if (UnfilteredDatabaseRows.Num() == 1)
+	{
+		ReasonForNoCandidates = LOCTEXT("ReasonForNoCandidates_OnlyContinuingPose", "The continuing pose cost cannot be lowered by searching the databases, so the search has been skipped");
+	}
+	else
+	{
+		ReasonForNoCandidates = FText::Format(LOCTEXT("ReasonForNoCandidates_AllFilteredOut", "All {0} databases poses have been filtered out"), UnfilteredDatabaseRows.Num());
+	}
 }
 
 TSharedRef<ITableRow> SDebuggerDatabaseView::HandleGenerateDatabaseRow(TSharedRef<FDebuggerDatabaseRowData> Item, const TSharedRef<STableViewBase>& OwnerTable) const
@@ -774,12 +809,35 @@ void SDebuggerDatabaseView::Construct(const FArguments& InArgs)
 					.VAlign(VAlign_Fill)
 					.Padding(0.0f)
 					[
-						SNew(SBorder)
-					
-						.BorderImage(FAppStyle::GetBrush("NoBorder"))
-						.Padding(0.0f)
+						SNew(SOverlay)
+						+ SOverlay::Slot()
 						[
-							ActiveView.ListView.ToSharedRef()
+							SNew(SBorder)
+							.BorderImage(FAppStyle::GetBrush("NoBorder"))
+							.Padding(0.0f)
+							[
+								ActiveView.ListView.ToSharedRef()
+							]
+						]
+						+ SOverlay::Slot()
+						[
+							SNew(SBorder)
+							.Visibility(EVisibility::SelfHitTestInvisible)
+							.Padding(FMargin(5.f, 5.f, 5.f, 5.f))
+							.HAlign(HAlign_Center)
+							.VAlign(VAlign_Fill)
+							[
+								SNew(STextBlock)
+								.Visibility_Lambda([this]()
+								{
+									return ReasonForNoActivePose.IsEmpty() ? EVisibility::Collapsed : EVisibility::HitTestInvisible;
+								})
+								.Margin(FMargin(5.f, 5.f, 5.f, 5.f))
+								.Text_Lambda([this]()
+								{
+									return ReasonForNoActivePose;
+								})
+							]
 						]
 					]
 
@@ -833,12 +891,35 @@ void SDebuggerDatabaseView::Construct(const FArguments& InArgs)
 					.VAlign(VAlign_Fill)
 					.Padding(0.0f)
 					[
-						SNew(SBorder)
-					
-						.BorderImage(FAppStyle::GetBrush("NoBorder"))
-						.Padding(0.0f)
+						SNew(SOverlay)
+						+ SOverlay::Slot()
+						[
+							SNew(SBorder)
+							.BorderImage(FAppStyle::GetBrush("NoBorder"))
+							.Padding(0.0f)
+							[
+								ContinuingPoseView.ListView.ToSharedRef()
+							]
+						]
+						+ SOverlay::Slot()
 						[
-							ContinuingPoseView.ListView.ToSharedRef()
+							SNew(SBorder)
+							.Visibility(EVisibility::SelfHitTestInvisible)
+							.Padding(FMargin(5.f, 5.f, 5.f, 5.f))
+							.HAlign(HAlign_Center)
+							.VAlign(VAlign_Fill)
+							[
+								SNew(STextBlock)
+								.Visibility_Lambda([this]()
+								{
+									return ReasonForNoContinuingPose.IsEmpty() ? EVisibility::Collapsed : EVisibility::HitTestInvisible;
+								})
+								.Margin(FMargin(5.f, 5.f, 5.f, 5.f))
+								.Text_Lambda([this]()
+								{
+									return ReasonForNoContinuingPose;
+								})
+							]
 						]
 					]
 
@@ -957,11 +1038,35 @@ void SDebuggerDatabaseView::Construct(const FArguments& InArgs)
 					+ SHorizontalBox::Slot()
 					.Padding(0.0f)
 					[
-						SNew(SBorder)
-						.BorderImage(FAppStyle::GetBrush("NoBorder"))
-						.Padding(0.0f)
+						SNew(SOverlay)
+						+SOverlay::Slot()
+						[
+							SNew(SBorder)
+							.BorderImage(FAppStyle::GetBrush("NoBorder"))
+							.Padding(0.0f)
+							[
+								FilteredDatabaseView.ListView.ToSharedRef()
+							]
+						]
+						+ SOverlay::Slot()
 						[
-							FilteredDatabaseView.ListView.ToSharedRef()
+							SNew(SBorder)
+							.Visibility(EVisibility::SelfHitTestInvisible)
+							.Padding(FMargin(5.f, 5.f, 5.f, 5.f))
+							.HAlign(HAlign_Center)
+							.VAlign(VAlign_Fill)
+							[
+								SNew(STextBlock)
+								.Visibility_Lambda([this]()
+								{
+									return ReasonForNoCandidates.IsEmpty() ? EVisibility::Collapsed : EVisibility::HitTestInvisible;
+								})
+								.Margin(FMargin(5.f, 5.f, 5.f, 5.f))
+								.Text_Lambda([this]()
+								{
+									return ReasonForNoCandidates;
+								})
+							]
 						]
 					]
 				
diff --git a/Engine/Plugins/Experimental/Animation/PoseSearch/Source/Editor/Private/PoseSearchDebuggerDatabaseView.h b/Engine/Plugins/Experimental/Animation/PoseSearch/Source/Editor/Private/PoseSearchDebuggerDatabaseView.h
--- a/Engine/Plugins/Experimental/Animation/PoseSearch/Source/Editor/Private/PoseSearchDebuggerDatabaseView.h
+++ b/Engine/Plugins/Experimental/Animation/PoseSearch/Source/Editor/Private/PoseSearchDebuggerDatabaseView.h
@@ -150,6 +150,10 @@ private:
 
 	/* Bool used to use FilterText as regex filter*/
 	bool bUseRegex = false;
+
+	FText ReasonForNoActivePose;
+	FText ReasonForNoContinuingPose;
+	FText ReasonForNoCandidates;
 };
 
 } // namespace UE::PoseSearch
