motion matching
- limiting the number of debug draws in SDebuggerView::DrawFeatures
- limiting the number of columns SDatabaseDataDetails
- adding Horizontal SScrollBox to SDatabaseDataDetails
- making Query column consistent with other columns by hiding expanded channels texts

diff --git a/Engine/Plugins/Experimental/Animation/PoseSearch/Source/Editor/Private/PoseSearchDatabaseDataDetails.cpp b/Engine/Plugins/Experimental/Animation/PoseSearch/Source/Editor/Private/PoseSearchDatabaseDataDetails.cpp
--- a/Engine/Plugins/Experimental/Animation/PoseSearch/Source/Editor/Private/PoseSearchDatabaseDataDetails.cpp
+++ b/Engine/Plugins/Experimental/Animation/PoseSearch/Source/Editor/Private/PoseSearchDatabaseDataDetails.cpp
@@ -165,32 +165,35 @@ public:
 			.Text_Lambda([this, ColumnName]() -> FText
 			{
 				TStringBuilder<256> StringBuilder;
-				if (const TSharedPtr<FDatabaseViewModel> ViewModel = EditorViewModel.Pin())
+				if (!ChannelItem->IsExpanded())
 				{
-					const int32 DataOffset = ChannelItem->GetDataOffset();
-					const int32 Cardinality = ChannelItem->GetCardinality();
-
-					const TConstArrayView<float> QueryValues = ViewModel->GetQueryVector();
-					if (DataOffset + Cardinality <= QueryValues.Num())
+					if (const TSharedPtr<FDatabaseViewModel> ViewModel = EditorViewModel.Pin())
 					{
-						if (Cardinality > 1)
+						const int32 DataOffset = ChannelItem->GetDataOffset();
+						const int32 Cardinality = ChannelItem->GetCardinality();
+
+						const TConstArrayView<float> QueryValues = ViewModel->GetQueryVector();
+						if (DataOffset + Cardinality <= QueryValues.Num())
 						{
-							// using only one decimal to keep the string compact
-							for (int32 i = 0; i < Cardinality; ++i)
+							if (Cardinality > 1)
 							{
-								if (i != 0)
+								// using only one decimal to keep the string compact
+								for (int32 i = 0; i < Cardinality; ++i)
 								{
-									StringBuilder.Append(TEXT(", "));
+									if (i != 0)
+									{
+										StringBuilder.Append(TEXT(", "));
+									}
+									const float Value = QueryValues[i + DataOffset];
+									StringBuilder.Appendf(TEXT("%.1f"), Value);
 								}
-								const float Value = QueryValues[i + DataOffset];
-								StringBuilder.Appendf(TEXT("%.1f"), Value);
 							}
-						}
-						else
-						{
-							// using all the float digits 
-							const float Value = QueryValues[DataOffset];
-							StringBuilder.Appendf(TEXT("%f"), Value);
+							else
+							{
+								// using all the float digits 
+								const float Value = QueryValues[DataOffset];
+								StringBuilder.Appendf(TEXT("%f"), Value);
+							}
 						}
 					}
 				}
@@ -199,6 +202,7 @@ public:
 		}
 		
 		return SNew(STextBlock)
+		.Margin(FMargin(1.f, 1.f, 1.f, 1.f))
 		.Text_Lambda([this, ColumnName]() -> FText
 		{
 			TStringBuilder<256> StringBuilder;
@@ -281,7 +285,7 @@ void SDatabaseDataDetails::SetExpandedItems(TArray<FChannelItemPtr>& ChannelItem
 	}
 };
 
-void SDatabaseDataDetails::Reconstruct()
+void SDatabaseDataDetails::Reconstruct(int32 MaxPreviewActors)
 {
 	const TSharedPtr<FDatabaseViewModel> ViewModel = EditorViewModel.Pin();
 	
@@ -306,7 +310,7 @@ void SDatabaseDataDetails::Reconstruct()
 	RebuildChannelItemsTreeRecursively(ChannelItems, PoseSearchDatabase->Schema->GetChannels());
 
 	// generating the header
-	TSharedPtr<SHeaderRow> HeaderRow = SNew(SHeaderRow); //.CanSelectGeneratedColumn(true).Visibility(EVisibility::Visible);
+	TSharedPtr<SHeaderRow> HeaderRow = SNew(SHeaderRow);
 	HeaderRow->AddColumn(
 		SHeaderRow::Column(TEXT("ChannelName"))
 			.DefaultLabel(LOCTEXT("ChannelName_Header", "Channel Name"))
@@ -327,8 +331,10 @@ void SDatabaseDataDetails::Reconstruct()
 			.ToolTipText(LOCTEXT("Query_ToolTip", "Query Values")));
 	}
 
-	for (const FDatabasePreviewActor& PreviewActor : ViewModel->GetPreviewActors())
+	const int32 PreviewActorNum = FMath::Min(MaxPreviewActors, ViewModel->GetPreviewActors().Num());
+	for (int32 PreviewActorIdx = 0; PreviewActorIdx < PreviewActorNum; ++PreviewActorIdx)
 	{
+		const FDatabasePreviewActor& PreviewActor = ViewModel->GetPreviewActors()[PreviewActorIdx];
 		HeaderRow->AddColumn(
 			SHeaderRow::Column(*PreviewActor.Actor->GetName())
 			.DefaultLabel(FText::FromString(PreviewActor.Sampler.GetAsset()->GetName())));
@@ -352,12 +358,13 @@ void SDatabaseDataDetails::Reconstruct()
 
 	ChildSlot
 	[
-		// @todo: add a SScrollBox to handle the complexity of multiple actors
-		//SNew(SScrollBox)
-		//+ SScrollBox::Slot()
-		//[
+		SNew(SScrollBox)
+		.Orientation(Orient_Horizontal)
+		+ SScrollBox::Slot()
+		.FillSize(1.f)
+		[
 			ChannelItemsTreeView.ToSharedRef()
-		//]
+		]
 	];
 
 	SetExpandedItems(ChannelItems, ExpandedItems, ChannelItemsTreeView.Get());
diff --git a/Engine/Plugins/Experimental/Animation/PoseSearch/Source/Editor/Private/PoseSearchDebuggerView.cpp b/Engine/Plugins/Experimental/Animation/PoseSearch/Source/Editor/Private/PoseSearchDebuggerView.cpp
--- a/Engine/Plugins/Experimental/Animation/PoseSearch/Source/Editor/Private/PoseSearchDebuggerView.cpp
+++ b/Engine/Plugins/Experimental/Animation/PoseSearch/Source/Editor/Private/PoseSearchDebuggerView.cpp
@@ -314,12 +314,7 @@ TArray<TSharedRef<FDebuggerDatabaseRowData>> SDebuggerView::GetSelectedDatabaseR
 	return DatabaseView->GetDatabaseRows()->GetSelectedItems();
 }
 
-void SDebuggerView::DrawFeatures(
-	const UWorld& DebuggerWorld,
-	const FTraceMotionMatchingStateMessage& State,
-	const FTransform& RootBoneWorldTransform,
-	const USkinnedMeshComponent* Mesh
-) const
+void SDebuggerView::DrawFeatures(const UWorld& DebuggerWorld, const FTraceMotionMatchingStateMessage& State, const FTransform& RootBoneWorldTransform, const USkinnedMeshComponent* Mesh, int32 MaxRowsToDraw) const
 {
 	// Draw world space trajectory
 #if ENABLE_ANIM_DEBUG
@@ -354,8 +349,10 @@ void SDebuggerView::DrawFeatures(
 	TArray<TSharedRef<FDebuggerDatabaseRowData>> SelectedRows = DatabaseRows->GetSelectedItems();
 	
 	// Draw any selected database vectors
-	for (const TSharedRef<FDebuggerDatabaseRowData>& Row : SelectedRows)
+	const int32 NumRowsToDraw = FMath::Min(MaxRowsToDraw, SelectedRows.Num());
+	for (int32 RowIdx = 0; RowIdx < NumRowsToDraw; ++RowIdx)
 	{
+		const TSharedRef<FDebuggerDatabaseRowData>& Row = SelectedRows[RowIdx];
 		const UPoseSearchDatabase* RowDatabase = Row->SharedData->SourceDatabase.Get();
 		if (FAsyncPoseSearchDatabasesManagement::RequestAsyncBuildIndex(RowDatabase, ERequestAsyncBuildFlag::ContinueRequest))
 		{
diff --git a/Engine/Plugins/Experimental/Animation/PoseSearch/Source/Editor/Private/PoseSearchDebuggerView.h b/Engine/Plugins/Experimental/Animation/PoseSearch/Source/Editor/Private/PoseSearchDebuggerView.h
--- a/Engine/Plugins/Experimental/Animation/PoseSearch/Source/Editor/Private/PoseSearchDebuggerView.h
+++ b/Engine/Plugins/Experimental/Animation/PoseSearch/Source/Editor/Private/PoseSearchDebuggerView.h
@@ -79,7 +79,7 @@ public:
 
 private:
 	/** Called each frame to draw features of the query vector & database selections */
-	void DrawFeatures(const UWorld& DebuggerWorld, const FTraceMotionMatchingStateMessage& State, const FTransform& RootBoneWorldTransform, const USkinnedMeshComponent* Mesh) const;
+	void DrawFeatures(const UWorld& DebuggerWorld, const FTraceMotionMatchingStateMessage& State, const FTransform& RootBoneWorldTransform, const USkinnedMeshComponent* Mesh, int32 MaxRowsToDraw = 250) const;
 	
 	/** Check if a node selection was made, true if a node is selected */
 	bool UpdateNodeSelection();
diff --git a/Engine/Plugins/Experimental/Animation/PoseSearch/Source/Editor/Public/PoseSearchDatabaseDataDetails.h b/Engine/Plugins/Experimental/Animation/PoseSearch/Source/Editor/Public/PoseSearchDatabaseDataDetails.h
--- a/Engine/Plugins/Experimental/Animation/PoseSearch/Source/Editor/Public/PoseSearchDatabaseDataDetails.h
+++ b/Engine/Plugins/Experimental/Animation/PoseSearch/Source/Editor/Public/PoseSearchDatabaseDataDetails.h
@@ -20,7 +20,7 @@ namespace UE::PoseSearch
 		SLATE_END_ARGS()
 
 		void Construct(const FArguments& Args, TSharedRef<FDatabaseViewModel> InEditorViewModel);
-		void Reconstruct();
+		void Reconstruct(int32 MaxPreviewActors = 15);
 
 	private:
 		static void RebuildChannelItemsTreeRecursively(TArray<FChannelItemPtr>& ChannelItems, TConstArrayView<TObjectPtr<UPoseSearchFeatureChannel>> Channels);
