motion matching - adding a simpler FAnimationAssetSampler::ExtractPose API

diff --git a/Engine/Plugins/Experimental/Animation/PoseSearch/Source/Editor/Private/PoseSearchDatabasePreviewScene.cpp b/Engine/Plugins/Experimental/Animation/PoseSearch/Source/Editor/Private/PoseSearchDatabasePreviewScene.cpp
--- a/Engine/Plugins/Experimental/Animation/PoseSearch/Source/Editor/Private/PoseSearchDatabasePreviewScene.cpp
+++ b/Engine/Plugins/Experimental/Animation/PoseSearch/Source/Editor/Private/PoseSearchDatabasePreviewScene.cpp
@@ -68,7 +68,7 @@ namespace UE::PoseSearch
 				{
 					if (UDebugSkelMeshComponent* Mesh = PreviewActor.GetDebugSkelMeshComponent())
 					{
-						UE::PoseSearch::FDebugDrawParams DrawParams(GetWorld(), Mesh, PreviewActor.QuantizedTimeRootMotion, Database);
+						UE::PoseSearch::FDebugDrawParams DrawParams(GetWorld(), Mesh, PreviewActor.QuantizedTimeRootTransform, Database);
 						DrawParams.DrawFeatureVector(PreviewActor.CurrentPoseIndex);
 
 						if (bDrawQueryVector)
diff --git a/Engine/Plugins/Experimental/Animation/PoseSearch/Source/Editor/Private/PoseSearchDatabaseViewModel.cpp b/Engine/Plugins/Experimental/Animation/PoseSearch/Source/Editor/Private/PoseSearchDatabaseViewModel.cpp
--- a/Engine/Plugins/Experimental/Animation/PoseSearch/Source/Editor/Private/PoseSearchDatabaseViewModel.cpp
+++ b/Engine/Plugins/Experimental/Animation/PoseSearch/Source/Editor/Private/PoseSearchDatabaseViewModel.cpp
@@ -34,9 +34,9 @@ namespace UE::PoseSearch
 		return Actor != nullptr;
 	}
 
-	void FDatabasePreviewActor::Process()
+	void FDatabasePreviewActor::Process(const FBoneContainer& BoneContainer)
 	{
-		Sampler.Process();
+		Sampler.Process(BoneContainer);
 	}
 
 	UDebugSkelMeshComponent* FDatabasePreviewActor::GetDebugSkelMeshComponent()
@@ -127,7 +127,7 @@ namespace UE::PoseSearch
 		PlayTime = MaxPreviewPlayLength;
 	}
 
-	FDatabasePreviewActor FDatabaseViewModel::SpawnPreviewActor(int32 IndexAssetIndex, const FBoneContainer& BoneContainer, float PlayTimeOffset)
+	FDatabasePreviewActor FDatabaseViewModel::SpawnPreviewActor(int32 IndexAssetIndex, float PlayTimeOffset)
 	{
 		FDatabasePreviewActor PreviewActor;
 
@@ -145,7 +145,7 @@ namespace UE::PoseSearch
 		const FInstancedStruct& DatabaseAsset = PoseSearchDatabase->GetAnimationAssetStruct(IndexAsset);
 		const FPoseSearchDatabaseAnimationAssetBase* DatabaseAnimationAssetBase = DatabaseAsset.GetPtr<FPoseSearchDatabaseAnimationAssetBase>();
 		check(DatabaseAnimationAssetBase);
-		PreviewActor.Sampler.Init(DatabaseAnimationAssetBase->GetAnimationAsset(), IndexAsset.BlendParameters, BoneContainer);
+		PreviewActor.Sampler.Init(DatabaseAnimationAssetBase->GetAnimationAsset(), IndexAsset.BlendParameters);
 		PreviewActor.IndexAssetIndex = IndexAssetIndex;
 		PreviewActor.CurrentPoseIndex = INDEX_NONE;
 		PreviewActor.PlayTimeOffset = PlayTimeOffset;
@@ -232,33 +232,33 @@ namespace UE::PoseSearch
 				continue;
 			}
 
-			float CurrentTime = 0.f;
+			PreviewActor.CurrentTime = 0.f;
 			const FPoseSearchIndexAsset& IndexAsset = SearchIndex.Assets[PreviewActor.IndexAssetIndex];
 			float CurrentPlayTime = PlayTime + IndexAsset.SamplingInterval.Min + PreviewActor.PlayTimeOffset;
-			FAnimationRuntime::AdvanceTime(false, CurrentPlayTime, CurrentTime, IndexAsset.SamplingInterval.Max);
+			FAnimationRuntime::AdvanceTime(false, CurrentPlayTime, PreviewActor.CurrentTime, IndexAsset.SamplingInterval.Max);
 			 
 			// time to pose index
-			PreviewActor.CurrentPoseIndex = PoseSearchDatabase->GetPoseIndexFromTime(CurrentTime, IndexAsset);
+			PreviewActor.CurrentPoseIndex = PoseSearchDatabase->GetPoseIndexFromTime(PreviewActor.CurrentTime, IndexAsset);
 
 			const float TimeToAssetTimeMultiplier = PreviewActor.Sampler.GetTimeToAssetTimeMultiplier();
 			check(!FMath::IsNearlyZero(TimeToAssetTimeMultiplier));
-			const float QuantizedTime = PreviewActor.CurrentPoseIndex >= 0 ? PoseSearchDatabase->GetAssetTime(PreviewActor.CurrentPoseIndex) / TimeToAssetTimeMultiplier : CurrentTime;
+			const float QuantizedTime = PreviewActor.CurrentPoseIndex >= 0 ? PoseSearchDatabase->GetAssetTime(PreviewActor.CurrentPoseIndex) / TimeToAssetTimeMultiplier : PreviewActor.CurrentTime;
 
 			if (bQuantizeAnimationToPoseData)
 			{
-				CurrentTime = QuantizedTime; 
+				PreviewActor.CurrentTime = QuantizedTime;
 			}
 
-			AnimInstance->SetPosition(CurrentTime * TimeToAssetTimeMultiplier);
+			AnimInstance->SetPosition(PreviewActor.CurrentTime * TimeToAssetTimeMultiplier);
 			AnimInstance->SetPlayRate(0.f);
 			AnimInstance->SetBlendSpacePosition(IndexAsset.BlendParameters);
 
-			PreviewActor.QuantizedTimeRootMotion = PreviewActor.Sampler.ExtractRootTransform(QuantizedTime);
-			FTransform RootMotion = PreviewActor.Sampler.ExtractRootTransform(CurrentTime);
+			PreviewActor.QuantizedTimeRootTransform = PreviewActor.Sampler.ExtractRootTransform(QuantizedTime);
+			FTransform RootMotion = PreviewActor.Sampler.ExtractRootTransform(PreviewActor.CurrentTime);
 			if (AnimInstance->GetMirrorDataTable())
 			{
-				RootMotion = MirrorRootMotion(RootMotion, AnimInstance->GetMirrorDataTable());
-				PreviewActor.QuantizedTimeRootMotion = MirrorRootMotion(PreviewActor.QuantizedTimeRootMotion, AnimInstance->GetMirrorDataTable());
+				RootMotion = MirrorRootTransform(RootMotion);
+				PreviewActor.QuantizedTimeRootTransform = MirrorRootTransform(PreviewActor.QuantizedTimeRootTransform);
 			}
 
 			if (PreviewActor.PlayTimeOffset != 0.f)
@@ -266,10 +266,10 @@ namespace UE::PoseSearch
 				FTransform OriginRootMotion = PreviewActor.Sampler.ExtractRootTransform(PreviewActor.PlayTimeOffset);
 				if (AnimInstance->GetMirrorDataTable())
 				{
-					OriginRootMotion = MirrorRootMotion(OriginRootMotion, AnimInstance->GetMirrorDataTable());
+					OriginRootMotion = MirrorRootTransform(OriginRootMotion);
 				}
 				RootMotion.SetToRelativeTransform(OriginRootMotion);
-				PreviewActor.QuantizedTimeRootMotion.SetToRelativeTransform(OriginRootMotion);
+				PreviewActor.QuantizedTimeRootTransform.SetToRelativeTransform(OriginRootMotion);
 			}
 
 			PreviewActor.Actor->SetActorTransform(RootMotion);
@@ -293,22 +293,25 @@ namespace UE::PoseSearch
 		PreviewActors.Reset();
 	}
 
-	FTransform FDatabaseViewModel::MirrorRootMotion(
-		FTransform RootMotion,
-		const UMirrorDataTable* MirrorDataTable)
+	FTransform FDatabaseViewModel::MirrorRootTransform(const FTransform& RootTransform)
 	{
+		check(PoseSearchDatabase && PoseSearchDatabase->Schema);
+		if (PoseSearchDatabase->Schema->MirrorDataTable && PoseSearchDatabase->Schema->Skeleton)
+		{
 		const FTransform RootReferenceTransform = PoseSearchDatabase->Schema->Skeleton->GetReferenceSkeleton().GetRefBonePose()[0];
 		const FQuat RootReferenceRotation = RootReferenceTransform.GetRotation();
 
-		const EAxis::Type MirrorAxis = MirrorDataTable->MirrorAxis;
-		FVector T = RootMotion.GetTranslation();
+			const EAxis::Type MirrorAxis = PoseSearchDatabase->Schema->MirrorDataTable->MirrorAxis;
+			FVector T = RootTransform.GetTranslation();
 		T = FAnimationRuntime::MirrorVector(T, MirrorAxis);
-		FQuat Q = RootMotion.GetRotation();
+			FQuat Q = RootTransform.GetRotation();
 		Q = FAnimationRuntime::MirrorQuat(Q, MirrorAxis);
 		Q *= FAnimationRuntime::MirrorQuat(RootReferenceRotation, MirrorAxis).Inverse() * RootReferenceRotation;
 
-		FTransform MirroredRootMotion = FTransform(Q, T, RootMotion.GetScale3D());
-		return MirroredRootMotion;
+			const FTransform MirroredRootTransform = FTransform(Q, T, RootTransform.GetScale3D());
+			return MirroredRootTransform;
+	}
+		return RootTransform;
 	}
 
 	void FDatabaseViewModel::OnSetPoseFeaturesDrawMode(EFeaturesDrawMode DrawMode)
@@ -401,9 +404,6 @@ namespace UE::PoseSearch
 
 		if (FAsyncPoseSearchDatabasesManagement::RequestAsyncBuildIndex(PoseSearchDatabase, ERequestAsyncBuildFlag::ContinueRequest))
 		{
-			FBoneContainer BoneContainer;
-			BoneContainer.InitializeTo(PoseSearchDatabase->Schema->BoneIndicesWithParents, UE::Anim::FCurveFilterSettings(UE::Anim::ECurveFilterMode::DisallowAll), *PoseSearchDatabase->Schema->Skeleton);
-
 			const FPoseSearchIndex& SearchIndex = PoseSearchDatabase->GetSearchIndex();
 			if (SearchIndex.PoseMetadata.IsValidIndex(PoseIdx))
 			{
@@ -413,7 +413,7 @@ namespace UE::PoseSearch
 					const FPoseSearchIndexAsset& IndexAsset = SearchIndex.Assets[IndexAssetIndex];
 					// clamping just to make sure the time interval is correct
 					const float PlayTimeOffset = FMath::Clamp(PoseSearchDatabase->GetAssetTime(PoseIdx), IndexAsset.SamplingInterval.Min, IndexAsset.SamplingInterval.Max);
-					FDatabasePreviewActor PreviewActor = SpawnPreviewActor(IndexAssetIndex, BoneContainer, PlayTimeOffset);
+					FDatabasePreviewActor PreviewActor = SpawnPreviewActor(IndexAssetIndex, PlayTimeOffset);
 					if (PreviewActor.IsValid())
 					{
 						MaxPreviewPlayLength = FMath::Max(MaxPreviewPlayLength, IndexAsset.SamplingInterval.Max - PlayTimeOffset);
@@ -424,7 +424,10 @@ namespace UE::PoseSearch
 				}
 			}
 
-			ParallelFor(PreviewActors.Num(), [this](int32 PreviewActorIndex) { PreviewActors[PreviewActorIndex].Process(); }, ParallelForFlags);
+			FBoneContainer BoneContainer;
+			BoneContainer.InitializeTo(PoseSearchDatabase->Schema->BoneIndicesWithParents, UE::Anim::FCurveFilterSettings(UE::Anim::ECurveFilterMode::DisallowAll), *PoseSearchDatabase->Schema->Skeleton);
+			// capturing BoneContainer by copy, since it has mutable properties
+			ParallelFor(PreviewActors.Num(), [this, BoneContainer](int32 PreviewActorIndex) { PreviewActors[PreviewActorIndex].Process(BoneContainer); }, ParallelForFlags);
 
 			DatabaseDataDetails.Pin()->Reconstruct();
 
@@ -444,9 +447,6 @@ namespace UE::PoseSearch
 
 		if (FAsyncPoseSearchDatabasesManagement::RequestAsyncBuildIndex(PoseSearchDatabase, ERequestAsyncBuildFlag::ContinueRequest))
 		{
-			FBoneContainer BoneContainer;
-			BoneContainer.InitializeTo(PoseSearchDatabase->Schema->BoneIndicesWithParents, UE::Anim::FCurveFilterSettings(UE::Anim::ECurveFilterMode::DisallowAll), *PoseSearchDatabase->Schema->Skeleton);
-
 			TMap<int32, int32> AssociatedAssetIndices;
 			for (int32 i = 0; i < InSelectedNodes.Num(); ++i)
 			{
@@ -461,7 +461,7 @@ namespace UE::PoseSearch
 				{
 					if (const int32* SelectedNodesIndex = AssociatedAssetIndices.Find(IndexAsset.SourceAssetIdx))
 					{
-						FDatabasePreviewActor PreviewActor = SpawnPreviewActor(IndexAssetIndex, BoneContainer, 0.f);
+						FDatabasePreviewActor PreviewActor = SpawnPreviewActor(IndexAssetIndex, 0.f);
 						if (PreviewActor.IsValid())
 						{
 							MaxPreviewPlayLength = FMath::Max(MaxPreviewPlayLength, IndexAsset.SamplingInterval.Max - IndexAsset.SamplingInterval.Min);
@@ -471,7 +471,10 @@ namespace UE::PoseSearch
 				}
 			}
 
-			ParallelFor(PreviewActors.Num(), [this](int32 PreviewActorIndex) { PreviewActors[PreviewActorIndex].Process(); }, ParallelForFlags);
+			FBoneContainer BoneContainer;
+			BoneContainer.InitializeTo(PoseSearchDatabase->Schema->BoneIndicesWithParents, UE::Anim::FCurveFilterSettings(UE::Anim::ECurveFilterMode::DisallowAll), *PoseSearchDatabase->Schema->Skeleton);
+			// capturing BoneContainer by copy, since it has mutable properties
+			ParallelFor(PreviewActors.Num(), [this, BoneContainer](int32 PreviewActorIndex) { PreviewActors[PreviewActorIndex].Process(BoneContainer); }, ParallelForFlags);
 
 			DatabaseDataDetails.Pin()->Reconstruct();
 			UpdatePreviewActors();
diff --git a/Engine/Plugins/Experimental/Animation/PoseSearch/Source/Editor/Public/PoseSearchDatabaseViewModel.h b/Engine/Plugins/Experimental/Animation/PoseSearch/Source/Editor/Public/PoseSearchDatabaseViewModel.h
--- a/Engine/Plugins/Experimental/Animation/PoseSearch/Source/Editor/Public/PoseSearchDatabaseViewModel.h
+++ b/Engine/Plugins/Experimental/Animation/PoseSearch/Source/Editor/Public/PoseSearchDatabaseViewModel.h
@@ -47,11 +47,12 @@ namespace UE::PoseSearch
 		int32 IndexAssetIndex = INDEX_NONE;
 		int32 CurrentPoseIndex = INDEX_NONE;
 		float PlayTimeOffset = 0.f;
+		float CurrentTime = 0.f;
 		FAnimationAssetSampler Sampler;
-		FTransform QuantizedTimeRootMotion = FTransform::Identity;
+		FTransform QuantizedTimeRootTransform = FTransform::Identity;
 
 		bool IsValid() const;
-		void Process();
+		void Process(const FBoneContainer& BoneContainer);
 		UDebugSkelMeshComponent* GetDebugSkelMeshComponent();
 		UAnimPreviewInstance* GetAnimPreviewInstance();
 	};
@@ -171,10 +172,10 @@ namespace UE::PoseSearch
 
 		UWorld* GetWorld();
 
-		FDatabasePreviewActor SpawnPreviewActor(int32 IndexAssetIndex, const FBoneContainer& BoneContainer, float PlayTimeOffset);
+		FDatabasePreviewActor SpawnPreviewActor(int32 IndexAssetIndex, float PlayTimeOffset);
 
 		void UpdatePreviewActors(bool bInTickPlayTime = false);
 
-		FTransform MirrorRootMotion(FTransform RootMotion, const class UMirrorDataTable* MirrorDataTable);
+		FTransform MirrorRootTransform(const FTransform& RootTransform);
 	};
 }
diff --git a/Engine/Plugins/Experimental/Animation/PoseSearch/Source/Runtime/Private/PoseSearchAssetIndexer.cpp b/Engine/Plugins/Experimental/Animation/PoseSearch/Source/Runtime/Private/PoseSearchAssetIndexer.cpp
--- a/Engine/Plugins/Experimental/Animation/PoseSearch/Source/Runtime/Private/PoseSearchAssetIndexer.cpp
+++ b/Engine/Plugins/Experimental/Animation/PoseSearch/Source/Runtime/Private/PoseSearchAssetIndexer.cpp
@@ -107,7 +107,7 @@ void FAssetIndexer::Process(int32 AssetIdx)
 	const float SequenceLength = IndexingContext.AssetSampler->GetPlayLength();
 	for (int32 SampleIdx = GetBeginSampleIdx(); SampleIdx != GetEndSampleIdx(); ++SampleIdx)
 	{
-		const float SampleTime = FMath::Min(SampleIdx * IndexingContext.Schema->GetSamplingInterval(), SequenceLength);
+		const float SampleTime = FMath::Min(CalculateSampleTime(SampleIdx), SequenceLength);
 		float CostAddend = IndexingContext.Schema->BaseCostBias;
 		float ContinuingPoseCostAddend = IndexingContext.Schema->ContinuingPoseCostBias;
 		bool bBlockTransition = false;
@@ -156,7 +156,7 @@ void FAssetIndexer::ComputeStats()
 
 	for (int32 SampleIdx = GetBeginSampleIdx(); SampleIdx != GetEndSampleIdx(); ++SampleIdx)
 	{
-		const float SampleTime = FMath::Min(SampleIdx * IndexingContext.Schema->GetSamplingInterval(), IndexingContext.AssetSampler->GetPlayLength());
+		const float SampleTime = FMath::Min(CalculateSampleTime(SampleIdx), IndexingContext.AssetSampler->GetPlayLength());
 
 		bool AnyClamped = false;
 		const FTransform TrajTransformsPast = GetTransform(SampleTime - SamplingContext->FiniteDelta, AnyClamped);
@@ -296,26 +296,15 @@ FAssetIndexer::CachedEntry& FAssetIndexer::GetEntry(float SampleTime)
 			CurrentTime = FMath::Clamp(CurrentTime, 0.f, PlayLength);
 		}
 
-		FDeltaTimeRecord DeltaTimeRecord;
-		DeltaTimeRecord.Set(CurrentTime, 0.f);
-		// no need to extract root motion here, since we use the precalculated Sample.RootTransform as root transform for the Entry
-		FAnimExtractContext ExtractionCtx(static_cast<double>(CurrentTime), false, DeltaTimeRecord, bLoopable);
-
 		FCompactPose Pose;
-		FBlendedCurve UnusedCurve;
-		FStackAttributeContainer UnusedAtrribute;
-		FAnimationPoseData AnimPoseData = { Pose, UnusedCurve, UnusedAtrribute };
-
-		UnusedCurve.InitFrom(BoneContainer);
 		Pose.SetBoneContainer(&BoneContainer);
-
-		IndexingContext.AssetSampler->ExtractPose(ExtractionCtx, AnimPoseData);
+		IndexingContext.AssetSampler->ExtractPose(CurrentTime, Pose);
 		Pose[FCompactPoseBoneIndex(RootBoneIndexType)].SetIdentity();
 
 		if (IndexingContext.bMirrored && IndexingContext.Schema->MirrorDataTable)
 		{
 			FAnimationRuntime::MirrorPose(
-				AnimPoseData.GetPose(),
+				Pose,
 				IndexingContext.Schema->MirrorDataTable->MirrorAxis,
 				SamplingContext->CompactPoseMirrorBones,
 				SamplingContext->ComponentSpaceRefRotations
@@ -367,13 +356,18 @@ FTransform FAssetIndexer::CalculateComponentSpaceTransform(FAssetIndexer::Cached
 	return Entry.ComponentSpacePose.GetComponentSpaceTransform(CompactBoneIndex);
 }
 
+float FAssetIndexer::CalculateSampleTime(int32 SampleIdx) const
+{
+	return SampleIdx * IndexingContext.Schema->GetSamplingInterval();
+}
+
 FQuat FAssetIndexer::GetSampleRotation(float SampleTimeOffset, int32 SampleIdx, int8 SchemaSampleBoneIdx, int8 SchemaOriginBoneIdx, EPermutationTimeType PermutationTimeType)
 {
 	float PermutationSampleTimeOffset = 0.f;
 	float PermutationOriginTimeOffset = 0.f;
 	UPoseSearchFeatureChannel::GetPermutationTimeOffsets(PermutationTimeType, CalculatePermutationTimeOffset(), PermutationSampleTimeOffset, PermutationOriginTimeOffset);
 
-	const float Time = SampleIdx * IndexingContext.Schema->GetSamplingInterval();
+	const float Time = CalculateSampleTime(SampleIdx);
 	const float SampleTime = Time + SampleTimeOffset + PermutationSampleTimeOffset;
 	const float OriginTime = Time + PermutationOriginTimeOffset;
 
@@ -405,7 +399,7 @@ FVector FAssetIndexer::GetSamplePosition(float SampleTimeOffset, int32 SampleIdx
 	float PermutationOriginTimeOffset = 0.f;
 	UPoseSearchFeatureChannel::GetPermutationTimeOffsets(PermutationTimeType, CalculatePermutationTimeOffset(), PermutationSampleTimeOffset, PermutationOriginTimeOffset);
 
-	const float Time = SampleIdx * IndexingContext.Schema->GetSamplingInterval();
+	const float Time = CalculateSampleTime(SampleIdx);
 	const float SampleTime = Time + SampleTimeOffset + PermutationSampleTimeOffset;
 	const float OriginTime = Time + PermutationOriginTimeOffset;
 	
@@ -450,7 +444,7 @@ FVector FAssetIndexer::GetSampleVelocity(float SampleTimeOffset, int32 SampleIdx
 	float PermutationOriginTimeOffset = 0.f;
 	UPoseSearchFeatureChannel::GetPermutationTimeOffsets(PermutationTimeType, CalculatePermutationTimeOffset(), PermutationSampleTimeOffset, PermutationOriginTimeOffset);
 
-	const float Time = SampleIdx * IndexingContext.Schema->GetSamplingInterval();
+	const float Time = CalculateSampleTime(SampleIdx);
 	const float SampleTime = Time + SampleTimeOffset + PermutationSampleTimeOffset;
 	const float OriginTime = Time + PermutationOriginTimeOffset;
 	const float FiniteDelta = IndexingContext.SamplingContext->FiniteDelta;
diff --git a/Engine/Plugins/Experimental/Animation/PoseSearch/Source/Runtime/Private/PoseSearchAssetSampler.cpp b/Engine/Plugins/Experimental/Animation/PoseSearch/Source/Runtime/Private/PoseSearchAssetSampler.cpp
--- a/Engine/Plugins/Experimental/Animation/PoseSearch/Source/Runtime/Private/PoseSearchAssetSampler.cpp
+++ b/Engine/Plugins/Experimental/Animation/PoseSearch/Source/Runtime/Private/PoseSearchAssetSampler.cpp
@@ -245,16 +245,15 @@ static int32 GetHighestWeightSample(const TArray<struct FBlendSampleData>& Sampl
 
 //////////////////////////////////////////////////////////////////////////
 // FAssetSamplerBase
-FAnimationAssetSampler::FAnimationAssetSampler(TWeakObjectPtr<const UAnimationAsset> InAnimationAsset, const FVector& InBlendParameters, const FBoneContainer& InBoneContainer, int32 InRootTransformSamplingRate)
+FAnimationAssetSampler::FAnimationAssetSampler(TWeakObjectPtr<const UAnimationAsset> InAnimationAsset, const FVector& InBlendParameters, int32 InRootTransformSamplingRate)
 {
-	Init(InAnimationAsset, InBlendParameters, InBoneContainer, InRootTransformSamplingRate);
+	Init(InAnimationAsset, InBlendParameters, InRootTransformSamplingRate);
 }
 
-void FAnimationAssetSampler::Init(TWeakObjectPtr<const UAnimationAsset> InAnimationAsset, const FVector& InBlendParameters, const FBoneContainer& InBoneContainer, int32 InRootTransformSamplingRate)
+void FAnimationAssetSampler::Init(TWeakObjectPtr<const UAnimationAsset> InAnimationAsset, const FVector& InBlendParameters, int32 InRootTransformSamplingRate)
 {
 	AnimationAsset = InAnimationAsset;
 	BlendParameters = InBlendParameters;
-	BoneContainer = InBoneContainer;
 	RootTransformSamplingRate = InRootTransformSamplingRate;
 
 	if (const UAnimationAsset* AnimAsset = AnimationAsset.Get())
@@ -390,6 +389,25 @@ void FAnimationAssetSampler::ExtractPose(const FAnimExtractContext& ExtractionCt
 	}
 }
 
+void FAnimationAssetSampler::ExtractPose(float Time, FCompactPose& OutPose) const
+{
+	using namespace UE::Anim;
+
+	FBlendedCurve UnusedCurve;
+	FStackAttributeContainer UnusedAtrribute;
+	FAnimationPoseData AnimPoseData = { OutPose, UnusedCurve, UnusedAtrribute };
+
+	check(OutPose.IsValid());
+	FBoneContainer& BoneContainer = OutPose.GetBoneContainer();
+	UnusedCurve.InitFrom(BoneContainer);
+
+	FDeltaTimeRecord DeltaTimeRecord;
+	DeltaTimeRecord.Set(Time, 0.f);
+	FAnimExtractContext ExtractionCtx(double(Time), false, DeltaTimeRecord, false);
+
+	ExtractPose(ExtractionCtx, AnimPoseData);
+}
+
 FTransform FAnimationAssetSampler::ExtractRootTransform(float Time) const
 {
 	FTransform RootTransform = FTransform::Identity;
@@ -530,7 +548,7 @@ FTransform FAnimationAssetSampler::ExtractRootTransform(float Time) const
 	return RootTransform;
 }
 
-void FAnimationAssetSampler::Process()
+void FAnimationAssetSampler::Process(const FBoneContainer& BoneContainer)
 {
 	if (const UBlendSpace* BlendSpace = Cast<UBlendSpace>(AnimationAsset.Get()))
 	{
diff --git a/Engine/Plugins/Experimental/Animation/PoseSearch/Source/Runtime/Private/PoseSearchDatabaseIndexingContext.cpp b/Engine/Plugins/Experimental/Animation/PoseSearch/Source/Runtime/Private/PoseSearchDatabaseIndexingContext.cpp
--- a/Engine/Plugins/Experimental/Animation/PoseSearch/Source/Runtime/Private/PoseSearchDatabaseIndexingContext.cpp
+++ b/Engine/Plugins/Experimental/Animation/PoseSearch/Source/Runtime/Private/PoseSearchDatabaseIndexingContext.cpp
@@ -78,7 +78,7 @@ bool FDatabaseIndexingContext::IndexDatabase(FPoseSearchIndexBase& SearchIndexBa
 						if (!SamplerMap.Contains({ DatabaseBlendSpace->BlendSpace, BlendParameters }))
 						{
 							SamplerMap.Add({ DatabaseBlendSpace->BlendSpace, BlendParameters }, Samplers.Num());
-							Samplers.Emplace(DatabaseBlendSpace->BlendSpace, BlendParameters, BoneContainer);
+							Samplers.Emplace(DatabaseBlendSpace->BlendSpace, BlendParameters);
 						}
 					}
 				}
@@ -97,7 +97,8 @@ bool FDatabaseIndexingContext::IndexDatabase(FPoseSearchIndexBase& SearchIndexBa
 		}
 	}
 
-	ParallelFor(Samplers.Num(), [this](int32 SamplerIdx) { Samplers[SamplerIdx].Process(); }, ParallelForFlags);
+	// capturing BoneContainer by copy, since it has mutable properties
+	ParallelFor(Samplers.Num(), [this, BoneContainer](int32 SamplerIdx) { Samplers[SamplerIdx].Process(BoneContainer); }, ParallelForFlags);
 
 	if (Owner.IsCanceled())
 	{
diff --git a/Engine/Plugins/Experimental/Animation/PoseSearch/Source/Runtime/Private/PoseSearchLibrary.cpp b/Engine/Plugins/Experimental/Animation/PoseSearch/Source/Runtime/Private/PoseSearchLibrary.cpp
--- a/Engine/Plugins/Experimental/Animation/PoseSearch/Source/Runtime/Private/PoseSearchLibrary.cpp
+++ b/Engine/Plugins/Experimental/Animation/PoseSearch/Source/Runtime/Private/PoseSearchLibrary.cpp
@@ -553,15 +553,7 @@ void UPoseSearchLibrary::MotionMatch(
 		{
 			const FBoneContainer& BoneContainer = AnimInstance->GetRequiredBonesOnAnyThread();
 			// @todo... add input BlendParameters to support sampling FutureAnimation blendspaces
-			const FAnimationAssetSampler Sampler(FutureAnimation, FVector::ZeroVector, BoneContainer);
-
-			FCompactPose Pose;
-			FBlendedCurve UnusedCurve;
-			FStackAttributeContainer UnusedAtrribute;
-			FAnimationPoseData AnimPoseData = { Pose, UnusedCurve, UnusedAtrribute };
-
-			UnusedCurve.InitFrom(BoneContainer);
-			Pose.SetBoneContainer(&BoneContainer);
+			const FAnimationAssetSampler Sampler(FutureAnimation, FVector::ZeroVector);
 
 			if (FutureAnimationStartTime < FiniteDelta)
 			{
@@ -582,11 +574,9 @@ void UPoseSearchLibrary::MotionMatch(
 				const float ExtractionTime = FutureAnimationStartTime + (i - 1) * FiniteDelta;
 				const float FutureAnimationTime = TimeToFutureAnimationStart + (i - 1) * FiniteDelta;
 
-				FDeltaTimeRecord DeltaTimeRecord;
-				DeltaTimeRecord.Set(ExtractionTime - FiniteDelta, FiniteDelta);
-				FAnimExtractContext ExtractionCtx(double(ExtractionTime), false, DeltaTimeRecord, false);
-
-				Sampler.ExtractPose(ExtractionCtx, AnimPoseData);
+				FCompactPose Pose;
+				Pose.SetBoneContainer(&BoneContainer);
+				Sampler.ExtractPose(ExtractionTime, Pose);
 
 				FCSPose<FCompactPose> ComponentSpacePose;
 				ComponentSpacePose.InitPose(Pose);
diff --git a/Engine/Plugins/Experimental/Animation/PoseSearch/Source/Runtime/Public/PoseSearch/PoseSearchAssetIndexer.h b/Engine/Plugins/Experimental/Animation/PoseSearch/Source/Runtime/Public/PoseSearch/PoseSearchAssetIndexer.h
--- a/Engine/Plugins/Experimental/Animation/PoseSearch/Source/Runtime/Public/PoseSearch/PoseSearchAssetIndexer.h
+++ b/Engine/Plugins/Experimental/Animation/PoseSearch/Source/Runtime/Public/PoseSearch/PoseSearchAssetIndexer.h
@@ -59,6 +59,7 @@ public:
 	TArrayView<float> GetPoseVector(int32 SampleIdx) const;
 	const UPoseSearchSchema* GetSchema() const;
 	float CalculatePermutationTimeOffset() const;
+	float CalculateSampleTime(int32 SampleIdx) const;
 
 private:
 	int32 GetVectorIdx(int32 SampleIdx) const;
diff --git a/Engine/Plugins/Experimental/Animation/PoseSearch/Source/Runtime/Public/PoseSearch/PoseSearchAssetSampler.h b/Engine/Plugins/Experimental/Animation/PoseSearch/Source/Runtime/Public/PoseSearch/PoseSearchAssetSampler.h
--- a/Engine/Plugins/Experimental/Animation/PoseSearch/Source/Runtime/Public/PoseSearch/PoseSearchAssetSampler.h
+++ b/Engine/Plugins/Experimental/Animation/PoseSearch/Source/Runtime/Public/PoseSearch/PoseSearchAssetSampler.h
@@ -17,8 +17,8 @@ namespace UE::PoseSearch
  */
 struct POSESEARCH_API FAnimationAssetSampler
 {
-	FAnimationAssetSampler(TWeakObjectPtr<const UAnimationAsset> InAnimationAsset = nullptr, const FVector& InBlendParameters = FVector::ZeroVector, const FBoneContainer& InBoneContainer = FBoneContainer(), int32 InRootTransformSamplingRate = 30);
-	void Init(TWeakObjectPtr<const UAnimationAsset> InAnimationAsset, const FVector& InBlendParameters = FVector::ZeroVector, const FBoneContainer& InBoneContainer = FBoneContainer(), int32 InRootTransformSamplingRate = 30);
+	FAnimationAssetSampler(TWeakObjectPtr<const UAnimationAsset> InAnimationAsset = nullptr, const FVector& InBlendParameters = FVector::ZeroVector, int32 InRootTransformSamplingRate = 30);
+	void Init(TWeakObjectPtr<const UAnimationAsset> InAnimationAsset, const FVector& InBlendParameters = FVector::ZeroVector, int32 InRootTransformSamplingRate = 30);
 
 	bool IsInitialized() const;
 	float GetPlayLength() const;
@@ -31,6 +31,9 @@ struct POSESEARCH_API FAnimationAssetSampler
 	// Extracts pose for this asset for a given context
 	void ExtractPose(const FAnimExtractContext& ExtractionCtx, FAnimationPoseData& OutAnimPoseData) const;
 
+	// Extracts pose for this asset at a given Time
+	void ExtractPose(float Time, FCompactPose& OutPose) const;
+
 	// Extracts root transform at the given time, using the extremities of the sequence to extrapolate beyond the 
 	// sequence limits when Time is less than zero or greater than the sequence length.
 	FTransform ExtractRootTransform(float Time) const;
@@ -41,14 +44,13 @@ struct POSESEARCH_API FAnimationAssetSampler
 
 	const UAnimationAsset* GetAsset() const;
 
-	void Process();
+	void Process(const FBoneContainer& BoneContainer);
 
 protected:
 	TWeakObjectPtr<const UAnimationAsset> AnimationAsset;
 
 	// members used to sample blend spaces only!
 	FVector BlendParameters = FVector::ZeroVector;
-	FBoneContainer BoneContainer;
 	int32 RootTransformSamplingRate = 30;
 	float CachedPlayLength = 0.f;
 	TArray<FTransform> AccumulatedRootTransform;
