motion matching - adding UPoseSearchFeatureChannel_SamplingTime

diff --git a/Engine/Plugins/Experimental/Animation/PoseSearch/Source/Runtime/Private/PoseSearchFeatureChannel_SamplingTime.cpp b/Engine/Plugins/Experimental/Animation/PoseSearch/Source/Runtime/Private/PoseSearchFeatureChannel_SamplingTime.cpp
new file mode 100644
--- /dev/null
+++ b/Engine/Plugins/Experimental/Animation/PoseSearch/Source/Runtime/Private/PoseSearchFeatureChannel_SamplingTime.cpp
@@ -0,0 +1,52 @@
+// Copyright Epic Games, Inc. All Rights Reserved.
+
+#include "PoseSearchFeatureChannel_SamplingTime.h"
+#include "PoseSearch/PoseSearchAssetIndexer.h"
+#include "PoseSearch/PoseSearchContext.h"
+#include "PoseSearch/PoseSearchSchema.h"
+
+void UPoseSearchFeatureChannel_SamplingTime::Finalize(UPoseSearchSchema* Schema)
+{
+	ChannelDataOffset = Schema->SchemaCardinality;
+	ChannelCardinality = 1;
+	Schema->SchemaCardinality += ChannelCardinality;
+}
+
+void UPoseSearchFeatureChannel_SamplingTime::BuildQuery(UE::PoseSearch::FSearchContext& SearchContext, UE::PoseSearch::FFeatureVectorBuilder& InOutQuery) const
+{
+	using namespace UE::PoseSearch;
+	FFeatureVectorHelper::EncodeFloat(InOutQuery.EditValues(), ChannelDataOffset, TimeToMatch);
+}
+
+#if WITH_EDITOR
+void UPoseSearchFeatureChannel_SamplingTime::FillWeights(TArray<float>& Weights) const
+{
+	for (int32 i = 0; i < ChannelCardinality; ++i)
+	{
+		Weights[ChannelDataOffset + i] = Weight;
+	}
+}
+
+void UPoseSearchFeatureChannel_SamplingTime::IndexAsset(UE::PoseSearch::FAssetIndexer& Indexer) const
+{
+	using namespace UE::PoseSearch;
+
+	for (int32 SampleIdx = Indexer.GetBeginSampleIdx(); SampleIdx != Indexer.GetEndSampleIdx(); ++SampleIdx)
+	{
+		FFeatureVectorHelper::EncodeFloat(Indexer.GetPoseVector(SampleIdx), ChannelDataOffset, Indexer.CalculateSampleTime(SampleIdx));
+	}
+}
+
+FString UPoseSearchFeatureChannel_SamplingTime::GetLabel() const
+{
+	TStringBuilder<256> Label;
+	if (const UPoseSearchFeatureChannel* OuterChannel = Cast<UPoseSearchFeatureChannel>(GetOuter()))
+	{
+		Label.Append(OuterChannel->GetLabel());
+		Label.Append(TEXT("_"));
+	}
+
+	Label.Append(TEXT("SamplTime"));
+	return Label.ToString();
+}
+#endif
\ No newline at end of file
diff --git a/Engine/Plugins/Experimental/Animation/PoseSearch/Source/Runtime/Private/PoseSearchFeatureChannel_SamplingTime.h b/Engine/Plugins/Experimental/Animation/PoseSearch/Source/Runtime/Private/PoseSearchFeatureChannel_SamplingTime.h
new file mode 100644
--- /dev/null
+++ b/Engine/Plugins/Experimental/Animation/PoseSearch/Source/Runtime/Private/PoseSearchFeatureChannel_SamplingTime.h
@@ -0,0 +1,30 @@
+// Copyright Epic Games, Inc. All Rights Reserved.
+
+#pragma once
+
+#include "PoseSearch/PoseSearchFeatureChannel.h"
+#include "PoseSearchFeatureChannel_SamplingTime.generated.h"
+
+// this channel is mainly for debug purposes to augment the features data with the sampling time (default weight is set to zero to be irrelevant during searches)
+UCLASS(EditInlineNew, meta = (DisplayName = "Sampling Time Channel"), CollapseCategories)
+class POSESEARCH_API UPoseSearchFeatureChannel_SamplingTime : public UPoseSearchFeatureChannel
+{
+	GENERATED_BODY()
+
+public:
+	UPROPERTY(EditAnywhere, Category = "Settings")
+	float Weight = 0.f;
+
+	UPROPERTY(EditAnywhere, Category = "Settings")
+	float TimeToMatch = 0.f;
+
+	// UPoseSearchFeatureChannel interface
+	virtual void Finalize(UPoseSearchSchema* Schema) override;
+	virtual void BuildQuery(UE::PoseSearch::FSearchContext& SearchContext, UE::PoseSearch::FFeatureVectorBuilder& InOutQuery) const override;
+
+#if WITH_EDITOR
+	virtual void FillWeights(TArray<float>& Weights) const override;
+	virtual void IndexAsset(UE::PoseSearch::FAssetIndexer& Indexer) const override;
+	virtual FString GetLabel() const override;
+#endif
+};
