motion matching
- database blend space support fixes
- added cvar a.DatabasePreview.DebugDrawSamplerSize to help visually validating the database sampler

diff --git a/Engine/Plugins/Experimental/Animation/PoseSearch/Source/Editor/Private/PoseSearchDatabasePreviewScene.cpp b/Engine/Plugins/Experimental/Animation/PoseSearch/Source/Editor/Private/PoseSearchDatabasePreviewScene.cpp
--- a/Engine/Plugins/Experimental/Animation/PoseSearch/Source/Editor/Private/PoseSearchDatabasePreviewScene.cpp
+++ b/Engine/Plugins/Experimental/Animation/PoseSearch/Source/Editor/Private/PoseSearchDatabasePreviewScene.cpp
@@ -2,6 +2,7 @@
 
 #include "PoseSearchDatabasePreviewScene.h"
 #include "Animation/DebugSkelMeshComponent.h"
+#include "AnimPreviewInstance.h"
 #include "Components/StaticMeshComponent.h"
 #include "EngineUtils.h"
 #include "GameFramework/WorldSettings.h"
@@ -11,6 +12,10 @@
 #include "PoseSearchDatabaseEditor.h"
 #include "PoseSearchDatabaseViewModel.h"
 
+#if ENABLE_ANIM_DEBUG
+static TAutoConsoleVariable<float> CVarDatabasePreviewDebugDrawSamplerSize(TEXT("a.DatabasePreview.DebugDrawSamplerSize"), 0.f, TEXT("Debug Draw Sampler Positions Size"));
+#endif
+
 namespace UE::PoseSearch
 {
 	FDatabasePreviewScene::FDatabasePreviewScene(
@@ -61,6 +66,11 @@ namespace UE::PoseSearch
 		if (ViewModel->IsPoseFeaturesDrawMode(EFeaturesDrawMode::All | EFeaturesDrawMode::Detailed) && !ViewModel->GetPreviewActors().IsEmpty() &&
 			FAsyncPoseSearchDatabasesManagement::RequestAsyncBuildIndex(Database, ERequestAsyncBuildFlag::ContinueRequest))
 		{
+#if ENABLE_ANIM_DEBUG
+			// memory marker required for PreviewActor.Sampler.ExtractPose
+			FMemMark Mark(FMemStack::Get());
+#endif // ENABLE_ANIM_DEBUG
+
 			bool bDrawQueryVector = ViewModel->ShouldDrawQueryVector();
 			for (FDatabasePreviewActor& PreviewActor : ViewModel->GetPreviewActors())
 			{
@@ -76,6 +86,28 @@ namespace UE::PoseSearch
 							DrawParams.DrawFeatureVector(ViewModel->GetQueryVector());
 							bDrawQueryVector = false;
 						}
+
+#if ENABLE_ANIM_DEBUG
+						const float DebugDrawSamplerSize = CVarDatabasePreviewDebugDrawSamplerSize.GetValueOnAnyThread();
+						if (DebugDrawSamplerSize > UE_KINDA_SMALL_NUMBER)
+						{
+							// drawing the pose extracted from the Sampler to visually compare with the pose features and the mesh drawing
+							FCompactPose Pose;
+							Pose.SetBoneContainer(&PreviewActor.GetAnimPreviewInstance()->GetRequiredBonesOnAnyThread());
+							PreviewActor.Sampler.ExtractPose(PreviewActor.CurrentTime, Pose);
+
+							const FTransform RootTransform = PreviewActor.Sampler.ExtractRootTransform(PreviewActor.CurrentTime);
+
+							FCSPose<FCompactPose> ComponentSpacePose;
+							ComponentSpacePose.InitPose(Pose);
+
+							for (int32 BoneIndex = 0; BoneIndex < Pose.GetNumBones(); ++BoneIndex)
+							{
+								const FTransform BoneWorldTransforms = ComponentSpacePose.GetComponentSpaceTransform(FCompactPoseBoneIndex(BoneIndex)) * RootTransform;
+								DrawParams.DrawPoint(BoneWorldTransforms.GetTranslation(), FColor::Red, DebugDrawSamplerSize);
+							}
+						}
+#endif // ENABLE_ANIM_DEBUG
 					}
 				}
 			}
diff --git a/Engine/Plugins/Experimental/Animation/PoseSearch/Source/Editor/Private/PoseSearchDatabaseViewModel.cpp b/Engine/Plugins/Experimental/Animation/PoseSearch/Source/Editor/Private/PoseSearchDatabaseViewModel.cpp
--- a/Engine/Plugins/Experimental/Animation/PoseSearch/Source/Editor/Private/PoseSearchDatabaseViewModel.cpp
+++ b/Engine/Plugins/Experimental/Animation/PoseSearch/Source/Editor/Private/PoseSearchDatabaseViewModel.cpp
@@ -127,7 +127,7 @@ namespace UE::PoseSearch
 		PlayTime = MaxPreviewPlayLength;
 	}
 
-	FDatabasePreviewActor FDatabaseViewModel::SpawnPreviewActor(int32 IndexAssetIndex, float PlayTimeOffset)
+	FDatabasePreviewActor FDatabaseViewModel::SpawnPreviewActor(int32 IndexAssetIndex, int32 PoseIdxForTimeOffset)
 	{
 		FDatabasePreviewActor PreviewActor;
 
@@ -148,7 +148,7 @@ namespace UE::PoseSearch
 		PreviewActor.Sampler.Init(DatabaseAnimationAssetBase->GetAnimationAsset(), IndexAsset.BlendParameters);
 		PreviewActor.IndexAssetIndex = IndexAssetIndex;
 		PreviewActor.CurrentPoseIndex = INDEX_NONE;
-		PreviewActor.PlayTimeOffset = PlayTimeOffset;
+		PreviewActor.PlayTimeOffset = PoseIdxForTimeOffset < 0 ? 0.f : PreviewActor.Sampler.ToRealTime(PoseSearchDatabase->GetAssetTime(PoseIdxForTimeOffset));
 
 		FActorSpawnParameters Params;
 		Params.SpawnCollisionHandlingOverride = ESpawnActorCollisionHandlingMethod::AdjustIfPossibleButAlwaysSpawn;
@@ -240,16 +240,14 @@ namespace UE::PoseSearch
 			// time to pose index
 			PreviewActor.CurrentPoseIndex = PoseSearchDatabase->GetPoseIndexFromTime(PreviewActor.CurrentTime, IndexAsset);
 
-			const float TimeToAssetTimeMultiplier = PreviewActor.Sampler.GetTimeToAssetTimeMultiplier();
-			check(!FMath::IsNearlyZero(TimeToAssetTimeMultiplier));
-			const float QuantizedTime = PreviewActor.CurrentPoseIndex >= 0 ? PoseSearchDatabase->GetAssetTime(PreviewActor.CurrentPoseIndex) / TimeToAssetTimeMultiplier : PreviewActor.CurrentTime;
-
+			const float QuantizedTime = PreviewActor.CurrentPoseIndex >= 0 ? PreviewActor.Sampler.ToRealTime(PoseSearchDatabase->GetAssetTime(PreviewActor.CurrentPoseIndex)) : PreviewActor.CurrentTime;
 			if (bQuantizeAnimationToPoseData)
 			{
 				PreviewActor.CurrentTime = QuantizedTime;
 			}
 
-			AnimInstance->SetPosition(PreviewActor.CurrentTime * TimeToAssetTimeMultiplier);
+			// SetPosition is in [0..1] range for blendspaces
+			AnimInstance->SetPosition(PreviewActor.Sampler.ToNormalizedTime(PreviewActor.CurrentTime));
 			AnimInstance->SetPlayRate(0.f);
 			AnimInstance->SetBlendSpacePosition(IndexAsset.BlendParameters);
 
@@ -410,14 +408,12 @@ namespace UE::PoseSearch
 				const uint32 IndexAssetIndex = SearchIndex.PoseMetadata[PoseIdx].GetAssetIndex();
 				if (SearchIndex.Assets.IsValidIndex(IndexAssetIndex))
 				{
-					const FPoseSearchIndexAsset& IndexAsset = SearchIndex.Assets[IndexAssetIndex];
-					// clamping just to make sure the time interval is correct
-					const float PlayTimeOffset = FMath::Clamp(PoseSearchDatabase->GetAssetTime(PoseIdx), IndexAsset.SamplingInterval.Min, IndexAsset.SamplingInterval.Max);
-					FDatabasePreviewActor PreviewActor = SpawnPreviewActor(IndexAssetIndex, PlayTimeOffset);
+					FDatabasePreviewActor PreviewActor = SpawnPreviewActor(IndexAssetIndex, PoseIdx);
 					if (PreviewActor.IsValid())
 					{
-						MaxPreviewPlayLength = FMath::Max(MaxPreviewPlayLength, IndexAsset.SamplingInterval.Max - PlayTimeOffset);
-						MinPreviewPlayLength = FMath::Min(MinPreviewPlayLength, IndexAsset.SamplingInterval.Min - PlayTimeOffset);
+						const FPoseSearchIndexAsset& IndexAsset = SearchIndex.Assets[IndexAssetIndex];
+						MaxPreviewPlayLength = FMath::Max(MaxPreviewPlayLength, IndexAsset.SamplingInterval.Max - PreviewActor.PlayTimeOffset);
+						MinPreviewPlayLength = FMath::Min(MinPreviewPlayLength, IndexAsset.SamplingInterval.Min - PreviewActor.PlayTimeOffset);
 						PreviewActors.Add(PreviewActor);
 						SelectedSourceAssetIdx = IndexAsset.SourceAssetIdx;
 					}
@@ -461,7 +457,7 @@ namespace UE::PoseSearch
 				{
 					if (const int32* SelectedNodesIndex = AssociatedAssetIndices.Find(IndexAsset.SourceAssetIdx))
 					{
-						FDatabasePreviewActor PreviewActor = SpawnPreviewActor(IndexAssetIndex, 0.f);
+						FDatabasePreviewActor PreviewActor = SpawnPreviewActor(IndexAssetIndex);
 						if (PreviewActor.IsValid())
 						{
 							MaxPreviewPlayLength = FMath::Max(MaxPreviewPlayLength, IndexAsset.SamplingInterval.Max - IndexAsset.SamplingInterval.Min);
@@ -552,7 +548,7 @@ namespace UE::PoseSearch
 					const FPoseSearchIndexAsset& IndexAsset = SearchIndex.Assets[PreviewActor.IndexAssetIndex];
 					if (IndexAsset.SourceAssetIdx == SourceAssetIdx)
 					{
-						CurrentPlayTime = PlayTime + IndexAsset.SamplingInterval.Min + PreviewActor.PlayTimeOffset;
+						CurrentPlayTime = PreviewActor.Sampler.ToNormalizedTime(PlayTime + IndexAsset.SamplingInterval.Min + PreviewActor.PlayTimeOffset);
 						BlendParameters = IndexAsset.BlendParameters;
 						return true;
 					}
@@ -565,6 +561,12 @@ namespace UE::PoseSearch
 				{
 					CurrentPlayTime = PlayTime + IndexAsset.SamplingInterval.Min;
 					BlendParameters = IndexAsset.BlendParameters;
+
+					const bool bIsBlendSpace = PoseSearchDatabase->GetAnimationAssetStruct(IndexAsset).GetPtr<FPoseSearchDatabaseBlendSpace>() != nullptr;
+					if (bIsBlendSpace && !FMath::IsNearlyEqual(MaxPreviewPlayLength, MinPreviewPlayLength))
+					{
+						CurrentPlayTime = (CurrentPlayTime - MaxPreviewPlayLength) / (MaxPreviewPlayLength - MinPreviewPlayLength);
+					}
 					return true;
 				}
 			}
diff --git a/Engine/Plugins/Experimental/Animation/PoseSearch/Source/Editor/Public/PoseSearchDatabaseViewModel.h b/Engine/Plugins/Experimental/Animation/PoseSearch/Source/Editor/Public/PoseSearchDatabaseViewModel.h
--- a/Engine/Plugins/Experimental/Animation/PoseSearch/Source/Editor/Public/PoseSearchDatabaseViewModel.h
+++ b/Engine/Plugins/Experimental/Animation/PoseSearch/Source/Editor/Public/PoseSearchDatabaseViewModel.h
@@ -172,7 +172,7 @@ namespace UE::PoseSearch
 
 		UWorld* GetWorld();
 
-		FDatabasePreviewActor SpawnPreviewActor(int32 IndexAssetIndex, float PlayTimeOffset);
+		FDatabasePreviewActor SpawnPreviewActor(int32 IndexAssetIndex, int32 PoseIdxForTimeOffset = -1);
 
 		void UpdatePreviewActors(bool bInTickPlayTime = false);
 
diff --git a/Engine/Plugins/Experimental/Animation/PoseSearch/Source/Runtime/Private/AnimNode_BlendStack.cpp b/Engine/Plugins/Experimental/Animation/PoseSearch/Source/Runtime/Private/AnimNode_BlendStack.cpp
--- a/Engine/Plugins/Experimental/Animation/PoseSearch/Source/Runtime/Private/AnimNode_BlendStack.cpp
+++ b/Engine/Plugins/Experimental/Animation/PoseSearch/Source/Runtime/Private/AnimNode_BlendStack.cpp
@@ -88,6 +88,9 @@ void FPoseSearchAnimPlayer::Initialize(UAnimationAsset* AnimationAsset, float Ac
 	}
 	else if (UBlendSpace* BlendSpace = Cast<UBlendSpace>(AnimationAsset))
 	{
+		// making sure AccumulatedTime is in normalized space
+		check(AccumulatedTime >= 0.f && AccumulatedTime <= 1.f);
+
 		BlendSpacePlayerNode.SetResetPlayTimeWhenBlendSpaceChanges(false /*!bReset*/);
 		BlendSpacePlayerNode.SetAccumulatedTime(AccumulatedTime);
 		BlendSpacePlayerNode.SetBlendSpace(BlendSpace);
@@ -197,6 +200,8 @@ float FPoseSearchAnimPlayer::GetAccumulatedTime() const
 	
 	if (BlendSpacePlayerNode.GetBlendSpace())
 	{
+		// making sure BlendSpacePlayerNode.GetAccumulatedTime() is in normalized space
+		check(BlendSpacePlayerNode.GetAccumulatedTime() >= 0.f && BlendSpacePlayerNode.GetAccumulatedTime() <= 1.f);
 		return BlendSpacePlayerNode.GetAccumulatedTime();
 	}
 
diff --git a/Engine/Plugins/Experimental/Animation/PoseSearch/Source/Runtime/Private/PoseSearchAssetIndexer.cpp b/Engine/Plugins/Experimental/Animation/PoseSearch/Source/Runtime/Private/PoseSearchAssetIndexer.cpp
--- a/Engine/Plugins/Experimental/Animation/PoseSearch/Source/Runtime/Private/PoseSearchAssetIndexer.cpp
+++ b/Engine/Plugins/Experimental/Animation/PoseSearch/Source/Runtime/Private/PoseSearchAssetIndexer.cpp
@@ -75,19 +75,20 @@ static FSamplingParam WrapOrClampSamplingParam(bool bCanWrap, float SamplingPara
 
 //////////////////////////////////////////////////////////////////////////
 // FAssetIndexer
-FAssetIndexer::FAssetIndexer(const FAssetIndexingContext& InIndexingContext, const FBoneContainer& InBoneContainer, const FPoseSearchIndexAsset& InSearchIndexAsset)
-: SearchIndexAsset(InSearchIndexAsset)
+FAssetIndexer::FAssetIndexer(const FBoneContainer& InBoneContainer, const FPoseSearchIndexAsset& InSearchIndexAsset, 
+	const FAssetSamplingContext& InSamplingContext, const UPoseSearchSchema& InSchema, const FAnimationAssetSampler& InAssetSampler)
+: BoneContainer(InBoneContainer)
+, CachedEntries()
+, SearchIndexAsset(InSearchIndexAsset)
+, SamplingContext(InSamplingContext)
+, Schema(InSchema)
+, AssetSampler(InAssetSampler)
 {
-	check(InIndexingContext.Schema);
-	check(InIndexingContext.Schema->IsValid());
-	check(InIndexingContext.AssetSampler);
+	check(Schema.IsValid());
 
-	BoneContainer = InBoneContainer;
-	IndexingContext = InIndexingContext;
-
-	FirstIndexedSample = FMath::FloorToInt(IndexingContext.RequestedSamplingRange.Min * IndexingContext.Schema->SampleRate);
-	// @todo: parhaps we should use FMath::CeilToInt to avoid sampling over teh length of the aniamation
-	LastIndexedSample = FMath::Max(0, FMath::CeilToInt(IndexingContext.RequestedSamplingRange.Max * IndexingContext.Schema->SampleRate));
+	FirstIndexedSample = FMath::FloorToInt(SearchIndexAsset.SamplingInterval.Min * Schema.SampleRate);
+	// @todo: parhaps we should use FMath::CeilToInt to avoid sampling over teh length of the animation
+	LastIndexedSample = FMath::Max(0, FMath::CeilToInt(SearchIndexAsset.SamplingInterval.Max * Schema.SampleRate));
 }
 
 void FAssetIndexer::AssignWorkingData(TArrayView<float> InOutFeatureVectorTable, TArrayView<FPoseSearchPoseMetadata> InOutPoseMetadata)
@@ -98,22 +99,21 @@ void FAssetIndexer::AssignWorkingData(TArrayView<float> InOutFeatureVectorTable,
 
 void FAssetIndexer::Process(int32 AssetIdx)
 {
-	check(GetSchema()->IsValid());
-	check(IndexingContext.AssetSampler);
+	check(Schema.IsValid());
 
 	FMemMark Mark(FMemStack::Get());
 
 	// Generate pose metadata
-	const float SequenceLength = IndexingContext.AssetSampler->GetPlayLength();
+	const float SequenceLength = AssetSampler.GetPlayLength();
 	for (int32 SampleIdx = GetBeginSampleIdx(); SampleIdx != GetEndSampleIdx(); ++SampleIdx)
 	{
 		const float SampleTime = FMath::Min(CalculateSampleTime(SampleIdx), SequenceLength);
-		float CostAddend = IndexingContext.Schema->BaseCostBias;
-		float ContinuingPoseCostAddend = IndexingContext.Schema->ContinuingPoseCostBias;
+		float CostAddend = Schema.BaseCostBias;
+		float ContinuingPoseCostAddend = Schema.ContinuingPoseCostBias;
 		bool bBlockTransition = false;
 
 		TArray<UAnimNotifyState_PoseSearchBase*> NotifyStates;
-		IndexingContext.AssetSampler->ExtractPoseSearchNotifyStates(SampleTime, NotifyStates);
+		AssetSampler.ExtractPoseSearchNotifyStates(SampleTime, NotifyStates);
 		for (const UAnimNotifyState_PoseSearchBase* PoseSearchNotify : NotifyStates)
 		{
 			if (PoseSearchNotify->GetClass()->IsChildOf<UAnimNotifyState_PoseSearchBlockTransition>())
@@ -126,18 +126,18 @@ void FAssetIndexer::Process(int32 AssetIdx)
 			}
 		}
 
-		if (IndexingContext.AssetSampler->IsLoopable())
+		if (AssetSampler.IsLoopable())
 		{
-			CostAddend += IndexingContext.Schema->LoopingCostBias;
+			CostAddend += Schema.LoopingCostBias;
 		}
 
 		PoseMetadata[GetVectorIdx(SampleIdx)].Init(AssetIdx, bBlockTransition, CostAddend);
 	}
 
 	// Generate pose features data
-	if (IndexingContext.Schema->SchemaCardinality > 0)
+	if (Schema.SchemaCardinality > 0)
 	{
-		for (const TObjectPtr<UPoseSearchFeatureChannel>& ChannelPtr : IndexingContext.Schema->GetChannels())
+		for (const TObjectPtr<UPoseSearchFeatureChannel>& ChannelPtr : Schema.GetChannels())
 		{
 			ChannelPtr->IndexAsset(*this);
 		}
@@ -151,27 +151,26 @@ void FAssetIndexer::ComputeStats()
 {
 	Stats = FStats();
 
-	const FAssetSamplingContext* SamplingContext = IndexingContext.SamplingContext;
-	check(SamplingContext->FiniteDelta > UE_KINDA_SMALL_NUMBER);
+	check(SamplingContext.FiniteDelta > UE_KINDA_SMALL_NUMBER);
 
 	for (int32 SampleIdx = GetBeginSampleIdx(); SampleIdx != GetEndSampleIdx(); ++SampleIdx)
 	{
-		const float SampleTime = FMath::Min(CalculateSampleTime(SampleIdx), IndexingContext.AssetSampler->GetPlayLength());
+		const float SampleTime = FMath::Min(CalculateSampleTime(SampleIdx), AssetSampler.GetPlayLength());
 
 		bool AnyClamped = false;
-		const FTransform TrajTransformsPast = GetTransform(SampleTime - SamplingContext->FiniteDelta, AnyClamped);
+		const FTransform TrajTransformsPast = GetTransform(SampleTime - SamplingContext.FiniteDelta, AnyClamped);
 		if (!AnyClamped)
 		{
 			const FTransform TrajTransformsPresent = GetTransform(SampleTime, AnyClamped);
 			if (!AnyClamped)
 			{
-				const FTransform TrajTransformsFuture = GetTransform(SampleTime + SamplingContext->FiniteDelta, AnyClamped);
+				const FTransform TrajTransformsFuture = GetTransform(SampleTime + SamplingContext.FiniteDelta, AnyClamped);
 				if (!AnyClamped)
 				{
 					// if any transform is clamped we just skip the sample entirely
-					const FVector LinearVelocityPresent = (TrajTransformsPresent.GetTranslation() - TrajTransformsPast.GetTranslation()) / SamplingContext->FiniteDelta;
-					const FVector LinearVelocityFuture = (TrajTransformsFuture.GetTranslation() - TrajTransformsPresent.GetTranslation()) / SamplingContext->FiniteDelta;
-					const FVector LinearAcceleration = (LinearVelocityFuture - LinearVelocityPresent) / SamplingContext->FiniteDelta;
+					const FVector LinearVelocityPresent = (TrajTransformsPresent.GetTranslation() - TrajTransformsPast.GetTranslation()) / SamplingContext.FiniteDelta;
+					const FVector LinearVelocityFuture = (TrajTransformsFuture.GetTranslation() - TrajTransformsPresent.GetTranslation()) / SamplingContext.FiniteDelta;
+					const FVector LinearAcceleration = (LinearVelocityFuture - LinearVelocityPresent) / SamplingContext.FiniteDelta;
 
 					const float Speed = LinearVelocityPresent.Length();
 					const float Acceleration = LinearAcceleration.Length();
@@ -193,10 +192,8 @@ FAssetIndexer::FSampleInfo FAssetIndexer::GetSampleInfo(float SampleTime) const
 {
 	FSampleInfo Sample;
 
-	check(IndexingContext.AssetSampler);
-
-	const float PlayLength = IndexingContext.AssetSampler->GetPlayLength();
-	const bool bCanWrap = IndexingContext.AssetSampler->IsLoopable();
+	const float PlayLength = AssetSampler.GetPlayLength();
+	const bool bCanWrap = AssetSampler.IsLoopable();
 
 	float MainRelativeTime = SampleTime;
 	if (SampleTime < 0.0f && bCanWrap)
@@ -212,7 +209,7 @@ FAssetIndexer::FSampleInfo FAssetIndexer::GetSampleInfo(float SampleTime) const
 	{
 		Sample.bClamped = true;
 		Sample.ClipTime = SamplingParam.WrappedParam + SamplingParam.Extrapolation;
-		Sample.RootTransform = IndexingContext.AssetSampler->ExtractRootTransform(Sample.ClipTime);
+		Sample.RootTransform = AssetSampler.ExtractRootTransform(Sample.ClipTime);
 	}
 	else
 	{
@@ -220,12 +217,12 @@ FAssetIndexer::FSampleInfo FAssetIndexer::GetSampleInfo(float SampleTime) const
 		Sample.RootTransform = FTransform::Identity;
 
 		// Find the remaining motion deltas after wrapping
-		FTransform RootMotionRemainder = IndexingContext.AssetSampler->ExtractRootTransform(Sample.ClipTime);
+		FTransform RootMotionRemainder = AssetSampler.ExtractRootTransform(Sample.ClipTime);
 
 		const bool bNegativeSampleTime = SampleTime < 0.f;
 		if (SamplingParam.NumCycles > 0 || bNegativeSampleTime)
 		{
-			const FTransform RootMotionLast = IndexingContext.AssetSampler->GetTotalRootTransform();
+			const FTransform RootMotionLast = AssetSampler.GetTotalRootTransform();
 
 			// Determine how to accumulate motion for every cycle of the anim. If the sample
 			// had to be clamped, this motion will end up not getting applied below.
@@ -259,7 +256,7 @@ FAssetIndexer::FSampleInfo FAssetIndexer::GetSampleInfo(float SampleTime) const
 
 FTransform FAssetIndexer::MirrorTransform(const FTransform& Transform) const
 {
-	return IndexingContext.bMirrored ? IndexingContext.SamplingContext->MirrorTransform(Transform) : Transform;
+	return SearchIndexAsset.bMirrored ? SamplingContext.MirrorTransform(Transform) : Transform;
 }
 
 FAssetIndexer::CachedEntry& FAssetIndexer::GetEntry(float SampleTime)
@@ -267,8 +264,6 @@ FAssetIndexer::CachedEntry& FAssetIndexer::GetEntry(float SampleTime)
 	using namespace UE::Anim;
 
 	CachedEntry* Entry = CachedEntries.Find(SampleTime);
-	const FAssetSamplingContext* SamplingContext = IndexingContext.SamplingContext;
-
 	if (!Entry)
 	{
 		Entry = &CachedEntries.Add(SampleTime);
@@ -279,8 +274,8 @@ FAssetIndexer::CachedEntry& FAssetIndexer::GetEntry(float SampleTime)
 			UE_LOG(LogPoseSearch,
 				Warning,
 				TEXT("Invalid BoneContainer encountered in FAssetIndexer::GetEntry. Asset: %s. Schema: %s. BoneContainerAsset: %s. NumBoneIndices: %d"),
-				*GetNameSafe(IndexingContext.AssetSampler->GetAsset()),
-				*GetNameSafe(IndexingContext.Schema),
+				*GetNameSafe(AssetSampler.GetAsset()),
+				*GetNameSafe(&Schema),
 				*GetNameSafe(BoneContainer.GetAsset()),
 				BoneContainer.GetCompactPoseNumBones());
 		}
@@ -288,8 +283,8 @@ FAssetIndexer::CachedEntry& FAssetIndexer::GetEntry(float SampleTime)
 		const FAssetIndexer::FSampleInfo Sample = GetSampleInfo(SampleTime);
 		float CurrentTime = Sample.ClipTime;
 
-		const bool bLoopable = IndexingContext.AssetSampler->IsLoopable();
-		const float PlayLength = IndexingContext.AssetSampler->GetPlayLength();
+		const bool bLoopable = AssetSampler.IsLoopable();
+		const float PlayLength = AssetSampler.GetPlayLength();
 
 		if (!bLoopable)
 		{
@@ -298,16 +293,16 @@ FAssetIndexer::CachedEntry& FAssetIndexer::GetEntry(float SampleTime)
 
 		FCompactPose Pose;
 		Pose.SetBoneContainer(&BoneContainer);
-		IndexingContext.AssetSampler->ExtractPose(CurrentTime, Pose);
+		AssetSampler.ExtractPose(CurrentTime, Pose);
 		Pose[FCompactPoseBoneIndex(RootBoneIndexType)].SetIdentity();
 
-		if (IndexingContext.bMirrored && IndexingContext.Schema->MirrorDataTable)
+		if (SearchIndexAsset.bMirrored && Schema.MirrorDataTable)
 		{
 			FAnimationRuntime::MirrorPose(
 				Pose,
-				IndexingContext.Schema->MirrorDataTable->MirrorAxis,
-				SamplingContext->CompactPoseMirrorBones,
-				SamplingContext->ComponentSpaceRefRotations
+				Schema.MirrorDataTable->MirrorAxis,
+				SamplingContext.CompactPoseMirrorBones,
+				SamplingContext.ComponentSpaceRefRotations
 			);
 			// Note curves and attributes are not used during the indexing process and therefore don't need to be mirrored
 		}
@@ -326,7 +321,7 @@ FTransform FAssetIndexer::GetComponentSpaceTransform(float SampleTime, bool& bCl
 	CachedEntry& Entry = GetEntry(SampleTime);
 	bClamped = Entry.bClamped;
 
-	if (!IndexingContext.Schema->IsRootBone(SchemaBoneIdx))
+	if (!Schema.IsRootBone(SchemaBoneIdx))
 	{
 		return CalculateComponentSpaceTransform(Entry, SchemaBoneIdx);
 	}
@@ -341,7 +336,7 @@ FTransform FAssetIndexer::GetTransform(float SampleTime, bool& bClamped, int8 Sc
 	bClamped = Entry.bClamped;
 
 	const FTransform MirroredRootTransform = MirrorTransform(Entry.RootTransform);
-	if (!IndexingContext.Schema->IsRootBone(SchemaBoneIdx))
+	if (!Schema.IsRootBone(SchemaBoneIdx))
 	{
 		return CalculateComponentSpaceTransform(Entry, SchemaBoneIdx) * MirroredRootTransform;
 	}
@@ -351,14 +346,14 @@ FTransform FAssetIndexer::GetTransform(float SampleTime, bool& bClamped, int8 Sc
 
 FTransform FAssetIndexer::CalculateComponentSpaceTransform(FAssetIndexer::CachedEntry& Entry, int8 SchemaBoneIdx)
 {
-	const FBoneReference& BoneReference = IndexingContext.Schema->BoneReferences[SchemaBoneIdx];
+	const FBoneReference& BoneReference = Schema.BoneReferences[SchemaBoneIdx];
 	const FCompactPoseBoneIndex CompactBoneIndex = BoneContainer.MakeCompactPoseIndex(FMeshPoseBoneIndex(BoneReference.BoneIndex));
 	return Entry.ComponentSpacePose.GetComponentSpaceTransform(CompactBoneIndex);
 }
 
 float FAssetIndexer::CalculateSampleTime(int32 SampleIdx) const
 {
-	return SampleIdx * IndexingContext.Schema->GetSamplingInterval();
+	return SampleIdx * Schema.GetSamplingInterval();
 }
 
 FQuat FAssetIndexer::GetSampleRotation(float SampleTimeOffset, int32 SampleIdx, int8 SchemaSampleBoneIdx, int8 SchemaOriginBoneIdx, EPermutationTimeType PermutationTimeType)
@@ -372,13 +367,13 @@ FQuat FAssetIndexer::GetSampleRotation(float SampleTimeOffset, int32 SampleIdx,
 	const float OriginTime = Time + PermutationOriginTimeOffset;
 
 	// @todo: add support for SchemaSampleBoneIdx
-	if (!IndexingContext.Schema->IsRootBone(SchemaOriginBoneIdx))
+	if (!Schema.IsRootBone(SchemaOriginBoneIdx))
 	{
 		UE_LOG(LogPoseSearch,
 			Error,
 			TEXT("FAssetIndexer::GetSampleRotation: support for non root origin bones not implemented (bone: '%s', schema: '%s'"), 
-			*IndexingContext.Schema->BoneReferences[SchemaOriginBoneIdx].BoneName.ToString(),
-			*GetNameSafe(IndexingContext.Schema));
+			*Schema.BoneReferences[SchemaOriginBoneIdx].BoneName.ToString(),
+			*GetNameSafe(&Schema));
 	}
 
 	bool bUnused;
@@ -411,7 +406,7 @@ FVector FAssetIndexer::GetSamplePositionInternal(float SampleTime, float OriginT
 {
 	if (SampleTime == OriginTime)
 	{
-		if (IndexingContext.Schema->IsRootBone(SchemaOriginBoneIdx))
+		if (Schema.IsRootBone(SchemaOriginBoneIdx))
 		{
 			return GetComponentSpaceTransform(SampleTime, bClamped, SchemaSampleBoneIdx).GetTranslation();
 		}
@@ -426,7 +421,7 @@ FVector FAssetIndexer::GetSamplePositionInternal(float SampleTime, float OriginT
 	bool Unused;
 	const FTransform RootBoneTransform = GetTransform(OriginTime, Unused, RootSchemaBoneIdx);
 	const FTransform SampleBoneTransform = GetTransform(SampleTime, bClamped, SchemaSampleBoneIdx);
-	if (IndexingContext.Schema->IsRootBone(SchemaOriginBoneIdx))
+	if (Schema.IsRootBone(SchemaOriginBoneIdx))
 	{
 		return RootBoneTransform.InverseTransformPosition(SampleBoneTransform.GetTranslation());
 	}
@@ -447,7 +442,7 @@ FVector FAssetIndexer::GetSampleVelocity(float SampleTimeOffset, int32 SampleIdx
 	const float Time = CalculateSampleTime(SampleIdx);
 	const float SampleTime = Time + SampleTimeOffset + PermutationSampleTimeOffset;
 	const float OriginTime = Time + PermutationOriginTimeOffset;
-	const float FiniteDelta = IndexingContext.SamplingContext->FiniteDelta;
+	const float FiniteDelta = SamplingContext.FiniteDelta;
 
 	bool bClampedPast, bClampedPresent;
 	const FVector BonePositionPast = GetSamplePositionInternal(SampleTime - FiniteDelta, bUseCharacterSpaceVelocities ? OriginTime - FiniteDelta : OriginTime, bClampedPast, SchemaSampleBoneIdx, SchemaOriginBoneIdx);
@@ -474,21 +469,18 @@ int32 FAssetIndexer::GetVectorIdx(int32 SampleIdx) const
 
 TArrayView<float> FAssetIndexer::GetPoseVector(int32 SampleIdx) const
 {
-	check(IndexingContext.Schema);
-	return MakeArrayView(&FeatureVectorTable[GetVectorIdx(SampleIdx) * IndexingContext.Schema->SchemaCardinality], IndexingContext.Schema->SchemaCardinality);
+	return MakeArrayView(&FeatureVectorTable[GetVectorIdx(SampleIdx) * Schema.SchemaCardinality], Schema.SchemaCardinality);
 }
 
 const UPoseSearchSchema* FAssetIndexer::GetSchema() const
 {
-	check(IndexingContext.Schema);
-	return IndexingContext.Schema;
+	return &Schema;
 }
 
 float FAssetIndexer::CalculatePermutationTimeOffset() const
 {
-	const UPoseSearchSchema* Schema = GetSchema();
-	check(Schema->PermutationsSampleRate > 0 && SearchIndexAsset.IsInitialized());
-	const float PermutationTimeOffset = Schema->PermutationsTimeOffset + SearchIndexAsset.PermutationIdx / float(Schema->PermutationsSampleRate);
+	check(Schema.PermutationsSampleRate > 0 && SearchIndexAsset.IsInitialized());
+	const float PermutationTimeOffset = Schema.PermutationsTimeOffset + SearchIndexAsset.PermutationIdx / float(Schema.PermutationsSampleRate);
 	return PermutationTimeOffset;
 }
 
diff --git a/Engine/Plugins/Experimental/Animation/PoseSearch/Source/Runtime/Private/PoseSearchAssetSampler.cpp b/Engine/Plugins/Experimental/Animation/PoseSearch/Source/Runtime/Private/PoseSearchAssetSampler.cpp
--- a/Engine/Plugins/Experimental/Animation/PoseSearch/Source/Runtime/Private/PoseSearchAssetSampler.cpp
+++ b/Engine/Plugins/Experimental/Animation/PoseSearch/Source/Runtime/Private/PoseSearchAssetSampler.cpp
@@ -154,6 +154,10 @@ static void ProcessPlayLength(const UBlendSpace* BlendSpace, const FVector& Blen
 	{
 		CachedPlayLength = BlendSpace->GetAnimationLengthFromSampleData(BlendSamples);
 	}
+	else
+	{
+		CachedPlayLength = 0.f;
+	}
 }
 
 static void ProcessRootTransform(const UBlendSpace* BlendSpace, const FVector& BlendParameters, float CachedPlayLength, const FBoneContainer& BoneContainer,
@@ -262,6 +266,14 @@ void FAnimationAssetSampler::Init(TWeakObjectPtr<const UAnimationAsset> InAnimat
 		{
 			FMemMark Mark(FMemStack::Get());
 			ProcessPlayLength(BlendSpace, BlendParameters, CachedPlayLength);
+			
+#if !NO_LOGGING
+			const TArray<FName>* UniqueMarkerNames = const_cast<UBlendSpace*>(BlendSpace)->GetUniqueMarkerNames();
+			if (UniqueMarkerNames && !UniqueMarkerNames->IsEmpty())
+			{
+				UE_LOG(LogPoseSearch, Warning, TEXT("FAnimationAssetSampler::Init: sampling blend space (%s) with synch markers is currently not supported"), *BlendSpace->GetName());
+			}
+#endif // !NO_LOGGING
 		}
 		else
 		{
@@ -280,15 +292,30 @@ const UAnimationAsset* FAnimationAssetSampler::GetAsset() const
 	return AnimationAsset.Get();
 }
 
-float FAnimationAssetSampler::GetTimeToAssetTimeMultiplier() const
+float FAnimationAssetSampler::ToRealTime(float NormalizedTime) const
+{
+	// Asset player time for blend spaces is normalized [0, 1] so we convert the sampling / animation time to asset time by multiplying it by CachedPlayLength
+	if (CachedPlayLength > UE_KINDA_SMALL_NUMBER && Cast<UBlendSpace>(AnimationAsset.Get()))
+	{
+		check(NormalizedTime >= 0.f && NormalizedTime <= 1.f);
+		const float RealTime = NormalizedTime * CachedPlayLength;
+		return RealTime;
+	}
+
+	return NormalizedTime;
+}
+
+float FAnimationAssetSampler::ToNormalizedTime(float RealTime) const
 {
 	// Asset player time for blend spaces is normalized [0, 1] so we convert the sampling / animation time to asset time by dividing it by CachedPlayLength
 	if (CachedPlayLength > UE_KINDA_SMALL_NUMBER && Cast<UBlendSpace>(AnimationAsset.Get()))
 	{
-		return 1.f / CachedPlayLength;
+		const float NormalizedTime = RealTime / CachedPlayLength;
+		check(NormalizedTime >= 0.f && NormalizedTime <= 1.f);
+		return NormalizedTime;
 	}
 
-	return 1.f;
+	return RealTime;
 }
 
 float FAnimationAssetSampler::GetPlayLength() const
diff --git a/Engine/Plugins/Experimental/Animation/PoseSearch/Source/Runtime/Private/PoseSearchDatabase.cpp b/Engine/Plugins/Experimental/Animation/PoseSearch/Source/Runtime/Private/PoseSearchDatabase.cpp
--- a/Engine/Plugins/Experimental/Animation/PoseSearch/Source/Runtime/Private/PoseSearchDatabase.cpp
+++ b/Engine/Plugins/Experimental/Animation/PoseSearch/Source/Runtime/Private/PoseSearchDatabase.cpp
@@ -455,11 +455,10 @@ const FPoseSearchIndex& UPoseSearchDatabase::GetSearchIndex() const
 int32 UPoseSearchDatabase::GetPoseIndexFromTime(float Time, const FPoseSearchIndexAsset& SearchIndexAsset) const
 {
 	const bool bIsLooping = IsSourceAssetLooping(SearchIndexAsset);
-	const FFloatInterval& Range = SearchIndexAsset.SamplingInterval;
-	const bool bHasPoseIndex = SearchIndexAsset.FirstPoseIdx != INDEX_NONE && SearchIndexAsset.NumPoses > 0 && (bIsLooping || Range.Contains(Time));
+	const bool bHasPoseIndex = SearchIndexAsset.FirstPoseIdx != INDEX_NONE && SearchIndexAsset.NumPoses > 0 && (bIsLooping || SearchIndexAsset.SamplingInterval.Contains(Time));
 	if (bHasPoseIndex)
 	{
-		int32 PoseOffset = FMath::RoundToInt(Schema->SampleRate * (Time - Range.Min));
+		int32 PoseOffset = FMath::RoundToInt(Schema->SampleRate * (Time - SearchIndexAsset.SamplingInterval.Min));
 		
 		if (PoseOffset < 0)
 		{
@@ -502,8 +501,7 @@ bool UPoseSearchDatabase::GetPoseIndicesAndLerpValueFromTime(float Time, const F
 		return false;
 	}
 
-	const FFloatInterval& Range = SearchIndexAsset.SamplingInterval;
-	const float FloatPoseOffset = Schema->SampleRate * (Time - Range.Min);
+	const float FloatPoseOffset = Schema->SampleRate * (Time - SearchIndexAsset.SamplingInterval.Min);
 	const int32 PoseOffset = FMath::RoundToInt(FloatPoseOffset);
 	LerpValue = FloatPoseOffset - float(PoseOffset);
 
@@ -692,16 +690,18 @@ float UPoseSearchDatabase::GetAssetTime(int32 PoseIdx) const
 	const bool bIsBlendSpace = AnimationAssets[Asset.SourceAssetIdx].GetPtr<FPoseSearchDatabaseBlendSpace>() != nullptr;
 	const FFloatInterval& SamplingRange = Asset.SamplingInterval;
 
-	if (bIsBlendSpace)
+	// sequences or anim composites
+	float AssetTime = FMath::Min(SamplingRange.Min + SamplingInterval * (PoseIdx - Asset.FirstPoseIdx), SamplingRange.Max);
+
+	if (bIsBlendSpace && Asset.NumPoses > 1)
 	{
+		check(SamplingInterval > UE_KINDA_SMALL_NUMBER);
 		// For BlendSpaces the AssetTime is in the range [0, 1] while the Sampling Range
 		// is in real time (seconds)
-		const float AssetTime = FMath::Min(SamplingRange.Min + SamplingInterval * (PoseIdx - Asset.FirstPoseIdx), SamplingRange.Max) / (Asset.NumPoses * SamplingInterval);
-		return AssetTime;
+		AssetTime = AssetTime / ((Asset.NumPoses - 1) * SamplingInterval);
+		check(AssetTime >= 0.f && AssetTime <= 1.f);
 	}
 
-	// sequences or anim composites
-	const float AssetTime = FMath::Min(SamplingRange.Min + SamplingInterval * (PoseIdx - Asset.FirstPoseIdx), SamplingRange.Max);
 	return AssetTime;
 }
 
diff --git a/Engine/Plugins/Experimental/Animation/PoseSearch/Source/Runtime/Private/PoseSearchDatabaseIndexingContext.cpp b/Engine/Plugins/Experimental/Animation/PoseSearch/Source/Runtime/Private/PoseSearchDatabaseIndexingContext.cpp
--- a/Engine/Plugins/Experimental/Animation/PoseSearch/Source/Runtime/Private/PoseSearchDatabaseIndexingContext.cpp
+++ b/Engine/Plugins/Experimental/Animation/PoseSearch/Source/Runtime/Private/PoseSearchDatabaseIndexingContext.cpp
@@ -114,21 +114,12 @@ bool FDatabaseIndexingContext::IndexDatabase(FPoseSearchIndexBase& SearchIndexBa
 		FPoseSearchIndexAsset& SearchIndexAsset = SearchIndexBase.Assets[AssetIdx];
 		SearchIndexAsset.FirstPoseIdx = TotalPoses;
 
-		FAssetIndexingContext IndexerContext;
-		IndexerContext.SamplingContext = &SamplingContext;
-		IndexerContext.Schema = Schema;
-		IndexerContext.RequestedSamplingRange = SearchIndexAsset.SamplingInterval;
-		IndexerContext.bMirrored = SearchIndexAsset.bMirrored;
-
 		const FInstancedStruct& DatabaseAsset = Database.GetAnimationAssetStruct(SearchIndexAsset.SourceAssetIdx);
 		const FPoseSearchDatabaseAnimationAssetBase* DatabaseAnimationAssetBase = DatabaseAsset.GetPtr<FPoseSearchDatabaseAnimationAssetBase>();
-		check(DatabaseAnimationAssetBase);
-		if (DatabaseAnimationAssetBase->GetAnimationAsset())
-		{
-			IndexerContext.AssetSampler = &Samplers[SamplerMap[{ DatabaseAnimationAssetBase->GetAnimationAsset(), SearchIndexAsset.BlendParameters }]];
-		}
+		check(DatabaseAnimationAssetBase && DatabaseAnimationAssetBase->GetAnimationAsset());
+		const FAnimationAssetSampler& AssetSampler = Samplers[SamplerMap[{ DatabaseAnimationAssetBase->GetAnimationAsset(), SearchIndexAsset.BlendParameters }]];
 
-		const int32 NewIndexerIdx = Indexers.Emplace(IndexerContext, BoneContainer, SearchIndexAsset);
+		const int32 NewIndexerIdx = Indexers.Emplace(BoneContainer, SearchIndexAsset, SamplingContext, *Schema, AssetSampler);
 		SearchIndexAsset.NumPoses = Indexers[NewIndexerIdx].GetNumIndexedPoses();
 		TotalPoses += SearchIndexAsset.NumPoses;
 	}
diff --git a/Engine/Plugins/Experimental/Animation/PoseSearch/Source/Runtime/Private/PoseSearchDerivedDataKey.cpp b/Engine/Plugins/Experimental/Animation/PoseSearch/Source/Runtime/Private/PoseSearchDerivedDataKey.cpp
--- a/Engine/Plugins/Experimental/Animation/PoseSearch/Source/Runtime/Private/PoseSearchDerivedDataKey.cpp
+++ b/Engine/Plugins/Experimental/Animation/PoseSearch/Source/Runtime/Private/PoseSearchDerivedDataKey.cpp
@@ -41,7 +41,7 @@ FKeyBuilder::FKeyBuilder(const UObject* Object, bool bUseDataVer)
 	if (bUseDataVer)
 	{
 		// used to invalidate the key without having to change POSESEARCHDB_DERIVEDDATA_VER all the times
-		int32 POSESEARCHDB_DERIVEDDATA_VER_SMALL = 139;
+		int32 POSESEARCHDB_DERIVEDDATA_VER_SMALL = 147;
 		FGuid VersionGuid = FDevSystemGuids::GetSystemGuid(FDevSystemGuids::Get().POSESEARCHDB_DERIVEDDATA_VER);
 
 		*this << VersionGuid;
diff --git a/Engine/Plugins/Experimental/Animation/PoseSearch/Source/Runtime/Private/PoseSearchResult.cpp b/Engine/Plugins/Experimental/Animation/PoseSearch/Source/Runtime/Private/PoseSearchResult.cpp
--- a/Engine/Plugins/Experimental/Animation/PoseSearch/Source/Runtime/Private/PoseSearchResult.cpp
+++ b/Engine/Plugins/Experimental/Animation/PoseSearch/Source/Runtime/Private/PoseSearchResult.cpp
@@ -58,6 +58,7 @@ void FSearchResult::Update(float NewAssetTime)
 
 			// Asset player time for blendspaces is normalized [0, 1] so we need to convert 
 			// to a real time before we advance it
+			check(NewAssetTime >= 0.f && NewAssetTime <= 1.f);
 			const float RealTime = NewAssetTime * PlayLength;
 			if (Database->GetPoseIndicesAndLerpValueFromTime(RealTime, SearchIndexAsset, PrevPoseIdx, PoseIdx, NextPoseIdx, LerpValue))
 			{
diff --git a/Engine/Plugins/Experimental/Animation/PoseSearch/Source/Runtime/Public/PoseSearch/PoseSearchAssetIndexer.h b/Engine/Plugins/Experimental/Animation/PoseSearch/Source/Runtime/Public/PoseSearch/PoseSearchAssetIndexer.h
--- a/Engine/Plugins/Experimental/Animation/PoseSearch/Source/Runtime/Public/PoseSearch/PoseSearchAssetIndexer.h
+++ b/Engine/Plugins/Experimental/Animation/PoseSearch/Source/Runtime/Public/PoseSearch/PoseSearchAssetIndexer.h
@@ -19,18 +19,6 @@ namespace UE::PoseSearch
 struct FAnimationAssetSampler;
 struct FAssetSamplingContext;
 
-/**
- * Inputs for asset indexing
- */
-struct FAssetIndexingContext
-{
-	const FAssetSamplingContext* SamplingContext = nullptr;
-	const UPoseSearchSchema* Schema = nullptr;
-	const FAnimationAssetSampler* AssetSampler = nullptr;
-	bool bMirrored = false;
-	FFloatInterval RequestedSamplingRange = FFloatInterval(0.f, 0.f);
-};
-
 class FAssetIndexer
 {
 public:
@@ -43,7 +31,8 @@ public:
 		float MaxAcceleration = 0.f;
 	};
 
-	FAssetIndexer(const FAssetIndexingContext& IndexingContext, const FBoneContainer& InBoneContainer, const FPoseSearchIndexAsset& InSearchIndexAsset);
+	FAssetIndexer(const FBoneContainer& InBoneContainer, const FPoseSearchIndexAsset& InSearchIndexAsset, 
+		const FAssetSamplingContext& InSamplingContext, const UPoseSearchSchema& InSchema, const FAnimationAssetSampler& InAssetSampler);
 	void AssignWorkingData(TArrayView<float> InOutFeatureVectorTable, TArrayView<FPoseSearchPoseMetadata> InOutPoseMetadata);
 	void Process(int32 AssetIdx);
 	const FStats& GetStats() const { return Stats; }
@@ -90,9 +79,11 @@ private:
 	void ComputeStats();
 
 	FBoneContainer BoneContainer;
-	FAssetIndexingContext IndexingContext;
 	TMap<float, CachedEntry> CachedEntries;
 	const FPoseSearchIndexAsset& SearchIndexAsset;
+	const FAssetSamplingContext& SamplingContext;
+	const UPoseSearchSchema& Schema;
+	const FAnimationAssetSampler& AssetSampler;
 	
 	int32 FirstIndexedSample = 0;
 	int32 LastIndexedSample = 0;
diff --git a/Engine/Plugins/Experimental/Animation/PoseSearch/Source/Runtime/Public/PoseSearch/PoseSearchAssetSampler.h b/Engine/Plugins/Experimental/Animation/PoseSearch/Source/Runtime/Public/PoseSearch/PoseSearchAssetSampler.h
--- a/Engine/Plugins/Experimental/Animation/PoseSearch/Source/Runtime/Public/PoseSearch/PoseSearchAssetSampler.h
+++ b/Engine/Plugins/Experimental/Animation/PoseSearch/Source/Runtime/Public/PoseSearch/PoseSearchAssetSampler.h
@@ -22,7 +22,8 @@ struct POSESEARCH_API FAnimationAssetSampler
 
 	bool IsInitialized() const;
 	float GetPlayLength() const;
-	float GetTimeToAssetTimeMultiplier() const;
+	float ToRealTime(float NormalizedTime) const;
+	float ToNormalizedTime(float RealTime) const;
 	bool IsLoopable() const;
 
 	// Gets the final root transformation at the end of the asset's playback time
