motion matching
- database normalized vs real time cleanup
- proper quantization of database poses (samples) time
- deprecated FPoseSearchIndexAsset::SamplingInterval in favor of FirstPoseIdx/LastPoseIdx leading to properly quantized time
- avoid  adding FPoseSearchIndexAsset with no associated poses in InitSearchIndexAssets
- saved 4 bytes per FPoseSearchIndexAsset
- removed no longer necessary FMotionMatchingState::CanAdvance since animation synchronization happens at the beginning of the MM update

diff --git a/Engine/Plugins/Experimental/Animation/PoseSearch/Source/Editor/Private/PoseSearchDatabaseViewModel.cpp b/Engine/Plugins/Experimental/Animation/PoseSearch/Source/Editor/Private/PoseSearchDatabaseViewModel.cpp
--- a/Engine/Plugins/Experimental/Animation/PoseSearch/Source/Editor/Private/PoseSearchDatabaseViewModel.cpp
+++ b/Engine/Plugins/Experimental/Animation/PoseSearch/Source/Editor/Private/PoseSearchDatabaseViewModel.cpp
@@ -148,7 +148,7 @@ namespace UE::PoseSearch
 		PreviewActor.Sampler.Init(DatabaseAnimationAssetBase->GetAnimationAsset(), IndexAsset.BlendParameters);
 		PreviewActor.IndexAssetIndex = IndexAssetIndex;
 		PreviewActor.CurrentPoseIndex = INDEX_NONE;
-		PreviewActor.PlayTimeOffset = PoseIdxForTimeOffset < 0 ? 0.f : PreviewActor.Sampler.ToRealTime(PoseSearchDatabase->GetAssetTime(PoseIdxForTimeOffset));
+		PreviewActor.PlayTimeOffset = PoseIdxForTimeOffset < 0 ? 0.f : PoseSearchDatabase->GetRealAssetTime(PoseIdxForTimeOffset);
 
 		FActorSpawnParameters Params;
 		Params.SpawnCollisionHandlingOverride = ESpawnActorCollisionHandlingMethod::AdjustIfPossibleButAlwaysSpawn;
@@ -234,13 +234,13 @@ namespace UE::PoseSearch
 
 			PreviewActor.CurrentTime = 0.f;
 			const FPoseSearchIndexAsset& IndexAsset = SearchIndex.Assets[PreviewActor.IndexAssetIndex];
-			float CurrentPlayTime = PlayTime + IndexAsset.SamplingInterval.Min + PreviewActor.PlayTimeOffset;
-			FAnimationRuntime::AdvanceTime(false, CurrentPlayTime, PreviewActor.CurrentTime, IndexAsset.SamplingInterval.Max);
+			float CurrentPlayTime = PlayTime + IndexAsset.GetFirstSampleTime(PoseSearchDatabase->Schema->SampleRate) + PreviewActor.PlayTimeOffset;
+			FAnimationRuntime::AdvanceTime(false, CurrentPlayTime, PreviewActor.CurrentTime, IndexAsset.GetLastSampleTime(PoseSearchDatabase->Schema->SampleRate));
 			 
 			// time to pose index
 			PreviewActor.CurrentPoseIndex = PoseSearchDatabase->GetPoseIndexFromTime(PreviewActor.CurrentTime, IndexAsset);
 
-			const float QuantizedTime = PreviewActor.CurrentPoseIndex >= 0 ? PreviewActor.Sampler.ToRealTime(PoseSearchDatabase->GetAssetTime(PreviewActor.CurrentPoseIndex)) : PreviewActor.CurrentTime;
+			const float QuantizedTime = PreviewActor.CurrentPoseIndex >= 0 ? PoseSearchDatabase->GetRealAssetTime(PreviewActor.CurrentPoseIndex) : PreviewActor.CurrentTime;
 			if (bQuantizeAnimationToPoseData)
 			{
 				PreviewActor.CurrentTime = QuantizedTime;
@@ -296,19 +296,19 @@ namespace UE::PoseSearch
 		check(PoseSearchDatabase && PoseSearchDatabase->Schema);
 		if (PoseSearchDatabase->Schema->MirrorDataTable && PoseSearchDatabase->Schema->Skeleton)
 		{
-		const FTransform RootReferenceTransform = PoseSearchDatabase->Schema->Skeleton->GetReferenceSkeleton().GetRefBonePose()[0];
-		const FQuat RootReferenceRotation = RootReferenceTransform.GetRotation();
+			const FTransform RootReferenceTransform = PoseSearchDatabase->Schema->Skeleton->GetReferenceSkeleton().GetRefBonePose()[0];
+			const FQuat RootReferenceRotation = RootReferenceTransform.GetRotation();
 
 			const EAxis::Type MirrorAxis = PoseSearchDatabase->Schema->MirrorDataTable->MirrorAxis;
 			FVector T = RootTransform.GetTranslation();
-		T = FAnimationRuntime::MirrorVector(T, MirrorAxis);
+			T = FAnimationRuntime::MirrorVector(T, MirrorAxis);
 			FQuat Q = RootTransform.GetRotation();
-		Q = FAnimationRuntime::MirrorQuat(Q, MirrorAxis);
-		Q *= FAnimationRuntime::MirrorQuat(RootReferenceRotation, MirrorAxis).Inverse() * RootReferenceRotation;
+			Q = FAnimationRuntime::MirrorQuat(Q, MirrorAxis);
+			Q *= FAnimationRuntime::MirrorQuat(RootReferenceRotation, MirrorAxis).Inverse() * RootReferenceRotation;
 
 			const FTransform MirroredRootTransform = FTransform(Q, T, RootTransform.GetScale3D());
 			return MirroredRootTransform;
-	}
+		}
 		return RootTransform;
 	}
 
@@ -412,8 +412,8 @@ namespace UE::PoseSearch
 					if (PreviewActor.IsValid())
 					{
 						const FPoseSearchIndexAsset& IndexAsset = SearchIndex.Assets[IndexAssetIndex];
-						MaxPreviewPlayLength = FMath::Max(MaxPreviewPlayLength, IndexAsset.SamplingInterval.Max - PreviewActor.PlayTimeOffset);
-						MinPreviewPlayLength = FMath::Min(MinPreviewPlayLength, IndexAsset.SamplingInterval.Min - PreviewActor.PlayTimeOffset);
+						MaxPreviewPlayLength = FMath::Max(MaxPreviewPlayLength, IndexAsset.GetLastSampleTime(PoseSearchDatabase->Schema->SampleRate) - PreviewActor.PlayTimeOffset);
+						MinPreviewPlayLength = FMath::Min(MinPreviewPlayLength, IndexAsset.GetFirstSampleTime(PoseSearchDatabase->Schema->SampleRate) - PreviewActor.PlayTimeOffset);
 						PreviewActors.Add(PreviewActor);
 						SelectedSourceAssetIdx = IndexAsset.SourceAssetIdx;
 					}
@@ -460,7 +460,7 @@ namespace UE::PoseSearch
 						FDatabasePreviewActor PreviewActor = SpawnPreviewActor(IndexAssetIndex);
 						if (PreviewActor.IsValid())
 						{
-							MaxPreviewPlayLength = FMath::Max(MaxPreviewPlayLength, IndexAsset.SamplingInterval.Max - IndexAsset.SamplingInterval.Min);
+							MaxPreviewPlayLength = FMath::Max(MaxPreviewPlayLength, IndexAsset.GetLastSampleTime(PoseSearchDatabase->Schema->SampleRate) - IndexAsset.GetFirstSampleTime(PoseSearchDatabase->Schema->SampleRate));
 							PreviewActors.Add(PreviewActor);
 						}
 					}
@@ -548,7 +548,7 @@ namespace UE::PoseSearch
 					const FPoseSearchIndexAsset& IndexAsset = SearchIndex.Assets[PreviewActor.IndexAssetIndex];
 					if (IndexAsset.SourceAssetIdx == SourceAssetIdx)
 					{
-						CurrentPlayTime = PreviewActor.Sampler.ToNormalizedTime(PlayTime + IndexAsset.SamplingInterval.Min + PreviewActor.PlayTimeOffset);
+						CurrentPlayTime = PreviewActor.Sampler.ToNormalizedTime(PlayTime + IndexAsset.GetFirstSampleTime(PoseSearchDatabase->Schema->SampleRate) + PreviewActor.PlayTimeOffset);
 						BlendParameters = IndexAsset.BlendParameters;
 						return true;
 					}
@@ -559,7 +559,7 @@ namespace UE::PoseSearch
 			{
 				if (IndexAsset.SourceAssetIdx == SourceAssetIdx)
 				{
-					CurrentPlayTime = PlayTime + IndexAsset.SamplingInterval.Min;
+					CurrentPlayTime = PlayTime + IndexAsset.GetFirstSampleTime(PoseSearchDatabase->Schema->SampleRate);
 					BlendParameters = IndexAsset.BlendParameters;
 
 					const bool bIsBlendSpace = PoseSearchDatabase->GetAnimationAssetStruct(IndexAsset).GetPtr<FPoseSearchDatabaseBlendSpace>() != nullptr;
diff --git a/Engine/Plugins/Experimental/Animation/PoseSearch/Source/Editor/Private/PoseSearchDebuggerDatabaseView.cpp b/Engine/Plugins/Experimental/Animation/PoseSearch/Source/Editor/Private/PoseSearchDebuggerDatabaseView.cpp
--- a/Engine/Plugins/Experimental/Animation/PoseSearch/Source/Editor/Private/PoseSearchDebuggerDatabaseView.cpp
+++ b/Engine/Plugins/Experimental/Animation/PoseSearch/Source/Editor/Private/PoseSearchDebuggerDatabaseView.cpp
@@ -161,7 +161,7 @@ static void AddUnfilteredDatabaseRow(const UPoseSearchDatabase* Database,
 	{
 		TSharedRef<FDebuggerDatabaseRowData>& Row = UnfilteredDatabaseRows.Add_GetRef(MakeShared<FDebuggerDatabaseRowData>(SharedData));
 
-		const float Time = Database->GetAssetTime(DbPoseIdx);
+		const float Time = Database->GetNormalizedAssetTime(DbPoseIdx);
 
 		Row->PoseIdx = DbPoseIdx;
 		Row->PoseCandidateFlags = PoseCandidateFlags;
diff --git a/Engine/Plugins/Experimental/Animation/PoseSearch/Source/Runtime/Private/PoseSearchAssetIndexer.cpp b/Engine/Plugins/Experimental/Animation/PoseSearch/Source/Runtime/Private/PoseSearchAssetIndexer.cpp
--- a/Engine/Plugins/Experimental/Animation/PoseSearch/Source/Runtime/Private/PoseSearchAssetIndexer.cpp
+++ b/Engine/Plugins/Experimental/Animation/PoseSearch/Source/Runtime/Private/PoseSearchAssetIndexer.cpp
@@ -84,11 +84,6 @@ FAssetIndexer::FAssetIndexer(const FBoneContainer& InBoneContainer, const FPoseS
 , Schema(InSchema)
 , AssetSampler(InAssetSampler)
 {
-	check(Schema.IsValid());
-
-	FirstIndexedSample = FMath::FloorToInt(SearchIndexAsset.SamplingInterval.Min * Schema.SampleRate);
-	// @todo: parhaps we should use FMath::CeilToInt to avoid sampling over teh length of the animation
-	LastIndexedSample = FMath::Max(0, FMath::CeilToInt(SearchIndexAsset.SamplingInterval.Max * Schema.SampleRate));
 }
 
 void FAssetIndexer::AssignWorkingData(TArrayView<float> InOutFeatureVectorTable, TArrayView<FPoseSearchPoseMetadata> InOutPoseMetadata)
@@ -353,7 +348,7 @@ FTransform FAssetIndexer::CalculateComponentSpaceTransform(FAssetIndexer::Cached
 
 float FAssetIndexer::CalculateSampleTime(int32 SampleIdx) const
 {
-	return SampleIdx * Schema.GetSamplingInterval();
+	return SampleIdx / float(Schema.SampleRate);
 }
 
 FQuat FAssetIndexer::GetSampleRotation(float SampleTimeOffset, int32 SampleIdx, int8 SchemaSampleBoneIdx, int8 SchemaOriginBoneIdx, EPermutationTimeType PermutationTimeType)
@@ -462,6 +457,21 @@ FVector FAssetIndexer::GetSampleVelocity(float SampleTimeOffset, int32 SampleIdx
 	return LinearVelocity;
 }
 
+int32 FAssetIndexer::GetBeginSampleIdx() const
+{
+	return SearchIndexAsset.GetBeginSampleIdx();
+}
+
+int32 FAssetIndexer::GetEndSampleIdx() const
+{
+	return SearchIndexAsset.GetEndSampleIdx();
+}
+
+int32 FAssetIndexer::GetNumIndexedPoses() const
+{
+	return SearchIndexAsset.GetNumPoses();
+}
+
 int32 FAssetIndexer::GetVectorIdx(int32 SampleIdx) const
 {
 	return SampleIdx - GetBeginSampleIdx();
diff --git a/Engine/Plugins/Experimental/Animation/PoseSearch/Source/Runtime/Private/PoseSearchAssetSampler.cpp b/Engine/Plugins/Experimental/Animation/PoseSearch/Source/Runtime/Private/PoseSearchAssetSampler.cpp
--- a/Engine/Plugins/Experimental/Animation/PoseSearch/Source/Runtime/Private/PoseSearchAssetSampler.cpp
+++ b/Engine/Plugins/Experimental/Animation/PoseSearch/Source/Runtime/Private/PoseSearchAssetSampler.cpp
@@ -146,20 +146,6 @@ static FTransform ExtractBlendSpaceRootMotion(float StartTime, float DeltaTime,
 	return RootMotionParams.GetRootMotionTransform();
 }
 
-static void ProcessPlayLength(const UBlendSpace* BlendSpace, const FVector& BlendParameters, float& CachedPlayLength)
-{
-	TArray<FBlendSampleData> BlendSamples;
-	int32 TriangulationIndex = 0;
-	if (BlendSpace->GetSamplesFromBlendInput(BlendParameters, BlendSamples, TriangulationIndex, true))
-	{
-		CachedPlayLength = BlendSpace->GetAnimationLengthFromSampleData(BlendSamples);
-	}
-	else
-	{
-		CachedPlayLength = 0.f;
-	}
-}
-
 static void ProcessRootTransform(const UBlendSpace* BlendSpace, const FVector& BlendParameters, float CachedPlayLength, const FBoneContainer& BoneContainer,
 	int32 RootTransformSamplingRate, bool bIsLoopable, TArray<FTransform>& AccumulatedRootTransform)
 {
@@ -259,14 +245,30 @@ void FAnimationAssetSampler::Init(TWeakObjectPtr<const UAnimationAsset> InAnimat
 	AnimationAsset = InAnimationAsset;
 	BlendParameters = InBlendParameters;
 	RootTransformSamplingRate = InRootTransformSamplingRate;
+	CachedPlayLength = GetPlayLength(AnimationAsset.Get(), BlendParameters);
+}
 
-	if (const UAnimationAsset* AnimAsset = AnimationAsset.Get())
+bool FAnimationAssetSampler::IsInitialized() const
+{
+	return AnimationAsset != nullptr;
+}
+
+float FAnimationAssetSampler::GetPlayLength(const UAnimationAsset* AnimAsset, const FVector& BlendParameters)
+{
+	float PlayLength = 0.f;
+	if (AnimAsset)
 	{
 		if (const UBlendSpace* BlendSpace = Cast<UBlendSpace>(AnimAsset))
 		{
 			FMemMark Mark(FMemStack::Get());
-			ProcessPlayLength(BlendSpace, BlendParameters, CachedPlayLength);
-			
+
+			TArray<FBlendSampleData> BlendSamples;
+			int32 TriangulationIndex = 0;
+			if (BlendSpace->GetSamplesFromBlendInput(BlendParameters, BlendSamples, TriangulationIndex, true))
+			{
+				PlayLength = BlendSpace->GetAnimationLengthFromSampleData(BlendSamples);
+			}
+
 #if !NO_LOGGING
 			const TArray<FName>* UniqueMarkerNames = const_cast<UBlendSpace*>(BlendSpace)->GetUniqueMarkerNames();
 			if (UniqueMarkerNames && !UniqueMarkerNames->IsEmpty())
@@ -277,14 +279,11 @@ void FAnimationAssetSampler::Init(TWeakObjectPtr<const UAnimationAsset> InAnimat
 		}
 		else
 		{
-			CachedPlayLength = AnimAsset->GetPlayLength();
+			PlayLength = AnimAsset->GetPlayLength();
 		}
 	}
-}
 
-bool FAnimationAssetSampler::IsInitialized() const
-{
-	return AnimationAsset != nullptr;
+	return PlayLength;
 }
 
 const UAnimationAsset* FAnimationAssetSampler::GetAsset() const
diff --git a/Engine/Plugins/Experimental/Animation/PoseSearch/Source/Runtime/Private/PoseSearchContext.cpp b/Engine/Plugins/Experimental/Animation/PoseSearch/Source/Runtime/Private/PoseSearchContext.cpp
--- a/Engine/Plugins/Experimental/Animation/PoseSearch/Source/Runtime/Private/PoseSearchContext.cpp
+++ b/Engine/Plugins/Experimental/Animation/PoseSearch/Source/Runtime/Private/PoseSearchContext.cpp
@@ -246,7 +246,7 @@ void FDebugDrawParams::DrawFeatureVector(int32 PoseIdx)
 //////////////////////////////////////////////////////////////////////////
 // FSearchContext
 FSearchContext::FSearchContext(const FPoseSearchQueryTrajectory* InTrajectory, const IPoseHistory* InHistory, float InDesiredPermutationTimeOffset, const FPoseIndicesHistory* InPoseIndicesHistory,
-	EPoseSearchBooleanRequest InQueryMirrorRequest, const FSearchResult& InCurrentResult, float InPoseJumpThresholdTime, bool bInForceInterrupt, bool bInCanAdvance)
+	EPoseSearchBooleanRequest InQueryMirrorRequest, const FSearchResult& InCurrentResult, float InPoseJumpThresholdTime, bool bInForceInterrupt)
 : Trajectory(InTrajectory)
 , History(InHistory)
 , DesiredPermutationTimeOffset(InDesiredPermutationTimeOffset)
@@ -255,7 +255,6 @@ FSearchContext::FSearchContext(const FPoseSearchQueryTrajectory* InTrajectory, c
 , CurrentResult(InCurrentResult)
 , PoseJumpThresholdTime(InPoseJumpThresholdTime)
 , bForceInterrupt(bInForceInterrupt)
-, bCanAdvance(bInCanAdvance)
 {
 	if (CurrentResult.IsValid())
 	{
diff --git a/Engine/Plugins/Experimental/Animation/PoseSearch/Source/Runtime/Private/PoseSearchDatabase.cpp b/Engine/Plugins/Experimental/Animation/PoseSearch/Source/Runtime/Private/PoseSearchDatabase.cpp
--- a/Engine/Plugins/Experimental/Animation/PoseSearch/Source/Runtime/Private/PoseSearchDatabase.cpp
+++ b/Engine/Plugins/Experimental/Animation/PoseSearch/Source/Runtime/Private/PoseSearchDatabase.cpp
@@ -42,7 +42,7 @@ static void PopulateNonSelectableIdx(FNonSelectableIdx& NonSelectableIdx, FSearc
 	const FPoseSearchIndexAsset* CurrentIndexAsset = SearchContext.GetCurrentResult().GetSearchIndexAsset();
 	if (CurrentIndexAsset && SearchContext.IsCurrentResultFromDatabase(Database) && SearchContext.GetPoseJumpThresholdTime() > 0.f)
 	{
-		const int32 PoseJumpIndexThreshold = FMath::FloorToInt(SearchContext.GetPoseJumpThresholdTime() / Database->Schema->GetSamplingInterval());
+		const int32 PoseJumpIndexThreshold = FMath::FloorToInt(SearchContext.GetPoseJumpThresholdTime() * Database->Schema->SampleRate);
 		const bool IsLooping = Database->IsSourceAssetLooping(*CurrentIndexAsset);
 
 		for (int32 i = -PoseJumpIndexThreshold; i <= -1; ++i)
@@ -55,7 +55,7 @@ static void PopulateNonSelectableIdx(FNonSelectableIdx& NonSelectableIdx, FSearc
 
 				while (PoseIdx < CurrentIndexAsset->FirstPoseIdx)
 				{
-					PoseIdx += CurrentIndexAsset->NumPoses;
+					PoseIdx += CurrentIndexAsset->GetNumPoses();
 				}
 			}
 			else if (CurrentIndexAsset->IsPoseInRange(PoseIdx))
@@ -87,9 +87,9 @@ static void PopulateNonSelectableIdx(FNonSelectableIdx& NonSelectableIdx, FSearc
 			{
 				bIsPoseInRange = true;
 
-				while (PoseIdx >= CurrentIndexAsset->FirstPoseIdx + CurrentIndexAsset->NumPoses)
+				while (PoseIdx >= CurrentIndexAsset->FirstPoseIdx + CurrentIndexAsset->GetNumPoses())
 				{
-					PoseIdx -= CurrentIndexAsset->NumPoses;
+					PoseIdx -= CurrentIndexAsset->GetNumPoses();
 				}
 			}
 			else if (CurrentIndexAsset->IsPoseInRange(PoseIdx))
@@ -455,39 +455,7 @@ const FPoseSearchIndex& UPoseSearchDatabase::GetSearchIndex() const
 int32 UPoseSearchDatabase::GetPoseIndexFromTime(float Time, const FPoseSearchIndexAsset& SearchIndexAsset) const
 {
 	const bool bIsLooping = IsSourceAssetLooping(SearchIndexAsset);
-	const bool bHasPoseIndex = SearchIndexAsset.FirstPoseIdx != INDEX_NONE && SearchIndexAsset.NumPoses > 0 && (bIsLooping || SearchIndexAsset.SamplingInterval.Contains(Time));
-	if (bHasPoseIndex)
-	{
-		int32 PoseOffset = FMath::RoundToInt(Schema->SampleRate * (Time - SearchIndexAsset.SamplingInterval.Min));
-		
-		if (PoseOffset < 0)
-		{
-			if (bIsLooping)
-			{
-				PoseOffset = (PoseOffset % SearchIndexAsset.NumPoses) + SearchIndexAsset.NumPoses;
-			}
-			else
-			{
-				PoseOffset = 0;
-			}
-		}
-		else if (PoseOffset >= SearchIndexAsset.NumPoses)
-		{
-			if (bIsLooping)
-			{
-				PoseOffset = PoseOffset % SearchIndexAsset.NumPoses;
-			}
-			else
-			{
-				PoseOffset = SearchIndexAsset.NumPoses - 1;
-			}
-		}
-
-		int32 PoseIdx = SearchIndexAsset.FirstPoseIdx + PoseOffset;
-		return PoseIdx;
-	}
-
-	return INDEX_NONE;
+	return SearchIndexAsset.GetPoseIndexFromTime(Time, bIsLooping, Schema->SampleRate);
 }
 
 bool UPoseSearchDatabase::GetPoseIndicesAndLerpValueFromTime(float Time, const FPoseSearchIndexAsset& SearchIndexAsset, int32& PrevPoseIdx, int32& PoseIdx, int32& NextPoseIdx, float& LerpValue) const
@@ -501,7 +469,7 @@ bool UPoseSearchDatabase::GetPoseIndicesAndLerpValueFromTime(float Time, const F
 		return false;
 	}
 
-	const float FloatPoseOffset = Schema->SampleRate * (Time - SearchIndexAsset.SamplingInterval.Min);
+	const float FloatPoseOffset = Schema->SampleRate * Time - SearchIndexAsset.FirstSampleIdx;
 	const int32 PoseOffset = FMath::RoundToInt(FloatPoseOffset);
 	LerpValue = FloatPoseOffset - float(PoseOffset);
 
@@ -682,23 +650,28 @@ void UPoseSearchDatabase::Serialize(FArchive& Ar)
 	}
 }
 
-float UPoseSearchDatabase::GetAssetTime(int32 PoseIdx) const
+float UPoseSearchDatabase::GetRealAssetTime(int32 PoseIdx) const
+{
+	check(Schema);
+	const FPoseSearchIndexAsset& Asset = GetSearchIndex().GetAssetForPose(PoseIdx);
+	return Asset.GetTimeFromPoseIndex(PoseIdx, Schema->SampleRate);
+}
+
+float UPoseSearchDatabase::GetNormalizedAssetTime(int32 PoseIdx) const
 {
 	check(Schema);
-	const float SamplingInterval = Schema->GetSamplingInterval();
 	const FPoseSearchIndexAsset& Asset = GetSearchIndex().GetAssetForPose(PoseIdx);
 	const bool bIsBlendSpace = AnimationAssets[Asset.SourceAssetIdx].GetPtr<FPoseSearchDatabaseBlendSpace>() != nullptr;
-	const FFloatInterval& SamplingRange = Asset.SamplingInterval;
 
 	// sequences or anim composites
-	float AssetTime = FMath::Min(SamplingRange.Min + SamplingInterval * (PoseIdx - Asset.FirstPoseIdx), SamplingRange.Max);
+	float AssetTime = Asset.GetTimeFromPoseIndex(PoseIdx, Schema->SampleRate);
 
-	if (bIsBlendSpace && Asset.NumPoses > 1)
+	if (bIsBlendSpace && Asset.GetNumPoses() > 1)
 	{
-		check(SamplingInterval > UE_KINDA_SMALL_NUMBER);
 		// For BlendSpaces the AssetTime is in the range [0, 1] while the Sampling Range
-		// is in real time (seconds)
-		AssetTime = AssetTime / ((Asset.NumPoses - 1) * SamplingInterval);
+		// is in real time (seconds). We should be using but FAnimationAssetSampler::GetPlayLength(...) to normalize precisely,
+		// but Asset.GetNumPoses() - 1 is a good enough estimator
+		AssetTime = AssetTime * Schema->SampleRate / float(Asset.GetNumPoses() - 1);
 		check(AssetTime >= 0.f && AssetTime <= 1.f);
 	}
 
@@ -765,7 +738,7 @@ FPoseSearchCost UPoseSearchDatabase::SearchContinuingPose(UE::PoseSearch::FSearc
 	const FPoseSearchDatabaseAnimationAssetBase* DatabaseAnimationAssetBase = GetAnimationAssetStruct(SearchIndexAsset).GetPtr<FPoseSearchDatabaseAnimationAssetBase>();
 	check(DatabaseAnimationAssetBase);
 	const FAnimationAssetSampler SequenceBaseSampler(DatabaseAnimationAssetBase->GetAnimationAsset(), SearchIndexAsset.BlendParameters);
-	const float SampleTime = GetAssetTime(PoseIdx);
+	const float SampleTime = GetNormalizedAssetTime(PoseIdx);
 
 	// @todo: change ExtractPoseSearchNotifyStates api to avoid NotifyStates allocation
 	TArray<UAnimNotifyState_PoseSearchBase*> NotifyStates;
@@ -925,7 +898,7 @@ UE::PoseSearch::FSearchResult UPoseSearchDatabase::SearchPCAKDTree(UE::PoseSearc
 	// finalizing Result properties
 	if (Result.PoseIdx != INDEX_NONE)
 	{
-		Result.AssetTime = GetAssetTime(Result.PoseIdx);
+		Result.AssetTime = GetNormalizedAssetTime(Result.PoseIdx);
 		Result.Database = this;
 	}
 
@@ -998,7 +971,7 @@ UE::PoseSearch::FSearchResult UPoseSearchDatabase::SearchBruteForce(UE::PoseSear
 	// finalizing Result properties
 	if (Result.PoseIdx != INDEX_NONE)
 	{
-		Result.AssetTime = GetAssetTime(Result.PoseIdx);
+		Result.AssetTime = GetNormalizedAssetTime(Result.PoseIdx);
 		Result.Database = this;
 	}
 
diff --git a/Engine/Plugins/Experimental/Animation/PoseSearch/Source/Runtime/Private/PoseSearchDatabaseIndexingContext.cpp b/Engine/Plugins/Experimental/Animation/PoseSearch/Source/Runtime/Private/PoseSearchDatabaseIndexingContext.cpp
--- a/Engine/Plugins/Experimental/Animation/PoseSearch/Source/Runtime/Private/PoseSearchDatabaseIndexingContext.cpp
+++ b/Engine/Plugins/Experimental/Animation/PoseSearch/Source/Runtime/Private/PoseSearchDatabaseIndexingContext.cpp
@@ -112,16 +112,14 @@ bool FDatabaseIndexingContext::IndexDatabase(FPoseSearchIndexBase& SearchIndexBa
 	for (int32 AssetIdx = 0; AssetIdx != SearchIndexBase.Assets.Num(); ++AssetIdx)
 	{
 		FPoseSearchIndexAsset& SearchIndexAsset = SearchIndexBase.Assets[AssetIdx];
-		SearchIndexAsset.FirstPoseIdx = TotalPoses;
+		check(SearchIndexAsset.FirstPoseIdx == TotalPoses);
 
-		const FInstancedStruct& DatabaseAsset = Database.GetAnimationAssetStruct(SearchIndexAsset.SourceAssetIdx);
-		const FPoseSearchDatabaseAnimationAssetBase* DatabaseAnimationAssetBase = DatabaseAsset.GetPtr<FPoseSearchDatabaseAnimationAssetBase>();
+		const FPoseSearchDatabaseAnimationAssetBase* DatabaseAnimationAssetBase = Database.GetAnimationAssetStruct(SearchIndexAsset).GetPtr<FPoseSearchDatabaseAnimationAssetBase>();
 		check(DatabaseAnimationAssetBase && DatabaseAnimationAssetBase->GetAnimationAsset());
 		const FAnimationAssetSampler& AssetSampler = Samplers[SamplerMap[{ DatabaseAnimationAssetBase->GetAnimationAsset(), SearchIndexAsset.BlendParameters }]];
 
-		const int32 NewIndexerIdx = Indexers.Emplace(BoneContainer, SearchIndexAsset, SamplingContext, *Schema, AssetSampler);
-		SearchIndexAsset.NumPoses = Indexers[NewIndexerIdx].GetNumIndexedPoses();
-		TotalPoses += SearchIndexAsset.NumPoses;
+		Indexers.Emplace(BoneContainer, SearchIndexAsset, SamplingContext, *Schema, AssetSampler);
+		TotalPoses += SearchIndexAsset.GetNumPoses();
 	}
 
 	// allocating Values and PoseMetadata
diff --git a/Engine/Plugins/Experimental/Animation/PoseSearch/Source/Runtime/Private/PoseSearchDerivedData.cpp b/Engine/Plugins/Experimental/Animation/PoseSearch/Source/Runtime/Private/PoseSearchDerivedData.cpp
--- a/Engine/Plugins/Experimental/Animation/PoseSearch/Source/Runtime/Private/PoseSearchDerivedData.cpp
+++ b/Engine/Plugins/Experimental/Animation/PoseSearch/Source/Runtime/Private/PoseSearchDerivedData.cpp
@@ -277,6 +277,7 @@ static void InitSearchIndexAssets(FPoseSearchIndexBase& SearchIndex, const UPose
 	TArray<FFloatRange> ValidRanges;
 	TArray<FBlendSampleData> BlendSamples;
 
+	int32 TotalPoses = 0;
 	for (int32 AnimationAssetIndex = 0; AnimationAssetIndex < Database->AnimationAssets.Num(); ++AnimationAssetIndex)
 	{
 		const FInstancedStruct& DatabaseAssetStruct = Database->GetAnimationAssetStruct(AnimationAssetIndex);
@@ -289,6 +290,7 @@ static void InitSearchIndexAssets(FPoseSearchIndexBase& SearchIndex, const UPose
 
 			const bool bAddUnmirrored = DatabaseAsset->GetMirrorOption() == EPoseSearchMirrorOption::UnmirroredOnly || DatabaseAsset->GetMirrorOption() == EPoseSearchMirrorOption::UnmirroredAndMirrored;
 			const bool bAddMirrored = DatabaseAsset->GetMirrorOption() == EPoseSearchMirrorOption::MirroredOnly || DatabaseAsset->GetMirrorOption() == EPoseSearchMirrorOption::UnmirroredAndMirrored;
+			const int32 SchemaSampleRate = Database->Schema->SampleRate;
 
 			if (const FPoseSearchDatabaseBlendSpace* DatabaseBlendSpace = DatabaseAssetStruct.GetPtr<FPoseSearchDatabaseBlendSpace>())
 			{
@@ -312,12 +314,22 @@ static void InitSearchIndexAssets(FPoseSearchIndexBase& SearchIndex, const UPose
 						{
 							if (bAddUnmirrored)
 							{
-								SearchIndex.Assets.Emplace(AnimationAssetIndex, false, FFloatInterval(0.0f, PlayLength), PermutationIdx, BlendParameters);
+								const FPoseSearchIndexAsset PoseSearchIndexAsset(AnimationAssetIndex, TotalPoses, false, FFloatInterval(0.f, PlayLength), SchemaSampleRate, PermutationIdx, BlendParameters);
+								if (PoseSearchIndexAsset.GetNumPoses() > 0)
+								{
+									SearchIndex.Assets.Add(PoseSearchIndexAsset);
+									TotalPoses += PoseSearchIndexAsset.GetNumPoses();
+								}
 							}
 
 							if (bAddMirrored)
 							{
-								SearchIndex.Assets.Emplace(AnimationAssetIndex, true, FFloatInterval(0.0f, PlayLength), PermutationIdx, BlendParameters);
+								const FPoseSearchIndexAsset PoseSearchIndexAsset(AnimationAssetIndex, TotalPoses, true, FFloatInterval(0.f, PlayLength), SchemaSampleRate, PermutationIdx, BlendParameters);
+								if (PoseSearchIndexAsset.GetNumPoses() > 0)
+								{
+									SearchIndex.Assets.Add(PoseSearchIndexAsset);
+									TotalPoses += PoseSearchIndexAsset.GetNumPoses();
+								}
 							}
 						}
 					}
@@ -335,12 +347,22 @@ static void InitSearchIndexAssets(FPoseSearchIndexBase& SearchIndex, const UPose
 					{
 						if (bAddUnmirrored)
 						{
-							SearchIndex.Assets.Emplace(AnimationAssetIndex, false, FFloatInterval(Range.GetLowerBoundValue(), Range.GetUpperBoundValue()), PermutationIdx);
+							const FPoseSearchIndexAsset PoseSearchIndexAsset(AnimationAssetIndex, TotalPoses, false, FFloatInterval(Range.GetLowerBoundValue(), Range.GetUpperBoundValue()), SchemaSampleRate, PermutationIdx);
+							if (PoseSearchIndexAsset.GetNumPoses() > 0)
+							{
+								SearchIndex.Assets.Add(PoseSearchIndexAsset);
+								TotalPoses += PoseSearchIndexAsset.GetNumPoses();
+							}
 						}
 
 						if (bAddMirrored)
 						{
-							SearchIndex.Assets.Emplace(AnimationAssetIndex, true, FFloatInterval(Range.GetLowerBoundValue(), Range.GetUpperBoundValue()), PermutationIdx);
+							const FPoseSearchIndexAsset PoseSearchIndexAsset(AnimationAssetIndex, TotalPoses, true, FFloatInterval(Range.GetLowerBoundValue(), Range.GetUpperBoundValue()), SchemaSampleRate, PermutationIdx);
+							if (PoseSearchIndexAsset.GetNumPoses() > 0)
+							{
+								SearchIndex.Assets.Add(PoseSearchIndexAsset);
+								TotalPoses += PoseSearchIndexAsset.GetNumPoses();
+							}
 						}
 					}
 				}
diff --git a/Engine/Plugins/Experimental/Animation/PoseSearch/Source/Runtime/Private/PoseSearchDerivedDataKey.cpp b/Engine/Plugins/Experimental/Animation/PoseSearch/Source/Runtime/Private/PoseSearchDerivedDataKey.cpp
--- a/Engine/Plugins/Experimental/Animation/PoseSearch/Source/Runtime/Private/PoseSearchDerivedDataKey.cpp
+++ b/Engine/Plugins/Experimental/Animation/PoseSearch/Source/Runtime/Private/PoseSearchDerivedDataKey.cpp
@@ -41,7 +41,7 @@ FKeyBuilder::FKeyBuilder(const UObject* Object, bool bUseDataVer)
 	if (bUseDataVer)
 	{
 		// used to invalidate the key without having to change POSESEARCHDB_DERIVEDDATA_VER all the times
-		int32 POSESEARCHDB_DERIVEDDATA_VER_SMALL = 147;
+		int32 POSESEARCHDB_DERIVEDDATA_VER_SMALL = 153;
 		FGuid VersionGuid = FDevSystemGuids::GetSystemGuid(FDevSystemGuids::Get().POSESEARCHDB_DERIVEDDATA_VER);
 
 		*this << VersionGuid;
diff --git a/Engine/Plugins/Experimental/Animation/PoseSearch/Source/Runtime/Private/PoseSearchFeatureChannel_Phase.cpp b/Engine/Plugins/Experimental/Animation/PoseSearch/Source/Runtime/Private/PoseSearchFeatureChannel_Phase.cpp
--- a/Engine/Plugins/Experimental/Animation/PoseSearch/Source/Runtime/Private/PoseSearchFeatureChannel_Phase.cpp
+++ b/Engine/Plugins/Experimental/Animation/PoseSearch/Source/Runtime/Private/PoseSearchFeatureChannel_Phase.cpp
@@ -363,9 +363,7 @@ void UPoseSearchFeatureChannel_Phase::IndexAsset(UE::PoseSearch::FAssetIndexer&
 	static float BoneSamplingCentralDifferencesTime = 0.2f; // seconds
 	static float SmoothingWindowTime = 0.3f; // seconds
 
-
 	const UPoseSearchSchema* Schema = Indexer.GetSchema();
-	const float FiniteDelta = Schema->GetSamplingInterval();
 
 	TArray<FVector2D> Phases;
 	TArray<float> Signal;
@@ -376,10 +374,10 @@ void UPoseSearchFeatureChannel_Phase::IndexAsset(UE::PoseSearch::FAssetIndexer&
 	CollectBonePositions(BonePositions, Indexer, SchemaBoneIdx);
 
 	// @todo: have different way of calculating signals, for example: height of the bone transform, acceleration, etc?
-	const int32 BoneSamplingCentralDifferencesOffset = FMath::Max(FMath::CeilToInt(BoneSamplingCentralDifferencesTime / FiniteDelta), 1);
+	const int32 BoneSamplingCentralDifferencesOffset = FMath::Max(FMath::CeilToInt(BoneSamplingCentralDifferencesTime * Schema->SampleRate), 1);
 	CalculateSignal(BonePositions, Signal, BoneSamplingCentralDifferencesOffset);
 
-	const int32 SmoothingWindowOffset = FMath::Max(FMath::CeilToInt(SmoothingWindowTime / FiniteDelta), 1);
+	const int32 SmoothingWindowOffset = FMath::Max(FMath::CeilToInt(SmoothingWindowTime * Schema->SampleRate), 1);
 	SmoothSignal(Signal, SmoothedSignal, SmoothingWindowOffset);
 
 	FindLocalMinMax(SmoothedSignal, LocalMinMax);
diff --git a/Engine/Plugins/Experimental/Animation/PoseSearch/Source/Runtime/Private/PoseSearchIndex.cpp b/Engine/Plugins/Experimental/Animation/PoseSearch/Source/Runtime/Private/PoseSearchIndex.cpp
--- a/Engine/Plugins/Experimental/Animation/PoseSearch/Source/Runtime/Private/PoseSearchIndex.cpp
+++ b/Engine/Plugins/Experimental/Animation/PoseSearch/Source/Runtime/Private/PoseSearchIndex.cpp
@@ -46,11 +46,11 @@ FArchive& operator<<(FArchive& Ar, FPoseSearchIndexAsset& IndexAsset)
 {
 	Ar << IndexAsset.SourceAssetIdx;
 	Ar << IndexAsset.bMirrored;
-	Ar << IndexAsset.SamplingInterval;
 	Ar << IndexAsset.PermutationIdx;
 	Ar << IndexAsset.BlendParameters;
 	Ar << IndexAsset.FirstPoseIdx;
-	Ar << IndexAsset.NumPoses;
+	Ar << IndexAsset.FirstSampleIdx;
+	Ar << IndexAsset.LastSampleIdx;
 	return Ar;
 }
 
diff --git a/Engine/Plugins/Experimental/Animation/PoseSearch/Source/Runtime/Private/PoseSearchLibrary.cpp b/Engine/Plugins/Experimental/Animation/PoseSearch/Source/Runtime/Private/PoseSearchLibrary.cpp
--- a/Engine/Plugins/Experimental/Animation/PoseSearch/Source/Runtime/Private/PoseSearchLibrary.cpp
+++ b/Engine/Plugins/Experimental/Animation/PoseSearch/Source/Runtime/Private/PoseSearchLibrary.cpp
@@ -53,41 +53,6 @@ void FMotionMatchingState::AdjustAssetTime(float AssetTime)
 	CurrentSearchResult.Update(AssetTime);
 }
 
-bool FMotionMatchingState::CanAdvance(float DeltaTime) const
-{
-	if (CurrentSearchResult.IsValid())
-	{
-		ETypeAdvanceAnim AdvanceType = ETypeAdvanceAnim::ETAA_Default;
-		float SteppedTime = CurrentSearchResult.AssetTime;
-		const FPoseSearchIndexAsset* SearchIndexAsset = CurrentSearchResult.GetSearchIndexAsset(true);
-		const FInstancedStruct& DatabaseAsset = CurrentSearchResult.Database->GetAnimationAssetStruct(*SearchIndexAsset);
-		if (const FPoseSearchDatabaseBlendSpace* DatabaseBlendSpace = DatabaseAsset.GetPtr<FPoseSearchDatabaseBlendSpace>())
-		{
-			TArray<FBlendSampleData> BlendSamples;
-			int32 TriangulationIndex = 0;
-			DatabaseBlendSpace->BlendSpace->GetSamplesFromBlendInput(SearchIndexAsset->BlendParameters, BlendSamples, TriangulationIndex, true);
-
-			const float PlayLength = DatabaseBlendSpace->BlendSpace->GetAnimationLengthFromSampleData(BlendSamples);
-
-			// Asset player time for blend spaces is normalized [0, 1] so we need to convert it back to real time before we advance it
-			SteppedTime = CurrentSearchResult.AssetTime * PlayLength;
-			AdvanceType = FAnimationRuntime::AdvanceTime(DatabaseBlendSpace->IsLooping(), DeltaTime, SteppedTime, PlayLength);
-		}
-		else if (const FPoseSearchDatabaseAnimationAssetBase* DatabaseAnimationAssetBase = DatabaseAsset.GetPtr<FPoseSearchDatabaseAnimationAssetBase>())
-		{
-			const float AssetLength = DatabaseAnimationAssetBase->GetAnimationAsset()->GetPlayLength();
-			AdvanceType = FAnimationRuntime::AdvanceTime(DatabaseAnimationAssetBase->IsLooping(), DeltaTime, SteppedTime, AssetLength);
-		}
-
-		if (AdvanceType != ETAA_Finished)
-		{
-			return SearchIndexAsset->SamplingInterval.Contains(SteppedTime);
-		}
-	}
-
-	return false;
-}
-
 static void RequestInertialBlend(const FAnimationUpdateContext& Context, float BlendTime)
 {
 	// Use inertial blending to smooth over the transition
@@ -343,17 +308,19 @@ void UPoseSearchLibrary::UpdateMotionMatchingState(
 
 	const FPoseSearchQueryTrajectory TrajectoryRootSpace = ProcessTrajectory(Trajectory, Context.AnimInstanceProxy->GetComponentTransform(), InOutMotionMatchingState.GetRootBoneDeltaYaw(), RootBoneDeltaYawBlendTime, TrajectorySpeedMultiplier);
 	FSearchContext SearchContext(&TrajectoryRootSpace, History, 0.f, &InOutMotionMatchingState.PoseIndicesHistory, QueryMirrorRequest,
-		InOutMotionMatchingState.CurrentSearchResult, PoseJumpThresholdTime, bForceInterrupt, InOutMotionMatchingState.CanAdvance(DeltaTime));
+		InOutMotionMatchingState.CurrentSearchResult, PoseJumpThresholdTime, bForceInterrupt);
+
+	const bool bMustSearch = !InOutMotionMatchingState.CurrentSearchResult.IsValid();
 
 	// If we can't advance or enough time has elapsed since the last pose jump then search
-	const bool bSearch = !SearchContext.CanAdvance() || (bShouldSearch && (InOutMotionMatchingState.ElapsedPoseSearchTime >= SearchThrottleTime));
+	const bool bSearch = bMustSearch || (bShouldSearch && (InOutMotionMatchingState.ElapsedPoseSearchTime >= SearchThrottleTime));
 	if (bSearch)
 	{
 		InOutMotionMatchingState.ElapsedPoseSearchTime = 0.f;
 
 		// Evaluate continuing pose
 		FSearchResult SearchResult;
-		if (!SearchContext.IsForceInterrupt() && SearchContext.CanAdvance() && SearchContext.GetCurrentResult().Database.IsValid())
+		if (!SearchContext.IsForceInterrupt() && SearchContext.GetCurrentResult().IsValid())
 		{
 			SearchResult.PoseCost = SearchContext.GetCurrentResult().Database->SearchContinuingPose(SearchContext);
 			SearchContext.UpdateCurrentBestCost(SearchResult.PoseCost);
diff --git a/Engine/Plugins/Experimental/Animation/PoseSearch/Source/Runtime/Public/PoseSearch/PoseSearchAssetIndexer.h b/Engine/Plugins/Experimental/Animation/PoseSearch/Source/Runtime/Public/PoseSearch/PoseSearchAssetIndexer.h
--- a/Engine/Plugins/Experimental/Animation/PoseSearch/Source/Runtime/Public/PoseSearch/PoseSearchAssetIndexer.h
+++ b/Engine/Plugins/Experimental/Animation/PoseSearch/Source/Runtime/Public/PoseSearch/PoseSearchAssetIndexer.h
@@ -41,9 +41,9 @@ public:
 	FVector GetSamplePosition(float SampleTimeOffset, int32 SampleIdx, int8 SchemaSampleBoneIdx = RootSchemaBoneIdx, int8 SchemaOriginBoneIdx = RootSchemaBoneIdx, EPermutationTimeType PermutationTimeType = EPermutationTimeType::UseSampleTime);
 	FVector GetSampleVelocity(float SampleTimeOffset, int32 SampleIdx, int8 SchemaSampleBoneIdx = RootSchemaBoneIdx, int8 SchemaOriginBoneIdx = RootSchemaBoneIdx, bool bUseCharacterSpaceVelocities = true, EPermutationTimeType PermutationTimeType = EPermutationTimeType::UseSampleTime);
 
-	int32 GetBeginSampleIdx() const { return FirstIndexedSample; }
-	int32 GetEndSampleIdx() const { return LastIndexedSample + 1; }
-	int32 GetNumIndexedPoses() const { return GetEndSampleIdx() - GetBeginSampleIdx(); }
+	int32 GetBeginSampleIdx() const;
+	int32 GetEndSampleIdx() const;
+	int32 GetNumIndexedPoses() const;
 	
 	TArrayView<float> GetPoseVector(int32 SampleIdx) const;
 	const UPoseSearchSchema* GetSchema() const;
@@ -85,9 +85,6 @@ private:
 	const UPoseSearchSchema& Schema;
 	const FAnimationAssetSampler& AssetSampler;
 	
-	int32 FirstIndexedSample = 0;
-	int32 LastIndexedSample = 0;
-
 	TArrayView<float> FeatureVectorTable;
 	TArrayView<FPoseSearchPoseMetadata> PoseMetadata;
 
diff --git a/Engine/Plugins/Experimental/Animation/PoseSearch/Source/Runtime/Public/PoseSearch/PoseSearchAssetSampler.h b/Engine/Plugins/Experimental/Animation/PoseSearch/Source/Runtime/Public/PoseSearch/PoseSearchAssetSampler.h
--- a/Engine/Plugins/Experimental/Animation/PoseSearch/Source/Runtime/Public/PoseSearch/PoseSearchAssetSampler.h
+++ b/Engine/Plugins/Experimental/Animation/PoseSearch/Source/Runtime/Public/PoseSearch/PoseSearchAssetSampler.h
@@ -47,6 +47,8 @@ struct POSESEARCH_API FAnimationAssetSampler
 
 	void Process(const FBoneContainer& BoneContainer);
 
+	static float GetPlayLength(const UAnimationAsset* AnimAsset, const FVector& BlendParameters = FVector::ZeroVector);
+
 protected:
 	TWeakObjectPtr<const UAnimationAsset> AnimationAsset;
 
diff --git a/Engine/Plugins/Experimental/Animation/PoseSearch/Source/Runtime/Public/PoseSearch/PoseSearchContext.h b/Engine/Plugins/Experimental/Animation/PoseSearch/Source/Runtime/Public/PoseSearch/PoseSearchContext.h
--- a/Engine/Plugins/Experimental/Animation/PoseSearch/Source/Runtime/Public/PoseSearch/PoseSearchContext.h
+++ b/Engine/Plugins/Experimental/Animation/PoseSearch/Source/Runtime/Public/PoseSearch/PoseSearchContext.h
@@ -141,7 +141,7 @@ struct POSESEARCH_API FSearchContext
 {
 	FSearchContext(const FPoseSearchQueryTrajectory* InTrajectory, const IPoseHistory* InHistory, float InDesiredPermutationTimeOffset,
 		const FPoseIndicesHistory* InPoseIndicesHistory = nullptr, EPoseSearchBooleanRequest InQueryMirrorRequest = EPoseSearchBooleanRequest::Indifferent,
-		const FSearchResult& InCurrentResult = FSearchResult(), float InPoseJumpThresholdTime = 0.f, bool bInForceInterrupt = false, bool bInCanAdvance = true);
+		const FSearchResult& InCurrentResult = FSearchResult(), float InPoseJumpThresholdTime = 0.f, bool bInForceInterrupt = false);
 
 	FQuat GetSampleRotation(float SampleTimeOffset, const UPoseSearchSchema* Schema, int8 SchemaSampleBoneIdx = RootSchemaBoneIdx, int8 SchemaOriginBoneIdx = RootSchemaBoneIdx, bool bUseHistoryRoot = false, EPermutationTimeType PermutationTimeType = EPermutationTimeType::UseSampleTime);
 	FVector GetSamplePosition(float SampleTimeOffset, const UPoseSearchSchema* Schema, int8 SchemaSampleBoneIdx = RootSchemaBoneIdx, int8 SchemaOriginBoneIdx = RootSchemaBoneIdx, bool bUseHistoryRoot = false, EPermutationTimeType PermutationTimeType = EPermutationTimeType::UseSampleTime);
@@ -169,7 +169,6 @@ struct POSESEARCH_API FSearchContext
 	bool IsHistoryValid() const { return History != nullptr; }
 	float GetDesiredPermutationTimeOffset() const { return DesiredPermutationTimeOffset; }
 	bool IsTrajectoryValid() const { return Trajectory != nullptr; }
-	bool CanAdvance() const { return bCanAdvance; }
 	bool IsForceInterrupt() const { return bForceInterrupt; }
 	FTransform GetRootAtTime(float Time, bool bUseHistoryRoot = false, bool bExtrapolate = true) const;
 
@@ -187,8 +186,6 @@ private:
 	FSearchResult CurrentResult;
 	float PoseJumpThresholdTime = 0.f;
 	bool bForceInterrupt = false;
-	// can the continuing pose advance? (if not we skip evaluating it)
-	bool bCanAdvance = true;
 
 	// cached previous, current, and next pose data
 	TConstArrayView<float> CurrentResultNextPoseVector;
diff --git a/Engine/Plugins/Experimental/Animation/PoseSearch/Source/Runtime/Public/PoseSearch/PoseSearchDatabase.h b/Engine/Plugins/Experimental/Animation/PoseSearch/Source/Runtime/Public/PoseSearch/PoseSearchDatabase.h
--- a/Engine/Plugins/Experimental/Animation/PoseSearch/Source/Runtime/Public/PoseSearch/PoseSearchDatabase.h
+++ b/Engine/Plugins/Experimental/Animation/PoseSearch/Source/Runtime/Public/PoseSearch/PoseSearchDatabase.h
@@ -334,7 +334,8 @@ public:
 	const bool IsSourceAssetLooping(const FPoseSearchIndexAsset& SearchIndexAsset) const;
 	const FString GetSourceAssetName(const FPoseSearchIndexAsset& SearchIndexAsset) const;
 	int32 GetNumberOfPrincipalComponents() const;
-	float GetAssetTime(int32 PoseIdx) const;
+	float GetRealAssetTime(int32 PoseIdx) const;
+	float GetNormalizedAssetTime(int32 PoseIdx) const;
 
 	// Begin UObject
 	virtual void PostLoad() override;
diff --git a/Engine/Plugins/Experimental/Animation/PoseSearch/Source/Runtime/Public/PoseSearch/PoseSearchIndex.h b/Engine/Plugins/Experimental/Animation/PoseSearch/Source/Runtime/Public/PoseSearch/PoseSearchIndex.h
--- a/Engine/Plugins/Experimental/Animation/PoseSearch/Source/Runtime/Public/PoseSearch/PoseSearchIndex.h
+++ b/Engine/Plugins/Experimental/Animation/PoseSearch/Source/Runtime/Public/PoseSearch/PoseSearchIndex.h
@@ -80,19 +80,22 @@ struct POSESEARCH_API FPoseSearchIndexAsset
 	FPoseSearchIndexAsset() {}
 
 	FPoseSearchIndexAsset(
-		int32 InSourceAssetIdx, 
+		int32 InSourceAssetIdx,
+		int32 InFirstPoseIdx,
 		bool bInMirrored, 
 		const FFloatInterval& InSamplingInterval,
+		int32 SchemaSampleRate,
 		int32 InPermutationIdx,
 		FVector InBlendParameters = FVector::Zero())
 		: SourceAssetIdx(InSourceAssetIdx)
 		, bMirrored(bInMirrored)
-		, SamplingInterval(InSamplingInterval)
 		, PermutationIdx(InPermutationIdx)
 		, BlendParameters(InBlendParameters)
-		, FirstPoseIdx(INDEX_NONE)
-		, NumPoses(INDEX_NONE)
+		, FirstPoseIdx(InFirstPoseIdx)
+		, FirstSampleIdx(FMath::CeilToInt(InSamplingInterval.Min * SchemaSampleRate))
+		, LastSampleIdx(FMath::FloorToInt(InSamplingInterval.Max * SchemaSampleRate))
 	{
+		check(SchemaSampleRate > 0);
 	}
 
 	// Index of the source asset in search index's container (i.e. UPoseSearchDatabase)
@@ -102,9 +105,6 @@ struct POSESEARCH_API FPoseSearchIndexAsset
 	UPROPERTY(meta = (NeverInHash))
 	bool bMirrored = false;
 
-	UPROPERTY(meta = (NeverInHash))
-	FFloatInterval SamplingInterval;
-
 	UPROPERTY(meta = (NeverInHash))
 	int32 PermutationIdx = INDEX_NONE;
 
@@ -115,10 +115,66 @@ struct POSESEARCH_API FPoseSearchIndexAsset
 	int32 FirstPoseIdx = INDEX_NONE;
 
 	UPROPERTY(meta = (NeverInHash))
-	int32 NumPoses = INDEX_NONE;
+	int32 FirstSampleIdx = INDEX_NONE;
+
+	UPROPERTY(meta = (NeverInHash))
+	int32 LastSampleIdx = INDEX_NONE;
+
+	bool IsPoseInRange(int32 PoseIdx) const { return (PoseIdx >= FirstPoseIdx) && (PoseIdx < FirstPoseIdx + GetNumPoses()); }
+	bool IsInitialized() const 
+	{
+		return 
+			SourceAssetIdx != INDEX_NONE &&
+			PermutationIdx != INDEX_NONE &&
+			FirstPoseIdx != INDEX_NONE &&
+			FirstSampleIdx != INDEX_NONE &&
+			LastSampleIdx != INDEX_NONE;
+	}
+
+	int32 GetBeginSampleIdx() const { return FirstSampleIdx; }
+	int32 GetEndSampleIdx() const {	return LastSampleIdx + 1; }
+	int32 GetNumPoses() const { return GetEndSampleIdx() - GetBeginSampleIdx(); }
 
-	bool IsPoseInRange(int32 PoseIdx) const { return (PoseIdx >= FirstPoseIdx) && (PoseIdx < FirstPoseIdx + NumPoses); }
-	bool IsInitialized() const { return SourceAssetIdx != INDEX_NONE && PermutationIdx != INDEX_NONE && FirstPoseIdx != INDEX_NONE && NumPoses != INDEX_NONE; }
+	float GetFirstSampleTime(int32 SchemaSampleRate) const { check(SchemaSampleRate > 0); return FirstSampleIdx / float(SchemaSampleRate); }
+	float GetLastSampleTime(int32 SchemaSampleRate) const { check(SchemaSampleRate > 0); return LastSampleIdx / float(SchemaSampleRate); }
+
+	int32 GetPoseIndexFromTime(float Time, bool bIsLooping, int32 SchemaSampleRate) const
+	{
+		check(IsInitialized());
+
+		const int32 NumPoses = GetNumPoses();
+		int32 PoseOffset = FMath::RoundToInt(SchemaSampleRate * Time) - FirstSampleIdx;
+		if (bIsLooping)
+		{
+			if (PoseOffset < 0)
+			{
+				PoseOffset = (PoseOffset % NumPoses) + NumPoses;
+			}
+			else if (PoseOffset >= NumPoses)
+			{
+				PoseOffset = PoseOffset % NumPoses;
+			}
+			return FirstPoseIdx + PoseOffset;
+		}
+
+		if (PoseOffset >= 0 && PoseOffset < NumPoses)
+		{
+			return FirstPoseIdx + PoseOffset;
+		}
+
+		return INDEX_NONE;
+	}
+
+	float GetTimeFromPoseIndex(int32 PoseIdx, int32 SchemaSampleRate) const
+	{
+		check(SchemaSampleRate > 0);
+
+		const int32 PoseOffset = PoseIdx - FirstPoseIdx;
+		check(PoseOffset >= 0 && PoseOffset < GetNumPoses());
+
+		const float Time = (FirstSampleIdx + PoseOffset) / float(SchemaSampleRate);
+		return Time;
+	}
 
 	friend FArchive& operator<<(FArchive& Ar, FPoseSearchIndexAsset& IndexAsset);
 };
diff --git a/Engine/Plugins/Experimental/Animation/PoseSearch/Source/Runtime/Public/PoseSearch/PoseSearchLibrary.h b/Engine/Plugins/Experimental/Animation/PoseSearch/Source/Runtime/Public/PoseSearch/PoseSearchLibrary.h
--- a/Engine/Plugins/Experimental/Animation/PoseSearch/Source/Runtime/Public/PoseSearch/PoseSearchLibrary.h
+++ b/Engine/Plugins/Experimental/Animation/PoseSearch/Source/Runtime/Public/PoseSearch/PoseSearchLibrary.h
@@ -31,9 +31,6 @@ struct POSESEARCH_API FMotionMatchingState
 	// Reset the state to a default state using the current Database
 	void Reset();
 
-	// Checks if the currently playing asset can advance and stay in bounds under the provided DeltaTime.
-	bool CanAdvance(float DeltaTime) const;
-
 	// Attempts to set the internal state to match the provided asset time including updating the internal DbPoseIdx. 
 	// If the provided asset time is out of bounds for the currently playing asset then this function will reset the 
 	// state back to the default state.
diff --git a/Engine/Plugins/Experimental/Animation/PoseSearch/Source/Runtime/Public/PoseSearch/PoseSearchSchema.h b/Engine/Plugins/Experimental/Animation/PoseSearch/Source/Runtime/Public/PoseSearch/PoseSearchSchema.h
--- a/Engine/Plugins/Experimental/Animation/PoseSearch/Source/Runtime/Public/PoseSearch/PoseSearchSchema.h
+++ b/Engine/Plugins/Experimental/Animation/PoseSearch/Source/Runtime/Public/PoseSearch/PoseSearchSchema.h
@@ -113,8 +113,6 @@ public:
 
 	bool IsValid () const;
 
-	float GetSamplingInterval() const { return 1.0f / SampleRate; }
-
 	TConstArrayView<TObjectPtr<UPoseSearchFeatureChannel>> GetChannels() const { return FinalizedChannels; }
 
 	void AddChannel(UPoseSearchFeatureChannel* Channel);
