motion matching - fix for pose search database DDC indexing while GC is reclaiming memory

diff --git a/Engine/Plugins/Experimental/Animation/PoseSearch/Source/Editor/Private/PoseSearchDatabaseViewModel.cpp b/Engine/Plugins/Experimental/Animation/PoseSearch/Source/Editor/Private/PoseSearchDatabaseViewModel.cpp
--- a/Engine/Plugins/Experimental/Animation/PoseSearch/Source/Editor/Private/PoseSearchDatabaseViewModel.cpp
+++ b/Engine/Plugins/Experimental/Animation/PoseSearch/Source/Editor/Private/PoseSearchDatabaseViewModel.cpp
@@ -285,7 +285,6 @@ namespace UE::PoseSearch
 
 		for (FDatabasePreviewActor& PreviewActor : PreviewActors)
 		{
-			// @todo: PreviewActor.Actor is a TWeakObjectPtr so it can be null.
 			PreviewActor.Actor->Destroy();
 		}
 		PreviewActors.Reset();
diff --git a/Engine/Plugins/Experimental/Animation/PoseSearch/Source/Editor/Private/PoseSearchDebuggerDatabaseRowData.h b/Engine/Plugins/Experimental/Animation/PoseSearch/Source/Editor/Private/PoseSearchDebuggerDatabaseRowData.h
--- a/Engine/Plugins/Experimental/Animation/PoseSearch/Source/Editor/Private/PoseSearchDebuggerDatabaseRowData.h
+++ b/Engine/Plugins/Experimental/Animation/PoseSearch/Source/Editor/Private/PoseSearchDebuggerDatabaseRowData.h
@@ -12,7 +12,7 @@ namespace UE::PoseSearch
 class FDebuggerDatabaseSharedData : public TSharedFromThis<FDebuggerDatabaseSharedData>
 {
 public:
-	TWeakObjectPtr<const UPoseSearchDatabase> SourceDatabase = nullptr;
+	TWeakObjectPtr<const UPoseSearchDatabase> SourceDatabase;
 	FString DatabaseName = "";
 	FString DatabasePath = "";
 	TArray<float> QueryVector;
diff --git a/Engine/Plugins/Experimental/Animation/PoseSearch/Source/Editor/Private/PoseSearchDebuggerViewModel.cpp b/Engine/Plugins/Experimental/Animation/PoseSearch/Source/Editor/Private/PoseSearchDebuggerViewModel.cpp
--- a/Engine/Plugins/Experimental/Animation/PoseSearch/Source/Editor/Private/PoseSearchDebuggerViewModel.cpp
+++ b/Engine/Plugins/Experimental/Animation/PoseSearch/Source/Editor/Private/PoseSearchDebuggerViewModel.cpp
@@ -27,7 +27,7 @@ FDebuggerViewModel::~FDebuggerViewModel()
 {
 	for (FSkeleton& Skeleton : Skeletons)
 	{
-		if (Skeleton.Actor.IsValid())
+		if (Skeleton.Actor != nullptr)
 		{
 			Skeleton.Actor->Destroy();
 		}
diff --git a/Engine/Plugins/Experimental/Animation/PoseSearch/Source/Editor/Private/PoseSearchDebuggerViewModel.h b/Engine/Plugins/Experimental/Animation/PoseSearch/Source/Editor/Private/PoseSearchDebuggerViewModel.h
--- a/Engine/Plugins/Experimental/Animation/PoseSearch/Source/Editor/Private/PoseSearchDebuggerViewModel.h
+++ b/Engine/Plugins/Experimental/Animation/PoseSearch/Source/Editor/Private/PoseSearchDebuggerViewModel.h
@@ -62,13 +62,13 @@ private:
 	struct FSkeleton
 	{
 		/** Actor object for the skeleton */
-		TWeakObjectPtr<AActor> Actor = nullptr;
+		TObjectPtr<AActor> Actor;
 
 		/** Derived skeletal mesh for setting the skeleton in the scene */
-		TWeakObjectPtr<UPoseSearchMeshComponent> Component = nullptr;
+		TWeakObjectPtr<UPoseSearchMeshComponent> Component;
 
 		/** Source database for this skeleton  */
-		TWeakObjectPtr<const UPoseSearchDatabase> SourceDatabase = nullptr;
+		TWeakObjectPtr<const UPoseSearchDatabase> SourceDatabase;
 
 		/** Source asset for this skeleton */
 		int32 AssetIdx = 0;
diff --git a/Engine/Plugins/Experimental/Animation/PoseSearch/Source/Editor/Public/PoseSearchDatabaseEdMode.h b/Engine/Plugins/Experimental/Animation/PoseSearch/Source/Editor/Public/PoseSearchDatabaseEdMode.h
--- a/Engine/Plugins/Experimental/Animation/PoseSearch/Source/Editor/Public/PoseSearchDatabaseEdMode.h
+++ b/Engine/Plugins/Experimental/Animation/PoseSearch/Source/Editor/Public/PoseSearchDatabaseEdMode.h
@@ -45,9 +45,6 @@ namespace UE::PoseSearch
 		virtual bool GetCustomInputCoordinateSystem(FMatrix& InMatrix, void* InData) override;
 
 	private:
-
 		FDatabaseViewModel* ViewModel = nullptr;
-
-		TWeakObjectPtr<AActor> SelectedActor = nullptr;
 	};
 }
diff --git a/Engine/Plugins/Experimental/Animation/PoseSearch/Source/Editor/Public/PoseSearchDatabaseViewModel.h b/Engine/Plugins/Experimental/Animation/PoseSearch/Source/Editor/Public/PoseSearchDatabaseViewModel.h
--- a/Engine/Plugins/Experimental/Animation/PoseSearch/Source/Editor/Public/PoseSearchDatabaseViewModel.h
+++ b/Engine/Plugins/Experimental/Animation/PoseSearch/Source/Editor/Public/PoseSearchDatabaseViewModel.h
@@ -43,7 +43,7 @@ namespace UE::PoseSearch
 	struct FDatabasePreviewActor
 	{
 	public:
-		TWeakObjectPtr<AActor> Actor = nullptr;
+		TObjectPtr<AActor> Actor;
 		int32 IndexAssetIndex = INDEX_NONE;
 		int32 CurrentPoseIndex = INDEX_NONE;
 		float PlayTimeOffset = 0.f;
diff --git a/Engine/Plugins/Experimental/Animation/PoseSearch/Source/Runtime/Private/PoseSearchAssetSampler.cpp b/Engine/Plugins/Experimental/Animation/PoseSearch/Source/Runtime/Private/PoseSearchAssetSampler.cpp
--- a/Engine/Plugins/Experimental/Animation/PoseSearch/Source/Runtime/Private/PoseSearchAssetSampler.cpp
+++ b/Engine/Plugins/Experimental/Animation/PoseSearch/Source/Runtime/Private/PoseSearchAssetSampler.cpp
@@ -235,12 +235,12 @@ static int32 GetHighestWeightSample(const TArray<struct FBlendSampleData>& Sampl
 
 //////////////////////////////////////////////////////////////////////////
 // FAssetSamplerBase
-FAnimationAssetSampler::FAnimationAssetSampler(TWeakObjectPtr<const UAnimationAsset> InAnimationAsset, const FVector& InBlendParameters, int32 InRootTransformSamplingRate)
+FAnimationAssetSampler::FAnimationAssetSampler(TObjectPtr<const UAnimationAsset> InAnimationAsset, const FVector& InBlendParameters, int32 InRootTransformSamplingRate)
 {
 	Init(InAnimationAsset, InBlendParameters, InRootTransformSamplingRate);
 }
 
-void FAnimationAssetSampler::Init(TWeakObjectPtr<const UAnimationAsset> InAnimationAsset, const FVector& InBlendParameters, int32 InRootTransformSamplingRate)
+void FAnimationAssetSampler::Init(TObjectPtr<const UAnimationAsset> InAnimationAsset, const FVector& InBlendParameters, int32 InRootTransformSamplingRate)
 {
 	AnimationAsset = InAnimationAsset;
 	BlendParameters = InBlendParameters;
diff --git a/Engine/Plugins/Experimental/Animation/PoseSearch/Source/Runtime/Private/PoseSearchDerivedData.cpp b/Engine/Plugins/Experimental/Animation/PoseSearch/Source/Runtime/Private/PoseSearchDerivedData.cpp
--- a/Engine/Plugins/Experimental/Animation/PoseSearch/Source/Runtime/Private/PoseSearchDerivedData.cpp
+++ b/Engine/Plugins/Experimental/Animation/PoseSearch/Source/Runtime/Private/PoseSearchDerivedData.cpp
@@ -625,6 +625,7 @@ struct FPoseSearchDatabaseAsyncCacheTask
 	void Wait(FCriticalSection& OuterMutex);
 	void Cancel(FCriticalSection& OuterMutex);
 	bool Poll(FCriticalSection& OuterMutex) const;
+	void AddReferencedObjects(FReferenceCollector& Collector, FCriticalSection& OuterMutex);
 	bool ContainsDatabase(const UPoseSearchDatabase* OtherDatabase, FCriticalSection& OuterMutex) const;
 	bool IsValid(FCriticalSection& OuterMutex) const;
 	const FIoHash& GetDerivedDataKey() const { return DerivedDataKey; }
@@ -883,8 +884,8 @@ void FPoseSearchDatabaseAsyncCacheTask::OnGetComplete(UE::DerivedData::FCacheGet
 				COOK_STAT(auto Timer = UsageStats.TimeSyncWork());
 
 				// collecting all the databases that need to be built to gather their FPoseSearchIndexBase
-				TArray<TWeakObjectPtr<const UPoseSearchDatabase>> IndexBaseDatabases;
-				IndexBaseDatabases.Add(Database); // the first one is always this Database
+				TArray<TObjectPtr<const UPoseSearchDatabase>> IndexBaseDatabases;
+				IndexBaseDatabases.Add(Database.Get()); // the first one is always this Database
 				if (Database->NormalizationSet)
 				{
 					Database->NormalizationSet->AddUniqueDatabases(IndexBaseDatabases);
@@ -897,7 +898,8 @@ void FPoseSearchDatabaseAsyncCacheTask::OnGetComplete(UE::DerivedData::FCacheGet
 				Schemas.AddDefaulted(IndexBaseDatabases.Num());
 				for (int32 IndexBaseIdx = 0; IndexBaseIdx < IndexBaseDatabases.Num(); ++IndexBaseIdx)
 				{
-					TWeakObjectPtr<const UPoseSearchDatabase> IndexBaseDatabase = IndexBaseDatabases[IndexBaseIdx];
+					const UPoseSearchDatabase* IndexBaseDatabase = IndexBaseDatabases[IndexBaseIdx].Get();
+					check(IndexBaseDatabase);
 					FPoseSearchIndexBase& SearchIndexBase = SearchIndexBases[IndexBaseIdx];
 					Schemas[IndexBaseIdx] = IndexBaseDatabase->Schema;
 
@@ -942,7 +944,7 @@ void FPoseSearchDatabaseAsyncCacheTask::OnGetComplete(UE::DerivedData::FCacheGet
 					}
 
 					// Building all the related FPoseSearchBaseIndex first
-					InitSearchIndexAssets(SearchIndexBase, IndexBaseDatabase.Get());
+					InitSearchIndexAssets(SearchIndexBase, IndexBaseDatabase);
 
 					if (Owner.IsCanceled())
 					{
@@ -1031,6 +1033,23 @@ void FPoseSearchDatabaseAsyncCacheTask::OnGetComplete(UE::DerivedData::FCacheGet
 	}
 }
 
+void FPoseSearchDatabaseAsyncCacheTask::AddReferencedObjects(FReferenceCollector& Collector, FCriticalSection& OuterMutex)
+{
+	FScopeLock Lock(&OuterMutex);
+
+	const EState State = GetState();
+	if (State != EState::Ended && State != EState::Failed)
+	{
+		// keeping around the assets for starting or in progress tasks
+		Collector.AddReferencedObject(Database);
+
+		for (TWeakObjectPtr<const UObject>& Dependency : DatabaseDependencies )
+		{
+			Collector.AddReferencedObject(Dependency);
+		}
+	}
+}
+
 //////////////////////////////////////////////////////////////////////////
 // FAsyncPoseSearchDatabasesManagement
 FCriticalSection FAsyncPoseSearchDatabasesManagement::Mutex;
@@ -1139,8 +1158,6 @@ void FAsyncPoseSearchDatabasesManagement::Tick(float DeltaTime)
 
 void FAsyncPoseSearchDatabasesManagement::TickCook(float DeltaTime, bool bCookCompete)
 {
-	FScopeLock Lock(&Mutex);
-
 	Tick(DeltaTime);
 }
 
@@ -1151,6 +1168,12 @@ TStatId FAsyncPoseSearchDatabasesManagement::GetStatId() const
 
 void FAsyncPoseSearchDatabasesManagement::AddReferencedObjects(FReferenceCollector& Collector)
 {
+	FScopeLock Lock(&Mutex);
+
+	for (TUniquePtr<FPoseSearchDatabaseAsyncCacheTask>& TaskPtr : Tasks)
+	{
+		TaskPtr->AddReferencedObjects(Collector, Mutex);
+	}
 }
 
 // returns true if the index has been built and the Database updated correctly  
diff --git a/Engine/Plugins/Experimental/Animation/PoseSearch/Source/Runtime/Private/PoseSearchLibrary.cpp b/Engine/Plugins/Experimental/Animation/PoseSearch/Source/Runtime/Private/PoseSearchLibrary.cpp
--- a/Engine/Plugins/Experimental/Animation/PoseSearch/Source/Runtime/Private/PoseSearchLibrary.cpp
+++ b/Engine/Plugins/Experimental/Animation/PoseSearch/Source/Runtime/Private/PoseSearchLibrary.cpp
@@ -390,7 +390,7 @@ void UPoseSearchLibrary::UpdateMotionMatchingState(
 
 #if WITH_EDITORONLY_DATA && ENABLE_ANIM_DEBUG
 	const FSearchResult& CurResult = InOutMotionMatchingState.CurrentSearchResult;
-	if ((bDebugDrawQuery || bDebugDrawCurResult) && CurResult.Database.IsValid())
+	if ((bDebugDrawQuery || bDebugDrawCurResult) && CurResult.Database != nullptr)
 	{
 		const UPoseSearchDatabase* CurResultDatabase = CurResult.Database.Get();
 
diff --git a/Engine/Plugins/Experimental/Animation/PoseSearch/Source/Runtime/Private/PoseSearchNormalizationSet.cpp b/Engine/Plugins/Experimental/Animation/PoseSearch/Source/Runtime/Private/PoseSearchNormalizationSet.cpp
--- a/Engine/Plugins/Experimental/Animation/PoseSearch/Source/Runtime/Private/PoseSearchNormalizationSet.cpp
+++ b/Engine/Plugins/Experimental/Animation/PoseSearch/Source/Runtime/Private/PoseSearchNormalizationSet.cpp
@@ -3,7 +3,7 @@
 #include "PoseSearch/PoseSearchNormalizationSet.h"
 #include "PoseSearch/PoseSearchDatabase.h"
 
-void UPoseSearchNormalizationSet::AddUniqueDatabases(TArray<TWeakObjectPtr<const UPoseSearchDatabase>>& UniqueDatabases) const
+void UPoseSearchNormalizationSet::AddUniqueDatabases(TArray<TObjectPtr<const UPoseSearchDatabase>>& UniqueDatabases) const
 {
 	for (const TObjectPtr<const UPoseSearchDatabase>& Database : Databases)
 	{
diff --git a/Engine/Plugins/Experimental/Animation/PoseSearch/Source/Runtime/Private/PoseSearchResult.cpp b/Engine/Plugins/Experimental/Animation/PoseSearch/Source/Runtime/Private/PoseSearchResult.cpp
--- a/Engine/Plugins/Experimental/Animation/PoseSearch/Source/Runtime/Private/PoseSearchResult.cpp
+++ b/Engine/Plugins/Experimental/Animation/PoseSearch/Source/Runtime/Private/PoseSearchResult.cpp
@@ -78,7 +78,7 @@ void FSearchResult::Update(float NewAssetTime)
 
 bool FSearchResult::IsValid() const
 {
-	return PoseIdx != INDEX_NONE && Database.IsValid();
+	return PoseIdx != INDEX_NONE && Database != nullptr;
 }
 
 void FSearchResult::Reset()
diff --git a/Engine/Plugins/Experimental/Animation/PoseSearch/Source/Runtime/Public/PoseSearch/PoseSearchAssetSampler.h b/Engine/Plugins/Experimental/Animation/PoseSearch/Source/Runtime/Public/PoseSearch/PoseSearchAssetSampler.h
--- a/Engine/Plugins/Experimental/Animation/PoseSearch/Source/Runtime/Public/PoseSearch/PoseSearchAssetSampler.h
+++ b/Engine/Plugins/Experimental/Animation/PoseSearch/Source/Runtime/Public/PoseSearch/PoseSearchAssetSampler.h
@@ -17,8 +17,8 @@ namespace UE::PoseSearch
  */
 struct POSESEARCH_API FAnimationAssetSampler
 {
-	FAnimationAssetSampler(TWeakObjectPtr<const UAnimationAsset> InAnimationAsset = nullptr, const FVector& InBlendParameters = FVector::ZeroVector, int32 InRootTransformSamplingRate = 30);
-	void Init(TWeakObjectPtr<const UAnimationAsset> InAnimationAsset, const FVector& InBlendParameters = FVector::ZeroVector, int32 InRootTransformSamplingRate = 30);
+	FAnimationAssetSampler(TObjectPtr<const UAnimationAsset> InAnimationAsset = nullptr, const FVector& InBlendParameters = FVector::ZeroVector, int32 InRootTransformSamplingRate = 30);
+	void Init(TObjectPtr<const UAnimationAsset> InAnimationAsset, const FVector& InBlendParameters = FVector::ZeroVector, int32 InRootTransformSamplingRate = 30);
 
 	bool IsInitialized() const;
 	float GetPlayLength() const;
@@ -50,7 +50,7 @@ struct POSESEARCH_API FAnimationAssetSampler
 	static float GetPlayLength(const UAnimationAsset* AnimAsset, const FVector& BlendParameters = FVector::ZeroVector);
 
 protected:
-	TWeakObjectPtr<const UAnimationAsset> AnimationAsset;
+	TObjectPtr<const UAnimationAsset> AnimationAsset;
 
 	// members used to sample blend spaces only!
 	FVector BlendParameters = FVector::ZeroVector;
diff --git a/Engine/Plugins/Experimental/Animation/PoseSearch/Source/Runtime/Public/PoseSearch/PoseSearchNormalizationSet.h b/Engine/Plugins/Experimental/Animation/PoseSearch/Source/Runtime/Public/PoseSearch/PoseSearchNormalizationSet.h
--- a/Engine/Plugins/Experimental/Animation/PoseSearch/Source/Runtime/Public/PoseSearch/PoseSearchNormalizationSet.h
+++ b/Engine/Plugins/Experimental/Animation/PoseSearch/Source/Runtime/Public/PoseSearch/PoseSearchNormalizationSet.h
@@ -17,5 +17,5 @@ public:
 	UPROPERTY(EditAnywhere, BlueprintReadOnly, Category = "NormalizationSet")
 	TArray<TObjectPtr<const UPoseSearchDatabase>> Databases;
 
-	void AddUniqueDatabases(TArray<TWeakObjectPtr<const UPoseSearchDatabase>>& UniqueDatabases) const;
+	void AddUniqueDatabases(TArray<TObjectPtr<const UPoseSearchDatabase>>& UniqueDatabases) const;
 };
\ No newline at end of file
diff --git a/Engine/Plugins/Experimental/Animation/PoseSearch/Source/Runtime/Public/PoseSearch/PoseSearchResult.h b/Engine/Plugins/Experimental/Animation/PoseSearch/Source/Runtime/Public/PoseSearch/PoseSearchResult.h
--- a/Engine/Plugins/Experimental/Animation/PoseSearch/Source/Runtime/Public/PoseSearch/PoseSearchResult.h
+++ b/Engine/Plugins/Experimental/Animation/PoseSearch/Source/Runtime/Public/PoseSearch/PoseSearchResult.h
@@ -27,7 +27,7 @@ public:
 	TConstArrayView<float> GetValues() const { return Values; }
 
 private:
-	TWeakObjectPtr<const UPoseSearchSchema> Schema = nullptr;
+	TObjectPtr<const UPoseSearchSchema> Schema;
 	TArray<float> Values;
 };
 	
@@ -43,7 +43,7 @@ struct FSearchResult
 	// lerp value to find AssetTime from PrevPoseIdx -> AssetTime -> NextPoseIdx, within range [-0.5, 0.5]
 	float LerpValue = 0.f;
 
-	TWeakObjectPtr<const UPoseSearchDatabase> Database;
+	TObjectPtr<const UPoseSearchDatabase> Database;
 
 	float AssetTime = 0.0f;
 
