motion matching
- UPoseSearchNormalizationSet as WITH_EDITORONLY_DATA
- channels GetLabel display improvements
- removed mirror mismatch bias
- minimized heap allocations

diff --git a/Engine/Plugins/Experimental/Animation/MotionTrajectory/Source/MotionTrajectory/Public/CharacterTrajectoryComponent.h b/Engine/Plugins/Experimental/Animation/MotionTrajectory/Source/MotionTrajectory/Public/CharacterTrajectoryComponent.h
--- a/Engine/Plugins/Experimental/Animation/MotionTrajectory/Source/MotionTrajectory/Public/CharacterTrajectoryComponent.h
+++ b/Engine/Plugins/Experimental/Animation/MotionTrajectory/Source/MotionTrajectory/Public/CharacterTrajectoryComponent.h
@@ -39,8 +39,8 @@ protected:
 
 	FRotator CalculateControllerRotationRate(float DeltaSeconds, bool bShouldRemainVertical);
 
-	// Trajectory stored in component space so it can be directly passed to Motion Matching.
-	UPROPERTY(BlueprintReadOnly, Category = "Trajectory")
+	// Trajectory stored in world space so it can be directly passed to Motion Matching.
+	UPROPERTY(Transient, BlueprintReadOnly, Category = "Trajectory")
 	FPoseSearchQueryTrajectory Trajectory;
 
 	// This should generally match the longest history required by a Motion Matching Database in the project.
diff --git a/Engine/Plugins/Experimental/Animation/PoseSearch/Source/Editor/Private/PoseSearchDebuggerDatabaseView.cpp b/Engine/Plugins/Experimental/Animation/PoseSearch/Source/Editor/Private/PoseSearchDebuggerDatabaseView.cpp
--- a/Engine/Plugins/Experimental/Animation/PoseSearch/Source/Editor/Private/PoseSearchDebuggerDatabaseView.cpp
+++ b/Engine/Plugins/Experimental/Animation/PoseSearch/Source/Editor/Private/PoseSearchDebuggerDatabaseView.cpp
@@ -186,8 +186,7 @@ static void AddUnfilteredDatabaseRow(const UPoseSearchDatabase* Database,
 		else
 		{
 			// @todo: perhaps reuse CompareFeatureVectors cost calculation
-			// @todo: initialize EPoseSearchBooleanRequest::Indifferent properly
-			Row->PoseCost = SearchIndex.ComparePoses(DbPoseIdx, EPoseSearchBooleanRequest::Indifferent, 0.f, Database->Schema->MirrorMismatchCostBias, PoseValues, QueryVector);
+			Row->PoseCost = SearchIndex.ComparePoses(DbPoseIdx, 0.f, PoseValues, QueryVector);
 		}
 
 		const FInstancedStruct& DatabaseAssetStruct = Database->GetAnimationAssetStruct(*SearchIndexAsset);
diff --git a/Engine/Plugins/Experimental/Animation/PoseSearch/Source/Runtime/Private/PoseSearchContext.cpp b/Engine/Plugins/Experimental/Animation/PoseSearch/Source/Runtime/Private/PoseSearchContext.cpp
--- a/Engine/Plugins/Experimental/Animation/PoseSearch/Source/Runtime/Private/PoseSearchContext.cpp
+++ b/Engine/Plugins/Experimental/Animation/PoseSearch/Source/Runtime/Private/PoseSearchContext.cpp
@@ -245,13 +245,12 @@ void FDebugDrawParams::DrawFeatureVector(int32 PoseIdx)
 
 //////////////////////////////////////////////////////////////////////////
 // FSearchContext
-FSearchContext::FSearchContext(const FPoseSearchQueryTrajectory* InTrajectory, const IPoseHistory* InHistory, float InDesiredPermutationTimeOffset, const FPoseIndicesHistory* InPoseIndicesHistory,
-	EPoseSearchBooleanRequest InQueryMirrorRequest, const FSearchResult& InCurrentResult, float InPoseJumpThresholdTime, bool bInForceInterrupt)
+FSearchContext::FSearchContext(const FPoseSearchQueryTrajectory* InTrajectory, const IPoseHistory* InHistory, float InDesiredPermutationTimeOffset,
+	const FPoseIndicesHistory* InPoseIndicesHistory, const FSearchResult& InCurrentResult, float InPoseJumpThresholdTime, bool bInForceInterrupt)
 : Trajectory(InTrajectory)
 , History(InHistory)
 , DesiredPermutationTimeOffset(InDesiredPermutationTimeOffset)
 , PoseIndicesHistory(InPoseIndicesHistory)
-, QueryMirrorRequest(InQueryMirrorRequest)
 , CurrentResult(InCurrentResult)
 , PoseJumpThresholdTime(InPoseJumpThresholdTime)
 , bForceInterrupt(bInForceInterrupt)
diff --git a/Engine/Plugins/Experimental/Animation/PoseSearch/Source/Runtime/Private/PoseSearchDatabase.cpp b/Engine/Plugins/Experimental/Animation/PoseSearch/Source/Runtime/Private/PoseSearchDatabase.cpp
--- a/Engine/Plugins/Experimental/Animation/PoseSearch/Source/Runtime/Private/PoseSearchDatabase.cpp
+++ b/Engine/Plugins/Experimental/Animation/PoseSearch/Source/Runtime/Private/PoseSearchDatabase.cpp
@@ -69,7 +69,7 @@ static void PopulateNonSelectableIdx(FNonSelectableIdx& NonSelectableIdx, FSearc
 
 #if UE_POSE_SEARCH_TRACE_ENABLED
 				const TArray<float> PoseValues = SearchIndex.GetPoseValuesSafe(PoseIdx);
-				const FPoseSearchCost PoseCost = SearchIndex.ComparePoses(PoseIdx, SearchContext.GetQueryMirrorRequest(), 0.f, Database->Schema->MirrorMismatchCostBias, PoseValues, QueryValues);
+				const FPoseSearchCost PoseCost = SearchIndex.ComparePoses(PoseIdx, 0.f, PoseValues, QueryValues);
 				SearchContext.BestCandidates.Add(PoseCost, PoseIdx, Database, EPoseCandidateFlags::DiscardedBy_PoseJumpThresholdTime);
 #endif
 			}
@@ -102,7 +102,7 @@ static void PopulateNonSelectableIdx(FNonSelectableIdx& NonSelectableIdx, FSearc
 				NonSelectableIdx.AddUnique(PoseIdx);
 
 #if UE_POSE_SEARCH_TRACE_ENABLED
-				const FPoseSearchCost PoseCost = SearchIndex.ComparePoses(PoseIdx, SearchContext.GetQueryMirrorRequest(), 0.f, Database->Schema->MirrorMismatchCostBias, SearchIndex.GetPoseValuesSafe(PoseIdx), QueryValues);
+				const FPoseSearchCost PoseCost = SearchIndex.ComparePoses(PoseIdx, 0.f, SearchIndex.GetPoseValuesSafe(PoseIdx), QueryValues);
 				SearchContext.BestCandidates.Add(PoseCost, PoseIdx, Database, EPoseCandidateFlags::DiscardedBy_PoseJumpThresholdTime);
 #endif
 			}
@@ -129,7 +129,7 @@ static void PopulateNonSelectableIdx(FNonSelectableIdx& NonSelectableIdx, FSearc
 				// if we're editing the database and removing assets it's possible that the PoseIndicesHistory contains invalid pose indexes
 				if (HistoricalPoseIndex.PoseIndex < SearchIndex.GetNumPoses())
 				{
-					const FPoseSearchCost PoseCost = SearchIndex.ComparePoses(HistoricalPoseIndex.PoseIndex, SearchContext.GetQueryMirrorRequest(), 0.f, Database->Schema->MirrorMismatchCostBias, SearchIndex.GetPoseValuesSafe(HistoricalPoseIndex.PoseIndex), QueryValues);
+					const FPoseSearchCost PoseCost = SearchIndex.ComparePoses(HistoricalPoseIndex.PoseIndex, 0.f, SearchIndex.GetPoseValuesSafe(HistoricalPoseIndex.PoseIndex), QueryValues);
 					SearchContext.BestCandidates.Add(PoseCost, HistoricalPoseIndex.PoseIndex, Database, EPoseCandidateFlags::DiscardedBy_PoseReselectHistory);
 				}
 #endif
@@ -182,12 +182,12 @@ struct FPoseFilters
 				}
 				else if (PoseFilter == &BlockTransitionPoseFilter)
 				{
-					const FPoseSearchCost PoseCost = SearchIndex.ComparePoses(PoseIdx, SearchContext.GetQueryMirrorRequest(), 0.f, Database->Schema->MirrorMismatchCostBias, PoseValues, QueryValues);
+					const FPoseSearchCost PoseCost = SearchIndex.ComparePoses(PoseIdx, 0.f, PoseValues, QueryValues);
 					SearchContext.BestCandidates.Add(PoseCost, PoseIdx, Database, EPoseCandidateFlags::DiscardedBy_BlockTransition);
 				}
 				else
 				{
-					const FPoseSearchCost PoseCost = SearchIndex.ComparePoses(PoseIdx, SearchContext.GetQueryMirrorRequest(), 0.f, Database->Schema->MirrorMismatchCostBias, PoseValues, QueryValues);
+					const FPoseSearchCost PoseCost = SearchIndex.ComparePoses(PoseIdx, 0.f, PoseValues, QueryValues);
 					SearchContext.BestCandidates.Add(PoseCost, PoseIdx, Database, EPoseCandidateFlags::DiscardedBy_PoseFilter);
 				}
 #endif
@@ -763,8 +763,7 @@ FPoseSearchCost UPoseSearchDatabase::SearchContinuingPose(UE::PoseSearch::FSearc
 		const TConstArrayView<float> PoseValues = SearchIndex.Values.IsEmpty() ? SearchIndex.GetReconstructedPoseValues(PoseIdx, ReconstructedPoseValuesBuffer) : SearchIndex.GetPoseValues(PoseIdx);
 
 		const int32 ContinuingPoseIdx = SearchContext.GetCurrentResult().PoseIdx;
-		ContinuingPoseCost = SearchIndex.ComparePoses(ContinuingPoseIdx, SearchContext.GetQueryMirrorRequest(),
-			ContinuingPoseCostBias, Schema->MirrorMismatchCostBias, PoseValues, SearchContext.GetOrBuildQuery(Schema).GetValues());
+		ContinuingPoseCost = SearchIndex.ComparePoses(ContinuingPoseIdx, ContinuingPoseCostBias, PoseValues, SearchContext.GetOrBuildQuery(Schema).GetValues());
 
 #if UE_POSE_SEARCH_TRACE_ENABLED
 		SearchContext.BestCandidates.Add(ContinuingPoseCost, ContinuingPoseIdx, this, EPoseCandidateFlags::Valid_ContinuingPose);
@@ -797,6 +796,7 @@ UE::PoseSearch::FSearchResult UPoseSearchDatabase::SearchPCAKDTree(UE::PoseSearc
 	RowMajorVectorMap ProjectedQueryValues((float*)FMemory_Alloca(ClampedNumberOfPrincipalComponents * sizeof(float)), 1, ClampedNumberOfPrincipalComponents);
 	TArrayView<float> ReconstructedPoseValuesBuffer((float*)FMemory_Alloca(NumDimensions * sizeof(float)), NumDimensions);
 
+#if DO_CHECK
 	// KDTree in PCA space search
 	if (PoseSearchMode == EPoseSearchMode::PCAKDTree_Validate)
 	{
@@ -830,6 +830,7 @@ UE::PoseSearch::FSearchResult UPoseSearchDatabase::SearchPCAKDTree(UE::PoseSearc
 			check(ResultIndex < ResultSet.Num());
 		}
 	}
+#endif //DO_CHECK
 
 	// since any PoseCost calculated here is at least SearchIndex.MinCostAddend,
 	// there's no point in performing the search if CurrentBestTotalCost is already better than that
@@ -872,7 +873,7 @@ UE::PoseSearch::FSearchResult UPoseSearchDatabase::SearchPCAKDTree(UE::PoseSearc
 #endif
 			))
 			{
-				const FPoseSearchCost PoseCost = SearchIndex.ComparePoses(PoseIdx, SearchContext.GetQueryMirrorRequest(), 0.f, Schema->MirrorMismatchCostBias, PoseValues, QueryValues);
+				const FPoseSearchCost PoseCost = SearchIndex.ComparePoses(PoseIdx, 0.f, PoseValues, QueryValues);
 				if (PoseCost < Result.PoseCost)
 				{
 					Result.PoseCost = PoseCost;
@@ -942,7 +943,7 @@ UE::PoseSearch::FSearchResult UPoseSearchDatabase::SearchBruteForce(UE::PoseSear
 #endif
 			))
 			{
-				const FPoseSearchCost PoseCost = SearchIndex.ComparePoses(PoseIdx, SearchContext.GetQueryMirrorRequest(), 0.f, Schema->MirrorMismatchCostBias, PoseValues, QueryValues);
+				const FPoseSearchCost PoseCost = SearchIndex.ComparePoses(PoseIdx, 0.f, PoseValues, QueryValues);
 				if (PoseCost < Result.PoseCost)
 				{
 					Result.PoseCost = PoseCost;
diff --git a/Engine/Plugins/Experimental/Animation/PoseSearch/Source/Runtime/Private/PoseSearchDerivedDataKey.cpp b/Engine/Plugins/Experimental/Animation/PoseSearch/Source/Runtime/Private/PoseSearchDerivedDataKey.cpp
--- a/Engine/Plugins/Experimental/Animation/PoseSearch/Source/Runtime/Private/PoseSearchDerivedDataKey.cpp
+++ b/Engine/Plugins/Experimental/Animation/PoseSearch/Source/Runtime/Private/PoseSearchDerivedDataKey.cpp
@@ -41,7 +41,7 @@ FKeyBuilder::FKeyBuilder(const UObject* Object, bool bUseDataVer)
 	if (bUseDataVer)
 	{
 		// used to invalidate the key without having to change POSESEARCHDB_DERIVEDDATA_VER all the times
-		int32 POSESEARCHDB_DERIVEDDATA_VER_SMALL = 153;
+		int32 POSESEARCHDB_DERIVEDDATA_VER_SMALL = 157;
 		FGuid VersionGuid = FDevSystemGuids::GetSystemGuid(FDevSystemGuids::Get().POSESEARCHDB_DERIVEDDATA_VER);
 
 		*this << VersionGuid;
diff --git a/Engine/Plugins/Experimental/Animation/PoseSearch/Source/Runtime/Private/PoseSearchFeatureChannel_Heading.cpp b/Engine/Plugins/Experimental/Animation/PoseSearch/Source/Runtime/Private/PoseSearchFeatureChannel_Heading.cpp
--- a/Engine/Plugins/Experimental/Animation/PoseSearch/Source/Runtime/Private/PoseSearchFeatureChannel_Heading.cpp
+++ b/Engine/Plugins/Experimental/Animation/PoseSearch/Source/Runtime/Private/PoseSearchFeatureChannel_Heading.cpp
@@ -139,7 +139,7 @@ FString UPoseSearchFeatureChannel_Heading::GetLabel() const
 		Label.Append(Schema->BoneReferences[SchemaBoneIdx].BoneName.ToString());
 	}
 
-	Label.Appendf(TEXT(" %.1f"), SampleTimeOffset);
+	Label.Appendf(TEXT(" %.2f"), SampleTimeOffset);
 	return Label.ToString();
 }
 #endif
\ No newline at end of file
diff --git a/Engine/Plugins/Experimental/Animation/PoseSearch/Source/Runtime/Private/PoseSearchFeatureChannel_Position.cpp b/Engine/Plugins/Experimental/Animation/PoseSearch/Source/Runtime/Private/PoseSearchFeatureChannel_Position.cpp
--- a/Engine/Plugins/Experimental/Animation/PoseSearch/Source/Runtime/Private/PoseSearchFeatureChannel_Position.cpp
+++ b/Engine/Plugins/Experimental/Animation/PoseSearch/Source/Runtime/Private/PoseSearchFeatureChannel_Position.cpp
@@ -159,7 +159,7 @@ FString UPoseSearchFeatureChannel_Position::GetLabel() const
 		Label.Append(Schema->BoneReferences[SchemaOriginBoneIdx].BoneName.ToString());
 	}
 
-	Label.Appendf(TEXT(" %.1f"), SampleTimeOffset);
+	Label.Appendf(TEXT(" %.2f"), SampleTimeOffset);
 	return Label.ToString();
 }
 #endif
\ No newline at end of file
diff --git a/Engine/Plugins/Experimental/Animation/PoseSearch/Source/Runtime/Private/PoseSearchFeatureChannel_Velocity.cpp b/Engine/Plugins/Experimental/Animation/PoseSearch/Source/Runtime/Private/PoseSearchFeatureChannel_Velocity.cpp
--- a/Engine/Plugins/Experimental/Animation/PoseSearch/Source/Runtime/Private/PoseSearchFeatureChannel_Velocity.cpp
+++ b/Engine/Plugins/Experimental/Animation/PoseSearch/Source/Runtime/Private/PoseSearchFeatureChannel_Velocity.cpp
@@ -128,7 +128,7 @@ FString UPoseSearchFeatureChannel_Velocity::GetLabel() const
 		Label.Append(Schema->BoneReferences[SchemaBoneIdx].BoneName.ToString());
 	}
 
-	Label.Appendf(TEXT(" %.1f"), SampleTimeOffset);
+	Label.Appendf(TEXT(" %.2f"), SampleTimeOffset);
 	return Label.ToString();
 }
 #endif
\ No newline at end of file
diff --git a/Engine/Plugins/Experimental/Animation/PoseSearch/Source/Runtime/Private/PoseSearchIndex.cpp b/Engine/Plugins/Experimental/Animation/PoseSearch/Source/Runtime/Private/PoseSearchIndex.cpp
--- a/Engine/Plugins/Experimental/Animation/PoseSearch/Source/Runtime/Private/PoseSearchIndex.cpp
+++ b/Engine/Plugins/Experimental/Animation/PoseSearch/Source/Runtime/Private/PoseSearchIndex.cpp
@@ -6,7 +6,7 @@
 namespace UE::PoseSearch
 {
 
-float CompareFeatureVectors(TConstArrayView<float> A, TConstArrayView<float> B, TConstArrayView<float> WeightsSqrt)
+static FORCEINLINE float CompareFeatureVectors(TConstArrayView<float> A, TConstArrayView<float> B, TConstArrayView<float> WeightsSqrt)
 {
 	check(A.Num() == B.Num() && A.Num() == WeightsSqrt.Num());
 
@@ -191,28 +191,14 @@ TArray<float> FPoseSearchIndex::GetPoseValuesSafe(int32 PoseIdx) const
 	return PoseValues;
 }
 
-FPoseSearchCost FPoseSearchIndex::ComparePoses(int32 PoseIdx, EPoseSearchBooleanRequest QueryMirrorRequest, float ContinuingPoseCostBias, float MirrorMismatchCostBias, TConstArrayView<float> PoseValues, TConstArrayView<float> QueryValues) const
+FPoseSearchCost FPoseSearchIndex::ComparePoses(int32 PoseIdx, float ContinuingPoseCostBias, TConstArrayView<float> PoseValues, TConstArrayView<float> QueryValues) const
 {
 	// base dissimilarity cost representing how the associated PoseIdx differ, in a weighted way, from the query pose (QueryValues)
 	const float DissimilarityCost = UE::PoseSearch::CompareFeatureVectors(PoseValues, QueryValues, WeightsSqrt);
 
-	// cost addend associated to a mismatch in mirror state between query and analyzed PoseIdx
-	float MirrorMismatchAddend = 0.f;
-	if (QueryMirrorRequest != EPoseSearchBooleanRequest::Indifferent)
-	{
-		const FPoseSearchIndexAsset& IndexAsset = GetAssetForPose(PoseIdx);
-		const bool bMirroringMismatch =
-			(IndexAsset.bMirrored && QueryMirrorRequest == EPoseSearchBooleanRequest::FalseValue) ||
-			(!IndexAsset.bMirrored && QueryMirrorRequest == EPoseSearchBooleanRequest::TrueValue);
-		if (bMirroringMismatch)
-		{
-			MirrorMismatchAddend = MirrorMismatchCostBias;
-		}
-	}
-
 	// cost addend associated to Schema->BaseCostBias or overriden by UAnimNotifyState_PoseSearchModifyCost
 	const float NotifyAddend = PoseMetadata[PoseIdx].GetCostAddend();
-	return FPoseSearchCost(DissimilarityCost, NotifyAddend, MirrorMismatchAddend, ContinuingPoseCostBias);
+	return FPoseSearchCost(DissimilarityCost, NotifyAddend, ContinuingPoseCostBias);
 }
 
 FArchive& operator<<(FArchive& Ar, FPoseSearchIndex& Index)
diff --git a/Engine/Plugins/Experimental/Animation/PoseSearch/Source/Runtime/Private/PoseSearchLibrary.cpp b/Engine/Plugins/Experimental/Animation/PoseSearch/Source/Runtime/Private/PoseSearchLibrary.cpp
--- a/Engine/Plugins/Experimental/Animation/PoseSearch/Source/Runtime/Private/PoseSearchLibrary.cpp
+++ b/Engine/Plugins/Experimental/Animation/PoseSearch/Source/Runtime/Private/PoseSearchLibrary.cpp
@@ -307,15 +307,8 @@ void UPoseSearchLibrary::UpdateMotionMatchingState(
 		History = &PoseHistoryProvider->GetPoseHistory();
 	}
 	
-	EPoseSearchBooleanRequest QueryMirrorRequest = EPoseSearchBooleanRequest::Indifferent;
-	if (const FPoseSearchIndexAsset* CurrentIndexAsset = InOutMotionMatchingState.CurrentSearchResult.GetSearchIndexAsset())
-	{
-		QueryMirrorRequest = CurrentIndexAsset->bMirrored ? EPoseSearchBooleanRequest::TrueValue : EPoseSearchBooleanRequest::FalseValue;
-	}
-
 	const FPoseSearchQueryTrajectory TrajectoryRootSpace = ProcessTrajectory(Trajectory, Context.AnimInstanceProxy->GetComponentTransform(), InOutMotionMatchingState.GetRootBoneDeltaYaw(), RootBoneDeltaYawBlendTime, TrajectorySpeedMultiplier);
-	FSearchContext SearchContext(&TrajectoryRootSpace, History, 0.f, &InOutMotionMatchingState.PoseIndicesHistory, QueryMirrorRequest,
-		InOutMotionMatchingState.CurrentSearchResult, PoseJumpThresholdTime, bForceInterrupt);
+	FSearchContext SearchContext(&TrajectoryRootSpace, History, 0.f, &InOutMotionMatchingState.PoseIndicesHistory, InOutMotionMatchingState.CurrentSearchResult, PoseJumpThresholdTime, bForceInterrupt);
 
 	const bool bMustSearch = !InOutMotionMatchingState.CurrentSearchResult.IsValid();
 
diff --git a/Engine/Plugins/Experimental/Animation/PoseSearch/Source/Runtime/Private/PoseSearchNormalizationSet.cpp b/Engine/Plugins/Experimental/Animation/PoseSearch/Source/Runtime/Private/PoseSearchNormalizationSet.cpp
--- a/Engine/Plugins/Experimental/Animation/PoseSearch/Source/Runtime/Private/PoseSearchNormalizationSet.cpp
+++ b/Engine/Plugins/Experimental/Animation/PoseSearch/Source/Runtime/Private/PoseSearchNormalizationSet.cpp
@@ -3,6 +3,8 @@
 #include "PoseSearch/PoseSearchNormalizationSet.h"
 #include "PoseSearch/PoseSearchDatabase.h"
 
+#if WITH_EDITORONLY_DATA
+
 void UPoseSearchNormalizationSet::AddUniqueDatabases(TArray<TObjectPtr<const UPoseSearchDatabase>>& UniqueDatabases) const
 {
 	for (const TObjectPtr<const UPoseSearchDatabase>& Database : Databases)
@@ -12,4 +14,6 @@ void UPoseSearchNormalizationSet::AddUniqueDatabases(TArray<TObjectPtr<const UPo
 			UniqueDatabases.AddUnique(Database);
 		}
 	}
-}
\ No newline at end of file
+}
+
+#endif // WITH_EDITORONLY_DATA
\ No newline at end of file
diff --git a/Engine/Plugins/Experimental/Animation/PoseSearch/Source/Runtime/Public/PoseSearch/PoseSearchContext.h b/Engine/Plugins/Experimental/Animation/PoseSearch/Source/Runtime/Public/PoseSearch/PoseSearchContext.h
--- a/Engine/Plugins/Experimental/Animation/PoseSearch/Source/Runtime/Public/PoseSearch/PoseSearchContext.h
+++ b/Engine/Plugins/Experimental/Animation/PoseSearch/Source/Runtime/Public/PoseSearch/PoseSearchContext.h
@@ -139,8 +139,7 @@ private:
 
 struct POSESEARCH_API FSearchContext
 {
-	FSearchContext(const FPoseSearchQueryTrajectory* InTrajectory, const IPoseHistory* InHistory, float InDesiredPermutationTimeOffset,
-		const FPoseIndicesHistory* InPoseIndicesHistory = nullptr, EPoseSearchBooleanRequest InQueryMirrorRequest = EPoseSearchBooleanRequest::Indifferent,
+	FSearchContext(const FPoseSearchQueryTrajectory* InTrajectory, const IPoseHistory* InHistory, float InDesiredPermutationTimeOffset, const FPoseIndicesHistory* InPoseIndicesHistory = nullptr,
 		const FSearchResult& InCurrentResult = FSearchResult(), float InPoseJumpThresholdTime = 0.f, bool bInForceInterrupt = false);
 
 	FQuat GetSampleRotation(float SampleTimeOffset, const UPoseSearchSchema* Schema, int8 SchemaSampleBoneIdx = RootSchemaBoneIdx, int8 SchemaOriginBoneIdx = RootSchemaBoneIdx, bool bUseHistoryRoot = false, EPermutationTimeType PermutationTimeType = EPermutationTimeType::UseSampleTime);
@@ -164,7 +163,6 @@ struct POSESEARCH_API FSearchContext
 
 	const FSearchResult& GetCurrentResult() const { return CurrentResult; }
 	float GetPoseJumpThresholdTime() const { return PoseJumpThresholdTime; }
-	EPoseSearchBooleanRequest GetQueryMirrorRequest() const { return QueryMirrorRequest; }
 	const FPoseIndicesHistory* GetPoseIndicesHistory() const { return PoseIndicesHistory; }
 	bool IsHistoryValid() const { return History != nullptr; }
 	float GetDesiredPermutationTimeOffset() const { return DesiredPermutationTimeOffset; }
@@ -182,7 +180,6 @@ private:
 	const IPoseHistory* History = nullptr;
 	float DesiredPermutationTimeOffset = 0.f;
 	const FPoseIndicesHistory* PoseIndicesHistory = nullptr;
-	EPoseSearchBooleanRequest QueryMirrorRequest = EPoseSearchBooleanRequest::Indifferent;
 	FSearchResult CurrentResult;
 	float PoseJumpThresholdTime = 0.f;
 	bool bForceInterrupt = false;
@@ -195,7 +192,8 @@ private:
 
 	// transforms cached in component space
 	FCachedTransforms<FTransform> CachedTransforms;
-	TArray<FFeatureVectorBuilder, TInlineAllocator<8>> CachedQueries;
+	// @todo: consider using FMemMark
+	TArray<FFeatureVectorBuilder, TInlineAllocator<PreallocatedCachedQueriesNum>> CachedQueries;
 
 	float CurrentBestTotalCost = MAX_flt;
 	
diff --git a/Engine/Plugins/Experimental/Animation/PoseSearch/Source/Runtime/Public/PoseSearch/PoseSearchCost.h b/Engine/Plugins/Experimental/Animation/PoseSearch/Source/Runtime/Public/PoseSearch/PoseSearchCost.h
--- a/Engine/Plugins/Experimental/Animation/PoseSearch/Source/Runtime/Public/PoseSearch/PoseSearchCost.h
+++ b/Engine/Plugins/Experimental/Animation/PoseSearch/Source/Runtime/Public/PoseSearch/PoseSearchCost.h
@@ -10,12 +10,11 @@ struct POSESEARCH_API FPoseSearchCost
 	GENERATED_BODY()
 public:
 	FPoseSearchCost() = default;
-	FPoseSearchCost(float InDissimilarityCost, float InNotifyCostAddend, float InMirrorMismatchCostAddend, float InContinuingPoseCostAddend)
-	: TotalCost(InDissimilarityCost + InNotifyCostAddend + InMirrorMismatchCostAddend + InContinuingPoseCostAddend)
+	FPoseSearchCost(float InDissimilarityCost, float InNotifyCostAddend, float InContinuingPoseCostAddend)
+	: TotalCost(InDissimilarityCost + InNotifyCostAddend + InContinuingPoseCostAddend)
 	{
 #if WITH_EDITORONLY_DATA
 		NotifyCostAddend = InNotifyCostAddend;
-		MirrorMismatchCostAddend = InMirrorMismatchCostAddend;
 		ContinuingPoseCostAddend = InContinuingPoseCostAddend;
 #endif // WITH_EDITORONLY_DATA
 	}
@@ -31,16 +30,12 @@ protected:
 #if WITH_EDITORONLY_DATA
 public:
 
-	float GetCostAddend() const { return NotifyCostAddend + MirrorMismatchCostAddend + ContinuingPoseCostAddend; }
+	float GetCostAddend() const { return NotifyCostAddend + ContinuingPoseCostAddend; }
 
 	// Contribution from ModifyCost anim notify
 	UPROPERTY()
 	float NotifyCostAddend = 0.f;
 
-	// Contribution from mirroring cost
-	UPROPERTY()
-	float MirrorMismatchCostAddend = 0.f;
-
 	UPROPERTY()
 	float ContinuingPoseCostAddend = 0.f;
 
diff --git a/Engine/Plugins/Experimental/Animation/PoseSearch/Source/Runtime/Public/PoseSearch/PoseSearchDatabase.h b/Engine/Plugins/Experimental/Animation/PoseSearch/Source/Runtime/Public/PoseSearch/PoseSearchDatabase.h
--- a/Engine/Plugins/Experimental/Animation/PoseSearch/Source/Runtime/Public/PoseSearch/PoseSearchDatabase.h
+++ b/Engine/Plugins/Experimental/Animation/PoseSearch/Source/Runtime/Public/PoseSearch/PoseSearchDatabase.h
@@ -302,10 +302,12 @@ public:
 	UPROPERTY(EditAnywhere, Category = "Performance")
 	bool bSkipSearchIfPossible = true;
 
+#if WITH_EDITORONLY_DATA
 	// This optional asset defines a list of databases you want to normalize together. Without it, it would be difficult to compare costs from separately normalized databases containing different types of animation,
 	// like only idles versus only runs animations, given that the range of movement would be dramatically different.
 	UPROPERTY(EditAnywhere, BlueprintReadOnly, Category = "Database")
 	TObjectPtr<const UPoseSearchNormalizationSet> NormalizationSet;
+#endif // WITH_EDITORONLY_DATA
 
 private:
 	// Do not use it directly. Use GetSearchIndex / SetSearchIndex interact with it and validate that is ok to do so.
diff --git a/Engine/Plugins/Experimental/Animation/PoseSearch/Source/Runtime/Public/PoseSearch/PoseSearchDefines.h b/Engine/Plugins/Experimental/Animation/PoseSearch/Source/Runtime/Public/PoseSearch/PoseSearchDefines.h
--- a/Engine/Plugins/Experimental/Animation/PoseSearch/Source/Runtime/Public/PoseSearch/PoseSearchDefines.h
+++ b/Engine/Plugins/Experimental/Animation/PoseSearch/Source/Runtime/Public/PoseSearch/PoseSearchDefines.h
@@ -24,4 +24,9 @@ DECLARE_LOG_CATEGORY_EXTERN(LogPoseSearch, Log, All);
 static constexpr int8 RootSchemaBoneIdx = 0;
 static constexpr FBoneIndexType RootBoneIndexType = 0;
 
+static constexpr int32 PreallocatedCachedQuerySize = 64;
+static constexpr int32 PreallocatedCachedQueriesNum = 8;
+
+
+
 
diff --git a/Engine/Plugins/Experimental/Animation/PoseSearch/Source/Runtime/Public/PoseSearch/PoseSearchIndex.h b/Engine/Plugins/Experimental/Animation/PoseSearch/Source/Runtime/Public/PoseSearch/PoseSearchIndex.h
--- a/Engine/Plugins/Experimental/Animation/PoseSearch/Source/Runtime/Public/PoseSearch/PoseSearchIndex.h
+++ b/Engine/Plugins/Experimental/Animation/PoseSearch/Source/Runtime/Public/PoseSearch/PoseSearchIndex.h
@@ -9,22 +9,10 @@
 namespace UE::PoseSearch
 {
 
-float CompareFeatureVectors(TConstArrayView<float> A, TConstArrayView<float> B, TConstArrayView<float> WeightsSqrt);
 POSESEARCH_API void CompareFeatureVectors(TConstArrayView<float> A, TConstArrayView<float> B, TConstArrayView<float> WeightsSqrt, TArrayView<float> Result);
 
 } // namespace UE::PoseSearch
 
-UENUM()
-enum class EPoseSearchBooleanRequest : uint8
-{
-	FalseValue,
-	TrueValue,
-	Indifferent, // if this is used, there will be no cost difference between true and false results
-
-	Num UMETA(Hidden),
-	Invalid = Num UMETA(Hidden)
-};
-
 /**
  * This is kept for each pose in the search index along side the feature vector values and is used to influence the search.
  */
@@ -282,7 +270,7 @@ struct POSESEARCH_API FPoseSearchIndex : public FPoseSearchIndexBase
 	TConstArrayView<float> GetPoseValues(int32 PoseIdx) const;
 	TConstArrayView<float> GetReconstructedPoseValues(int32 PoseIdx, TArrayView<float> BufferUsedForReconstruction) const;
 	TArray<float> GetPoseValuesSafe(int32 PoseIdx) const;
-	FPoseSearchCost ComparePoses(int32 PoseIdx, EPoseSearchBooleanRequest QueryMirrorRequest, float ContinuingPoseCostBias, float MirrorMismatchCostBias, TConstArrayView<float> PoseValues, TConstArrayView<float> QueryValues) const;
+	FPoseSearchCost ComparePoses(int32 PoseIdx, float ContinuingPoseCostBias, TConstArrayView<float> PoseValues, TConstArrayView<float> QueryValues) const;
 
 	friend FArchive& operator<<(FArchive& Ar, FPoseSearchIndex& Index);
 };
diff --git a/Engine/Plugins/Experimental/Animation/PoseSearch/Source/Runtime/Public/PoseSearch/PoseSearchResult.h b/Engine/Plugins/Experimental/Animation/PoseSearch/Source/Runtime/Public/PoseSearch/PoseSearchResult.h
--- a/Engine/Plugins/Experimental/Animation/PoseSearch/Source/Runtime/Public/PoseSearch/PoseSearchResult.h
+++ b/Engine/Plugins/Experimental/Animation/PoseSearch/Source/Runtime/Public/PoseSearch/PoseSearchResult.h
@@ -3,6 +3,7 @@
 #pragma once
 
 #include "PoseSearch/PoseSearchCost.h"
+#include "PoseSearch/PoseSearchDefines.h"
 
 struct FPoseSearchIndexAsset;
 class UPoseSearchDatabase;
@@ -23,12 +24,13 @@ public:
 
 	const UPoseSearchSchema* GetSchema() const { return Schema.Get(); }
 
-	TArray<float>& EditValues() { return Values; }
+	TArrayView<float> EditValues() { return Values; }
 	TConstArrayView<float> GetValues() const { return Values; }
 
 private:
 	TObjectPtr<const UPoseSearchSchema> Schema;
-	TArray<float> Values;
+	// @todo: consider using FMemMark
+	TArray<float, TInlineAllocator<PreallocatedCachedQuerySize>> Values;
 };
 	
 struct FSearchResult
diff --git a/Engine/Plugins/Experimental/Animation/PoseSearch/Source/Runtime/Public/PoseSearch/PoseSearchSchema.h b/Engine/Plugins/Experimental/Animation/PoseSearch/Source/Runtime/Public/PoseSearch/PoseSearchSchema.h
--- a/Engine/Plugins/Experimental/Animation/PoseSearch/Source/Runtime/Public/PoseSearch/PoseSearchSchema.h
+++ b/Engine/Plugins/Experimental/Animation/PoseSearch/Source/Runtime/Public/PoseSearch/PoseSearchSchema.h
@@ -82,11 +82,6 @@ public:
 	UPROPERTY(EditAnywhere, Category = "Bias")
 	float BaseCostBias = 0.f;
 
-	// If there’s a mirroring mismatch between the currently playing asset and a search candidate, this cost will be added or removed to the candidate,
-	// making it more or less likely to be selected. Negative values make it more likely to be picked, or stayed in, Positive values make it less likely to be picked or stay in.
-	UPROPERTY(EditAnywhere, Category = "Bias")
-	float MirrorMismatchCostBias = 0.f;
-
 	// Cost added to all looping animation assets in a database that uses this config. This allows users to make it more or less likely to pick the looping animation segments.
 	// Negative values make it more likely to be picked, or stayed in, Positive values make it less likely to be picked or stay in.
 	UPROPERTY(EditAnywhere, Category = "Bias")
