motion matching - updating posesearch profiling markers

diff --git a/Engine/Plugins/Experimental/Animation/PoseSearch/Source/Runtime/Private/AnimNode_BlendStack.cpp b/Engine/Plugins/Experimental/Animation/PoseSearch/Source/Runtime/Private/AnimNode_BlendStack.cpp
--- a/Engine/Plugins/Experimental/Animation/PoseSearch/Source/Runtime/Private/AnimNode_BlendStack.cpp
+++ b/Engine/Plugins/Experimental/Animation/PoseSearch/Source/Runtime/Private/AnimNode_BlendStack.cpp
@@ -288,9 +288,10 @@ bool FPoseSearchAnimPlayer::GetBlendInWeights(TArray<float>& Weights) const
 // FAnimNode_BlendStack_Standalone
 void FAnimNode_BlendStack_Standalone::Evaluate_AnyThread(FPoseContext& Output)
 {
-	Super::Evaluate_AnyThread(Output);
-
 	DECLARE_SCOPE_HIERARCHICAL_COUNTER_ANIMNODE(Evaluate_AnyThread);
+	QUICK_SCOPE_CYCLE_COUNTER(STAT_BlendStack_Evaluate_AnyThread);
+
+	Super::Evaluate_AnyThread(Output);
 
 	const int32 BlendStackSize = AnimPlayers.Num();
 	if (BlendStackSize <= 0)
diff --git a/Engine/Plugins/Experimental/Animation/PoseSearch/Source/Runtime/Private/AnimNode_MotionMatching.cpp b/Engine/Plugins/Experimental/Animation/PoseSearch/Source/Runtime/Private/AnimNode_MotionMatching.cpp
--- a/Engine/Plugins/Experimental/Animation/PoseSearch/Source/Runtime/Private/AnimNode_MotionMatching.cpp
+++ b/Engine/Plugins/Experimental/Animation/PoseSearch/Source/Runtime/Private/AnimNode_MotionMatching.cpp
@@ -83,6 +83,8 @@ void FAnimNode_MotionMatching::UpdateAssetPlayer(const FAnimationUpdateContext&
 {
 	DECLARE_SCOPE_HIERARCHICAL_COUNTER_ANIMNODE(UpdateAssetPlayer);
 
+	QUICK_SCOPE_CYCLE_COUNTER(STAT_PoseSearch_UpdateAssetPlayer);
+
 	using namespace UE::PoseSearch;
 
 	GetEvaluateGraphExposedInputs().Execute(Context);
diff --git a/Engine/Plugins/Experimental/Animation/PoseSearch/Source/Runtime/Private/KDTree.cpp b/Engine/Plugins/Experimental/Animation/PoseSearch/Source/Runtime/Private/KDTree.cpp
--- a/Engine/Plugins/Experimental/Animation/PoseSearch/Source/Runtime/Private/KDTree.cpp
+++ b/Engine/Plugins/Experimental/Animation/PoseSearch/Source/Runtime/Private/KDTree.cpp
@@ -1,8 +1,11 @@
 // Copyright Epic Games, Inc. All Rights Reserved.
 
 #include "PoseSearch/KDTree.h"
+#include "Stats/Stats.h"
 
-#define UE_POSE_SEARCH_USE_NANOFLANN 1
+#ifndef UE_POSE_SEARCH_USE_NANOFLANN
+	#define UE_POSE_SEARCH_USE_NANOFLANN 1
+#endif
 
 // @third party code - BEGIN nanoflann
 #if UE_POSE_SEARCH_USE_NANOFLANN
@@ -148,6 +151,8 @@ void FKDTree::Construct(int32 Count, int32 Dim, const float* Data, int32 MaxLeaf
 bool FKDTree::FindNeighbors(KNNResultSet& Result, const float* Query) const
 {
 #if UE_POSE_SEARCH_USE_NANOFLANN
+	QUICK_SCOPE_CYCLE_COUNTER(STAT_FKDTree_FindNeighbors);
+
 	check(Query && Impl->root_node);
 	const nanoflann::SearchParams SearchParams(
 		32,			// Ignored parameter (Kept for compatibility with the FLANN interface).
@@ -155,7 +160,7 @@ bool FKDTree::FindNeighbors(KNNResultSet& Result, const float* Query) const
 		false);		// only for radius search, require neighbours sorted by
 	return Impl->findNeighbors(Result, Query, SearchParams);
 #else
-	check(false); // unimplemented
+	checkNoEntry(); // unimplemented
 	return false;
 #endif
 }
diff --git a/Engine/Plugins/Experimental/Animation/PoseSearch/Source/Runtime/Private/PoseSearchContext.cpp b/Engine/Plugins/Experimental/Animation/PoseSearch/Source/Runtime/Private/PoseSearchContext.cpp
--- a/Engine/Plugins/Experimental/Animation/PoseSearch/Source/Runtime/Private/PoseSearchContext.cpp
+++ b/Engine/Plugins/Experimental/Animation/PoseSearch/Source/Runtime/Private/PoseSearchContext.cpp
@@ -475,6 +475,8 @@ const FFeatureVectorBuilder* FSearchContext::GetCachedQuery(const UPoseSearchSch
 
 const FFeatureVectorBuilder& FSearchContext::GetOrBuildQuery(const UPoseSearchSchema* Schema)
 {
+	QUICK_SCOPE_CYCLE_COUNTER(STAT_PoseSearch_GetOrBuildQuery);
+
 	check(Schema && Schema->IsValid());
 	const FFeatureVectorBuilder* CachedFeatureVectorBuilder = GetCachedQuery(Schema);
 	if (CachedFeatureVectorBuilder)
diff --git a/Engine/Plugins/Experimental/Animation/PoseSearch/Source/Runtime/Private/PoseSearchDatabase.cpp b/Engine/Plugins/Experimental/Animation/PoseSearch/Source/Runtime/Private/PoseSearchDatabase.cpp
--- a/Engine/Plugins/Experimental/Animation/PoseSearch/Source/Runtime/Private/PoseSearchDatabase.cpp
+++ b/Engine/Plugins/Experimental/Animation/PoseSearch/Source/Runtime/Private/PoseSearchDatabase.cpp
@@ -18,10 +18,10 @@
 #include "UObject/ObjectSaveContext.h"
 
 DECLARE_STATS_GROUP(TEXT("PoseSearch"), STATGROUP_PoseSearch, STATCAT_Advanced);
-DECLARE_CYCLE_STAT_EXTERN(TEXT("Search Brute Force"), STAT_PoseSearchBruteForce, STATGROUP_PoseSearch, );
-DECLARE_CYCLE_STAT_EXTERN(TEXT("Search PCA/KNN"), STAT_PoseSearchPCAKNN, STATGROUP_PoseSearch, );
-DEFINE_STAT(STAT_PoseSearchBruteForce);
-DEFINE_STAT(STAT_PoseSearchPCAKNN);
+DECLARE_CYCLE_STAT_EXTERN(TEXT("Search Brute Force"), STAT_PoseSearch_BruteForce, STATGROUP_PoseSearch, );
+DECLARE_CYCLE_STAT_EXTERN(TEXT("Search PCA/KNN"), STAT_PoseSearch_PCAKNN, STATGROUP_PoseSearch, );
+DEFINE_STAT(STAT_PoseSearch_BruteForce);
+DEFINE_STAT(STAT_PoseSearch_PCAKNN);
 
 namespace UE::PoseSearch
 {
@@ -775,8 +775,7 @@ FPoseSearchCost UPoseSearchDatabase::SearchContinuingPose(UE::PoseSearch::FSearc
 
 UE::PoseSearch::FSearchResult UPoseSearchDatabase::SearchPCAKDTree(UE::PoseSearch::FSearchContext& SearchContext) const
 {
-	QUICK_SCOPE_CYCLE_COUNTER(STAT_PoseSearch_PCA_KNN);
-	SCOPE_CYCLE_COUNTER(STAT_PoseSearchPCAKNN);
+	SCOPE_CYCLE_COUNTER(STAT_PoseSearch_PCAKNN);
 
 	using namespace UE::PoseSearch;
 
@@ -908,8 +907,7 @@ UE::PoseSearch::FSearchResult UPoseSearchDatabase::SearchPCAKDTree(UE::PoseSearc
 
 UE::PoseSearch::FSearchResult UPoseSearchDatabase::SearchBruteForce(UE::PoseSearch::FSearchContext& SearchContext) const
 {
-	QUICK_SCOPE_CYCLE_COUNTER(STAT_PoseSearch_Brute_Force);
-	SCOPE_CYCLE_COUNTER(STAT_PoseSearchBruteForce);
+	SCOPE_CYCLE_COUNTER(STAT_PoseSearch_BruteForce);
 	
 	using namespace UE::PoseSearch;
 	
diff --git a/Engine/Plugins/Experimental/Animation/PoseSearch/Source/Runtime/Private/PoseSearchLibrary.cpp b/Engine/Plugins/Experimental/Animation/PoseSearch/Source/Runtime/Private/PoseSearchLibrary.cpp
--- a/Engine/Plugins/Experimental/Animation/PoseSearch/Source/Runtime/Private/PoseSearchLibrary.cpp
+++ b/Engine/Plugins/Experimental/Animation/PoseSearch/Source/Runtime/Private/PoseSearchLibrary.cpp
@@ -425,6 +425,8 @@ void UPoseSearchLibrary::UpdateMotionMatchingState(
 // transforms Trajectory from world space to root bone space, and scale it by TrajectorySpeedMultiplier
 FPoseSearchQueryTrajectory UPoseSearchLibrary::ProcessTrajectory(const FPoseSearchQueryTrajectory& Trajectory, const FTransform& ComponentWorldTransform, float RootBoneDeltaYaw, float RootBoneDeltaYawBlendTime, float TrajectorySpeedMultiplier)
 {
+	QUICK_SCOPE_CYCLE_COUNTER(STAT_PoseSearch_ProcessTrajectory);
+
 	const float TrajectorySpeedMultiplierInv = FMath::IsNearlyZero(TrajectorySpeedMultiplier) ? 1.f : 1.f / TrajectorySpeedMultiplier;
 
 	FPoseSearchQueryTrajectory TrajectoryRootSpace = Trajectory;
