motion matching - YawFromAnimation improvements

diff --git a/Engine/Plugins/Experimental/Animation/PoseSearch/Source/Runtime/Private/AnimNode_MotionMatching.cpp b/Engine/Plugins/Experimental/Animation/PoseSearch/Source/Runtime/Private/AnimNode_MotionMatching.cpp
--- a/Engine/Plugins/Experimental/Animation/PoseSearch/Source/Runtime/Private/AnimNode_MotionMatching.cpp
+++ b/Engine/Plugins/Experimental/Animation/PoseSearch/Source/Runtime/Private/AnimNode_MotionMatching.cpp
@@ -98,7 +98,7 @@ void FAnimNode_MotionMatching::UpdateAssetPlayer(const FAnimationUpdateContext&
 	// If we just became relevant and haven't been initialized yet, then reset motion matching state, otherwise update the asset time using the player node.
 	if (bNeedsReset)
 	{
-		MotionMatchingState.Reset();
+		MotionMatchingState.Reset(Context.AnimInstanceProxy->GetComponentTransform());
 	}
 	else
 	{
@@ -112,13 +112,13 @@ void FAnimNode_MotionMatching::UpdateAssetPlayer(const FAnimationUpdateContext&
 			{
 				// MotionMatchingState is out of sync with CurrentResultDatabase: we need to reset the MM state. This could happen if PIE is paused, and we edit the database,
 				// so FAnimNode_MotionMatching::UpdateAssetPlayer is never called and FAsyncPoseSearchDatabasesManagement::RequestAsyncBuildIndex never returns false here
-				MotionMatchingState.Reset();
+				MotionMatchingState.Reset(Context.AnimInstanceProxy->GetComponentTransform());
 			}
 		}
 		else
 		{
 			// we're still indexing MotionMatchingState.CurrentSearchResult.Database, so we Reset the MotionMatchingState
-			MotionMatchingState.Reset();
+			MotionMatchingState.Reset(Context.AnimInstanceProxy->GetComponentTransform());
 		}
 #endif // WITH_EDITOR
 
@@ -163,7 +163,7 @@ void FAnimNode_MotionMatching::UpdateAssetPlayer(const FAnimationUpdateContext&
 		SearchThrottleTime,
 		PlayRate,
 		MotionMatchingState,
-		YawFromAnimation,
+		YawFromAnimationBlendRate,
 		YawFromAnimationTrajectoryBlendTime,
 		bForceInterruptNextUpdate,
 		bShouldSearch
@@ -183,7 +183,7 @@ void FAnimNode_MotionMatching::UpdateAssetPlayer(const FAnimationUpdateContext&
 			if (const FPoseSearchDatabaseAnimationAssetBase* DatabaseAsset = CurrentResultDatabase->GetAnimationAssetBase(*SearchIndexAsset))
 			{
 				// root bone blending needs to be immediate if MM node controls the offset between mesh component and root bone
-				const float RootBoneBlendTime = FMath::IsNearlyZero(YawFromAnimation) ? BlendTime : 0.f;
+				const float RootBoneBlendTime = YawFromAnimationBlendRate < 0.f ? BlendTime : 0.f;
 				BlendStackNode.BlendTo(DatabaseAsset->GetAnimationAsset(), MotionMatchingState.CurrentSearchResult.AssetTime,
 					DatabaseAsset->IsLooping(), SearchIndexAsset->bMirrored, CurrentResultDatabase->Schema->MirrorDataTable.Get(),
 					MaxActiveBlends, BlendTime, RootBoneBlendTime, BlendProfile, BlendOption, SearchIndexAsset->BlendParameters, MotionMatchingState.WantedPlayRate);
diff --git a/Engine/Plugins/Experimental/Animation/PoseSearch/Source/Runtime/Private/PoseSearchLibrary.cpp b/Engine/Plugins/Experimental/Animation/PoseSearch/Source/Runtime/Private/PoseSearchLibrary.cpp
--- a/Engine/Plugins/Experimental/Animation/PoseSearch/Source/Runtime/Private/PoseSearchLibrary.cpp
+++ b/Engine/Plugins/Experimental/Animation/PoseSearch/Source/Runtime/Private/PoseSearchLibrary.cpp
@@ -34,7 +34,7 @@ TAutoConsoleVariable<int32> CVarAnimMotionMatchDrawHistoryEnable(TEXT("a.MotionM
 //////////////////////////////////////////////////////////////////////////
 // FMotionMatchingState
 
-void FMotionMatchingState::Reset()
+void FMotionMatchingState::Reset(const FTransform& ComponentTransform)
 {
 	CurrentSearchResult.Reset();
 	// Set the elapsed time to INFINITY to trigger a search right away
@@ -42,7 +42,8 @@ void FMotionMatchingState::Reset()
 	WantedPlayRate = 1.f;
 	bJumpedToPose = false;
 	RootBoneDeltaYaw = 0.f;
-	RootBoneWorldYaw = 0.f;
+	RootBoneWorldYaw = FRotator(ComponentTransform.GetRotation()).Yaw;
+
 #if UE_POSE_SEARCH_TRACE_ENABLED
 	RootMotionTransformDelta = FTransform::Identity;
 #endif // UE_POSE_SEARCH_TRACE_ENABLED
@@ -119,16 +120,17 @@ void FMotionMatchingState::UpdateWantedPlayRate(const UE::PoseSearch::FSearchCon
 	}
 }
 
-void FMotionMatchingState::UpdateRootBoneControl(const FAnimInstanceProxy* AnimInstanceProxy, float RootBoneYawFromAnimation)
+void FMotionMatchingState::UpdateRootBoneControl(const FAnimationUpdateContext& Context, float YawFromAnimationBlendRate)
 {
+	const FAnimInstanceProxy* AnimInstanceProxy = Context.AnimInstanceProxy;
 	static const FName RootMotionAttributeName = "RootMotionDelta";
 	static const UE::Anim::FAttributeId RootMotionAttributeId = { RootMotionAttributeName , FCompactPoseBoneIndex(0) };
 
 	const FRotator ComponentWorldRotator(AnimInstanceProxy->GetComponentTransform().GetRotation());
-	if (FMath::IsNearlyZero(RootBoneYawFromAnimation))
+	if (YawFromAnimationBlendRate < 0.f)
 	{
 		RootBoneWorldYaw = ComponentWorldRotator.Yaw;
-		RootBoneWorldYaw = 0.f;
+		RootBoneDeltaYaw = 0.f;
 	}
 	else
 	{
@@ -142,22 +144,21 @@ void FMotionMatchingState::UpdateRootBoneControl(const FAnimInstanceProxy* AnimI
 			const float RootBoneAnimationDelta = RootMotionRotatorDelta.Yaw;
 			const float RootBoneToComponentDelta = FRotator::NormalizeAxis(ComponentWorldRotator.Yaw - RootBoneWorldYaw);
 
-			// @todo: RootBoneYawFromAnimation should be a speed (influenced by dt)
 			// lerping the animation delta with the capsule delta
-			const float RootBoneDelta = FMath::Lerp(RootBoneToComponentDelta, RootBoneAnimationDelta, RootBoneYawFromAnimation);
+			const float RootBoneDelta = FMath::Lerp(RootBoneAnimationDelta, RootBoneToComponentDelta, YawFromAnimationBlendRate * Context.GetDeltaTime());
 
 			RootBoneWorldYaw = FRotator::NormalizeAxis(RootBoneWorldYaw + RootBoneDelta);
+
+			// @todo: handle the case when the character is on top of a rotating platform
+			RootBoneDeltaYaw = FRotator::NormalizeAxis(RootBoneWorldYaw - ComponentWorldRotator.Yaw);
 		}
 		else
 		{
 			UE_LOG(LogPoseSearch, Warning, TEXT("FMotionMatchingState::UpdateRootBoneControl - Couldn't find FTransformAnimationAttribute. Root bone yaw control will not be performed"));
 
 			RootBoneWorldYaw = ComponentWorldRotator.Yaw;
-			RootBoneWorldYaw = 0.f;
+			RootBoneDeltaYaw = 0.f;
 		}
-
-		// @todo: handle the case when the character is on top of a rotating platform
-		RootBoneDeltaYaw = FRotator::NormalizeAxis(RootBoneWorldYaw - ComponentWorldRotator.Yaw);
 	}
 }
 
@@ -271,8 +272,8 @@ void UPoseSearchLibrary::UpdateMotionMatchingState(
 	float SearchThrottleTime,
 	const FFloatInterval& PlayRate,
 	FMotionMatchingState& InOutMotionMatchingState,
-	float RootBoneYawFromAnimation,
-	float RootBoneDeltaYawBlendTime,
+	float YawFromAnimationBlendRate,
+	float YawFromAnimationTrajectoryBlendTime,
 	bool bForceInterrupt,
 	bool bShouldSearch,
 	bool bDebugDrawQuery,
@@ -290,7 +291,7 @@ void UPoseSearchLibrary::UpdateMotionMatchingState(
 		return;
 	}
 
-	InOutMotionMatchingState.UpdateRootBoneControl(Context.AnimInstanceProxy, RootBoneYawFromAnimation);
+	InOutMotionMatchingState.UpdateRootBoneControl(Context, YawFromAnimationBlendRate);
 
 	const float DeltaTime = Context.GetDeltaTime();
 
@@ -307,7 +308,7 @@ void UPoseSearchLibrary::UpdateMotionMatchingState(
 		History = &PoseHistoryProvider->GetPoseHistory();
 	}
 	
-	const FPoseSearchQueryTrajectory TrajectoryRootSpace = ProcessTrajectory(Trajectory, Context.AnimInstanceProxy->GetComponentTransform(), InOutMotionMatchingState.GetRootBoneDeltaYaw(), RootBoneDeltaYawBlendTime, TrajectorySpeedMultiplier);
+	const FPoseSearchQueryTrajectory TrajectoryRootSpace = ProcessTrajectory(Trajectory, Context.AnimInstanceProxy->GetComponentTransform(), InOutMotionMatchingState.GetRootBoneDeltaYaw(), YawFromAnimationTrajectoryBlendTime, TrajectorySpeedMultiplier);
 	FSearchContext SearchContext(&TrajectoryRootSpace, History, 0.f, &InOutMotionMatchingState.PoseIndicesHistory, InOutMotionMatchingState.CurrentSearchResult, PoseJumpThresholdTime, bForceInterrupt);
 
 	const bool bMustSearch = !InOutMotionMatchingState.CurrentSearchResult.IsValid();
@@ -423,7 +424,7 @@ void UPoseSearchLibrary::UpdateMotionMatchingState(
 }
 
 // transforms Trajectory from world space to root bone space, and scale it by TrajectorySpeedMultiplier
-FPoseSearchQueryTrajectory UPoseSearchLibrary::ProcessTrajectory(const FPoseSearchQueryTrajectory& Trajectory, const FTransform& ComponentWorldTransform, float RootBoneDeltaYaw, float RootBoneDeltaYawBlendTime, float TrajectorySpeedMultiplier)
+FPoseSearchQueryTrajectory UPoseSearchLibrary::ProcessTrajectory(const FPoseSearchQueryTrajectory& Trajectory, const FTransform& ComponentWorldTransform, float RootBoneDeltaYaw, float YawFromAnimationTrajectoryBlendTime, float TrajectorySpeedMultiplier)
 {
 	QUICK_SCOPE_CYCLE_COUNTER(STAT_PoseSearch_ProcessTrajectory);
 
@@ -437,7 +438,7 @@ FPoseSearchQueryTrajectory UPoseSearchLibrary::ProcessTrajectory(const FPoseSear
 
 		Sample.Position = ToRootSpace.TransformPosition(Sample.Position);
 
-		const float BlendParam = RootBoneDeltaYawBlendTime < UE_KINDA_SMALL_NUMBER ? 1 : FMath::Clamp(1.f - (Sample.AccumulatedSeconds - RootBoneDeltaYawBlendTime) / RootBoneDeltaYawBlendTime, 0.f, 1.f);
+		const float BlendParam = YawFromAnimationTrajectoryBlendTime < UE_KINDA_SMALL_NUMBER ? 1 : FMath::Clamp(1.f - (Sample.AccumulatedSeconds - YawFromAnimationTrajectoryBlendTime) / YawFromAnimationTrajectoryBlendTime, 0.f, 1.f);
 		const FQuat RootBoneDelta(FRotator(0.f, RootBoneDeltaYaw * BlendParam, 0.f));
 		const FQuat ToRootSpaceRotation = ToRootSpace.GetRotation() * RootBoneDelta;
 		Sample.Facing = ToRootSpaceRotation * Sample.Facing;
diff --git a/Engine/Plugins/Experimental/Animation/PoseSearch/Source/Runtime/Public/PoseSearch/AnimNode_MotionMatching.h b/Engine/Plugins/Experimental/Animation/PoseSearch/Source/Runtime/Public/PoseSearch/AnimNode_MotionMatching.h
--- a/Engine/Plugins/Experimental/Animation/PoseSearch/Source/Runtime/Public/PoseSearch/AnimNode_MotionMatching.h
+++ b/Engine/Plugins/Experimental/Animation/PoseSearch/Source/Runtime/Public/PoseSearch/AnimNode_MotionMatching.h
@@ -103,13 +103,16 @@ private:
 	UPROPERTY(EditAnywhere, Category = Settings, meta = (PinHiddenByDefault))
 	bool bShouldSearch = true;
 
-	// blend time over which the YawFromAnimation is distributed across the trajectory samples (negative values implies YawFromAnimation is constant over the entire trajectory, so the trajectory will not try to recover towards the capsule orientation)
+	// blend time over which the yaw from the animation is distributed across the trajectory samples (negative values implies yaw from the animation is constant over the entire trajectory, so the trajectory will not try to recover towards the capsule orientation)
 	UPROPERTY(EditAnywhere, Category = RootMotion, meta = (PinHiddenByDefault))
 	float YawFromAnimationTrajectoryBlendTime = 0.1f;
 
-	// percentage of delta yaw from the animation controlling the root bone (0% the capsule is authoritative over the root bone orientation, 100% the orientation will be fully controlled by animation)
+	// rate at which the root bone orientation catches up to the capsule orientation after being controlled by animation
+	// (negative values mean the capsule is authoritative over the root bone orientation and the root bone is always synchronized with the capsule,
+	// 0 means the orientation will be fully controlled by animation, and potentially never converge over the capsule orientation,
+	// positive values represent the rate at which the orientation drifts towards the capsule orientation after being controlled by animation)
 	UPROPERTY(EditAnywhere, Category = RootMotion, meta = (PinHiddenByDefault))
-	float YawFromAnimation = 0.f;
+	float YawFromAnimationBlendRate = -1.f;
 
 	FAnimNode_BlendStack_Standalone BlendStackNode;
 
diff --git a/Engine/Plugins/Experimental/Animation/PoseSearch/Source/Runtime/Public/PoseSearch/PoseSearchLibrary.h b/Engine/Plugins/Experimental/Animation/PoseSearch/Source/Runtime/Public/PoseSearch/PoseSearchLibrary.h
--- a/Engine/Plugins/Experimental/Animation/PoseSearch/Source/Runtime/Public/PoseSearch/PoseSearchLibrary.h
+++ b/Engine/Plugins/Experimental/Animation/PoseSearch/Source/Runtime/Public/PoseSearch/PoseSearchLibrary.h
@@ -29,7 +29,7 @@ struct POSESEARCH_API FMotionMatchingState
 	GENERATED_BODY()
 
 	// Reset the state to a default state using the current Database
-	void Reset();
+	void Reset(const FTransform& ComponentTransform);
 
 	// Attempts to set the internal state to match the provided asset time including updating the internal DbPoseIdx. 
 	// If the provided asset time is out of bounds for the currently playing asset then this function will reset the 
@@ -41,7 +41,7 @@ struct POSESEARCH_API FMotionMatchingState
 
 	void UpdateWantedPlayRate(const UE::PoseSearch::FSearchContext& SearchContext, const FFloatInterval& PlayRate, float TrajectorySpeedMultiplier);
 
-	void UpdateRootBoneControl(const FAnimInstanceProxy* AnimInstanceProxy, float RootBoneYawFromAnimation);
+	void UpdateRootBoneControl(const FAnimationUpdateContext& Context, float YawFromAnimationBlendRate);
 
 	float GetRootBoneDeltaYaw() const { return RootBoneDeltaYaw; }
 
@@ -97,7 +97,7 @@ class POSESEARCH_API UPoseSearchLibrary : public UBlueprintFunctionLibrary
 		const FPoseSearchQueryTrajectory& Trajectory,
 		const FTransform& ComponentWorldTransform,
 		float RootBoneDeltaYaw,
-		float RootBoneDeltaYawBlendTime,
+		float YawFromAnimationTrajectoryBlendTime,
 		float TrajectorySpeedMultiplier);
 
 public:
@@ -131,8 +131,8 @@ public:
 		float SearchThrottleTime,
 		const FFloatInterval& PlayRate,
 		FMotionMatchingState& InOutMotionMatchingState,
-		float RootBoneYawFromAnimation,
-		float RootBoneDeltaYawBlendTime,
+		float YawFromAnimationBlendRate,
+		float YawFromAnimationTrajectoryBlendTime,
 		bool bForceInterrupt = false,
 		bool bShouldSearch = true,
 		bool bDebugDrawQuery = false,
