motion matching
- renaming indexing structs to place them in UE::PoseSearch namespace
- renaming IPoseFilter to IPoseSearchFilter
- renaming FPoseFilters to FSearchFilters
- aligning data in FSearchIndex and FSearchIndexBase

diff --git a/Engine/Plugins/Experimental/Animation/PoseSearch/Source/Editor/Private/PoseSearchDatabaseAssetListItem.cpp b/Engine/Plugins/Experimental/Animation/PoseSearch/Source/Editor/Private/PoseSearchDatabaseAssetListItem.cpp
--- a/Engine/Plugins/Experimental/Animation/PoseSearch/Source/Editor/Private/PoseSearchDatabaseAssetListItem.cpp
+++ b/Engine/Plugins/Experimental/Animation/PoseSearch/Source/Editor/Private/PoseSearchDatabaseAssetListItem.cpp
@@ -361,7 +361,7 @@ namespace UE::PoseSearch
 	{
 		TSharedPtr<FDatabaseViewModel> ViewModelPtr = EditorViewModel.Pin();
 		TSharedPtr<FDatabaseAssetTreeNode> TreeNodePtr = WeakAssetTreeNode.Pin();
-		if (const FPoseSearchIndexAsset* SelectedIndexAsset = ViewModelPtr->GetSelectedActorIndexAsset())
+		if (const FSearchIndexAsset* SelectedIndexAsset = ViewModelPtr->GetSelectedActorIndexAsset())
 		{
 			if (TreeNodePtr->SourceAssetIdx == SelectedIndexAsset->SourceAssetIdx)
 			{
diff --git a/Engine/Plugins/Experimental/Animation/PoseSearch/Source/Editor/Private/PoseSearchDatabaseDataDetails.cpp b/Engine/Plugins/Experimental/Animation/PoseSearch/Source/Editor/Private/PoseSearchDatabaseDataDetails.cpp
--- a/Engine/Plugins/Experimental/Animation/PoseSearch/Source/Editor/Private/PoseSearchDatabaseDataDetails.cpp
+++ b/Engine/Plugins/Experimental/Animation/PoseSearch/Source/Editor/Private/PoseSearchDatabaseDataDetails.cpp
@@ -219,10 +219,7 @@ public:
 							{ return *PreviewActor.Actor->GetName() == ColumnName; }))
 						{
 							const int32 PoseIdx = FoundPreviewActor->CurrentPoseIndex;
-							const int32 NumDimensions = PoseSearchDatabase->Schema->SchemaCardinality;
-							const FPoseSearchIndex& SearchIndex = PoseSearchDatabase->GetSearchIndex();
-							TArrayView<float> ReconstructedPoseValuesBuffer((float*)FMemory_Alloca(NumDimensions * sizeof(float)), NumDimensions);
-							const TConstArrayView<float> PoseValues = SearchIndex.Values.IsEmpty() ? SearchIndex.GetReconstructedPoseValues(PoseIdx, ReconstructedPoseValuesBuffer) : SearchIndex.GetPoseValues(PoseIdx);
+							const TArray<float> PoseValues = PoseSearchDatabase->GetSearchIndex().GetPoseValuesSafe(PoseIdx);
 
 							const int32 DataOffset = ChannelItem->GetDataOffset();
 							const int32 Cardinality = ChannelItem->GetCardinality();
diff --git a/Engine/Plugins/Experimental/Animation/PoseSearch/Source/Editor/Private/PoseSearchDatabaseEditorReflection.cpp b/Engine/Plugins/Experimental/Animation/PoseSearch/Source/Editor/Private/PoseSearchDatabaseEditorReflection.cpp
--- a/Engine/Plugins/Experimental/Animation/PoseSearch/Source/Editor/Private/PoseSearchDatabaseEditorReflection.cpp
+++ b/Engine/Plugins/Experimental/Animation/PoseSearch/Source/Editor/Private/PoseSearchDatabaseEditorReflection.cpp
@@ -111,7 +111,7 @@ void UPoseSearchDatabaseStatistics::Initialize(const UPoseSearchDatabase* PoseSe
 	
 	if (PoseSearchDatabase)
 	{
-		const FPoseSearchIndex& SearchIndex = PoseSearchDatabase->GetSearchIndex();
+		const UE::PoseSearch::FSearchIndex& SearchIndex = PoseSearchDatabase->GetSearchIndex();
 		// General Information
 	
 		AnimationSequences = PoseSearchDatabase->AnimationAssets.Num();
@@ -121,7 +121,7 @@ void UPoseSearchDatabaseStatistics::Initialize(const UPoseSearchDatabase* PoseSe
 		TotalAnimationPosesInTime = FText::Format(TimeFormat, static_cast<double>(TotalAnimationPosesInFrames) / SampleRate);
 			
 		uint32 NumOfSearchablePoses = 0;
-		for (const FPoseSearchPoseMetadata& PoseMetadata : SearchIndex.PoseMetadata)
+		for (const UE::PoseSearch::FPoseMetadata& PoseMetadata : SearchIndex.PoseMetadata)
 		{
 			if (!PoseMetadata.IsBlockTransition())
 			{
diff --git a/Engine/Plugins/Experimental/Animation/PoseSearch/Source/Editor/Private/PoseSearchDatabaseViewModel.cpp b/Engine/Plugins/Experimental/Animation/PoseSearch/Source/Editor/Private/PoseSearchDatabaseViewModel.cpp
--- a/Engine/Plugins/Experimental/Animation/PoseSearch/Source/Editor/Private/PoseSearchDatabaseViewModel.cpp
+++ b/Engine/Plugins/Experimental/Animation/PoseSearch/Source/Editor/Private/PoseSearchDatabaseViewModel.cpp
@@ -132,8 +132,8 @@ namespace UE::PoseSearch
 		FDatabasePreviewActor PreviewActor;
 
 		check(PoseSearchDatabase);
-		const FPoseSearchIndex& SearchIndex = PoseSearchDatabase->GetSearchIndex();
-		const FPoseSearchIndexAsset& IndexAsset = SearchIndex.Assets[IndexAssetIndex];
+		const FSearchIndex& SearchIndex = PoseSearchDatabase->GetSearchIndex();
+		const FSearchIndexAsset& IndexAsset = SearchIndex.Assets[IndexAssetIndex];
 
 		const FPoseSearchDatabaseAnimationAssetBase* DatabaseAnimationAsset = PoseSearchDatabase->GetAnimationAssetBase(IndexAsset.SourceAssetIdx);
 		UAnimationAsset* PreviewAsset = DatabaseAnimationAsset->GetAnimationAsset();
@@ -216,7 +216,7 @@ namespace UE::PoseSearch
 
 	void FDatabaseViewModel::UpdatePreviewActors(bool bInTickPlayTime)
 	{
-		const FPoseSearchIndex& SearchIndex = PoseSearchDatabase->GetSearchIndex();
+		const FSearchIndex& SearchIndex = PoseSearchDatabase->GetSearchIndex();
 		TSharedPtr<FDatabasePreviewScene> PreviewScene = PreviewScenePtr.Pin();
 		for (FDatabasePreviewActor& PreviewActor : GetPreviewActors())
 		{
@@ -233,7 +233,7 @@ namespace UE::PoseSearch
 			}
 
 			PreviewActor.CurrentTime = 0.f;
-			const FPoseSearchIndexAsset& IndexAsset = SearchIndex.Assets[PreviewActor.IndexAssetIndex];
+			const FSearchIndexAsset& IndexAsset = SearchIndex.Assets[PreviewActor.IndexAssetIndex];
 			float CurrentPlayTime = PlayTime + IndexAsset.GetFirstSampleTime(PoseSearchDatabase->Schema->SampleRate) + PreviewActor.PlayTimeOffset;
 			FAnimationRuntime::AdvanceTime(false, CurrentPlayTime, PreviewActor.CurrentTime, IndexAsset.GetLastSampleTime(PoseSearchDatabase->Schema->SampleRate));
 			 
@@ -401,7 +401,7 @@ namespace UE::PoseSearch
 
 		if (FAsyncPoseSearchDatabasesManagement::RequestAsyncBuildIndex(PoseSearchDatabase, ERequestAsyncBuildFlag::ContinueRequest))
 		{
-			const FPoseSearchIndex& SearchIndex = PoseSearchDatabase->GetSearchIndex();
+			const FSearchIndex& SearchIndex = PoseSearchDatabase->GetSearchIndex();
 			if (SearchIndex.PoseMetadata.IsValidIndex(PoseIdx))
 			{
 				const uint32 IndexAssetIndex = SearchIndex.PoseMetadata[PoseIdx].GetAssetIndex();
@@ -410,7 +410,7 @@ namespace UE::PoseSearch
 					FDatabasePreviewActor PreviewActor = SpawnPreviewActor(IndexAssetIndex, PoseIdx);
 					if (PreviewActor.IsValid())
 					{
-						const FPoseSearchIndexAsset& IndexAsset = SearchIndex.Assets[IndexAssetIndex];
+						const FSearchIndexAsset& IndexAsset = SearchIndex.Assets[IndexAssetIndex];
 						MaxPreviewPlayLength = FMath::Max(MaxPreviewPlayLength, IndexAsset.GetLastSampleTime(PoseSearchDatabase->Schema->SampleRate) - PreviewActor.PlayTimeOffset);
 						MinPreviewPlayLength = FMath::Min(MinPreviewPlayLength, IndexAsset.GetFirstSampleTime(PoseSearchDatabase->Schema->SampleRate) - PreviewActor.PlayTimeOffset);
 						PreviewActors.Add(PreviewActor);
@@ -448,10 +448,10 @@ namespace UE::PoseSearch
 				AssociatedAssetIndices.FindOrAdd(InSelectedNodes[i]->SourceAssetIdx) = i;
 			}
 
-			const FPoseSearchIndex& SearchIndex = PoseSearchDatabase->GetSearchIndex();
+			const FSearchIndex& SearchIndex = PoseSearchDatabase->GetSearchIndex();
 			for (int32 IndexAssetIndex = 0; IndexAssetIndex < SearchIndex.Assets.Num(); ++IndexAssetIndex)
 			{
-				const FPoseSearchIndexAsset& IndexAsset = SearchIndex.Assets[IndexAssetIndex];
+				const FSearchIndexAsset& IndexAsset = SearchIndex.Assets[IndexAssetIndex];
 				if (AnimationPreviewMode == EAnimationPreviewMode::OriginalAndMirrored || !IndexAsset.bMirrored)
 				{
 					if (const int32* SelectedNodesIndex = AssociatedAssetIndices.Find(IndexAsset.SourceAssetIdx))
@@ -496,11 +496,11 @@ namespace UE::PoseSearch
 		}
 	}
 
-	const FPoseSearchIndexAsset* FDatabaseViewModel::GetSelectedActorIndexAsset() const
+	const FSearchIndexAsset* FDatabaseViewModel::GetSelectedActorIndexAsset() const
 	{
 		if (FAsyncPoseSearchDatabasesManagement::RequestAsyncBuildIndex(PoseSearchDatabase, ERequestAsyncBuildFlag::ContinueRequest))
 		{
-			const FPoseSearchIndex& SearchIndex = PoseSearchDatabase->GetSearchIndex();
+			const FSearchIndex& SearchIndex = PoseSearchDatabase->GetSearchIndex();
 			if (SearchIndex.Assets.IsValidIndex(SelectedActorIndexAssetIndex))
 			{
 				return &SearchIndex.Assets[SelectedActorIndexAssetIndex];
@@ -539,12 +539,12 @@ namespace UE::PoseSearch
 	{
 		if (FAsyncPoseSearchDatabasesManagement::RequestAsyncBuildIndex(PoseSearchDatabase, ERequestAsyncBuildFlag::ContinueRequest))
 		{
-			const FPoseSearchIndex& SearchIndex = PoseSearchDatabase->GetSearchIndex();
+			const FSearchIndex& SearchIndex = PoseSearchDatabase->GetSearchIndex();
 			for (const FDatabasePreviewActor& PreviewActor : GetPreviewActors())
 			{
 				if (PreviewActor.IndexAssetIndex >= 0 && PreviewActor.IndexAssetIndex < SearchIndex.Assets.Num() && PreviewActor.Sampler.IsInitialized())
 				{
-					const FPoseSearchIndexAsset& IndexAsset = SearchIndex.Assets[PreviewActor.IndexAssetIndex];
+					const FSearchIndexAsset& IndexAsset = SearchIndex.Assets[PreviewActor.IndexAssetIndex];
 					if (IndexAsset.SourceAssetIdx == SourceAssetIdx)
 					{
 						CurrentPlayTime = PreviewActor.Sampler.ToNormalizedTime(PlayTime + IndexAsset.GetFirstSampleTime(PoseSearchDatabase->Schema->SampleRate) + PreviewActor.PlayTimeOffset);
@@ -554,7 +554,7 @@ namespace UE::PoseSearch
 				}
 			}
 
-			for (const FPoseSearchIndexAsset& IndexAsset : SearchIndex.Assets)
+			for (const FSearchIndexAsset& IndexAsset : SearchIndex.Assets)
 			{
 				if (IndexAsset.SourceAssetIdx == SourceAssetIdx)
 				{
diff --git a/Engine/Plugins/Experimental/Animation/PoseSearch/Source/Editor/Private/PoseSearchDebuggerDatabaseView.cpp b/Engine/Plugins/Experimental/Animation/PoseSearch/Source/Editor/Private/PoseSearchDebuggerDatabaseView.cpp
--- a/Engine/Plugins/Experimental/Animation/PoseSearch/Source/Editor/Private/PoseSearchDebuggerDatabaseView.cpp
+++ b/Engine/Plugins/Experimental/Animation/PoseSearch/Source/Editor/Private/PoseSearchDebuggerDatabaseView.cpp
@@ -156,8 +156,8 @@ static void AddUnfilteredDatabaseRow(const UPoseSearchDatabase* Database,
 	TArray<TSharedRef<FDebuggerDatabaseRowData>>& UnfilteredDatabaseRows, TSharedRef<FDebuggerDatabaseSharedData> SharedData,
 	int32 DbPoseIdx, EPoseCandidateFlags PoseCandidateFlags, const FPoseSearchCost& Cost = FPoseSearchCost())
 {
-	const FPoseSearchIndex& SearchIndex = Database->GetSearchIndex();
-	if (const FPoseSearchIndexAsset* SearchIndexAsset = SearchIndex.GetAssetForPoseSafe(DbPoseIdx))
+	const FSearchIndex& SearchIndex = Database->GetSearchIndex();
+	if (const FSearchIndexAsset* SearchIndexAsset = SearchIndex.GetAssetForPoseSafe(DbPoseIdx))
 	{
 		TSharedRef<FDebuggerDatabaseRowData>& Row = UnfilteredDatabaseRows.Add_GetRef(MakeShared<FDebuggerDatabaseRowData>(SharedData));
 
@@ -255,7 +255,7 @@ void SDebuggerDatabaseView::Update(const FTraceMotionMatchingStateMessage& State
 					PoseEntriesIdx.Add(PoseEntry.DbPoseIdx);
 				}
 
-				const FPoseSearchIndex& SearchIndex = Database->GetSearchIndex();
+				const FSearchIndex& SearchIndex = Database->GetSearchIndex();
 				for (int32 DbPoseIdx = 0; DbPoseIdx < SearchIndex.GetNumPoses(); ++DbPoseIdx)
 				{
 					if (!PoseEntriesIdx.Find(DbPoseIdx))
diff --git a/Engine/Plugins/Experimental/Animation/PoseSearch/Source/Editor/Private/PoseSearchDebuggerViewModel.cpp b/Engine/Plugins/Experimental/Animation/PoseSearch/Source/Editor/Private/PoseSearchDebuggerViewModel.cpp
--- a/Engine/Plugins/Experimental/Animation/PoseSearch/Source/Editor/Private/PoseSearchDebuggerViewModel.cpp
+++ b/Engine/Plugins/Experimental/Animation/PoseSearch/Source/Editor/Private/PoseSearchDebuggerViewModel.cpp
@@ -59,17 +59,18 @@ void FDebuggerViewModel::ShowSelectedSkeleton(const UPoseSearchDatabase* Databas
 		return;
 	}
 
-	const FPoseSearchIndex& SearchIndex = Database->GetSearchIndex();
-	const FPoseSearchIndexAsset& IndexAsset = SearchIndex.GetAssetForPose(DbPoseIdx);
-	
-	Component->ResetToStart(); 
-	bSelecting = true;
-
-	Skeletons[SelectedPose].Time = Time;
-	Skeletons[SelectedPose].bMirrored = IndexAsset.bMirrored;
-	Skeletons[SelectedPose].SourceDatabase = Database;
-	Skeletons[SelectedPose].AssetIdx = IndexAsset.SourceAssetIdx;
-	Skeletons[SelectedPose].BlendParameters = IndexAsset.BlendParameters;
+	const FSearchIndex& SearchIndex = Database->GetSearchIndex();
+	if (const FSearchIndexAsset* IndexAsset = SearchIndex.GetAssetForPoseSafe(DbPoseIdx))
+	{
+		Component->ResetToStart();
+		bSelecting = true;
+
+		Skeletons[SelectedPose].Time = Time;
+		Skeletons[SelectedPose].bMirrored = IndexAsset->bMirrored;
+		Skeletons[SelectedPose].SourceDatabase = Database;
+		Skeletons[SelectedPose].AssetIdx = IndexAsset->SourceAssetIdx;
+		Skeletons[SelectedPose].BlendParameters = IndexAsset->BlendParameters;
+	}
 }
 
 void FDebuggerViewModel::ClearSelectedSkeleton()
@@ -140,9 +141,9 @@ void FDebuggerViewModel::OnUpdateNodeSelection(int32 InNodeId)
 		const UPoseSearchDatabase* CurrentDatabase = ActiveMotionMatchingState->GetCurrentDatabase();
 		if (FAsyncPoseSearchDatabasesManagement::RequestAsyncBuildIndex(CurrentDatabase, ERequestAsyncBuildFlag::ContinueRequest))
 		{
-			const FPoseSearchIndex& CurrentSearchIndex = CurrentDatabase->GetSearchIndex();
+			const FSearchIndex& CurrentSearchIndex = CurrentDatabase->GetSearchIndex();
 			int32 CurrentPoseIdx = ActiveMotionMatchingState->GetCurrentDatabasePoseIndex();
-			if (const FPoseSearchIndexAsset* IndexAsset = CurrentSearchIndex.GetAssetForPoseSafe(CurrentPoseIdx))
+			if (const FSearchIndexAsset* IndexAsset = CurrentSearchIndex.GetAssetForPoseSafe(CurrentPoseIdx))
 			{
 				Skeletons[Asset].bMirrored = IndexAsset->bMirrored;
 				Skeletons[Asset].SourceDatabase = CurrentDatabase;
diff --git a/Engine/Plugins/Experimental/Animation/PoseSearch/Source/Editor/Public/PoseSearchDatabaseViewModel.h b/Engine/Plugins/Experimental/Animation/PoseSearch/Source/Editor/Public/PoseSearchDatabaseViewModel.h
--- a/Engine/Plugins/Experimental/Animation/PoseSearch/Source/Editor/Public/PoseSearchDatabaseViewModel.h
+++ b/Engine/Plugins/Experimental/Animation/PoseSearch/Source/Editor/Public/PoseSearchDatabaseViewModel.h
@@ -12,7 +12,6 @@
 
 class UWorld;
 class UPoseSearchDatabase;
-struct FPoseSearchIndexAsset;
 class UAnimPreviewInstance;
 class UDebugSkelMeshComponent;
 class UAnimComposite;
@@ -23,6 +22,7 @@ class UMirrorDataTable;
 namespace UE::PoseSearch
 {
 	class FDatabaseAssetTreeNode;
+	struct FSearchIndexAsset;
 	class SDatabaseDataDetails;
 
 	enum class EFeaturesDrawMode : uint8
@@ -121,7 +121,7 @@ namespace UE::PoseSearch
 		void SetDrawQueryVector(bool bValue);
 		bool ShouldDrawQueryVector() const { return bDrawQueryVector && !bIsEditorSelection; }
 
-		const FPoseSearchIndexAsset* GetSelectedActorIndexAsset() const;
+		const FSearchIndexAsset* GetSelectedActorIndexAsset() const;
 
 		TRange<double> GetPreviewPlayRange() const;
 
diff --git a/Engine/Plugins/Experimental/Animation/PoseSearch/Source/Runtime/Private/AnimNode_MotionMatching.cpp b/Engine/Plugins/Experimental/Animation/PoseSearch/Source/Runtime/Private/AnimNode_MotionMatching.cpp
--- a/Engine/Plugins/Experimental/Animation/PoseSearch/Source/Runtime/Private/AnimNode_MotionMatching.cpp
+++ b/Engine/Plugins/Experimental/Animation/PoseSearch/Source/Runtime/Private/AnimNode_MotionMatching.cpp
@@ -106,7 +106,7 @@ void FAnimNode_MotionMatching::UpdateAssetPlayer(const FAnimationUpdateContext&
 		if (FAsyncPoseSearchDatabasesManagement::RequestAsyncBuildIndex(MotionMatchingState.CurrentSearchResult.Database.Get(), ERequestAsyncBuildFlag::ContinueRequest))
 		{
 			const UPoseSearchDatabase* CurrentResultDatabase = MotionMatchingState.CurrentSearchResult.Database.Get();
-			const FPoseSearchIndex& SearchIndex = CurrentResultDatabase->GetSearchIndex();
+			const FSearchIndex& SearchIndex = CurrentResultDatabase->GetSearchIndex();
 			if (!SearchIndex.IsValidPoseIndex(MotionMatchingState.CurrentSearchResult.PrevPoseIdx) ||
 				!SearchIndex.IsValidPoseIndex(MotionMatchingState.CurrentSearchResult.NextPoseIdx))
 			{
@@ -176,7 +176,7 @@ void FAnimNode_MotionMatching::UpdateAssetPlayer(const FAnimationUpdateContext&
 	// If a new pose is requested, blend into the new asset via BlendStackNode
 	if (MotionMatchingState.bJumpedToPose)
 	{
-		const FPoseSearchIndexAsset* SearchIndexAsset = MotionMatchingState.CurrentSearchResult.GetSearchIndexAsset();
+		const FSearchIndexAsset* SearchIndexAsset = MotionMatchingState.CurrentSearchResult.GetSearchIndexAsset();
 		const UPoseSearchDatabase* CurrentResultDatabase = MotionMatchingState.CurrentSearchResult.Database.Get();
 		if (SearchIndexAsset && CurrentResultDatabase && CurrentResultDatabase->Schema)
 		{
diff --git a/Engine/Plugins/Experimental/Animation/PoseSearch/Source/Runtime/Private/PoseSearchAssetIndexer.cpp b/Engine/Plugins/Experimental/Animation/PoseSearch/Source/Runtime/Private/PoseSearchAssetIndexer.cpp
--- a/Engine/Plugins/Experimental/Animation/PoseSearch/Source/Runtime/Private/PoseSearchAssetIndexer.cpp
+++ b/Engine/Plugins/Experimental/Animation/PoseSearch/Source/Runtime/Private/PoseSearchAssetIndexer.cpp
@@ -75,7 +75,7 @@ static FSamplingParam WrapOrClampSamplingParam(bool bCanWrap, float SamplingPara
 
 //////////////////////////////////////////////////////////////////////////
 // FAssetIndexer
-FAssetIndexer::FAssetIndexer(const FBoneContainer& InBoneContainer, const FPoseSearchIndexAsset& InSearchIndexAsset, 
+FAssetIndexer::FAssetIndexer(const FBoneContainer& InBoneContainer, const FSearchIndexAsset& InSearchIndexAsset, 
 	const FAssetSamplingContext& InSamplingContext, const UPoseSearchSchema& InSchema, const FAnimationAssetSampler& InAssetSampler)
 : BoneContainer(InBoneContainer)
 , CachedEntries()
@@ -86,7 +86,7 @@ FAssetIndexer::FAssetIndexer(const FBoneContainer& InBoneContainer, const FPoseS
 {
 }
 
-void FAssetIndexer::AssignWorkingData(TArrayView<float> InOutFeatureVectorTable, TArrayView<FPoseSearchPoseMetadata> InOutPoseMetadata)
+void FAssetIndexer::AssignWorkingData(TArrayView<float> InOutFeatureVectorTable, TArrayView<FPoseMetadata> InOutPoseMetadata)
 {
 	FeatureVectorTable = InOutFeatureVectorTable;
 	PoseMetadata = InOutPoseMetadata;
diff --git a/Engine/Plugins/Experimental/Animation/PoseSearch/Source/Runtime/Private/PoseSearchContext.cpp b/Engine/Plugins/Experimental/Animation/PoseSearch/Source/Runtime/Private/PoseSearchContext.cpp
--- a/Engine/Plugins/Experimental/Animation/PoseSearch/Source/Runtime/Private/PoseSearchContext.cpp
+++ b/Engine/Plugins/Experimental/Animation/PoseSearch/Source/Runtime/Private/PoseSearchContext.cpp
@@ -40,7 +40,7 @@ bool FDebugDrawParams::CanDraw() const
 	return (AnimInstanceProxy || World) && Database && Database->Schema && Database->Schema->IsValid();
 }
 
-const FPoseSearchIndex* FDebugDrawParams::GetSearchIndex() const
+const FSearchIndex* FDebugDrawParams::GetSearchIndex() const
 {
 	return Database ? &Database->GetSearchIndex() : nullptr;
 }
@@ -257,7 +257,7 @@ FSearchContext::FSearchContext(const FPoseSearchQueryTrajectory* InTrajectory, c
 {
 	if (CurrentResult.IsValid())
 	{
-		const FPoseSearchIndex& SearchIndex = CurrentResult.Database->GetSearchIndex();
+		const FSearchIndex& SearchIndex = CurrentResult.Database->GetSearchIndex();
 		if (SearchIndex.Values.IsEmpty())
 		{
 			const int32 NumDimensions = CurrentResult.Database->Schema->SchemaCardinality;
diff --git a/Engine/Plugins/Experimental/Animation/PoseSearch/Source/Runtime/Private/PoseSearchDatabase.cpp b/Engine/Plugins/Experimental/Animation/PoseSearch/Source/Runtime/Private/PoseSearchDatabase.cpp
--- a/Engine/Plugins/Experimental/Animation/PoseSearch/Source/Runtime/Private/PoseSearchDatabase.cpp
+++ b/Engine/Plugins/Experimental/Animation/PoseSearch/Source/Runtime/Private/PoseSearchDatabase.cpp
@@ -35,11 +35,11 @@ static void PopulateNonSelectableIdx(FNonSelectableIdx& NonSelectableIdx, FSearc
 {
 	check(Database);
 #if UE_POSE_SEARCH_TRACE_ENABLED
-	const FPoseSearchIndex& SearchIndex = Database->GetSearchIndex();
+	const FSearchIndex& SearchIndex = Database->GetSearchIndex();
 #endif
 
 	NonSelectableIdx.Reset();
-	const FPoseSearchIndexAsset* CurrentIndexAsset = SearchContext.GetCurrentResult().GetSearchIndexAsset();
+	const FSearchIndexAsset* CurrentIndexAsset = SearchContext.GetCurrentResult().GetSearchIndexAsset();
 	if (CurrentIndexAsset && SearchContext.IsCurrentResultFromDatabase(Database) && SearchContext.GetPoseJumpThresholdTime() > 0.f)
 	{
 		const int32 PoseJumpIndexThreshold = FMath::FloorToInt(SearchContext.GetPoseJumpThresholdTime() * Database->Schema->SampleRate);
@@ -140,47 +140,47 @@ static void PopulateNonSelectableIdx(FNonSelectableIdx& NonSelectableIdx, FSearc
 	NonSelectableIdx.Sort();
 }
 
-struct FPoseFilters
+struct FSearchFilters
 {
-	FPoseFilters(const UPoseSearchSchema* Schema, TConstArrayView<size_t> NonSelectableIdx, bool bAnyBlockTransition)
+	FSearchFilters(const UPoseSearchSchema* Schema, TConstArrayView<size_t> NonSelectableIdx, bool bAnyBlockTransition)
 	{
-		NonSelectableIdxPoseFilter.NonSelectableIdx = NonSelectableIdx;
+		NonSelectableIdxFilter.NonSelectableIdx = NonSelectableIdx;
 
 		if (bAnyBlockTransition)
 		{
-			AllPoseFilters.Add(&BlockTransitionPoseFilter);
+			Filters.Add(&BlockTransitionFilter);
 		}
 
-		if (NonSelectableIdxPoseFilter.IsPoseFilterActive())
+		if (NonSelectableIdxFilter.IsFilterActive())
 		{
-			AllPoseFilters.Add(&NonSelectableIdxPoseFilter);
+			Filters.Add(&NonSelectableIdxFilter);
 		}
 
-		for (const IPoseFilter* ChannelPoseFilter : Schema->GetChannels())
+		for (const IPoseSearchFilter* Filter : Schema->GetChannels())
 		{
-			if (ChannelPoseFilter->IsPoseFilterActive())
+			if (Filter->IsFilterActive())
 			{
-				AllPoseFilters.Add(ChannelPoseFilter);
+				Filters.Add(Filter);
 			}
 		}
 	}
 
-	bool AreFiltersValid(const FPoseSearchIndex& SearchIndex, TConstArrayView<float> PoseValues, TConstArrayView<float> QueryValues, int32 PoseIdx, const FPoseSearchPoseMetadata& Metadata
+	bool AreFiltersValid(const FSearchIndex& SearchIndex, TConstArrayView<float> PoseValues, TConstArrayView<float> QueryValues, int32 PoseIdx, const FPoseMetadata& Metadata
 #if UE_POSE_SEARCH_TRACE_ENABLED
 		, UE::PoseSearch::FSearchContext& SearchContext, const UPoseSearchDatabase* Database
 #endif
 	) const
 	{
-		for (const IPoseFilter* PoseFilter : AllPoseFilters)
+		for (const IPoseSearchFilter* Filter : Filters)
 		{
-			if (!PoseFilter->IsPoseValid(PoseValues, QueryValues, PoseIdx, Metadata))
+			if (!Filter->IsFilterValid(PoseValues, QueryValues, PoseIdx, Metadata))
 			{
 #if UE_POSE_SEARCH_TRACE_ENABLED
-				if (PoseFilter == &NonSelectableIdxPoseFilter)
+				if (Filter == &NonSelectableIdxFilter)
 				{
 					// candidate already added to SearchContext.BestCandidates by PopulateNonSelectableIdx
 				}
-				else if (PoseFilter == &BlockTransitionPoseFilter)
+				else if (Filter == &BlockTransitionFilter)
 				{
 					const FPoseSearchCost PoseCost = SearchIndex.ComparePoses(PoseIdx, 0.f, PoseValues, QueryValues);
 					SearchContext.BestCandidates.Add(PoseCost, PoseIdx, Database, EPoseCandidateFlags::DiscardedBy_BlockTransition);
@@ -198,14 +198,14 @@ struct FPoseFilters
 	};
 
 private:
-	struct FNonSelectableIdxPoseFilter : public IPoseFilter
+	struct FNonSelectableIdxFilter : public IPoseSearchFilter
 	{
-		virtual bool IsPoseFilterActive() const override
+		virtual bool IsFilterActive() const override
 		{
 			return !NonSelectableIdx.IsEmpty();
 		}
 
-		virtual bool IsPoseValid(TConstArrayView<float> PoseValues, TConstArrayView<float> QueryValues, int32 PoseIdx, const FPoseSearchPoseMetadata& Metadata) const override
+		virtual bool IsFilterValid(TConstArrayView<float> PoseValues, TConstArrayView<float> QueryValues, int32 PoseIdx, const FPoseMetadata& Metadata) const override
 		{
 			return Algo::BinarySearch(NonSelectableIdx, PoseIdx) == INDEX_NONE;
 		}
@@ -213,23 +213,23 @@ private:
 		TConstArrayView<size_t> NonSelectableIdx;
 	};
 
-	struct FBlockTransitionPoseFilter : public IPoseFilter
+	struct FBlockTransitionFilter : public IPoseSearchFilter
 	{
-		virtual bool IsPoseFilterActive() const override
+		virtual bool IsFilterActive() const override
 		{
 			return true;
 		}
 
-		virtual bool IsPoseValid(TConstArrayView<float> PoseValues, TConstArrayView<float> QueryValues, int32 PoseIdx, const FPoseSearchPoseMetadata& Metadata) const override
+		virtual bool IsFilterValid(TConstArrayView<float> PoseValues, TConstArrayView<float> QueryValues, int32 PoseIdx, const FPoseMetadata& Metadata) const override
 		{
 			return !Metadata.IsBlockTransition();
 		}
 	};
 
-	FNonSelectableIdxPoseFilter NonSelectableIdxPoseFilter;
-	FBlockTransitionPoseFilter BlockTransitionPoseFilter;
+	FNonSelectableIdxFilter NonSelectableIdxFilter;
+	FBlockTransitionFilter BlockTransitionFilter;
 
-	TArray<const IPoseFilter*, TInlineAllocator<64>> AllPoseFilters;
+	TArray<const IPoseSearchFilter*, TInlineAllocator<64>> Filters;
 };
 
 } // namespace UE::PoseSearch
@@ -439,26 +439,26 @@ UPoseSearchDatabase::~UPoseSearchDatabase()
 {
 }
 
-void UPoseSearchDatabase::SetSearchIndex(const FPoseSearchIndex& SearchIndex)
+void UPoseSearchDatabase::SetSearchIndex(const UE::PoseSearch::FSearchIndex& SearchIndex)
 {
 	check(IsInGameThread());
 	SearchIndexPrivate = SearchIndex;
 }
 
-const FPoseSearchIndex& UPoseSearchDatabase::GetSearchIndex() const
+const UE::PoseSearch::FSearchIndex& UPoseSearchDatabase::GetSearchIndex() const
 {
 	// making sure the search index is consistent. if it fails the calling code hasn't been protected by FAsyncPoseSearchDatabasesManagement::RequestAsyncBuildIndex
 	check(Schema && Schema->IsValid() && !SearchIndexPrivate.IsEmpty() && SearchIndexPrivate.WeightsSqrt.Num() == Schema->SchemaCardinality);
 	return SearchIndexPrivate;
 }
 
-int32 UPoseSearchDatabase::GetPoseIndexFromTime(float Time, const FPoseSearchIndexAsset& SearchIndexAsset) const
+int32 UPoseSearchDatabase::GetPoseIndexFromTime(float Time, const UE::PoseSearch::FSearchIndexAsset& SearchIndexAsset) const
 {
 	const bool bIsLooping = IsSourceAssetLooping(SearchIndexAsset);
 	return SearchIndexAsset.GetPoseIndexFromTime(Time, bIsLooping, Schema->SampleRate);
 }
 
-bool UPoseSearchDatabase::GetPoseIndicesAndLerpValueFromTime(float Time, const FPoseSearchIndexAsset& SearchIndexAsset, int32& PrevPoseIdx, int32& PoseIdx, int32& NextPoseIdx, float& LerpValue) const
+bool UPoseSearchDatabase::GetPoseIndicesAndLerpValueFromTime(float Time, const UE::PoseSearch::FSearchIndexAsset& SearchIndexAsset, int32& PrevPoseIdx, int32& PoseIdx, int32& NextPoseIdx, float& LerpValue) const
 {
 	PoseIdx = GetPoseIndexFromTime(Time, SearchIndexAsset);
 	if (PoseIdx == INDEX_NONE)
@@ -499,7 +499,7 @@ const FInstancedStruct& UPoseSearchDatabase::GetAnimationAssetStruct(int32 Anima
 	return AnimationAssets[AnimationAssetIndex];
 }
 
-const FInstancedStruct& UPoseSearchDatabase::GetAnimationAssetStruct(const FPoseSearchIndexAsset& SearchIndexAsset) const
+const FInstancedStruct& UPoseSearchDatabase::GetAnimationAssetStruct(const UE::PoseSearch::FSearchIndexAsset& SearchIndexAsset) const
 {
 	return GetAnimationAssetStruct(SearchIndexAsset.SourceAssetIdx);
 }
@@ -510,7 +510,7 @@ FInstancedStruct& UPoseSearchDatabase::GetMutableAnimationAssetStruct(int32 Anim
 	return AnimationAssets[AnimationAssetIndex];
 }
 
-FInstancedStruct& UPoseSearchDatabase::GetMutableAnimationAssetStruct(const FPoseSearchIndexAsset& SearchIndexAsset)
+FInstancedStruct& UPoseSearchDatabase::GetMutableAnimationAssetStruct(const UE::PoseSearch::FSearchIndexAsset& SearchIndexAsset)
 {
 	return GetMutableAnimationAssetStruct(SearchIndexAsset.SourceAssetIdx);
 }
@@ -525,7 +525,7 @@ const FPoseSearchDatabaseAnimationAssetBase* UPoseSearchDatabase::GetAnimationAs
 	return nullptr;
 }
 
-const FPoseSearchDatabaseAnimationAssetBase* UPoseSearchDatabase::GetAnimationAssetBase(const FPoseSearchIndexAsset& SearchIndexAsset) const
+const FPoseSearchDatabaseAnimationAssetBase* UPoseSearchDatabase::GetAnimationAssetBase(const UE::PoseSearch::FSearchIndexAsset& SearchIndexAsset) const
 {
 	return GetAnimationAssetBase(SearchIndexAsset.SourceAssetIdx);
 }
@@ -540,17 +540,17 @@ FPoseSearchDatabaseAnimationAssetBase* UPoseSearchDatabase::GetMutableAnimationA
 	return nullptr;
 }
 
-FPoseSearchDatabaseAnimationAssetBase* UPoseSearchDatabase::GetMutableAnimationAssetBase(const FPoseSearchIndexAsset& SearchIndexAsset)
+FPoseSearchDatabaseAnimationAssetBase* UPoseSearchDatabase::GetMutableAnimationAssetBase(const UE::PoseSearch::FSearchIndexAsset& SearchIndexAsset)
 {
 	return GetMutableAnimationAssetBase(SearchIndexAsset.SourceAssetIdx);
 }
 
-const bool UPoseSearchDatabase::IsSourceAssetLooping(const FPoseSearchIndexAsset& SearchIndexAsset) const
+const bool UPoseSearchDatabase::IsSourceAssetLooping(const UE::PoseSearch::FSearchIndexAsset& SearchIndexAsset) const
 {
 	return GetAnimationAssetBase(SearchIndexAsset.SourceAssetIdx)->IsLooping();
 }
 
-const FString UPoseSearchDatabase::GetSourceAssetName(const FPoseSearchIndexAsset& SearchIndexAsset) const
+const FString UPoseSearchDatabase::GetSourceAssetName(const UE::PoseSearch::FSearchIndexAsset& SearchIndexAsset) const
 {
 	return GetAnimationAssetBase(SearchIndexAsset.SourceAssetIdx)->GetName();
 }
@@ -653,14 +653,14 @@ void UPoseSearchDatabase::Serialize(FArchive& Ar)
 float UPoseSearchDatabase::GetRealAssetTime(int32 PoseIdx) const
 {
 	check(Schema);
-	const FPoseSearchIndexAsset& Asset = GetSearchIndex().GetAssetForPose(PoseIdx);
+	const UE::PoseSearch::FSearchIndexAsset& Asset = GetSearchIndex().GetAssetForPose(PoseIdx);
 	return Asset.GetTimeFromPoseIndex(PoseIdx, Schema->SampleRate);
 }
 
 float UPoseSearchDatabase::GetNormalizedAssetTime(int32 PoseIdx) const
 {
 	check(Schema);
-	const FPoseSearchIndexAsset& Asset = GetSearchIndex().GetAssetForPose(PoseIdx);
+	const UE::PoseSearch::FSearchIndexAsset& Asset = GetSearchIndex().GetAssetForPose(PoseIdx);
 	const bool bIsBlendSpace = AnimationAssets[Asset.SourceAssetIdx].GetPtr<FPoseSearchDatabaseBlendSpace>() != nullptr;
 
 	// sequences or anim composites
@@ -732,9 +732,9 @@ FPoseSearchCost UPoseSearchDatabase::SearchContinuingPose(UE::PoseSearch::FSearc
 #endif
 
 	// extracting notifies from the database animation asset at time SampleTime to search for UAnimNotifyState_PoseSearchOverrideContinuingPoseCostBias eventually overriding the schema ContinuingPoseCostBias
-	const FPoseSearchIndex& SearchIndex = GetSearchIndex();
+	const FSearchIndex& SearchIndex = GetSearchIndex();
 	const int32 PoseIdx = SearchContext.GetCurrentResult().PoseIdx;
-	const FPoseSearchIndexAsset& SearchIndexAsset = SearchIndex.GetAssetForPose(PoseIdx);
+	const FSearchIndexAsset& SearchIndexAsset = SearchIndex.GetAssetForPose(PoseIdx);
 	const FPoseSearchDatabaseAnimationAssetBase* DatabaseAnimationAssetBase = GetAnimationAssetStruct(SearchIndexAsset).GetPtr<FPoseSearchDatabaseAnimationAssetBase>();
 	check(DatabaseAnimationAssetBase);
 	const FAnimationAssetSampler SequenceBaseSampler(DatabaseAnimationAssetBase->GetAnimationAsset(), SearchIndexAsset.BlendParameters);
@@ -782,7 +782,7 @@ UE::PoseSearch::FSearchResult UPoseSearchDatabase::SearchPCAKDTree(UE::PoseSearc
 	FSearchResult Result;
 
 	const int32 NumDimensions = Schema->SchemaCardinality;
-	const FPoseSearchIndex& SearchIndex = GetSearchIndex();
+	const FSearchIndex& SearchIndex = GetSearchIndex();
 
 	const uint32 ClampedNumberOfPrincipalComponents = GetNumberOfPrincipalComponents();
 	const uint32 ClampedKDTreeQueryNumNeighbors = FMath::Clamp<uint32>(KDTreeQueryNumNeighbors, 1, SearchIndex.GetNumPoses());
@@ -860,13 +860,13 @@ UE::PoseSearch::FSearchResult UPoseSearchDatabase::SearchPCAKDTree(UE::PoseSearc
 		SearchIndex.KDTree.FindNeighbors(ResultSet, ProjectedQueryValues.data());
 
 		// NonSelectableIdx are already filtered out inside the kdtree search
-		const FPoseFilters PoseFilters(Schema, TConstArrayView<size_t>(), SearchIndex.bAnyBlockTransition);
+		const FSearchFilters SearchFilters(Schema, TConstArrayView<size_t>(), SearchIndex.bAnyBlockTransition);
 		for (size_t ResultIndex = 0; ResultIndex < ResultSet.Num(); ++ResultIndex)
 		{
 			const int32 PoseIdx = ResultIndexes[ResultIndex];
 			const TConstArrayView<float> PoseValues = SearchIndex.Values.IsEmpty() ? SearchIndex.GetReconstructedPoseValues(PoseIdx, ReconstructedPoseValuesBuffer) : SearchIndex.GetPoseValues(PoseIdx);
 
-			if (PoseFilters.AreFiltersValid(SearchIndex, PoseValues, QueryValues, PoseIdx, SearchIndex.PoseMetadata[PoseIdx]
+			if (SearchFilters.AreFiltersValid(SearchIndex, PoseValues, QueryValues, PoseIdx, SearchIndex.PoseMetadata[PoseIdx]
 #if UE_POSE_SEARCH_TRACE_ENABLED
 				, SearchContext, this
 #endif
@@ -913,7 +913,7 @@ UE::PoseSearch::FSearchResult UPoseSearchDatabase::SearchBruteForce(UE::PoseSear
 	
 	FSearchResult Result;
 
-	const FPoseSearchIndex& SearchIndex = GetSearchIndex();
+	const FSearchIndex& SearchIndex = GetSearchIndex();
 
 	// since any PoseCost calculated here is at least SearchIndex.MinCostAddend,
 	// there's no point in performing the search if CurrentBestTotalCost is already better than that
@@ -931,11 +931,11 @@ UE::PoseSearch::FSearchResult UPoseSearchDatabase::SearchBruteForce(UE::PoseSear
 
 		const int32 NumDimensions = Schema->SchemaCardinality;
 		TArrayView<float> ReconstructedPoseValuesBuffer((float*)FMemory_Alloca(NumDimensions * sizeof(float)), NumDimensions);
-		const FPoseFilters PoseFilters(Schema, NonSelectableIdx, SearchIndex.bAnyBlockTransition);
+		const FSearchFilters SearchFilters(Schema, NonSelectableIdx, SearchIndex.bAnyBlockTransition);
 		for (int32 PoseIdx = 0; PoseIdx < SearchIndex.GetNumPoses(); ++PoseIdx)
 		{
 			const TConstArrayView<float> PoseValues = SearchIndex.Values.IsEmpty() ? SearchIndex.GetReconstructedPoseValues(PoseIdx, ReconstructedPoseValuesBuffer) : SearchIndex.GetPoseValues(PoseIdx);
-			if (PoseFilters.AreFiltersValid(SearchIndex, PoseValues, QueryValues, PoseIdx, SearchIndex.PoseMetadata[PoseIdx]
+			if (SearchFilters.AreFiltersValid(SearchIndex, PoseValues, QueryValues, PoseIdx, SearchIndex.PoseMetadata[PoseIdx]
 #if UE_POSE_SEARCH_TRACE_ENABLED
 				, SearchContext, this
 #endif
diff --git a/Engine/Plugins/Experimental/Animation/PoseSearch/Source/Runtime/Private/PoseSearchDatabaseIndexingContext.cpp b/Engine/Plugins/Experimental/Animation/PoseSearch/Source/Runtime/Private/PoseSearchDatabaseIndexingContext.cpp
--- a/Engine/Plugins/Experimental/Animation/PoseSearch/Source/Runtime/Private/PoseSearchDatabaseIndexingContext.cpp
+++ b/Engine/Plugins/Experimental/Animation/PoseSearch/Source/Runtime/Private/PoseSearchDatabaseIndexingContext.cpp
@@ -42,7 +42,7 @@ FTransform FAssetSamplingContext::MirrorTransform(const FTransform& InTransform)
 
 //////////////////////////////////////////////////////////////////////////
 // FDatabaseIndexingContext
-bool FDatabaseIndexingContext::IndexDatabase(FPoseSearchIndexBase& SearchIndexBase, const UPoseSearchDatabase& Database, UE::DerivedData::FRequestOwner& Owner)
+bool FDatabaseIndexingContext::IndexDatabase(FSearchIndexBase& SearchIndexBase, const UPoseSearchDatabase& Database, UE::DerivedData::FRequestOwner& Owner)
 {
 	const UPoseSearchSchema* Schema = Database.Schema;
 	check(Schema);
@@ -111,7 +111,7 @@ bool FDatabaseIndexingContext::IndexDatabase(FPoseSearchIndexBase& SearchIndexBa
 	int32 TotalPoses = 0;
 	for (int32 AssetIdx = 0; AssetIdx != SearchIndexBase.Assets.Num(); ++AssetIdx)
 	{
-		FPoseSearchIndexAsset& SearchIndexAsset = SearchIndexBase.Assets[AssetIdx];
+		FSearchIndexAsset& SearchIndexAsset = SearchIndexBase.Assets[AssetIdx];
 		check(SearchIndexAsset.FirstPoseIdx == TotalPoses);
 
 		const FPoseSearchDatabaseAnimationAssetBase* DatabaseAnimationAssetBase = Database.GetAnimationAssetStruct(SearchIndexAsset).GetPtr<FPoseSearchDatabaseAnimationAssetBase>();
@@ -155,7 +155,7 @@ bool FDatabaseIndexingContext::IndexDatabase(FPoseSearchIndexBase& SearchIndexBa
 
 	// Joining Metadata.Flags into OverallFlags
 	SearchIndexBase.bAnyBlockTransition = false;
-	for (const FPoseSearchPoseMetadata& Metadata : SearchIndexBase.PoseMetadata)
+	for (const FPoseMetadata& Metadata : SearchIndexBase.PoseMetadata)
 	{
 		if (Metadata.IsBlockTransition())
 		{
@@ -166,7 +166,7 @@ bool FDatabaseIndexingContext::IndexDatabase(FPoseSearchIndexBase& SearchIndexBa
 
 	// Joining Stats
 	int32 NumAccumulatedSamples = 0;
-	SearchIndexBase.Stats = FPoseSearchStats();
+	SearchIndexBase.Stats = FSearchStats();
 	for (int32 AssetIdx = 0; AssetIdx != SearchIndexBase.Assets.Num(); ++AssetIdx)
 	{
 		const FAssetIndexer::FStats& Stats = Indexers[AssetIdx].GetStats();
@@ -190,7 +190,7 @@ bool FDatabaseIndexingContext::IndexDatabase(FPoseSearchIndexBase& SearchIndexBa
 	if (!SearchIndexBase.PoseMetadata.IsEmpty())
 	{
 		SearchIndexBase.MinCostAddend = MAX_FLT;
-		for (const FPoseSearchPoseMetadata& PoseMetadata : SearchIndexBase.PoseMetadata)
+		for (const FPoseMetadata& PoseMetadata : SearchIndexBase.PoseMetadata)
 		{
 			if (PoseMetadata.GetCostAddend() < SearchIndexBase.MinCostAddend)
 			{
diff --git a/Engine/Plugins/Experimental/Animation/PoseSearch/Source/Runtime/Private/PoseSearchDatabaseIndexingContext.h b/Engine/Plugins/Experimental/Animation/PoseSearch/Source/Runtime/Private/PoseSearchDatabaseIndexingContext.h
--- a/Engine/Plugins/Experimental/Animation/PoseSearch/Source/Runtime/Private/PoseSearchDatabaseIndexingContext.h
+++ b/Engine/Plugins/Experimental/Animation/PoseSearch/Source/Runtime/Private/PoseSearchDatabaseIndexingContext.h
@@ -7,12 +7,12 @@
 #include "PoseSearch/PoseSearchAssetIndexer.h"
 #include "PoseSearch/PoseSearchAssetSampler.h"
 
-struct FPoseSearchIndexBase;
 class UPoseSearchDatabase;
 namespace UE::DerivedData { class FRequestOwner; }
 
 namespace UE::PoseSearch
 {
+struct FSearchIndexBase;
 
 struct FAssetSamplingContext
 {
@@ -35,7 +35,7 @@ struct FAssetSamplingContext
 
 struct FDatabaseIndexingContext
 {
-	bool IndexDatabase(FPoseSearchIndexBase& SearchIndexBase, const UPoseSearchDatabase& Database, UE::DerivedData::FRequestOwner& Owner);
+	bool IndexDatabase(FSearchIndexBase& SearchIndexBase, const UPoseSearchDatabase& Database, UE::DerivedData::FRequestOwner& Owner);
 
 private:
 	FAssetSamplingContext SamplingContext;
diff --git a/Engine/Plugins/Experimental/Animation/PoseSearch/Source/Runtime/Private/PoseSearchDerivedData.cpp b/Engine/Plugins/Experimental/Animation/PoseSearch/Source/Runtime/Private/PoseSearchDerivedData.cpp
--- a/Engine/Plugins/Experimental/Animation/PoseSearch/Source/Runtime/Private/PoseSearchDerivedData.cpp
+++ b/Engine/Plugins/Experimental/Animation/PoseSearch/Source/Runtime/Private/PoseSearchDerivedData.cpp
@@ -108,7 +108,7 @@ private:
 
 	// given an array of channels that can be normalized together (Entries), with the same cardinality (Entries[0].Channel->GetChannelCardinality()),
 	// it'll calculate the mean deviation of the associated data (from SearchIndexBases)
-	static float CalculateEntriesMeanDeviation(const FEntries& Entries, TConstArrayView<FPoseSearchIndexBase> SearchIndexBases, TConstArrayView<const UPoseSearchSchema*> Schemas)
+	static float CalculateEntriesMeanDeviation(const FEntries& Entries, TConstArrayView<FSearchIndexBase> SearchIndexBases, TConstArrayView<const UPoseSearchSchema*> Schemas)
 	{
 		check(Schemas.Num() == SearchIndexBases.Num());
 	
@@ -134,7 +134,7 @@ private:
 			const int32 DataSetIdx = Entry.SchemaIndex;
 
 			const UPoseSearchSchema* Schema = Schemas[DataSetIdx];
-			const FPoseSearchIndexBase& SearchIndex = SearchIndexBases[DataSetIdx];
+			const FSearchIndexBase& SearchIndex = SearchIndexBases[DataSetIdx];
 
 			const int32 NumPoses = SearchIndex.GetNumPoses();
 
@@ -160,7 +160,7 @@ public:
 	// it returns an array of dimension Schemas[0]->SchemaCardinality containing the mean deviation calculated from the data passed in with SearchIndexBases following the layout described in the schemas channels:
 	// channels from all the schemas get collected in groups that can be normalized together (FEntriesGroup, populated in AnalyzeSchemas) and then those homogeneous (in cardinality and meaning) groups get processed 
 	// one by one in CalculateEntriesMeanDeviation to extract the group mean deviation against the input data contained in SearchIndexBases
-	static TArray<float> Calculate(TConstArrayView<FPoseSearchIndexBase> SearchIndexBases, TConstArrayView<const UPoseSearchSchema*> Schemas)
+	static TArray<float> Calculate(TConstArrayView<FSearchIndexBase> SearchIndexBases, TConstArrayView<const UPoseSearchSchema*> Schemas)
 	{
 		// This method performs a modified z-score normalization where features are normalized
 		// by mean absolute deviation rather than standard deviation. Both methods are preferable
@@ -269,7 +269,7 @@ static void FindValidSequenceIntervals(const UAnimSequenceBase* SequenceBase, FF
 	}
 }
 
-static void InitSearchIndexAssets(FPoseSearchIndexBase& SearchIndex, const UPoseSearchDatabase* Database)
+static void InitSearchIndexAssets(FSearchIndexBase& SearchIndex, const UPoseSearchDatabase* Database)
 {
 	using namespace UE::PoseSearch;
 
@@ -314,7 +314,7 @@ static void InitSearchIndexAssets(FPoseSearchIndexBase& SearchIndex, const UPose
 						{
 							if (bAddUnmirrored)
 							{
-								const FPoseSearchIndexAsset PoseSearchIndexAsset(AnimationAssetIndex, TotalPoses, false, FFloatInterval(0.f, PlayLength), SchemaSampleRate, PermutationIdx, BlendParameters);
+								const FSearchIndexAsset PoseSearchIndexAsset(AnimationAssetIndex, TotalPoses, false, FFloatInterval(0.f, PlayLength), SchemaSampleRate, PermutationIdx, BlendParameters);
 								if (PoseSearchIndexAsset.GetNumPoses() > 0)
 								{
 									SearchIndex.Assets.Add(PoseSearchIndexAsset);
@@ -324,7 +324,7 @@ static void InitSearchIndexAssets(FPoseSearchIndexBase& SearchIndex, const UPose
 
 							if (bAddMirrored)
 							{
-								const FPoseSearchIndexAsset PoseSearchIndexAsset(AnimationAssetIndex, TotalPoses, true, FFloatInterval(0.f, PlayLength), SchemaSampleRate, PermutationIdx, BlendParameters);
+								const FSearchIndexAsset PoseSearchIndexAsset(AnimationAssetIndex, TotalPoses, true, FFloatInterval(0.f, PlayLength), SchemaSampleRate, PermutationIdx, BlendParameters);
 								if (PoseSearchIndexAsset.GetNumPoses() > 0)
 								{
 									SearchIndex.Assets.Add(PoseSearchIndexAsset);
@@ -347,7 +347,7 @@ static void InitSearchIndexAssets(FPoseSearchIndexBase& SearchIndex, const UPose
 					{
 						if (bAddUnmirrored)
 						{
-							const FPoseSearchIndexAsset PoseSearchIndexAsset(AnimationAssetIndex, TotalPoses, false, FFloatInterval(Range.GetLowerBoundValue(), Range.GetUpperBoundValue()), SchemaSampleRate, PermutationIdx);
+							const FSearchIndexAsset PoseSearchIndexAsset(AnimationAssetIndex, TotalPoses, false, FFloatInterval(Range.GetLowerBoundValue(), Range.GetUpperBoundValue()), SchemaSampleRate, PermutationIdx);
 							if (PoseSearchIndexAsset.GetNumPoses() > 0)
 							{
 								SearchIndex.Assets.Add(PoseSearchIndexAsset);
@@ -357,7 +357,7 @@ static void InitSearchIndexAssets(FPoseSearchIndexBase& SearchIndex, const UPose
 
 						if (bAddMirrored)
 						{
-							const FPoseSearchIndexAsset PoseSearchIndexAsset(AnimationAssetIndex, TotalPoses, true, FFloatInterval(Range.GetLowerBoundValue(), Range.GetUpperBoundValue()), SchemaSampleRate, PermutationIdx);
+							const FSearchIndexAsset PoseSearchIndexAsset(AnimationAssetIndex, TotalPoses, true, FFloatInterval(Range.GetLowerBoundValue(), Range.GetUpperBoundValue()), SchemaSampleRate, PermutationIdx);
 							if (PoseSearchIndexAsset.GetNumPoses() > 0)
 							{
 								SearchIndex.Assets.Add(PoseSearchIndexAsset);
@@ -375,7 +375,7 @@ static void InitSearchIndexAssets(FPoseSearchIndexBase& SearchIndex, const UPose
 	}
 }
 
-static void PreprocessSearchIndexWeights(FPoseSearchIndex& SearchIndex, const UPoseSearchSchema* Schema, TConstArrayView<float> Deviation)
+static void PreprocessSearchIndexWeights(FSearchIndex& SearchIndex, const UPoseSearchSchema* Schema, TConstArrayView<float> Deviation)
 {
 	const int32 NumDimensions = Schema->SchemaCardinality;
 	SearchIndex.WeightsSqrt.Init(1.f, NumDimensions);
@@ -414,7 +414,7 @@ static void PreprocessSearchIndexWeights(FPoseSearchIndex& SearchIndex, const UP
 }
 
 // it calculates Mean, PCAValues, and PCAProjectionMatrix
-static void PreprocessSearchIndexPCAData(FPoseSearchIndex& SearchIndex, int32 NumDimensions, uint32 NumberOfPrincipalComponents, EPoseSearchMode PoseSearchMode)
+static void PreprocessSearchIndexPCAData(FSearchIndex& SearchIndex, int32 NumDimensions, uint32 NumberOfPrincipalComponents, EPoseSearchMode PoseSearchMode)
 {
 	// binding SearchIndex.Values and SearchIndex.PCAValues Eigen row major matrix maps
 	const int32 NumPoses = SearchIndex.GetNumPoses();
@@ -525,7 +525,7 @@ static void PreprocessSearchIndexPCAData(FPoseSearchIndex& SearchIndex, int32 Nu
 	}
 }
 
-static void PreprocessSearchIndexKDTree(FPoseSearchIndex& SearchIndex, int32 NumDimensions, uint32 NumberOfPrincipalComponents, EPoseSearchMode PoseSearchMode, int32 KDTreeMaxLeafSize, int32 KDTreeQueryNumNeighbors)
+static void PreprocessSearchIndexKDTree(FSearchIndex& SearchIndex, int32 NumDimensions, uint32 NumberOfPrincipalComponents, EPoseSearchMode PoseSearchMode, int32 KDTreeMaxLeafSize, int32 KDTreeQueryNumNeighbors)
 {
 	SearchIndex.KDTree.Reset();
 	if (PoseSearchMode != EPoseSearchMode::BruteForce && NumDimensions > 0)
@@ -644,7 +644,7 @@ private:
 	void SetState(EState State) { ThreadSafeState.Set(int32(State)); }
 
 	TWeakObjectPtr<UPoseSearchDatabase> Database;
-	FPoseSearchIndex SearchIndex;
+	FSearchIndex SearchIndex;
 	UE::DerivedData::FRequestOwner Owner;
 	FIoHash DerivedDataKey = FIoHash::Zero;
 	TSet<TWeakObjectPtr<const UObject>> DatabaseDependencies;
@@ -799,7 +799,7 @@ void FPoseSearchDatabaseAsyncCacheTask::Wait(FCriticalSection& OuterMutex)
 	const bool bFailedIndexing = SearchIndex.IsEmpty();
 	if (!bFailedIndexing)
 	{
-		Database->SetSearchIndex(SearchIndex); // @todo: implement FPoseSearchIndex move ctor and assignment operator and use a MoveTemp(SearchIndex) here
+		Database->SetSearchIndex(SearchIndex); // @todo: implement FSearchIndex move ctor and assignment operator and use a MoveTemp(SearchIndex) here
 
 		check(Database->Schema && Database->Schema->IsValid() && !SearchIndex.IsEmpty() && SearchIndex.WeightsSqrt.Num() == Database->Schema->SchemaCardinality);
 
@@ -883,7 +883,7 @@ void FPoseSearchDatabaseAsyncCacheTask::OnGetComplete(UE::DerivedData::FCacheGet
 			{
 				COOK_STAT(auto Timer = UsageStats.TimeSyncWork());
 
-				// collecting all the databases that need to be built to gather their FPoseSearchIndexBase
+				// collecting all the databases that need to be built to gather their FSearchIndexBase
 				TArray<TObjectPtr<const UPoseSearchDatabase>> IndexBaseDatabases;
 				IndexBaseDatabases.Add(Database.Get()); // the first one is always this Database
 				if (Database->NormalizationSet)
@@ -892,7 +892,7 @@ void FPoseSearchDatabaseAsyncCacheTask::OnGetComplete(UE::DerivedData::FCacheGet
 				}
 
 				// @todo: DDC or parallelize this code
-				TArray<FPoseSearchIndexBase> SearchIndexBases;
+				TArray<FSearchIndexBase> SearchIndexBases;
 				TArray<const UPoseSearchSchema*> Schemas;
 				SearchIndexBases.AddDefaulted(IndexBaseDatabases.Num());
 				Schemas.AddDefaulted(IndexBaseDatabases.Num());
@@ -900,7 +900,7 @@ void FPoseSearchDatabaseAsyncCacheTask::OnGetComplete(UE::DerivedData::FCacheGet
 				{
 					const UPoseSearchDatabase* IndexBaseDatabase = IndexBaseDatabases[IndexBaseIdx].Get();
 					check(IndexBaseDatabase);
-					FPoseSearchIndexBase& SearchIndexBase = SearchIndexBases[IndexBaseIdx];
+					FSearchIndexBase& SearchIndexBase = SearchIndexBases[IndexBaseIdx];
 					Schemas[IndexBaseIdx] = IndexBaseDatabase->Schema;
 
 					// early out for invalid indexing conditions
@@ -974,11 +974,11 @@ void FPoseSearchDatabaseAsyncCacheTask::OnGetComplete(UE::DerivedData::FCacheGet
 					}
 				}
 
-				static_cast<FPoseSearchIndexBase&>(SearchIndex) = SearchIndexBases[0];
+				static_cast<FSearchIndexBase&>(SearchIndex) = SearchIndexBases[0];
 				
 				TArray<float> Deviation = FMeanDeviationCalculator::Calculate(SearchIndexBases, Schemas);
 
-				// Building FPoseSearchIndex
+				// Building FSearchIndex
 				PreprocessSearchIndexWeights(SearchIndex, Database->Schema, Deviation);
 				if (Owner.IsCanceled())
 				{
@@ -1003,7 +1003,7 @@ void FPoseSearchDatabaseAsyncCacheTask::OnGetComplete(UE::DerivedData::FCacheGet
 					return;
 				}
 
-				// removing SearchIndex.Values and relying on FPoseSearchIndex::GetReconstructedPoseValues to reconstruct the Values data from the PCAValues
+				// removing SearchIndex.Values and relying on FSearchIndex::GetReconstructedPoseValues to reconstruct the Values data from the PCAValues
 				if (Database->PoseSearchMode == EPoseSearchMode::PCAKDTree && Database->KDTreeQueryNumNeighbors <= 1)
 				{
 					SearchIndex.Values.Reset();
diff --git a/Engine/Plugins/Experimental/Animation/PoseSearch/Source/Runtime/Private/PoseSearchFeatureChannel_FilterCrashingLegs.cpp b/Engine/Plugins/Experimental/Animation/PoseSearch/Source/Runtime/Private/PoseSearchFeatureChannel_FilterCrashingLegs.cpp
--- a/Engine/Plugins/Experimental/Animation/PoseSearch/Source/Runtime/Private/PoseSearchFeatureChannel_FilterCrashingLegs.cpp
+++ b/Engine/Plugins/Experimental/Animation/PoseSearch/Source/Runtime/Private/PoseSearchFeatureChannel_FilterCrashingLegs.cpp
@@ -115,13 +115,13 @@ void UPoseSearchFeatureChannel_FilterCrashingLegs::DebugDraw(const UE::PoseSearc
 }
 #endif // ENABLE_DRAW_DEBUG
 
-// IPoseFilter interface
-bool UPoseSearchFeatureChannel_FilterCrashingLegs::IsPoseFilterActive() const
+// IPoseSearchFilter interface
+bool UPoseSearchFeatureChannel_FilterCrashingLegs::IsFilterActive() const
 {
 	return true;
 }
 
-bool UPoseSearchFeatureChannel_FilterCrashingLegs::IsPoseValid(TConstArrayView<float> PoseValues, TConstArrayView<float> QueryValues, int32 PoseIdx, const FPoseSearchPoseMetadata& Metadata) const
+bool UPoseSearchFeatureChannel_FilterCrashingLegs::IsFilterValid(TConstArrayView<float> PoseValues, TConstArrayView<float> QueryValues, int32 PoseIdx, const UE::PoseSearch::FPoseMetadata& Metadata) const
 {
 	using namespace UE::PoseSearch;
 
@@ -138,7 +138,7 @@ bool UPoseSearchFeatureChannel_FilterCrashingLegs::IsPoseValid(TConstArrayView<f
 }
 
 #if WITH_EDITOR
-void UPoseSearchFeatureChannel_FilterCrashingLegs::FillWeights(TArray<float>& Weights) const
+void UPoseSearchFeatureChannel_FilterCrashingLegs::FillWeights(TArrayView<float> Weights) const
 {
 	for (int32 i = 0; i < ChannelCardinality; ++i)
 	{
diff --git a/Engine/Plugins/Experimental/Animation/PoseSearch/Source/Runtime/Private/PoseSearchFeatureChannel_FilterCrashingLegs.h b/Engine/Plugins/Experimental/Animation/PoseSearch/Source/Runtime/Private/PoseSearchFeatureChannel_FilterCrashingLegs.h
--- a/Engine/Plugins/Experimental/Animation/PoseSearch/Source/Runtime/Private/PoseSearchFeatureChannel_FilterCrashingLegs.h
+++ b/Engine/Plugins/Experimental/Animation/PoseSearch/Source/Runtime/Private/PoseSearchFeatureChannel_FilterCrashingLegs.h
@@ -12,7 +12,7 @@
 // 0.5 if the right foot is exactly in front of the left foot (for example when a character is running  following a line)
 // -0.5 if the left foot is exactly in front of the right foot
 // close to 1 or -1 if the feet (and so the legs) are completely crossed
-// at runtime we'll match the CrashingLegsValue and also filter by discarding pose candidates that don't respect the 'AllowedTolerance' between query and database values (happening in IsPoseValid)
+// at runtime we'll match the CrashingLegsValue and also filter by discarding pose candidates that don't respect the 'AllowedTolerance' between query and database values (happening in IsFilterValid)
 UCLASS(BlueprintType, EditInlineNew, meta = (DisplayName = "CrashingLegs Channel"), CollapseCategories)
 class POSESEARCH_API UPoseSearchFeatureChannel_FilterCrashingLegs : public UPoseSearchFeatureChannel
 {
@@ -63,12 +63,12 @@ public:
 #endif // ENABLE_DRAW_DEBUG
 
 #if WITH_EDITOR
-	virtual void FillWeights(TArray<float>& Weights) const override;
+	virtual void FillWeights(TArrayView<float> Weights) const override;
 	virtual void IndexAsset(UE::PoseSearch::FAssetIndexer& Indexer) const override;
 	virtual FString GetLabel() const override;
 #endif
 
-	// IPoseFilter interface
-	virtual bool IsPoseFilterActive() const override;
-	virtual bool IsPoseValid(TConstArrayView<float> PoseValues, TConstArrayView<float> QueryValues, int32 PoseIdx, const FPoseSearchPoseMetadata& Metadata) const override;
+	// IPoseSearchFilter interface
+	virtual bool IsFilterActive() const override;
+	virtual bool IsFilterValid(TConstArrayView<float> PoseValues, TConstArrayView<float> QueryValues, int32 PoseIdx, const UE::PoseSearch::FPoseMetadata& Metadata) const override;
 };
diff --git a/Engine/Plugins/Experimental/Animation/PoseSearch/Source/Runtime/Private/PoseSearchFeatureChannel_Group.cpp b/Engine/Plugins/Experimental/Animation/PoseSearch/Source/Runtime/Private/PoseSearchFeatureChannel_Group.cpp
--- a/Engine/Plugins/Experimental/Animation/PoseSearch/Source/Runtime/Private/PoseSearchFeatureChannel_Group.cpp
+++ b/Engine/Plugins/Experimental/Animation/PoseSearch/Source/Runtime/Private/PoseSearchFeatureChannel_Group.cpp
@@ -28,7 +28,7 @@ void UPoseSearchFeatureChannel_GroupBase::AddDependentChannels(UPoseSearchSchema
 }
 
 #if WITH_EDITOR
-void UPoseSearchFeatureChannel_GroupBase::FillWeights(TArray<float>& Weights) const
+void UPoseSearchFeatureChannel_GroupBase::FillWeights(TArrayView<float> Weights) const
 {
 	for (const TObjectPtr<UPoseSearchFeatureChannel>& SubChannelPtr : GetSubChannels())
 	{
@@ -75,14 +75,14 @@ void UPoseSearchFeatureChannel_GroupBase::DebugDraw(const UE::PoseSearch::FDebug
 }
 #endif // ENABLE_DRAW_DEBUG
 
-// IPoseFilter interface
-bool UPoseSearchFeatureChannel_GroupBase::IsPoseFilterActive() const
+// IPoseSearchFilter interface
+bool UPoseSearchFeatureChannel_GroupBase::IsFilterActive() const
 {
 	for (const TObjectPtr<UPoseSearchFeatureChannel>& SubChannelPtr : GetSubChannels())
 	{
 		if (const UPoseSearchFeatureChannel* SubChannel = SubChannelPtr.Get())
 		{
-			if (SubChannel->IsPoseFilterActive())
+			if (SubChannel->IsFilterActive())
 			{
 				return true;
 			}
@@ -92,13 +92,13 @@ bool UPoseSearchFeatureChannel_GroupBase::IsPoseFilterActive() const
 	return false;
 }
 
-bool UPoseSearchFeatureChannel_GroupBase::IsPoseValid(TConstArrayView<float> PoseValues, TConstArrayView<float> QueryValues, int32 PoseIdx, const FPoseSearchPoseMetadata& Metadata) const
+bool UPoseSearchFeatureChannel_GroupBase::IsFilterValid(TConstArrayView<float> PoseValues, TConstArrayView<float> QueryValues, int32 PoseIdx, const UE::PoseSearch::FPoseMetadata& Metadata) const
 {
 	for (const TObjectPtr<UPoseSearchFeatureChannel>& SubChannelPtr : GetSubChannels())
 	{
 		if (const UPoseSearchFeatureChannel* SubChannel = SubChannelPtr.Get())
 		{
-			if (!SubChannel->IsPoseValid(PoseValues, QueryValues, PoseIdx, Metadata))
+			if (!SubChannel->IsFilterValid(PoseValues, QueryValues, PoseIdx, Metadata))
 			{
 				return false;
 			}
diff --git a/Engine/Plugins/Experimental/Animation/PoseSearch/Source/Runtime/Private/PoseSearchFeatureChannel_Group.h b/Engine/Plugins/Experimental/Animation/PoseSearch/Source/Runtime/Private/PoseSearchFeatureChannel_Group.h
--- a/Engine/Plugins/Experimental/Animation/PoseSearch/Source/Runtime/Private/PoseSearchFeatureChannel_Group.h
+++ b/Engine/Plugins/Experimental/Animation/PoseSearch/Source/Runtime/Private/PoseSearchFeatureChannel_Group.h
@@ -23,13 +23,13 @@ public:
 #endif // ENABLE_DRAW_DEBUG
 
 #if WITH_EDITOR
-	virtual void FillWeights(TArray<float>& Weights) const override;
+	virtual void FillWeights(TArrayView<float> Weights) const override;
 	virtual void IndexAsset(UE::PoseSearch::FAssetIndexer& Indexer) const override;
 #endif //WITH_EDITOR
 
-	// IPoseFilter interface
-	virtual bool IsPoseFilterActive() const override;
-	virtual bool IsPoseValid(TConstArrayView<float> PoseValues, TConstArrayView<float> QueryValues, int32 PoseIdx, const FPoseSearchPoseMetadata& Metadata) const override;
+	// IPoseSearchFilter interface
+	virtual bool IsFilterActive() const override;
+	virtual bool IsFilterValid(TConstArrayView<float> PoseValues, TConstArrayView<float> QueryValues, int32 PoseIdx, const UE::PoseSearch::FPoseMetadata& Metadata) const override;
 };
 
 UCLASS(BlueprintType, EditInlineNew, meta = (DisplayName = "Group Channel"), CollapseCategories)
diff --git a/Engine/Plugins/Experimental/Animation/PoseSearch/Source/Runtime/Private/PoseSearchFeatureChannel_Heading.cpp b/Engine/Plugins/Experimental/Animation/PoseSearch/Source/Runtime/Private/PoseSearchFeatureChannel_Heading.cpp
--- a/Engine/Plugins/Experimental/Animation/PoseSearch/Source/Runtime/Private/PoseSearchFeatureChannel_Heading.cpp
+++ b/Engine/Plugins/Experimental/Animation/PoseSearch/Source/Runtime/Private/PoseSearchFeatureChannel_Heading.cpp
@@ -80,7 +80,7 @@ void UPoseSearchFeatureChannel_Heading::DebugDraw(const UE::PoseSearch::FDebugDr
 #endif // ENABLE_DRAW_DEBUG
 
 #if WITH_EDITOR
-void UPoseSearchFeatureChannel_Heading::FillWeights(TArray<float>& Weights) const
+void UPoseSearchFeatureChannel_Heading::FillWeights(TArrayView<float> Weights) const
 {
 	for (int32 i = 0; i < ChannelCardinality; ++i)
 	{
diff --git a/Engine/Plugins/Experimental/Animation/PoseSearch/Source/Runtime/Private/PoseSearchFeatureChannel_Heading.h b/Engine/Plugins/Experimental/Animation/PoseSearch/Source/Runtime/Private/PoseSearchFeatureChannel_Heading.h
--- a/Engine/Plugins/Experimental/Animation/PoseSearch/Source/Runtime/Private/PoseSearchFeatureChannel_Heading.h
+++ b/Engine/Plugins/Experimental/Animation/PoseSearch/Source/Runtime/Private/PoseSearchFeatureChannel_Heading.h
@@ -58,7 +58,7 @@ public:
 #endif // ENABLE_DRAW_DEBUG
 
 #if WITH_EDITOR
-	virtual void FillWeights(TArray<float>& Weights) const override;
+	virtual void FillWeights(TArrayView<float> Weights) const override;
 	virtual void IndexAsset(UE::PoseSearch::FAssetIndexer& Indexer) const override;
 	virtual FString GetLabel() const override;
 #endif
diff --git a/Engine/Plugins/Experimental/Animation/PoseSearch/Source/Runtime/Private/PoseSearchFeatureChannel_Padding.cpp b/Engine/Plugins/Experimental/Animation/PoseSearch/Source/Runtime/Private/PoseSearchFeatureChannel_Padding.cpp
new file mode 100644
--- /dev/null
+++ b/Engine/Plugins/Experimental/Animation/PoseSearch/Source/Runtime/Private/PoseSearchFeatureChannel_Padding.cpp
@@ -0,0 +1,53 @@
+// Copyright Epic Games, Inc. All Rights Reserved.
+
+#include "PoseSearchFeatureChannel_Padding.h"
+#include "PoseSearch/PoseSearchAssetIndexer.h"
+#include "PoseSearch/PoseSearchContext.h"
+#include "PoseSearch/PoseSearchSchema.h"
+
+void UPoseSearchFeatureChannel_Padding::AddToSchema(UPoseSearchSchema* Schema, int32 PaddingSize)
+{
+	UPoseSearchFeatureChannel_Padding* Padding = NewObject<UPoseSearchFeatureChannel_Padding>(Schema, NAME_None, RF_Transient);
+	Padding->PaddingSize = PaddingSize;
+	Schema->AddTemporaryChannel(Padding);
+}
+
+void UPoseSearchFeatureChannel_Padding::Finalize(UPoseSearchSchema* Schema)
+{
+	ChannelDataOffset = Schema->SchemaCardinality;
+	ChannelCardinality = PaddingSize;
+	Schema->SchemaCardinality += ChannelCardinality;
+}
+
+void UPoseSearchFeatureChannel_Padding::BuildQuery(UE::PoseSearch::FSearchContext& SearchContext, UE::PoseSearch::FFeatureVectorBuilder& InOutQuery) const
+{
+	// nothing to do here
+}
+
+#if WITH_EDITOR
+void UPoseSearchFeatureChannel_Padding::FillWeights(TArrayView<float> Weights) const
+{
+	for (int32 i = 0; i < ChannelCardinality; ++i)
+	{
+		Weights[ChannelDataOffset + i] = 0.f;
+	}
+}
+
+void UPoseSearchFeatureChannel_Padding::IndexAsset(UE::PoseSearch::FAssetIndexer& Indexer) const
+{
+	UE_LOG(LogPoseSearch, Log, TEXT("UPoseSearchFeatureChannel_Padding::IndexAsset: padding data with '%d' additional floats per pose in UPoseSearchSchema '%s'"), PaddingSize, *GetNameSafe(GetSchema()));
+}
+
+FString UPoseSearchFeatureChannel_Padding::GetLabel() const
+{
+	TStringBuilder<256> Label;
+	if (const UPoseSearchFeatureChannel* OuterChannel = Cast<UPoseSearchFeatureChannel>(GetOuter()))
+	{
+		Label.Append(OuterChannel->GetLabel());
+		Label.Append(TEXT("_"));
+	}
+
+	Label.Appendf(TEXT("Pad_%d_%d"), ChannelDataOffset, ChannelDataOffset + ChannelCardinality);
+	return Label.ToString();
+}
+#endif
\ No newline at end of file
diff --git a/Engine/Plugins/Experimental/Animation/PoseSearch/Source/Runtime/Private/PoseSearchFeatureChannel_Padding.h b/Engine/Plugins/Experimental/Animation/PoseSearch/Source/Runtime/Private/PoseSearchFeatureChannel_Padding.h
new file mode 100644
--- /dev/null
+++ b/Engine/Plugins/Experimental/Animation/PoseSearch/Source/Runtime/Private/PoseSearchFeatureChannel_Padding.h
@@ -0,0 +1,28 @@
+// Copyright Epic Games, Inc. All Rights Reserved.
+
+#pragma once
+
+#include "PoseSearch/PoseSearchFeatureChannel.h"
+#include "PoseSearchFeatureChannel_Padding.generated.h"
+
+UCLASS(EditInlineNew, meta = (DisplayName = "Padding Channel"), CollapseCategories)
+class POSESEARCH_API UPoseSearchFeatureChannel_Padding : public UPoseSearchFeatureChannel
+{
+	GENERATED_BODY()
+
+public:
+	UPROPERTY(EditAnywhere, Category = "Settings")
+	int32 PaddingSize = 1;
+
+	// UPoseSearchFeatureChannel interface
+	virtual void Finalize(UPoseSearchSchema* Schema) override;
+	virtual void BuildQuery(UE::PoseSearch::FSearchContext& SearchContext, UE::PoseSearch::FFeatureVectorBuilder& InOutQuery) const override;
+
+#if WITH_EDITOR
+	virtual void FillWeights(TArrayView<float> Weights) const override;
+	virtual void IndexAsset(UE::PoseSearch::FAssetIndexer& Indexer) const override;
+	virtual FString GetLabel() const override;
+#endif
+
+	static void AddToSchema(UPoseSearchSchema* Schema, int32 PaddingSize);
+};
diff --git a/Engine/Plugins/Experimental/Animation/PoseSearch/Source/Runtime/Private/PoseSearchFeatureChannel_PermutationTime.cpp b/Engine/Plugins/Experimental/Animation/PoseSearch/Source/Runtime/Private/PoseSearchFeatureChannel_PermutationTime.cpp
--- a/Engine/Plugins/Experimental/Animation/PoseSearch/Source/Runtime/Private/PoseSearchFeatureChannel_PermutationTime.cpp
+++ b/Engine/Plugins/Experimental/Animation/PoseSearch/Source/Runtime/Private/PoseSearchFeatureChannel_PermutationTime.cpp
@@ -29,7 +29,7 @@ void UPoseSearchFeatureChannel_PermutationTime::BuildQuery(UE::PoseSearch::FSear
 }
 
 #if WITH_EDITOR
-void UPoseSearchFeatureChannel_PermutationTime::FillWeights(TArray<float>& Weights) const
+void UPoseSearchFeatureChannel_PermutationTime::FillWeights(TArrayView<float> Weights) const
 {
 	for (int32 i = 0; i < ChannelCardinality; ++i)
 	{
diff --git a/Engine/Plugins/Experimental/Animation/PoseSearch/Source/Runtime/Private/PoseSearchFeatureChannel_PermutationTime.h b/Engine/Plugins/Experimental/Animation/PoseSearch/Source/Runtime/Private/PoseSearchFeatureChannel_PermutationTime.h
--- a/Engine/Plugins/Experimental/Animation/PoseSearch/Source/Runtime/Private/PoseSearchFeatureChannel_PermutationTime.h
+++ b/Engine/Plugins/Experimental/Animation/PoseSearch/Source/Runtime/Private/PoseSearchFeatureChannel_PermutationTime.h
@@ -19,7 +19,7 @@ public:
 	virtual void BuildQuery(UE::PoseSearch::FSearchContext& SearchContext, UE::PoseSearch::FFeatureVectorBuilder& InOutQuery) const override;
 
 #if WITH_EDITOR
-	virtual void FillWeights(TArray<float>& Weights) const override;
+	virtual void FillWeights(TArrayView<float> Weights) const override;
 	virtual void IndexAsset(UE::PoseSearch::FAssetIndexer& Indexer) const override;
 	virtual FString GetLabel() const override;
 #endif
diff --git a/Engine/Plugins/Experimental/Animation/PoseSearch/Source/Runtime/Private/PoseSearchFeatureChannel_Phase.cpp b/Engine/Plugins/Experimental/Animation/PoseSearch/Source/Runtime/Private/PoseSearchFeatureChannel_Phase.cpp
--- a/Engine/Plugins/Experimental/Animation/PoseSearch/Source/Runtime/Private/PoseSearchFeatureChannel_Phase.cpp
+++ b/Engine/Plugins/Experimental/Animation/PoseSearch/Source/Runtime/Private/PoseSearchFeatureChannel_Phase.cpp
@@ -346,7 +346,7 @@ void UPoseSearchFeatureChannel_Phase::DebugDraw(const UE::PoseSearch::FDebugDraw
 #endif // ENABLE_DRAW_DEBUG
 
 #if WITH_EDITOR
-void UPoseSearchFeatureChannel_Phase::FillWeights(TArray<float>& Weights) const
+void UPoseSearchFeatureChannel_Phase::FillWeights(TArrayView<float> Weights) const
 {
 	for (int32 i = 0; i < ChannelCardinality; ++i)
 	{
diff --git a/Engine/Plugins/Experimental/Animation/PoseSearch/Source/Runtime/Private/PoseSearchFeatureChannel_Phase.h b/Engine/Plugins/Experimental/Animation/PoseSearch/Source/Runtime/Private/PoseSearchFeatureChannel_Phase.h
--- a/Engine/Plugins/Experimental/Animation/PoseSearch/Source/Runtime/Private/PoseSearchFeatureChannel_Phase.h
+++ b/Engine/Plugins/Experimental/Animation/PoseSearch/Source/Runtime/Private/PoseSearchFeatureChannel_Phase.h
@@ -38,7 +38,7 @@ public:
 #endif // ENABLE_DRAW_DEBUG
 
 #if WITH_EDITOR
-	virtual void FillWeights(TArray<float>& Weights) const override;
+	virtual void FillWeights(TArrayView<float> Weights) const override;
 	virtual void IndexAsset(UE::PoseSearch::FAssetIndexer& Indexer) const override;
 	virtual FString GetLabel() const override;
 #endif
diff --git a/Engine/Plugins/Experimental/Animation/PoseSearch/Source/Runtime/Private/PoseSearchFeatureChannel_Position.cpp b/Engine/Plugins/Experimental/Animation/PoseSearch/Source/Runtime/Private/PoseSearchFeatureChannel_Position.cpp
--- a/Engine/Plugins/Experimental/Animation/PoseSearch/Source/Runtime/Private/PoseSearchFeatureChannel_Position.cpp
+++ b/Engine/Plugins/Experimental/Animation/PoseSearch/Source/Runtime/Private/PoseSearchFeatureChannel_Position.cpp
@@ -106,7 +106,7 @@ void UPoseSearchFeatureChannel_Position::DebugDraw(const UE::PoseSearch::FDebugD
 #endif // ENABLE_DRAW_DEBUG
 
 #if WITH_EDITOR
-void UPoseSearchFeatureChannel_Position::FillWeights(TArray<float>& Weights) const
+void UPoseSearchFeatureChannel_Position::FillWeights(TArrayView<float> Weights) const
 {
 	for (int32 i = 0; i < ChannelCardinality; ++i)
 	{
diff --git a/Engine/Plugins/Experimental/Animation/PoseSearch/Source/Runtime/Private/PoseSearchFeatureChannel_Position.h b/Engine/Plugins/Experimental/Animation/PoseSearch/Source/Runtime/Private/PoseSearchFeatureChannel_Position.h
--- a/Engine/Plugins/Experimental/Animation/PoseSearch/Source/Runtime/Private/PoseSearchFeatureChannel_Position.h
+++ b/Engine/Plugins/Experimental/Animation/PoseSearch/Source/Runtime/Private/PoseSearchFeatureChannel_Position.h
@@ -54,7 +54,7 @@ public:
 #endif // ENABLE_DRAW_DEBUG
 
 #if WITH_EDITOR
-	virtual void FillWeights(TArray<float>& Weights) const override;
+	virtual void FillWeights(TArrayView<float> Weights) const override;
 	virtual void IndexAsset(UE::PoseSearch::FAssetIndexer& Indexer) const override;
 	virtual FString GetLabel() const override;
 #endif
diff --git a/Engine/Plugins/Experimental/Animation/PoseSearch/Source/Runtime/Private/PoseSearchFeatureChannel_SamplingTime.cpp b/Engine/Plugins/Experimental/Animation/PoseSearch/Source/Runtime/Private/PoseSearchFeatureChannel_SamplingTime.cpp
--- a/Engine/Plugins/Experimental/Animation/PoseSearch/Source/Runtime/Private/PoseSearchFeatureChannel_SamplingTime.cpp
+++ b/Engine/Plugins/Experimental/Animation/PoseSearch/Source/Runtime/Private/PoseSearchFeatureChannel_SamplingTime.cpp
@@ -19,7 +19,7 @@ void UPoseSearchFeatureChannel_SamplingTime::BuildQuery(UE::PoseSearch::FSearchC
 }
 
 #if WITH_EDITOR
-void UPoseSearchFeatureChannel_SamplingTime::FillWeights(TArray<float>& Weights) const
+void UPoseSearchFeatureChannel_SamplingTime::FillWeights(TArrayView<float> Weights) const
 {
 	for (int32 i = 0; i < ChannelCardinality; ++i)
 	{
diff --git a/Engine/Plugins/Experimental/Animation/PoseSearch/Source/Runtime/Private/PoseSearchFeatureChannel_SamplingTime.h b/Engine/Plugins/Experimental/Animation/PoseSearch/Source/Runtime/Private/PoseSearchFeatureChannel_SamplingTime.h
--- a/Engine/Plugins/Experimental/Animation/PoseSearch/Source/Runtime/Private/PoseSearchFeatureChannel_SamplingTime.h
+++ b/Engine/Plugins/Experimental/Animation/PoseSearch/Source/Runtime/Private/PoseSearchFeatureChannel_SamplingTime.h
@@ -23,7 +23,7 @@ public:
 	virtual void BuildQuery(UE::PoseSearch::FSearchContext& SearchContext, UE::PoseSearch::FFeatureVectorBuilder& InOutQuery) const override;
 
 #if WITH_EDITOR
-	virtual void FillWeights(TArray<float>& Weights) const override;
+	virtual void FillWeights(TArrayView<float> Weights) const override;
 	virtual void IndexAsset(UE::PoseSearch::FAssetIndexer& Indexer) const override;
 	virtual FString GetLabel() const override;
 #endif
diff --git a/Engine/Plugins/Experimental/Animation/PoseSearch/Source/Runtime/Private/PoseSearchFeatureChannel_Velocity.cpp b/Engine/Plugins/Experimental/Animation/PoseSearch/Source/Runtime/Private/PoseSearchFeatureChannel_Velocity.cpp
--- a/Engine/Plugins/Experimental/Animation/PoseSearch/Source/Runtime/Private/PoseSearchFeatureChannel_Velocity.cpp
+++ b/Engine/Plugins/Experimental/Animation/PoseSearch/Source/Runtime/Private/PoseSearchFeatureChannel_Velocity.cpp
@@ -73,7 +73,7 @@ void UPoseSearchFeatureChannel_Velocity::DebugDraw(const UE::PoseSearch::FDebugD
 #endif // ENABLE_DRAW_DEBUG
 
 #if WITH_EDITOR
-void UPoseSearchFeatureChannel_Velocity::FillWeights(TArray<float>& Weights) const
+void UPoseSearchFeatureChannel_Velocity::FillWeights(TArrayView<float> Weights) const
 {
 	for (int32 i = 0; i < ChannelCardinality; ++i)
 	{
diff --git a/Engine/Plugins/Experimental/Animation/PoseSearch/Source/Runtime/Private/PoseSearchFeatureChannel_Velocity.h b/Engine/Plugins/Experimental/Animation/PoseSearch/Source/Runtime/Private/PoseSearchFeatureChannel_Velocity.h
--- a/Engine/Plugins/Experimental/Animation/PoseSearch/Source/Runtime/Private/PoseSearchFeatureChannel_Velocity.h
+++ b/Engine/Plugins/Experimental/Animation/PoseSearch/Source/Runtime/Private/PoseSearchFeatureChannel_Velocity.h
@@ -52,7 +52,7 @@ public:
 #endif // ENABLE_DRAW_DEBUG
 
 #if WITH_EDITOR
-	virtual void FillWeights(TArray<float>& Weights) const override;
+	virtual void FillWeights(TArrayView<float> Weights) const override;
 	virtual void IndexAsset(UE::PoseSearch::FAssetIndexer& Indexer) const override;
 	virtual FString GetLabel() const override;
 #endif
diff --git a/Engine/Plugins/Experimental/Animation/PoseSearch/Source/Runtime/Private/PoseSearchIndex.cpp b/Engine/Plugins/Experimental/Animation/PoseSearch/Source/Runtime/Private/PoseSearchIndex.cpp
--- a/Engine/Plugins/Experimental/Animation/PoseSearch/Source/Runtime/Private/PoseSearchIndex.cpp
+++ b/Engine/Plugins/Experimental/Animation/PoseSearch/Source/Runtime/Private/PoseSearchIndex.cpp
@@ -29,11 +29,9 @@ void CompareFeatureVectors(TConstArrayView<float> A, TConstArrayView<float> B, T
 	VR = ((VA - VB) * VW).square();
 }
 
-} // namespace UE::PoseSearch
-
 //////////////////////////////////////////////////////////////////////////
-// FPoseSearchPoseMetadata
-FArchive& operator<<(FArchive& Ar, FPoseSearchPoseMetadata& Metadata)
+// FPoseMetadata
+FArchive& operator<<(FArchive& Ar, FPoseMetadata& Metadata)
 {
 	Ar << Metadata.Data;
 	Ar << Metadata.CostAddend;
@@ -41,8 +39,8 @@ FArchive& operator<<(FArchive& Ar, FPoseSearchPoseMetadata& Metadata)
 }
 
 //////////////////////////////////////////////////////////////////////////
-// FPoseSearchIndexAsset
-FArchive& operator<<(FArchive& Ar, FPoseSearchIndexAsset& IndexAsset)
+// FSearchIndexAsset
+FArchive& operator<<(FArchive& Ar, FSearchIndexAsset& IndexAsset)
 {
 	Ar << IndexAsset.SourceAssetIdx;
 	Ar << IndexAsset.bMirrored;
@@ -55,8 +53,8 @@ FArchive& operator<<(FArchive& Ar, FPoseSearchIndexAsset& IndexAsset)
 }
 
 //////////////////////////////////////////////////////////////////////////
-// FPoseSearchStats
-FArchive& operator<<(FArchive& Ar, FPoseSearchStats& Stats)
+// FSearchStats
+FArchive& operator<<(FArchive& Ar, FSearchStats& Stats)
 {
 	Ar << Stats.AverageSpeed;
 	Ar << Stats.MaxSpeed;
@@ -67,13 +65,13 @@ FArchive& operator<<(FArchive& Ar, FPoseSearchStats& Stats)
 
 //////////////////////////////////////////////////////////////////////////
 // FPoseSearchBaseIndex
-const FPoseSearchIndexAsset& FPoseSearchIndexBase::GetAssetForPose(int32 PoseIdx) const
+const FSearchIndexAsset& FSearchIndexBase::GetAssetForPose(int32 PoseIdx) const
 {
 	const uint32 AssetIndex = PoseMetadata[PoseIdx].GetAssetIndex();
 	return Assets[AssetIndex];
 }
 
-const FPoseSearchIndexAsset* FPoseSearchIndexBase::GetAssetForPoseSafe(int32 PoseIdx) const
+const FSearchIndexAsset* FSearchIndexBase::GetAssetForPoseSafe(int32 PoseIdx) const
 {
 	if (PoseMetadata.IsValidIndex(PoseIdx))
 	{
@@ -86,17 +84,17 @@ const FPoseSearchIndexAsset* FPoseSearchIndexBase::GetAssetForPoseSafe(int32 Pos
 	return nullptr;
 }
 
-bool FPoseSearchIndexBase::IsEmpty() const
+bool FSearchIndexBase::IsEmpty() const
 {
 	return Assets.IsEmpty() || PoseMetadata.IsEmpty();
 }
 
-void FPoseSearchIndexBase::Reset()
+void FSearchIndexBase::Reset()
 {
-	*this = FPoseSearchIndexBase();
+	*this = FSearchIndexBase();
 }
 
-FArchive& operator<<(FArchive& Ar, FPoseSearchIndexBase& Index)
+FArchive& operator<<(FArchive& Ar, FSearchIndexBase& Index)
 {
 	Ar << Index.Values;
 	Ar << Index.PoseMetadata;
@@ -108,9 +106,9 @@ FArchive& operator<<(FArchive& Ar, FPoseSearchIndexBase& Index)
 }
 
 //////////////////////////////////////////////////////////////////////////
-// FPoseSearchIndex
-FPoseSearchIndex::FPoseSearchIndex(const FPoseSearchIndex& Other)
-	: FPoseSearchIndexBase(Other)
+// FSearchIndex
+FSearchIndex::FSearchIndex(const FSearchIndex& Other)
+	: FSearchIndexBase(Other)
 	, WeightsSqrt(Other.WeightsSqrt)
 	, PCAValues(Other.PCAValues)
 	, PCAProjectionMatrix(Other.PCAProjectionMatrix)
@@ -122,23 +120,23 @@ FPoseSearchIndex::FPoseSearchIndex(const FPoseSearchIndex& Other)
 	KDTree.DataSource.Data = PCAValues.IsEmpty() ? nullptr : PCAValues.GetData();
 }
 
-FPoseSearchIndex& FPoseSearchIndex::operator=(const FPoseSearchIndex& Other)
+FSearchIndex& FSearchIndex::operator=(const FSearchIndex& Other)
 {
 	if (this != &Other)
 	{
-		this->~FPoseSearchIndex();
-		new(this)FPoseSearchIndex(Other);
+		this->~FSearchIndex();
+		new(this)FSearchIndex(Other);
 	}
 	return *this;
 }
 
-void FPoseSearchIndex::Reset()
+void FSearchIndex::Reset()
 {
-	FPoseSearchIndex Default;
+	FSearchIndex Default;
 	*this = Default;
 }
 
-TConstArrayView<float> FPoseSearchIndex::GetPoseValues(int32 PoseIdx) const
+TConstArrayView<float> FSearchIndex::GetPoseValues(int32 PoseIdx) const
 {
 	const int32 SchemaCardinality = WeightsSqrt.Num();
 	check(!Values.IsEmpty() && PoseIdx >= 0 && PoseIdx < GetNumPoses() && SchemaCardinality > 0);
@@ -146,10 +144,8 @@ TConstArrayView<float> FPoseSearchIndex::GetPoseValues(int32 PoseIdx) const
 	return MakeArrayView(&Values[ValueOffset], SchemaCardinality);
 }
 
-TConstArrayView<float> FPoseSearchIndex::GetReconstructedPoseValues(int32 PoseIdx, TArrayView<float> BufferUsedForReconstruction) const
+TConstArrayView<float> FSearchIndex::GetReconstructedPoseValues(int32 PoseIdx, TArrayView<float> BufferUsedForReconstruction) const
 {
-	using namespace UE::PoseSearch;
-
 	const int32 NumDimensions = WeightsSqrt.Num();
 	const int32 NumPoses = GetNumPoses();
 	check(PoseIdx >= 0 && PoseIdx < NumPoses&& NumDimensions > 0);
@@ -172,7 +168,7 @@ TConstArrayView<float> FPoseSearchIndex::GetReconstructedPoseValues(int32 PoseId
 	return BufferUsedForReconstruction;
 }
 
-TArray<float> FPoseSearchIndex::GetPoseValuesSafe(int32 PoseIdx) const
+TArray<float> FSearchIndex::GetPoseValuesSafe(int32 PoseIdx) const
 {
 	TArray<float> PoseValues;
 	if (PoseIdx >= 0 && PoseIdx < GetNumPoses())
@@ -191,19 +187,19 @@ TArray<float> FPoseSearchIndex::GetPoseValuesSafe(int32 PoseIdx) const
 	return PoseValues;
 }
 
-FPoseSearchCost FPoseSearchIndex::ComparePoses(int32 PoseIdx, float ContinuingPoseCostBias, TConstArrayView<float> PoseValues, TConstArrayView<float> QueryValues) const
+FPoseSearchCost FSearchIndex::ComparePoses(int32 PoseIdx, float ContinuingPoseCostBias, TConstArrayView<float> PoseValues, TConstArrayView<float> QueryValues) const
 {
 	// base dissimilarity cost representing how the associated PoseIdx differ, in a weighted way, from the query pose (QueryValues)
-	const float DissimilarityCost = UE::PoseSearch::CompareFeatureVectors(PoseValues, QueryValues, WeightsSqrt);
+	const float DissimilarityCost = CompareFeatureVectors(PoseValues, QueryValues, WeightsSqrt);
 
 	// cost addend associated to Schema->BaseCostBias or overriden by UAnimNotifyState_PoseSearchModifyCost
 	const float NotifyAddend = PoseMetadata[PoseIdx].GetCostAddend();
 	return FPoseSearchCost(DissimilarityCost, NotifyAddend, ContinuingPoseCostBias);
 }
 
-FArchive& operator<<(FArchive& Ar, FPoseSearchIndex& Index)
+FArchive& operator<<(FArchive& Ar, FSearchIndex& Index)
 {
-	Ar << static_cast<FPoseSearchIndexBase&>(Index);
+	Ar << static_cast<FSearchIndexBase&>(Index);
 
 	Ar << Index.WeightsSqrt;
 	Ar << Index.PCAValues;
@@ -213,4 +209,6 @@ FArchive& operator<<(FArchive& Ar, FPoseSearchIndex& Index)
 
 	Serialize(Ar, Index.KDTree, Index.PCAValues.GetData());
 	return Ar;
-}
\ No newline at end of file
+}
+
+} // namespace UE::PoseSearch
diff --git a/Engine/Plugins/Experimental/Animation/PoseSearch/Source/Runtime/Private/PoseSearchLibrary.cpp b/Engine/Plugins/Experimental/Animation/PoseSearch/Source/Runtime/Private/PoseSearchLibrary.cpp
--- a/Engine/Plugins/Experimental/Animation/PoseSearch/Source/Runtime/Private/PoseSearchLibrary.cpp
+++ b/Engine/Plugins/Experimental/Animation/PoseSearch/Source/Runtime/Private/PoseSearchLibrary.cpp
@@ -575,7 +575,7 @@ void UPoseSearchLibrary::MotionMatch(
 		FSearchResult SearchResult = Database->Search(SearchContext);
 		if (SearchResult.IsValid())
 		{
-			const FPoseSearchIndexAsset* SearchIndexAsset = SearchResult.GetSearchIndexAsset();
+			const FSearchIndexAsset* SearchIndexAsset = SearchResult.GetSearchIndexAsset();
 			if (const FPoseSearchDatabaseAnimationAssetBase* DatabaseAsset = SearchResult.Database->GetAnimationAssetBase(*SearchIndexAsset))
 			{
 				SelectedAnimation = DatabaseAsset->GetAnimationAsset();
diff --git a/Engine/Plugins/Experimental/Animation/PoseSearch/Source/Runtime/Private/PoseSearchResult.cpp b/Engine/Plugins/Experimental/Animation/PoseSearch/Source/Runtime/Private/PoseSearchResult.cpp
--- a/Engine/Plugins/Experimental/Animation/PoseSearch/Source/Runtime/Private/PoseSearchResult.cpp
+++ b/Engine/Plugins/Experimental/Animation/PoseSearch/Source/Runtime/Private/PoseSearchResult.cpp
@@ -35,7 +35,7 @@ void FSearchResult::Update(float NewAssetTime)
 	}
 	else
 	{
-		const FPoseSearchIndexAsset& SearchIndexAsset = Database->GetSearchIndex().GetAssetForPose(PoseIdx);
+		const FSearchIndexAsset& SearchIndexAsset = Database->GetSearchIndex().GetAssetForPose(PoseIdx);
 		const FInstancedStruct& DatabaseAsset = Database->GetAnimationAssetStruct(SearchIndexAsset);
 		if (DatabaseAsset.GetPtr<FPoseSearchDatabaseSequence>() || DatabaseAsset.GetPtr<FPoseSearchDatabaseAnimComposite>())
 		{
@@ -88,7 +88,7 @@ void FSearchResult::Reset()
 	AssetTime = 0.0f;
 }
 
-const FPoseSearchIndexAsset* FSearchResult::GetSearchIndexAsset(bool bMandatory) const
+const FSearchIndexAsset* FSearchResult::GetSearchIndexAsset(bool bMandatory) const
 {
 	if (bMandatory)
 	{
diff --git a/Engine/Plugins/Experimental/Animation/PoseSearch/Source/Runtime/Private/PoseSearchSchema.cpp b/Engine/Plugins/Experimental/Animation/PoseSearch/Source/Runtime/Private/PoseSearchSchema.cpp
--- a/Engine/Plugins/Experimental/Animation/PoseSearch/Source/Runtime/Private/PoseSearchSchema.cpp
+++ b/Engine/Plugins/Experimental/Animation/PoseSearch/Source/Runtime/Private/PoseSearchSchema.cpp
@@ -4,6 +4,7 @@
 #include "AnimationRuntime.h"
 #include "PoseSearch/PoseSearchDefines.h"
 #include "PoseSearch/PoseSearchResult.h"
+#include "PoseSearchFeatureChannel_Padding.h"
 #include "PoseSearchFeatureChannel_PermutationTime.h"
 #include "UObject/ObjectSaveContext.h"
 
@@ -134,6 +135,17 @@ void UPoseSearchSchema::Finalize()
 		}
 	}
 	
+	// adding padding if required
+	if (bAddDataPadding)
+	{
+		// calculating how many floats of padding are required to make the data 16 bytes padded
+		const int32 PaddingSize = SchemaCardinality % (16 / sizeof(float));
+		if (PaddingSize > 0)
+		{
+			UPoseSearchFeatureChannel_Padding::AddToSchema(this, PaddingSize);
+		}
+	}
+
 	// Initialize references to obtain bone indices and fill out bone index array
 	for (FBoneReference& BoneRef : BoneReferences)
 	{
diff --git a/Engine/Plugins/Experimental/Animation/PoseSearch/Source/Runtime/Public/PoseSearch/PoseSearchAssetIndexer.h b/Engine/Plugins/Experimental/Animation/PoseSearch/Source/Runtime/Public/PoseSearch/PoseSearchAssetIndexer.h
--- a/Engine/Plugins/Experimental/Animation/PoseSearch/Source/Runtime/Public/PoseSearch/PoseSearchAssetIndexer.h
+++ b/Engine/Plugins/Experimental/Animation/PoseSearch/Source/Runtime/Public/PoseSearch/PoseSearchAssetIndexer.h
@@ -10,14 +10,13 @@
 #include "PoseSearch/PoseSearchFeatureChannel.h"
 #include "PoseSearch/PoseSearchSchema.h"
 
-struct FPoseSearchIndexAsset;
-struct FPoseSearchPoseMetadata;
-
 namespace UE::PoseSearch
 {
 
 struct FAnimationAssetSampler;
 struct FAssetSamplingContext;
+struct FSearchIndexAsset;
+struct FPoseMetadata;
 
 class FAssetIndexer
 {
@@ -31,9 +30,9 @@ public:
 		float MaxAcceleration = 0.f;
 	};
 
-	FAssetIndexer(const FBoneContainer& InBoneContainer, const FPoseSearchIndexAsset& InSearchIndexAsset, 
+	FAssetIndexer(const FBoneContainer& InBoneContainer, const FSearchIndexAsset& InSearchIndexAsset, 
 		const FAssetSamplingContext& InSamplingContext, const UPoseSearchSchema& InSchema, const FAnimationAssetSampler& InAssetSampler);
-	void AssignWorkingData(TArrayView<float> InOutFeatureVectorTable, TArrayView<FPoseSearchPoseMetadata> InOutPoseMetadata);
+	void AssignWorkingData(TArrayView<float> InOutFeatureVectorTable, TArrayView<FPoseMetadata> InOutPoseMetadata);
 	void Process(int32 AssetIdx);
 	const FStats& GetStats() const { return Stats; }
 
@@ -80,13 +79,13 @@ private:
 
 	FBoneContainer BoneContainer;
 	TMap<float, CachedEntry> CachedEntries;
-	const FPoseSearchIndexAsset& SearchIndexAsset;
+	const FSearchIndexAsset& SearchIndexAsset;
 	const FAssetSamplingContext& SamplingContext;
 	const UPoseSearchSchema& Schema;
 	const FAnimationAssetSampler& AssetSampler;
 	
 	TArrayView<float> FeatureVectorTable;
-	TArrayView<FPoseSearchPoseMetadata> PoseMetadata;
+	TArrayView<FPoseMetadata> PoseMetadata;
 
 	FStats Stats;
 };
diff --git a/Engine/Plugins/Experimental/Animation/PoseSearch/Source/Runtime/Public/PoseSearch/PoseSearchContext.h b/Engine/Plugins/Experimental/Animation/PoseSearch/Source/Runtime/Public/PoseSearch/PoseSearchContext.h
--- a/Engine/Plugins/Experimental/Animation/PoseSearch/Source/Runtime/Public/PoseSearch/PoseSearchContext.h
+++ b/Engine/Plugins/Experimental/Animation/PoseSearch/Source/Runtime/Public/PoseSearch/PoseSearchContext.h
@@ -109,7 +109,7 @@ struct POSESEARCH_API FDebugDrawParams
 	FDebugDrawParams(FAnimInstanceProxy* InAnimInstanceProxy, const FTransform& InRootMotionTransform, const UPoseSearchDatabase* InDatabase, EDebugDrawFlags InFlags = EDebugDrawFlags::None);
 	FDebugDrawParams(const UWorld* InWorld, const USkinnedMeshComponent* InMesh, const FTransform& InRootMotionTransform, const UPoseSearchDatabase* InDatabase, EDebugDrawFlags InFlags = EDebugDrawFlags::None);
 
-	const FPoseSearchIndex* GetSearchIndex() const;
+	const FSearchIndex* GetSearchIndex() const;
 	const UPoseSearchSchema* GetSchema() const;
 
 	FVector ExtractPosition(TConstArrayView<float> PoseVector, const UPoseSearchFeatureChannel_Position* Position) const;
diff --git a/Engine/Plugins/Experimental/Animation/PoseSearch/Source/Runtime/Public/PoseSearch/PoseSearchDatabase.h b/Engine/Plugins/Experimental/Animation/PoseSearch/Source/Runtime/Public/PoseSearch/PoseSearchDatabase.h
--- a/Engine/Plugins/Experimental/Animation/PoseSearch/Source/Runtime/Public/PoseSearch/PoseSearchDatabase.h
+++ b/Engine/Plugins/Experimental/Animation/PoseSearch/Source/Runtime/Public/PoseSearch/PoseSearchDatabase.h
@@ -311,30 +311,29 @@ public:
 
 private:
 	// Do not use it directly. Use GetSearchIndex / SetSearchIndex interact with it and validate that is ok to do so.
-	UPROPERTY(Transient, meta = (ExcludeFromHash))
-	FPoseSearchIndex SearchIndexPrivate;
+	UE::PoseSearch::FSearchIndex SearchIndexPrivate;
 
 public:
 	virtual ~UPoseSearchDatabase();
 
-	void SetSearchIndex(const FPoseSearchIndex& SearchIndex);
-	const FPoseSearchIndex& GetSearchIndex() const;
+	void SetSearchIndex(const UE::PoseSearch::FSearchIndex& SearchIndex);
+	const UE::PoseSearch::FSearchIndex& GetSearchIndex() const;
 	
 	bool GetSkipSearchIfPossible() const;
 
-	int32 GetPoseIndexFromTime(float AssetTime, const FPoseSearchIndexAsset& SearchIndexAsset) const;
-	bool GetPoseIndicesAndLerpValueFromTime(float Time, const FPoseSearchIndexAsset& SearchIndexAsset, int32& PrevPoseIdx, int32& PoseIdx, int32& NextPoseIdx, float& LerpValue) const;
+	int32 GetPoseIndexFromTime(float AssetTime, const UE::PoseSearch::FSearchIndexAsset& SearchIndexAsset) const;
+	bool GetPoseIndicesAndLerpValueFromTime(float Time, const UE::PoseSearch::FSearchIndexAsset& SearchIndexAsset, int32& PrevPoseIdx, int32& PoseIdx, int32& NextPoseIdx, float& LerpValue) const;
 
 	const FInstancedStruct& GetAnimationAssetStruct(int32 AnimationAssetIndex) const;
-	const FInstancedStruct& GetAnimationAssetStruct(const FPoseSearchIndexAsset& SearchIndexAsset) const;
+	const FInstancedStruct& GetAnimationAssetStruct(const UE::PoseSearch::FSearchIndexAsset& SearchIndexAsset) const;
 	FInstancedStruct& GetMutableAnimationAssetStruct(int32 AnimationAssetIndex);
-	FInstancedStruct& GetMutableAnimationAssetStruct(const FPoseSearchIndexAsset& SearchIndexAsset);
+	FInstancedStruct& GetMutableAnimationAssetStruct(const UE::PoseSearch::FSearchIndexAsset& SearchIndexAsset);
 	const FPoseSearchDatabaseAnimationAssetBase* GetAnimationAssetBase(int32 AnimationAssetIndex) const;
-	const FPoseSearchDatabaseAnimationAssetBase* GetAnimationAssetBase(const FPoseSearchIndexAsset& SearchIndexAsset) const;
+	const FPoseSearchDatabaseAnimationAssetBase* GetAnimationAssetBase(const UE::PoseSearch::FSearchIndexAsset& SearchIndexAsset) const;
 	FPoseSearchDatabaseAnimationAssetBase* GetMutableAnimationAssetBase(int32 AnimationAssetIndex);
-	FPoseSearchDatabaseAnimationAssetBase* GetMutableAnimationAssetBase(const FPoseSearchIndexAsset& SearchIndexAsset);
-	const bool IsSourceAssetLooping(const FPoseSearchIndexAsset& SearchIndexAsset) const;
-	const FString GetSourceAssetName(const FPoseSearchIndexAsset& SearchIndexAsset) const;
+	FPoseSearchDatabaseAnimationAssetBase* GetMutableAnimationAssetBase(const UE::PoseSearch::FSearchIndexAsset& SearchIndexAsset);
+	const bool IsSourceAssetLooping(const UE::PoseSearch::FSearchIndexAsset& SearchIndexAsset) const;
+	const FString GetSourceAssetName(const UE::PoseSearch::FSearchIndexAsset& SearchIndexAsset) const;
 	int32 GetNumberOfPrincipalComponents() const;
 	float GetRealAssetTime(int32 PoseIdx) const;
 	float GetNormalizedAssetTime(int32 PoseIdx) const;
diff --git a/Engine/Plugins/Experimental/Animation/PoseSearch/Source/Runtime/Public/PoseSearch/PoseSearchFeatureChannel.h b/Engine/Plugins/Experimental/Animation/PoseSearch/Source/Runtime/Public/PoseSearch/PoseSearchFeatureChannel.h
--- a/Engine/Plugins/Experimental/Animation/PoseSearch/Source/Runtime/Public/PoseSearch/PoseSearchFeatureChannel.h
+++ b/Engine/Plugins/Experimental/Animation/PoseSearch/Source/Runtime/Public/PoseSearch/PoseSearchFeatureChannel.h
@@ -8,7 +8,6 @@
 #include "PoseSearchFeatureChannel.generated.h"
 
 class UPoseSearchSchema;
-struct FPoseSearchPoseMetadata;
 
 UENUM()
 enum class EComponentStrippingVector : uint8
@@ -61,6 +60,7 @@ namespace UE::PoseSearch
 struct FDebugDrawParams;
 struct FFeatureVectorBuilder;
 struct FSearchContext;
+struct FPoseMetadata;
 
 #if WITH_EDITOR
 class FAssetIndexer;
@@ -86,22 +86,22 @@ public:
 
 } // namespace UE::PoseSearch
 
-class POSESEARCH_API IPoseFilter
+class POSESEARCH_API IPoseSearchFilter
 {
 public:
-	virtual ~IPoseFilter() {}
+	virtual ~IPoseSearchFilter() {}
 
 	// if true this filter will be evaluated
-	virtual bool IsPoseFilterActive() const { return false; }
+	virtual bool IsFilterActive() const { return false; }
 	
 	// if it returns false the pose candidate will be discarded
-	virtual bool IsPoseValid(TConstArrayView<float> PoseValues, TConstArrayView<float> QueryValues, int32 PoseIdx, const FPoseSearchPoseMetadata& Metadata) const { return true; }
+	virtual bool IsFilterValid(TConstArrayView<float> PoseValues, TConstArrayView<float> QueryValues, int32 PoseIdx, const UE::PoseSearch::FPoseMetadata& Metadata) const { return true; }
 };
 
 //////////////////////////////////////////////////////////////////////////
 // Feature channels interface
 UCLASS(Abstract, BlueprintType, EditInlineNew)
-class POSESEARCH_API UPoseSearchFeatureChannel : public UObject, public IBoneReferenceSkeletonProvider, public IPoseFilter
+class POSESEARCH_API UPoseSearchFeatureChannel : public UObject, public IBoneReferenceSkeletonProvider, public IPoseSearchFilter
 {
 	GENERATED_BODY()
 
@@ -134,7 +134,7 @@ public:
 	// Called at database build time to collect feature weights.
 	// Weights is sized to the cardinality of the schema and the feature channel should write
 	// its weights at the channel's data offset. Channels should provide a weight for each dimension.
-	virtual void FillWeights(TArray<float>& Weights) const PURE_VIRTUAL(UPoseSearchFeatureChannel::FillWeights, );
+	virtual void FillWeights(TArrayView<float> Weights) const PURE_VIRTUAL(UPoseSearchFeatureChannel::FillWeights, );
 
 	// Called at database build time to populate pose vectors with this channel's data
 	virtual void IndexAsset(UE::PoseSearch::FAssetIndexer& Indexer) const PURE_VIRTUAL(UPoseSearchFeatureChannel::IndexAsset, );
diff --git a/Engine/Plugins/Experimental/Animation/PoseSearch/Source/Runtime/Public/PoseSearch/PoseSearchIndex.h b/Engine/Plugins/Experimental/Animation/PoseSearch/Source/Runtime/Public/PoseSearch/PoseSearchIndex.h
--- a/Engine/Plugins/Experimental/Animation/PoseSearch/Source/Runtime/Public/PoseSearch/PoseSearchIndex.h
+++ b/Engine/Plugins/Experimental/Animation/PoseSearch/Source/Runtime/Public/PoseSearch/PoseSearchIndex.h
@@ -4,29 +4,24 @@
 
 #include "PoseSearch/KDTree.h"
 #include "PoseSearch/PoseSearchCost.h"
-#include "PoseSearchIndex.generated.h"
 
 namespace UE::PoseSearch
 {
 
-POSESEARCH_API void CompareFeatureVectors(TConstArrayView<float> A, TConstArrayView<float> B, TConstArrayView<float> WeightsSqrt, TArrayView<float> Result);
+constexpr int32 TMax(int32 A, int32 B) { return (A > B ? A : B); }
+template<typename ElementType>
+using TAlignedArray = TArray<ElementType, TAlignedHeapAllocator<TMax(alignof(ElementType), 16)>>;
 
-} // namespace UE::PoseSearch
+POSESEARCH_API void CompareFeatureVectors(TConstArrayView<float> A, TConstArrayView<float> B, TConstArrayView<float> WeightsSqrt, TArrayView<float> Result);
 
 /**
  * This is kept for each pose in the search index along side the feature vector values and is used to influence the search.
  */
-USTRUCT()
-struct POSESEARCH_API FPoseSearchPoseMetadata
+struct FPoseMetadata
 {
-	GENERATED_BODY()
-
 private:
 	// bits 0-30 represent the AssetIndex, bit 31 represents bBlockTransition
-	UPROPERTY(meta = (NeverInHash))
 	uint32 Data = 0;
-
-	UPROPERTY(meta = (NeverInHash))
 	float CostAddend = 0.f;
 
 public:
@@ -53,21 +48,18 @@ public:
 		return CostAddend;
 	}
 
-	friend FArchive& operator<<(FArchive& Ar, FPoseSearchPoseMetadata& Metadata);
+	friend FArchive& operator<<(FArchive& Ar, FPoseMetadata& Metadata);
 };
 
 /**
 * Information about a source animation asset used by a search index.
-* Some source animation entries may generate multiple FPoseSearchIndexAsset entries.
+* Some source animation entries may generate multiple FSearchIndexAsset entries.
 **/
-USTRUCT()
-struct POSESEARCH_API FPoseSearchIndexAsset
+struct FSearchIndexAsset
 {
-	GENERATED_BODY()
-
-	FPoseSearchIndexAsset() {}
+	FSearchIndexAsset() {}
 
-	FPoseSearchIndexAsset(
+	FSearchIndexAsset(
 		int32 InSourceAssetIdx,
 		int32 InFirstPoseIdx,
 		bool bInMirrored, 
@@ -87,28 +79,17 @@ struct POSESEARCH_API FPoseSearchIndexAsset
 	}
 
 	// Index of the source asset in search index's container (i.e. UPoseSearchDatabase)
-	UPROPERTY(meta = (NeverInHash))
 	int32 SourceAssetIdx = INDEX_NONE;
-
-	UPROPERTY(meta = (NeverInHash))
 	bool bMirrored = false;
-
-	UPROPERTY(meta = (NeverInHash))
 	int32 PermutationIdx = INDEX_NONE;
-
-	UPROPERTY(meta = (NeverInHash))
 	FVector BlendParameters = FVector::Zero();
-
-	UPROPERTY(meta = (NeverInHash))
 	int32 FirstPoseIdx = INDEX_NONE;
-
-	UPROPERTY(meta = (NeverInHash))
 	int32 FirstSampleIdx = INDEX_NONE;
-
-	UPROPERTY(meta = (NeverInHash))
 	int32 LastSampleIdx = INDEX_NONE;
 
 	bool IsPoseInRange(int32 PoseIdx) const { return (PoseIdx >= FirstPoseIdx) && (PoseIdx < FirstPoseIdx + GetNumPoses()); }
+
+#if DO_CHECK
 	bool IsInitialized() const 
 	{
 		return 
@@ -118,6 +99,7 @@ struct POSESEARCH_API FPoseSearchIndexAsset
 			FirstSampleIdx != INDEX_NONE &&
 			LastSampleIdx != INDEX_NONE;
 	}
+#endif // DO_CHECK
 
 	int32 GetBeginSampleIdx() const { return FirstSampleIdx; }
 	int32 GetEndSampleIdx() const {	return LastSampleIdx + 1; }
@@ -164,113 +146,81 @@ struct POSESEARCH_API FPoseSearchIndexAsset
 		return Time;
 	}
 
-	friend FArchive& operator<<(FArchive& Ar, FPoseSearchIndexAsset& IndexAsset);
+	friend FArchive& operator<<(FArchive& Ar, FSearchIndexAsset& IndexAsset);
 };
 
-USTRUCT()
-struct POSESEARCH_API FPoseSearchStats
+struct FSearchStats
 {
-	GENERATED_BODY()
-
-	UPROPERTY(meta = (NeverInHash))
 	float AverageSpeed = 0.f;
-
-	UPROPERTY(meta = (NeverInHash))
 	float MaxSpeed = 0.f;
-
-	UPROPERTY(meta = (NeverInHash))
 	float AverageAcceleration = 0.f;
-
-	UPROPERTY(meta = (NeverInHash))
 	float MaxAcceleration = 0.f;
-
-	friend FArchive& operator<<(FArchive& Ar, FPoseSearchStats& Stats);
+	friend FArchive& operator<<(FArchive& Ar, FSearchStats& Stats);
 };
 
 /**
-* case class for FPoseSearchIndex. building block used to gather data for data mining and calculate weights, pca, kdtree stuff
+* case class for FSearchIndex. building block used to gather data for data mining and calculate weights, pca, kdtree stuff
 */
-USTRUCT()
-struct POSESEARCH_API FPoseSearchIndexBase
+struct FSearchIndexBase
 {
-	GENERATED_BODY()
-
-	UPROPERTY(meta = (NeverInHash))
-	TArray<float> Values;
-	
-	UPROPERTY(meta = (NeverInHash))
-	TArray<FPoseSearchPoseMetadata> PoseMetadata;
-
-	UPROPERTY(meta = (NeverInHash))
+	TAlignedArray<float> Values;
+	TAlignedArray<FPoseMetadata> PoseMetadata;
 	bool bAnyBlockTransition = false;
-
-	UPROPERTY(meta = (NeverInHash))
-	TArray<FPoseSearchIndexAsset> Assets;
+	TAlignedArray<FSearchIndexAsset> Assets;
 
 	// minimum of the database metadata CostAddend: it represents the minimum cost of any search for the associated database (we'll skip the search in case the search result total cost is already less than MinCostAddend)
-	UPROPERTY(meta = (NeverInHash))
 	float MinCostAddend = -MAX_FLT;
 
 	// @todo: this property should be editor only
-	UPROPERTY(meta = (NeverInHash))
-	FPoseSearchStats Stats;
+	FSearchStats Stats;
 
 	int32 GetNumPoses() const { return PoseMetadata.Num(); }
 	bool IsValidPoseIndex(int32 PoseIdx) const { return PoseIdx < GetNumPoses(); }
 	bool IsEmpty() const;
 
-	const FPoseSearchIndexAsset& GetAssetForPose(int32 PoseIdx) const;
-	const FPoseSearchIndexAsset* GetAssetForPoseSafe(int32 PoseIdx) const;
+	const FSearchIndexAsset& GetAssetForPose(int32 PoseIdx) const;
+	POSESEARCH_API const FSearchIndexAsset* GetAssetForPoseSafe(int32 PoseIdx) const;
 
 	void Reset();
 	
-	friend FArchive& operator<<(FArchive& Ar, FPoseSearchIndexBase& Index);
+	friend FArchive& operator<<(FArchive& Ar, FSearchIndexBase& Index);
 };
 
 /**
 * A search index for animation poses. The structure of the search index is determined by its UPoseSearchSchema.
 * May represent a single animation (see UPoseSearchSequenceMetaData) or a collection (see UPoseSearchDatabase).
 */
-USTRUCT()
-struct POSESEARCH_API FPoseSearchIndex : public FPoseSearchIndexBase
+struct FSearchIndex : public FSearchIndexBase
 {
-	GENERATED_BODY()
-
 	// we store weights square roots to reduce numerical errors when CompareFeatureVectors 
 	// ((VA - VB) * VW).square().sum()
 	// instead of
 	// ((VA - VB).square() * VW).sum()
 	// since (VA - VB).square() could lead to big numbers, and VW being multiplied by the variance of the dataset
-	UPROPERTY(Category = Info, VisibleAnywhere, meta = (NeverInHash))
-	TArray<float> WeightsSqrt;
-
-	UPROPERTY(meta = (NeverInHash))
-	TArray<float> PCAValues;
-
-	UPROPERTY(Category = Info, VisibleAnywhere, meta = (NeverInHash))
-	TArray<float> PCAProjectionMatrix;
+	TAlignedArray<float> WeightsSqrt;
+	TAlignedArray<float> PCAValues;
+	TAlignedArray<float> PCAProjectionMatrix;
+	TAlignedArray<float> Mean;
 
-	UPROPERTY(Category = Info, VisibleAnywhere, meta = (NeverInHash))
-	TArray<float> Mean;
-
-	UE::PoseSearch::FKDTree KDTree;
+	FKDTree KDTree;
 
 	// @todo: this property should be editor only
-	UPROPERTY(Category = Info, VisibleAnywhere, meta = (NeverInHash))
 	float PCAExplainedVariance = 0.f;
 
-	FPoseSearchIndex() = default;
-	~FPoseSearchIndex() = default;
-	FPoseSearchIndex(const FPoseSearchIndex& Other); // custom copy constructor to deal with the KDTree DataSrc
-	FPoseSearchIndex(FPoseSearchIndex&& Other) = delete;
-	FPoseSearchIndex& operator=(const FPoseSearchIndex& Other); // custom equal operator to deal with the KDTree DataSrc
-	FPoseSearchIndex& operator=(FPoseSearchIndex&& Other) = delete;
+	FSearchIndex() = default;
+	~FSearchIndex() = default;
+	FSearchIndex(const FSearchIndex& Other); // custom copy constructor to deal with the KDTree DataSrc
+	FSearchIndex(FSearchIndex&& Other) = delete;
+	FSearchIndex& operator=(const FSearchIndex& Other); // custom equal operator to deal with the KDTree DataSrc
+	FSearchIndex& operator=(FSearchIndex&& Other) = delete;
 
 	void Reset();
 	TConstArrayView<float> GetPoseValues(int32 PoseIdx) const;
 	TConstArrayView<float> GetReconstructedPoseValues(int32 PoseIdx, TArrayView<float> BufferUsedForReconstruction) const;
-	TArray<float> GetPoseValuesSafe(int32 PoseIdx) const;
-	FPoseSearchCost ComparePoses(int32 PoseIdx, float ContinuingPoseCostBias, TConstArrayView<float> PoseValues, TConstArrayView<float> QueryValues) const;
+	POSESEARCH_API TArray<float> GetPoseValuesSafe(int32 PoseIdx) const;
+	POSESEARCH_API FPoseSearchCost ComparePoses(int32 PoseIdx, float ContinuingPoseCostBias, TConstArrayView<float> PoseValues, TConstArrayView<float> QueryValues) const;
 
-	friend FArchive& operator<<(FArchive& Ar, FPoseSearchIndex& Index);
+	friend FArchive& operator<<(FArchive& Ar, FSearchIndex& Index);
 };
+
+} // namespace UE::PoseSearch
diff --git a/Engine/Plugins/Experimental/Animation/PoseSearch/Source/Runtime/Public/PoseSearch/PoseSearchResult.h b/Engine/Plugins/Experimental/Animation/PoseSearch/Source/Runtime/Public/PoseSearch/PoseSearchResult.h
--- a/Engine/Plugins/Experimental/Animation/PoseSearch/Source/Runtime/Public/PoseSearch/PoseSearchResult.h
+++ b/Engine/Plugins/Experimental/Animation/PoseSearch/Source/Runtime/Public/PoseSearch/PoseSearchResult.h
@@ -5,12 +5,12 @@
 #include "PoseSearch/PoseSearchCost.h"
 #include "PoseSearch/PoseSearchDefines.h"
 
-struct FPoseSearchIndexAsset;
 class UPoseSearchDatabase;
 class UPoseSearchSchema;
 
 namespace UE::PoseSearch
 {
+struct FSearchIndexAsset;
 
 /**
 * float buffer of features according to a UPoseSearchSchema layout.
@@ -62,7 +62,7 @@ struct FSearchResult
 
 	void Reset();
 
-	const FPoseSearchIndexAsset* GetSearchIndexAsset(bool bMandatory = false) const;
+	const FSearchIndexAsset* GetSearchIndexAsset(bool bMandatory = false) const;
 };
 
 } // namespace UE::PoseSearch
\ No newline at end of file
diff --git a/Engine/Plugins/Experimental/Animation/PoseSearch/Source/Runtime/Public/PoseSearch/PoseSearchSchema.h b/Engine/Plugins/Experimental/Animation/PoseSearch/Source/Runtime/Public/PoseSearch/PoseSearchSchema.h
--- a/Engine/Plugins/Experimental/Animation/PoseSearch/Source/Runtime/Public/PoseSearch/PoseSearchSchema.h
+++ b/Engine/Plugins/Experimental/Animation/PoseSearch/Source/Runtime/Public/PoseSearch/PoseSearchSchema.h
@@ -100,6 +100,10 @@ public:
 	UPROPERTY(EditAnywhere, Category = "Permutations")
 	float PermutationsTimeOffset = 0.f;
 
+	// if true a padding channel will be added to make sure the data is 16 bytes (aligned) and padded, to facilitate performance improvements at cost of eventual additional memory
+	UPROPERTY(EditAnywhere, Category = "Performance")
+	bool bAddDataPadding = false;
+
 	// If bInjectAdditionalDebugChannels is true, channels will be asked to inject additional channels into this schema.
 	// the original intent is to add UPoseSearchFeatureChannel_Position(s) to help with the complexity of the debug drawing
 	// (the database will have all the necessary positions to draw lines at the right location and time).
