motion matching - fix for FMotionMatchingState::UpdateRootBoneControl that was using the RootMotionTransformDelta from the end of the frame (and potentially updated by other anim nodes), rather than the end of FAnimNode_MotionMatching::Evaluate_AnyThread

diff --git a/Engine/Plugins/Experimental/Animation/PoseSearch/Source/Runtime/Private/AnimNode_MotionMatching.cpp b/Engine/Plugins/Experimental/Animation/PoseSearch/Source/Runtime/Private/AnimNode_MotionMatching.cpp
--- a/Engine/Plugins/Experimental/Animation/PoseSearch/Source/Runtime/Private/AnimNode_MotionMatching.cpp
+++ b/Engine/Plugins/Experimental/Animation/PoseSearch/Source/Runtime/Private/AnimNode_MotionMatching.cpp
@@ -49,35 +49,39 @@ void FAnimNode_MotionMatching::Evaluate_AnyThread(FPoseContext& Output)
 
 	Source.Evaluate(Output);
 
-	// applying MotionMatchingState.RootBoneDeltaYaw to the root bone and the root motion delta transform
-	if (!FMath::IsNearlyZero(MotionMatchingState.GetRootBoneDeltaYaw()))
+	const UE::Anim::IAnimRootMotionProvider* RootMotionProvider = UE::Anim::IAnimRootMotionProvider::Get();
+	FTransform RootMotionTransformDelta;
+	if (RootMotionProvider && RootMotionProvider->HasRootMotion(Output.CustomAttributes))
+	{
+		RootMotionProvider->ExtractRootMotion(Output.CustomAttributes, RootMotionTransformDelta);
+	}
+	else
 	{
-		const FQuat RootBoneDelta(FRotator(0.f, MotionMatchingState.GetRootBoneDeltaYaw(), 0.f));
+		RootMotionTransformDelta = FTransform::Identity;
+		RootMotionProvider = nullptr;
+	}
+
+	// applying MotionMatchingState.ComponentDeltaYaw (considered as root bone delta yaw) to the root bone and the root motion delta transform
+	if (!FMath::IsNearlyZero(MotionMatchingState.ComponentDeltaYaw))
+	{
+		const FQuat RootBoneDelta(FRotator(0.f, MotionMatchingState.ComponentDeltaYaw, 0.f));
 		FCompactPoseBoneIndex RootBoneIndex(RootBoneIndexType);
 		Output.Pose[RootBoneIndex].SetRotation(Output.Pose[RootBoneIndex].GetRotation() * RootBoneDelta);
 		Output.Pose[RootBoneIndex].NormalizeRotation();
 
-		const UE::Anim::IAnimRootMotionProvider* RootMotionProvider = UE::Anim::IAnimRootMotionProvider::Get();
-		if (RootMotionProvider && RootMotionProvider->HasRootMotion(Output.CustomAttributes))
+		RootMotionTransformDelta.SetTranslation(RootBoneDelta.RotateVector(RootMotionTransformDelta.GetTranslation()));
+
+		if (RootMotionProvider)
 		{
-			FTransform RootMotionTransformDelta;
-			RootMotionProvider->ExtractRootMotion(Output.CustomAttributes, RootMotionTransformDelta);
-			RootMotionTransformDelta.SetTranslation(RootBoneDelta.RotateVector(RootMotionTransformDelta.GetTranslation()));
 			RootMotionProvider->OverrideRootMotion(RootMotionTransformDelta, Output.CustomAttributes);
 		}
 	}
-	
+
+	MotionMatchingState.AnimationDeltaYaw = FRotator(RootMotionTransformDelta.GetRotation()).Yaw;
+
 #if UE_POSE_SEARCH_TRACE_ENABLED
-	const UE::Anim::IAnimRootMotionProvider* RootMotionProvider = UE::Anim::IAnimRootMotionProvider::Get();
-	if (RootMotionProvider && RootMotionProvider->HasRootMotion(Output.CustomAttributes))
-	{
-		RootMotionProvider->ExtractRootMotion(Output.CustomAttributes, MotionMatchingState.RootMotionTransformDelta);
-	}
-	else
-	{
-		MotionMatchingState.RootMotionTransformDelta = FTransform::Identity;
-	}
-#endif // UE_POSE_SEARCH_TRACE_ENABLED
+	MotionMatchingState.RootMotionTransformDelta = RootMotionTransformDelta;
+#endif //UE_POSE_SEARCH_TRACE_ENABLED
 }
 
 void FAnimNode_MotionMatching::UpdateAssetPlayer(const FAnimationUpdateContext& Context)
diff --git a/Engine/Plugins/Experimental/Animation/PoseSearch/Source/Runtime/Private/PoseSearchLibrary.cpp b/Engine/Plugins/Experimental/Animation/PoseSearch/Source/Runtime/Private/PoseSearchLibrary.cpp
--- a/Engine/Plugins/Experimental/Animation/PoseSearch/Source/Runtime/Private/PoseSearchLibrary.cpp
+++ b/Engine/Plugins/Experimental/Animation/PoseSearch/Source/Runtime/Private/PoseSearchLibrary.cpp
@@ -41,13 +41,15 @@ void FMotionMatchingState::Reset(const FTransform& ComponentTransform)
 	ElapsedPoseSearchTime = INFINITY;
 	WantedPlayRate = 1.f;
 	bJumpedToPose = false;
-	RootBoneDeltaYaw = 0.f;
-	RootBoneWorldYaw = FRotator(ComponentTransform.GetRotation()).Yaw;
+	ComponentDeltaYaw = 0.f;
+	ComponentWorldYaw = FRotator(ComponentTransform.GetRotation()).Yaw;
+	AnimationDeltaYaw = 0.f;
+
+	PoseIndicesHistory.Reset();
 
 #if UE_POSE_SEARCH_TRACE_ENABLED
 	RootMotionTransformDelta = FTransform::Identity;
 #endif // UE_POSE_SEARCH_TRACE_ENABLED
-	PoseIndicesHistory.Reset();
 }
 
 void FMotionMatchingState::AdjustAssetTime(float AssetTime)
@@ -123,42 +125,27 @@ void FMotionMatchingState::UpdateWantedPlayRate(const UE::PoseSearch::FSearchCon
 void FMotionMatchingState::UpdateRootBoneControl(const FAnimationUpdateContext& Context, float YawFromAnimationBlendRate)
 {
 	const FAnimInstanceProxy* AnimInstanceProxy = Context.AnimInstanceProxy;
-	static const FName RootMotionAttributeName = "RootMotionDelta";
-	static const UE::Anim::FAttributeId RootMotionAttributeId = { RootMotionAttributeName , FCompactPoseBoneIndex(0) };
 
-	const FRotator ComponentWorldRotator(AnimInstanceProxy->GetComponentTransform().GetRotation());
+	const float CurrentComponentWorldYaw = FRotator(AnimInstanceProxy->GetComponentTransform().GetRotation()).Yaw;
 	if (YawFromAnimationBlendRate < 0.f)
 	{
-		RootBoneWorldYaw = ComponentWorldRotator.Yaw;
-		RootBoneDeltaYaw = 0.f;
+		ComponentWorldYaw = CurrentComponentWorldYaw;
+		ComponentDeltaYaw = 0.f;
 	}
 	else
 	{
-		// integrating MotionMatchingState.RootBoneTargetYaw with the previous frame RootMotionDelta
-		const USkeletalMeshComponent* Mesh = AnimInstanceProxy->GetSkelMeshComponent();
-		check(Mesh);
-		
-		if (const FTransformAnimationAttribute* RootMotionAttribute = Mesh->GetCustomAttributes().Find<FTransformAnimationAttribute>(RootMotionAttributeId))
-		{
-			const FRotator RootMotionRotatorDelta(RootMotionAttribute->Value.GetRotation());
-			const float RootBoneAnimationDelta = RootMotionRotatorDelta.Yaw;
-			const float RootBoneToComponentDelta = FRotator::NormalizeAxis(ComponentWorldRotator.Yaw - RootBoneWorldYaw);
+		// integrating MotionMatchingState.ComponentWorldYaw with a lerped value between 
+		// the previous frame AnimationDeltaYaw and CurrentComponentDeltaYaw (component delta yaw happened during the delta time)
+		const float CurrentComponentDeltaYaw = FRotator::NormalizeAxis(CurrentComponentWorldYaw - ComponentWorldYaw);
 
-			// lerping the animation delta with the capsule delta
-			const float RootBoneDelta = FMath::Lerp(RootBoneAnimationDelta, RootBoneToComponentDelta, YawFromAnimationBlendRate * Context.GetDeltaTime());
+		// lerping the animation delta with the capsule delta
+		const float LerpValue = FMath::Min(YawFromAnimationBlendRate * Context.GetDeltaTime(), 1.f);
+		const float LerpedDeltaYaw = FMath::Lerp(AnimationDeltaYaw, CurrentComponentDeltaYaw, LerpValue);
 
-			RootBoneWorldYaw = FRotator::NormalizeAxis(RootBoneWorldYaw + RootBoneDelta);
+		ComponentWorldYaw = FRotator::NormalizeAxis(ComponentWorldYaw + LerpedDeltaYaw);
 
-			// @todo: handle the case when the character is on top of a rotating platform
-			RootBoneDeltaYaw = FRotator::NormalizeAxis(RootBoneWorldYaw - ComponentWorldRotator.Yaw);
-		}
-		else
-		{
-			UE_LOG(LogPoseSearch, Warning, TEXT("FMotionMatchingState::UpdateRootBoneControl - Couldn't find FTransformAnimationAttribute. Root bone yaw control will not be performed"));
-
-			RootBoneWorldYaw = ComponentWorldRotator.Yaw;
-			RootBoneDeltaYaw = 0.f;
-		}
+		// @todo: handle the case when the character is on top of a rotating platform
+		ComponentDeltaYaw = FRotator::NormalizeAxis(ComponentWorldYaw - CurrentComponentWorldYaw);
 	}
 }
 
@@ -308,7 +295,7 @@ void UPoseSearchLibrary::UpdateMotionMatchingState(
 		History = &PoseHistoryProvider->GetPoseHistory();
 	}
 	
-	const FPoseSearchQueryTrajectory TrajectoryRootSpace = ProcessTrajectory(Trajectory, Context.AnimInstanceProxy->GetComponentTransform(), InOutMotionMatchingState.GetRootBoneDeltaYaw(), YawFromAnimationTrajectoryBlendTime, TrajectorySpeedMultiplier);
+	const FPoseSearchQueryTrajectory TrajectoryRootSpace = ProcessTrajectory(Trajectory, Context.AnimInstanceProxy->GetComponentTransform(), InOutMotionMatchingState.ComponentDeltaYaw, YawFromAnimationTrajectoryBlendTime, TrajectorySpeedMultiplier);
 	FSearchContext SearchContext(&TrajectoryRootSpace, History, 0.f, &InOutMotionMatchingState.PoseIndicesHistory, InOutMotionMatchingState.CurrentSearchResult, PoseJumpThresholdTime, bForceInterrupt);
 
 	const bool bMustSearch = !InOutMotionMatchingState.CurrentSearchResult.IsValid();
@@ -403,7 +390,7 @@ void UPoseSearchLibrary::UpdateMotionMatchingState(
 		else
 #endif // WITH_EDITOR
 		{
-			const FRotator DeltaRotation(0.f, InOutMotionMatchingState.GetRootBoneDeltaYaw(), 0.f);
+			const FRotator DeltaRotation(0.f, InOutMotionMatchingState.ComponentDeltaYaw, 0.f);
 			const FTransform DeltaTransform(DeltaRotation);
 			const FTransform RootBoneTransform = DeltaTransform * Context.AnimInstanceProxy->GetComponentTransform();
 
diff --git a/Engine/Plugins/Experimental/Animation/PoseSearch/Source/Runtime/Public/PoseSearch/PoseSearchLibrary.h b/Engine/Plugins/Experimental/Animation/PoseSearch/Source/Runtime/Public/PoseSearch/PoseSearchLibrary.h
--- a/Engine/Plugins/Experimental/Animation/PoseSearch/Source/Runtime/Public/PoseSearch/PoseSearchLibrary.h
+++ b/Engine/Plugins/Experimental/Animation/PoseSearch/Source/Runtime/Public/PoseSearch/PoseSearchLibrary.h
@@ -43,8 +43,6 @@ struct POSESEARCH_API FMotionMatchingState
 
 	void UpdateRootBoneControl(const FAnimationUpdateContext& Context, float YawFromAnimationBlendRate);
 
-	float GetRootBoneDeltaYaw() const { return RootBoneDeltaYaw; }
-
 	UE::PoseSearch::FSearchResult CurrentSearchResult;
 
 	// Time since the last pose jump
@@ -59,16 +57,21 @@ struct POSESEARCH_API FMotionMatchingState
 	UPROPERTY(Transient, VisibleAnywhere, BlueprintReadOnly, Category=State)
 	bool bJumpedToPose = false;
 
+	UE::PoseSearch::FPoseIndicesHistory PoseIndicesHistory;
+
+	// Component delta yaw (also considered as root bone delta yaw)
+	float ComponentDeltaYaw = 0.f;
+
+	// Internal component yaw in world space. Initialized as FRotator(AnimInstanceProxy->GetComponentTransform().GetRotation()).Yaw, but then integrated by ComponentDeltaYaw
+	float ComponentWorldYaw = 0.f;
+	
+	// RootMotionTransformDelta yaw at the end of FAnimNode_MotionMatching::Evaluate_AnyThread (it represents the previous frame animation delta yaw)
+	float AnimationDeltaYaw = 0.f;
+
 #if UE_POSE_SEARCH_TRACE_ENABLED
 	// Root motion delta for currently playing animation (or animation tree if from the blend stack)
 	FTransform RootMotionTransformDelta = FTransform::Identity;
 #endif //UE_POSE_SEARCH_TRACE_ENABLED
-
-	UE::PoseSearch::FPoseIndicesHistory PoseIndicesHistory;
-
-private:
-	float RootBoneDeltaYaw = 0.f;
-	float RootBoneWorldYaw = 0.f;
 };
 
 UCLASS()
