motion matching - pose search performance optimizations (more data memory aligment, use of TMemStackAllocator, implemented CompareAlignedFeatureVectors using vector registers)

diff --git a/Engine/Plugins/Experimental/Animation/PoseSearch/Source/Runtime/Private/PoseSearchDatabase.cpp b/Engine/Plugins/Experimental/Animation/PoseSearch/Source/Runtime/Private/PoseSearchDatabase.cpp
--- a/Engine/Plugins/Experimental/Animation/PoseSearch/Source/Runtime/Private/PoseSearchDatabase.cpp
+++ b/Engine/Plugins/Experimental/Animation/PoseSearch/Source/Runtime/Private/PoseSearchDatabase.cpp
@@ -716,6 +716,36 @@ UE::PoseSearch::FSearchResult UPoseSearchDatabase::Search(UE::PoseSearch::FSearc
 	return Result;
 }
 
+template<bool bReconstructPoseValues, bool bAlignedAndPadded>
+static inline void EvaluatePoseKernel(UE::PoseSearch::FSearchResult& Result, const UE::PoseSearch::FSearchIndex& SearchIndex, TConstArrayView<float> QueryValues, TArrayView<float> ReconstructedPoseValuesBuffer,
+	int32 PoseIdx, const UE::PoseSearch::FSearchFilters& SearchFilters, UE::PoseSearch::FSearchContext& SearchContext, const UPoseSearchDatabase* Database, bool bUpdateBestCandidates = true)
+{
+	using namespace UE::PoseSearch;
+
+	const TConstArrayView<float> PoseValues = bReconstructPoseValues ? SearchIndex.GetReconstructedPoseValues(PoseIdx, ReconstructedPoseValuesBuffer) : SearchIndex.GetPoseValues(PoseIdx);
+
+	if (SearchFilters.AreFiltersValid(SearchIndex, PoseValues, QueryValues, PoseIdx, SearchIndex.PoseMetadata[PoseIdx]
+#if UE_POSE_SEARCH_TRACE_ENABLED
+		, SearchContext, Database
+#endif
+	))
+	{
+		const FPoseSearchCost PoseCost = bAlignedAndPadded ? SearchIndex.CompareAlignedPoses(PoseIdx, 0.f, PoseValues, QueryValues) : SearchIndex.ComparePoses(PoseIdx, 0.f, PoseValues, QueryValues);;
+		if (PoseCost < Result.PoseCost)
+		{
+			Result.PoseCost = PoseCost;
+			Result.PoseIdx = PoseIdx;
+		}
+
+#if UE_POSE_SEARCH_TRACE_ENABLED
+		if (bUpdateBestCandidates)
+		{
+			SearchContext.BestCandidates.Add(PoseCost, PoseIdx, Database, EPoseCandidateFlags::Valid_Pose);
+		}
+#endif
+	}
+}
+
 FPoseSearchCost UPoseSearchDatabase::SearchContinuingPose(UE::PoseSearch::FSearchContext& SearchContext) const
 {
 	using namespace UE::PoseSearch;
@@ -759,11 +789,22 @@ FPoseSearchCost UPoseSearchDatabase::SearchContinuingPose(UE::PoseSearch::FSearc
 	if (!GetSkipSearchIfPossible() || SearchContext.GetCurrentBestTotalCost() > SearchIndex.MinCostAddend + ContinuingPoseCostBias)
 	{
 		const int32 NumDimensions = Schema->SchemaCardinality;
+		// FMemory_Alloca is forced 16 bytes aligned
 		TArrayView<float> ReconstructedPoseValuesBuffer((float*)FMemory_Alloca(NumDimensions * sizeof(float)), NumDimensions);
+		check(IsAligned(ReconstructedPoseValuesBuffer.GetData(), alignof(VectorRegister4Float)));
 		const TConstArrayView<float> PoseValues = SearchIndex.Values.IsEmpty() ? SearchIndex.GetReconstructedPoseValues(PoseIdx, ReconstructedPoseValuesBuffer) : SearchIndex.GetPoseValues(PoseIdx);
 
 		const int32 ContinuingPoseIdx = SearchContext.GetCurrentResult().PoseIdx;
-		ContinuingPoseCost = SearchIndex.ComparePoses(ContinuingPoseIdx, ContinuingPoseCostBias, PoseValues, SearchContext.GetOrBuildQuery(Schema).GetValues());
+		// is the data padded at 16 bytes (and 16 bytes aligned by construction)?
+		if (NumDimensions % 4 == 0)
+		{
+			ContinuingPoseCost = SearchIndex.CompareAlignedPoses(ContinuingPoseIdx, ContinuingPoseCostBias, PoseValues, SearchContext.GetOrBuildQuery(Schema).GetValues());
+		}
+		// data is not 16 bytes padded
+		else
+		{
+			ContinuingPoseCost = SearchIndex.ComparePoses(ContinuingPoseIdx, ContinuingPoseCostBias, PoseValues, SearchContext.GetOrBuildQuery(Schema).GetValues());
+		}
 
 #if UE_POSE_SEARCH_TRACE_ENABLED
 		SearchContext.BestCandidates.Add(ContinuingPoseCost, ContinuingPoseIdx, this, EPoseCandidateFlags::Valid_ContinuingPose);
@@ -793,7 +834,6 @@ UE::PoseSearch::FSearchResult UPoseSearchDatabase::SearchPCAKDTree(UE::PoseSearc
 	RowMajorVectorMap WeightedQueryValues((float*)FMemory_Alloca(NumDimensions * sizeof(float)), 1, NumDimensions);
 	RowMajorVectorMap CenteredQueryValues((float*)FMemory_Alloca(NumDimensions * sizeof(float)), 1, NumDimensions);
 	RowMajorVectorMap ProjectedQueryValues((float*)FMemory_Alloca(ClampedNumberOfPrincipalComponents * sizeof(float)), 1, ClampedNumberOfPrincipalComponents);
-	TArrayView<float> ReconstructedPoseValuesBuffer((float*)FMemory_Alloca(NumDimensions * sizeof(float)), NumDimensions);
 
 #if DO_CHECK
 	// KDTree in PCA space search
@@ -861,27 +901,32 @@ UE::PoseSearch::FSearchResult UPoseSearchDatabase::SearchPCAKDTree(UE::PoseSearc
 
 		// NonSelectableIdx are already filtered out inside the kdtree search
 		const FSearchFilters SearchFilters(Schema, TConstArrayView<size_t>(), SearchIndex.bAnyBlockTransition);
-		for (size_t ResultIndex = 0; ResultIndex < ResultSet.Num(); ++ResultIndex)
+		
+		// do we need to reconstruct pose values?
+		if (SearchIndex.Values.IsEmpty())
 		{
-			const int32 PoseIdx = ResultIndexes[ResultIndex];
-			const TConstArrayView<float> PoseValues = SearchIndex.Values.IsEmpty() ? SearchIndex.GetReconstructedPoseValues(PoseIdx, ReconstructedPoseValuesBuffer) : SearchIndex.GetPoseValues(PoseIdx);
-
-			if (SearchFilters.AreFiltersValid(SearchIndex, PoseValues, QueryValues, PoseIdx, SearchIndex.PoseMetadata[PoseIdx]
-#if UE_POSE_SEARCH_TRACE_ENABLED
-				, SearchContext, this
-#endif
-			))
+			// FMemory_Alloca is forced 16 bytes aligned
+			TArrayView<float> ReconstructedPoseValuesBuffer((float*)FMemory_Alloca(NumDimensions * sizeof(float)), NumDimensions);
+			check(IsAligned(ReconstructedPoseValuesBuffer.GetData(), alignof(VectorRegister4Float)));
+			for (size_t ResultIndex = 0; ResultIndex < ResultSet.Num(); ++ResultIndex)
 			{
-				const FPoseSearchCost PoseCost = SearchIndex.ComparePoses(PoseIdx, 0.f, PoseValues, QueryValues);
-				if (PoseCost < Result.PoseCost)
-				{
-					Result.PoseCost = PoseCost;
-					Result.PoseIdx = PoseIdx;
-				}
-
-#if UE_POSE_SEARCH_TRACE_ENABLED
-				SearchContext.BestCandidates.Add(PoseCost, PoseIdx, this, EPoseCandidateFlags::Valid_Pose);
-#endif
+				EvaluatePoseKernel<true, false>(Result, SearchIndex, QueryValues, ReconstructedPoseValuesBuffer, ResultIndexes[ResultIndex], SearchFilters, SearchContext, this);
+			}
+		}
+		// is the data padded at 16 bytes (and 16 bytes aligned by construction)?
+		else if (NumDimensions % 4 == 0)
+		{
+			for (size_t ResultIndex = 0; ResultIndex < ResultSet.Num(); ++ResultIndex)
+			{
+				EvaluatePoseKernel<false, true>(Result, SearchIndex, QueryValues, TArrayView<float>(), ResultIndexes[ResultIndex], SearchFilters, SearchContext, this);
+			}
+		}
+		// no reconstruction, but data is not 16 bytes padded
+		else
+		{
+			for (size_t ResultIndex = 0; ResultIndex < ResultSet.Num(); ++ResultIndex)
+			{
+				EvaluatePoseKernel<false, false>(Result, SearchIndex, QueryValues, TArrayView<float>(), ResultIndexes[ResultIndex], SearchFilters, SearchContext, this);
 			}
 		}
 	}
@@ -930,30 +975,34 @@ UE::PoseSearch::FSearchResult UPoseSearchDatabase::SearchBruteForce(UE::PoseSear
 		check(Algo::IsSorted(NonSelectableIdx));
 
 		const int32 NumDimensions = Schema->SchemaCardinality;
-		TArrayView<float> ReconstructedPoseValuesBuffer((float*)FMemory_Alloca(NumDimensions * sizeof(float)), NumDimensions);
 		const FSearchFilters SearchFilters(Schema, NonSelectableIdx, SearchIndex.bAnyBlockTransition);
-		for (int32 PoseIdx = 0; PoseIdx < SearchIndex.GetNumPoses(); ++PoseIdx)
+		const bool bUpdateBestCandidates = PoseSearchMode == EPoseSearchMode::BruteForce;
+
+		// do we need to reconstruct pose values?
+		if (SearchIndex.Values.IsEmpty())
 		{
-			const TConstArrayView<float> PoseValues = SearchIndex.Values.IsEmpty() ? SearchIndex.GetReconstructedPoseValues(PoseIdx, ReconstructedPoseValuesBuffer) : SearchIndex.GetPoseValues(PoseIdx);
-			if (SearchFilters.AreFiltersValid(SearchIndex, PoseValues, QueryValues, PoseIdx, SearchIndex.PoseMetadata[PoseIdx]
-#if UE_POSE_SEARCH_TRACE_ENABLED
-				, SearchContext, this
-#endif
-			))
+			// FMemory_Alloca is forced 16 bytes aligned
+			TArrayView<float> ReconstructedPoseValuesBuffer((float*)FMemory_Alloca(NumDimensions * sizeof(float)), NumDimensions);
+			check(IsAligned(ReconstructedPoseValuesBuffer.GetData(), alignof(VectorRegister4Float)));
+			for (int32 PoseIdx = 0; PoseIdx < SearchIndex.GetNumPoses(); ++PoseIdx)
 			{
-				const FPoseSearchCost PoseCost = SearchIndex.ComparePoses(PoseIdx, 0.f, PoseValues, QueryValues);
-				if (PoseCost < Result.PoseCost)
-				{
-					Result.PoseCost = PoseCost;
-					Result.PoseIdx = PoseIdx;
-				}
-
-#if UE_POSE_SEARCH_TRACE_ENABLED
-				if (PoseSearchMode == EPoseSearchMode::BruteForce)
-				{
-					SearchContext.BestCandidates.Add(PoseCost, PoseIdx, this, EPoseCandidateFlags::Valid_Pose);
-				}
-#endif
+				EvaluatePoseKernel<true, false>(Result, SearchIndex, QueryValues, ReconstructedPoseValuesBuffer, PoseIdx, SearchFilters, SearchContext, this, bUpdateBestCandidates);
+			}
+		}
+		// is the data padded at 16 bytes (and 16 bytes aligned by construction)?
+		else if (NumDimensions % 4 == 0)
+		{
+			for (int32 PoseIdx = 0; PoseIdx < SearchIndex.GetNumPoses(); ++PoseIdx)
+			{
+				EvaluatePoseKernel<false, true>(Result, SearchIndex, QueryValues, TArrayView<float>(), PoseIdx, SearchFilters, SearchContext, this, bUpdateBestCandidates);
+			}
+		}
+		// no reconstruction, but data is not 16 bytes padded
+		else
+		{
+			for (int32 PoseIdx = 0; PoseIdx < SearchIndex.GetNumPoses(); ++PoseIdx)
+			{
+				EvaluatePoseKernel<false, false>(Result, SearchIndex, QueryValues, TArrayView<float>(), PoseIdx, SearchFilters, SearchContext, this, bUpdateBestCandidates);
 			}
 		}
 	}
diff --git a/Engine/Plugins/Experimental/Animation/PoseSearch/Source/Runtime/Private/PoseSearchIndex.cpp b/Engine/Plugins/Experimental/Animation/PoseSearch/Source/Runtime/Private/PoseSearchIndex.cpp
--- a/Engine/Plugins/Experimental/Animation/PoseSearch/Source/Runtime/Private/PoseSearchIndex.cpp
+++ b/Engine/Plugins/Experimental/Animation/PoseSearch/Source/Runtime/Private/PoseSearchIndex.cpp
@@ -17,6 +17,52 @@ static FORCEINLINE float CompareFeatureVectors(TConstArrayView<float> A, TConstA
 	return ((VA - VB) * VW).square().sum();
 }
 
+static FORCEINLINE float CompareAlignedFeatureVectors(TConstArrayView<float> A, TConstArrayView<float> B, TConstArrayView<float> WeightsSqrt)
+{
+	check(A.Num() == B.Num() && A.Num() == WeightsSqrt.Num());
+	check(A.Num() % 4 == 0);
+	check(IsAligned(A.GetData(), alignof(VectorRegister4Float)));
+	check(IsAligned(B.GetData(), alignof(VectorRegister4Float)));
+	check(IsAligned(WeightsSqrt.GetData(), alignof(VectorRegister4Float)));
+	// sufficient condition to check for pointer overlapping
+	check(A.GetData() != B.GetData() && A.GetData() != WeightsSqrt.GetData());
+
+	const int32 NumVectors = A.Num() / 4;
+
+	const VectorRegister4Float* RESTRICT VA = reinterpret_cast<const VectorRegister4Float*>(A.GetData());
+	const VectorRegister4Float* RESTRICT VB = reinterpret_cast<const VectorRegister4Float*>(B.GetData());
+	const VectorRegister4Float* RESTRICT VW = reinterpret_cast<const VectorRegister4Float*>(WeightsSqrt.GetData());
+
+	VectorRegister4Float PartialCost = VectorZero();
+	for (int32 VectorIdx = 0; VectorIdx < NumVectors; ++VectorIdx, ++VA, ++VB, ++VW)
+	{
+		const VectorRegister4Float Diff = VectorSubtract(*VA, *VB);
+		const VectorRegister4Float WeightedDiff = VectorMultiply(Diff, *VW);
+		PartialCost = VectorMultiplyAdd(WeightedDiff, WeightedDiff, PartialCost);
+	}
+
+	// calculating PartialCost.X + PartialCost.Y + PartialCost.Z + PartialCost.W
+	VectorRegister4Float Swizzle = VectorSwizzle(PartialCost, 1, 0, 3, 2);	// (Y, X, W, Z) of PartialCost
+	PartialCost = VectorAdd(PartialCost, Swizzle);							// (X + Y, Y + X, Z + W, W + Z)
+	Swizzle = VectorSwizzle(PartialCost, 2, 3, 0, 1);						// (Z + W, W + Z, X + Y, Y + X)
+	PartialCost = VectorAdd(PartialCost, Swizzle);							// (X + Y + Z + W, Y + X + W + Z, Z + W + X + Y, W + Z + Y + X)
+	float Cost;
+	VectorStoreFloat1(PartialCost, &Cost);
+
+// keeping this debug code to validate CompareAlignedFeatureVectors against CompareFeatureVectors
+//#if DO_CHECK
+//	const float EigenCost = CompareFeatureVectors(A, B, WeightsSqrt);
+//
+//	if (!FMath::IsNearlyEqual(Cost, EigenCost))
+//	{
+//		const float RelativeDifference = Cost > EigenCost ? (Cost - EigenCost) / Cost : (EigenCost - Cost) / EigenCost;
+//		check(FMath::IsNearlyZero(RelativeDifference, UE_KINDA_SMALL_NUMBER));
+//	}
+//#endif //DO_CHECK
+	
+	return Cost;
+}
+
 void CompareFeatureVectors(TConstArrayView<float> A, TConstArrayView<float> B, TConstArrayView<float> WeightsSqrt, TArrayView<float> Result)
 {
 	check(A.Num() == B.Num() && A.Num() == WeightsSqrt.Num() && A.Num() == Result.Num());
@@ -197,6 +243,16 @@ FPoseSearchCost FSearchIndex::ComparePoses(int32 PoseIdx, float ContinuingPoseCo
 	return FPoseSearchCost(DissimilarityCost, NotifyAddend, ContinuingPoseCostBias);
 }
 
+FPoseSearchCost FSearchIndex::CompareAlignedPoses(int32 PoseIdx, float ContinuingPoseCostBias, TConstArrayView<float> PoseValues, TConstArrayView<float> QueryValues) const
+{
+	// base dissimilarity cost representing how the associated PoseIdx differ, in a weighted way, from the query pose (QueryValues)
+	const float DissimilarityCost = CompareAlignedFeatureVectors(PoseValues, QueryValues, WeightsSqrt);
+
+	// cost addend associated to Schema->BaseCostBias or overriden by UAnimNotifyState_PoseSearchModifyCost
+	const float NotifyAddend = PoseMetadata[PoseIdx].GetCostAddend();
+	return FPoseSearchCost(DissimilarityCost, NotifyAddend, ContinuingPoseCostBias);
+}
+
 FArchive& operator<<(FArchive& Ar, FSearchIndex& Index)
 {
 	Ar << static_cast<FSearchIndexBase&>(Index);
diff --git a/Engine/Plugins/Experimental/Animation/PoseSearch/Source/Runtime/Private/PoseSearchLibrary.cpp b/Engine/Plugins/Experimental/Animation/PoseSearch/Source/Runtime/Private/PoseSearchLibrary.cpp
--- a/Engine/Plugins/Experimental/Animation/PoseSearch/Source/Runtime/Private/PoseSearchLibrary.cpp
+++ b/Engine/Plugins/Experimental/Animation/PoseSearch/Source/Runtime/Private/PoseSearchLibrary.cpp
@@ -296,6 +296,8 @@ void UPoseSearchLibrary::UpdateMotionMatchingState(
 	}
 	
 	const FPoseSearchQueryTrajectory TrajectoryRootSpace = ProcessTrajectory(Trajectory, Context.AnimInstanceProxy->GetComponentTransform(), InOutMotionMatchingState.ComponentDeltaYaw, YawFromAnimationTrajectoryBlendTime, TrajectorySpeedMultiplier);
+	
+	FMemMark Mark(FMemStack::Get());
 	FSearchContext SearchContext(&TrajectoryRootSpace, History, 0.f, &InOutMotionMatchingState.PoseIndicesHistory, InOutMotionMatchingState.CurrentSearchResult, PoseJumpThresholdTime, bForceInterrupt);
 
 	const bool bMustSearch = !InOutMotionMatchingState.CurrentSearchResult.IsValid();
@@ -556,7 +558,7 @@ void UPoseSearchLibrary::MotionMatch(
 #endif // ENABLE_DRAW_DEBUG && ENABLE_ANIM_DEBUG
 		}
 
-		// @todo: finish set up SearchContext by exposing or calculating additional members
+		FMemMark Mark(FMemStack::Get());
 		FSearchContext SearchContext(&TrajectoryRootSpace, ExtendedPoseHistory.IsInitialized() ? &ExtendedPoseHistory : nullptr, TimeToFutureAnimationStart);
 
 		FSearchResult SearchResult = Database->Search(SearchContext);
diff --git a/Engine/Plugins/Experimental/Animation/PoseSearch/Source/Runtime/Public/PoseSearch/PoseSearchContext.h b/Engine/Plugins/Experimental/Animation/PoseSearch/Source/Runtime/Public/PoseSearch/PoseSearchContext.h
--- a/Engine/Plugins/Experimental/Animation/PoseSearch/Source/Runtime/Public/PoseSearch/PoseSearchContext.h
+++ b/Engine/Plugins/Experimental/Animation/PoseSearch/Source/Runtime/Public/PoseSearch/PoseSearchContext.h
@@ -188,11 +188,10 @@ private:
 	TConstArrayView<float> CurrentResultNextPoseVector;
 	TConstArrayView<float> CurrentResultPoseVector;
 	TConstArrayView<float> CurrentResultPrevPoseVector;
-	TArray<float> CurrentResultPoseVectorData;
+	TStackAlignedArray<float> CurrentResultPoseVectorData;
 
 	// transforms cached in component space
 	FCachedTransforms<FTransform> CachedTransforms;
-	// @todo: consider using FMemMark
 	TArray<FFeatureVectorBuilder, TInlineAllocator<PreallocatedCachedQueriesNum>> CachedQueries;
 
 	float CurrentBestTotalCost = MAX_flt;
diff --git a/Engine/Plugins/Experimental/Animation/PoseSearch/Source/Runtime/Public/PoseSearch/PoseSearchDefines.h b/Engine/Plugins/Experimental/Animation/PoseSearch/Source/Runtime/Public/PoseSearch/PoseSearchDefines.h
--- a/Engine/Plugins/Experimental/Animation/PoseSearch/Source/Runtime/Public/PoseSearch/PoseSearchDefines.h
+++ b/Engine/Plugins/Experimental/Animation/PoseSearch/Source/Runtime/Public/PoseSearch/PoseSearchDefines.h
@@ -24,9 +24,15 @@ DECLARE_LOG_CATEGORY_EXTERN(LogPoseSearch, Log, All);
 static constexpr int8 RootSchemaBoneIdx = 0;
 static constexpr FBoneIndexType RootBoneIndexType = 0;
 
-static constexpr int32 PreallocatedCachedQuerySize = 64;
 static constexpr int32 PreallocatedCachedQueriesNum = 8;
 
+namespace UE::PoseSearch
+{
+constexpr int32 TMax(int32 A, int32 B) { return (A > B ? A : B); }
+template<typename ElementType> constexpr int32 TAlignOf() { return TMax(alignof(ElementType), 16); }
+template<typename ElementType> using TAlignedArray = TArray<ElementType, TAlignedHeapAllocator<TAlignOf<ElementType>()>>;
+template<typename ElementType> using TStackAlignedArray = TArray<ElementType, TMemStackAllocator<TAlignOf<ElementType>()>>;
+}
 
 
 
diff --git a/Engine/Plugins/Experimental/Animation/PoseSearch/Source/Runtime/Public/PoseSearch/PoseSearchIndex.h b/Engine/Plugins/Experimental/Animation/PoseSearch/Source/Runtime/Public/PoseSearch/PoseSearchIndex.h
--- a/Engine/Plugins/Experimental/Animation/PoseSearch/Source/Runtime/Public/PoseSearch/PoseSearchIndex.h
+++ b/Engine/Plugins/Experimental/Animation/PoseSearch/Source/Runtime/Public/PoseSearch/PoseSearchIndex.h
@@ -4,14 +4,11 @@
 
 #include "PoseSearch/KDTree.h"
 #include "PoseSearch/PoseSearchCost.h"
+#include "PoseSearch/PoseSearchDefines.h"
 
 namespace UE::PoseSearch
 {
 
-constexpr int32 TMax(int32 A, int32 B) { return (A > B ? A : B); }
-template<typename ElementType>
-using TAlignedArray = TArray<ElementType, TAlignedHeapAllocator<TMax(alignof(ElementType), 16)>>;
-
 POSESEARCH_API void CompareFeatureVectors(TConstArrayView<float> A, TConstArrayView<float> B, TConstArrayView<float> WeightsSqrt, TArrayView<float> Result);
 
 /**
@@ -219,6 +216,7 @@ struct FSearchIndex : public FSearchIndexBase
 	TConstArrayView<float> GetReconstructedPoseValues(int32 PoseIdx, TArrayView<float> BufferUsedForReconstruction) const;
 	POSESEARCH_API TArray<float> GetPoseValuesSafe(int32 PoseIdx) const;
 	POSESEARCH_API FPoseSearchCost ComparePoses(int32 PoseIdx, float ContinuingPoseCostBias, TConstArrayView<float> PoseValues, TConstArrayView<float> QueryValues) const;
+	FPoseSearchCost CompareAlignedPoses(int32 PoseIdx, float ContinuingPoseCostBias, TConstArrayView<float> PoseValues, TConstArrayView<float> QueryValues) const;
 
 	friend FArchive& operator<<(FArchive& Ar, FSearchIndex& Index);
 };
diff --git a/Engine/Plugins/Experimental/Animation/PoseSearch/Source/Runtime/Public/PoseSearch/PoseSearchResult.h b/Engine/Plugins/Experimental/Animation/PoseSearch/Source/Runtime/Public/PoseSearch/PoseSearchResult.h
--- a/Engine/Plugins/Experimental/Animation/PoseSearch/Source/Runtime/Public/PoseSearch/PoseSearchResult.h
+++ b/Engine/Plugins/Experimental/Animation/PoseSearch/Source/Runtime/Public/PoseSearch/PoseSearchResult.h
@@ -28,9 +28,8 @@ public:
 	TConstArrayView<float> GetValues() const { return Values; }
 
 private:
+	TStackAlignedArray<float> Values;
 	TObjectPtr<const UPoseSearchSchema> Schema;
-	// @todo: consider using FMemMark
-	TArray<float, TInlineAllocator<PreallocatedCachedQuerySize>> Values;
 };
 	
 struct FSearchResult
