motion matching - fix for KNNResultSet::addPoint

diff --git a/Engine/Plugins/Experimental/Animation/PoseSearch/Source/Runtime/Public/PoseSearch/KDTree.h b/Engine/Plugins/Experimental/Animation/PoseSearch/Source/Runtime/Public/PoseSearch/KDTree.h
--- a/Engine/Plugins/Experimental/Animation/PoseSearch/Source/Runtime/Public/PoseSearch/KDTree.h
+++ b/Engine/Plugins/Experimental/Animation/PoseSearch/Source/Runtime/Public/PoseSearch/KDTree.h
@@ -82,7 +82,7 @@ struct POSESEARCH_API FKDTree
 			return Count == NumNeighbors;
 		}
 
-		inline bool addPoint(size_t dist, size_t index)
+		inline bool addPoint(float dist, size_t index)
 		{
 			if (Algo::BinarySearch(ExcludeFromSearchIndexes, index) != INDEX_NONE)
 			{
