motion matching - fix for ~FDebuggerViewModel cleanup (FDebuggerViewModel::FSkeleton::Actor from TObjectPtr to TWeakObjectPtr)

diff --git a/Engine/Plugins/Experimental/Animation/PoseSearch/Source/Editor/Private/PoseSearchDebuggerViewModel.h b/Engine/Plugins/Experimental/Animation/PoseSearch/Source/Editor/Private/PoseSearchDebuggerViewModel.h
--- a/Engine/Plugins/Experimental/Animation/PoseSearch/Source/Editor/Private/PoseSearchDebuggerViewModel.h
+++ b/Engine/Plugins/Experimental/Animation/PoseSearch/Source/Editor/Private/PoseSearchDebuggerViewModel.h
@@ -62,7 +62,7 @@ private:
 	struct FSkeleton
 	{
 		/** Actor object for the skeleton */
-		TObjectPtr<AActor> Actor;
+		TWeakObjectPtr<AActor> Actor;
 
 		/** Derived skeletal mesh for setting the skeleton in the scene */
 		TWeakObjectPtr<UPoseSearchMeshComponent> Component;
