motion matching - exposing PCA costs in pose search debugger

diff --git a/Engine/Plugins/Experimental/Animation/PoseSearch/Source/Editor/Private/PoseSearchDebuggerDatabaseColumns.h b/Engine/Plugins/Experimental/Animation/PoseSearch/Source/Editor/Private/PoseSearchDebuggerDatabaseColumns.h
--- a/Engine/Plugins/Experimental/Animation/PoseSearch/Source/Editor/Private/PoseSearchDebuggerDatabaseColumns.h
+++ b/Engine/Plugins/Experimental/Animation/PoseSearch/Source/Editor/Private/PoseSearchDebuggerDatabaseColumns.h
@@ -379,6 +379,36 @@ struct FCost : ITextColumn
 	}
 };
 
+struct FPCACost : ITextColumn
+{
+	using ITextColumn::ITextColumn;
+
+	virtual FText GetLabel() const override
+	{
+		return LOCTEXT("ColumnLabelPCACost", "PCA Cost");
+	}
+	
+	virtual FText GetLabelTooltip() const override
+	{
+		return LOCTEXT("ColumnLabelTooltipPCACost", "Total PCA Cost of the associated Pose");
+	}
+
+	virtual FSortPredicate GetSortPredicate() const override
+	{
+		return [](const FRowDataRef& Row0, const FRowDataRef& Row1) -> bool { return Row0->PosePCACost < Row1->PosePCACost; };
+	}
+		
+	virtual FText GetRowText(const FRowDataRef& Row) const override
+	{
+		return FText::AsNumber(Row->PosePCACost);
+    }
+
+	virtual FSlateColor GetColorAndOpacity(const FRowDataRef& Row) const override
+	{
+		return Row->PCACostColor;
+	}
+};
+
 struct FChannelBreakdownCostColumn : ITextColumn
 {
 	FChannelBreakdownCostColumn(int32 SortIndex, int32 InBreakdownCostIndex, const FText& InLabel)
diff --git a/Engine/Plugins/Experimental/Animation/PoseSearch/Source/Editor/Private/PoseSearchDebuggerDatabaseRowData.h b/Engine/Plugins/Experimental/Animation/PoseSearch/Source/Editor/Private/PoseSearchDebuggerDatabaseRowData.h
--- a/Engine/Plugins/Experimental/Animation/PoseSearch/Source/Editor/Private/PoseSearchDebuggerDatabaseRowData.h
+++ b/Engine/Plugins/Experimental/Animation/PoseSearch/Source/Editor/Private/PoseSearchDebuggerDatabaseRowData.h
@@ -15,7 +15,8 @@ public:
 	TWeakObjectPtr<const UPoseSearchDatabase> SourceDatabase;
 	FString DatabaseName = "";
 	FString DatabasePath = "";
-	TArray<float> QueryVector;
+	TAlignedArray<float> QueryVector;
+	TAlignedArray<float> PCAQueryVector;
 };
 
 class FDebuggerDatabaseRowData : public TSharedFromThis<FDebuggerDatabaseRowData>
@@ -36,6 +37,8 @@ public:
 	FVector BlendParameters = FVector::Zero();
 	FPoseSearchCost PoseCost;
 	FLinearColor CostColor = FLinearColor::White;
+	float PosePCACost = 0.f;
+	FLinearColor PCACostColor = FLinearColor::White;
 	TArray<float> CostBreakdowns;
 	TArray<FLinearColor> CostBreakdownsColors;
 	TArray<float> CostVector;
diff --git a/Engine/Plugins/Experimental/Animation/PoseSearch/Source/Editor/Private/PoseSearchDebuggerDatabaseView.cpp b/Engine/Plugins/Experimental/Animation/PoseSearch/Source/Editor/Private/PoseSearchDebuggerDatabaseView.cpp
--- a/Engine/Plugins/Experimental/Animation/PoseSearch/Source/Editor/Private/PoseSearchDebuggerDatabaseView.cpp
+++ b/Engine/Plugins/Experimental/Animation/PoseSearch/Source/Editor/Private/PoseSearchDebuggerDatabaseView.cpp
@@ -14,6 +14,7 @@
 #include "PoseSearch/PoseSearchDerivedData.h"
 #include "PoseSearch/PoseSearchFeatureChannel.h"
 #include "PoseSearch/PoseSearchSchema.h"
+#include "PoseSearchEditor.h"
 #include "Trace/PoseSearchTraceProvider.h"
 #include "Widgets/Input/SCheckBox.h"
 #include "Widgets/Layout/SScrollBar.h"
@@ -175,20 +176,40 @@ static void AddUnfilteredDatabaseRow(const UPoseSearchDatabase* Database,
 		// in case we modify the schema while PIE is paused and displaying the Pose Search Editor, we could end up with a stale State with a SharedData->QueryVector saved with the previous schema
 		// so the cardinality of SharedData->QueryVector and PoseValues don't match. In that case we just use PoseValues as query to have all costs set to zero
 		const bool bIsQueryVectorValid = SharedData->QueryVector.Num() == PoseValues.Num();
-		const TArray<float>& QueryVector = bIsQueryVectorValid ? SharedData->QueryVector : PoseValues;
+		TConstArrayView<float> QueryVector = bIsQueryVectorValid ? TConstArrayView<float>(SharedData->QueryVector) : TConstArrayView<float>(PoseValues);
 
 		CompareFeatureVectors(PoseValues, QueryVector, SearchIndex.WeightsSqrt, Row->CostVector);
 
 		if (Cost.IsValid())
 		{
 			Row->PoseCost = Cost;
+
+			// @todo: perhaps reuse this code to recalculate Row->PoseCost instead of using Cost
+			//const FPoseSearchCost RecalculatedCost = SearchIndex.ComparePoses(DbPoseIdx, 0.f, PoseValues, QueryVector);
+			//if (Cost.GetTotalCost() != RecalculatedCost.GetTotalCost())
+			//{
+			//	UE_LOG(LogPoseSearchEditor, Error, TEXT("RecalculatedCost (%f) differs from Cost (%f) for Pose '%d'"), RecalculatedCost.GetTotalCost(), Cost.GetTotalCost(), DbPoseIdx);
+			//}
+		}
+		// is PoseValues padded at 16 bytes (and 16 bytes aligned by construction)?
+		else if (PoseValues.Num() % 4 == 0)
+		{
+			Row->PoseCost = SearchIndex.CompareAlignedPoses(DbPoseIdx, 0.f, PoseValues, QueryVector);
 		}
 		else
 		{
-			// @todo: perhaps reuse CompareFeatureVectors cost calculation
 			Row->PoseCost = SearchIndex.ComparePoses(DbPoseIdx, 0.f, PoseValues, QueryVector);
 		}
 
+		if (!SharedData->PCAQueryVector.IsEmpty())
+		{
+			TConstArrayView<float> PCAPoseValues = SearchIndex.GetPCAPoseValues(DbPoseIdx);
+			if (SharedData->PCAQueryVector.Num() == PCAPoseValues.Num())
+			{
+				Row->PosePCACost = CompareFeatureVectors(SharedData->PCAQueryVector, PCAPoseValues);
+			}
+		}
+
 		const FInstancedStruct& DatabaseAssetStruct = Database->GetAnimationAssetStruct(*SearchIndexAsset);
 		if (const FPoseSearchDatabaseAnimationAssetBase* DatabaseAsset = DatabaseAssetStruct.GetPtr<FPoseSearchDatabaseAnimationAssetBase>())
 		{
@@ -231,6 +252,7 @@ void SDebuggerDatabaseView::Update(const FTraceMotionMatchingStateMessage& State
 
 	UnfilteredDatabaseRows.Reset();
 
+	bool bAddPCACost = false;
 	for (const FTraceMotionMatchingStateDatabaseEntry& DbEntry : State.DatabaseEntries)
 	{
 		const UPoseSearchDatabase* Database = FTraceMotionMatchingState::GetObjectFromId<UPoseSearchDatabase>(DbEntry.DatabaseId);
@@ -242,6 +264,13 @@ void SDebuggerDatabaseView::Update(const FTraceMotionMatchingStateMessage& State
 			SharedData->DatabasePath = Database->GetPathName();
 			SharedData->QueryVector = DbEntry.QueryVector;
 
+			if (Database->PoseSearchMode != EPoseSearchMode::BruteForce)
+			{
+				SharedData->PCAQueryVector.SetNumZeroed(Database->GetNumberOfPrincipalComponents());
+				Database->GetSearchIndex().PCAProject(DbEntry.QueryVector, SharedData->PCAQueryVector);
+				bAddPCACost = true;
+			}
+
 			for (const FTraceMotionMatchingStatePoseEntry& PoseEntry : DbEntry.PoseEntries)
 			{
 				AddUnfilteredDatabaseRow(Database, UnfilteredDatabaseRows, SharedData, PoseEntry.DbPoseIdx, PoseEntry.PoseCandidateFlags, PoseEntry.Cost);
@@ -272,87 +301,20 @@ void SDebuggerDatabaseView::Update(const FTraceMotionMatchingStateMessage& State
 	{
 		CostBreakDownData.ProcessData(UnfilteredDatabaseRows);
 
-		// calculating breakdowns min max and colors
+		// calculating min/max for Cost, PCACost, and breakdowns
+		float MinCost = UE_MAX_FLT;
+		float MaxCost = -UE_MAX_FLT;
+
+		float MinPCACost = UE_MAX_FLT;
+		float MaxPCACost = -UE_MAX_FLT;
+
 		TArray<float> MinCostBreakdowns;
 		TArray<float> MaxCostBreakdowns;
 
 		const int32 CostBreakdownsCardinality = UnfilteredDatabaseRows[0]->CostBreakdowns.Num();
 		MinCostBreakdowns.Init(UE_MAX_FLT, CostBreakdownsCardinality);
 		MaxCostBreakdowns.Init(-UE_MAX_FLT, CostBreakdownsCardinality);
-
-		auto ArrayMinMax = [](TConstArrayView<float> View, TArrayView<float> Min, TArrayView<float> Max, float InvalidValue)
-		{
-			const int32 Num = View.Num();
-			check(Num == Min.Num() && Num == Max.Num());
-			for (int i = 0; i < Num; ++i)
-			{
-				const float Value = View[i];
-				if (Value != InvalidValue)
-				{
-					Min[i] = FMath::Min(Min[i], Value);
-					Max[i] = FMath::Max(Max[i], Value);
-				}
-			}
-		};
-
-		for (TSharedRef<FDebuggerDatabaseRowData>& UnfilteredRow : UnfilteredDatabaseRows)
-		{
-			if (EnumHasAnyFlags(UnfilteredRow->PoseCandidateFlags, EPoseCandidateFlags::AnyValidMask))
-			{
-				ArrayMinMax(UnfilteredRow->CostBreakdowns, MinCostBreakdowns, MaxCostBreakdowns, UE_MAX_FLT);
-			}
-		}
 		
-		auto ArraySafeNormalize = [](TConstArrayView<float> View, TConstArrayView<float> Min, TConstArrayView<float> Max, TArrayView<float> NormalizedView)
-		{
-			const int32 Num = View.Num();
-			check(Num == Min.Num() && Num == Max.Num() && Num == NormalizedView.Num());
-			for (int i = 0; i < Num; ++i)
-			{
-				const float Delta = Max[i] - Min[i];
-				if (FMath::IsNearlyZero(Delta, UE_KINDA_SMALL_NUMBER))
-				{
-					NormalizedView[i] = 0.f;
-				}
-				else
-				{
-					NormalizedView[i] = (View[i] - Min[i]) / Delta;
-				}
-			}
-		};
-
-		auto LinearColorBlend = [](FLinearColor LinearColorA, FLinearColor LinearColorB, float BlendParam) -> FLinearColor
-		{
-			return LinearColorA + (LinearColorB - LinearColorA) * BlendParam;
-		};
-
-		auto LinearColorArrayBlend = [](FLinearColor LinearColorA, FLinearColor LinearColorB, TConstArrayView<float> BlendParam, TArray<FLinearColor>& BlendedColors) -> void
-		{
-			const int32 Num = BlendParam.Num();
-			BlendedColors.SetNumUninitialized(Num);
-			for (int i = 0; i < Num; ++i)
-			{
-				BlendedColors[i] = LinearColorA + (LinearColorB - LinearColorA) * BlendParam[i];
-			}
-		};
-
-		TArray<float> CostBreakdownsColorBlend;
-		CostBreakdownsColorBlend.Init(0, CostBreakdownsCardinality);
-		for (TSharedRef<FDebuggerDatabaseRowData>& UnfilteredRow : UnfilteredDatabaseRows)
-		{
-			if (EnumHasAnyFlags(UnfilteredRow->PoseCandidateFlags, EPoseCandidateFlags::AnyValidMask))
-			{
-				ArraySafeNormalize(UnfilteredRow->CostBreakdowns, MinCostBreakdowns, MaxCostBreakdowns, CostBreakdownsColorBlend);
-				LinearColorArrayBlend(BestScoreRowColor, WorstScoreRowColor, CostBreakdownsColorBlend, UnfilteredRow->CostBreakdownsColors);
-			}
-			else
-			{
-				UnfilteredRow->CostBreakdownsColors.Init(DiscardedRowColor, CostBreakdownsCardinality);
-			}
-		}
-
-		float MinCost = UE_MAX_FLT;
-		float MaxCost = -UE_MAX_FLT;
 		for (TSharedRef<FDebuggerDatabaseRowData>& UnfilteredRow : UnfilteredDatabaseRows)
 		{
 			if (EnumHasAnyFlags(UnfilteredRow->PoseCandidateFlags, EPoseCandidateFlags::AnyValidMask))
@@ -360,26 +322,64 @@ void SDebuggerDatabaseView::Update(const FTraceMotionMatchingStateMessage& State
 				const float Cost = UnfilteredRow->PoseCost.GetTotalCost();
 				MinCost = FMath::Min(MinCost, Cost);
 				MaxCost = FMath::Max(MaxCost, Cost);
+
+				const float PCACost = UnfilteredRow->PosePCACost;
+				MinPCACost = FMath::Min(MinPCACost, PCACost);
+				MaxPCACost = FMath::Max(MaxPCACost, PCACost);
+
+				for (int Index = 0; Index < CostBreakdownsCardinality; ++Index)
+				{
+					const float Value = UnfilteredRow->CostBreakdowns[Index];
+					if (Value != UE_MAX_FLT)
+					{
+						MinCostBreakdowns[Index] = FMath::Min(MinCostBreakdowns[Index], Value);
+						MaxCostBreakdowns[Index] = FMath::Max(MaxCostBreakdowns[Index], Value);
+					}
+				}
 			}
 		}
+		
+		// calculating min max deltas (always >= UE_KINDA_SMALL_NUMBER to avoid division by zero)
+		const float DeltaCost = FMath::Max(MaxCost - MinCost, UE_KINDA_SMALL_NUMBER);
+		const float DeltaPCACost = FMath::Max(MaxPCACost - MinPCACost, UE_KINDA_SMALL_NUMBER);
+		TArray<float> DeltaCostBreakdowns;
+		DeltaCostBreakdowns.SetNumUninitialized(CostBreakdownsCardinality);
+		for (int32 Index = 0; Index < CostBreakdownsCardinality; ++Index)
+		{
+			DeltaCostBreakdowns[Index] = FMath::Max(MaxCostBreakdowns[Index] - MinCostBreakdowns[Index], UE_KINDA_SMALL_NUMBER);
+		}
 
-		const float DeltaCost = MaxCost - MinCost;
+		// calculating colors for Cost, PCACost, and breakdowns (using min/max as range)
 		for (TSharedRef<FDebuggerDatabaseRowData>& UnfilteredRow : UnfilteredDatabaseRows)
 		{
 			if (EnumHasAnyFlags(UnfilteredRow->PoseCandidateFlags, EPoseCandidateFlags::AnyValidMask))
 			{
-				const float CostColorBlend = DeltaCost > UE_KINDA_SMALL_NUMBER ? (UnfilteredRow->PoseCost.GetTotalCost() - MinCost) / DeltaCost : 0.f;
-				UnfilteredRow->CostColor = LinearColorBlend(BestScoreRowColor, WorstScoreRowColor, CostColorBlend);
+				const float CostBlend = (UnfilteredRow->PoseCost.GetTotalCost() - MinCost) / DeltaCost;
+				UnfilteredRow->CostColor = BestScoreRowColor + (WorstScoreRowColor - BestScoreRowColor) * CostBlend;
+
+				const float PCACostBlend = (UnfilteredRow->PosePCACost - MinPCACost) / DeltaPCACost;
+				UnfilteredRow->PCACostColor = BestScoreRowColor + (WorstScoreRowColor - BestScoreRowColor) * PCACostBlend;
+
+				UnfilteredRow->CostBreakdownsColors.SetNumUninitialized(CostBreakdownsCardinality);
+				for (int Index = 0; Index < CostBreakdownsCardinality; ++Index)
+				{
+					const float CostBreakdownBlend = (UnfilteredRow->CostBreakdowns[Index] - MinCostBreakdowns[Index]) / DeltaCostBreakdowns[Index];
+					UnfilteredRow->CostBreakdownsColors[Index] = BestScoreRowColor + (WorstScoreRowColor - BestScoreRowColor) * CostBreakdownBlend;
+				}
 			}
 			else
 			{
 				UnfilteredRow->CostColor = DiscardedRowColor;
+				UnfilteredRow->PCACostColor = DiscardedRowColor;
+				UnfilteredRow->CostBreakdownsColors.Init(DiscardedRowColor, CostBreakdownsCardinality);
 			}
 		}
 	}
 
-	if (!CostBreakDownData.AreLabelsEqualTo(OldLabels))
+	if (bHasPCACostColumn != bAddPCACost || !CostBreakDownData.AreLabelsEqualTo(OldLabels))
 	{
+		bHasPCACostColumn = bAddPCACost;
+
 		OldLabels = CostBreakDownData.GetLabels();
 
 		// recreating and binding the columns
@@ -393,6 +393,11 @@ void SDebuggerDatabaseView::Update(const FTraceMotionMatchingStateMessage& State
 		TSharedRef<FCost> CostColumn = MakeShared<FCost>(ColumnIdx++);
 		AddColumn(CostColumn);
 
+		if (bAddPCACost)
+		{
+			AddColumn(MakeShared<FPCACost>(ColumnIdx++));
+		}
+
 		int32 LabelIdx = 0;
 		for (const FText& Label : CostBreakDownData.GetLabels())
 		{
@@ -1082,6 +1087,7 @@ void SDebuggerDatabaseView::Construct(const FArguments& InArgs)
 	SortMode = EColumnSortMode::Ascending;
 	OldLabels.Reset();
 	Columns.Reset();
+	bHasPCACostColumn = false;
 }
 
 } // namespace UE::PoseSearch
diff --git a/Engine/Plugins/Experimental/Animation/PoseSearch/Source/Editor/Private/PoseSearchDebuggerDatabaseView.h b/Engine/Plugins/Experimental/Animation/PoseSearch/Source/Editor/Private/PoseSearchDebuggerDatabaseView.h
--- a/Engine/Plugins/Experimental/Animation/PoseSearch/Source/Editor/Private/PoseSearchDebuggerDatabaseView.h
+++ b/Engine/Plugins/Experimental/Animation/PoseSearch/Source/Editor/Private/PoseSearchDebuggerDatabaseView.h
@@ -101,6 +101,7 @@ private:
     FColumnMap Columns;
 
 	TArray<FText> OldLabels;
+	bool bHasPCACostColumn = false;
 
 	struct FTable
 	{
diff --git a/Engine/Plugins/Experimental/Animation/PoseSearch/Source/Runtime/Private/PoseSearchDatabase.cpp b/Engine/Plugins/Experimental/Animation/PoseSearch/Source/Runtime/Private/PoseSearchDatabase.cpp
--- a/Engine/Plugins/Experimental/Animation/PoseSearch/Source/Runtime/Private/PoseSearchDatabase.cpp
+++ b/Engine/Plugins/Experimental/Animation/PoseSearch/Source/Runtime/Private/PoseSearchDatabase.cpp
@@ -14,7 +14,6 @@
 #include "PoseSearch/PoseSearchFeatureChannel.h"
 #include "PoseSearch/PoseSearchHistory.h"
 #include "PoseSearch/PoseSearchSchema.h"
-#include "PoseSearchEigenHelper.h"
 #include "UObject/ObjectSaveContext.h"
 
 DECLARE_STATS_GROUP(TEXT("PoseSearch"), STATGROUP_PoseSearch, STATCAT_Advanced);
@@ -831,31 +830,17 @@ UE::PoseSearch::FSearchResult UPoseSearchDatabase::SearchPCAKDTree(UE::PoseSearc
 	//stack allocated temporaries
 	TArrayView<size_t> ResultIndexes((size_t*)FMemory_Alloca((ClampedKDTreeQueryNumNeighbors + 1) * sizeof(size_t)), ClampedKDTreeQueryNumNeighbors + 1);
 	TArrayView<float> ResultDistanceSqr((float*)FMemory_Alloca((ClampedKDTreeQueryNumNeighbors + 1) * sizeof(float)), ClampedKDTreeQueryNumNeighbors + 1);
-	RowMajorVectorMap WeightedQueryValues((float*)FMemory_Alloca(NumDimensions * sizeof(float)), 1, NumDimensions);
-	RowMajorVectorMap CenteredQueryValues((float*)FMemory_Alloca(NumDimensions * sizeof(float)), 1, NumDimensions);
-	RowMajorVectorMap ProjectedQueryValues((float*)FMemory_Alloca(ClampedNumberOfPrincipalComponents * sizeof(float)), 1, ClampedNumberOfPrincipalComponents);
+	TArrayView<float> ProjectedQueryValues((float*)FMemory_Alloca(ClampedNumberOfPrincipalComponents * sizeof(float)), ClampedNumberOfPrincipalComponents);
 
 #if DO_CHECK
 	// KDTree in PCA space search
 	if (PoseSearchMode == EPoseSearchMode::PCAKDTree_Validate)
 	{
-		const RowMajorVectorMapConst MapWeightsSqrt(SearchIndex.WeightsSqrt.GetData(), 1, NumDimensions);
-
 		// testing the KDTree is returning the proper searches for all the original points transformed in pca space
 		for (int32 PoseIdx = 0; PoseIdx < SearchIndex.GetNumPoses(); ++PoseIdx)
 		{
 			FKDTree::KNNResultSet ResultSet(ClampedKDTreeQueryNumNeighbors, ResultIndexes, ResultDistanceSqr);
-			TConstArrayView<float> PoseValues = SearchIndex.GetPoseValues(PoseIdx);
-
-			const RowMajorVectorMapConst Mean(SearchIndex.Mean.GetData(), 1, NumDimensions);
-			const ColMajorMatrixMapConst PCAProjectionMatrix(SearchIndex.PCAProjectionMatrix.GetData(), NumDimensions, ClampedNumberOfPrincipalComponents);
-
-			const RowMajorVectorMapConst QueryValues(PoseValues.GetData(), 1, NumDimensions);
-			WeightedQueryValues = QueryValues.array() * MapWeightsSqrt.array();
-			CenteredQueryValues.noalias() = WeightedQueryValues - Mean;
-			ProjectedQueryValues.noalias() = CenteredQueryValues * PCAProjectionMatrix;
-
-			SearchIndex.KDTree.FindNeighbors(ResultSet, ProjectedQueryValues.data());
+			SearchIndex.KDTree.FindNeighbors(ResultSet, SearchIndex.PCAProject(SearchIndex.GetPoseValues(PoseIdx), ProjectedQueryValues).GetData());
 
 			size_t ResultIndex = 0;
 			for (; ResultIndex < ResultSet.Num(); ++ResultIndex)
@@ -886,18 +871,8 @@ UE::PoseSearch::FSearchResult UPoseSearchDatabase::SearchPCAKDTree(UE::PoseSearc
 		FKDTree::KNNResultSet ResultSet(ClampedKDTreeQueryNumNeighbors, ResultIndexes, ResultDistanceSqr, NonSelectableIdx);
 
 		check(QueryValues.Num() == NumDimensions);
-
-		const RowMajorVectorMapConst Mean(SearchIndex.Mean.GetData(), 1, NumDimensions);
-		const ColMajorMatrixMapConst PCAProjectionMatrix(SearchIndex.PCAProjectionMatrix.GetData(), NumDimensions, ClampedNumberOfPrincipalComponents);
-
-		// transforming query values into PCA space to query the KDTree
-		const RowMajorVectorMapConst QueryValuesMap(QueryValues.GetData(), 1, NumDimensions);
-		const RowMajorVectorMapConst MapWeightsSqrt(SearchIndex.WeightsSqrt.GetData(), 1, NumDimensions);
-		WeightedQueryValues = QueryValuesMap.array() * MapWeightsSqrt.array();
-		CenteredQueryValues.noalias() = WeightedQueryValues - Mean;
-		ProjectedQueryValues.noalias() = CenteredQueryValues * PCAProjectionMatrix;
-
-		SearchIndex.KDTree.FindNeighbors(ResultSet, ProjectedQueryValues.data());
+		// projecting QueryValues into the PCA space ProjectedQueryValues and query the KDTree
+		SearchIndex.KDTree.FindNeighbors(ResultSet, SearchIndex.PCAProject(QueryValues, ProjectedQueryValues).GetData());
 
 		// NonSelectableIdx are already filtered out inside the kdtree search
 		const FSearchFilters SearchFilters(Schema, TConstArrayView<size_t>(), SearchIndex.bAnyBlockTransition);
diff --git a/Engine/Plugins/Experimental/Animation/PoseSearch/Source/Runtime/Private/PoseSearchIndex.cpp b/Engine/Plugins/Experimental/Animation/PoseSearch/Source/Runtime/Private/PoseSearchIndex.cpp
--- a/Engine/Plugins/Experimental/Animation/PoseSearch/Source/Runtime/Private/PoseSearchIndex.cpp
+++ b/Engine/Plugins/Experimental/Animation/PoseSearch/Source/Runtime/Private/PoseSearchIndex.cpp
@@ -75,6 +75,16 @@ void CompareFeatureVectors(TConstArrayView<float> A, TConstArrayView<float> B, T
 	VR = ((VA - VB) * VW).square();
 }
 
+float CompareFeatureVectors(TConstArrayView<float> A, TConstArrayView<float> B)
+{
+	check(A.Num() == B.Num() && A.Num());
+
+	Eigen::Map<const Eigen::ArrayXf> VA(A.GetData(), A.Num());
+	Eigen::Map<const Eigen::ArrayXf> VB(B.GetData(), B.Num());
+
+	return (VA - VB).square().sum();
+}
+
 //////////////////////////////////////////////////////////////////////////
 // FPoseMetadata
 FArchive& operator<<(FArchive& Ar, FPoseMetadata& Metadata)
@@ -184,10 +194,10 @@ void FSearchIndex::Reset()
 
 TConstArrayView<float> FSearchIndex::GetPoseValues(int32 PoseIdx) const
 {
-	const int32 SchemaCardinality = WeightsSqrt.Num();
-	check(!Values.IsEmpty() && PoseIdx >= 0 && PoseIdx < GetNumPoses() && SchemaCardinality > 0);
-	const int32 ValueOffset = PoseIdx * SchemaCardinality;
-	return MakeArrayView(&Values[ValueOffset], SchemaCardinality);
+	const int32 NumDimensions = WeightsSqrt.Num();
+	check(!Values.IsEmpty() && PoseIdx >= 0 && PoseIdx < GetNumPoses() && NumDimensions > 0);
+	const int32 ValueOffset = PoseIdx * NumDimensions;
+	return MakeArrayView(&Values[ValueOffset], NumDimensions);
 }
 
 TConstArrayView<float> FSearchIndex::GetReconstructedPoseValues(int32 PoseIdx, TArrayView<float> BufferUsedForReconstruction) const
@@ -214,6 +224,31 @@ TConstArrayView<float> FSearchIndex::GetReconstructedPoseValues(int32 PoseIdx, T
 	return BufferUsedForReconstruction;
 }
 
+TConstArrayView<float> FSearchIndex::PCAProject(TConstArrayView<float> PoseValues, TArrayView<float> BufferUsedForProjection) const
+{
+	const int32 NumDimensions = WeightsSqrt.Num();
+	const int32 NumberOfPrincipalComponents = PCAProjectionMatrix.Num() / NumDimensions;
+
+	check(PCAProjectionMatrix.Num() > 0 && PCAProjectionMatrix.Num() % NumDimensions == 0);
+	check(BufferUsedForProjection.Num() == NumberOfPrincipalComponents);
+
+	const RowMajorVectorMapConst WeightsSqrtMap(WeightsSqrt.GetData(), 1, NumDimensions);
+	const RowMajorVectorMapConst MeanMap(Mean.GetData(), 1, NumDimensions);
+	const ColMajorMatrixMapConst PCAProjectionMatrixMap(PCAProjectionMatrix.GetData(), NumDimensions, NumberOfPrincipalComponents);
+	const RowMajorVectorMapConst PoseValuesMap(PoseValues.GetData(), 1, NumDimensions);
+
+	RowMajorVectorMap WeightedPoseValuesMap((float*)FMemory_Alloca(NumDimensions * sizeof(float)), 1, NumDimensions);
+	WeightedPoseValuesMap = PoseValuesMap.array() * WeightsSqrtMap.array();
+
+	RowMajorVectorMap CenteredPoseValuesMap((float*)FMemory_Alloca(NumDimensions * sizeof(float)), 1, NumDimensions);
+	CenteredPoseValuesMap.noalias() = WeightedPoseValuesMap - MeanMap;
+
+	RowMajorVectorMap ProjectedPoseValuesMap(BufferUsedForProjection.GetData(), 1, NumberOfPrincipalComponents);
+	ProjectedPoseValuesMap.noalias() = CenteredPoseValuesMap * PCAProjectionMatrixMap;
+
+	return BufferUsedForProjection;
+}
+
 TArray<float> FSearchIndex::GetPoseValuesSafe(int32 PoseIdx) const
 {
 	TArray<float> PoseValues;
@@ -233,6 +268,23 @@ TArray<float> FSearchIndex::GetPoseValuesSafe(int32 PoseIdx) const
 	return PoseValues;
 }
 
+TConstArrayView<float> FSearchIndex::GetPCAPoseValues(int32 PoseIdx) const
+{
+	if (PCAValues.IsEmpty())
+	{
+		return TConstArrayView<float>();
+	}
+
+	const int32 NumDimensions = WeightsSqrt.Num();
+	const int32 NumberOfPrincipalComponents = PCAProjectionMatrix.Num() / NumDimensions;
+
+	check(PCAProjectionMatrix.Num() > 0 && PCAProjectionMatrix.Num() % NumDimensions == 0);
+	check(PoseIdx >= 0 && PoseIdx < GetNumPoses() && NumDimensions > 0);
+
+	const int32 ValueOffset = PoseIdx * NumberOfPrincipalComponents;
+	return MakeArrayView(&PCAValues[ValueOffset], NumberOfPrincipalComponents);
+}
+
 FPoseSearchCost FSearchIndex::ComparePoses(int32 PoseIdx, float ContinuingPoseCostBias, TConstArrayView<float> PoseValues, TConstArrayView<float> QueryValues) const
 {
 	// base dissimilarity cost representing how the associated PoseIdx differ, in a weighted way, from the query pose (QueryValues)
diff --git a/Engine/Plugins/Experimental/Animation/PoseSearch/Source/Runtime/Public/PoseSearch/PoseSearchIndex.h b/Engine/Plugins/Experimental/Animation/PoseSearch/Source/Runtime/Public/PoseSearch/PoseSearchIndex.h
--- a/Engine/Plugins/Experimental/Animation/PoseSearch/Source/Runtime/Public/PoseSearch/PoseSearchIndex.h
+++ b/Engine/Plugins/Experimental/Animation/PoseSearch/Source/Runtime/Public/PoseSearch/PoseSearchIndex.h
@@ -10,6 +10,7 @@ namespace UE::PoseSearch
 {
 
 POSESEARCH_API void CompareFeatureVectors(TConstArrayView<float> A, TConstArrayView<float> B, TConstArrayView<float> WeightsSqrt, TArrayView<float> Result);
+POSESEARCH_API float CompareFeatureVectors(TConstArrayView<float> A, TConstArrayView<float> B);
 
 /**
  * This is kept for each pose in the search index along side the feature vector values and is used to influence the search.
@@ -214,9 +215,12 @@ struct FSearchIndex : public FSearchIndexBase
 	void Reset();
 	TConstArrayView<float> GetPoseValues(int32 PoseIdx) const;
 	TConstArrayView<float> GetReconstructedPoseValues(int32 PoseIdx, TArrayView<float> BufferUsedForReconstruction) const;
+	POSESEARCH_API TConstArrayView<float> PCAProject(TConstArrayView<float> PoseValues, TArrayView<float> BufferUsedForProjection) const;
+
 	POSESEARCH_API TArray<float> GetPoseValuesSafe(int32 PoseIdx) const;
+	POSESEARCH_API TConstArrayView<float> GetPCAPoseValues(int32 PoseIdx) const;
 	POSESEARCH_API FPoseSearchCost ComparePoses(int32 PoseIdx, float ContinuingPoseCostBias, TConstArrayView<float> PoseValues, TConstArrayView<float> QueryValues) const;
-	FPoseSearchCost CompareAlignedPoses(int32 PoseIdx, float ContinuingPoseCostBias, TConstArrayView<float> PoseValues, TConstArrayView<float> QueryValues) const;
+	POSESEARCH_API FPoseSearchCost CompareAlignedPoses(int32 PoseIdx, float ContinuingPoseCostBias, TConstArrayView<float> PoseValues, TConstArrayView<float> QueryValues) const;
 
 	friend FArchive& operator<<(FArchive& Ar, FSearchIndex& Index);
 };
