motion matching - logging errors in case of missing pose history node, when required

diff --git a/Engine/Plugins/Experimental/Animation/PoseSearch/Source/Runtime/Private/PoseSearchFeatureChannel_FilterCrashingLegs.cpp b/Engine/Plugins/Experimental/Animation/PoseSearch/Source/Runtime/Private/PoseSearchFeatureChannel_FilterCrashingLegs.cpp
--- a/Engine/Plugins/Experimental/Animation/PoseSearch/Source/Runtime/Private/PoseSearchFeatureChannel_FilterCrashingLegs.cpp
+++ b/Engine/Plugins/Experimental/Animation/PoseSearch/Source/Runtime/Private/PoseSearchFeatureChannel_FilterCrashingLegs.cpp
@@ -63,7 +63,11 @@ void UPoseSearchFeatureChannel_FilterCrashingLegs::BuildQuery(UE::PoseSearch::FS
 			const float LerpValue = InputQueryPose == EInputQueryPose::UseInterpolatedContinuingPose ? SearchContext.GetCurrentResult().LerpValue : 0.f;
 			FFeatureVectorHelper::EncodeFloat(InOutQuery.EditValues(), ChannelDataOffset, SearchContext.GetCurrentResultPrevPoseVector(), SearchContext.GetCurrentResultPoseVector(), SearchContext.GetCurrentResultNextPoseVector(), LerpValue);
 		}
-		// else leave the InOutQuery set to zero since the SearchContext.History is invalid and it'll fail if we continue
+		else
+		{
+			// we leave the InOutQuery set to zero since the SearchContext.History is invalid and it'll fail if we continue
+			UE_LOG(LogPoseSearch, Error, TEXT("UPoseSearchFeatureChannel_FilterCrashingLegs::BuildQuery - Failed because Pose History Node is missing."));
+		}
 	}
 	else
 	{
diff --git a/Engine/Plugins/Experimental/Animation/PoseSearch/Source/Runtime/Private/PoseSearchFeatureChannel_Heading.cpp b/Engine/Plugins/Experimental/Animation/PoseSearch/Source/Runtime/Private/PoseSearchFeatureChannel_Heading.cpp
--- a/Engine/Plugins/Experimental/Animation/PoseSearch/Source/Runtime/Private/PoseSearchFeatureChannel_Heading.cpp
+++ b/Engine/Plugins/Experimental/Animation/PoseSearch/Source/Runtime/Private/PoseSearchFeatureChannel_Heading.cpp
@@ -55,7 +55,11 @@ void UPoseSearchFeatureChannel_Heading::BuildQuery(UE::PoseSearch::FSearchContex
 			// @todo: we should normalize if LerpValue != 0
 			FFeatureVectorHelper::EncodeVector(InOutQuery.EditValues(), ChannelDataOffset, SearchContext.GetCurrentResultPrevPoseVector(), SearchContext.GetCurrentResultPoseVector(), SearchContext.GetCurrentResultNextPoseVector(), LerpValue, true, ComponentStripping);
 		}
-		// else leave the InOutQuery set to zero since the SearchContext.History is invalid and it'll fail if we continue
+		else
+		{
+			// we leave the InOutQuery set to zero since the SearchContext.History is invalid and it'll fail if we continue
+			UE_LOG(LogPoseSearch, Error, TEXT("UPoseSearchFeatureChannel_Heading::BuildQuery - Failed because Pose History Node is missing."));
+		}
 	}
 	else
 	{
diff --git a/Engine/Plugins/Experimental/Animation/PoseSearch/Source/Runtime/Private/PoseSearchFeatureChannel_Phase.cpp b/Engine/Plugins/Experimental/Animation/PoseSearch/Source/Runtime/Private/PoseSearchFeatureChannel_Phase.cpp
--- a/Engine/Plugins/Experimental/Animation/PoseSearch/Source/Runtime/Private/PoseSearchFeatureChannel_Phase.cpp
+++ b/Engine/Plugins/Experimental/Animation/PoseSearch/Source/Runtime/Private/PoseSearchFeatureChannel_Phase.cpp
@@ -309,7 +309,11 @@ void UPoseSearchFeatureChannel_Phase::BuildQuery(UE::PoseSearch::FSearchContext&
 			const float LerpValue = InputQueryPose == EInputQueryPose::UseInterpolatedContinuingPose ? SearchContext.GetCurrentResult().LerpValue : 0.f;
 			FFeatureVectorHelper::EncodeVector2D(InOutQuery.EditValues(), ChannelDataOffset, SearchContext.GetCurrentResultPrevPoseVector(), SearchContext.GetCurrentResultPoseVector(), SearchContext.GetCurrentResultNextPoseVector(), LerpValue);
 		}
-		// else leave the InOutQuery set to zero since the SearchContext.History is invalid and it'll fail if we continue
+		else
+		{
+			// we leave the InOutQuery set to zero since the SearchContext.History is invalid and it'll fail if we continue
+			UE_LOG(LogPoseSearch, Error, TEXT("UPoseSearchFeatureChannel_Phase::BuildQuery - Failed because Pose History Node is missing."));
+		}
 	}
 	else
 	{
diff --git a/Engine/Plugins/Experimental/Animation/PoseSearch/Source/Runtime/Private/PoseSearchFeatureChannel_Position.cpp b/Engine/Plugins/Experimental/Animation/PoseSearch/Source/Runtime/Private/PoseSearchFeatureChannel_Position.cpp
--- a/Engine/Plugins/Experimental/Animation/PoseSearch/Source/Runtime/Private/PoseSearchFeatureChannel_Position.cpp
+++ b/Engine/Plugins/Experimental/Animation/PoseSearch/Source/Runtime/Private/PoseSearchFeatureChannel_Position.cpp
@@ -69,7 +69,11 @@ void UPoseSearchFeatureChannel_Position::BuildQuery(UE::PoseSearch::FSearchConte
 			const float LerpValue = InputQueryPose == EInputQueryPose::UseInterpolatedContinuingPose ? SearchContext.GetCurrentResult().LerpValue : 0.f;
 			FFeatureVectorHelper::EncodeVector(InOutQuery.EditValues(), ChannelDataOffset, SearchContext.GetCurrentResultPrevPoseVector(), SearchContext.GetCurrentResultPoseVector(), SearchContext.GetCurrentResultNextPoseVector(), LerpValue, false, ComponentStripping);
 		}
-		// else leave the InOutQuery set to zero since the SearchContext.History is invalid and it'll fail if we continue
+		else
+		{
+			// we leave the InOutQuery set to zero since the SearchContext.History is invalid and it'll fail if we continue
+			UE_LOG(LogPoseSearch, Error, TEXT("UPoseSearchFeatureChannel_Position::BuildQuery - Failed because Pose History Node is missing."));
+		}
 	}
 	else
 	{
diff --git a/Engine/Plugins/Experimental/Animation/PoseSearch/Source/Runtime/Private/PoseSearchFeatureChannel_Velocity.cpp b/Engine/Plugins/Experimental/Animation/PoseSearch/Source/Runtime/Private/PoseSearchFeatureChannel_Velocity.cpp
--- a/Engine/Plugins/Experimental/Animation/PoseSearch/Source/Runtime/Private/PoseSearchFeatureChannel_Velocity.cpp
+++ b/Engine/Plugins/Experimental/Animation/PoseSearch/Source/Runtime/Private/PoseSearchFeatureChannel_Velocity.cpp
@@ -41,7 +41,11 @@ void UPoseSearchFeatureChannel_Velocity::BuildQuery(UE::PoseSearch::FSearchConte
 			// @todo: we should normalize if EPoseSearchVelocityFlags::Normalize && LerpValue != 0
 			FFeatureVectorHelper::EncodeVector(InOutQuery.EditValues(), ChannelDataOffset, SearchContext.GetCurrentResultPrevPoseVector(), SearchContext.GetCurrentResultPoseVector(), SearchContext.GetCurrentResultNextPoseVector(), LerpValue, false, ComponentStripping);
 		}
-		// else leave the InOutQuery set to zero since the SearchContext.History is invalid and it'll fail if we continue
+		else
+		{
+			// we leave the InOutQuery set to zero since the SearchContext.History is invalid and it'll fail if we continue
+			UE_LOG(LogPoseSearch, Error, TEXT("UPoseSearchFeatureChannel_Velocity::BuildQuery - Failed because Pose History Node is missing."));
+		}
 	}
 	else
 	{
