motion matching - reintroducing CanAdvance to fix animation freezing if delta ticks brings over the end of the animation (passing the additional blend frames at the end of the animation)

diff --git a/Engine/Plugins/Experimental/Animation/PoseSearch/Source/Runtime/Private/PoseSearchLibrary.cpp b/Engine/Plugins/Experimental/Animation/PoseSearch/Source/Runtime/Private/PoseSearchLibrary.cpp
--- a/Engine/Plugins/Experimental/Animation/PoseSearch/Source/Runtime/Private/PoseSearchLibrary.cpp
+++ b/Engine/Plugins/Experimental/Animation/PoseSearch/Source/Runtime/Private/PoseSearchLibrary.cpp
@@ -300,17 +300,17 @@ void UPoseSearchLibrary::UpdateMotionMatchingState(
 	FMemMark Mark(FMemStack::Get());
 	FSearchContext SearchContext(&TrajectoryRootSpace, History, 0.f, &InOutMotionMatchingState.PoseIndicesHistory, InOutMotionMatchingState.CurrentSearchResult, PoseJumpThresholdTime, bForceInterrupt);
 
-	const bool bMustSearch = !InOutMotionMatchingState.CurrentSearchResult.IsValid();
+	const bool bCanAdvance = InOutMotionMatchingState.CurrentSearchResult.CanAdvance(DeltaTime);
 
 	// If we can't advance or enough time has elapsed since the last pose jump then search
-	const bool bSearch = bMustSearch || (bShouldSearch && (InOutMotionMatchingState.ElapsedPoseSearchTime >= SearchThrottleTime));
+	const bool bSearch = !bCanAdvance || (bShouldSearch && (InOutMotionMatchingState.ElapsedPoseSearchTime >= SearchThrottleTime));
 	if (bSearch)
 	{
 		InOutMotionMatchingState.ElapsedPoseSearchTime = 0.f;
 
 		// Evaluate continuing pose
 		FSearchResult SearchResult;
-		if (!SearchContext.IsForceInterrupt() && SearchContext.GetCurrentResult().IsValid())
+		if (!SearchContext.IsForceInterrupt() && bCanAdvance)
 		{
 			SearchResult.PoseCost = SearchContext.GetCurrentResult().Database->SearchContinuingPose(SearchContext);
 			SearchContext.UpdateCurrentBestCost(SearchResult.PoseCost);
diff --git a/Engine/Plugins/Experimental/Animation/PoseSearch/Source/Runtime/Private/PoseSearchResult.cpp b/Engine/Plugins/Experimental/Animation/PoseSearch/Source/Runtime/Private/PoseSearchResult.cpp
--- a/Engine/Plugins/Experimental/Animation/PoseSearch/Source/Runtime/Private/PoseSearchResult.cpp
+++ b/Engine/Plugins/Experimental/Animation/PoseSearch/Source/Runtime/Private/PoseSearchResult.cpp
@@ -102,4 +102,33 @@ const FSearchIndexAsset* FSearchResult::GetSearchIndexAsset(bool bMandatory) con
 	return &Database->GetSearchIndex().GetAssetForPose(PoseIdx);
 }
 
+bool FSearchResult::CanAdvance(float DeltaTime) const
+{
+	bool bCanAdvance = false;
+	if (IsValid())
+	{
+		float SteppedTime = AssetTime;
+		const FSearchIndexAsset* SearchIndexAsset = GetSearchIndexAsset(true);
+		const FInstancedStruct& DatabaseAsset = Database->GetAnimationAssetStruct(*SearchIndexAsset);
+		if (const FPoseSearchDatabaseBlendSpace* DatabaseBlendSpace = DatabaseAsset.GetPtr<FPoseSearchDatabaseBlendSpace>())
+		{
+			TArray<FBlendSampleData> BlendSamples;
+			int32 TriangulationIndex = 0;
+			DatabaseBlendSpace->BlendSpace->GetSamplesFromBlendInput(SearchIndexAsset->BlendParameters, BlendSamples, TriangulationIndex, true);
+
+			const float PlayLength = DatabaseBlendSpace->BlendSpace->GetAnimationLengthFromSampleData(BlendSamples);
+
+			// Asset player time for blend spaces is normalized [0, 1] so we need to convert it back to real time before we advance it
+			SteppedTime = AssetTime * PlayLength;
+			bCanAdvance = ETAA_Finished != FAnimationRuntime::AdvanceTime(DatabaseBlendSpace->IsLooping(), DeltaTime, SteppedTime, PlayLength);
+		}
+		else if (const FPoseSearchDatabaseAnimationAssetBase* DatabaseAnimationAssetBase = DatabaseAsset.GetPtr<FPoseSearchDatabaseAnimationAssetBase>())
+		{
+			const float AssetLength = DatabaseAnimationAssetBase->GetAnimationAsset()->GetPlayLength();
+			bCanAdvance = ETAA_Finished != FAnimationRuntime::AdvanceTime(DatabaseAnimationAssetBase->IsLooping(), DeltaTime, SteppedTime, AssetLength);
+		}
+	}
+	return bCanAdvance;
+}
+
 } // namespace UE::PoseSearch
diff --git a/Engine/Plugins/Experimental/Animation/PoseSearch/Source/Runtime/Public/PoseSearch/PoseSearchResult.h b/Engine/Plugins/Experimental/Animation/PoseSearch/Source/Runtime/Public/PoseSearch/PoseSearchResult.h
--- a/Engine/Plugins/Experimental/Animation/PoseSearch/Source/Runtime/Public/PoseSearch/PoseSearchResult.h
+++ b/Engine/Plugins/Experimental/Animation/PoseSearch/Source/Runtime/Public/PoseSearch/PoseSearchResult.h
@@ -62,6 +62,8 @@ struct FSearchResult
 	void Reset();
 
 	const FSearchIndexAsset* GetSearchIndexAsset(bool bMandatory = false) const;
+	
+	bool CanAdvance(float DeltaTime) const;
 };
 
 } // namespace UE::PoseSearch
\ No newline at end of file
