motion matching - pose search debugger: adding show best asset pose only and FAssetIdx column

diff --git a/Engine/Plugins/Experimental/Animation/PoseSearch/Source/Editor/Private/PoseSearchDebuggerDatabaseColumns.h b/Engine/Plugins/Experimental/Animation/PoseSearch/Source/Editor/Private/PoseSearchDebuggerDatabaseColumns.h
--- a/Engine/Plugins/Experimental/Animation/PoseSearch/Source/Editor/Private/PoseSearchDebuggerDatabaseColumns.h
+++ b/Engine/Plugins/Experimental/Animation/PoseSearch/Source/Editor/Private/PoseSearchDebuggerDatabaseColumns.h
@@ -94,7 +94,7 @@ struct FPoseIdx : ITextColumn
 
 	virtual FText GetLabel() const override
 	{
-		return LOCTEXT("ColumnLabelPoseIndex", "Index");
+		return LOCTEXT("ColumnLabelPoseIndex", "PoseIndex");
 	}
 
 	virtual FText GetLabelTooltip() const override
@@ -113,6 +113,32 @@ struct FPoseIdx : ITextColumn
 	}
 };
 
+struct FAssetIdx : ITextColumn
+{
+	using ITextColumn::ITextColumn;
+
+	virtual FText GetLabel() const override
+	{
+		return LOCTEXT("ColumnLabelAssetIndex", "AssetIndex");
+	}
+
+	virtual FText GetLabelTooltip() const override
+	{
+		return LOCTEXT("ColumnLabelTooltipAssetIndex", "Index of the Asset in the Database");
+	}
+
+	virtual FSortPredicate GetSortPredicate() const override
+	{
+		return [](const FRowDataRef& Row0, const FRowDataRef& Row1) -> bool { return Row0->DbAssetIdx < Row1->DbAssetIdx; };
+	}
+
+	virtual FText GetRowText(const FRowDataRef& Row) const override
+	{
+		return FText::AsNumber(Row->DbAssetIdx, &FNumberFormattingOptions::DefaultNoGrouping());
+	}
+};
+
+
 struct FDatabaseName : IColumn
 {
 	TSharedPtr<FDebuggerViewModel> DebuggerViewModel;
diff --git a/Engine/Plugins/Experimental/Animation/PoseSearch/Source/Editor/Private/PoseSearchDebuggerDatabaseView.cpp b/Engine/Plugins/Experimental/Animation/PoseSearch/Source/Editor/Private/PoseSearchDebuggerDatabaseView.cpp
--- a/Engine/Plugins/Experimental/Animation/PoseSearch/Source/Editor/Private/PoseSearchDebuggerDatabaseView.cpp
+++ b/Engine/Plugins/Experimental/Animation/PoseSearch/Source/Editor/Private/PoseSearchDebuggerDatabaseView.cpp
@@ -411,6 +411,7 @@ void SDebuggerDatabaseView::Update(const FTraceMotionMatchingStateMessage& State
 		AddColumn(MakeShared<FMirrored>(ColumnIdx++));
 		AddColumn(MakeShared<FLooping>(ColumnIdx++));
 		AddColumn(MakeShared<FPoseIdx>(ColumnIdx++));
+		AddColumn(MakeShared<FAssetIdx>(ColumnIdx++));
 		AddColumn(MakeShared<FBlendParameters>(ColumnIdx++));
 		AddColumn(MakeShared<FPoseCandidateFlags>(ColumnIdx++));
 
@@ -525,6 +526,20 @@ void SDebuggerDatabaseView::OnShowAllPosesCheckboxChanged(ECheckBoxState State)
 	}
 }
 
+void SDebuggerDatabaseView::OnShowOnlyBestAssetPoseCheckboxChanged(ECheckBoxState State)
+{
+	if (State == ECheckBoxState::Checked)
+	{
+		bShowOnlyBestAssetPose = true;
+	}
+	else
+	{
+		bShowOnlyBestAssetPose = false;
+	}
+
+	PopulateViewRows();
+}
+
 void SDebuggerDatabaseView::OnHideInvalidPosesCheckboxChanged(ECheckBoxState State)
 {
 	if (State == ECheckBoxState::Checked) 
@@ -585,50 +600,104 @@ void SDebuggerDatabaseView::PopulateViewRows()
 	ContinuingPoseView.Rows.Reset();
 	FilteredDatabaseView.Rows.Empty();
 
-	FString FilterString = FilterText.ToString();
-	TArray<FString> Tokens;
-	FilterString.ParseIntoArrayWS(Tokens);
-	const bool bHasNameFilter = !Tokens.IsEmpty();
-	FRegexPattern Pattern(FilterString);
-
-	for (const TSharedRef<FDebuggerDatabaseRowData>& UnfilteredRow : UnfilteredDatabaseRows)
+	const int32 UnfilteredDatabaseRowsNum = UnfilteredDatabaseRows.Num();
+	if (UnfilteredDatabaseRowsNum > 0)
 	{
-		bool bTryAddToFilteredDatabaseViewRows = true;
-		if (EnumHasAnyFlags(UnfilteredRow->PoseCandidateFlags, EPoseCandidateFlags::Valid_ContinuingPose))
+		FString FilterString = FilterText.ToString();
+		TArray<FString> Tokens;
+		FilterString.ParseIntoArrayWS(Tokens);
+		const bool bHasNameFilter = !Tokens.IsEmpty();
+		FRegexPattern Pattern(FilterString);
+	
+		TSet<int32> BestAssetPoseIndex;
+		if (bShowOnlyBestAssetPose)
 		{
-			ContinuingPoseView.Rows.Add(UnfilteredRow);
-			bTryAddToFilteredDatabaseViewRows = false;
-		}
+			// @todo: perhaps optimize via:
+			//TArrayView<int32> SortIndex((int32*)FMemory_Alloca(UnfilteredDatabaseRowsNum * sizeof(int32)), UnfilteredDatabaseRowsNum);
+			TArray<int32> SortIndex;
+			SortIndex.SetNumUninitialized(UnfilteredDatabaseRowsNum);
+			for (int32 i = 0; i < UnfilteredDatabaseRowsNum; ++i)
+			{
+				SortIndex[i] = i;
+			}
 
-		if (EnumHasAnyFlags(UnfilteredRow->PoseCandidateFlags, EPoseCandidateFlags::Valid_CurrentPose))
-		{
-			ActiveView.Rows.Add(UnfilteredRow);
-			bTryAddToFilteredDatabaseViewRows = false;
+			// sorting SortIndex elements by ascending order in cost, after grouping them by asset index
+			Algo::Sort(SortIndex, [this](int32 IndexA, int32 IndexB)
+			{
+				const FDebuggerDatabaseRowData& RowDataA = UnfilteredDatabaseRows[IndexA].Get();
+				const FDebuggerDatabaseRowData& RowDataB = UnfilteredDatabaseRows[IndexB].Get();
+
+				if (RowDataA.DbAssetIdx == RowDataB.DbAssetIdx)
+				{
+					return RowDataA.PoseCost.GetTotalCost() < RowDataB.PoseCost.GetTotalCost();
+				}
+				return RowDataA.DbAssetIdx < RowDataB.DbAssetIdx;
+			});
+
+			// populating BestAssetPoseIndex only with the best pose index (lowest total cost) for each asset
+			int32 PreviousDbAssetIdx = INDEX_NONE;
+			for (int32 i = 0; i < UnfilteredDatabaseRowsNum; ++i)
+			{
+				const FDebuggerDatabaseRowData& RowData = UnfilteredDatabaseRows[SortIndex[i]].Get();
+				if (PreviousDbAssetIdx != RowData.DbAssetIdx)
+				{
+					// adding only the best pose index for the RowData.DbAssetIdx asset to BestAssetPoseIndex 
+					// (the first we find by iterating over UnfilteredDatabaseRows[SortIndex[i]])
+					BestAssetPoseIndex.Add(RowData.PoseIdx);
+					PreviousDbAssetIdx = RowData.DbAssetIdx;
+				}
+			}
 		}
 
-		if (bTryAddToFilteredDatabaseViewRows)
+		for (int32 i = 0; i < UnfilteredDatabaseRowsNum; ++i)
 		{
-			bool bPassesNameFilter = true;
-			if (bHideInvalidPoses && !EnumHasAnyFlags(UnfilteredRow->PoseCandidateFlags, EPoseCandidateFlags::AnyValidMask))
+			const TSharedRef<FDebuggerDatabaseRowData>& UnfilteredRow = UnfilteredDatabaseRows[i];
+
+			bool bTryAddToFilteredDatabaseViewRows = true;
+			if (EnumHasAnyFlags(UnfilteredRow->PoseCandidateFlags, EPoseCandidateFlags::Valid_ContinuingPose))
 			{
-				bPassesNameFilter = false;
+				ContinuingPoseView.Rows.Add(UnfilteredRow);
+				bTryAddToFilteredDatabaseViewRows = false;
 			}
-			else if (bUseRegex)
+
+			if (EnumHasAnyFlags(UnfilteredRow->PoseCandidateFlags, EPoseCandidateFlags::Valid_CurrentPose))
 			{
-				FRegexMatcher Matcher(Pattern, UnfilteredRow->AssetName);
-				bPassesNameFilter = Matcher.FindNext();
+				ActiveView.Rows.Add(UnfilteredRow);
+				bTryAddToFilteredDatabaseViewRows = false;
 			}
-			else if (bHasNameFilter)
+
+			if (bShowOnlyBestAssetPose)
 			{
-				bPassesNameFilter = Algo::AllOf(Tokens, [&](FString Token)
-					{
-						return UnfilteredRow->AssetName.Contains(Token);
-					});
+				if (!BestAssetPoseIndex.Find(UnfilteredRow->PoseIdx))
+				{
+					bTryAddToFilteredDatabaseViewRows = false;
+				}
 			}
-				
-			if (bPassesNameFilter)
+
+			if (bTryAddToFilteredDatabaseViewRows)
 			{
-				FilteredDatabaseView.Rows.Add(UnfilteredRow);
+				bool bPassesNameFilter = true;
+				if (bHideInvalidPoses && !EnumHasAnyFlags(UnfilteredRow->PoseCandidateFlags, EPoseCandidateFlags::AnyValidMask))
+				{
+					bPassesNameFilter = false;
+				}
+				else if (bUseRegex)
+				{
+					FRegexMatcher Matcher(Pattern, UnfilteredRow->AssetName);
+					bPassesNameFilter = Matcher.FindNext();
+				}
+				else if (bHasNameFilter)
+				{
+					bPassesNameFilter = Algo::AllOf(Tokens, [&](FString Token)
+						{
+							return UnfilteredRow->AssetName.Contains(Token);
+						});
+				}
+				
+				if (bPassesNameFilter)
+				{
+					FilteredDatabaseView.Rows.Add(UnfilteredRow);
+				}
 			}
 		}
 	}
@@ -1009,6 +1078,20 @@ void SDebuggerDatabaseView::Construct(const FArguments& InArgs)
 					+ SHorizontalBox::Slot()
 					.AutoWidth()
 					.Padding(10, 5, 10, 5)
+					[
+						SNew(SCheckBox)
+						.OnCheckStateChanged_Lambda([this](ECheckBoxState State)
+						{
+							SDebuggerDatabaseView::OnShowOnlyBestAssetPoseCheckboxChanged(State);
+						})
+						[
+							SNew(STextBlock)
+							.Text(LOCTEXT("PoseSearchDebuggerShowOnlyBestAssetPoseFlag", "Only Best Asset Pose"))
+						]
+					]
+					+ SHorizontalBox::Slot()
+					.AutoWidth()
+					.Padding(10, 5, 10, 5)
 					[
 						SNew(SCheckBox)
 						.OnCheckStateChanged_Lambda([this](ECheckBoxState State)
diff --git a/Engine/Plugins/Experimental/Animation/PoseSearch/Source/Editor/Private/PoseSearchDebuggerDatabaseView.h b/Engine/Plugins/Experimental/Animation/PoseSearch/Source/Editor/Private/PoseSearchDebuggerDatabaseView.h
--- a/Engine/Plugins/Experimental/Animation/PoseSearch/Source/Editor/Private/PoseSearchDebuggerDatabaseView.h
+++ b/Engine/Plugins/Experimental/Animation/PoseSearch/Source/Editor/Private/PoseSearchDebuggerDatabaseView.h
@@ -69,6 +69,9 @@ private:
 	/** Called when the ShowAllPoses Checkbox is Toggled */
 	void OnShowAllPosesCheckboxChanged(ECheckBoxState State);
 	
+	/** Called when the ShowOnlyBestAssetPose Checkbox is Toggled */
+	void OnShowOnlyBestAssetPoseCheckboxChanged(ECheckBoxState State);	
+
 	/** Called when the HideInvalidPoses Checkbox is Toggled */
 	void OnHideInvalidPosesCheckboxChanged(ECheckBoxState State);
 
@@ -146,6 +149,9 @@ private:
 	/* Bool used to show all poses from display */
 	bool bShowAllPoses = false;
 	
+	/* Bool used to show only the best pose of every asset */
+	bool bShowOnlyBestAssetPose = false;
+
 	/* Bool used to hide invalid poses from display */
 	bool bHideInvalidPoses = false;
 
