motion matching - compatible skeletons runtime support

diff --git a/Engine/Plugins/Experimental/Animation/PoseSearch/Source/Runtime/Private/PoseSearchContext.cpp b/Engine/Plugins/Experimental/Animation/PoseSearch/Source/Runtime/Private/PoseSearchContext.cpp
--- a/Engine/Plugins/Experimental/Animation/PoseSearch/Source/Runtime/Private/PoseSearchContext.cpp
+++ b/Engine/Plugins/Experimental/Animation/PoseSearch/Source/Runtime/Private/PoseSearchContext.cpp
@@ -367,7 +367,7 @@ FTransform FSearchContext::GetComponentSpaceTransform(float SampleTime, const UP
 		// collecting the local bone transforms from the IPoseHistory
 		check(History);
 		FTransform BoneComponentSpaceTransform;
-		if (!History->GetComponentSpaceTransformAtTime(SampleTime, BoneIndexType, BoneComponentSpaceTransform))
+		if (!History->GetComponentSpaceTransformAtTime(SampleTime, BoneIndexType, Schema->Skeleton, BoneComponentSpaceTransform))
 		{
 			if (const USkeleton* Skeleton = Schema->Skeleton)
 			{
diff --git a/Engine/Plugins/Experimental/Animation/PoseSearch/Source/Runtime/Private/PoseSearchHistory.cpp b/Engine/Plugins/Experimental/Animation/PoseSearch/Source/Runtime/Private/PoseSearchHistory.cpp
--- a/Engine/Plugins/Experimental/Animation/PoseSearch/Source/Runtime/Private/PoseSearchHistory.cpp
+++ b/Engine/Plugins/Experimental/Animation/PoseSearch/Source/Runtime/Private/PoseSearchHistory.cpp
@@ -5,6 +5,8 @@
 #include "Animation/AnimInstanceProxy.h"
 #include "Animation/AnimNodeBase.h"
 #include "Animation/AnimRootMotionProvider.h"
+#include "Animation/SkeletonRemapping.h"
+#include "Animation/SkeletonRemappingRegistry.h"
 #include "BonePose.h"
 #include "DrawDebugHelpers.h"
 #include "PoseSearch/PoseSearchResult.h"
@@ -120,8 +122,25 @@ void FPoseHistory::Init(int32 InNumPoses, float InTimeHorizon, const TArray<FBon
 	Entries.Reserve(InNumPoses);
 }
 
-bool FPoseHistory::GetComponentSpaceTransformAtTime(float Time, FBoneIndexType BoneIndexType, FTransform& OutBoneTransform, bool bExtrapolate) const
+FBoneIndexType FPoseHistory::GetRemappedBoneIndexType(FBoneIndexType BoneIndexType, const USkeleton* BoneIndexSkeleton) const
 {
+	// remapping BoneIndexType in case the skeleton used to store history (LastUpdateSkeleton) is different from BoneIndexSkeleton
+	if (LastUpdateSkeleton != nullptr && LastUpdateSkeleton != BoneIndexSkeleton)
+	{
+		const FSkeletonRemapping& SkeletonRemapping = UE::Anim::FSkeletonRemappingRegistry::Get().GetRemapping(BoneIndexSkeleton, LastUpdateSkeleton.Get());
+		if (SkeletonRemapping.IsValid())
+		{
+			BoneIndexType = SkeletonRemapping.GetTargetSkeletonBoneIndex(BoneIndexType);
+		}
+	}
+
+	return BoneIndexType;
+}
+
+bool FPoseHistory::GetComponentSpaceTransformAtTime(float Time, FBoneIndexType BoneIndexType, const USkeleton* BoneIndexSkeleton, FTransform& OutBoneTransform, bool bExtrapolate) const
+{
+	BoneIndexType = GetRemappedBoneIndexType(BoneIndexType, BoneIndexSkeleton);
+
 	FComponentSpaceTransformIndex TransformIndex = FComponentSpaceTransformIndex(BoneIndexType);
 	if (!BoneToTransformMap.IsEmpty())
 	{
@@ -214,6 +233,14 @@ bool FPoseHistory::IsEmpty() const
 
 void FPoseHistory::Update(float SecondsElapsed, FCSPose<FCompactPose>& ComponentSpacePose, const FTransform& ComponentTransform)
 {
+	const USkeleton* Skeleton = ComponentSpacePose.GetPose().GetBoneContainer().GetSkeletonAsset();
+	if (LastUpdateSkeleton != Skeleton)
+	{
+		// @todo: support a different USkeleton per FPoseHistoryEntry if required
+		Entries.Reset();
+		LastUpdateSkeleton = Skeleton;
+	}
+
 	// Age our elapsed times
 	for (FPoseHistoryEntry& Entry : Entries)
 	{
@@ -320,10 +347,14 @@ float FExtendedPoseHistory::GetSampleTimeInterval() const
 	return PoseHistory->GetSampleTimeInterval();
 }
 
-bool FExtendedPoseHistory::GetComponentSpaceTransformAtTime(float Time, FBoneIndexType BoneIndexType, FTransform& OutBoneTransform, bool bExtrapolate) const
+bool FExtendedPoseHistory::GetComponentSpaceTransformAtTime(float Time, FBoneIndexType BoneIndexType, const USkeleton* BoneIndexSkeleton, FTransform& OutBoneTransform, bool bExtrapolate) const
 {
 	check(PoseHistory);
 
+	const FBoneIndexType OriginalBoneIndex = BoneIndexType;
+
+	BoneIndexType = PoseHistory->GetRemappedBoneIndexType(BoneIndexType, BoneIndexSkeleton);
+
 	if (Time > 0.f)
 	{
  		FComponentSpaceTransformIndex TransformIndex = FComponentSpaceTransformIndex(BoneIndexType);
@@ -369,7 +400,7 @@ bool FExtendedPoseHistory::GetComponentSpaceTransformAtTime(float Time, FBoneInd
 		Time = 0.f;
 	}
 
-	return PoseHistory->GetComponentSpaceTransformAtTime(Time, BoneIndexType, OutBoneTransform, bExtrapolate);
+	return PoseHistory->GetComponentSpaceTransformAtTime(Time, OriginalBoneIndex, BoneIndexSkeleton, OutBoneTransform, bExtrapolate);
 }
 
 void FExtendedPoseHistory::GetRootTransformAtTime(float Time, FTransform& OutRootTransform, bool bExtrapolate) const
diff --git a/Engine/Plugins/Experimental/Animation/PoseSearch/Source/Runtime/Public/PoseSearch/PoseSearchHistory.h b/Engine/Plugins/Experimental/Animation/PoseSearch/Source/Runtime/Public/PoseSearch/PoseSearchHistory.h
--- a/Engine/Plugins/Experimental/Animation/PoseSearch/Source/Runtime/Public/PoseSearch/PoseSearchHistory.h
+++ b/Engine/Plugins/Experimental/Animation/PoseSearch/Source/Runtime/Public/PoseSearch/PoseSearchHistory.h
@@ -24,7 +24,7 @@ struct POSESEARCH_API IPoseHistory
 {
 	virtual ~IPoseHistory() {}
 	virtual float GetSampleTimeInterval() const = 0;
-	virtual bool GetComponentSpaceTransformAtTime(float Time, FBoneIndexType BoneIndexType, FTransform& OutBoneTransform, bool bExtrapolate = true) const = 0;
+	virtual bool GetComponentSpaceTransformAtTime(float Time, FBoneIndexType BoneIndexType, const USkeleton* BoneIndexSkeleton, FTransform& OutBoneTransform, bool bExtrapolate = true) const = 0;
 	virtual void GetRootTransformAtTime(float Time, FTransform& OutRootTransform, bool bExtrapolate = true) const = 0;
 	virtual bool IsEmpty() const = 0;
 };
@@ -52,16 +52,19 @@ struct FPoseHistory : public IPoseHistory
 
 	// IPoseHistory interface
 	virtual float GetSampleTimeInterval() const override;
-	virtual bool GetComponentSpaceTransformAtTime(float Time, FBoneIndexType BoneIndexType, FTransform& OutBoneTransform, bool bExtrapolate = true) const override;
+	virtual bool GetComponentSpaceTransformAtTime(float Time, FBoneIndexType BoneIndexType, const USkeleton* BoneIndexSkeleton, FTransform& OutBoneTransform, bool bExtrapolate = true) const override;
 	virtual void GetRootTransformAtTime(float Time, FTransform& OutRootTransform, bool bExtrapolate = true) const override;
 	virtual bool IsEmpty() const override;
 	// End of IPoseHistory interface
 
+	FBoneIndexType GetRemappedBoneIndexType(FBoneIndexType BoneIndexType, const USkeleton* BoneIndexSkeleton) const;
+
 #if ENABLE_DRAW_DEBUG && ENABLE_ANIM_DEBUG
 	void DebugDraw(FAnimInstanceProxy& AnimInstanceProxy) const;
 #endif
 
 private:
+	TWeakObjectPtr<const USkeleton> LastUpdateSkeleton;
 	FBoneToTransformMap BoneToTransformMap;
 	FPoseHistoryEntries Entries;
 	float TimeHorizon = 0.f;
@@ -75,7 +78,7 @@ struct FExtendedPoseHistory : public IPoseHistory
 
 	// IPoseHistory interface
 	virtual float GetSampleTimeInterval() const override;
-	virtual bool GetComponentSpaceTransformAtTime(float Time, FBoneIndexType BoneIndexType, FTransform& OutBoneTransform, bool bExtrapolate = true) const override;
+	virtual bool GetComponentSpaceTransformAtTime(float Time, FBoneIndexType BoneIndexType, const USkeleton* BoneIndexSkeleton, FTransform& OutBoneTransform, bool bExtrapolate = true) const override;
 	virtual void GetRootTransformAtTime(float Time, FTransform& OutRootTransform, bool bExtrapolate = true) const override;
 	virtual bool IsEmpty() const override;
 	// End of IPoseHistory interface
