motion matching - removing query interpolation

diff --git a/Engine/Plugins/Experimental/Animation/PoseSearch/Source/Runtime/Private/AnimNode_MotionMatching.cpp b/Engine/Plugins/Experimental/Animation/PoseSearch/Source/Runtime/Private/AnimNode_MotionMatching.cpp
--- a/Engine/Plugins/Experimental/Animation/PoseSearch/Source/Runtime/Private/AnimNode_MotionMatching.cpp
+++ b/Engine/Plugins/Experimental/Animation/PoseSearch/Source/Runtime/Private/AnimNode_MotionMatching.cpp
@@ -94,11 +94,18 @@ void FAnimNode_MotionMatching::UpdateAssetPlayer(const FAnimationUpdateContext&
 
 	GetEvaluateGraphExposedInputs().Execute(Context);
 
-	const bool bNeedsReset =
+	bool bNeedsReset =
 		bResetOnBecomingRelevant &&
 		UpdateCounter.HasEverBeenUpdated() &&
 		!UpdateCounter.WasSynchronizedCounter(Context.AnimInstanceProxy->GetUpdateCounter());
 
+#if WITH_EDITOR
+	if (!FAsyncPoseSearchDatabasesManagement::RequestAsyncBuildIndex(MotionMatchingState.CurrentSearchResult.Database.Get(), ERequestAsyncBuildFlag::ContinueRequest))
+	{
+		bNeedsReset = true;
+	}
+#endif // WITH_EDITOR
+
 	// If we just became relevant and haven't been initialized yet, then reset motion matching state, otherwise update the asset time using the player node.
 	if (bNeedsReset)
 	{
@@ -106,26 +113,6 @@ void FAnimNode_MotionMatching::UpdateAssetPlayer(const FAnimationUpdateContext&
 	}
 	else
 	{
-#if WITH_EDITOR
-		if (FAsyncPoseSearchDatabasesManagement::RequestAsyncBuildIndex(MotionMatchingState.CurrentSearchResult.Database.Get(), ERequestAsyncBuildFlag::ContinueRequest))
-		{
-			const UPoseSearchDatabase* CurrentResultDatabase = MotionMatchingState.CurrentSearchResult.Database.Get();
-			const FSearchIndex& SearchIndex = CurrentResultDatabase->GetSearchIndex();
-			if (!SearchIndex.IsValidPoseIndex(MotionMatchingState.CurrentSearchResult.PrevPoseIdx) ||
-				!SearchIndex.IsValidPoseIndex(MotionMatchingState.CurrentSearchResult.NextPoseIdx))
-			{
-				// MotionMatchingState is out of sync with CurrentResultDatabase: we need to reset the MM state. This could happen if PIE is paused, and we edit the database,
-				// so FAnimNode_MotionMatching::UpdateAssetPlayer is never called and FAsyncPoseSearchDatabasesManagement::RequestAsyncBuildIndex never returns false here
-				MotionMatchingState.Reset(Context.AnimInstanceProxy->GetComponentTransform());
-			}
-		}
-		else
-		{
-			// we're still indexing MotionMatchingState.CurrentSearchResult.Database, so we Reset the MotionMatchingState
-			MotionMatchingState.Reset(Context.AnimInstanceProxy->GetComponentTransform());
-		}
-#endif // WITH_EDITOR
-
 		// We adjust the motion matching state asset time to the current player node's asset time. This is done 
 		// because the player node may have ticked more or less time than we expected due to variable dt or the 
 		// dynamic playback rate adjustment and as such the motion matching state does not update by itself
diff --git a/Engine/Plugins/Experimental/Animation/PoseSearch/Source/Runtime/Private/PoseSearchContext.cpp b/Engine/Plugins/Experimental/Animation/PoseSearch/Source/Runtime/Private/PoseSearchContext.cpp
--- a/Engine/Plugins/Experimental/Animation/PoseSearch/Source/Runtime/Private/PoseSearchContext.cpp
+++ b/Engine/Plugins/Experimental/Animation/PoseSearch/Source/Runtime/Private/PoseSearchContext.cpp
@@ -261,20 +261,13 @@ FSearchContext::FSearchContext(const FPoseSearchQueryTrajectory* InTrajectory, c
 		if (SearchIndex.Values.IsEmpty())
 		{
 			const int32 NumDimensions = CurrentResult.Database->Schema->SchemaCardinality;
-			CurrentResultPoseVectorData.AddUninitialized(NumDimensions * 3);
-
-			CurrentResultNextPoseVector = SearchIndex.GetReconstructedPoseValues(CurrentResult.NextPoseIdx, MakeArrayView(CurrentResultPoseVectorData.GetData(), NumDimensions));
+			CurrentResultPoseVectorData.AddUninitialized(NumDimensions);
 			CurrentResultPoseVector = SearchIndex.GetReconstructedPoseValues(CurrentResult.PoseIdx, MakeArrayView(CurrentResultPoseVectorData.GetData() + NumDimensions, NumDimensions));
-			CurrentResultPrevPoseVector = SearchIndex.GetReconstructedPoseValues(CurrentResult.PrevPoseIdx, MakeArrayView(CurrentResultPoseVectorData.GetData() + 2 * NumDimensions, NumDimensions));
 		}
 		else
 		{
-			CurrentResultNextPoseVector = SearchIndex.GetPoseValues(CurrentResult.NextPoseIdx);
 			CurrentResultPoseVector = SearchIndex.GetPoseValues(CurrentResult.PoseIdx);
-			CurrentResultPrevPoseVector = SearchIndex.GetPoseValues(CurrentResult.PrevPoseIdx);
 		}
-
-		check(CurrentResultNextPoseVector.Num() == CurrentResultPoseVector.Num() && CurrentResultPrevPoseVector.Num() == CurrentResultPoseVector.Num());
 	}
 }
 
diff --git a/Engine/Plugins/Experimental/Animation/PoseSearch/Source/Runtime/Private/PoseSearchDatabase.cpp b/Engine/Plugins/Experimental/Animation/PoseSearch/Source/Runtime/Private/PoseSearchDatabase.cpp
--- a/Engine/Plugins/Experimental/Animation/PoseSearch/Source/Runtime/Private/PoseSearchDatabase.cpp
+++ b/Engine/Plugins/Experimental/Animation/PoseSearch/Source/Runtime/Private/PoseSearchDatabase.cpp
@@ -457,41 +457,6 @@ int32 UPoseSearchDatabase::GetPoseIndexFromTime(float Time, const UE::PoseSearch
 	return SearchIndexAsset.GetPoseIndexFromTime(Time, bIsLooping, Schema->SampleRate);
 }
 
-bool UPoseSearchDatabase::GetPoseIndicesAndLerpValueFromTime(float Time, const UE::PoseSearch::FSearchIndexAsset& SearchIndexAsset, int32& PrevPoseIdx, int32& PoseIdx, int32& NextPoseIdx, float& LerpValue) const
-{
-	PoseIdx = GetPoseIndexFromTime(Time, SearchIndexAsset);
-	if (PoseIdx == INDEX_NONE)
-	{
-		PrevPoseIdx = INDEX_NONE;
-		NextPoseIdx = INDEX_NONE;
-		LerpValue = 0.f;
-		return false;
-	}
-
-	const float FloatPoseOffset = Schema->SampleRate * Time - SearchIndexAsset.FirstSampleIdx;
-	const int32 PoseOffset = FMath::RoundToInt(FloatPoseOffset);
-	LerpValue = FloatPoseOffset - float(PoseOffset);
-
-	const float PrevTime = Time - 1.f / Schema->SampleRate;
-	const float NextTime = Time + 1.f / Schema->SampleRate;
-
-	PrevPoseIdx = GetPoseIndexFromTime(PrevTime, SearchIndexAsset);
-	if (PrevPoseIdx == INDEX_NONE)
-	{
-		PrevPoseIdx = PoseIdx;
-	}
-
-	NextPoseIdx = GetPoseIndexFromTime(NextTime, SearchIndexAsset);
-	if (NextPoseIdx == INDEX_NONE)
-	{
-		NextPoseIdx = PoseIdx;
-	}
-
-	check(LerpValue >= -0.5f && LerpValue <= 0.5f);
-
-	return true;
-}
-
 const FInstancedStruct& UPoseSearchDatabase::GetAnimationAssetStruct(int32 AnimationAssetIndex) const
 {
 	check(AnimationAssets.IsValidIndex(AnimationAssetIndex));
diff --git a/Engine/Plugins/Experimental/Animation/PoseSearch/Source/Runtime/Private/PoseSearchFeatureChannel.cpp b/Engine/Plugins/Experimental/Animation/PoseSearch/Source/Runtime/Private/PoseSearchFeatureChannel.cpp
--- a/Engine/Plugins/Experimental/Animation/PoseSearch/Source/Runtime/Private/PoseSearchFeatureChannel.cpp
+++ b/Engine/Plugins/Experimental/Animation/PoseSearch/Source/Runtime/Private/PoseSearchFeatureChannel.cpp
@@ -48,32 +48,6 @@ void FFeatureVectorHelper::EncodeVector(TArrayView<float> Values, int32 DataOffs
 	}
 }
 
-void FFeatureVectorHelper::EncodeVector(TArrayView<float> Values, int32 DataOffset, TConstArrayView<float> PrevValues, TConstArrayView<float> CurValues, TConstArrayView<float> NextValues, float LerpValue, bool bNormalize, EComponentStrippingVector ComponentStrippingVector)
-{
-	FVector Vector = DecodeVector(CurValues, DataOffset, ComponentStrippingVector);
-
-	// linear interpolation
-	if (!FMath::IsNearlyZero(LerpValue))
-	{
-		if (LerpValue < 0.f)
-		{
-			Vector = FMath::Lerp(Vector, DecodeVector(PrevValues, DataOffset, ComponentStrippingVector), -LerpValue);
-		}
-		else
-		{
-			Vector = FMath::Lerp(Vector, DecodeVector(NextValues, DataOffset, ComponentStrippingVector), LerpValue);
-		}
-	}
-
-	// @todo: do we need to add options for cubic interpolation?
-	if (bNormalize)
-	{
-		Vector = Vector.GetSafeNormal(UE_SMALL_NUMBER, FVector::XAxisVector);
-	}
-
-	EncodeVector(Values, DataOffset, Vector, ComponentStrippingVector);
-}
-
 FVector FFeatureVectorHelper::DecodeVector(TConstArrayView<float> Values, int32 DataOffset, EComponentStrippingVector ComponentStrippingVector)
 {
 	switch (ComponentStrippingVector)
@@ -96,27 +70,6 @@ void FFeatureVectorHelper::EncodeVector2D(TArrayView<float> Values, int32 DataOf
 	Values[DataOffset + 1] = Vector2D.Y;
 }
 
-void FFeatureVectorHelper::EncodeVector2D(TArrayView<float> Values, int32 DataOffset, TConstArrayView<float> PrevValues, TConstArrayView<float> CurValues, TConstArrayView<float> NextValues, float LerpValue)
-{
-	FVector2D Vector2D = DecodeVector2D(CurValues, DataOffset);
-
-	// linear interpolation
-	if (!FMath::IsNearlyZero(LerpValue))
-	{
-		if (LerpValue < 0.f)
-		{
-			Vector2D = FMath::Lerp(Vector2D, DecodeVector2D(PrevValues, DataOffset), -LerpValue);
-		}
-		else
-		{
-			Vector2D = FMath::Lerp(Vector2D, DecodeVector2D(NextValues, DataOffset), LerpValue);
-		}
-	}
-
-	// @todo: do we need to add options for cubic interpolation?
-	EncodeVector2D(Values, DataOffset, Vector2D);
-}
-
 FVector2D FFeatureVectorHelper::DecodeVector2D(TConstArrayView<float> Values, int32 DataOffset)
 {
 	return FVector2D(Values[DataOffset + 0], Values[DataOffset + 1]);
@@ -127,30 +80,17 @@ void FFeatureVectorHelper::EncodeFloat(TArrayView<float> Values, int32 DataOffse
 	Values[DataOffset + 0] = Value;
 }
 
-void FFeatureVectorHelper::EncodeFloat(TArrayView<float> Values, int32 DataOffset, TConstArrayView<float> PrevValues, TConstArrayView<float> CurValues, TConstArrayView<float> NextValues, float LerpValue)
+float FFeatureVectorHelper::DecodeFloat(TConstArrayView<float> Values, int32 DataOffset)
 {
-	float Value = DecodeFloat(CurValues, DataOffset);
-
-	// linear interpolation
-	if (!FMath::IsNearlyZero(LerpValue))
-	{
-		if (LerpValue < 0.f)
-		{
-			Value = FMath::Lerp(Value, DecodeFloat(PrevValues, DataOffset), -LerpValue);
-		}
-		else
-		{
-			Value = FMath::Lerp(Value, DecodeFloat(NextValues, DataOffset), LerpValue);
-		}
-	}
-
-	// @todo: do we need to add options for cubic interpolation?
-	EncodeFloat(Values, DataOffset, Value);
+	return Values[DataOffset];
 }
 
-float FFeatureVectorHelper::DecodeFloat(TConstArrayView<float> Values, int32 DataOffset)
+void FFeatureVectorHelper::Copy(TArrayView<float> Values, int32 DataOffset, int32 DataCardinality, TConstArrayView<float> OriginValues)
 {
-	return Values[DataOffset];
+	for (int32 i = 0; i < DataCardinality; ++i)
+	{
+		Values[DataOffset + i] = OriginValues[DataOffset + i];
+	}
 }
 
 } // namespace UE::PoseSearch
diff --git a/Engine/Plugins/Experimental/Animation/PoseSearch/Source/Runtime/Private/PoseSearchFeatureChannel_FilterCrashingLegs.cpp b/Engine/Plugins/Experimental/Animation/PoseSearch/Source/Runtime/Private/PoseSearchFeatureChannel_FilterCrashingLegs.cpp
--- a/Engine/Plugins/Experimental/Animation/PoseSearch/Source/Runtime/Private/PoseSearchFeatureChannel_FilterCrashingLegs.cpp
+++ b/Engine/Plugins/Experimental/Animation/PoseSearch/Source/Runtime/Private/PoseSearchFeatureChannel_FilterCrashingLegs.cpp
@@ -60,8 +60,7 @@ void UPoseSearchFeatureChannel_FilterCrashingLegs::BuildQuery(UE::PoseSearch::FS
 	{
 		if (bIsCurrentResultValid)
 		{
-			const float LerpValue = InputQueryPose == EInputQueryPose::UseInterpolatedContinuingPose ? SearchContext.GetCurrentResult().LerpValue : 0.f;
-			FFeatureVectorHelper::EncodeFloat(InOutQuery.EditValues(), ChannelDataOffset, SearchContext.GetCurrentResultPrevPoseVector(), SearchContext.GetCurrentResultPoseVector(), SearchContext.GetCurrentResultNextPoseVector(), LerpValue);
+			FFeatureVectorHelper::Copy(InOutQuery.EditValues(), ChannelDataOffset, ChannelCardinality, SearchContext.GetCurrentResultPoseVector());
 		}
 		else
 		{
diff --git a/Engine/Plugins/Experimental/Animation/PoseSearch/Source/Runtime/Private/PoseSearchFeatureChannel_Heading.cpp b/Engine/Plugins/Experimental/Animation/PoseSearch/Source/Runtime/Private/PoseSearchFeatureChannel_Heading.cpp
--- a/Engine/Plugins/Experimental/Animation/PoseSearch/Source/Runtime/Private/PoseSearchFeatureChannel_Heading.cpp
+++ b/Engine/Plugins/Experimental/Animation/PoseSearch/Source/Runtime/Private/PoseSearchFeatureChannel_Heading.cpp
@@ -51,9 +51,7 @@ void UPoseSearchFeatureChannel_Heading::BuildQuery(UE::PoseSearch::FSearchContex
 	{
 		if (bIsCurrentResultValid)
 		{
-			const float LerpValue = InputQueryPose == EInputQueryPose::UseInterpolatedContinuingPose ? SearchContext.GetCurrentResult().LerpValue : 0.f;
-			// @todo: we should normalize if LerpValue != 0
-			FFeatureVectorHelper::EncodeVector(InOutQuery.EditValues(), ChannelDataOffset, SearchContext.GetCurrentResultPrevPoseVector(), SearchContext.GetCurrentResultPoseVector(), SearchContext.GetCurrentResultNextPoseVector(), LerpValue, true, ComponentStripping);
+			FFeatureVectorHelper::Copy(InOutQuery.EditValues(), ChannelDataOffset, ChannelCardinality, SearchContext.GetCurrentResultPoseVector());
 		}
 		else
 		{
diff --git a/Engine/Plugins/Experimental/Animation/PoseSearch/Source/Runtime/Private/PoseSearchFeatureChannel_Phase.cpp b/Engine/Plugins/Experimental/Animation/PoseSearch/Source/Runtime/Private/PoseSearchFeatureChannel_Phase.cpp
--- a/Engine/Plugins/Experimental/Animation/PoseSearch/Source/Runtime/Private/PoseSearchFeatureChannel_Phase.cpp
+++ b/Engine/Plugins/Experimental/Animation/PoseSearch/Source/Runtime/Private/PoseSearchFeatureChannel_Phase.cpp
@@ -306,8 +306,7 @@ void UPoseSearchFeatureChannel_Phase::BuildQuery(UE::PoseSearch::FSearchContext&
 	{
 		if (bIsCurrentResultValid)
 		{
-			const float LerpValue = InputQueryPose == EInputQueryPose::UseInterpolatedContinuingPose ? SearchContext.GetCurrentResult().LerpValue : 0.f;
-			FFeatureVectorHelper::EncodeVector2D(InOutQuery.EditValues(), ChannelDataOffset, SearchContext.GetCurrentResultPrevPoseVector(), SearchContext.GetCurrentResultPoseVector(), SearchContext.GetCurrentResultNextPoseVector(), LerpValue);
+			FFeatureVectorHelper::Copy(InOutQuery.EditValues(), ChannelDataOffset, ChannelCardinality, SearchContext.GetCurrentResultPoseVector());
 		}
 		else
 		{
diff --git a/Engine/Plugins/Experimental/Animation/PoseSearch/Source/Runtime/Private/PoseSearchFeatureChannel_Position.cpp b/Engine/Plugins/Experimental/Animation/PoseSearch/Source/Runtime/Private/PoseSearchFeatureChannel_Position.cpp
--- a/Engine/Plugins/Experimental/Animation/PoseSearch/Source/Runtime/Private/PoseSearchFeatureChannel_Position.cpp
+++ b/Engine/Plugins/Experimental/Animation/PoseSearch/Source/Runtime/Private/PoseSearchFeatureChannel_Position.cpp
@@ -66,8 +66,7 @@ void UPoseSearchFeatureChannel_Position::BuildQuery(UE::PoseSearch::FSearchConte
 	{
 		if (bIsCurrentResultValid)
 		{
-			const float LerpValue = InputQueryPose == EInputQueryPose::UseInterpolatedContinuingPose ? SearchContext.GetCurrentResult().LerpValue : 0.f;
-			FFeatureVectorHelper::EncodeVector(InOutQuery.EditValues(), ChannelDataOffset, SearchContext.GetCurrentResultPrevPoseVector(), SearchContext.GetCurrentResultPoseVector(), SearchContext.GetCurrentResultNextPoseVector(), LerpValue, false, ComponentStripping);
+			FFeatureVectorHelper::Copy(InOutQuery.EditValues(), ChannelDataOffset, ChannelCardinality, SearchContext.GetCurrentResultPoseVector());
 		}
 		else
 		{
diff --git a/Engine/Plugins/Experimental/Animation/PoseSearch/Source/Runtime/Private/PoseSearchFeatureChannel_Velocity.cpp b/Engine/Plugins/Experimental/Animation/PoseSearch/Source/Runtime/Private/PoseSearchFeatureChannel_Velocity.cpp
--- a/Engine/Plugins/Experimental/Animation/PoseSearch/Source/Runtime/Private/PoseSearchFeatureChannel_Velocity.cpp
+++ b/Engine/Plugins/Experimental/Animation/PoseSearch/Source/Runtime/Private/PoseSearchFeatureChannel_Velocity.cpp
@@ -37,9 +37,7 @@ void UPoseSearchFeatureChannel_Velocity::BuildQuery(UE::PoseSearch::FSearchConte
 	{
 		if (bIsCurrentResultValid)
 		{
-			const float LerpValue = InputQueryPose == EInputQueryPose::UseInterpolatedContinuingPose ? SearchContext.GetCurrentResult().LerpValue : 0.f;
-			// @todo: we should normalize if EPoseSearchVelocityFlags::Normalize && LerpValue != 0
-			FFeatureVectorHelper::EncodeVector(InOutQuery.EditValues(), ChannelDataOffset, SearchContext.GetCurrentResultPrevPoseVector(), SearchContext.GetCurrentResultPoseVector(), SearchContext.GetCurrentResultNextPoseVector(), LerpValue, false, ComponentStripping);
+			FFeatureVectorHelper::Copy(InOutQuery.EditValues(), ChannelDataOffset, ChannelCardinality, SearchContext.GetCurrentResultPoseVector());
 		}
 		else
 		{
diff --git a/Engine/Plugins/Experimental/Animation/PoseSearch/Source/Runtime/Private/PoseSearchResult.cpp b/Engine/Plugins/Experimental/Animation/PoseSearch/Source/Runtime/Private/PoseSearchResult.cpp
--- a/Engine/Plugins/Experimental/Animation/PoseSearch/Source/Runtime/Private/PoseSearchResult.cpp
+++ b/Engine/Plugins/Experimental/Animation/PoseSearch/Source/Runtime/Private/PoseSearchResult.cpp
@@ -39,14 +39,8 @@ void FSearchResult::Update(float NewAssetTime)
 		const FInstancedStruct& DatabaseAsset = Database->GetAnimationAssetStruct(SearchIndexAsset);
 		if (DatabaseAsset.GetPtr<FPoseSearchDatabaseSequence>() || DatabaseAsset.GetPtr<FPoseSearchDatabaseAnimComposite>())
 		{
-			if (Database->GetPoseIndicesAndLerpValueFromTime(NewAssetTime, SearchIndexAsset, PrevPoseIdx, PoseIdx, NextPoseIdx, LerpValue))
-			{
-				AssetTime = NewAssetTime;
-			}
-			else
-			{
-				Reset();
-			}
+			PoseIdx = Database->GetPoseIndexFromTime(NewAssetTime, SearchIndexAsset);
+			AssetTime = NewAssetTime;
 		}
 		else if (const FPoseSearchDatabaseBlendSpace* DatabaseBlendSpace = DatabaseAsset.GetPtr<FPoseSearchDatabaseBlendSpace>())
 		{
@@ -60,14 +54,9 @@ void FSearchResult::Update(float NewAssetTime)
 			// to a real time before we advance it
 			check(NewAssetTime >= 0.f && NewAssetTime <= 1.f);
 			const float RealTime = NewAssetTime * PlayLength;
-			if (Database->GetPoseIndicesAndLerpValueFromTime(RealTime, SearchIndexAsset, PrevPoseIdx, PoseIdx, NextPoseIdx, LerpValue))
-			{
-				AssetTime = NewAssetTime;
-			}
-			else
-			{
-				Reset();
-			}
+
+			PoseIdx = Database->GetPoseIndexFromTime(RealTime, SearchIndexAsset);
+			AssetTime = NewAssetTime;
 		}
 		else
 		{
diff --git a/Engine/Plugins/Experimental/Animation/PoseSearch/Source/Runtime/Public/PoseSearch/PoseSearchContext.h b/Engine/Plugins/Experimental/Animation/PoseSearch/Source/Runtime/Public/PoseSearch/PoseSearchContext.h
--- a/Engine/Plugins/Experimental/Animation/PoseSearch/Source/Runtime/Public/PoseSearch/PoseSearchContext.h
+++ b/Engine/Plugins/Experimental/Animation/PoseSearch/Source/Runtime/Public/PoseSearch/PoseSearchContext.h
@@ -157,9 +157,7 @@ struct POSESEARCH_API FSearchContext
 
 	bool IsCurrentResultFromDatabase(const UPoseSearchDatabase* Database) const;
 
-	TConstArrayView<float> GetCurrentResultPrevPoseVector() const { return CurrentResultPrevPoseVector; }
 	TConstArrayView<float> GetCurrentResultPoseVector() const { return CurrentResultPoseVector; }
-	TConstArrayView<float> GetCurrentResultNextPoseVector() const { return CurrentResultNextPoseVector; }
 
 	const FSearchResult& GetCurrentResult() const { return CurrentResult; }
 	float GetPoseJumpThresholdTime() const { return PoseJumpThresholdTime; }
@@ -184,10 +182,7 @@ private:
 	float PoseJumpThresholdTime = 0.f;
 	bool bForceInterrupt = false;
 
-	// cached previous, current, and next pose data
-	TConstArrayView<float> CurrentResultNextPoseVector;
 	TConstArrayView<float> CurrentResultPoseVector;
-	TConstArrayView<float> CurrentResultPrevPoseVector;
 	TStackAlignedArray<float> CurrentResultPoseVectorData;
 
 	// transforms cached in component space
diff --git a/Engine/Plugins/Experimental/Animation/PoseSearch/Source/Runtime/Public/PoseSearch/PoseSearchDatabase.h b/Engine/Plugins/Experimental/Animation/PoseSearch/Source/Runtime/Public/PoseSearch/PoseSearchDatabase.h
--- a/Engine/Plugins/Experimental/Animation/PoseSearch/Source/Runtime/Public/PoseSearch/PoseSearchDatabase.h
+++ b/Engine/Plugins/Experimental/Animation/PoseSearch/Source/Runtime/Public/PoseSearch/PoseSearchDatabase.h
@@ -322,7 +322,6 @@ public:
 	bool GetSkipSearchIfPossible() const;
 
 	int32 GetPoseIndexFromTime(float AssetTime, const UE::PoseSearch::FSearchIndexAsset& SearchIndexAsset) const;
-	bool GetPoseIndicesAndLerpValueFromTime(float Time, const UE::PoseSearch::FSearchIndexAsset& SearchIndexAsset, int32& PrevPoseIdx, int32& PoseIdx, int32& NextPoseIdx, float& LerpValue) const;
 
 	const FInstancedStruct& GetAnimationAssetStruct(int32 AnimationAssetIndex) const;
 	const FInstancedStruct& GetAnimationAssetStruct(const UE::PoseSearch::FSearchIndexAsset& SearchIndexAsset) const;
diff --git a/Engine/Plugins/Experimental/Animation/PoseSearch/Source/Runtime/Public/PoseSearch/PoseSearchFeatureChannel.h b/Engine/Plugins/Experimental/Animation/PoseSearch/Source/Runtime/Public/PoseSearch/PoseSearchFeatureChannel.h
--- a/Engine/Plugins/Experimental/Animation/PoseSearch/Source/Runtime/Public/PoseSearch/PoseSearchFeatureChannel.h
+++ b/Engine/Plugins/Experimental/Animation/PoseSearch/Source/Runtime/Public/PoseSearch/PoseSearchFeatureChannel.h
@@ -72,16 +72,15 @@ class POSESEARCH_API FFeatureVectorHelper
 public:
 	static int32 GetVectorCardinality(EComponentStrippingVector ComponentStrippingVector);
 	static void EncodeVector(TArrayView<float> Values, int32 DataOffset, const FVector& Vector, EComponentStrippingVector ComponentStrippingVector);
-	static void EncodeVector(TArrayView<float> Values, int32 DataOffset, TConstArrayView<float> PrevValues, TConstArrayView<float> CurValues, TConstArrayView<float> NextValues, float LerpValue, bool bNormalize, EComponentStrippingVector ComponentStrippingVector);
 	static FVector DecodeVector(TConstArrayView<float> Values, int32 DataOffset, EComponentStrippingVector ComponentStrippingVector);
 
 	static void EncodeVector2D(TArrayView<float> Values, int32 DataOffset, const FVector2D& Vector2D);
-	static void EncodeVector2D(TArrayView<float> Values, int32 DataOffset, TConstArrayView<float> PrevValues, TConstArrayView<float> CurValues, TConstArrayView<float> NextValues, float LerpValue);
 	static FVector2D DecodeVector2D(TConstArrayView<float> Values, int32 DataOffset);
 
 	static void EncodeFloat(TArrayView<float> Values, int32 DataOffset, const float Value);
-	static void EncodeFloat(TArrayView<float> Values, int32 DataOffset, TConstArrayView<float> PrevValues, TConstArrayView<float> CurValues, TConstArrayView<float> NextValues, float LerpValue);
 	static float DecodeFloat(TConstArrayView<float> Values, int32 DataOffset);
+
+	static void Copy(TArrayView<float> Values, int32 DataOffset, int32 DataCardinality, TConstArrayView<float> OriginValues);
 };
 
 } // namespace UE::PoseSearch
diff --git a/Engine/Plugins/Experimental/Animation/PoseSearch/Source/Runtime/Public/PoseSearch/PoseSearchResult.h b/Engine/Plugins/Experimental/Animation/PoseSearch/Source/Runtime/Public/PoseSearch/PoseSearchResult.h
--- a/Engine/Plugins/Experimental/Animation/PoseSearch/Source/Runtime/Public/PoseSearch/PoseSearchResult.h
+++ b/Engine/Plugins/Experimental/Animation/PoseSearch/Source/Runtime/Public/PoseSearch/PoseSearchResult.h
@@ -38,12 +38,6 @@ struct FSearchResult
 	FPoseSearchCost PoseCost;
 	int32 PoseIdx = INDEX_NONE;
 
-	int32 PrevPoseIdx = INDEX_NONE;
-	int32 NextPoseIdx = INDEX_NONE;
-
-	// lerp value to find AssetTime from PrevPoseIdx -> AssetTime -> NextPoseIdx, within range [-0.5, 0.5]
-	float LerpValue = 0.f;
-
 	TObjectPtr<const UPoseSearchDatabase> Database;
 
 	float AssetTime = 0.0f;
