motion matching
- optimization for FMotionMatchingState::UpdateWantedPlayRate
- optimization for FFeatureVectorHelper::Copy
- small improvements for FSearchContext::GetOrBuildQuery
- added STAT_PoseSearch_ContinuingPose, STAT_PoseSearch_PCAReconstruct, STAT_PoseSearch_PCAProject, STAT_BlendStack_UpdateAssetPlayer markers
- added SCostTimelineView::BestPosePosView curve to keep track of the position/index of the best pose during the kdtree search to help fine tuning the database KDTreeQueryNumNeighbors

diff --git a/Engine/Plugins/Experimental/Animation/PoseSearch/Source/Editor/Private/PoseSearchDebugger.cpp b/Engine/Plugins/Experimental/Animation/PoseSearch/Source/Editor/Private/PoseSearchDebugger.cpp
--- a/Engine/Plugins/Experimental/Animation/PoseSearch/Source/Editor/Private/PoseSearchDebugger.cpp
+++ b/Engine/Plugins/Experimental/Animation/PoseSearch/Source/Editor/Private/PoseSearchDebugger.cpp
@@ -45,18 +45,23 @@ protected:
 	TSharedPtr<SCostCurveTimelineView> BruteForceCostView;
 	TSharedPtr<SCostCurveTimelineView::FTimelineCurveData> BruteForceCostData;
 
+	TSharedPtr<SCostCurveTimelineView> BestPosePosView;
+	TSharedPtr<SCostCurveTimelineView::FTimelineCurveData> BestPosePosData;
+
 	TSharedPtr<SToolTip> CostToolTip;
 
 	FText ToolTipTime;
 	FText ToolTipCost;
 	FText ToolTipCostBruteForce;
+	FText ToolTipBestPosePos;
 };
 
 void SCostTimelineView::Construct(const FArguments& InArgs)
 {
 	BestCostData = MakeShared<SCurveTimelineView::FTimelineCurveData>();
 	BruteForceCostData = MakeShared<SCurveTimelineView::FTimelineCurveData>();
-
+	BestPosePosData = MakeShared<SCurveTimelineView::FTimelineCurveData>();
+	
 	BestCostView = SNew(SCostCurveTimelineView)
 	.CurveColor(FLinearColor::White)
 	.ViewRange_Lambda([]()
@@ -80,6 +85,18 @@ void SCostTimelineView::Construct(const FArguments& InArgs)
 	{
 		return BruteForceCostData;
 	});
+
+	BestPosePosView = SNew(SCostCurveTimelineView)
+	.CurveColor(FLinearColor::Blue)
+	.ViewRange_Lambda([]()
+	{
+		return IRewindDebugger::Instance()->GetCurrentViewRange();
+	})
+	.RenderFill(false)
+	.CurveData_Lambda([this]()
+	{
+		return BestPosePosData;
+	});
 		
 	AddSlot()
 	[
@@ -89,6 +106,10 @@ void SCostTimelineView::Construct(const FArguments& InArgs)
 	[
 		BestCostView.ToSharedRef()
 	];
+	AddSlot()
+	[
+		BestPosePosView.ToSharedRef()
+	];
 }
 
 void SCostTimelineView::UpdateInternal(uint64 ObjectId)
@@ -107,20 +128,24 @@ void SCostTimelineView::UpdateInternal(uint64 ObjectId)
 		TArray<FCurvePoint>& BruteForceCostPoints = BruteForceCostData->Points;
 		BruteForceCostPoints.Reset();
 
+		TArray<FCurvePoint>& BestPosePosPoints = BestPosePosData->Points;
+		BestPosePosPoints.Reset();
+
 		// convert time range to from rewind debugger times to profiler times
 		TRange<double> TraceTimeRange = RewindDebugger->GetCurrentTraceRange();
 		double StartTime = TraceTimeRange.GetLowerBoundValue();
 		double EndTime = TraceTimeRange.GetUpperBoundValue();
 
-		PoseSearchProvider->EnumerateMotionMatchingStateTimelines(ObjectId, [StartTime, EndTime, &BestCostPoints, &BruteForceCostPoints](const FTraceProvider::FMotionMatchingStateTimeline& InTimeline)
+		PoseSearchProvider->EnumerateMotionMatchingStateTimelines(ObjectId, [StartTime, EndTime, &BestCostPoints, &BruteForceCostPoints, &BestPosePosPoints](const FTraceProvider::FMotionMatchingStateTimeline& InTimeline)
 		{
 			// this isn't very efficient, and it gets called every frame.  will need optimizing
-			InTimeline.EnumerateEvents(StartTime, EndTime, [StartTime, EndTime, &BestCostPoints, &BruteForceCostPoints](double InStartTime, double InEndTime, uint32 InDepth, const FTraceMotionMatchingStateMessage& InMessage)
+			InTimeline.EnumerateEvents(StartTime, EndTime, [StartTime, EndTime, &BestCostPoints, &BruteForceCostPoints, &BestPosePosPoints](double InStartTime, double InEndTime, uint32 InDepth, const FTraceMotionMatchingStateMessage& InMessage)
 			{
 				if (InEndTime > StartTime && InStartTime < EndTime)
 				{
 					BestCostPoints.Add({ InMessage.RecordingTime, InMessage.SearchBestCost });
 					BruteForceCostPoints.Add({ InMessage.RecordingTime, InMessage.SearchBruteForceCost });
+					BestPosePosPoints.Add({ InMessage.RecordingTime, float(InMessage.SearchBestPosePos) });
 				}
 				return TraceServices::EEventEnumerate::Continue;
 			});
@@ -181,11 +206,13 @@ FReply SCostTimelineView::OnMouseMove(const FGeometry& MyGeometry, const FPointe
 					const float Time = CurvePoints[TargetPointIndex].Time;
 					const float BestCost = CurvePoints[TargetPointIndex].Value;
 					const float BruteForceCost = BruteForceCostData->Points[TargetPointIndex].Value;
+					const int32 BestPosePos = FMath::RoundToInt(BestPosePosData->Points[TargetPointIndex].Value);
 
 					// Tooltip text formatting
 					FNumberFormattingOptions FormattingOptions;
 					FormattingOptions.MaximumFractionalDigits = 3;
 
+					ToolTipBestPosePos = FText::Format(LOCTEXT("CostTimelineViewToolTip_BestPosePosFormat", "Best Index: {0}"), FText::AsNumber(BestPosePos, &FormattingOptions));
 					ToolTipTime = FText::Format(LOCTEXT("CostTimelineViewToolTip_TimeFormat", "Search Time: {0}"), FText::AsNumber(Time, &FormattingOptions));
 					ToolTipCost = FText::Format(LOCTEXT("CostTimelineViewToolTip_CostFormat", "Search Cost: {0}"), FText::AsNumber(BestCost, &FormattingOptions));
 
@@ -214,6 +241,13 @@ FReply SCostTimelineView::OnMouseMove(const FGeometry& MyGeometry, const FPointe
 									.ColorAndOpacity(FLinearColor::Black)
 								]
 								+ SVerticalBox::Slot()
+								[
+									SNew(STextBlock)
+									.Text_Lambda([this]() { return ToolTipBestPosePos; })
+									.Font(FCoreStyle::Get().GetFontStyle("ToolTip.LargerFont"))
+									.ColorAndOpacity(FLinearColor::Black)
+								]
+								+ SVerticalBox::Slot()
 								[
 									SNew(STextBlock)
 									.Text_Lambda([this]() { return ToolTipCost; })
diff --git a/Engine/Plugins/Experimental/Animation/PoseSearch/Source/Runtime/Private/AnimNode_BlendStack.cpp b/Engine/Plugins/Experimental/Animation/PoseSearch/Source/Runtime/Private/AnimNode_BlendStack.cpp
--- a/Engine/Plugins/Experimental/Animation/PoseSearch/Source/Runtime/Private/AnimNode_BlendStack.cpp
+++ b/Engine/Plugins/Experimental/Animation/PoseSearch/Source/Runtime/Private/AnimNode_BlendStack.cpp
@@ -377,6 +377,8 @@ void FAnimNode_BlendStack_Standalone::Evaluate_AnyThread(FPoseContext& Output)
 
 void FAnimNode_BlendStack_Standalone::UpdateAssetPlayer(const FAnimationUpdateContext& Context)
 {
+	QUICK_SCOPE_CYCLE_COUNTER(STAT_BlendStack_UpdateAssetPlayer);
+
 	Super::UpdateAssetPlayer(Context);
 
 	// AnimPlayers[0] is the most newly inserted AnimPlayer, AnimPlayers[AnimPlayers.Num()-1] is the oldest, so to calculate the weights
diff --git a/Engine/Plugins/Experimental/Animation/PoseSearch/Source/Runtime/Private/PoseSearchContext.cpp b/Engine/Plugins/Experimental/Animation/PoseSearch/Source/Runtime/Private/PoseSearchContext.cpp
--- a/Engine/Plugins/Experimental/Animation/PoseSearch/Source/Runtime/Private/PoseSearchContext.cpp
+++ b/Engine/Plugins/Experimental/Animation/PoseSearch/Source/Runtime/Private/PoseSearchContext.cpp
@@ -454,16 +454,10 @@ void FSearchContext::UpdateCurrentBestCost(const FPoseSearchCost& PoseSearchCost
 
 const FFeatureVectorBuilder* FSearchContext::GetCachedQuery(const UPoseSearchSchema* Schema) const
 {
-	const FFeatureVectorBuilder* CachedQuery = CachedQueries.FindByPredicate([Schema](const FFeatureVectorBuilder& CachedQuery)
+	return CachedQueries.FindByPredicate([Schema](const FFeatureVectorBuilder& CachedQuery)
 	{
 		return CachedQuery.GetSchema() == Schema;
 	});
-
-	if (CachedQuery)
-	{
-		return CachedQuery;
-	}
-	return nullptr;
 }
 
 const FFeatureVectorBuilder& FSearchContext::GetOrBuildQuery(const UPoseSearchSchema* Schema)
@@ -471,13 +465,12 @@ const FFeatureVectorBuilder& FSearchContext::GetOrBuildQuery(const UPoseSearchSc
 	QUICK_SCOPE_CYCLE_COUNTER(STAT_PoseSearch_GetOrBuildQuery);
 
 	check(Schema && Schema->IsValid());
-	const FFeatureVectorBuilder* CachedFeatureVectorBuilder = GetCachedQuery(Schema);
-	if (CachedFeatureVectorBuilder)
+	if (const FFeatureVectorBuilder* CachedFeatureVectorBuilder = GetCachedQuery(Schema))
 	{
 		return *CachedFeatureVectorBuilder;
 	}
 	
-	FFeatureVectorBuilder& NewCachedQuery = CachedQueries[CachedQueries.AddDefaulted()];
+	FFeatureVectorBuilder& NewCachedQuery = CachedQueries[CachedQueries.Emplace(Schema)];
 	Schema->BuildQuery(*this, NewCachedQuery);
 	return NewCachedQuery;
 }
diff --git a/Engine/Plugins/Experimental/Animation/PoseSearch/Source/Runtime/Private/PoseSearchDatabase.cpp b/Engine/Plugins/Experimental/Animation/PoseSearch/Source/Runtime/Private/PoseSearchDatabase.cpp
--- a/Engine/Plugins/Experimental/Animation/PoseSearch/Source/Runtime/Private/PoseSearchDatabase.cpp
+++ b/Engine/Plugins/Experimental/Animation/PoseSearch/Source/Runtime/Private/PoseSearchDatabase.cpp
@@ -662,19 +662,19 @@ UE::PoseSearch::FSearchResult UPoseSearchDatabase::Search(UE::PoseSearch::FSearc
 
 	if (PoseSearchMode != EPoseSearchMode::BruteForce)
 	{
-#if WITH_EDITORONLY_DATA
+#if UE_POSE_SEARCH_TRACE_ENABLED
 		FPoseSearchCost BruteForcePoseCost = Result.BruteForcePoseCost;
-#endif
+#endif // UE_POSE_SEARCH_TRACE_ENABLED
 
 		Result = SearchPCAKDTree(SearchContext);
 
-#if WITH_EDITORONLY_DATA
+#if UE_POSE_SEARCH_TRACE_ENABLED
 		Result.BruteForcePoseCost = BruteForcePoseCost;
 		if (PoseSearchMode == EPoseSearchMode::PCAKDTree_Compare)
 		{
 			check(Result.BruteForcePoseCost.GetTotalCost() <= Result.PoseCost.GetTotalCost());
 		}
-#endif
+#endif // UE_POSE_SEARCH_TRACE_ENABLED
 	}
 	
 	return Result;
@@ -682,7 +682,7 @@ UE::PoseSearch::FSearchResult UPoseSearchDatabase::Search(UE::PoseSearch::FSearc
 
 template<bool bReconstructPoseValues, bool bAlignedAndPadded>
 static inline void EvaluatePoseKernel(UE::PoseSearch::FSearchResult& Result, const UE::PoseSearch::FSearchIndex& SearchIndex, TConstArrayView<float> QueryValues, TArrayView<float> ReconstructedPoseValuesBuffer,
-	int32 PoseIdx, const UE::PoseSearch::FSearchFilters& SearchFilters, UE::PoseSearch::FSearchContext& SearchContext, const UPoseSearchDatabase* Database, bool bUpdateBestCandidates = true)
+	int32 PoseIdx, const UE::PoseSearch::FSearchFilters& SearchFilters, UE::PoseSearch::FSearchContext& SearchContext, const UPoseSearchDatabase* Database, bool bUpdateBestCandidates, size_t ResultIndex)
 {
 	using namespace UE::PoseSearch;
 
@@ -694,11 +694,18 @@ static inline void EvaluatePoseKernel(UE::PoseSearch::FSearchResult& Result, con
 #endif
 	))
 	{
-		const FPoseSearchCost PoseCost = bAlignedAndPadded ? SearchIndex.CompareAlignedPoses(PoseIdx, 0.f, PoseValues, QueryValues) : SearchIndex.ComparePoses(PoseIdx, 0.f, PoseValues, QueryValues);;
+		const FPoseSearchCost PoseCost = bAlignedAndPadded ? SearchIndex.CompareAlignedPoses(PoseIdx, 0.f, PoseValues, QueryValues) : SearchIndex.ComparePoses(PoseIdx, 0.f, PoseValues, QueryValues);
 		if (PoseCost < Result.PoseCost)
 		{
 			Result.PoseCost = PoseCost;
 			Result.PoseIdx = PoseIdx;
+
+#if WITH_EDITORONLY_DATA
+			if (bUpdateBestCandidates)
+			{
+				Result.BestPosePos = ResultIndex;
+			}
+#endif // WITH_EDITORONLY_DATA
 		}
 
 #if UE_POSE_SEARCH_TRACE_ENABLED
@@ -712,6 +719,8 @@ static inline void EvaluatePoseKernel(UE::PoseSearch::FSearchResult& Result, con
 
 FPoseSearchCost UPoseSearchDatabase::SearchContinuingPose(UE::PoseSearch::FSearchContext& SearchContext) const
 {
+	QUICK_SCOPE_CYCLE_COUNTER(STAT_PoseSearch_ContinuingPose);
+
 	using namespace UE::PoseSearch;
 
 	check(SearchContext.GetCurrentResult().Database.Get() == this);
@@ -850,7 +859,7 @@ UE::PoseSearch::FSearchResult UPoseSearchDatabase::SearchPCAKDTree(UE::PoseSearc
 			check(IsAligned(ReconstructedPoseValuesBuffer.GetData(), alignof(VectorRegister4Float)));
 			for (size_t ResultIndex = 0; ResultIndex < ResultSet.Num(); ++ResultIndex)
 			{
-				EvaluatePoseKernel<true, false>(Result, SearchIndex, QueryValues, ReconstructedPoseValuesBuffer, ResultIndexes[ResultIndex], SearchFilters, SearchContext, this);
+				EvaluatePoseKernel<true, false>(Result, SearchIndex, QueryValues, ReconstructedPoseValuesBuffer, ResultIndexes[ResultIndex], SearchFilters, SearchContext, this, true, ResultIndex);
 			}
 		}
 		// is the data padded at 16 bytes (and 16 bytes aligned by construction)?
@@ -858,7 +867,7 @@ UE::PoseSearch::FSearchResult UPoseSearchDatabase::SearchPCAKDTree(UE::PoseSearc
 		{
 			for (size_t ResultIndex = 0; ResultIndex < ResultSet.Num(); ++ResultIndex)
 			{
-				EvaluatePoseKernel<false, true>(Result, SearchIndex, QueryValues, TArrayView<float>(), ResultIndexes[ResultIndex], SearchFilters, SearchContext, this);
+				EvaluatePoseKernel<false, true>(Result, SearchIndex, QueryValues, TArrayView<float>(), ResultIndexes[ResultIndex], SearchFilters, SearchContext, this, true, ResultIndex);
 			}
 		}
 		// no reconstruction, but data is not 16 bytes padded
@@ -866,7 +875,7 @@ UE::PoseSearch::FSearchResult UPoseSearchDatabase::SearchPCAKDTree(UE::PoseSearc
 		{
 			for (size_t ResultIndex = 0; ResultIndex < ResultSet.Num(); ++ResultIndex)
 			{
-				EvaluatePoseKernel<false, false>(Result, SearchIndex, QueryValues, TArrayView<float>(), ResultIndexes[ResultIndex], SearchFilters, SearchContext, this);
+				EvaluatePoseKernel<false, false>(Result, SearchIndex, QueryValues, TArrayView<float>(), ResultIndexes[ResultIndex], SearchFilters, SearchContext, this, true, ResultIndex);
 			}
 		}
 	}
@@ -926,7 +935,7 @@ UE::PoseSearch::FSearchResult UPoseSearchDatabase::SearchBruteForce(UE::PoseSear
 			check(IsAligned(ReconstructedPoseValuesBuffer.GetData(), alignof(VectorRegister4Float)));
 			for (int32 PoseIdx = 0; PoseIdx < SearchIndex.GetNumPoses(); ++PoseIdx)
 			{
-				EvaluatePoseKernel<true, false>(Result, SearchIndex, QueryValues, ReconstructedPoseValuesBuffer, PoseIdx, SearchFilters, SearchContext, this, bUpdateBestCandidates);
+				EvaluatePoseKernel<true, false>(Result, SearchIndex, QueryValues, ReconstructedPoseValuesBuffer, PoseIdx, SearchFilters, SearchContext, this, bUpdateBestCandidates, PoseIdx);
 			}
 		}
 		// is the data padded at 16 bytes (and 16 bytes aligned by construction)?
@@ -934,7 +943,7 @@ UE::PoseSearch::FSearchResult UPoseSearchDatabase::SearchBruteForce(UE::PoseSear
 		{
 			for (int32 PoseIdx = 0; PoseIdx < SearchIndex.GetNumPoses(); ++PoseIdx)
 			{
-				EvaluatePoseKernel<false, true>(Result, SearchIndex, QueryValues, TArrayView<float>(), PoseIdx, SearchFilters, SearchContext, this, bUpdateBestCandidates);
+				EvaluatePoseKernel<false, true>(Result, SearchIndex, QueryValues, TArrayView<float>(), PoseIdx, SearchFilters, SearchContext, this, bUpdateBestCandidates, PoseIdx);
 			}
 		}
 		// no reconstruction, but data is not 16 bytes padded
@@ -942,7 +951,7 @@ UE::PoseSearch::FSearchResult UPoseSearchDatabase::SearchBruteForce(UE::PoseSear
 		{
 			for (int32 PoseIdx = 0; PoseIdx < SearchIndex.GetNumPoses(); ++PoseIdx)
 			{
-				EvaluatePoseKernel<false, false>(Result, SearchIndex, QueryValues, TArrayView<float>(), PoseIdx, SearchFilters, SearchContext, this, bUpdateBestCandidates);
+				EvaluatePoseKernel<false, false>(Result, SearchIndex, QueryValues, TArrayView<float>(), PoseIdx, SearchFilters, SearchContext, this, bUpdateBestCandidates, PoseIdx);
 			}
 		}
 	}
@@ -963,9 +972,9 @@ UE::PoseSearch::FSearchResult UPoseSearchDatabase::SearchBruteForce(UE::PoseSear
 		Result.Database = this;
 	}
 
-#if WITH_EDITORONLY_DATA
+#if UE_POSE_SEARCH_TRACE_ENABLED
 	Result.BruteForcePoseCost = Result.PoseCost; 
-#endif
+#endif // UE_POSE_SEARCH_TRACE_ENABLED
 
 	return Result;
 }
\ No newline at end of file
diff --git a/Engine/Plugins/Experimental/Animation/PoseSearch/Source/Runtime/Private/PoseSearchFeatureChannel.cpp b/Engine/Plugins/Experimental/Animation/PoseSearch/Source/Runtime/Private/PoseSearchFeatureChannel.cpp
--- a/Engine/Plugins/Experimental/Animation/PoseSearch/Source/Runtime/Private/PoseSearchFeatureChannel.cpp
+++ b/Engine/Plugins/Experimental/Animation/PoseSearch/Source/Runtime/Private/PoseSearchFeatureChannel.cpp
@@ -87,9 +87,16 @@ float FFeatureVectorHelper::DecodeFloat(TConstArrayView<float> Values, int32 Dat
 
 void FFeatureVectorHelper::Copy(TArrayView<float> Values, int32 DataOffset, int32 DataCardinality, TConstArrayView<float> OriginValues)
 {
-	for (int32 i = 0; i < DataCardinality; ++i)
+	check(Values.Num() == OriginValues.Num() && DataOffset + DataCardinality <= Values.Num());
+
+	float* RESTRICT ValuesData = Values.GetData() + DataOffset;
+	const float* RESTRICT OriginValuesData = OriginValues.GetData() + DataOffset;
+
+	check(ValuesData != OriginValuesData);
+
+	for (int32 i = 0; i < DataCardinality; ++i, ++ValuesData, ++OriginValuesData)
 	{
-		Values[DataOffset + i] = OriginValues[DataOffset + i];
+		*ValuesData = *OriginValuesData;
 	}
 }
 
diff --git a/Engine/Plugins/Experimental/Animation/PoseSearch/Source/Runtime/Private/PoseSearchIndex.cpp b/Engine/Plugins/Experimental/Animation/PoseSearch/Source/Runtime/Private/PoseSearchIndex.cpp
--- a/Engine/Plugins/Experimental/Animation/PoseSearch/Source/Runtime/Private/PoseSearchIndex.cpp
+++ b/Engine/Plugins/Experimental/Animation/PoseSearch/Source/Runtime/Private/PoseSearchIndex.cpp
@@ -202,6 +202,8 @@ TConstArrayView<float> FSearchIndex::GetPoseValues(int32 PoseIdx) const
 
 TConstArrayView<float> FSearchIndex::GetReconstructedPoseValues(int32 PoseIdx, TArrayView<float> BufferUsedForReconstruction) const
 {
+	QUICK_SCOPE_CYCLE_COUNTER(STAT_PoseSearch_PCAReconstruct);
+
 	const int32 NumDimensions = WeightsSqrt.Num();
 	const int32 NumPoses = GetNumPoses();
 	check(PoseIdx >= 0 && PoseIdx < NumPoses&& NumDimensions > 0);
@@ -226,6 +228,8 @@ TConstArrayView<float> FSearchIndex::GetReconstructedPoseValues(int32 PoseIdx, T
 
 TConstArrayView<float> FSearchIndex::PCAProject(TConstArrayView<float> PoseValues, TArrayView<float> BufferUsedForProjection) const
 {
+	QUICK_SCOPE_CYCLE_COUNTER(STAT_PoseSearch_PCAProject);
+
 	const int32 NumDimensions = WeightsSqrt.Num();
 	const int32 NumberOfPrincipalComponents = PCAProjectionMatrix.Num() / NumDimensions;
 
diff --git a/Engine/Plugins/Experimental/Animation/PoseSearch/Source/Runtime/Private/PoseSearchLibrary.cpp b/Engine/Plugins/Experimental/Animation/PoseSearch/Source/Runtime/Private/PoseSearchLibrary.cpp
--- a/Engine/Plugins/Experimental/Animation/PoseSearch/Source/Runtime/Private/PoseSearchLibrary.cpp
+++ b/Engine/Plugins/Experimental/Animation/PoseSearch/Source/Runtime/Private/PoseSearchLibrary.cpp
@@ -101,7 +101,12 @@ void FMotionMatchingState::UpdateWantedPlayRate(const UE::PoseSearch::FSearchCon
 {
 	if (CurrentSearchResult.IsValid())
 	{
-		if (!FMath::IsNearlyEqual(PlayRate.Min, 1.f, UE_KINDA_SMALL_NUMBER) || !FMath::IsNearlyEqual(PlayRate.Max, 1.f, UE_KINDA_SMALL_NUMBER))
+		if (!ensure(PlayRate.Min <= PlayRate.Max && PlayRate.Min > UE_KINDA_SMALL_NUMBER))
+		{
+			UE_LOG(LogPoseSearch, Error, TEXT("Couldn't update the WantedPlayRate in FMotionMatchingState::UpdateWantedPlayRate, because of invalid PlayRate interval (%f, %f)"), PlayRate.Min, PlayRate.Max);
+			WantedPlayRate = 1.f;
+		}
+		else if (!FMath::IsNearlyEqual(PlayRate.Min, PlayRate.Max, UE_KINDA_SMALL_NUMBER))
 		{
 			if (const UE::PoseSearch::FFeatureVectorBuilder* PoseSearchFeatureVectorBuilder = SearchContext.GetCachedQuery(CurrentSearchResult.Database->Schema))
 			{
@@ -110,7 +115,7 @@ void FMotionMatchingState::UpdateWantedPlayRate(const UE::PoseSearch::FSearchCon
 					TConstArrayView<float> QueryData = PoseSearchFeatureVectorBuilder->GetValues();
 					TConstArrayView<float> ResultData = CurrentSearchResult.Database->GetSearchIndex().GetPoseValues(CurrentSearchResult.PoseIdx);
 					const float EstimatedSpeedRatio = TrajectoryChannel->GetEstimatedSpeedRatio(QueryData, ResultData);
-					check(PlayRate.Min <= PlayRate.Max);
+
 					WantedPlayRate = FMath::Clamp(EstimatedSpeedRatio, PlayRate.Min, PlayRate.Max);
 				}
 				else
@@ -123,7 +128,11 @@ void FMotionMatchingState::UpdateWantedPlayRate(const UE::PoseSearch::FSearchCon
 		}
 		else if (!FMath::IsNearlyZero(TrajectorySpeedMultiplier))
 		{
-			WantedPlayRate = 1.f / TrajectorySpeedMultiplier;
+			WantedPlayRate = PlayRate.Min / TrajectorySpeedMultiplier;
+		}
+		else
+		{
+			WantedPlayRate = PlayRate.Min;
 		}
 	}
 }
@@ -160,7 +169,6 @@ void UPoseSearchLibrary::TraceMotionMatchingState(
 	const FPoseSearchQueryTrajectory& Trajectory,
 	UE::PoseSearch::FSearchContext& SearchContext,
 	const UE::PoseSearch::FSearchResult& CurrentResult,
-	const UE::PoseSearch::FSearchResult& LastResult,
 	float ElapsedPoseSearchTime,
 	const FTransform& RootMotionTransformDelta,
 	const UObject* AnimInstance,
@@ -169,7 +177,8 @@ void UPoseSearchLibrary::TraceMotionMatchingState(
 	bool bSearch,
 	float RecordingTime,
 	float SearchBestCost,
-	float SearchBruteForceCost)
+	float SearchBruteForceCost,
+	int32 SearchBestPosePos)
 {
 	using namespace UE::PoseSearch;
 	
@@ -246,6 +255,7 @@ void UPoseSearchLibrary::TraceMotionMatchingState(
 	TraceState.RecordingTime = RecordingTime;
 	TraceState.SearchBestCost = SearchBestCost;
 	TraceState.SearchBruteForceCost = SearchBruteForceCost;
+	TraceState.SearchBestPosePos = SearchBestPosePos;
 
 	TraceState.Trajectory = Trajectory;
 
@@ -290,11 +300,6 @@ void UPoseSearchLibrary::UpdateMotionMatchingState(
 
 	InOutMotionMatchingState.bJumpedToPose = false;
 
-#if UE_POSE_SEARCH_TRACE_ENABLED
-	// Record Current Pose Index for Debugger
-	const FSearchResult LastResult = InOutMotionMatchingState.CurrentSearchResult;
-#endif
-
 	const IPoseHistory* History = nullptr;
 	if (IPoseHistoryProvider* PoseHistoryProvider = Context.GetMessage<IPoseHistoryProvider>())
 	{
@@ -337,12 +342,12 @@ void UPoseSearchLibrary::UpdateMotionMatchingState(
 			}
 		}
 
-#if WITH_EDITORONLY_DATA
+#if UE_POSE_SEARCH_TRACE_ENABLED
 		if (!SearchResult.BruteForcePoseCost.IsValid())
 		{
 			SearchResult.BruteForcePoseCost = SearchResult.PoseCost;
 		}
-#endif // WITH_EDITORONLY_DATA
+#endif // UE_POSE_SEARCH_TRACE_ENABLED
 
 		if (bJumpToPose)
 		{
@@ -351,9 +356,9 @@ void UPoseSearchLibrary::UpdateMotionMatchingState(
 		else
 		{
 			// copying few properties of SearchResult into CurrentSearchResult to facilitate debug drawing
-#if WITH_EDITORONLY_DATA
+#if UE_POSE_SEARCH_TRACE_ENABLED
 			InOutMotionMatchingState.CurrentSearchResult.BruteForcePoseCost = SearchResult.BruteForcePoseCost;
-#endif // WITH_EDITORONLY_DATA
+#endif // UE_POSE_SEARCH_TRACE_ENABLED
 			InOutMotionMatchingState.CurrentSearchResult.PoseCost = SearchResult.PoseCost;
 		}
 	}
@@ -373,14 +378,10 @@ void UPoseSearchLibrary::UpdateMotionMatchingState(
 		const UAnimInstance* AnimInstance = Cast<const UAnimInstance>(Context.AnimInstanceProxy->GetAnimInstanceObject());
 
 		const float SearchBestCost = InOutMotionMatchingState.CurrentSearchResult.PoseCost.GetTotalCost();
-		float SearchBruteForceCost = SearchBestCost;
-#if WITH_EDITORONLY_DATA
-		SearchBruteForceCost = InOutMotionMatchingState.CurrentSearchResult.BruteForcePoseCost.GetTotalCost();
-#endif // WITH_EDITORONLY_DATA
-
-		TraceMotionMatchingState(Trajectory, SearchContext, InOutMotionMatchingState.CurrentSearchResult, LastResult, InOutMotionMatchingState.ElapsedPoseSearchTime,
+		const float SearchBruteForceCost = InOutMotionMatchingState.CurrentSearchResult.BruteForcePoseCost.GetTotalCost();
+		TraceMotionMatchingState(Trajectory, SearchContext, InOutMotionMatchingState.CurrentSearchResult, InOutMotionMatchingState.ElapsedPoseSearchTime,
 			InOutMotionMatchingState.RootMotionTransformDelta, Context.AnimInstanceProxy->GetAnimInstanceObject(), Context.GetCurrentNodeId(), DeltaTime, bSearch,
-			AnimInstance ? FObjectTrace::GetWorldElapsedTime(AnimInstance->GetWorld()) : 0.f, SearchBestCost, SearchBruteForceCost);
+			AnimInstance ? FObjectTrace::GetWorldElapsedTime(AnimInstance->GetWorld()) : 0.f, SearchBestCost, SearchBruteForceCost, InOutMotionMatchingState.CurrentSearchResult.BestPosePos);
 	}
 #endif
 
@@ -602,12 +603,9 @@ void UPoseSearchLibrary::MotionMatch(
 
 #if UE_POSE_SEARCH_TRACE_ENABLED
 		const float SearchBestCost = SearchResult.PoseCost.GetTotalCost();
-		float SearchBruteForceCost = SearchBestCost;
-#if WITH_EDITORONLY_DATA
-		SearchBruteForceCost = SearchResult.BruteForcePoseCost.GetTotalCost();
-#endif // WITH_EDITORONLY_DATA
-		TraceMotionMatchingState(Trajectory, SearchContext, SearchResult, FSearchResult(), 0.f, FTransform::Identity, AnimInstance, DebugSessionUniqueIdentifier,
-			AnimInstance->GetDeltaSeconds(), true, FObjectTrace::GetWorldElapsedTime(AnimInstance->GetWorld()), SearchBestCost, SearchBruteForceCost);
+		const float SearchBruteForceCost = SearchResult.BruteForcePoseCost.GetTotalCost();
+		TraceMotionMatchingState(Trajectory, SearchContext, SearchResult, 0.f, FTransform::Identity, AnimInstance, DebugSessionUniqueIdentifier,
+			AnimInstance->GetDeltaSeconds(), true, FObjectTrace::GetWorldElapsedTime(AnimInstance->GetWorld()), SearchBestCost, SearchBruteForceCost, SearchResult.BestPosePos);
 #endif // UE_POSE_SEARCH_TRACE_ENABLED
 	}
 }
diff --git a/Engine/Plugins/Experimental/Animation/PoseSearch/Source/Runtime/Private/PoseSearchResult.cpp b/Engine/Plugins/Experimental/Animation/PoseSearch/Source/Runtime/Private/PoseSearchResult.cpp
--- a/Engine/Plugins/Experimental/Animation/PoseSearch/Source/Runtime/Private/PoseSearchResult.cpp
+++ b/Engine/Plugins/Experimental/Animation/PoseSearch/Source/Runtime/Private/PoseSearchResult.cpp
@@ -11,20 +11,13 @@ namespace UE::PoseSearch
 
 //////////////////////////////////////////////////////////////////////////
 // FFeatureVectorBuilder
-void FFeatureVectorBuilder::Init(const UPoseSearchSchema* InSchema)
+FFeatureVectorBuilder::FFeatureVectorBuilder(const UPoseSearchSchema* InSchema)
 {
 	check(InSchema && InSchema->IsValid());
 	Schema = InSchema;
-	Values.Reset();
 	Values.SetNumZeroed(Schema->SchemaCardinality);
 }
 
-void FFeatureVectorBuilder::Reset()
-{
-	Schema = nullptr;
-	Values.Reset();
-}
-
 //////////////////////////////////////////////////////////////////////////
 // FSearchResult
 void FSearchResult::Update(float NewAssetTime)
diff --git a/Engine/Plugins/Experimental/Animation/PoseSearch/Source/Runtime/Private/PoseSearchSchema.cpp b/Engine/Plugins/Experimental/Animation/PoseSearch/Source/Runtime/Private/PoseSearchSchema.cpp
--- a/Engine/Plugins/Experimental/Animation/PoseSearch/Source/Runtime/Private/PoseSearchSchema.cpp
+++ b/Engine/Plugins/Experimental/Animation/PoseSearch/Source/Runtime/Private/PoseSearchSchema.cpp
@@ -28,7 +28,8 @@ void UPoseSearchSchema::BuildQuery(UE::PoseSearch::FSearchContext& SearchContext
 {
 	QUICK_SCOPE_CYCLE_COUNTER(STAT_PoseSearch_BuildQuery);
 
-	InOutQuery.Init(this);
+	check(InOutQuery.GetSchema() == this);
+	check(InOutQuery.GetValues().Num() == SchemaCardinality);
 
 	for (const TObjectPtr<UPoseSearchFeatureChannel>& ChannelPtr : GetChannels())
 	{
diff --git a/Engine/Plugins/Experimental/Animation/PoseSearch/Source/Runtime/Private/Trace/PoseSearchTraceLogger.cpp b/Engine/Plugins/Experimental/Animation/PoseSearch/Source/Runtime/Private/Trace/PoseSearchTraceLogger.cpp
--- a/Engine/Plugins/Experimental/Animation/PoseSearch/Source/Runtime/Private/Trace/PoseSearchTraceLogger.cpp
+++ b/Engine/Plugins/Experimental/Animation/PoseSearch/Source/Runtime/Private/Trace/PoseSearchTraceLogger.cpp
@@ -80,6 +80,7 @@ FArchive& operator<<(FArchive& Ar, FTraceMotionMatchingState& State)
 	Ar << State.RecordingTime;
 	Ar << State.SearchBestCost;
 	Ar << State.SearchBruteForceCost;
+	Ar << State.SearchBestPosePos;
 	Ar << State.DatabaseEntries;
 	Ar << State.Trajectory;
 	Ar << State.CurrentDbEntryIdx;
diff --git a/Engine/Plugins/Experimental/Animation/PoseSearch/Source/Runtime/Private/Trace/PoseSearchTraceLogger.h b/Engine/Plugins/Experimental/Animation/PoseSearch/Source/Runtime/Private/Trace/PoseSearchTraceLogger.h
--- a/Engine/Plugins/Experimental/Animation/PoseSearch/Source/Runtime/Private/Trace/PoseSearchTraceLogger.h
+++ b/Engine/Plugins/Experimental/Animation/PoseSearch/Source/Runtime/Private/Trace/PoseSearchTraceLogger.h
@@ -78,6 +78,7 @@ struct POSESEARCH_API FTraceMotionMatchingState
 	float RecordingTime = 0.f;
 	float SearchBestCost = 0.f;
 	float SearchBruteForceCost = 0.f;
+	int32 SearchBestPosePos = 0;
 
 	TArray<FTraceMotionMatchingStateDatabaseEntry> DatabaseEntries;
 
diff --git a/Engine/Plugins/Experimental/Animation/PoseSearch/Source/Runtime/Public/PoseSearch/PoseSearchLibrary.h b/Engine/Plugins/Experimental/Animation/PoseSearch/Source/Runtime/Public/PoseSearch/PoseSearchLibrary.h
--- a/Engine/Plugins/Experimental/Animation/PoseSearch/Source/Runtime/Public/PoseSearch/PoseSearchLibrary.h
+++ b/Engine/Plugins/Experimental/Animation/PoseSearch/Source/Runtime/Public/PoseSearch/PoseSearchLibrary.h
@@ -84,7 +84,6 @@ class POSESEARCH_API UPoseSearchLibrary : public UBlueprintFunctionLibrary
 		const FPoseSearchQueryTrajectory& Trajectory,
 		UE::PoseSearch::FSearchContext& SearchContext,
 		const UE::PoseSearch::FSearchResult& CurrentResult,
-		const UE::PoseSearch::FSearchResult& LastResult,
 		float ElapsedPoseSearchTime,
 		const FTransform& RootMotionTransformDelta,
 		const UObject* AnimInstance,
@@ -93,7 +92,8 @@ class POSESEARCH_API UPoseSearchLibrary : public UBlueprintFunctionLibrary
 		bool bSearch,
 		float RecordingTime,
 		float SearchBestCost,
-		float SearchBruteForceCost);
+		float SearchBruteForceCost,
+		int32 SearchBestPosePos);
 #endif // UE_POSE_SEARCH_TRACE_ENABLED
 
 	static FPoseSearchQueryTrajectory ProcessTrajectory(
diff --git a/Engine/Plugins/Experimental/Animation/PoseSearch/Source/Runtime/Public/PoseSearch/PoseSearchResult.h b/Engine/Plugins/Experimental/Animation/PoseSearch/Source/Runtime/Public/PoseSearch/PoseSearchResult.h
--- a/Engine/Plugins/Experimental/Animation/PoseSearch/Source/Runtime/Public/PoseSearch/PoseSearchResult.h
+++ b/Engine/Plugins/Experimental/Animation/PoseSearch/Source/Runtime/Public/PoseSearch/PoseSearchResult.h
@@ -19,9 +19,7 @@ struct FSearchIndexAsset;
 struct FFeatureVectorBuilder
 {
 public:
-	void Init(const UPoseSearchSchema* Schema);
-	void Reset();
-
+	explicit FFeatureVectorBuilder(const UPoseSearchSchema* Schema);
 	const UPoseSearchSchema* GetSchema() const { return Schema.Get(); }
 
 	TArrayView<float> EditValues() { return Values; }
@@ -42,9 +40,10 @@ struct FSearchResult
 
 	float AssetTime = 0.0f;
 
-#if WITH_EDITORONLY_DATA
+#if UE_POSE_SEARCH_TRACE_ENABLED
 	FPoseSearchCost BruteForcePoseCost;
-#endif // WITH_EDITORONLY_DATA
+	int32 BestPosePos = 0;
+#endif // UE_POSE_SEARCH_TRACE_ENABLED
 
 	// Attempts to set the internal state to match the provided asset time including updating the internal DbPoseIdx. 
 	// If the provided asset time is out of bounds for the currently playing asset then this function will reset the 
