motion matching - populating all the searched databases to the pose search debugger

diff --git a/Engine/Plugins/Experimental/Animation/PoseSearch/Source/Runtime/Private/PoseSearchDatabase.cpp b/Engine/Plugins/Experimental/Animation/PoseSearch/Source/Runtime/Private/PoseSearchDatabase.cpp
--- a/Engine/Plugins/Experimental/Animation/PoseSearch/Source/Runtime/Private/PoseSearchDatabase.cpp
+++ b/Engine/Plugins/Experimental/Animation/PoseSearch/Source/Runtime/Private/PoseSearchDatabase.cpp
@@ -781,7 +781,13 @@ FPoseSearchCost UPoseSearchDatabase::SearchContinuingPose(UE::PoseSearch::FSearc
 
 #if UE_POSE_SEARCH_TRACE_ENABLED
 		SearchContext.BestCandidates.Add(ContinuingPoseCost, ContinuingPoseIdx, this, EPoseCandidateFlags::Valid_ContinuingPose);
-#endif
+#endif // UE_POSE_SEARCH_TRACE_ENABLED
+	}
+	else
+	{
+#if UE_POSE_SEARCH_TRACE_ENABLED
+		SearchContext.BestCandidates.Add(this);
+#endif // UE_POSE_SEARCH_TRACE_ENABLED
 	}
 
 	return ContinuingPoseCost;
@@ -886,9 +892,13 @@ UE::PoseSearch::FSearchResult UPoseSearchDatabase::SearchPCAKDTree(UE::PoseSearc
 		TConstArrayView<float> QueryValues = SearchContext.GetOrBuildQuery(Schema).GetValues();
 		FNonSelectableIdx NonSelectableIdx;
 		PopulateNonSelectableIdx(NonSelectableIdx, SearchContext, this, QueryValues);
-#endif
+#endif // UE_POSE_SEARCH_TRACE_ENABLED
 	}
 
+#if UE_POSE_SEARCH_TRACE_ENABLED
+	SearchContext.BestCandidates.Add(this);
+#endif // UE_POSE_SEARCH_TRACE_ENABLED
+
 	// finalizing Result properties
 	if (Result.PoseIdx != INDEX_NONE)
 	{
@@ -962,9 +972,13 @@ UE::PoseSearch::FSearchResult UPoseSearchDatabase::SearchBruteForce(UE::PoseSear
 		TConstArrayView<float> QueryValues = SearchContext.GetOrBuildQuery(Schema).GetValues();
 		FNonSelectableIdx NonSelectableIdx;
 		PopulateNonSelectableIdx(NonSelectableIdx, SearchContext, this, QueryValues);
-#endif
+#endif // UE_POSE_SEARCH_TRACE_ENABLED
 	}
 
+#if UE_POSE_SEARCH_TRACE_ENABLED
+	SearchContext.BestCandidates.Add(this);
+#endif // UE_POSE_SEARCH_TRACE_ENABLED
+
 	// finalizing Result properties
 	if (Result.PoseIdx != INDEX_NONE)
 	{
diff --git a/Engine/Plugins/Experimental/Animation/PoseSearch/Source/Runtime/Private/PoseSearchLibrary.cpp b/Engine/Plugins/Experimental/Animation/PoseSearch/Source/Runtime/Private/PoseSearchLibrary.cpp
--- a/Engine/Plugins/Experimental/Animation/PoseSearch/Source/Runtime/Private/PoseSearchLibrary.cpp
+++ b/Engine/Plugins/Experimental/Animation/PoseSearch/Source/Runtime/Private/PoseSearchLibrary.cpp
@@ -208,6 +208,12 @@ void UPoseSearchLibrary::TraceMotionMatchingState(
 
 	const int32 CurrentPoseIdx = bSearch && CurrentResult.PoseCost.IsValid() ? CurrentResult.PoseIdx : -1;
 	FTraceMotionMatchingState TraceState;
+
+	for (const UPoseSearchDatabase* SearchedDatabase : SearchContext.BestCandidates.GetSearchedDatabases())
+	{
+		AddUniqueDatabase(TraceState.DatabaseEntries, SearchedDatabase, SearchContext);
+	}
+
 	while (!SearchContext.BestCandidates.IsEmpty())
 	{
 		FSearchContext::FPoseCandidate PoseCandidate;
diff --git a/Engine/Plugins/Experimental/Animation/PoseSearch/Source/Runtime/Public/PoseSearch/PoseSearchContext.h b/Engine/Plugins/Experimental/Animation/PoseSearch/Source/Runtime/Public/PoseSearch/PoseSearchContext.h
--- a/Engine/Plugins/Experimental/Animation/PoseSearch/Source/Runtime/Public/PoseSearch/PoseSearchContext.h
+++ b/Engine/Plugins/Experimental/Animation/PoseSearch/Source/Runtime/Public/PoseSearch/PoseSearchContext.h
@@ -211,6 +211,11 @@ public:
 
 	struct FBestPoseCandidates
 	{
+		void Add(const UPoseSearchDatabase* Database)
+		{
+			SearchedDatabases.Add(Database);
+		}
+
 		void Add(const FPoseSearchCost& Cost, int32 PoseIdx, const UPoseSearchDatabase* Database, EPoseCandidateFlags PoseCandidateFlags)
 		{
 			if (EPoseCandidateFlags* PoseIdxPoseCandidateFlags = PoseIdxToFlags.Find(PoseIdx))
@@ -252,6 +257,8 @@ public:
 				PoseCandidateHeap.HeapPush(PoseCandidate);
 				PoseIdxToFlags.Add(PoseIdx, PoseCandidateFlags);
 			}
+
+			Add(Database);
 		}
 
 		void Pop(FPoseCandidate& OutItem)
@@ -270,7 +277,13 @@ public:
 			MaxPoseCandidates = Value;
 		}
 
+		const TSet<const UPoseSearchDatabase*>& GetSearchedDatabases() const
+		{
+			return SearchedDatabases;
+		}
+
 	private:
+		TSet<const UPoseSearchDatabase*> SearchedDatabases;
 		TArray<FPoseCandidateBase> PoseCandidateHeap;
 		TMap<int32, EPoseCandidateFlags> PoseIdxToFlags;
 		int32 MaxPoseCandidates = 200;
