motion matching - fix for error C2039: 'BestPosePos': is not a member of 'UE::PoseSearch::FSearchResult'

diff --git a/Engine/Plugins/Experimental/Animation/PoseSearch/Source/Runtime/Private/PoseSearchDatabase.cpp b/Engine/Plugins/Experimental/Animation/PoseSearch/Source/Runtime/Private/PoseSearchDatabase.cpp
--- a/Engine/Plugins/Experimental/Animation/PoseSearch/Source/Runtime/Private/PoseSearchDatabase.cpp
+++ b/Engine/Plugins/Experimental/Animation/PoseSearch/Source/Runtime/Private/PoseSearchDatabase.cpp
@@ -700,7 +700,7 @@ static inline void EvaluatePoseKernel(UE::PoseSearch::FSearchResult& Result, con
 			Result.PoseCost = PoseCost;
 			Result.PoseIdx = PoseIdx;
 
-#if WITH_EDITORONLY_DATA
+#if UE_POSE_SEARCH_TRACE_ENABLED
 			if (bUpdateBestCandidates)
 			{
 				Result.BestPosePos = ResultIndex;
