motion matching - clean up UE_POSE_SEARCH_TRACE_ENABLED macro usage

diff --git a/Engine/Plugins/Experimental/Animation/PoseSearch/Source/Runtime/Private/PoseSearchDatabase.cpp b/Engine/Plugins/Experimental/Animation/PoseSearch/Source/Runtime/Private/PoseSearchDatabase.cpp
--- a/Engine/Plugins/Experimental/Animation/PoseSearch/Source/Runtime/Private/PoseSearchDatabase.cpp
+++ b/Engine/Plugins/Experimental/Animation/PoseSearch/Source/Runtime/Private/PoseSearchDatabase.cpp
@@ -35,7 +35,7 @@ static void PopulateNonSelectableIdx(FNonSelectableIdx& NonSelectableIdx, FSearc
 	check(Database);
 #if UE_POSE_SEARCH_TRACE_ENABLED
 	const FSearchIndex& SearchIndex = Database->GetSearchIndex();
-#endif
+#endif // UE_POSE_SEARCH_TRACE_ENABLED
 
 	NonSelectableIdx.Reset();
 	const FSearchIndexAsset* CurrentIndexAsset = SearchContext.GetCurrentResult().GetSearchIndexAsset();
@@ -70,7 +70,7 @@ static void PopulateNonSelectableIdx(FNonSelectableIdx& NonSelectableIdx, FSearc
 				const TArray<float> PoseValues = SearchIndex.GetPoseValuesSafe(PoseIdx);
 				const FPoseSearchCost PoseCost = SearchIndex.ComparePoses(PoseIdx, 0.f, PoseValues, QueryValues);
 				SearchContext.BestCandidates.Add(PoseCost, PoseIdx, Database, EPoseCandidateFlags::DiscardedBy_PoseJumpThresholdTime);
-#endif
+#endif // UE_POSE_SEARCH_TRACE_ENABLED
 			}
 			else
 			{
@@ -103,7 +103,7 @@ static void PopulateNonSelectableIdx(FNonSelectableIdx& NonSelectableIdx, FSearc
 #if UE_POSE_SEARCH_TRACE_ENABLED
 				const FPoseSearchCost PoseCost = SearchIndex.ComparePoses(PoseIdx, 0.f, SearchIndex.GetPoseValuesSafe(PoseIdx), QueryValues);
 				SearchContext.BestCandidates.Add(PoseCost, PoseIdx, Database, EPoseCandidateFlags::DiscardedBy_PoseJumpThresholdTime);
-#endif
+#endif // UE_POSE_SEARCH_TRACE_ENABLED
 			}
 			else
 			{
@@ -131,7 +131,7 @@ static void PopulateNonSelectableIdx(FNonSelectableIdx& NonSelectableIdx, FSearc
 					const FPoseSearchCost PoseCost = SearchIndex.ComparePoses(HistoricalPoseIndex.PoseIndex, 0.f, SearchIndex.GetPoseValuesSafe(HistoricalPoseIndex.PoseIndex), QueryValues);
 					SearchContext.BestCandidates.Add(PoseCost, HistoricalPoseIndex.PoseIndex, Database, EPoseCandidateFlags::DiscardedBy_PoseReselectHistory);
 				}
-#endif
+#endif // UE_POSE_SEARCH_TRACE_ENABLED
 			}
 		}
 	}
@@ -167,7 +167,7 @@ struct FSearchFilters
 	bool AreFiltersValid(const FSearchIndex& SearchIndex, TConstArrayView<float> PoseValues, TConstArrayView<float> QueryValues, int32 PoseIdx, const FPoseMetadata& Metadata
 #if UE_POSE_SEARCH_TRACE_ENABLED
 		, UE::PoseSearch::FSearchContext& SearchContext, const UPoseSearchDatabase* Database
-#endif
+#endif // UE_POSE_SEARCH_TRACE_ENABLED
 	) const
 	{
 		for (const IPoseSearchFilter* Filter : Filters)
@@ -189,7 +189,7 @@ struct FSearchFilters
 					const FPoseSearchCost PoseCost = SearchIndex.ComparePoses(PoseIdx, 0.f, PoseValues, QueryValues);
 					SearchContext.BestCandidates.Add(PoseCost, PoseIdx, Database, EPoseCandidateFlags::DiscardedBy_PoseFilter);
 				}
-#endif
+#endif // UE_POSE_SEARCH_TRACE_ENABLED
 				return false;
 			}
 		}
@@ -691,7 +691,7 @@ static inline void EvaluatePoseKernel(UE::PoseSearch::FSearchResult& Result, con
 	if (SearchFilters.AreFiltersValid(SearchIndex, PoseValues, QueryValues, PoseIdx, SearchIndex.PoseMetadata[PoseIdx]
 #if UE_POSE_SEARCH_TRACE_ENABLED
 		, SearchContext, Database
-#endif
+#endif // UE_POSE_SEARCH_TRACE_ENABLED
 	))
 	{
 		const FPoseSearchCost PoseCost = bAlignedAndPadded ? SearchIndex.CompareAlignedPoses(PoseIdx, 0.f, PoseValues, QueryValues) : SearchIndex.ComparePoses(PoseIdx, 0.f, PoseValues, QueryValues);
@@ -705,7 +705,7 @@ static inline void EvaluatePoseKernel(UE::PoseSearch::FSearchResult& Result, con
 			{
 				Result.BestPosePos = ResultIndex;
 			}
-#endif // WITH_EDITORONLY_DATA
+#endif // UE_POSE_SEARCH_TRACE_ENABLED
 		}
 
 #if UE_POSE_SEARCH_TRACE_ENABLED
@@ -713,7 +713,7 @@ static inline void EvaluatePoseKernel(UE::PoseSearch::FSearchResult& Result, con
 		{
 			SearchContext.BestCandidates.Add(PoseCost, PoseIdx, Database, EPoseCandidateFlags::Valid_Pose);
 		}
-#endif
+#endif // UE_POSE_SEARCH_TRACE_ENABLED
 	}
 }
 
@@ -732,7 +732,7 @@ FPoseSearchCost UPoseSearchDatabase::SearchContinuingPose(UE::PoseSearch::FSearc
 	{
 		return ContinuingPoseCost;
 	}
-#endif
+#endif // WITH_EDITOR
 
 	// extracting notifies from the database animation asset at time SampleTime to search for UAnimNotifyState_PoseSearchOverrideContinuingPoseCostBias eventually overriding the schema ContinuingPoseCostBias
 	const FSearchIndex& SearchIndex = GetSearchIndex();
diff --git a/Engine/Plugins/Experimental/Animation/PoseSearch/Source/Runtime/Private/PoseSearchLibrary.cpp b/Engine/Plugins/Experimental/Animation/PoseSearch/Source/Runtime/Private/PoseSearchLibrary.cpp
--- a/Engine/Plugins/Experimental/Animation/PoseSearch/Source/Runtime/Private/PoseSearchLibrary.cpp
+++ b/Engine/Plugins/Experimental/Animation/PoseSearch/Source/Runtime/Private/PoseSearchLibrary.cpp
@@ -389,7 +389,7 @@ void UPoseSearchLibrary::UpdateMotionMatchingState(
 			InOutMotionMatchingState.RootMotionTransformDelta, Context.AnimInstanceProxy->GetAnimInstanceObject(), Context.GetCurrentNodeId(), DeltaTime, bSearch,
 			AnimInstance ? FObjectTrace::GetWorldElapsedTime(AnimInstance->GetWorld()) : 0.f, SearchBestCost, SearchBruteForceCost, InOutMotionMatchingState.CurrentSearchResult.BestPosePos);
 	}
-#endif
+#endif // UE_POSE_SEARCH_TRACE_ENABLED
 
 #if WITH_EDITORONLY_DATA && ENABLE_ANIM_DEBUG
 	const FSearchResult& CurResult = InOutMotionMatchingState.CurrentSearchResult;
diff --git a/Engine/Plugins/Experimental/Animation/PoseSearch/Source/Runtime/Private/Trace/PoseSearchTraceLogger.cpp b/Engine/Plugins/Experimental/Animation/PoseSearch/Source/Runtime/Private/Trace/PoseSearchTraceLogger.cpp
--- a/Engine/Plugins/Experimental/Animation/PoseSearch/Source/Runtime/Private/Trace/PoseSearchTraceLogger.cpp
+++ b/Engine/Plugins/Experimental/Animation/PoseSearch/Source/Runtime/Private/Trace/PoseSearchTraceLogger.cpp
@@ -37,9 +37,9 @@ bool IsTracing(const FAnimationBaseContext& InContext)
 
 	check(InContext.AnimInstanceProxy);
 	return !CANNOT_TRACE_OBJECT(InContext.AnimInstanceProxy->GetSkelMeshComponent());
-#else
+#else // UE_POSE_SEARCH_TRACE_ENABLED
 	return false;
-#endif
+#endif // UE_POSE_SEARCH_TRACE_ENABLED
 }
 
 FArchive& operator<<(FArchive& Ar, FTraceMessage& State)
diff --git a/Engine/Plugins/Experimental/Animation/PoseSearch/Source/Runtime/Public/PoseSearch/PoseSearchContext.h b/Engine/Plugins/Experimental/Animation/PoseSearch/Source/Runtime/Public/PoseSearch/PoseSearchContext.h
--- a/Engine/Plugins/Experimental/Animation/PoseSearch/Source/Runtime/Public/PoseSearch/PoseSearchContext.h
+++ b/Engine/Plugins/Experimental/Animation/PoseSearch/Source/Runtime/Public/PoseSearch/PoseSearchContext.h
@@ -290,7 +290,7 @@ public:
 	};
 	
 	FBestPoseCandidates BestCandidates;
-#endif
+#endif // UE_POSE_SEARCH_TRACE_ENABLED
 };
 
 POSESEARCH_API FTransform MirrorTransform(const FTransform& InTransform, EAxis::Type MirrorAxis, const FQuat& ReferenceRotation);
