Blend Stack:
-Added initial support for blend stack graphs.
-Changed Max Active Blends to be a compile time fixed value for blend stack.
-review-25906631 [REDACTED_HANDLE], [REDACTED_HANDLE], [REDACTED_HANDLE]

diff --git a/Engine/Plugins/Experimental/Animation/PoseSearch/Source/Editor/Private/AnimGraphNode_BlendStack.cpp b/Engine/Plugins/Experimental/Animation/PoseSearch/Source/Editor/Private/AnimGraphNode_BlendStack.cpp
--- a/Engine/Plugins/Experimental/Animation/PoseSearch/Source/Editor/Private/AnimGraphNode_BlendStack.cpp
+++ b/Engine/Plugins/Experimental/Animation/PoseSearch/Source/Editor/Private/AnimGraphNode_BlendStack.cpp
@@ -2,6 +2,14 @@
 
 #include "AnimGraphNode_BlendStack.h"
 #include "Animation/AnimRootMotionProvider.h"
+#include "AnimationBlendStackGraphSchema.h"
+#include "AnimationBlendStackGraph.h"
+#include "Kismet2/BlueprintEditorUtils.h"
+#include "Kismet2/KismetEditorUtilities.h"
+#include "EdGraphUtilities.h"
+#include "IAnimBlueprintCompilationContext.h"
+#include "AnimGraphNode_Root.h"
+#include "AnimGraphNode_BlendStackInput.h"
 
 #define LOCTEXT_NAMESPACE "AnimGraphNode_BlendStack"
 
@@ -26,7 +34,7 @@ FText UAnimGraphNode_BlendStack::GetMenuCategory() const
 	return LOCTEXT("NodeCategory", "Pose Search");
 }
 
-void UAnimGraphNode_BlendStack::GetOutputLinkAttributes(FNodeAttributeArray& OutAttributes) const
+void UAnimGraphNode_BlendStack_Base::GetOutputLinkAttributes(FNodeAttributeArray& OutAttributes) const
 {
 	if (UE::Anim::IAnimRootMotionProvider::Get())
 	{
@@ -40,24 +48,229 @@ void UAnimGraphNode_BlendStack::BakeDataDuringCompilation(class FCompilerResults
 	AnimBlueprint->FindOrAddGroup(Node.GetGroupName());
 }
 
-bool UAnimGraphNode_BlendStack::DoesSupportTimeForTransitionGetter() const
+void UAnimGraphNode_BlendStack_Base::ExpandGraphAndProcessNodes(
+	UEdGraph* SourceGraph, 
+	UAnimGraphNode_Base* SourceRootNode, UAnimGraphNode_Base* SourceInputNode,
+	IAnimBlueprintCompilationContext& InCompilationContext, IAnimBlueprintGeneratedClassCompiledData& OutCompiledData,
+	UAnimGraphNode_Base*& OutRootNode, UAnimGraphNode_Base*& OutInputNode)
+{
+	// Note: This is mostly copied from UAnimGraphNode_BlendSpaceGraphBase::ExpandGraphAndProcessNodes
+	
+	// Clone the nodes from the source graph
+	// Note that we outer this graph to the ConsolidatedEventGraph to allow ExpansionStep to 
+	// correctly retrieve the context for any expanded function calls (custom make/break structs etc.)
+	UEdGraph* ClonedGraph = FEdGraphUtilities::CloneGraph(SourceGraph, InCompilationContext.GetConsolidatedEventGraph(), &InCompilationContext.GetMessageLog(), true);
+
+	// Grab all the animation nodes and find the corresponding 
+	// root node and input pose node in the cloned set
+	TArray<UAnimGraphNode_Base*> AnimNodeList;
+
+	const UObject* SourceRootObject = InCompilationContext.GetMessageLog().FindSourceObject(SourceRootNode);
+	const UObject* SourceInputObject = InCompilationContext.GetMessageLog().FindSourceObject(SourceInputNode);
+	for (auto NodeIt = ClonedGraph->Nodes.CreateIterator(); NodeIt; ++NodeIt)
+	{
+		UEdGraphNode* ClonedNode = *NodeIt;
+
+		if (UAnimGraphNode_Base* TestNode = Cast<UAnimGraphNode_Base>(ClonedNode))
+		{
+			AnimNodeList.Add(TestNode);
+
+			//@TODO: There ought to be a better way to determine this
+			const UObject* TestObject = InCompilationContext.GetMessageLog().FindSourceObject(TestNode);
+			if (TestObject == SourceRootObject)
+			{
+				OutRootNode = TestNode;
+			}
+			if (TestObject == SourceInputObject)
+			{
+				OutInputNode = TestNode;
+			}
+		}
+	}
+
+	check(OutRootNode && OutInputNode);
+
+	// Run another expansion pass to catch the graph we just added (this is slightly wasteful)
+	InCompilationContext.ExpansionStep(ClonedGraph, false);
+
+	// Validate graph now we have expanded/pruned
+	InCompilationContext.ValidateGraphIsWellFormed(ClonedGraph);
+
+	// Move the cloned nodes into the consolidated event graph
+	const bool bIsLoading = InCompilationContext.GetBlueprint()->bIsRegeneratingOnLoad || IsAsyncLoading();
+	const bool bIsCompiling = InCompilationContext.GetBlueprint()->bBeingCompiled;
+	ClonedGraph->MoveNodesToAnotherGraph(InCompilationContext.GetConsolidatedEventGraph(), bIsLoading, bIsCompiling);
+
+	// Process any animation nodes
+	{
+		TArray<UAnimGraphNode_Base*> RootSet;
+		RootSet.Add(OutRootNode);
+
+		InCompilationContext.PruneIsolatedAnimationNodes(RootSet, AnimNodeList);
+		InCompilationContext.ProcessAnimationNodes(AnimNodeList);
+	}
+}
+
+void UAnimGraphNode_BlendStack_Base::OnProcessDuringCompilation(IAnimBlueprintCompilationContext& InCompilationContext, IAnimBlueprintGeneratedClassCompiledData& OutCompiledData)
+{
+	Super::OnProcessDuringCompilation(InCompilationContext, OutCompiledData);
+
+	FAnimNode_BlendStack_Standalone* AnimNode = GetBlendStackNode();
+	check(AnimNode);
+
+	const int32 MaxBlendsNum = GetMaxActiveBlends();
+	// Set MaxActiveBlends of the blend stack at compile time.
+	// @todo: Allow reducing the blend stack size (but not decreasing), to enable scalability control under i.e. lower LODs
+	AnimNode->MaxActiveBlends = GetMaxActiveBlends();
+
+	UAnimationBlendStackGraph* SampleGraph = CastChecked<UAnimationBlendStackGraph>(BoundGraph);
+	if (!SampleGraph->ResultNode->IsNodeRootSet())
+	{
+		// Input Pose is connected to Output Pose, so the sample graph does nothing.
+		// No need to use allocate sample graphs in that case.
+		AnimNode->SampleGraphPoseLinks.Reset();
+		return;
+	}
+
+	// Allocate one sample graph per-active blend plus an extra one for the stored pose.
+	AnimNode->SampleGraphPoseLinks.SetNum(MaxBlendsNum + 1);
+
+	for(FBlendStack_SampleGraphPoseLink& GraphPoseLink : AnimNode->SampleGraphPoseLinks)
+	{
+		UAnimGraphNode_Base *RootNode, *InputNode;
+		ExpandGraphAndProcessNodes(SampleGraph, SampleGraph->ResultNode, SampleGraph->InputNode, InCompilationContext, OutCompiledData, RootNode, InputNode);
+
+		// Blend stack node is potentially nested in the struct, so we can't use FPoseLinkMappingRecord. Patch at runtime instead.
+		GraphPoseLink.RootNodeIndex = InCompilationContext.GetAllocationIndexOfNode(RootNode);
+		GraphPoseLink.InputPoseNodeIndex = InCompilationContext.GetAllocationIndexOfNode(InputNode);
+	}
+}
+
+bool UAnimGraphNode_BlendStack_Base::DoesSupportTimeForTransitionGetter() const
 {
 	return true;
 }
 
-UAnimationAsset* UAnimGraphNode_BlendStack::GetAnimationAsset() const
+UAnimationAsset* UAnimGraphNode_BlendStack_Base::GetAnimationAsset() const
 {
 	return nullptr;
 }
 
-const TCHAR* UAnimGraphNode_BlendStack::GetTimePropertyName() const
+const TCHAR* UAnimGraphNode_BlendStack_Base::GetTimePropertyName() const
 {
 	return TEXT("InternalTimeAccumulator");
 }
 
-UScriptStruct* UAnimGraphNode_BlendStack::GetTimePropertyStruct() const
+UScriptStruct* UAnimGraphNode_BlendStack_Base::GetTimePropertyStruct() const
 {
 	return FAnimNode_BlendStack::StaticStruct();
 }
 
+void UAnimGraphNode_BlendStack_Base::Serialize(FArchive& Ar)
+{
+	Super::Serialize(Ar);
+
+	if (Ar.IsLoading() && (BoundGraph == nullptr))
+	{
+		// Update older nodes without a graph to create one on load.
+		CreateGraph();
+	}
+}
+
+int32 UAnimGraphNode_BlendStack_Base::GetMaxActiveBlends() const
+{
+	FAnimNode_BlendStack_Standalone* AnimNode = GetBlendStackNode();
+	check(AnimNode);
+	return AnimNode->MaxActiveBlends;
+}
+
+void UAnimGraphNode_BlendStack_Base::CreateGraph()
+{
+	// Create a new animation graph
+	check(BoundGraph == nullptr);
+	BoundGraph = FBlueprintEditorUtils::CreateNewGraph(
+		this,
+		NAME_None,
+		UAnimationBlendStackGraph::StaticClass(),
+		UAnimationBlendStackGraphSchema::StaticClass());
+	check(BoundGraph);
+
+	// Initialize the anim graph
+	const UEdGraphSchema* Schema = BoundGraph->GetSchema();
+	Schema->CreateDefaultNodesForGraph(*BoundGraph);
+
+	// Add the new graph as a child of our parent graph
+	UEdGraph* ParentGraph = GetGraph();
+	if(ParentGraph->SubGraphs.Find(BoundGraph) == INDEX_NONE)
+	{
+		ParentGraph->Modify();
+		ParentGraph->SubGraphs.Add(BoundGraph);
+	}
+}
+
+void UAnimGraphNode_BlendStack_Base::PostPlacedNewNode()
+{
+	Super::PostPlacedNewNode();
+	CreateGraph();
+}
+
+TArray<UEdGraph*> UAnimGraphNode_BlendStack_Base::GetSubGraphs() const
+{
+	return TArray<UEdGraph*>({BoundGraph});
+}
+
+UObject* UAnimGraphNode_BlendStack_Base::GetJumpTargetForDoubleClick() const
+{
+	// Open the blend stack graph
+	return BoundGraph;
+}
+
+void UAnimGraphNode_BlendStack_Base::JumpToDefinition() const
+{
+	if (UObject* HyperlinkTarget = GetJumpTargetForDoubleClick())
+	{
+		FKismetEditorUtilities::BringKismetToFocusAttentionOnObject(HyperlinkTarget);
+	}
+}
+
+void UAnimGraphNode_BlendStack_Base::DestroyNode()
+{
+	UBlueprint* Blueprint = GetBlueprint();
+	UEdGraph* GraphToRemove = BoundGraph;
+	BoundGraph = nullptr;
+	Super::DestroyNode();
+
+	if (GraphToRemove)
+	{
+		GraphToRemove->Modify();
+		FBlueprintEditorUtils::RemoveGraph(Blueprint, GraphToRemove, EGraphRemoveFlags::Recompile);
+	}
+}
+
+void UAnimGraphNode_BlendStack_Base::PostPasteNode()
+{
+	Super::PostPasteNode();
+
+	if(BoundGraph)
+	{
+		// Add the new graph as a child of our parent graph
+		UEdGraph* ParentGraph = GetGraph();
+
+		if(ParentGraph->SubGraphs.Find(BoundGraph) == INDEX_NONE)
+		{
+			ParentGraph->SubGraphs.Add(BoundGraph);
+		}
+
+		for (UEdGraphNode* GraphNode : BoundGraph->Nodes)
+		{
+			GraphNode->CreateNewGuid();
+			GraphNode->PostPasteNode();
+			GraphNode->ReconstructNode();
+		}
+
+		//restore transactional flag that is lost during copy/paste process
+		BoundGraph->SetFlags(RF_Transactional);
+	}
+}
+
 #undef LOCTEXT_NAMESPACE
diff --git a/Engine/Plugins/Experimental/Animation/PoseSearch/Source/Editor/Private/AnimGraphNode_BlendStackInput.cpp b/Engine/Plugins/Experimental/Animation/PoseSearch/Source/Editor/Private/AnimGraphNode_BlendStackInput.cpp
new file mode 100644
--- /dev/null
+++ b/Engine/Plugins/Experimental/Animation/PoseSearch/Source/Editor/Private/AnimGraphNode_BlendStackInput.cpp
@@ -0,0 +1,34 @@
+// Copyright Epic Games, Inc. All Rights Reserved.
+
+#include "AnimGraphNode_BlendStackInput.h"
+#include "AnimationBlendStackGraph.h"
+
+#define LOCTEXT_NAMESPACE "BlendStackInput"
+
+FText UAnimGraphNode_BlendStackInput::GetNodeTitle(ENodeTitleType::Type TitleType) const
+{
+	FText DefaultTitle = LOCTEXT("Title", "Blend Stack Input");
+	return DefaultTitle;
+}
+
+FText UAnimGraphNode_BlendStackInput::GetMenuCategory() const
+{
+	return LOCTEXT("BlendStackInputCategory", "Animation|Blend Stack");
+}
+
+bool UAnimGraphNode_BlendStackInput::CanUserDeleteNode() const
+{
+	return false;
+}
+
+bool UAnimGraphNode_BlendStackInput::CanDuplicateNode() const
+{
+	return false;
+}
+
+bool UAnimGraphNode_BlendStackInput::IsCompatibleWithGraph(UEdGraph const* Graph) const
+{
+	return Graph->IsA(UAnimationBlendStackGraph::StaticClass());
+}
+
+#undef LOCTEXT_NAMESPACE
diff --git a/Engine/Plugins/Experimental/Animation/PoseSearch/Source/Editor/Private/AnimGraphNode_BlendStackResult.cpp b/Engine/Plugins/Experimental/Animation/PoseSearch/Source/Editor/Private/AnimGraphNode_BlendStackResult.cpp
new file mode 100644
--- /dev/null
+++ b/Engine/Plugins/Experimental/Animation/PoseSearch/Source/Editor/Private/AnimGraphNode_BlendStackResult.cpp
@@ -0,0 +1,18 @@
+// Copyright Epic Games, Inc. All Rights Reserved.
+
+#include "AnimGraphNode_BlendStackResult.h"
+#include "AnimGraphNode_BlendStackInput.h"
+
+bool UAnimGraphNode_BlendStackResult::IsNodeRootSet() const
+{
+	if (!Pins.IsEmpty() && !Pins[0]->LinkedTo.IsEmpty())
+	{
+		if (Pins[0]->LinkedTo[0]->GetOuter()->IsA<UAnimGraphNode_BlendStackInput>())
+		{
+			// If output is connected to input directly, graph is a no-op. Consider it not set for pruning purposes.
+			return false;
+		}
+	}
+
+	return true;
+}
diff --git a/Engine/Plugins/Experimental/Animation/PoseSearch/Source/Editor/Private/AnimGraphNode_MotionMatching.cpp b/Engine/Plugins/Experimental/Animation/PoseSearch/Source/Editor/Private/AnimGraphNode_MotionMatching.cpp
--- a/Engine/Plugins/Experimental/Animation/PoseSearch/Source/Editor/Private/AnimGraphNode_MotionMatching.cpp
+++ b/Engine/Plugins/Experimental/Animation/PoseSearch/Source/Editor/Private/AnimGraphNode_MotionMatching.cpp
@@ -26,35 +26,12 @@ FText UAnimGraphNode_MotionMatching::GetMenuCategory() const
 	return LOCTEXT("NodeCategory", "Pose Search");
 }
 
-void UAnimGraphNode_MotionMatching::GetOutputLinkAttributes(FNodeAttributeArray& OutAttributes) const
-{
-	if (UE::Anim::IAnimRootMotionProvider::Get())
-	{
-		OutAttributes.Add(UE::Anim::IAnimRootMotionProvider::AttributeName);
-	}
-}
-
 void UAnimGraphNode_MotionMatching::BakeDataDuringCompilation(class FCompilerResultsLog& MessageLog)
 {
 	UAnimBlueprint* AnimBlueprint = GetAnimBlueprint();
 	AnimBlueprint->FindOrAddGroup(Node.GetGroupName());
 }
 
-bool UAnimGraphNode_MotionMatching::DoesSupportTimeForTransitionGetter() const
-{
-	return true;
-}
-
-UAnimationAsset* UAnimGraphNode_MotionMatching::GetAnimationAsset() const
-{
-	return nullptr;
-}
-
-const TCHAR* UAnimGraphNode_MotionMatching::GetTimePropertyName() const
-{
-	return TEXT("InternalTimeAccumulator");
-}
-
 UScriptStruct* UAnimGraphNode_MotionMatching::GetTimePropertyStruct() const
 {
 	return FAnimNode_MotionMatching::StaticStruct();
diff --git a/Engine/Plugins/Experimental/Animation/PoseSearch/Source/Editor/Private/AnimationBlendStackGraphSchema.cpp b/Engine/Plugins/Experimental/Animation/PoseSearch/Source/Editor/Private/AnimationBlendStackGraphSchema.cpp
new file mode 100644
--- /dev/null
+++ b/Engine/Plugins/Experimental/Animation/PoseSearch/Source/Editor/Private/AnimationBlendStackGraphSchema.cpp
@@ -0,0 +1,41 @@
+// Copyright Epic Games, Inc. All Rights Reserved.
+
+#include "AnimationBlendStackGraphSchema.h"
+#include "AnimationBlendStackGraph.h"
+#include "AnimGraphNode_BlendStackInput.h"
+#include "AnimGraphNode_BlendStackResult.h"
+
+/////////////////////////////////////////////////////
+// UAnimationBlendStackGraphSchema
+
+void UAnimationBlendStackGraphSchema::CreateDefaultNodesForGraph(UEdGraph& Graph) const
+{
+	UAnimationBlendStackGraph* TypedGraph = CastChecked<UAnimationBlendStackGraph>(&Graph);
+
+	UAnimGraphNode_BlendStackResult* ResultNode;
+	UAnimGraphNode_BlendStackInput* InputNode;
+
+	{
+		// Create the result node
+		FGraphNodeCreator<UAnimGraphNode_BlendStackResult> NodeCreator(Graph);
+		ResultNode = NodeCreator.CreateNode();
+		NodeCreator.Finalize();
+		SetNodeMetaData(ResultNode, FNodeMetadata::DefaultGraphNode);
+		TypedGraph->ResultNode = ResultNode;
+	}
+
+	{
+		// Create the result node
+		FGraphNodeCreator<UAnimGraphNode_BlendStackInput> NodeCreator(Graph);
+		InputNode = NodeCreator.CreateNode();
+		NodeCreator.Finalize();
+		SetNodeMetaData(InputNode, FNodeMetadata::DefaultGraphNode);
+
+		// Move input node to the left of the output pose.
+		InputNode->NodePosX -= 200; 
+		TypedGraph->InputNode = InputNode;
+	}
+
+	// Connect input node to output pose.
+	InputNode->Pins[0]->MakeLinkTo(ResultNode->Pins[0]);
+}
\ No newline at end of file
diff --git a/Engine/Plugins/Experimental/Animation/PoseSearch/Source/Editor/Public/AnimGraphNode_BlendStack.h b/Engine/Plugins/Experimental/Animation/PoseSearch/Source/Editor/Public/AnimGraphNode_BlendStack.h
--- a/Engine/Plugins/Experimental/Animation/PoseSearch/Source/Editor/Public/AnimGraphNode_BlendStack.h
+++ b/Engine/Plugins/Experimental/Animation/PoseSearch/Source/Editor/Public/AnimGraphNode_BlendStack.h
@@ -2,15 +2,54 @@
 
 #pragma once
 
-#include "CoreMinimal.h"
-#include "UObject/ObjectMacros.h"
 #include "AnimGraphNode_AssetPlayerBase.h"
 #include "PoseSearch/AnimNode_BlendStack.h"
 #include "AnimGraphNode_BlendStack.generated.h"
 
 
 UCLASS(MinimalAPI, Experimental)
-class UAnimGraphNode_BlendStack : public UAnimGraphNode_AssetPlayerBase
+class UAnimGraphNode_BlendStack_Base : public UAnimGraphNode_AssetPlayerBase
+{
+	GENERATED_BODY()
+
+	virtual void GetOutputLinkAttributes(FNodeAttributeArray& OutAttributes) const override;
+	virtual void OnProcessDuringCompilation(IAnimBlueprintCompilationContext& InCompilationContext, IAnimBlueprintGeneratedClassCompiledData& OutCompiledData) override;
+
+	virtual bool DoesSupportTimeForTransitionGetter() const override;
+	virtual UAnimationAsset* GetAnimationAsset() const override;
+	virtual const TCHAR* GetTimePropertyName() const override;
+	virtual UScriptStruct* GetTimePropertyStruct() const override;
+
+	virtual void PostPlacedNewNode() override;
+	virtual UObject* GetJumpTargetForDoubleClick() const override;
+	virtual void JumpToDefinition() const override;
+	virtual void DestroyNode() override;
+	virtual void PostPasteNode() override;
+	virtual void Serialize(FArchive& Ar) override;
+
+	// Helper function for compilation
+	void ExpandGraphAndProcessNodes(
+		UEdGraph* SourceGraph, 
+		UAnimGraphNode_Base* SourceRootNode, UAnimGraphNode_Base* SourceInputNode,
+		IAnimBlueprintCompilationContext& InCompilationContext, IAnimBlueprintGeneratedClassCompiledData& OutCompiledData,
+		UAnimGraphNode_Base*& OutRootNode, UAnimGraphNode_Base*& OutInputNode);
+
+	virtual TArray<UEdGraph*> GetSubGraphs() const override;
+
+protected:
+	virtual FAnimNode_BlendStack_Standalone* GetBlendStackNode() const PURE_VIRTUAL(UAnimGraphNode_BlendStack_Base::GetBlendStackNode, return nullptr;);
+	virtual int32 GetMaxActiveBlends() const;
+
+private:
+
+	void CreateGraph();
+
+	UPROPERTY()
+	TObjectPtr<UEdGraph> BoundGraph = nullptr;
+};
+
+UCLASS(MinimalAPI, Experimental)
+class UAnimGraphNode_BlendStack : public UAnimGraphNode_BlendStack_Base
 {
 	GENERATED_BODY()
 
@@ -21,11 +60,8 @@ class UAnimGraphNode_BlendStack : public UAnimGraphNode_AssetPlayerBase
 	virtual FText GetTooltipText() const override;
 	virtual FText GetNodeTitle(ENodeTitleType::Type TitleType) const override;
 	virtual FText GetMenuCategory() const override;
-	virtual void GetOutputLinkAttributes(FNodeAttributeArray& OutAttributes) const override;
-
 	virtual void BakeDataDuringCompilation(class FCompilerResultsLog& MessageLog) override;
-	virtual bool DoesSupportTimeForTransitionGetter() const override;
-	virtual UAnimationAsset* GetAnimationAsset() const override;
-	virtual const TCHAR* GetTimePropertyName() const override;
-	virtual UScriptStruct* GetTimePropertyStruct() const override;
-};
+
+protected:
+	virtual FAnimNode_BlendStack_Standalone* GetBlendStackNode() const override { return (FAnimNode_BlendStack_Standalone*)(&Node); }
+};
\ No newline at end of file
diff --git a/Engine/Plugins/Experimental/Animation/PoseSearch/Source/Editor/Public/AnimGraphNode_BlendStackInput.h b/Engine/Plugins/Experimental/Animation/PoseSearch/Source/Editor/Public/AnimGraphNode_BlendStackInput.h
new file mode 100644
--- /dev/null
+++ b/Engine/Plugins/Experimental/Animation/PoseSearch/Source/Editor/Public/AnimGraphNode_BlendStackInput.h
@@ -0,0 +1,29 @@
+// Copyright Epic Games, Inc. All Rights Reserved.
+
+#pragma once
+
+#include "AnimGraphNode_Base.h"
+#include "PoseSearch/AnimNode_BlendStackInput.h"
+
+#include "AnimGraphNode_BlendStackInput.generated.h"
+
+UCLASS(MinimalAPI)
+class UAnimGraphNode_BlendStackInput : public UAnimGraphNode_Base
+{
+	GENERATED_BODY()
+
+public:
+	UPROPERTY(EditAnywhere, Category = "Inputs")
+	FAnimNode_BlendStackInput Node;
+
+	virtual FText GetNodeTitle(ENodeTitleType::Type TitleType) const override;
+	virtual FText GetMenuCategory() const override;
+
+	virtual bool CanUserDeleteNode() const override;
+	virtual bool CanDuplicateNode() const override;
+	virtual bool IsCompatibleWithGraph(UEdGraph const* Graph) const override;
+
+	/** @return whether this node is editable. Non-editable nodes are implemented from function interfaces */
+	bool IsEditable() const { return CanUserDeleteNode(); }
+
+};
\ No newline at end of file
diff --git a/Engine/Plugins/Experimental/Animation/PoseSearch/Source/Editor/Public/AnimGraphNode_BlendStackResult.h b/Engine/Plugins/Experimental/Animation/PoseSearch/Source/Editor/Public/AnimGraphNode_BlendStackResult.h
new file mode 100644
--- /dev/null
+++ b/Engine/Plugins/Experimental/Animation/PoseSearch/Source/Editor/Public/AnimGraphNode_BlendStackResult.h
@@ -0,0 +1,15 @@
+// Copyright Epic Games, Inc. All Rights Reserved.
+
+#pragma once
+
+#include "AnimGraphNode_Root.h"
+
+#include "AnimGraphNode_BlendStackResult.generated.h"
+
+UCLASS(MinimalAPI)
+class UAnimGraphNode_BlendStackResult : public UAnimGraphNode_Root
+{
+	GENERATED_BODY()
+
+	virtual bool IsNodeRootSet() const override;
+};
\ No newline at end of file
diff --git a/Engine/Plugins/Experimental/Animation/PoseSearch/Source/Editor/Public/AnimGraphNode_MotionMatching.h b/Engine/Plugins/Experimental/Animation/PoseSearch/Source/Editor/Public/AnimGraphNode_MotionMatching.h
--- a/Engine/Plugins/Experimental/Animation/PoseSearch/Source/Editor/Public/AnimGraphNode_MotionMatching.h
+++ b/Engine/Plugins/Experimental/Animation/PoseSearch/Source/Editor/Public/AnimGraphNode_MotionMatching.h
@@ -4,13 +4,13 @@
 
 #include "CoreMinimal.h"
 #include "UObject/ObjectMacros.h"
-#include "AnimGraphNode_AssetPlayerBase.h"
+#include "AnimGraphNode_BlendStack.h"
 #include "PoseSearch/AnimNode_MotionMatching.h"
 #include "AnimGraphNode_MotionMatching.generated.h"
 
 
 UCLASS(MinimalAPI, Experimental)
-class UAnimGraphNode_MotionMatching : public UAnimGraphNode_AssetPlayerBase
+class UAnimGraphNode_MotionMatching : public UAnimGraphNode_BlendStack_Base
 {
 	GENERATED_BODY()
 
@@ -21,11 +21,9 @@ class UAnimGraphNode_MotionMatching : public UAnimGraphNode_AssetPlayerBase
 	virtual FText GetTooltipText() const override;
 	virtual FText GetNodeTitle(ENodeTitleType::Type TitleType) const override;
 	virtual FText GetMenuCategory() const override;
-	virtual void GetOutputLinkAttributes(FNodeAttributeArray& OutAttributes) const override;
-
-	virtual void BakeDataDuringCompilation(class FCompilerResultsLog& MessageLog) override;
-	virtual bool DoesSupportTimeForTransitionGetter() const override;
-	virtual UAnimationAsset* GetAnimationAsset() const override;
-	virtual const TCHAR* GetTimePropertyName() const override;
 	virtual UScriptStruct* GetTimePropertyStruct() const override;
+	virtual void BakeDataDuringCompilation(class FCompilerResultsLog& MessageLog) override;
+
+	virtual int32 GetMaxActiveBlends() const override { return Node.MaxActiveBlends; }
+	virtual FAnimNode_BlendStack_Standalone* GetBlendStackNode() const override { return (FAnimNode_BlendStack_Standalone*)(&Node.BlendStackNode); }
 };
diff --git a/Engine/Plugins/Experimental/Animation/PoseSearch/Source/Editor/Public/AnimationBlendStackGraph.h b/Engine/Plugins/Experimental/Animation/PoseSearch/Source/Editor/Public/AnimationBlendStackGraph.h
new file mode 100644
--- /dev/null
+++ b/Engine/Plugins/Experimental/Animation/PoseSearch/Source/Editor/Public/AnimationBlendStackGraph.h
@@ -0,0 +1,29 @@
+// Copyright Epic Games, Inc. All Rights Reserved.
+
+#pragma once
+
+#include "EdGraph/EdGraph.h"
+
+#include "AnimationBlendStackGraph.generated.h"
+
+class UObject;
+
+/** Animation graph used for blend stacks.
+*	The result node is the root of the blend graph.
+*	The input node is dynamically linked to the blend stack's sample pose.
+*/
+UCLASS(MinimalAPI)
+class UAnimationBlendStackGraph : public UEdGraph
+{
+	GENERATED_BODY()
+
+public:
+	// Result node within the state's animation graph
+	UPROPERTY()
+	TObjectPtr<class UAnimGraphNode_Root> ResultNode = nullptr;
+
+	// Result node within the state's animation graph
+	UPROPERTY()
+	TObjectPtr<class UAnimGraphNode_BlendStackInput> InputNode = nullptr;
+};
+
diff --git a/Engine/Plugins/Experimental/Animation/PoseSearch/Source/Editor/Public/AnimationBlendStackGraphSchema.h b/Engine/Plugins/Experimental/Animation/PoseSearch/Source/Editor/Public/AnimationBlendStackGraphSchema.h
new file mode 100644
--- /dev/null
+++ b/Engine/Plugins/Experimental/Animation/PoseSearch/Source/Editor/Public/AnimationBlendStackGraphSchema.h
@@ -0,0 +1,16 @@
+// Copyright Epic Games, Inc. All Rights Reserved.
+
+#pragma once
+
+#include "AnimationGraphSchema.h"
+#include "AnimationBlendStackGraphSchema.generated.h"
+
+UCLASS(MinimalAPI)
+class UAnimationBlendStackGraphSchema : public UAnimationGraphSchema
+{
+	GENERATED_BODY()
+
+	//~ Begin UEdGraphSchema Interface.
+	virtual void CreateDefaultNodesForGraph(UEdGraph& Graph) const override;
+	//~ End UEdGraphSchema Interface.
+};
diff --git a/Engine/Plugins/Experimental/Animation/PoseSearch/Source/Runtime/Private/AnimNode_BlendStack.cpp b/Engine/Plugins/Experimental/Animation/PoseSearch/Source/Runtime/Private/AnimNode_BlendStack.cpp
--- a/Engine/Plugins/Experimental/Animation/PoseSearch/Source/Runtime/Private/AnimNode_BlendStack.cpp
+++ b/Engine/Plugins/Experimental/Animation/PoseSearch/Source/Runtime/Private/AnimNode_BlendStack.cpp
@@ -8,6 +8,7 @@
 #include "Animation/BlendSpace.h"
 #include "Animation/AnimMontage.h"
 #include "PoseSearch/PoseSearchDefines.h"
+#include "PoseSearch/AnimNode_BlendStackInput.h"
 
 #include UE_INLINE_GENERATED_CPP_BY_NAME(AnimNode_BlendStack)
 
@@ -18,7 +19,7 @@ TAutoConsoleVariable<int32> CVarAnimBlendStackPruningEnable(TEXT("a.AnimNode.Ble
 
 /////////////////////////////////////////////////////
 // FPoseSearchAnimPlayer
-void FPoseSearchAnimPlayer::Initialize(UAnimationAsset* AnimationAsset, float AccumulatedTime, bool bLoop, bool bMirrored, UMirrorDataTable* MirrorDataTable, float BlendTime, float RootBoneBlendTime, const UBlendProfile* BlendProfile, EAlphaBlendOption InBlendOption, FVector BlendParameters, float PlayRate)
+void FPoseSearchAnimPlayer::Initialize(UAnimationAsset* AnimationAsset, float AccumulatedTime, bool bLoop, bool bMirrored, UMirrorDataTable* MirrorDataTable, float BlendTime, float RootBoneBlendTime, const UBlendProfile* BlendProfile, EAlphaBlendOption InBlendOption, FVector BlendParameters, float PlayRate, int32 InPoseLinkIdx)
 {
 	check(AnimationAsset);
 
@@ -104,6 +105,7 @@ void FPoseSearchAnimPlayer::Initialize(UAnimationAsset* AnimationAsset, float Ac
 	}
 
 	UpdateSourceLinkNode();
+	PoseLinkIndex = InPoseLinkIdx;
 }
 
 void FPoseSearchAnimPlayer::UpdatePlayRate(float PlayRate)
@@ -182,13 +184,10 @@ void FPoseSearchAnimPlayer::Evaluate_AnyThread(FPoseContext& Output)
 	}
 }
 
-void FPoseSearchAnimPlayer::Update_AnyThread(const FAnimationUpdateContext& Context, float BlendWeight)
+void FPoseSearchAnimPlayer::Update_AnyThread(const FAnimationUpdateContext& Context)
 {
-	const FAnimationUpdateContext AnimPlayerContext = Context.FractionalWeightAndRootMotion(BlendWeight, BlendWeight);
-
 	UpdateSourceLinkNode();
-	MirrorNode.Update_AnyThread(AnimPlayerContext);
-	CurrentBlendInTime += AnimPlayerContext.GetDeltaTime();
+	MirrorNode.Update_AnyThread(Context);
 }
 
 float FPoseSearchAnimPlayer::GetAccumulatedTime() const
@@ -275,8 +274,8 @@ bool FPoseSearchAnimPlayer::GetBlendInWeights(TArray<float>& Weights) const
 			}
 			else
 			{
-				const float unclampedLinearWeight = CurrentBlendInTime / TotalBlendInTimeBoneIdx;
-				Weights[BoneIdx] = FAlphaBlend::AlphaToBlendOption(unclampedLinearWeight, BlendOption);
+				const float UnclampedLinearWeight = CurrentBlendInTime / TotalBlendInTimeBoneIdx;
+				Weights[BoneIdx] = FAlphaBlend::AlphaToBlendOption(UnclampedLinearWeight, BlendOption);
 			}
 		}
 		return true;
@@ -300,9 +299,9 @@ void FAnimNode_BlendStack_Standalone::Evaluate_AnyThread(FPoseContext& Output)
 	}
 	else if (BlendStackSize == 1)
 	{
-		AnimPlayers.First().Evaluate_AnyThread(Output);
+		EvaluateSample(Output, 0);
 	}
-	else if (RequestedMaxActiveBlends <= 0 
+	else if (MaxActiveBlends <= 0 
 #if ENABLE_ANIM_DEBUG
 		|| !CVarAnimBlendStackEnable.GetValueOnAnyThread()
 #endif // ENABLE_ANIM_DEBUG
@@ -313,13 +312,13 @@ void FAnimNode_BlendStack_Standalone::Evaluate_AnyThread(FPoseContext& Output)
 		{
 			AnimPlayers.PopLast();
 		}
-		AnimPlayers.First().Evaluate_AnyThread(Output);
+		EvaluateSample(Output, 0);
 	}
 	else
 	{
 		// evaluating the last AnimPlayer into Output...
-		AnimPlayers[BlendStackSize - 1].Evaluate_AnyThread(Output);
-		 
+		EvaluateSample(Output, BlendStackSize - 1);
+
 		FPoseContext EvaluationPoseContext(Output);
 		FPoseContext BlendedPoseContext(Output); // @todo: this should not be necessary (but FBaseBlendedCurve::InitFrom complains about "ensure(&InCurveToInitFrom != this)"): optimize it away!
 		FAnimationPoseData BlendedAnimationPoseData(BlendedPoseContext);
@@ -331,50 +330,103 @@ void FAnimNode_BlendStack_Standalone::Evaluate_AnyThread(FPoseContext& Output)
 		const int32 NumSkeletonBones = RefSkeleton.GetNum();
 		TArray<float> Weights;
 
-		// evaluating from the second last to the first into EvaluationPoseContext and then blend it with the with the Output (initialized with the last AnimPlayer evaluation)
-		for (int32 i = BlendStackSize - 2; i >= 0; --i)
+		auto EvaluateAndBlendPlayerByIndex = [this, &EvaluationPoseContext, &BlendedAnimationPoseData, &Output, &Weights, &BlendedPoseContext](int32 PlayerIndex)
 		{
-			AnimPlayers[i].Evaluate_AnyThread(EvaluationPoseContext);
-			
-			if (AnimPlayers[i].GetBlendInWeights(Weights))
+			// Evaluate into EvaluationPoseContext and then blend it with the Output (initialized with the last AnimPlayer evaluation)
+			EvaluateSample(EvaluationPoseContext, PlayerIndex);
+			if (AnimPlayers[PlayerIndex].GetBlendInWeights(Weights))
 			{
 				// @todo: have BlendTwoPosesTogetherPerBone using a TArrayView for the Weights to avoid allocations
 				FAnimationRuntime::BlendTwoPosesTogetherPerBone(FAnimationPoseData(Output), FAnimationPoseData(EvaluationPoseContext), Weights, BlendedAnimationPoseData);
 			}
 			else
 			{
-				const float Weight = 1.f - FAlphaBlend::AlphaToBlendOption(AnimPlayers[i].GetBlendInPercentage(), AnimPlayers[i].GetBlendOption());
+				const float Weight = 1.f - FAlphaBlend::AlphaToBlendOption(AnimPlayers[PlayerIndex].GetBlendInPercentage(), AnimPlayers[PlayerIndex].GetBlendOption());
 				FAnimationRuntime::BlendTwoPosesTogether(FAnimationPoseData(Output), FAnimationPoseData(EvaluationPoseContext), Weight, BlendedAnimationPoseData);
 			}
 			Output = BlendedPoseContext; // @todo: this should not be necessary either: optimize it away!
+		};
 
-			if (i >= RequestedMaxActiveBlends
 #if ENABLE_ANIM_DEBUG
-				&& CVarAnimBlendStackPruningEnable.GetValueOnAnyThread()
+			const bool bEnablePruning = CVarAnimBlendStackPruningEnable.GetValueOnAnyThread() > 0;
+#else
+			const bool bEnablePruning = true;
 #endif // ENABLE_ANIM_DEBUG
-				)
+
+		// Evaluate our players from the second last to the first.
+		int32 PlayerIndex = BlendStackSize - 2;
+		// Start evaluating with our least significant players.
+		for (; PlayerIndex >= MaxActiveBlends; --PlayerIndex)
+		{
+			EvaluateAndBlendPlayerByIndex(PlayerIndex);
+
+			if (bEnablePruning)
 			{
-				// too many AnimPlayers! we don't have enough available blends to hold them all, so we accumulate the blended poses into Output / BlendedPoseContext, until...
+				// too many AnimPlayers! we don't have enough available blends to hold them all, so we accumulate the blended poses into Output / BlendedPoseContext.
 				AnimPlayers.PopLast();
-				
-				if (i == RequestedMaxActiveBlends)
-				{
-					check(AnimPlayers.Num() == RequestedMaxActiveBlends + 1);
-
-					// we can store Output / BlendedPoseContext into the last AnimPlayer, that will hold a static pose, no longer an animation playing
-					AnimPlayers[i].StorePoseContext(Output);
-				}
 			}
 		}
 
+		// Even if we're not pruning, we must use the stored pose if we have a limited number of graphs to execute.
+		const bool bNeedsStoredPose = (bEnablePruning || !SampleGraphPoseLinks.IsEmpty()) && (PlayerIndex == (MaxActiveBlends - 1));
+		if (bNeedsStoredPose)
+		{
+			check(AnimPlayers.Num() == MaxActiveBlends + 1);
+
+			// We store Output / BlendedPoseContext into the last AnimPlayer, that will hold a static pose, no longer an animation playing.
+			AnimPlayers.Last().StorePoseContext(Output);
+
+			if (!SampleGraphPoseLinks.IsEmpty())
+			{
+				const int32 PoseLinkIdx = AnimPlayers[MaxActiveBlends].GetPoseLinkIndex();
+				FBlendStack_SampleGraphPoseLink& PoseLink = SampleGraphPoseLinks[PoseLinkIdx];
+				// No players should have evaluated a graph before this point.
+				// Evaluate the graph on the blended result.
+				PoseLink.EvaluatePlayer(Output, AnimPlayers[MaxActiveBlends]);
+			}
+		}
+
+		// Continue with our most significant players.
+		for (; PlayerIndex >= 0; --PlayerIndex)
+		{
+			EvaluateAndBlendPlayerByIndex(PlayerIndex);
+		}
+
 		const int32 ActiveBlends = AnimPlayers.Num() - 1;
-		if (ActiveBlends > RequestedMaxActiveBlends)
+		if (ActiveBlends > MaxActiveBlends)
+		{
+			UE_LOG(LogPoseSearch, Display, TEXT("FAnimNode_BlendStack_Standalone NumBlends/MaxNumBlends %d / %d"), ActiveBlends, MaxActiveBlends);
+		}
+	}
+}
+
+void FAnimNode_BlendStack_Standalone::Initialize_AnyThread(const FAnimationInitializeContext& Context)
+{
+	Super::Initialize_AnyThread(Context);
+
+	IAnimClassInterface* AnimBlueprintClass = Context.GetAnimClass();
+	if (SampleGraphPoseLinks.IsEmpty() == false)
+	{
+		// Patch our pose links
+		for (FBlendStack_SampleGraphPoseLink& GraphPoseLink : SampleGraphPoseLinks)
 		{
-			UE_LOG(LogPoseSearch, Display, TEXT("FAnimNode_BlendStack_Standalone NumBlends/MaxNumBlends %d / %d"), ActiveBlends, RequestedMaxActiveBlends);
+			if (GraphPoseLink.InputPoseNodeIndex != INDEX_NONE)
+			{
+				GraphPoseLink.InputPose.LinkID = AnimBlueprintClass->GetAnimNodeProperties().Num() - 1 - GraphPoseLink.InputPoseNodeIndex;
+				GraphPoseLink.InputPose.AttemptRelink(Context);
+			}
+
+			if (GraphPoseLink.RootNodeIndex != INDEX_NONE)
+			{
+				GraphPoseLink.Root.LinkID = AnimBlueprintClass->GetAnimNodeProperties().Num() - 1 - GraphPoseLink.RootNodeIndex;
+			}
+
+			GraphPoseLink.CacheBoneCounter.Reset();
 		}
 	}
 }
 
+
 void FAnimNode_BlendStack_Standalone::UpdateAssetPlayer(const FAnimationUpdateContext& Context)
 {
 	QUICK_SCOPE_CYCLE_COUNTER(STAT_BlendStack_UpdateAssetPlayer);
@@ -400,8 +452,8 @@ void FAnimNode_BlendStack_Standalone::UpdateAssetPlayer(const FAnimationUpdateCo
 			break;
 		}
 
-		AnimPlayer.Update_AnyThread(Context, AnimPlayerBlendWeight);
-
+		const FAnimationUpdateContext AnimPlayerContext = Context.FractionalWeightAndRootMotion(BlendWeight, BlendWeight);
+		UpdateSample(AnimPlayerContext, AnimPlayerIndex);
 		CurrentWeightMultiplier *= (1.f - BlendInPercentage);
 	}
 
@@ -413,16 +465,116 @@ void FAnimNode_BlendStack_Standalone::UpdateAssetPlayer(const FAnimationUpdateCo
 	}
 }
 
+bool FAnimNode_BlendStack_Standalone::IsSampleGraphAvailableForPlayer(const int32 PlayerIndex)
+{
+	// If we have any sample graphs, our player has been assigned a pose link index.
+	// If we are within X most relelvant players, then the graph is available.
+	return !SampleGraphPoseLinks.IsEmpty() && (PlayerIndex < MaxActiveBlends);
+}
+
+void FAnimNode_BlendStack_Standalone::EvaluateSample(FPoseContext& Output, const int32 PlayerIndex)
+{
+	FPoseSearchAnimPlayer& SamplePlayer = AnimPlayers[PlayerIndex];
+	// If we have any sample graphs, our player has been assigned a pose link index.
+	// If we are within X most relelvant players, then the graph is available.
+	// If PlayerIndex == MaxActiveBlends, don't evaluate that graph. It's reserved for the stored pose.
+	const bool bIsSampleGraphAvailable = !SampleGraphPoseLinks.IsEmpty() && (PlayerIndex < MaxActiveBlends);
+	if (!bIsSampleGraphAvailable)
+	{
+		// If we have no sample graph, evaluate the player directly.
+		SamplePlayer.Evaluate_AnyThread(Output);
+		return;
+	}
+
+	FBlendStack_SampleGraphPoseLink& PoseLink = SampleGraphPoseLinks[SamplePlayer.GetPoseLinkIndex()];
+	PoseLink.EvaluatePlayer(Output, SamplePlayer);
+}
+
+void FBlendStack_SampleGraphPoseLink::EvaluatePlayer(FPoseContext& Output, FPoseSearchAnimPlayer& SamplePlayer)
+{
+	SetInputPosePlayer(SamplePlayer);
+
+	// Make sure CacheBones has been called before evaluating.
+	ConditionalCacheBones(Output);
+	// The anim player may or may not have its Evaluate_AnyThread called through the graph update. 
+	Root.Evaluate(Output);
+}
+
+void FBlendStack_SampleGraphPoseLink::ConditionalCacheBones(const FAnimationBaseContext& Context)
+{
+	// Only call CacheBones when needed.
+	if (!CacheBoneCounter.IsSynchronized_Counter(Context.AnimInstanceProxy->GetCachedBonesCounter()))
+	{
+		// Keep track of samples that have had CacheBones called on.
+		CacheBoneCounter.SynchronizeWith(Context.AnimInstanceProxy->GetCachedBonesCounter());
+
+		FAnimationCacheBonesContext CacheBoneContext(Context.AnimInstanceProxy);
+		Root.CacheBones(CacheBoneContext);
+	}
+}
+
+void FAnimNode_BlendStack_Standalone::UpdateSample(const FAnimationUpdateContext& Context, const int32 PlayerIndex)
+{
+	FPoseSearchAnimPlayer& SamplePlayer = AnimPlayers[PlayerIndex];
+
+	// If we have any sample graphs, our player has been assigned a pose link index.
+	// If we are within X most relelvant players, then the graph is available.
+	// @todo: If PlayerIndex == MaxActiveBlends, this will likely become a stored pose. What do we update that graph with? 
+	// For now, just use the same player.
+	const bool bHasSampleGraph = !SampleGraphPoseLinks.IsEmpty() && (PlayerIndex <= MaxActiveBlends);
+	if (bHasSampleGraph)
+	{
+		FBlendStack_SampleGraphPoseLink& PoseLink = SampleGraphPoseLinks[SamplePlayer.GetPoseLinkIndex()];
+		PoseLink.SetInputPosePlayer(SamplePlayer);
+		// The anim player may or may not have its Update_AnyThread called through the graph update. 
+		PoseLink.Root.Update(Context);
+	}
+	else
+	{
+		// If we have no sample graph, update the player directly.
+		SamplePlayer.Update_AnyThread(Context);
+	}
+
+	// Advance the blend-in time regardless of whether or not the player was updated.
+	SamplePlayer.AdvanceBlendInTime(Context.GetDeltaTime());
+}
+
+void FAnimNode_BlendStack_Standalone::InitializeSample(const FAnimationInitializeContext& Context, FPoseSearchAnimPlayer& SamplePlayer)
+{
+	if (SamplePlayer.GetPoseLinkIndex() != INDEX_NONE)
+	{
+		FBlendStack_SampleGraphPoseLink& PoseLink = SampleGraphPoseLinks[SamplePlayer.GetPoseLinkIndex()];
+		PoseLink.Root.Initialize(Context);
+	}
+}
+
 float FAnimNode_BlendStack_Standalone::GetAccumulatedTime() const
 {
 	return AnimPlayers.IsEmpty() ? 0.f : AnimPlayers.First().GetAccumulatedTime();
 }
 
-void FAnimNode_BlendStack_Standalone::BlendTo(UAnimationAsset* AnimationAsset, float AccumulatedTime, bool bLoop, bool bMirrored, UMirrorDataTable* MirrorDataTable, int32 MaxActiveBlends, float BlendTime, float RootBoneBlendTime, const UBlendProfile* BlendProfile, EAlphaBlendOption BlendOption, FVector BlendParameters, float PlayRate)
+void FAnimNode_BlendStack_Standalone::BlendTo(const FAnimationUpdateContext& Context, UAnimationAsset* AnimationAsset, float AccumulatedTime, bool bLoop, bool bMirrored, UMirrorDataTable* MirrorDataTable, float BlendTime, float RootBoneBlendTime, const UBlendProfile* BlendProfile, EAlphaBlendOption BlendOption, FVector BlendParameters, float PlayRate)
 {
-	RequestedMaxActiveBlends = MaxActiveBlends;
 	AnimPlayers.PushFirst(FPoseSearchAnimPlayer());
-	AnimPlayers.First().Initialize(AnimationAsset, AccumulatedTime, bLoop, bMirrored, MirrorDataTable, BlendTime, RootBoneBlendTime, BlendProfile, BlendOption, BlendParameters, PlayRate);
+	FPoseSearchAnimPlayer& AnimPlayer = AnimPlayers.First();
+	AnimPlayer.Initialize(AnimationAsset, AccumulatedTime, bLoop, bMirrored, MirrorDataTable, BlendTime, RootBoneBlendTime, BlendProfile, BlendOption, BlendParameters, PlayRate, GetNextPoseLinkIndex());
+
+	FAnimationInitializeContext InitContext(Context.AnimInstanceProxy, Context.SharedContext);
+	InitializeSample(InitContext, AnimPlayer);
+}
+
+int32 FAnimNode_BlendStack_Standalone::GetNextPoseLinkIndex()
+{
+	if (SampleGraphPoseLinks.IsEmpty())
+	{
+		return INDEX_NONE;
+	}
+
+	const int32 NumPoseLinks = SampleGraphPoseLinks.Num();
+	CurrentSamplePoseLink = ++CurrentSamplePoseLink;
+	if (CurrentSamplePoseLink == NumPoseLinks) { CurrentSamplePoseLink = 0; }
+
+	return CurrentSamplePoseLink;
 }
 
 void FAnimNode_BlendStack_Standalone::UpdatePlayRate(float PlayRate)
@@ -493,7 +645,7 @@ void FAnimNode_BlendStack::UpdateAssetPlayer(const FAnimationUpdateContext& Cont
 
 		if (bExecuteBlendTo)
 		{
-			BlendTo(AnimationAsset, AnimationTime, bLoop, bMirrored, MirrorDataTable.Get(), MaxActiveBlends, BlendTime, RootBoneBlendTime, BlendProfile, BlendOption, BlendParameters, WantedPlayRate);
+			BlendTo(Context, AnimationAsset, AnimationTime, bLoop, bMirrored, MirrorDataTable.Get(), BlendTime, RootBoneBlendTime, BlendProfile, BlendOption, BlendParameters, WantedPlayRate);
 		}
 	}
 	
@@ -501,3 +653,13 @@ void FAnimNode_BlendStack::UpdateAssetPlayer(const FAnimationUpdateContext& Cont
 
 	Super::UpdateAssetPlayer(Context);
 }
+
+void FBlendStack_SampleGraphPoseLink::SetInputPosePlayer(FPoseSearchAnimPlayer& Player)
+{
+	// Because our anim players may get reallocated, or change indices due to push/pops,
+	// we must call this before every operation that might end up needing the anim player through the graph's input node.
+	// @todo: There's probably a better way to do this.
+	FAnimNode_BlendStackInput* InputNode = static_cast<FAnimNode_BlendStackInput*>(InputPose.GetLinkNode());
+	// Link our anim player to the input pose node.
+	InputNode->Player = &Player;
+}
diff --git a/Engine/Plugins/Experimental/Animation/PoseSearch/Source/Runtime/Private/AnimNode_BlendStackInput.cpp b/Engine/Plugins/Experimental/Animation/PoseSearch/Source/Runtime/Private/AnimNode_BlendStackInput.cpp
new file mode 100644
--- /dev/null
+++ b/Engine/Plugins/Experimental/Animation/PoseSearch/Source/Runtime/Private/AnimNode_BlendStackInput.cpp
@@ -0,0 +1,20 @@
+// Copyright Epic Games, Inc. All Rights Reserved.
+
+#include "PoseSearch/AnimNode_BlendStackInput.h"
+#include "PoseSearch/AnimNode_BlendStack.h"
+#include "Animation/AnimInstanceProxy.h"
+
+#include UE_INLINE_GENERATED_CPP_BY_NAME(AnimNode_BlendStackInput)
+
+void FAnimNode_BlendStackInput::Evaluate_AnyThread(FPoseContext& Output)
+{
+	check(Player);
+	Player->Evaluate_AnyThread(Output);
+}
+
+void FAnimNode_BlendStackInput::Update_AnyThread(const FAnimationUpdateContext& Context)
+{
+	check(Player);
+	Player->Update_AnyThread(Context);
+}
+
diff --git a/Engine/Plugins/Experimental/Animation/PoseSearch/Source/Runtime/Private/AnimNode_MotionMatching.cpp b/Engine/Plugins/Experimental/Animation/PoseSearch/Source/Runtime/Private/AnimNode_MotionMatching.cpp
--- a/Engine/Plugins/Experimental/Animation/PoseSearch/Source/Runtime/Private/AnimNode_MotionMatching.cpp
+++ b/Engine/Plugins/Experimental/Animation/PoseSearch/Source/Runtime/Private/AnimNode_MotionMatching.cpp
@@ -176,9 +176,9 @@ void FAnimNode_MotionMatching::UpdateAssetPlayer(const FAnimationUpdateContext&
 			{
 				// root bone blending needs to be immediate if MM node controls the offset between mesh component and root bone
 				const float RootBoneBlendTime = YawFromAnimationBlendRate < 0.f ? BlendTime : 0.f;
-				BlendStackNode.BlendTo(DatabaseAsset->GetAnimationAsset(), MotionMatchingState.CurrentSearchResult.AssetTime,
+				BlendStackNode.BlendTo(Context, DatabaseAsset->GetAnimationAsset(), MotionMatchingState.CurrentSearchResult.AssetTime,
 					DatabaseAsset->IsLooping(), SearchIndexAsset->bMirrored, CurrentResultDatabase->Schema->MirrorDataTable.Get(),
-					MaxActiveBlends, BlendTime, RootBoneBlendTime, BlendProfile, BlendOption, SearchIndexAsset->BlendParameters, MotionMatchingState.WantedPlayRate);
+					BlendTime, RootBoneBlendTime, BlendProfile, BlendOption, SearchIndexAsset->BlendParameters, MotionMatchingState.WantedPlayRate);
 			}
 		}
 	}
diff --git a/Engine/Plugins/Experimental/Animation/PoseSearch/Source/Runtime/Public/PoseSearch/AnimNode_BlendStack.h b/Engine/Plugins/Experimental/Animation/PoseSearch/Source/Runtime/Public/PoseSearch/AnimNode_BlendStack.h
--- a/Engine/Plugins/Experimental/Animation/PoseSearch/Source/Runtime/Public/PoseSearch/AnimNode_BlendStack.h
+++ b/Engine/Plugins/Experimental/Animation/PoseSearch/Source/Runtime/Public/PoseSearch/AnimNode_BlendStack.h
@@ -15,10 +15,10 @@ struct FPoseSearchAnimPlayer
 {
 	GENERATED_BODY()
 	
-	void Initialize(UAnimationAsset* AnimationAsset, float AccumulatedTime, bool bLoop, bool bMirrored, UMirrorDataTable* MirrorDataTable, float BlendTime, float RootBoneBlendTime, const UBlendProfile* BlendProfile, EAlphaBlendOption InBlendOption, FVector BlendParameters, float PlayRate);
+	void Initialize(UAnimationAsset* AnimationAsset, float AccumulatedTime, bool bLoop, bool bMirrored, UMirrorDataTable* MirrorDataTable, float BlendTime, float RootBoneBlendTime, const UBlendProfile* BlendProfile, EAlphaBlendOption InBlendOption, FVector BlendParameters, float PlayRate, int32 InPoseLinkIdx);
 	void UpdatePlayRate(float PlayRate);
 	void Evaluate_AnyThread(FPoseContext& Output);
-	void Update_AnyThread(const FAnimationUpdateContext& Context, float BlendWeight);
+	void Update_AnyThread(const FAnimationUpdateContext& Context);
 	float GetAccumulatedTime() const;
 	
 	float GetBlendInPercentage() const;
@@ -28,14 +28,16 @@ struct FPoseSearchAnimPlayer
 
 	float GetTotalBlendInTime() const { return TotalBlendInTime; }
 	float GetCurrentBlendInTime() const { return CurrentBlendInTime; }
+	void AdvanceBlendInTime(const float DeltaTime) { CurrentBlendInTime += DeltaTime; }
 	bool GetMirror() const { return MirrorNode.GetMirror(); }
 	FVector GetBlendParameters() const;
 	FString GetAnimationName() const;
 	const UAnimationAsset* GetAnimationAsset() const;
 
 	FAnimNode_Mirror_Standalone& GetMirrorNode() { return MirrorNode; }
+	int32 GetPoseLinkIndex() const { return PoseLinkIndex; }
 
-protected:
+public:
 	void RestorePoseContext(FPoseContext& PoseContext) const;
 	void UpdateSourceLinkNode();
 
@@ -56,33 +58,71 @@ protected:
 	UE::Anim::FHeapAttributeContainer StoredAttributes;
 
 	EAlphaBlendOption BlendOption = EAlphaBlendOption::Linear;
-
+	int32 PoseLinkIndex = INDEX_NONE;
 	TCustomBoneIndexArray<float, FSkeletonPoseBoneIndex> TotalBlendInTimePerBone;
 
 	float TotalBlendInTime = 0.f;
 	float CurrentBlendInTime = 0.f;
 };
 
+USTRUCT(BlueprintInternalUseOnly)
+struct FBlendStack_SampleGraphPoseLink
+{
+	GENERATED_BODY()
+
+	UPROPERTY()
+	int32 InputPoseNodeIndex;
+	UPROPERTY()
+	int32 RootNodeIndex;
+
+	FPoseLink InputPose;
+	FPoseLink Root;
+	FGraphTraversalCounter CacheBoneCounter;
+
+	void SetInputPosePlayer(FPoseSearchAnimPlayer& Player);
+	void EvaluatePlayer(FPoseContext& Output, FPoseSearchAnimPlayer& SamplePlayer);
+	void ConditionalCacheBones(const FAnimationBaseContext& Output);
+};
+
 USTRUCT(BlueprintInternalUseOnly)
 struct POSESEARCH_API FAnimNode_BlendStack_Standalone : public FAnimNode_AssetPlayerBase
 {
 	GENERATED_BODY()
 
+	UPROPERTY()
+	TArray<FBlendStack_SampleGraphPoseLink> SampleGraphPoseLinks;
+	int32 CurrentSamplePoseLink = -1;
+
 	TDeque<FPoseSearchAnimPlayer> AnimPlayers;
-	int32 RequestedMaxActiveBlends = 0;
 
 	// FAnimNode_Base interface
 	virtual void Evaluate_AnyThread(FPoseContext& Output) override;
+	virtual void Initialize_AnyThread(const FAnimationInitializeContext& Context) override;
 	virtual void GatherDebugData(FNodeDebugData& DebugData) override;
 	// End of FAnimNode_Base interface
 
-	void BlendTo(UAnimationAsset* AnimationAsset, float AccumulatedTime = 0.f, bool bLoop = false, bool bMirrored = false, UMirrorDataTable* MirrorDataTable = nullptr, int32 MaxActiveBlends = 3, float BlendTime = 0.2f, float RootBoneBlendTime = -1.f, const UBlendProfile* BlendProfile = nullptr, EAlphaBlendOption BlendOption = EAlphaBlendOption::Linear, FVector BlendParameters = FVector::Zero(), float PlayRate = 1.f);
+	void BlendTo(const FAnimationUpdateContext& Context, UAnimationAsset* AnimationAsset, float AccumulatedTime = 0.f, bool bLoop = false, bool bMirrored = false, UMirrorDataTable* MirrorDataTable = nullptr, float BlendTime = 0.2f, float RootBoneBlendTime = -1.f, const UBlendProfile* BlendProfile = nullptr, EAlphaBlendOption BlendOption = EAlphaBlendOption::Linear, FVector BlendParameters = FVector::Zero(), float PlayRate = 1.f);
 	void UpdatePlayRate(float PlayRate);
 
 	// FAnimNode_AssetPlayerBase interface
 	virtual float GetAccumulatedTime() const override;
 	virtual void UpdateAssetPlayer(const FAnimationUpdateContext& Context) override;
 	// End of FAnimNode_AssetPlayerBase interface
+
+	int32 GetNextPoseLinkIndex();
+
+private:
+
+	void InitializeSample(const FAnimationInitializeContext& Context, FPoseSearchAnimPlayer& SamplePlayer);
+	void EvaluateSample(FPoseContext& Output, const int32 PlayerIndex);
+	void UpdateSample(const FAnimationUpdateContext& Context, const int32 PlayerIndex);
+	bool IsSampleGraphAvailableForPlayer(const int32 PlayerIndex);
+
+	// Number of max active blending animation in the blend stack. If MaxActiveBlends is zero then blend stack is disabled
+	UPROPERTY(EditAnywhere, Category = Settings, meta = (ClampMin = "0"))
+	int32 MaxActiveBlends = 4;
+
+	friend class UAnimGraphNode_BlendStack_Base;
 };
 
 
@@ -124,10 +164,6 @@ struct POSESEARCH_API FAnimNode_BlendStack : public FAnimNode_BlendStack_Standal
 	UPROPERTY(EditAnywhere, BlueprintReadWrite, Category = Settings, meta = (PinHiddenByDefault))
 	float MaxAnimationDeltaTime = -1.f;
 
-	// Number of max active blending animation in the blend stack. If MaxActiveBlends is zero then blend stack is disabled
-	UPROPERTY(EditAnywhere, BlueprintReadWrite, Category = Settings, meta = (PinHiddenByDefault, ClampMin = "0"))
-	int32 MaxActiveBlends = 4;
-
 	UPROPERTY(EditAnywhere, BlueprintReadWrite, Category = Settings, meta = (PinHiddenByDefault, UseAsBlendProfile = true))
 	TObjectPtr<UBlendProfile> BlendProfile;
 
diff --git a/Engine/Plugins/Experimental/Animation/PoseSearch/Source/Runtime/Public/PoseSearch/AnimNode_BlendStackInput.h b/Engine/Plugins/Experimental/Animation/PoseSearch/Source/Runtime/Public/PoseSearch/AnimNode_BlendStackInput.h
new file mode 100644
--- /dev/null
+++ b/Engine/Plugins/Experimental/Animation/PoseSearch/Source/Runtime/Public/PoseSearch/AnimNode_BlendStackInput.h
@@ -0,0 +1,24 @@
+// Copyright Epic Games, Inc. All Rights Reserved.
+
+#pragma once
+
+#include "Animation/AnimNodeBase.h"
+#include "AnimNode_BlendStackInput.generated.h"
+
+/** Input pose that links the blend stack's sample graph with the sample/pose chosen by the blend stack.
+*	@Todo: It might be better to reuse FAnimNode_LinkedInputPose, since we will most likely need variable input pins in the future too.
+*/
+USTRUCT(BlueprintInternalUseOnly, Experimental)
+struct POSESEARCH_API FAnimNode_BlendStackInput : public FAnimNode_Base
+{
+	GENERATED_BODY()
+	
+public:
+
+	virtual void Update_AnyThread(const FAnimationUpdateContext& Context) override;
+	virtual void Evaluate_AnyThread(FPoseContext& Output) override;
+
+	// The player is guaranteed to be valid for the whole duration of update/eval.
+	struct FPoseSearchAnimPlayer* Player;
+};
+
diff --git a/Engine/Plugins/Experimental/Animation/PoseSearch/Source/Runtime/Public/PoseSearch/AnimNode_MotionMatching.h b/Engine/Plugins/Experimental/Animation/PoseSearch/Source/Runtime/Public/PoseSearch/AnimNode_MotionMatching.h
--- a/Engine/Plugins/Experimental/Animation/PoseSearch/Source/Runtime/Public/PoseSearch/AnimNode_MotionMatching.h
+++ b/Engine/Plugins/Experimental/Animation/PoseSearch/Source/Runtime/Public/PoseSearch/AnimNode_MotionMatching.h
@@ -68,7 +68,7 @@ private:
 	float BlendTime = 0.2f;
 
 	// Number of max active animation segments being blended together in the blend stack. If MaxActiveBlends is zero then the blend stack is disabled.
-	UPROPERTY(EditAnywhere, Category = Settings, meta = (PinHiddenByDefault, ClampMin="0"))
+	UPROPERTY(EditAnywhere, Category = Settings, meta = (ClampMin="0"))
 	int32 MaxActiveBlends = 4;
 
 	// Set Blend Profiles (editable in the skeleton) to determine how the blending is distributed among your character's bones. It could be used to differentiate between upper body and lower body to blend timing.
@@ -114,6 +114,7 @@ private:
 	UPROPERTY(EditAnywhere, Category = RootMotion, meta = (PinHiddenByDefault))
 	float YawFromAnimationBlendRate = -1.f;
 
+	UPROPERTY()
 	FAnimNode_BlendStack_Standalone BlendStackNode;
 
 	// Encapsulated motion matching algorithm and internal state
@@ -137,4 +138,6 @@ private:
 	UPROPERTY(EditAnywhere, Category = Relevancy, meta = (FoldProperty, PinHiddenByDefault))
 	bool bIgnoreForRelevancyTest = false;
 #endif // WITH_EDITORONLY_DATA
+
+	friend class UAnimGraphNode_MotionMatching;
 };
\ No newline at end of file
diff --git a/Engine/Source/Editor/AnimGraph/Public/AnimGraphNode_Root.h b/Engine/Source/Editor/AnimGraph/Public/AnimGraphNode_Root.h
--- a/Engine/Source/Editor/AnimGraph/Public/AnimGraphNode_Root.h
+++ b/Engine/Source/Editor/AnimGraph/Public/AnimGraphNode_Root.h
@@ -19,23 +19,23 @@ class UAnimGraphNode_Root : public UAnimGraphNode_Base
 	FAnimNode_Root Node;
 
 	//~ Begin UEdGraphNode Interface.
-	virtual FLinearColor GetNodeTitleColor() const override;
-	virtual FText GetTooltipText() const override;
-	virtual FText GetNodeTitle(ENodeTitleType::Type TitleType) const override;
-	virtual bool CanUserDeleteNode() const override { return false; }
-	virtual bool CanDuplicateNode() const override { return false; }
-	virtual void GetMenuActions(FBlueprintActionDatabaseRegistrar& ActionRegistrar) const override;
-	virtual bool IsNodeRootSet() const override { return true; }
+	ANIMGRAPH_API virtual FLinearColor GetNodeTitleColor() const override;
+	ANIMGRAPH_API virtual FText GetTooltipText() const override;
+	ANIMGRAPH_API virtual FText GetNodeTitle(ENodeTitleType::Type TitleType) const override;
+	ANIMGRAPH_API virtual bool CanUserDeleteNode() const override { return false; }
+	ANIMGRAPH_API virtual bool CanDuplicateNode() const override { return false; }
+	ANIMGRAPH_API virtual void GetMenuActions(FBlueprintActionDatabaseRegistrar& ActionRegistrar) const override;
+	ANIMGRAPH_API virtual bool IsNodeRootSet() const override { return true; }
 	//~ End UEdGraphNode Interface.
 
 	//~ Begin UAnimGraphNode_Base Interface
-	virtual bool IsPoseWatchable() const override;
-	virtual bool IsSinkNode() const override;
-	virtual void OnProcessDuringCompilation(IAnimBlueprintCompilationContext& InCompilationContext, IAnimBlueprintGeneratedClassCompiledData& OutCompiledData) override;
-	virtual void OnCopyTermDefaultsToDefaultObject(IAnimBlueprintCopyTermDefaultsContext& InCompilationContext, IAnimBlueprintNodeCopyTermDefaultsContext& InPerNodeContext, IAnimBlueprintGeneratedClassCompiledData& OutCompiledData) override;
+	ANIMGRAPH_API virtual bool IsPoseWatchable() const override;
+	ANIMGRAPH_API virtual bool IsSinkNode() const override;
+	ANIMGRAPH_API virtual void OnProcessDuringCompilation(IAnimBlueprintCompilationContext& InCompilationContext, IAnimBlueprintGeneratedClassCompiledData& OutCompiledData) override;
+	ANIMGRAPH_API virtual void OnCopyTermDefaultsToDefaultObject(IAnimBlueprintCopyTermDefaultsContext& InCompilationContext, IAnimBlueprintNodeCopyTermDefaultsContext& InPerNodeContext, IAnimBlueprintGeneratedClassCompiledData& OutCompiledData) override;
 
 	// Get the link to the documentation
-	virtual FString GetDocumentationLink() const override;
+	ANIMGRAPH_API virtual FString GetDocumentationLink() const override;
 
 	//~ End UAnimGraphNode_Base Interface
 };
diff --git a/Engine/Source/Editor/AnimGraph/Public/AnimationGraphSchema.h b/Engine/Source/Editor/AnimGraph/Public/AnimationGraphSchema.h
--- a/Engine/Source/Editor/AnimGraph/Public/AnimationGraphSchema.h
+++ b/Engine/Source/Editor/AnimGraph/Public/AnimationGraphSchema.h
@@ -48,38 +48,38 @@ class UAnimationGraphSchema : public UEdGraphSchema_K2
 	FName DefaultEvaluationHandlerName;
 
 	//~ Begin UEdGraphSchema Interface.
-	virtual FLinearColor GetPinTypeColor(const FEdGraphPinType& PinType) const override;
-	virtual EGraphType GetGraphType(const UEdGraph* TestEdGraph) const override;
-	virtual void CreateDefaultNodesForGraph(UEdGraph& Graph) const override;
-	virtual void HandleGraphBeingDeleted(UEdGraph& GraphBeingRemoved) const override;
-	virtual bool CreateAutomaticConversionNodeAndConnections(UEdGraphPin* PinA, UEdGraphPin* PinB) const override;
-	virtual void DroppedAssetsOnGraph(const TArray<FAssetData>& Assets, const FVector2D& GraphPosition, UEdGraph* Graph) const override;
-	virtual void DroppedAssetsOnNode(const TArray<FAssetData>& Assets, const FVector2D& GraphPosition, UEdGraphNode* Node) const override;
-	virtual void DroppedAssetsOnPin(const TArray<FAssetData>& Assets, const FVector2D& GraphPosition, UEdGraphPin* Pin) const override;
-	virtual void GetAssetsNodeHoverMessage(const TArray<FAssetData>& Assets, const UEdGraphNode* HoverNode, FString& OutTooltipText, bool& OutOkIcon) const override;
-	virtual void GetAssetsPinHoverMessage(const TArray<FAssetData>& Assets, const UEdGraphPin* HoverPin, FString& OutTooltipText, bool& OutOkIcon) const override;
-	virtual void GetAssetsGraphHoverMessage(const TArray<FAssetData>& Assets, const UEdGraph* HoverGraph, FString& OutTooltipText, bool& OutOkIcon) const override;
-	virtual void GetContextMenuActions(class UToolMenu* Menu, class UGraphNodeContextMenuContext* Context) const override;
-	virtual FText GetPinDisplayName(const UEdGraphPin* Pin) const override;
-	virtual bool CanDuplicateGraph(UEdGraph* InSourceGraph) const override;
-	virtual bool ShouldAlwaysPurgeOnModification() const override { return true; }
-	virtual bool TryCreateConnection(UEdGraphPin* A, UEdGraphPin* B) const override;
-	virtual void GetGraphDisplayInformation(const UEdGraph& Graph, /*out*/ FGraphDisplayInfo& DisplayInfo) const override;
-	virtual bool CanGraphBeDropped(TSharedPtr<FEdGraphSchemaAction> InAction) const override;
-	virtual FReply BeginGraphDragAction(TSharedPtr<FEdGraphSchemaAction> InAction, const FPointerEvent& MouseEvent = FPointerEvent()) const override;
+	ANIMGRAPH_API virtual FLinearColor GetPinTypeColor(const FEdGraphPinType& PinType) const override;
+	ANIMGRAPH_API virtual EGraphType GetGraphType(const UEdGraph* TestEdGraph) const override;
+	ANIMGRAPH_API virtual void CreateDefaultNodesForGraph(UEdGraph& Graph) const override;
+	ANIMGRAPH_API virtual void HandleGraphBeingDeleted(UEdGraph& GraphBeingRemoved) const override;
+	ANIMGRAPH_API virtual bool CreateAutomaticConversionNodeAndConnections(UEdGraphPin* PinA, UEdGraphPin* PinB) const override;
+	ANIMGRAPH_API virtual void DroppedAssetsOnGraph(const TArray<FAssetData>& Assets, const FVector2D& GraphPosition, UEdGraph* Graph) const override;
+	ANIMGRAPH_API virtual void DroppedAssetsOnNode(const TArray<FAssetData>& Assets, const FVector2D& GraphPosition, UEdGraphNode* Node) const override;
+	ANIMGRAPH_API virtual void DroppedAssetsOnPin(const TArray<FAssetData>& Assets, const FVector2D& GraphPosition, UEdGraphPin* Pin) const override;
+	ANIMGRAPH_API virtual void GetAssetsNodeHoverMessage(const TArray<FAssetData>& Assets, const UEdGraphNode* HoverNode, FString& OutTooltipText, bool& OutOkIcon) const override;
+	ANIMGRAPH_API virtual void GetAssetsPinHoverMessage(const TArray<FAssetData>& Assets, const UEdGraphPin* HoverPin, FString& OutTooltipText, bool& OutOkIcon) const override;
+	ANIMGRAPH_API virtual void GetAssetsGraphHoverMessage(const TArray<FAssetData>& Assets, const UEdGraph* HoverGraph, FString& OutTooltipText, bool& OutOkIcon) const override;
+	ANIMGRAPH_API virtual void GetContextMenuActions(class UToolMenu* Menu, class UGraphNodeContextMenuContext* Context) const override;
+	ANIMGRAPH_API virtual FText GetPinDisplayName(const UEdGraphPin* Pin) const override;
+	ANIMGRAPH_API virtual bool CanDuplicateGraph(UEdGraph* InSourceGraph) const override;
+	ANIMGRAPH_API virtual bool ShouldAlwaysPurgeOnModification() const override { return true; }
+	ANIMGRAPH_API virtual bool TryCreateConnection(UEdGraphPin* A, UEdGraphPin* B) const override;
+	ANIMGRAPH_API virtual void GetGraphDisplayInformation(const UEdGraph& Graph, /*out*/ FGraphDisplayInfo& DisplayInfo) const override;
+	ANIMGRAPH_API virtual bool CanGraphBeDropped(TSharedPtr<FEdGraphSchemaAction> InAction) const override;
+	ANIMGRAPH_API virtual FReply BeginGraphDragAction(TSharedPtr<FEdGraphSchemaAction> InAction, const FPointerEvent& MouseEvent = FPointerEvent()) const override;
 	//~ End UEdGraphSchema Interface.
 
 	//~ Begin UEdGraphSchema_K2 Interface
-	virtual const FPinConnectionResponse DetermineConnectionResponseOfCompatibleTypedPins(const UEdGraphPin* PinA, const UEdGraphPin* PinB, const UEdGraphPin* InputPin, const UEdGraphPin* OutputPin) const override;
+	ANIMGRAPH_API virtual const FPinConnectionResponse DetermineConnectionResponseOfCompatibleTypedPins(const UEdGraphPin* PinA, const UEdGraphPin* PinB, const UEdGraphPin* InputPin, const UEdGraphPin* OutputPin) const override;
 	UE_DEPRECATED(5.2, "Use the FSearchForAutocastFunctionResults variant.")
-	virtual bool SearchForAutocastFunction(const FEdGraphPinType& OutputPinType, const FEdGraphPinType& InputPinType, /*out*/ FName& TargetFunction, /*out*/ UClass*& FunctionOwner) const override;
-	virtual TOptional<FSearchForAutocastFunctionResults> SearchForAutocastFunction(const FEdGraphPinType& OutputPinType, const FEdGraphPinType& InputPinType) const;
-	virtual bool ArePinsCompatible(const UEdGraphPin* PinA, const UEdGraphPin* PinB, const UClass* CallingContext = NULL, bool bIgnoreArray = false) const override;
-	virtual bool DoesSupportCollapsedNodes() const override { return false; }
-	virtual bool DoesSupportEventDispatcher() const	override { return false; }
-	virtual bool DoesSupportAnimNotifyActions() const override;
-	virtual void CreateFunctionGraphTerminators(UEdGraph& Graph, UClass* Class) const override;
-	virtual bool CanShowDataTooltipForPin(const UEdGraphPin& Pin) const override;
+	ANIMGRAPH_API virtual bool SearchForAutocastFunction(const FEdGraphPinType& OutputPinType, const FEdGraphPinType& InputPinType, /*out*/ FName& TargetFunction, /*out*/ UClass*& FunctionOwner) const override;
+	ANIMGRAPH_API virtual TOptional<FSearchForAutocastFunctionResults> SearchForAutocastFunction(const FEdGraphPinType& OutputPinType, const FEdGraphPinType& InputPinType) const;
+	ANIMGRAPH_API virtual bool ArePinsCompatible(const UEdGraphPin* PinA, const UEdGraphPin* PinB, const UClass* CallingContext = NULL, bool bIgnoreArray = false) const override;
+	ANIMGRAPH_API virtual bool DoesSupportCollapsedNodes() const override { return false; }
+	ANIMGRAPH_API virtual bool DoesSupportEventDispatcher() const	override { return false; }
+	ANIMGRAPH_API virtual bool DoesSupportAnimNotifyActions() const override;
+	ANIMGRAPH_API virtual void CreateFunctionGraphTerminators(UEdGraph& Graph, UClass* Class) const override;
+	ANIMGRAPH_API virtual bool CanShowDataTooltipForPin(const UEdGraphPin& Pin) const override;
 	//~ End UEdGraphSchema_K2 Interface
 
 	/** Spawn the correct node in the Animation Graph using the given AnimationAsset at the supplied location */
