Blend stack: Fix for CIS error (mark AnimGraphNode_BlendStack_Base as abstract)
#rnx

diff --git a/Engine/Plugins/Experimental/Animation/PoseSearch/Source/Editor/Public/AnimGraphNode_BlendStack.h b/Engine/Plugins/Experimental/Animation/PoseSearch/Source/Editor/Public/AnimGraphNode_BlendStack.h
--- a/Engine/Plugins/Experimental/Animation/PoseSearch/Source/Editor/Public/AnimGraphNode_BlendStack.h
+++ b/Engine/Plugins/Experimental/Animation/PoseSearch/Source/Editor/Public/AnimGraphNode_BlendStack.h
@@ -7,7 +7,7 @@
 #include "AnimGraphNode_BlendStack.generated.h"
 
 
-UCLASS(MinimalAPI, Experimental)
+UCLASS(MinimalAPI, Abstract, Experimental)
 class UAnimGraphNode_BlendStack_Base : public UAnimGraphNode_AssetPlayerBase
 {
 	GENERATED_BODY()
