Fix for blend stack sample blend weights being incorrect.
#rnx

diff --git a/Engine/Plugins/Experimental/Animation/PoseSearch/Source/Runtime/Private/AnimNode_BlendStack.cpp b/Engine/Plugins/Experimental/Animation/PoseSearch/Source/Runtime/Private/AnimNode_BlendStack.cpp
--- a/Engine/Plugins/Experimental/Animation/PoseSearch/Source/Runtime/Private/AnimNode_BlendStack.cpp
+++ b/Engine/Plugins/Experimental/Animation/PoseSearch/Source/Runtime/Private/AnimNode_BlendStack.cpp
@@ -452,7 +452,7 @@ void FAnimNode_BlendStack_Standalone::UpdateAssetPlayer(const FAnimationUpdateCo
 			break;
 		}
 
-		const FAnimationUpdateContext AnimPlayerContext = Context.FractionalWeightAndRootMotion(BlendWeight, BlendWeight);
+		const FAnimationUpdateContext AnimPlayerContext = Context.FractionalWeightAndRootMotion(AnimPlayerBlendWeight, AnimPlayerBlendWeight);
 		UpdateSample(AnimPlayerContext, AnimPlayerIndex);
 		CurrentWeightMultiplier *= (1.f - BlendInPercentage);
 	}
