Blend Stack: Fix for CIS error
#rnx

diff --git a/Engine/Plugins/Experimental/Animation/PoseSearch/Source/Runtime/Public/PoseSearch/AnimNode_BlendStack.h b/Engine/Plugins/Experimental/Animation/PoseSearch/Source/Runtime/Public/PoseSearch/AnimNode_BlendStack.h
--- a/Engine/Plugins/Experimental/Animation/PoseSearch/Source/Runtime/Public/PoseSearch/AnimNode_BlendStack.h
+++ b/Engine/Plugins/Experimental/Animation/PoseSearch/Source/Runtime/Public/PoseSearch/AnimNode_BlendStack.h
@@ -71,9 +71,9 @@ struct FBlendStack_SampleGraphPoseLink
 	GENERATED_BODY()
 
 	UPROPERTY()
-	int32 InputPoseNodeIndex;
+	int32 InputPoseNodeIndex = INDEX_NONE;
 	UPROPERTY()
-	int32 RootNodeIndex;
+	int32 RootNodeIndex = INDEX_NONE;
 
 	FPoseLink InputPose;
 	FPoseLink Root;
