motion matching - pose feature values pruning

diff --git a/Engine/Plugins/Experimental/Animation/PoseSearch/Source/Editor/Private/PoseSearchDatabaseEditorReflection.cpp b/Engine/Plugins/Experimental/Animation/PoseSearch/Source/Editor/Private/PoseSearchDatabaseEditorReflection.cpp
--- a/Engine/Plugins/Experimental/Animation/PoseSearch/Source/Editor/Private/PoseSearchDatabaseEditorReflection.cpp
+++ b/Engine/Plugins/Experimental/Animation/PoseSearch/Source/Editor/Private/PoseSearchDatabaseEditorReflection.cpp
@@ -130,9 +130,19 @@ void UPoseSearchDatabaseStatistics::Initialize(const UPoseSearchDatabase* PoseSe
 		}
 		SearchableFrames = NumOfSearchablePoses;
 		SearchableTime = FText::Format(TimeFormat, static_cast<double>(NumOfSearchablePoses) / SampleRate);
-	
+
 		ConfigCardinality = PoseSearchDatabase->Schema->SchemaCardinality;
 
+		if (ConfigCardinality > 0)
+		{
+			const int32 TotalAnimationFeatureVectors = SearchIndex.GetValues().Num() / ConfigCardinality;
+			PrunedFrames = TotalAnimationPosesInFrames - TotalAnimationFeatureVectors;
+		}
+		else
+		{
+			PrunedFrames = 0;
+		}
+
 		// Kinematic Information
 	
 		// using FText instead of meta = (ForceUnits = "cm/s") to keep properties consistent
@@ -148,7 +158,7 @@ void UPoseSearchDatabaseStatistics::Initialize(const UPoseSearchDatabase* PoseSe
 		// Memory Information
 			
 		{
-			const uint32 ValuesBytesSize = SearchIndex.Values.GetAllocatedSize();
+			const uint32 ValuesBytesSize = SearchIndex.GetValues().GetAllocatedSize();
 			const uint32 PCAValuesBytesSize = SearchIndex.PCAValues.GetAllocatedSize();
 			const uint32 KDTreeBytesSize = SearchIndex.KDTree.GetAllocatedSize();
 			const uint32 PoseMetadataBytesSize = SearchIndex.PoseMetadata.GetAllocatedSize();
diff --git a/Engine/Plugins/Experimental/Animation/PoseSearch/Source/Editor/Public/PoseSearchDatabaseEditorReflection.h b/Engine/Plugins/Experimental/Animation/PoseSearch/Source/Editor/Public/PoseSearchDatabaseEditorReflection.h
--- a/Engine/Plugins/Experimental/Animation/PoseSearch/Source/Editor/Public/PoseSearchDatabaseEditorReflection.h
+++ b/Engine/Plugins/Experimental/Animation/PoseSearch/Source/Editor/Public/PoseSearchDatabaseEditorReflection.h
@@ -198,6 +198,10 @@ public:
 	UPROPERTY(VisibleAnywhere, Category = "General Information")
 	FText SearchableTime;
 
+	// Amount of animation frames (database poses) pose feature data that has been pruned out because close to each other more than PosePruningSimilartyThreshod
+	UPROPERTY(VisibleAnywhere, Category = "General Information")
+	uint32 PrunedFrames;
+
 	// Cardinality for the database config (how many floats per pose to store the pose features data)
 	UPROPERTY(VisibleAnywhere, Category = "General Information")
 	uint32 ConfigCardinality;
diff --git a/Engine/Plugins/Experimental/Animation/PoseSearch/Source/Runtime/Private/PoseSearchAssetIndexer.cpp b/Engine/Plugins/Experimental/Animation/PoseSearch/Source/Runtime/Private/PoseSearchAssetIndexer.cpp
--- a/Engine/Plugins/Experimental/Animation/PoseSearch/Source/Runtime/Private/PoseSearchAssetIndexer.cpp
+++ b/Engine/Plugins/Experimental/Animation/PoseSearch/Source/Runtime/Private/PoseSearchAssetIndexer.cpp
@@ -86,10 +86,13 @@ FAssetIndexer::FAssetIndexer(const FBoneContainer& InBoneContainer, const FSearc
 {
 }
 
-void FAssetIndexer::AssignWorkingData(TArrayView<float> InOutFeatureVectorTable, TArrayView<FPoseMetadata> InOutPoseMetadata)
+void FAssetIndexer::AssignWorkingData(int32 InStartPoseIdx, TArrayView<float> InOutFeatureVectorTable, TArrayView<FPoseMetadata> InOutPoseMetadata)
 {
-	FeatureVectorTable = InOutFeatureVectorTable;
-	PoseMetadata = InOutPoseMetadata;
+	const int32 NumIndexedPoses = GetNumIndexedPoses();
+
+	StartPoseIdx = InStartPoseIdx;
+	FeatureVectorTable = MakeArrayView(InOutFeatureVectorTable.GetData() + Schema.SchemaCardinality * StartPoseIdx, Schema.SchemaCardinality * NumIndexedPoses);
+	PoseMetadata = MakeArrayView(InOutPoseMetadata.GetData() + StartPoseIdx, NumIndexedPoses);
 }
 
 void FAssetIndexer::Process(int32 AssetIdx)
@@ -126,7 +129,9 @@ void FAssetIndexer::Process(int32 AssetIdx)
 			CostAddend += Schema.LoopingCostBias;
 		}
 
-		PoseMetadata[GetVectorIdx(SampleIdx)].Init(AssetIdx, bBlockTransition, CostAddend);
+		const int32 ValueOffset = (StartPoseIdx + GetVectorIdx(SampleIdx)) * Schema.SchemaCardinality;
+		check(ValueOffset >= 0 && AssetIdx >= 0);
+		PoseMetadata[GetVectorIdx(SampleIdx)] = FPoseMetadata(ValueOffset, AssetIdx, bBlockTransition, CostAddend);
 	}
 
 	// Generate pose features data
diff --git a/Engine/Plugins/Experimental/Animation/PoseSearch/Source/Runtime/Private/PoseSearchContext.cpp b/Engine/Plugins/Experimental/Animation/PoseSearch/Source/Runtime/Private/PoseSearchContext.cpp
--- a/Engine/Plugins/Experimental/Animation/PoseSearch/Source/Runtime/Private/PoseSearchContext.cpp
+++ b/Engine/Plugins/Experimental/Animation/PoseSearch/Source/Runtime/Private/PoseSearchContext.cpp
@@ -258,7 +258,7 @@ FSearchContext::FSearchContext(const FPoseSearchQueryTrajectory* InTrajectory, c
 	if (CurrentResult.IsValid())
 	{
 		const FSearchIndex& SearchIndex = CurrentResult.Database->GetSearchIndex();
-		if (SearchIndex.Values.IsEmpty())
+		if (SearchIndex.IsValuesEmpty())
 		{
 			const int32 NumDimensions = CurrentResult.Database->Schema->SchemaCardinality;
 			CurrentResultPoseVectorData.AddUninitialized(NumDimensions);
diff --git a/Engine/Plugins/Experimental/Animation/PoseSearch/Source/Runtime/Private/PoseSearchDatabase.cpp b/Engine/Plugins/Experimental/Animation/PoseSearch/Source/Runtime/Private/PoseSearchDatabase.cpp
--- a/Engine/Plugins/Experimental/Animation/PoseSearch/Source/Runtime/Private/PoseSearchDatabase.cpp
+++ b/Engine/Plugins/Experimental/Animation/PoseSearch/Source/Runtime/Private/PoseSearchDatabase.cpp
@@ -765,7 +765,7 @@ FPoseSearchCost UPoseSearchDatabase::SearchContinuingPose(UE::PoseSearch::FSearc
 		// FMemory_Alloca is forced 16 bytes aligned
 		TArrayView<float> ReconstructedPoseValuesBuffer((float*)FMemory_Alloca(NumDimensions * sizeof(float)), NumDimensions);
 		check(IsAligned(ReconstructedPoseValuesBuffer.GetData(), alignof(VectorRegister4Float)));
-		const TConstArrayView<float> PoseValues = SearchIndex.Values.IsEmpty() ? SearchIndex.GetReconstructedPoseValues(PoseIdx, ReconstructedPoseValuesBuffer) : SearchIndex.GetPoseValues(PoseIdx);
+		const TConstArrayView<float> PoseValues = SearchIndex.IsValuesEmpty() ? SearchIndex.GetReconstructedPoseValues(PoseIdx, ReconstructedPoseValuesBuffer) : SearchIndex.GetPoseValues(PoseIdx);
 
 		const int32 ContinuingPoseIdx = SearchContext.GetCurrentResult().PoseIdx;
 		// is the data padded at 16 bytes (and 16 bytes aligned by construction)?
@@ -858,7 +858,7 @@ UE::PoseSearch::FSearchResult UPoseSearchDatabase::SearchPCAKDTree(UE::PoseSearc
 		const FSearchFilters SearchFilters(Schema, TConstArrayView<size_t>(), SearchIndex.bAnyBlockTransition);
 		
 		// do we need to reconstruct pose values?
-		if (SearchIndex.Values.IsEmpty())
+		if (SearchIndex.IsValuesEmpty())
 		{
 			// FMemory_Alloca is forced 16 bytes aligned
 			TArrayView<float> ReconstructedPoseValuesBuffer((float*)FMemory_Alloca(NumDimensions * sizeof(float)), NumDimensions);
@@ -938,7 +938,7 @@ UE::PoseSearch::FSearchResult UPoseSearchDatabase::SearchBruteForce(UE::PoseSear
 		const bool bUpdateBestCandidates = PoseSearchMode == EPoseSearchMode::BruteForce;
 
 		// do we need to reconstruct pose values?
-		if (SearchIndex.Values.IsEmpty())
+		if (SearchIndex.IsValuesEmpty())
 		{
 			// FMemory_Alloca is forced 16 bytes aligned
 			TArrayView<float> ReconstructedPoseValuesBuffer((float*)FMemory_Alloca(NumDimensions * sizeof(float)), NumDimensions);
diff --git a/Engine/Plugins/Experimental/Animation/PoseSearch/Source/Runtime/Private/PoseSearchDatabaseIndexingContext.cpp b/Engine/Plugins/Experimental/Animation/PoseSearch/Source/Runtime/Private/PoseSearchDatabaseIndexingContext.cpp
--- a/Engine/Plugins/Experimental/Animation/PoseSearch/Source/Runtime/Private/PoseSearchDatabaseIndexingContext.cpp
+++ b/Engine/Plugins/Experimental/Animation/PoseSearch/Source/Runtime/Private/PoseSearchDatabaseIndexingContext.cpp
@@ -123,21 +123,14 @@ bool FDatabaseIndexingContext::IndexDatabase(FSearchIndexBase& SearchIndexBase,
 	}
 
 	// allocating Values and PoseMetadata
-	SearchIndexBase.Values.Reset();
-	SearchIndexBase.PoseMetadata.Reset();
-
-	SearchIndexBase.Values.SetNumZeroed(Schema->SchemaCardinality * TotalPoses);
-	SearchIndexBase.PoseMetadata.SetNumZeroed(TotalPoses);
+	SearchIndexBase.AllocateData(Schema->SchemaCardinality, TotalPoses);
 
 	// assigning local data to each Indexer
 	TotalPoses = 0;
 	for (int32 AssetIdx = 0; AssetIdx != SearchIndexBase.Assets.Num(); ++AssetIdx)
 	{
-		const int32 NumIndexedPoses = Indexers[AssetIdx].GetNumIndexedPoses();
-		Indexers[AssetIdx].AssignWorkingData(
-			MakeArrayView(SearchIndexBase.Values.GetData() + Schema->SchemaCardinality * TotalPoses, Schema->SchemaCardinality * NumIndexedPoses),
-			MakeArrayView(SearchIndexBase.PoseMetadata.GetData() + TotalPoses, NumIndexedPoses));
-		TotalPoses += NumIndexedPoses;
+		Indexers[AssetIdx].AssignWorkingData(TotalPoses, SearchIndexBase.EditValues(), SearchIndexBase.PoseMetadata);
+		TotalPoses += Indexers[AssetIdx].GetNumIndexedPoses();
 	}
 
 	if (Owner.IsCanceled())
diff --git a/Engine/Plugins/Experimental/Animation/PoseSearch/Source/Runtime/Private/PoseSearchDerivedData.cpp b/Engine/Plugins/Experimental/Animation/PoseSearch/Source/Runtime/Private/PoseSearchDerivedData.cpp
--- a/Engine/Plugins/Experimental/Animation/PoseSearch/Source/Runtime/Private/PoseSearchDerivedData.cpp
+++ b/Engine/Plugins/Experimental/Animation/PoseSearch/Source/Runtime/Private/PoseSearchDerivedData.cpp
@@ -118,32 +118,37 @@ private:
 		const int32 Cardinality = Entries[0].Channel->GetChannelCardinality();
 		check(Cardinality > 0);
 
-		int32 TotalNumPoses = 0;
+		int32 TotalNumValuesVectors = 0;
 		for (int32 EntryIdx = 0; EntryIdx < EntriesNum; ++EntryIdx)
 		{
-			TotalNumPoses += SearchIndexBases[Entries[EntryIdx].SchemaIndex].GetNumPoses();
+			const FEntry& Entry = Entries[EntryIdx];
+			check(Cardinality == Entry.Channel->GetChannelCardinality());
+
+			const int32 DataSetIdx = Entry.SchemaIndex;
+			const UPoseSearchSchema* Schema = Schemas[DataSetIdx];
+			const FSearchIndexBase& SearchIndexBase = SearchIndexBases[DataSetIdx];
+
+			TotalNumValuesVectors += SearchIndexBase.GetNumValuesVectors(Schema->SchemaCardinality);
 		}
 
-		int32 AccumulatedNumPoses = 0;
-		RowMajorMatrix CenteredSubPoseMatrix(TotalNumPoses, Cardinality);
+		int32 AccumulatedNumValuesVectors = 0;
+		RowMajorMatrix CenteredSubPoseMatrix(TotalNumValuesVectors, Cardinality);
 		for (int32 EntryIdx = 0; EntryIdx < EntriesNum; ++EntryIdx)
 		{
 			const FEntry& Entry = Entries[EntryIdx];
-			check(Cardinality == Entry.Channel->GetChannelCardinality());
-
 			const int32 DataSetIdx = Entry.SchemaIndex;
 
 			const UPoseSearchSchema* Schema = Schemas[DataSetIdx];
-			const FSearchIndexBase& SearchIndex = SearchIndexBases[DataSetIdx];
+			const FSearchIndexBase& SearchIndexBase = SearchIndexBases[DataSetIdx];
 
-			const int32 NumPoses = SearchIndex.GetNumPoses();
+			const int32 NumValuesVectors = SearchIndexBase.GetNumValuesVectors(Schema->SchemaCardinality);
 
-			// Map input buffer with NumPoses as rows and NumDimensions	as cols
-			RowMajorMatrixMapConst PoseMatrixSourceMap(SearchIndex.Values.GetData(), NumPoses, Schema->SchemaCardinality);
+			// Map input buffer with NumValuesVectors as rows and NumDimensions	as cols
+			RowMajorMatrixMapConst PoseMatrixSourceMap(SearchIndexBase.GetValues().GetData(), NumValuesVectors, Schema->SchemaCardinality);
 
 			// Given the sub matrix for the features, find the average distance to the feature's centroid.
-			CenteredSubPoseMatrix.block(AccumulatedNumPoses, 0, NumPoses, Cardinality) = PoseMatrixSourceMap.block(0, Entry.Channel->GetChannelDataOffset(), NumPoses, Cardinality);
-			AccumulatedNumPoses += NumPoses;
+			CenteredSubPoseMatrix.block(AccumulatedNumValuesVectors, 0, NumValuesVectors, Cardinality) = PoseMatrixSourceMap.block(0, Entry.Channel->GetChannelDataOffset(), NumValuesVectors, Cardinality);
+			AccumulatedNumValuesVectors += NumValuesVectors;
 		}
 
 		RowMajorVector SampleMean = CenteredSubPoseMatrix.colwise().mean();
@@ -431,8 +436,16 @@ static void PreprocessSearchIndexPCAData(FSearchIndex& SearchIndex, int32 NumDim
 		SearchIndex.Mean.AddZeroed(NumDimensions);
 		SearchIndex.PCAProjectionMatrix.AddZeroed(NumDimensions * NumberOfPrincipalComponents);
 
+		// recreating the full pose values data to have a 1:1 mapping between PCAValues/NumDimensions and PoseIdx
+		TArray<float> AllValuesWithDuplicateData;
+		AllValuesWithDuplicateData.SetNum(NumPoses * NumDimensions);
+		for (int32 PoseIdx = 0; PoseIdx < NumPoses; ++PoseIdx)
+		{
+			FMemory::Memcpy(AllValuesWithDuplicateData.GetData() + PoseIdx * NumDimensions, SearchIndex.GetPoseValuesBase(PoseIdx, NumDimensions).GetData(), NumDimensions * sizeof(float));
+		}
+
 		const RowMajorVectorMapConst MapWeightsSqrt(SearchIndex.WeightsSqrt.GetData(), 1, NumDimensions);
-		const RowMajorMatrixMapConst MapValues(SearchIndex.Values.GetData(), NumPoses, NumDimensions);
+		const RowMajorMatrixMapConst MapValues(AllValuesWithDuplicateData.GetData(), NumPoses, NumDimensions);
 		const RowMajorMatrix WeightedValues = MapValues.array().rowwise() * MapWeightsSqrt.array();
 		RowMajorMatrixMap MapPCAValues(SearchIndex.PCAValues.GetData(), NumPoses, NumberOfPrincipalComponents);
 
@@ -573,7 +586,7 @@ static void PreprocessSearchIndexKDTree(FSearchIndex& SearchIndex, int32 NumDime
 				ResultDistanceSqr.SetNum(KDTreeQueryNumNeighbors + 1);
 				FKDTree::KNNResultSet ResultSet(KDTreeQueryNumNeighbors, ResultIndexes, ResultDistanceSqr);
 
-				const RowMajorVectorMapConst MapValues(&SearchIndex.Values[PointIndex * NumDimensions], 1, NumDimensions);
+				const RowMajorVectorMapConst MapValues(SearchIndex.GetPoseValuesBase(PointIndex, NumDimensions).GetData(), 1, NumDimensions);
 				const RowMajorVectorMapConst MapWeightsSqrt(SearchIndex.WeightsSqrt.GetData(), 1, NumDimensions);
 				const RowMajorVectorMapConst Mean(SearchIndex.Mean.GetData(), 1, NumDimensions);
 				const ColMajorMatrixMapConst PCAProjectionMatrix(SearchIndex.PCAProjectionMatrix.GetData(), NumDimensions, NumberOfPrincipalComponents);
@@ -976,6 +989,8 @@ void FPoseSearchDatabaseAsyncCacheTask::OnGetComplete(UE::DerivedData::FCacheGet
 
 				static_cast<FSearchIndexBase&>(SearchIndex) = SearchIndexBases[0];
 				
+				SearchIndex.PruneDuplicateValues(Database->PosePruningSimilarityThreshold, Database->Schema->SchemaCardinality);
+
 				TArray<float> Deviation = FMeanDeviationCalculator::Calculate(SearchIndexBases, Schemas);
 
 				// Building FSearchIndex
@@ -1006,7 +1021,7 @@ void FPoseSearchDatabaseAsyncCacheTask::OnGetComplete(UE::DerivedData::FCacheGet
 				// removing SearchIndex.Values and relying on FSearchIndex::GetReconstructedPoseValues to reconstruct the Values data from the PCAValues
 				if (Database->PoseSearchMode == EPoseSearchMode::PCAKDTree && Database->KDTreeQueryNumNeighbors <= 1)
 				{
-					SearchIndex.Values.Reset();
+					SearchIndex.ResetValues();
 				}
 
 				UE_LOG(LogPoseSearch, Log, TEXT("%s - %s BuildIndex Succeeded"), *LexToString(FullIndexKey.Hash), *Database->GetName());
diff --git a/Engine/Plugins/Experimental/Animation/PoseSearch/Source/Runtime/Private/PoseSearchDerivedDataKey.cpp b/Engine/Plugins/Experimental/Animation/PoseSearch/Source/Runtime/Private/PoseSearchDerivedDataKey.cpp
--- a/Engine/Plugins/Experimental/Animation/PoseSearch/Source/Runtime/Private/PoseSearchDerivedDataKey.cpp
+++ b/Engine/Plugins/Experimental/Animation/PoseSearch/Source/Runtime/Private/PoseSearchDerivedDataKey.cpp
@@ -41,7 +41,7 @@ FKeyBuilder::FKeyBuilder(const UObject* Object, bool bUseDataVer)
 	if (bUseDataVer)
 	{
 		// used to invalidate the key without having to change POSESEARCHDB_DERIVEDDATA_VER all the times
-		int32 POSESEARCHDB_DERIVEDDATA_VER_SMALL = 157;
+		int32 POSESEARCHDB_DERIVEDDATA_VER_SMALL = 167;
 		FGuid VersionGuid = FDevSystemGuids::GetSystemGuid(FDevSystemGuids::Get().POSESEARCHDB_DERIVEDDATA_VER);
 
 		*this << VersionGuid;
diff --git a/Engine/Plugins/Experimental/Animation/PoseSearch/Source/Runtime/Private/PoseSearchIndex.cpp b/Engine/Plugins/Experimental/Animation/PoseSearch/Source/Runtime/Private/PoseSearchIndex.cpp
--- a/Engine/Plugins/Experimental/Animation/PoseSearch/Source/Runtime/Private/PoseSearchIndex.cpp
+++ b/Engine/Plugins/Experimental/Animation/PoseSearch/Source/Runtime/Private/PoseSearchIndex.cpp
@@ -89,8 +89,18 @@ float CompareFeatureVectors(TConstArrayView<float> A, TConstArrayView<float> B)
 // FPoseMetadata
 FArchive& operator<<(FArchive& Ar, FPoseMetadata& Metadata)
 {
-	Ar << Metadata.Data;
-	Ar << Metadata.CostAddend;
+	// storing more data than necessary for now to avoid to deal with endiannes
+	uint32 ValueOffset = Metadata.GetValueOffset();
+	uint32 AssetIndex = Metadata.GetAssetIndex();
+	bool bInBlockTransition = Metadata.IsBlockTransition();
+	FFloat16 CostAddend = Metadata.CostAddend;
+	
+	Ar << ValueOffset;
+	Ar << AssetIndex;
+	Ar << bInBlockTransition;
+	Ar << CostAddend;
+
+	Metadata = FPoseMetadata(ValueOffset, AssetIndex, bInBlockTransition, CostAddend);
 	return Ar;
 }
 
@@ -150,6 +160,173 @@ void FSearchIndexBase::Reset()
 	*this = FSearchIndexBase();
 }
 
+void FSearchIndexBase::PruneDuplicateValues(float SimilarityThreshold, int32 DataCardinality)
+{
+	constexpr bool bParalleEvaluation = true;
+
+	const int32 NumPoses = GetNumPoses();
+	if (SimilarityThreshold > 0.f && NumPoses >= 2)
+	{
+		struct FPosePair
+		{
+			int32 PoseIdxA = 0;
+			int32 PoseIdxB = 0;
+		};
+		struct FPosePairSimilarity : public FPosePair
+		{
+			float Similarity = 0.f;
+		};
+		TArray<FPosePairSimilarity> PosePairSimilarities;
+		PosePairSimilarities.Reserve(1024 * 64);
+
+		// calculating a sparse proximity matrix with Similarity up to SimilarityThreshold (to perform something similar to hierarchical clustering)
+		if (bParalleEvaluation)
+		{
+			// @todo: we should use a smarter approach, knowing we need NumPoses * (NumPoses - 1) / 2 comparisons, and reconstruct the pose indexes from the comparison index...
+			
+			// doing it in batches of items to avoid using too much memory
+			constexpr int32 BatchSize = 1024 * 1024;
+
+			TArray<FPosePair> PosePairs;
+			PosePairs.Reserve(BatchSize);
+			FPosePair PosePair;
+			FCriticalSection Mutex;
+
+			auto EvaluatePosePairs = [this, DataCardinality, SimilarityThreshold, &PosePairs, &PosePairSimilarities, &Mutex]()
+			{
+				ParallelFor(PosePairs.Num(), [this, DataCardinality, SimilarityThreshold, &PosePairs, &PosePairSimilarities, &Mutex](int32 ComparisonIndex)
+				{
+					const int PoseIdxA = PosePairs[ComparisonIndex].PoseIdxA;
+					const int PoseIdxB = PosePairs[ComparisonIndex].PoseIdxB;
+					TConstArrayView<float> ValuesA = GetPoseValuesBase(PoseIdxA, DataCardinality);
+					TConstArrayView<float> ValuesB = GetPoseValuesBase(PoseIdxB, DataCardinality);
+					const float Similarity = CompareFeatureVectors(ValuesA, ValuesB);
+					if (Similarity < SimilarityThreshold)
+					{
+						// since this condition doesn't happen often, we're not gonna spend too much time in mutex contention
+						FScopeLock Lock(&Mutex);
+						FPosePairSimilarity PosePairSimilarity;
+						PosePairSimilarity.PoseIdxA = PoseIdxA;
+						PosePairSimilarity.PoseIdxB = PoseIdxB;
+						PosePairSimilarity.Similarity = Similarity;
+						PosePairSimilarities.Emplace(PosePairSimilarity);
+					}
+				});
+			};
+
+			for (PosePair.PoseIdxA = 0; PosePair.PoseIdxA < NumPoses; ++PosePair.PoseIdxA)
+			{
+				for (PosePair.PoseIdxB = PosePair.PoseIdxA + 1; PosePair.PoseIdxB < NumPoses; ++PosePair.PoseIdxB)
+				{
+					PosePairs.Add(PosePair);
+					if (PosePairs.Num() == BatchSize)
+					{
+						EvaluatePosePairs();
+
+						PosePairs.Reset();
+					}
+				}
+			}
+
+			EvaluatePosePairs();
+		}
+		else
+		{
+			for (int32 PoseIdxA = 0; PoseIdxA < NumPoses; ++PoseIdxA)
+			{
+				TConstArrayView<float> ValuesA = GetPoseValuesBase(PoseIdxA, DataCardinality);
+				for (int32 PoseIdxB = PoseIdxA + 1; PoseIdxB < NumPoses; ++PoseIdxB)
+				{
+					TConstArrayView<float> ValuesB = GetPoseValuesBase(PoseIdxB, DataCardinality);
+					const float Similarity = CompareFeatureVectors(ValuesA, ValuesB);
+					if (Similarity < SimilarityThreshold)
+					{
+						FPosePairSimilarity PosePairSimilarity;
+						PosePairSimilarity.PoseIdxA = PoseIdxA;
+						PosePairSimilarity.PoseIdxB = PoseIdxB;
+						PosePairSimilarity.Similarity = Similarity;
+						PosePairSimilarities.Emplace(PosePairSimilarity);
+					}
+				}
+			}
+		}
+
+		if (!PosePairSimilarities.IsEmpty())
+		{
+			PosePairSimilarities.Sort([](const FPosePairSimilarity& A, const FPosePairSimilarity& B)
+			{
+				return A.Similarity < B.Similarity;
+			});
+
+			// mapping between the one value offset and all the poses sharing it
+			TMap<uint32, TArray<int32>> ValueOffsetToPoses;
+			for (int32 PoseIdx = 0; PoseIdx < NumPoses; ++PoseIdx)
+			{
+				const uint32 ValueOffset = PoseMetadata[PoseIdx].GetValueOffset();
+				// FindOrAdd to support the eventuality of having multiple poses already sharing the same value offset
+				ValueOffsetToPoses.FindOrAdd(ValueOffset).Add(PoseIdx);
+			}
+
+			// at this point ValueOffsetToPoses is fully populated with all the possible value offset, and since we're not adding, but eventually removing entries we can just use the [] operator
+			int32 LastValidPoseIdx = NumPoses - 1;
+			for (int32 PosePairSimilarityIdx = 0; PosePairSimilarityIdx < PosePairSimilarities.Num(); ++PosePairSimilarityIdx)
+			{
+				const FPosePairSimilarity& PosePairSimilarity = PosePairSimilarities[PosePairSimilarityIdx];
+				const uint32 ValueOffsetA = PoseMetadata[PosePairSimilarity.PoseIdxA].GetValueOffset();
+				const uint32 ValueOffsetB = PoseMetadata[PosePairSimilarity.PoseIdxB].GetValueOffset();
+				
+				// if the two poses don't point already to the same value offset, we can remove one of them
+				if (ValueOffsetA != ValueOffsetB)
+				{
+					// transferring all the poses associated to ValueOffsetB to ValueOffsetA
+					TArray<int32>& PosesAtValueOffsetA = ValueOffsetToPoses[ValueOffsetA];
+					TArray<int32>& PosesAtValueOffsetB = ValueOffsetToPoses[ValueOffsetB];
+
+					for (int32 PoseAtValueOffsetB : PosesAtValueOffsetB)
+					{
+						PoseMetadata[PoseAtValueOffsetB].SetValueOffset(ValueOffsetA);
+						PosesAtValueOffsetA.Add(PoseAtValueOffsetB);
+					}
+
+					// moving the LastValidPoseIdx values into the location ValueOffsetB, that we just free up
+					const uint32 ValueOffsetLast = LastValidPoseIdx * DataCardinality;
+					if (ValueOffsetB != ValueOffsetLast)
+					{
+						FMemory::Memcpy(&Values[ValueOffsetB], &Values[ValueOffsetLast], DataCardinality * sizeof(float));
+						TArray<int32>& PosesAtValueOffsetLast = ValueOffsetToPoses[ValueOffsetLast];
+						
+						for (int32 PoseAtValueOffsetLast : PosesAtValueOffsetLast)
+						{
+							PoseMetadata[PoseAtValueOffsetLast].SetValueOffset(ValueOffsetB);
+						}
+
+						PosesAtValueOffsetB = PosesAtValueOffsetLast;
+						PosesAtValueOffsetLast.Reset();
+					}
+					else
+					{
+						PosesAtValueOffsetB.Reset();
+					}
+
+					--LastValidPoseIdx;
+				}
+			}
+
+			// resizing the Values array  
+			Values.SetNum((LastValidPoseIdx + 1) * DataCardinality);
+		}
+	}
+}
+
+void FSearchIndexBase::AllocateData(int32 DataCardinality, int32 NumPoses)
+{
+	Values.Reset();
+	PoseMetadata.Reset();
+
+	Values.SetNumZeroed(DataCardinality * NumPoses);
+	PoseMetadata.SetNumZeroed(NumPoses);
+}
+
 FArchive& operator<<(FArchive& Ar, FSearchIndexBase& Index)
 {
 	Ar << Index.Values;
@@ -194,10 +371,7 @@ void FSearchIndex::Reset()
 
 TConstArrayView<float> FSearchIndex::GetPoseValues(int32 PoseIdx) const
 {
-	const int32 NumDimensions = WeightsSqrt.Num();
-	check(!Values.IsEmpty() && PoseIdx >= 0 && PoseIdx < GetNumPoses() && NumDimensions > 0);
-	const int32 ValueOffset = PoseIdx * NumDimensions;
-	return MakeArrayView(&Values[ValueOffset], NumDimensions);
+	return GetPoseValuesBase(PoseIdx, WeightsSqrt.Num());
 }
 
 TConstArrayView<float> FSearchIndex::GetReconstructedPoseValues(int32 PoseIdx, TArrayView<float> BufferUsedForReconstruction) const
@@ -258,7 +432,7 @@ TArray<float> FSearchIndex::GetPoseValuesSafe(int32 PoseIdx) const
 	TArray<float> PoseValues;
 	if (PoseIdx >= 0 && PoseIdx < GetNumPoses())
 	{
-		if (Values.IsEmpty())
+		if (IsValuesEmpty())
 		{
 			const int32 NumDimensions = WeightsSqrt.Num();
 			PoseValues.SetNumUninitialized(NumDimensions);
diff --git a/Engine/Plugins/Experimental/Animation/PoseSearch/Source/Runtime/Public/PoseSearch/PoseSearchAssetIndexer.h b/Engine/Plugins/Experimental/Animation/PoseSearch/Source/Runtime/Public/PoseSearch/PoseSearchAssetIndexer.h
--- a/Engine/Plugins/Experimental/Animation/PoseSearch/Source/Runtime/Public/PoseSearch/PoseSearchAssetIndexer.h
+++ b/Engine/Plugins/Experimental/Animation/PoseSearch/Source/Runtime/Public/PoseSearch/PoseSearchAssetIndexer.h
@@ -32,7 +32,7 @@ public:
 
 	FAssetIndexer(const FBoneContainer& InBoneContainer, const FSearchIndexAsset& InSearchIndexAsset, 
 		const FAssetSamplingContext& InSamplingContext, const UPoseSearchSchema& InSchema, const FAnimationAssetSampler& InAssetSampler);
-	void AssignWorkingData(TArrayView<float> InOutFeatureVectorTable, TArrayView<FPoseMetadata> InOutPoseMetadata);
+	void AssignWorkingData(int32 InStartPoseIdx, TArrayView<float> InOutFeatureVectorTable, TArrayView<FPoseMetadata> InOutPoseMetadata);
 	void Process(int32 AssetIdx);
 	const FStats& GetStats() const { return Stats; }
 
@@ -64,8 +64,8 @@ private:
 
 	struct CachedEntry
 	{
-		float SampleTime;
-		bool bClamped;
+		float SampleTime = 0.f;
+		bool bClamped = false;
 
 		FTransform RootTransform;
 		FCSPose<FCompactPose> ComponentSpacePose;
@@ -83,6 +83,7 @@ private:
 	const FAssetSamplingContext& SamplingContext;
 	const UPoseSearchSchema& Schema;
 	const FAnimationAssetSampler& AssetSampler;
+	int32 StartPoseIdx = 0;
 	
 	TArrayView<float> FeatureVectorTable;
 	TArrayView<FPoseMetadata> PoseMetadata;
diff --git a/Engine/Plugins/Experimental/Animation/PoseSearch/Source/Runtime/Public/PoseSearch/PoseSearchDatabase.h b/Engine/Plugins/Experimental/Animation/PoseSearch/Source/Runtime/Public/PoseSearch/PoseSearchDatabase.h
--- a/Engine/Plugins/Experimental/Animation/PoseSearch/Source/Runtime/Public/PoseSearch/PoseSearchDatabase.h
+++ b/Engine/Plugins/Experimental/Animation/PoseSearch/Source/Runtime/Public/PoseSearch/PoseSearchDatabase.h
@@ -302,6 +302,11 @@ public:
 	UPROPERTY(EditAnywhere, Category = "Performance")
 	bool bSkipSearchIfPossible = true;
 
+	// if two poses values (multi dimensional point with the schema cardinality) are closer than PosePruningSimilarityThreshold,
+	// only one will be saved into the database (to save memory) and accessed by the two different pose indexes
+	UPROPERTY(EditAnywhere, Category = "Performance")
+	float PosePruningSimilarityThreshold = 0.f;
+
 #if WITH_EDITORONLY_DATA
 	// This optional asset defines a list of databases you want to normalize together. Without it, it would be difficult to compare costs from separately normalized databases containing different types of animation,
 	// like only idles versus only runs animations, given that the range of movement would be dramatically different.
diff --git a/Engine/Plugins/Experimental/Animation/PoseSearch/Source/Runtime/Public/PoseSearch/PoseSearchIndex.h b/Engine/Plugins/Experimental/Animation/PoseSearch/Source/Runtime/Public/PoseSearch/PoseSearchIndex.h
--- a/Engine/Plugins/Experimental/Animation/PoseSearch/Source/Runtime/Public/PoseSearch/PoseSearchIndex.h
+++ b/Engine/Plugins/Experimental/Animation/PoseSearch/Source/Runtime/Public/PoseSearch/PoseSearchIndex.h
@@ -18,27 +18,35 @@ POSESEARCH_API float CompareFeatureVectors(TConstArrayView<float> A, TConstArray
 struct FPoseMetadata
 {
 private:
-	// bits 0-30 represent the AssetIndex, bit 31 represents bBlockTransition
-	uint32 Data = 0;
-	float CostAddend = 0.f;
+	enum { ValueOffsetNumBits = 27 };
+	enum { AssetIndexNumBits = 20 };
+	enum { BlockTransitionNumBits = 1 };
+
+	uint32 ValueOffset : ValueOffsetNumBits = 0;
+	uint32 AssetIndex : AssetIndexNumBits = 0;
+	bool bBlockTransition : BlockTransitionNumBits = 0;
+	FFloat16 CostAddend = 0.f;
 
 public:
-	void Init(uint32 AssetIndex, bool bBlockTransition, float InCostAddend)
+	FPoseMetadata(uint32 InValueOffset = 0, uint32 InAssetIndex = 0, bool bInBlockTransition = false, float InCostAddend = 0.f)
+	: ValueOffset(InValueOffset)
+	, AssetIndex(InAssetIndex)
+	, bBlockTransition(bInBlockTransition)
+	, CostAddend(InCostAddend)
 	{
-		check((AssetIndex & (1 << 31)) == 0);
-		Data = AssetIndex | (bBlockTransition ? 1 << 31 : 0);
-
-		CostAddend = InCostAddend;
+		// checking for overflowing inputs
+		check(InValueOffset < (1 << ValueOffsetNumBits));
+		check(InValueOffset < (1 << AssetIndexNumBits));
 	}
 
 	bool IsBlockTransition() const
 	{
-		return Data & (1 << 31);
+		return bBlockTransition;
 	}
 
 	uint32 GetAssetIndex() const
 	{
-		return Data & ~(1 << 31);
+		return AssetIndex;
 	}
 
 	float GetCostAddend() const
@@ -46,6 +54,16 @@ public:
 		return CostAddend;
 	}
 
+	uint32 GetValueOffset() const
+	{
+		return ValueOffset;
+	}
+
+	void SetValueOffset(uint32 Value)
+	{
+		ValueOffset = Value;
+	}
+
 	friend FArchive& operator<<(FArchive& Ar, FPoseMetadata& Metadata);
 };
 
@@ -161,7 +179,10 @@ struct FSearchStats
 */
 struct FSearchIndexBase
 {
+private:
 	TAlignedArray<float> Values;
+
+public:
 	TAlignedArray<FPoseMetadata> PoseMetadata;
 	bool bAnyBlockTransition = false;
 	TAlignedArray<FSearchIndexAsset> Assets;
@@ -173,14 +194,38 @@ struct FSearchIndexBase
 	FSearchStats Stats;
 
 	int32 GetNumPoses() const { return PoseMetadata.Num(); }
+	int32 GetNumValuesVectors(int32 DataCardinality) const
+	{
+		check(DataCardinality > 0);
+		check(Values.Num() % DataCardinality == 0);
+		return Values.Num() / DataCardinality;
+	}
+
 	bool IsValidPoseIndex(int32 PoseIdx) const { return PoseIdx < GetNumPoses(); }
 	bool IsEmpty() const;
+	bool IsValuesEmpty() const { return Values.IsEmpty(); }
+
+	void ResetValues() { Values.Reset(); }
+	void AllocateData(int32 DataCardinality, int32 NumPoses);
+	
+	const TAlignedArray<float>& GetValues() const { return Values; }
+	TAlignedArray<float>& EditValues() { return Values; }
 
 	const FSearchIndexAsset& GetAssetForPose(int32 PoseIdx) const;
 	POSESEARCH_API const FSearchIndexAsset* GetAssetForPoseSafe(int32 PoseIdx) const;
 
 	void Reset();
 	
+	void PruneDuplicateValues(float SimilarityThreshold, int32 DataCardinality);
+
+	TConstArrayView<float> GetPoseValuesBase(int32 PoseIdx, int32 DataCardinality) const
+	{
+		check(!IsValuesEmpty() && PoseIdx >= 0 && PoseIdx < GetNumPoses());
+		check(Values.Num() % DataCardinality == 0);
+		const int32 ValueOffset = PoseMetadata[PoseIdx].GetValueOffset();
+		return MakeArrayView(&Values[ValueOffset], DataCardinality);
+	}
+
 	friend FArchive& operator<<(FArchive& Ar, FSearchIndexBase& Index);
 };
 
