Blend stack: Evaluate the graph when MaxActiveBlends == 0 (used for inertialization)
#rnx

diff --git a/Engine/Plugins/Experimental/Animation/PoseSearch/Source/Runtime/Private/AnimNode_BlendStack.cpp b/Engine/Plugins/Experimental/Animation/PoseSearch/Source/Runtime/Private/AnimNode_BlendStack.cpp
--- a/Engine/Plugins/Experimental/Animation/PoseSearch/Source/Runtime/Private/AnimNode_BlendStack.cpp
+++ b/Engine/Plugins/Experimental/Animation/PoseSearch/Source/Runtime/Private/AnimNode_BlendStack.cpp
@@ -478,7 +478,9 @@ void FAnimNode_BlendStack_Standalone::EvaluateSample(FPoseContext& Output, const
 	// If we have any sample graphs, our player has been assigned a pose link index.
 	// If we are within X most relelvant players, then the graph is available.
 	// If PlayerIndex == MaxActiveBlends, don't evaluate that graph. It's reserved for the stored pose.
-	const bool bIsSampleGraphAvailable = !SampleGraphPoseLinks.IsEmpty() && (PlayerIndex < MaxActiveBlends);
+	// MaxActiveBlends == 0, means we're using inertialization. Run the the graph.
+	const bool bIsSampleGraphAvailable = !SampleGraphPoseLinks.IsEmpty() && 
+										((PlayerIndex < MaxActiveBlends) || (MaxActiveBlends == 0));
 	if (!bIsSampleGraphAvailable)
 	{
 		// If we have no sample graph, evaluate the player directly.
