motion matching - heading channel OriginTimeOffset support: the data relative to the sampling time associated to this channel origin (root / trajectory bone) will be offsetted by OriginTimeOffset seconds. For example, if Bone is the head bone, SampleTimeOffset is 0.5, and OriginTimeOffset is 0.5, this channel will try to match the future heading of the character head bone 0.5 seconds ahead, relative to the future root bone 0.5 seconds ahead

diff --git a/Engine/Plugins/Experimental/Animation/PoseSearch/Source/Runtime/Private/PoseSearchAssetIndexer.cpp b/Engine/Plugins/Experimental/Animation/PoseSearch/Source/Runtime/Private/PoseSearchAssetIndexer.cpp
--- a/Engine/Plugins/Experimental/Animation/PoseSearch/Source/Runtime/Private/PoseSearchAssetIndexer.cpp
+++ b/Engine/Plugins/Experimental/Animation/PoseSearch/Source/Runtime/Private/PoseSearchAssetIndexer.cpp
@@ -356,7 +356,7 @@ float FAssetIndexer::CalculateSampleTime(int32 SampleIdx) const
 	return SampleIdx / float(Schema.SampleRate);
 }
 
-FQuat FAssetIndexer::GetSampleRotation(float SampleTimeOffset, int32 SampleIdx, int8 SchemaSampleBoneIdx, int8 SchemaOriginBoneIdx, EPermutationTimeType PermutationTimeType)
+FQuat FAssetIndexer::GetSampleRotation(float SampleTimeOffset, float OriginTimeOffset, int32 SampleIdx, int8 SchemaSampleBoneIdx, int8 SchemaOriginBoneIdx, EPermutationTimeType PermutationTimeType)
 {
 	float PermutationSampleTimeOffset = 0.f;
 	float PermutationOriginTimeOffset = 0.f;
@@ -364,7 +364,7 @@ FQuat FAssetIndexer::GetSampleRotation(float SampleTimeOffset, int32 SampleIdx,
 
 	const float Time = CalculateSampleTime(SampleIdx);
 	const float SampleTime = Time + SampleTimeOffset + PermutationSampleTimeOffset;
-	const float OriginTime = Time + PermutationOriginTimeOffset;
+	const float OriginTime = Time + OriginTimeOffset + PermutationOriginTimeOffset;
 
 	// @todo: add support for SchemaSampleBoneIdx
 	if (!Schema.IsRootBone(SchemaOriginBoneIdx))
@@ -388,7 +388,7 @@ FQuat FAssetIndexer::GetSampleRotation(float SampleTimeOffset, int32 SampleIdx,
 	return BoneTransform.GetRotation();
 }
 
-FVector FAssetIndexer::GetSamplePosition(float SampleTimeOffset, int32 SampleIdx, int8 SchemaSampleBoneIdx, int8 SchemaOriginBoneIdx, EPermutationTimeType PermutationTimeType)
+FVector FAssetIndexer::GetSamplePosition(float SampleTimeOffset, float OriginTimeOffset, int32 SampleIdx, int8 SchemaSampleBoneIdx, int8 SchemaOriginBoneIdx, EPermutationTimeType PermutationTimeType)
 {
 	float PermutationSampleTimeOffset = 0.f;
 	float PermutationOriginTimeOffset = 0.f;
@@ -396,7 +396,7 @@ FVector FAssetIndexer::GetSamplePosition(float SampleTimeOffset, int32 SampleIdx
 
 	const float Time = CalculateSampleTime(SampleIdx);
 	const float SampleTime = Time + SampleTimeOffset + PermutationSampleTimeOffset;
-	const float OriginTime = Time + PermutationOriginTimeOffset;
+	const float OriginTime = Time + OriginTimeOffset + PermutationOriginTimeOffset;
 	
 	bool bUnused;
 	return GetSamplePositionInternal(SampleTime, OriginTime, bUnused, SchemaSampleBoneIdx, SchemaOriginBoneIdx);
@@ -433,7 +433,7 @@ FVector FAssetIndexer::GetSamplePositionInternal(float SampleTime, float OriginT
 	return RootBoneTransform.InverseTransformVector(DeltaBoneTranslation);
 }
 
-FVector FAssetIndexer::GetSampleVelocity(float SampleTimeOffset, int32 SampleIdx, int8 SchemaSampleBoneIdx, int8 SchemaOriginBoneIdx, bool bUseCharacterSpaceVelocities, EPermutationTimeType PermutationTimeType)
+FVector FAssetIndexer::GetSampleVelocity(float SampleTimeOffset, float OriginTimeOffset, int32 SampleIdx, int8 SchemaSampleBoneIdx, int8 SchemaOriginBoneIdx, bool bUseCharacterSpaceVelocities, EPermutationTimeType PermutationTimeType)
 {
 	float PermutationSampleTimeOffset = 0.f;
 	float PermutationOriginTimeOffset = 0.f;
@@ -441,7 +441,7 @@ FVector FAssetIndexer::GetSampleVelocity(float SampleTimeOffset, int32 SampleIdx
 
 	const float Time = CalculateSampleTime(SampleIdx);
 	const float SampleTime = Time + SampleTimeOffset + PermutationSampleTimeOffset;
-	const float OriginTime = Time + PermutationOriginTimeOffset;
+	const float OriginTime = Time + OriginTimeOffset + PermutationOriginTimeOffset;
 	const float FiniteDelta = SamplingContext.FiniteDelta;
 
 	bool bClampedPast, bClampedPresent;
diff --git a/Engine/Plugins/Experimental/Animation/PoseSearch/Source/Runtime/Private/PoseSearchContext.cpp b/Engine/Plugins/Experimental/Animation/PoseSearch/Source/Runtime/Private/PoseSearchContext.cpp
--- a/Engine/Plugins/Experimental/Animation/PoseSearch/Source/Runtime/Private/PoseSearchContext.cpp
+++ b/Engine/Plugins/Experimental/Animation/PoseSearch/Source/Runtime/Private/PoseSearchContext.cpp
@@ -271,37 +271,37 @@ FSearchContext::FSearchContext(const FPoseSearchQueryTrajectory* InTrajectory, c
 	}
 }
 
-FQuat FSearchContext::GetSampleRotation(float SampleTimeOffset, const UPoseSearchSchema* Schema, int8 SchemaSampleBoneIdx, int8 SchemaOriginBoneIdx, bool bUseHistoryRoot, EPermutationTimeType PermutationTimeType)
+FQuat FSearchContext::GetSampleRotation(float SampleTimeOffset, float OriginTimeOffset, const UPoseSearchSchema* Schema, int8 SchemaSampleBoneIdx, int8 SchemaOriginBoneIdx, bool bUseHistoryRoot, EPermutationTimeType PermutationTimeType)
 {
 	float PermutationSampleTimeOffset = 0.f;
 	float PermutationOriginTimeOffset = 0.f;
 	UPoseSearchFeatureChannel::GetPermutationTimeOffsets(PermutationTimeType, DesiredPermutationTimeOffset, PermutationSampleTimeOffset, PermutationOriginTimeOffset);
 
 	const float SampleTime = SampleTimeOffset + PermutationSampleTimeOffset;
-	const float OriginTime = PermutationOriginTimeOffset;
+	const float OriginTime = OriginTimeOffset + PermutationOriginTimeOffset;
 
 	return GetSampleRotationInternal(SampleTime, OriginTime, Schema, SchemaSampleBoneIdx, SchemaOriginBoneIdx, bUseHistoryRoot);
 }
 
-FVector FSearchContext::GetSamplePosition(float SampleTimeOffset, const UPoseSearchSchema* Schema, int8 SchemaSampleBoneIdx, int8 SchemaOriginBoneIdx, bool bUseHistoryRoot, EPermutationTimeType PermutationTimeType)
+FVector FSearchContext::GetSamplePosition(float SampleTimeOffset, float OriginTimeOffset, const UPoseSearchSchema* Schema, int8 SchemaSampleBoneIdx, int8 SchemaOriginBoneIdx, bool bUseHistoryRoot, EPermutationTimeType PermutationTimeType)
 {
 	float PermutationSampleTimeOffset = 0.f;
 	float PermutationOriginTimeOffset = 0.f;
 	UPoseSearchFeatureChannel::GetPermutationTimeOffsets(PermutationTimeType, DesiredPermutationTimeOffset, PermutationSampleTimeOffset, PermutationOriginTimeOffset);
 
 	const float SampleTime = SampleTimeOffset + PermutationSampleTimeOffset;
-	const float OriginTime = PermutationOriginTimeOffset;
+	const float OriginTime = OriginTimeOffset + PermutationOriginTimeOffset;
 	return GetSamplePositionInternal(SampleTime, OriginTime, Schema, SchemaSampleBoneIdx, SchemaOriginBoneIdx, bUseHistoryRoot);
 }
 
-FVector FSearchContext::GetSampleVelocity(float SampleTimeOffset, const UPoseSearchSchema* Schema, int8 SchemaSampleBoneIdx, int8 SchemaOriginBoneIdx, bool bUseCharacterSpaceVelocities, bool bUseHistoryRoot, EPermutationTimeType PermutationTimeType)
+FVector FSearchContext::GetSampleVelocity(float SampleTimeOffset, float OriginTimeOffset, const UPoseSearchSchema* Schema, int8 SchemaSampleBoneIdx, int8 SchemaOriginBoneIdx, bool bUseCharacterSpaceVelocities, bool bUseHistoryRoot, EPermutationTimeType PermutationTimeType)
 {
 	float PermutationSampleTimeOffset = 0.f;
 	float PermutationOriginTimeOffset = 0.f;
 	UPoseSearchFeatureChannel::GetPermutationTimeOffsets(PermutationTimeType, DesiredPermutationTimeOffset, PermutationSampleTimeOffset, PermutationOriginTimeOffset);
 
 	const float SampleTime = SampleTimeOffset + PermutationSampleTimeOffset;
-	const float OriginTime = PermutationOriginTimeOffset;
+	const float OriginTime = OriginTimeOffset + PermutationOriginTimeOffset;
 	const float FiniteDelta = History ? History->GetSampleTimeInterval() : 1 / 60.0f;
 	check(FiniteDelta > UE_KINDA_SMALL_NUMBER);
 
diff --git a/Engine/Plugins/Experimental/Animation/PoseSearch/Source/Runtime/Private/PoseSearchFeatureChannel_FilterCrashingLegs.cpp b/Engine/Plugins/Experimental/Animation/PoseSearch/Source/Runtime/Private/PoseSearchFeatureChannel_FilterCrashingLegs.cpp
--- a/Engine/Plugins/Experimental/Animation/PoseSearch/Source/Runtime/Private/PoseSearchFeatureChannel_FilterCrashingLegs.cpp
+++ b/Engine/Plugins/Experimental/Animation/PoseSearch/Source/Runtime/Private/PoseSearchFeatureChannel_FilterCrashingLegs.cpp
@@ -70,10 +70,10 @@ void UPoseSearchFeatureChannel_FilterCrashingLegs::BuildQuery(UE::PoseSearch::FS
 	}
 	else
 	{
-		const FVector RightThighPosition = SearchContext.GetSamplePosition(0.f, InOutQuery.GetSchema(), RightThighIdx);
-		const FVector LeftThighPosition = SearchContext.GetSamplePosition(0.f, InOutQuery.GetSchema(), LeftThighIdx);
-		const FVector RightFootPosition = SearchContext.GetSamplePosition(0.f, InOutQuery.GetSchema(), RightFootIdx);
-		const FVector LeftFootPosition = SearchContext.GetSamplePosition(0.f, InOutQuery.GetSchema(), LeftFootIdx);
+		const FVector RightThighPosition = SearchContext.GetSamplePosition(0.f, 0.f, InOutQuery.GetSchema(), RightThighIdx);
+		const FVector LeftThighPosition = SearchContext.GetSamplePosition(0.f, 0.f, InOutQuery.GetSchema(), LeftThighIdx);
+		const FVector RightFootPosition = SearchContext.GetSamplePosition(0.f, 0.f, InOutQuery.GetSchema(), RightFootIdx);
+		const FVector LeftFootPosition = SearchContext.GetSamplePosition(0.f, 0.f, InOutQuery.GetSchema(), LeftFootIdx);
 
 		const float CrashingLegsValue = ComputeCrashingLegsValue(RightThighPosition, LeftThighPosition, RightFootPosition, LeftFootPosition);
 
@@ -155,10 +155,10 @@ void UPoseSearchFeatureChannel_FilterCrashingLegs::IndexAsset(UE::PoseSearch::FA
 
 	for (int32 SampleIdx = Indexer.GetBeginSampleIdx(); SampleIdx != Indexer.GetEndSampleIdx(); ++SampleIdx)
 	{
-		const FVector RightThighPosition = Indexer.GetSamplePosition(0.f, SampleIdx, RightThighIdx);
-		const FVector LeftThighPosition = Indexer.GetSamplePosition(0.f, SampleIdx, LeftThighIdx);
-		const FVector RightFootPosition = Indexer.GetSamplePosition(0.f, SampleIdx, RightFootIdx);
-		const FVector LeftFootPosition = Indexer.GetSamplePosition(0.f, SampleIdx, LeftFootIdx);
+		const FVector RightThighPosition = Indexer.GetSamplePosition(0.f, 0.f, SampleIdx, RightThighIdx);
+		const FVector LeftThighPosition = Indexer.GetSamplePosition(0.f, 0.f, SampleIdx, LeftThighIdx);
+		const FVector RightFootPosition = Indexer.GetSamplePosition(0.f, 0.f, SampleIdx, RightFootIdx);
+		const FVector LeftFootPosition = Indexer.GetSamplePosition(0.f, 0.f, SampleIdx, LeftFootIdx);
 
 		const float CrashingLegsValue = ComputeCrashingLegsValue(RightThighPosition, LeftThighPosition, RightFootPosition, LeftFootPosition);
 
diff --git a/Engine/Plugins/Experimental/Animation/PoseSearch/Source/Runtime/Private/PoseSearchFeatureChannel_Heading.cpp b/Engine/Plugins/Experimental/Animation/PoseSearch/Source/Runtime/Private/PoseSearchFeatureChannel_Heading.cpp
--- a/Engine/Plugins/Experimental/Animation/PoseSearch/Source/Runtime/Private/PoseSearchFeatureChannel_Heading.cpp
+++ b/Engine/Plugins/Experimental/Animation/PoseSearch/Source/Runtime/Private/PoseSearchFeatureChannel_Heading.cpp
@@ -62,7 +62,7 @@ void UPoseSearchFeatureChannel_Heading::BuildQuery(UE::PoseSearch::FSearchContex
 	else
 	{
 		// calculating the BoneRotation in component space for the bone indexed by SchemaBoneIdx
-		const FQuat BoneRotation = SearchContext.GetSampleRotation(SampleTimeOffset, InOutQuery.GetSchema(), SchemaBoneIdx, RootSchemaBoneIdx, !bIsRootBone);
+		const FQuat BoneRotation = SearchContext.GetSampleRotation(SampleTimeOffset, OriginTimeOffset, InOutQuery.GetSchema(), SchemaBoneIdx, RootSchemaBoneIdx, !bIsRootBone);
 		FFeatureVectorHelper::EncodeVector(InOutQuery.EditValues(), ChannelDataOffset, GetAxis(BoneRotation), ComponentStripping);
 	}
 }
@@ -96,7 +96,7 @@ void UPoseSearchFeatureChannel_Heading::IndexAsset(UE::PoseSearch::FAssetIndexer
 
 	for (int32 SampleIdx = Indexer.GetBeginSampleIdx(); SampleIdx != Indexer.GetEndSampleIdx(); ++SampleIdx)
 	{
-		const FVector Heading = GetAxis(Indexer.GetSampleRotation(SampleTimeOffset, SampleIdx, SchemaBoneIdx));
+		const FVector Heading = GetAxis(Indexer.GetSampleRotation(SampleTimeOffset, OriginTimeOffset, SampleIdx, SchemaBoneIdx));
 		FFeatureVectorHelper::EncodeVector(Indexer.GetPoseVector(SampleIdx), ChannelDataOffset, Heading, ComponentStripping);
 	}
 }
@@ -142,6 +142,12 @@ FString UPoseSearchFeatureChannel_Heading::GetLabel() const
 	}
 
 	Label.Appendf(TEXT(" %.2f"), SampleTimeOffset);
+
+	if (!FMath::IsNearlyZero(OriginTimeOffset))
+	{
+		Label.Appendf(TEXT("-%.2f"), OriginTimeOffset);
+	}
+
 	return Label.ToString();
 }
 #endif
\ No newline at end of file
diff --git a/Engine/Plugins/Experimental/Animation/PoseSearch/Source/Runtime/Private/PoseSearchFeatureChannel_Heading.h b/Engine/Plugins/Experimental/Animation/PoseSearch/Source/Runtime/Private/PoseSearchFeatureChannel_Heading.h
--- a/Engine/Plugins/Experimental/Animation/PoseSearch/Source/Runtime/Private/PoseSearchFeatureChannel_Heading.h
+++ b/Engine/Plugins/Experimental/Animation/PoseSearch/Source/Runtime/Private/PoseSearchFeatureChannel_Heading.h
@@ -29,12 +29,21 @@ public:
 	UPROPERTY(EditAnywhere, Category = "Settings")
 	float Weight = 1.f;
 
+	// the data relative to the sampling time associated to this channel will be offsetted by SampleTimeOffset seconds.
+	// For example, if Bone is the head bone, and SampleTimeOffset is 0.5, this channel will try to match the future heading of the character head bone 0.5 seconds ahead
 	UPROPERTY(EditAnywhere, Category = "Settings")
 	float SampleTimeOffset = 0.f;
 
+	// the data relative to the sampling time associated to this channel origin (root / trajectory bone) will be offsetted by OriginTimeOffset seconds.
+	// For example, if Bone is the head bone, SampleTimeOffset is 0.5, and OriginTimeOffset is 0.5, this channel will try to match 
+	// the future heading of the character head bone 0.5 seconds ahead, relative to the future root bone 0.5 seconds ahead
+	UPROPERTY(EditAnywhere, Category = "Settings")
+	float OriginTimeOffset = 0.f;
+
 	UPROPERTY(EditAnywhere, Category = "Settings")
 	EHeadingAxis HeadingAxis = EHeadingAxis::X;	
 
+	// index referencing the associated bone in UPoseSearchSchema::BoneReferences
 	UPROPERTY(Transient)
 	int8 SchemaBoneIdx = 0;
 
diff --git a/Engine/Plugins/Experimental/Animation/PoseSearch/Source/Runtime/Private/PoseSearchFeatureChannel_Phase.cpp b/Engine/Plugins/Experimental/Animation/PoseSearch/Source/Runtime/Private/PoseSearchFeatureChannel_Phase.cpp
--- a/Engine/Plugins/Experimental/Animation/PoseSearch/Source/Runtime/Private/PoseSearchFeatureChannel_Phase.cpp
+++ b/Engine/Plugins/Experimental/Animation/PoseSearch/Source/Runtime/Private/PoseSearchFeatureChannel_Phase.cpp
@@ -50,7 +50,7 @@ namespace UE::PoseSearch
 		BonePositions.AddDefaulted(NumSamples);
 		for (int32 SampleIdx = 0; SampleIdx != NumSamples; ++SampleIdx)
 		{
-			BonePositions[SampleIdx] = Indexer.GetSamplePosition(0.f, SampleIdx, SchemaBoneIdx);
+			BonePositions[SampleIdx] = Indexer.GetSamplePosition(0.f, 0.f, SampleIdx, SchemaBoneIdx, RootSchemaBoneIdx, EPermutationTimeType::UseSampleTime);
 		}
 	}
 
diff --git a/Engine/Plugins/Experimental/Animation/PoseSearch/Source/Runtime/Private/PoseSearchFeatureChannel_Phase.h b/Engine/Plugins/Experimental/Animation/PoseSearch/Source/Runtime/Private/PoseSearchFeatureChannel_Phase.h
--- a/Engine/Plugins/Experimental/Animation/PoseSearch/Source/Runtime/Private/PoseSearchFeatureChannel_Phase.h
+++ b/Engine/Plugins/Experimental/Animation/PoseSearch/Source/Runtime/Private/PoseSearchFeatureChannel_Phase.h
@@ -18,6 +18,7 @@ public:
 	UPROPERTY(EditAnywhere, Category = "Settings")
 	float Weight = 1.f;
 
+	// index referencing the associated bone in UPoseSearchSchema::BoneReferences
 	UPROPERTY(Transient)
 	int8 SchemaBoneIdx = 0;
 
diff --git a/Engine/Plugins/Experimental/Animation/PoseSearch/Source/Runtime/Private/PoseSearchFeatureChannel_Position.cpp b/Engine/Plugins/Experimental/Animation/PoseSearch/Source/Runtime/Private/PoseSearchFeatureChannel_Position.cpp
--- a/Engine/Plugins/Experimental/Animation/PoseSearch/Source/Runtime/Private/PoseSearchFeatureChannel_Position.cpp
+++ b/Engine/Plugins/Experimental/Animation/PoseSearch/Source/Runtime/Private/PoseSearchFeatureChannel_Position.cpp
@@ -77,7 +77,7 @@ void UPoseSearchFeatureChannel_Position::BuildQuery(UE::PoseSearch::FSearchConte
 	else
 	{
 		// calculating the BonePosition in component space for the bone indexed by SchemaBoneIdx
-		const FVector BonePosition = SearchContext.GetSamplePosition(SampleTimeOffset, InOutQuery.GetSchema(), SchemaBoneIdx, SchemaOriginBoneIdx, !bIsRootBone, PermutationTimeType);
+		const FVector BonePosition = SearchContext.GetSamplePosition(SampleTimeOffset, 0.f, InOutQuery.GetSchema(), SchemaBoneIdx, SchemaOriginBoneIdx, !bIsRootBone, PermutationTimeType);
 		FFeatureVectorHelper::EncodeVector(InOutQuery.EditValues(), ChannelDataOffset, BonePosition, ComponentStripping);
 	}
 }
@@ -123,7 +123,7 @@ void UPoseSearchFeatureChannel_Position::IndexAsset(UE::PoseSearch::FAssetIndexe
 
 	for (int32 SampleIdx = Indexer.GetBeginSampleIdx(); SampleIdx != Indexer.GetEndSampleIdx(); ++SampleIdx)
 	{
-		const FVector BonePosition = Indexer.GetSamplePosition(SampleTimeOffset, SampleIdx, SchemaBoneIdx, SchemaOriginBoneIdx, PermutationTimeType);
+		const FVector BonePosition = Indexer.GetSamplePosition(SampleTimeOffset, 0.f, SampleIdx, SchemaBoneIdx, SchemaOriginBoneIdx, PermutationTimeType);
 		FFeatureVectorHelper::EncodeVector(Indexer.GetPoseVector(SampleIdx), ChannelDataOffset, BonePosition, ComponentStripping);
 	}
 }
diff --git a/Engine/Plugins/Experimental/Animation/PoseSearch/Source/Runtime/Private/PoseSearchFeatureChannel_Position.h b/Engine/Plugins/Experimental/Animation/PoseSearch/Source/Runtime/Private/PoseSearchFeatureChannel_Position.h
--- a/Engine/Plugins/Experimental/Animation/PoseSearch/Source/Runtime/Private/PoseSearchFeatureChannel_Position.h
+++ b/Engine/Plugins/Experimental/Animation/PoseSearch/Source/Runtime/Private/PoseSearchFeatureChannel_Position.h
@@ -21,9 +21,12 @@ public:
 	UPROPERTY(EditAnywhere, Category = "Settings")
 	float Weight = 1.f;
 
+	// the data relative to the sampling time associated to this channel will be offsetted by SampleTimeOffset seconds.
+	// For example, if Bone is the head bone, and SampleTimeOffset is 0.5, this channel will try to match the future position of the character head bone 0.5 seconds ahead
 	UPROPERTY(EditAnywhere, Category = "Settings")
 	float SampleTimeOffset = 0.f;
 
+	// index referencing the associated bone in UPoseSearchSchema::BoneReferences
 	UPROPERTY(Transient)
 	int8 SchemaBoneIdx = 0;
 
diff --git a/Engine/Plugins/Experimental/Animation/PoseSearch/Source/Runtime/Private/PoseSearchFeatureChannel_Trajectory.h b/Engine/Plugins/Experimental/Animation/PoseSearch/Source/Runtime/Private/PoseSearchFeatureChannel_Trajectory.h
--- a/Engine/Plugins/Experimental/Animation/PoseSearch/Source/Runtime/Private/PoseSearchFeatureChannel_Trajectory.h
+++ b/Engine/Plugins/Experimental/Animation/PoseSearch/Source/Runtime/Private/PoseSearchFeatureChannel_Trajectory.h
@@ -27,8 +27,10 @@ struct POSESEARCH_API FPoseSearchTrajectorySample
 {
 	GENERATED_BODY()
 
+	// the data relative to the sampling time associated to this FPoseSearchTrajectorySample will be offsetted by Offset seconds.
+	// For example, Flags is Position, and Offset is 0.5, this channel will try to match the future position of the trajectory 0.5 seconds ahead
 	UPROPERTY(EditAnywhere, Category = Config)
-	float Offset = 0.f; // offset in time or distance depending on UPoseSearchFeatureChannel_Trajectory.Domain
+	float Offset = 0.f;
 
 	// This allows the user to define what information from the channel you want to compare to.
 	UPROPERTY(EditAnywhere, meta = (Bitmask, BitmaskEnum = "/Script/PoseSearch.EPoseSearchTrajectoryFlags"), Category = Config)
diff --git a/Engine/Plugins/Experimental/Animation/PoseSearch/Source/Runtime/Private/PoseSearchFeatureChannel_Velocity.cpp b/Engine/Plugins/Experimental/Animation/PoseSearch/Source/Runtime/Private/PoseSearchFeatureChannel_Velocity.cpp
--- a/Engine/Plugins/Experimental/Animation/PoseSearch/Source/Runtime/Private/PoseSearchFeatureChannel_Velocity.cpp
+++ b/Engine/Plugins/Experimental/Animation/PoseSearch/Source/Runtime/Private/PoseSearchFeatureChannel_Velocity.cpp
@@ -48,7 +48,7 @@ void UPoseSearchFeatureChannel_Velocity::BuildQuery(UE::PoseSearch::FSearchConte
 	else
 	{
 		// calculating the LinearVelocity for the bone indexed by SchemaBoneIdx
-		FVector LinearVelocity = SearchContext.GetSampleVelocity(SampleTimeOffset, InOutQuery.GetSchema(), SchemaBoneIdx, RootSchemaBoneIdx, bUseCharacterSpaceVelocities, !bIsRootBone);
+		FVector LinearVelocity = SearchContext.GetSampleVelocity(SampleTimeOffset, 0.f, InOutQuery.GetSchema(), SchemaBoneIdx, RootSchemaBoneIdx, bUseCharacterSpaceVelocities, !bIsRootBone, EPermutationTimeType::UseSampleTime);
 		if (bNormalize)
 		{
 			LinearVelocity = LinearVelocity.GetClampedToMaxSize(1.f);
@@ -89,7 +89,7 @@ void UPoseSearchFeatureChannel_Velocity::IndexAsset(UE::PoseSearch::FAssetIndexe
 
 	for (int32 SampleIdx = Indexer.GetBeginSampleIdx(); SampleIdx != Indexer.GetEndSampleIdx(); ++SampleIdx)
 	{
-		FVector LinearVelocity = Indexer.GetSampleVelocity(SampleTimeOffset, SampleIdx, SchemaBoneIdx, RootSchemaBoneIdx, bUseCharacterSpaceVelocities);
+		FVector LinearVelocity = Indexer.GetSampleVelocity(SampleTimeOffset, 0.f, SampleIdx, SchemaBoneIdx, RootSchemaBoneIdx, bUseCharacterSpaceVelocities, EPermutationTimeType::UseSampleTime);
 		if (bNormalize)
 		{
 			LinearVelocity = LinearVelocity.GetClampedToMaxSize(1.f);
diff --git a/Engine/Plugins/Experimental/Animation/PoseSearch/Source/Runtime/Private/PoseSearchFeatureChannel_Velocity.h b/Engine/Plugins/Experimental/Animation/PoseSearch/Source/Runtime/Private/PoseSearchFeatureChannel_Velocity.h
--- a/Engine/Plugins/Experimental/Animation/PoseSearch/Source/Runtime/Private/PoseSearchFeatureChannel_Velocity.h
+++ b/Engine/Plugins/Experimental/Animation/PoseSearch/Source/Runtime/Private/PoseSearchFeatureChannel_Velocity.h
@@ -18,9 +18,12 @@ public:
 	UPROPERTY(EditAnywhere, Category = "Settings")
 	float Weight = 1.f;
 
+	// the data relative to the sampling time associated to this channel will be offsetted by SampleTimeOffset seconds.
+	// For example, if Bone is the head bone, and SampleTimeOffset is 0.5, this channel will try to match the future velocity of the character head bone 0.5 seconds ahead
 	UPROPERTY(EditAnywhere, Category = "Settings")
 	float SampleTimeOffset = 0.f;
 
+	// index referencing the associated bone in UPoseSearchSchema::BoneReferences
 	UPROPERTY(Transient)
 	int8 SchemaBoneIdx = 0;
 
diff --git a/Engine/Plugins/Experimental/Animation/PoseSearch/Source/Runtime/Public/PoseSearch/PoseSearchAssetIndexer.h b/Engine/Plugins/Experimental/Animation/PoseSearch/Source/Runtime/Public/PoseSearch/PoseSearchAssetIndexer.h
--- a/Engine/Plugins/Experimental/Animation/PoseSearch/Source/Runtime/Public/PoseSearch/PoseSearchAssetIndexer.h
+++ b/Engine/Plugins/Experimental/Animation/PoseSearch/Source/Runtime/Public/PoseSearch/PoseSearchAssetIndexer.h
@@ -36,9 +36,21 @@ public:
 	void Process(int32 AssetIdx);
 	const FStats& GetStats() const { return Stats; }
 
-	FQuat GetSampleRotation(float SampleTimeOffset, int32 SampleIdx, int8 SchemaSampleBoneIdx = RootSchemaBoneIdx, int8 SchemaOriginBoneIdx = RootSchemaBoneIdx, EPermutationTimeType PermutationTimeType = EPermutationTimeType::UseSampleTime);
-	FVector GetSamplePosition(float SampleTimeOffset, int32 SampleIdx, int8 SchemaSampleBoneIdx = RootSchemaBoneIdx, int8 SchemaOriginBoneIdx = RootSchemaBoneIdx, EPermutationTimeType PermutationTimeType = EPermutationTimeType::UseSampleTime);
-	FVector GetSampleVelocity(float SampleTimeOffset, int32 SampleIdx, int8 SchemaSampleBoneIdx = RootSchemaBoneIdx, int8 SchemaOriginBoneIdx = RootSchemaBoneIdx, bool bUseCharacterSpaceVelocities = true, EPermutationTimeType PermutationTimeType = EPermutationTimeType::UseSampleTime);
+	// Returns the rotation of the bone Schema.BoneReferences[SchemaSampleBoneIdx] at time CalculateSampleTime(SampleIdx) + SampleTimeOffset relative to the
+	// transform of the bone Schema.BoneReferences[SchemaOriginBoneIdx] at time CalculateSampleTime(SampleIdx) + OriginTimeOffset 
+	// Times will be processed by GetPermutationTimeOffsets(PermutationTimeType, ...)
+	FQuat GetSampleRotation(float SampleTimeOffset, float OriginTimeOffset, int32 SampleIdx, int8 SchemaSampleBoneIdx = RootSchemaBoneIdx, int8 SchemaOriginBoneIdx = RootSchemaBoneIdx, EPermutationTimeType PermutationTimeType = EPermutationTimeType::UseSampleTime);
+
+	// Returns the position of the bone Schema.BoneReferences[SchemaSampleBoneIdx] at time CalculateSampleTime(SampleIdx) + SampleTimeOffset relative to the
+	// transform of the bone Schema.BoneReferences[SchemaOriginBoneIdx] at time CalculateSampleTime(SampleIdx) + OriginTimeOffset.
+	// Times will be processed by GetPermutationTimeOffsets(PermutationTimeType, ...)
+	FVector GetSamplePosition(float SampleTimeOffset, float OriginTimeOffset, int32 SampleIdx, int8 SchemaSampleBoneIdx = RootSchemaBoneIdx, int8 SchemaOriginBoneIdx = RootSchemaBoneIdx, EPermutationTimeType PermutationTimeType = EPermutationTimeType::UseSampleTime);
+
+	// Returns the delta velocity of the bone Schema.BoneReferences[SchemaSampleBoneIdx] velocity at time CalculateSampleTime(SampleIdx) + SampleTimeOffset minus
+	// the velocity of the bone Schema.BoneReferences[SchemaOriginBoneIdx] at time CalculateSampleTime(SampleIdx) + OriginTimeOffset.
+	// Times will be processed by GetPermutationTimeOffsets(PermutationTimeType, ...)
+	// if bUseCharacterSpaceVelocities is true, velocities will be computed in root bone space, rather than animation (world) space
+	FVector GetSampleVelocity(float SampleTimeOffset, float OriginTimeOffset, int32 SampleIdx, int8 SchemaSampleBoneIdx = RootSchemaBoneIdx, int8 SchemaOriginBoneIdx = RootSchemaBoneIdx, bool bUseCharacterSpaceVelocities = true, EPermutationTimeType PermutationTimeType = EPermutationTimeType::UseSampleTime);
 
 	int32 GetBeginSampleIdx() const;
 	int32 GetEndSampleIdx() const;
@@ -51,8 +63,18 @@ public:
 
 private:
 	int32 GetVectorIdx(int32 SampleIdx) const;
+
+	// Returns the animation (world) space transform of the bone Schema.BoneReferences[SchemaBoneIdx] at time SampleTime
+	// bClamped will be true if SampleTime is outside the animation duration boundaries
 	FTransform GetTransform(float SampleTime, bool& bClamped, int8 SchemaBoneIdx = RootSchemaBoneIdx);
+
+	// Returns the component space transform of the bone Schema.BoneReferences[SchemaBoneIdx] at time SampleTime
+	// bClamped will be true if SampleTime is outside the animation duration boundaries
 	FTransform GetComponentSpaceTransform(float SampleTime, bool& bClamped, int8 SchemaBoneIdx = RootSchemaBoneIdx);
+
+	// Returns the position of the bone Schema.BoneReferences[SchemaSampleBoneIdx] at time SampleTime relative to the
+	// transform of the bone Schema.BoneReferences[SchemaOriginBoneIdx] at time OriginTime.
+	// bClamped will be true if SampleTime or OriginTime are outside the animation duration boundaries
 	FVector GetSamplePositionInternal(float SampleTime, float OriginTime, bool& bClamped, int8 SchemaSampleBoneIdx, int8 SchemaOriginBoneIdx);
 
 	struct FSampleInfo
diff --git a/Engine/Plugins/Experimental/Animation/PoseSearch/Source/Runtime/Public/PoseSearch/PoseSearchContext.h b/Engine/Plugins/Experimental/Animation/PoseSearch/Source/Runtime/Public/PoseSearch/PoseSearchContext.h
--- a/Engine/Plugins/Experimental/Animation/PoseSearch/Source/Runtime/Public/PoseSearch/PoseSearchContext.h
+++ b/Engine/Plugins/Experimental/Animation/PoseSearch/Source/Runtime/Public/PoseSearch/PoseSearchContext.h
@@ -142,9 +142,24 @@ struct POSESEARCH_API FSearchContext
 	FSearchContext(const FPoseSearchQueryTrajectory* InTrajectory, const IPoseHistory* InHistory, float InDesiredPermutationTimeOffset, const FPoseIndicesHistory* InPoseIndicesHistory = nullptr,
 		const FSearchResult& InCurrentResult = FSearchResult(), float InPoseJumpThresholdTime = 0.f, bool bInForceInterrupt = false);
 
-	FQuat GetSampleRotation(float SampleTimeOffset, const UPoseSearchSchema* Schema, int8 SchemaSampleBoneIdx = RootSchemaBoneIdx, int8 SchemaOriginBoneIdx = RootSchemaBoneIdx, bool bUseHistoryRoot = false, EPermutationTimeType PermutationTimeType = EPermutationTimeType::UseSampleTime);
-	FVector GetSamplePosition(float SampleTimeOffset, const UPoseSearchSchema* Schema, int8 SchemaSampleBoneIdx = RootSchemaBoneIdx, int8 SchemaOriginBoneIdx = RootSchemaBoneIdx, bool bUseHistoryRoot = false, EPermutationTimeType PermutationTimeType = EPermutationTimeType::UseSampleTime);
-	FVector GetSampleVelocity(float SampleTimeOffset, const UPoseSearchSchema* Schema, int8 SchemaSampleBoneIdx = RootSchemaBoneIdx, int8 SchemaOriginBoneIdx = RootSchemaBoneIdx, bool bUseCharacterSpaceVelocities = true, bool bUseHistoryRoot = false, EPermutationTimeType PermutationTimeType = EPermutationTimeType::UseSampleTime);
+	// Returns the rotation of the bone Schema.BoneReferences[SchemaSampleBoneIdx] at an offset time of SampleTimeOffset relative to the
+	// transform of the bone Schema.BoneReferences[SchemaOriginBoneIdx] at an offset time of time OriginTimeOffset 
+	// Times will be processed by GetPermutationTimeOffsets(PermutationTimeType, ...)
+	// if bUseHistoryRoot is true, eventually required root transforms will be gathered from the pose history node data, otherwise the context trajectory will be used
+	FQuat GetSampleRotation(float SampleTimeOffset, float OriginTimeOffset, const UPoseSearchSchema* Schema, int8 SchemaSampleBoneIdx = RootSchemaBoneIdx, int8 SchemaOriginBoneIdx = RootSchemaBoneIdx, bool bUseHistoryRoot = false, EPermutationTimeType PermutationTimeType = EPermutationTimeType::UseSampleTime);
+	
+	// Returns the position of the bone Schema.BoneReferences[SchemaSampleBoneIdx] at an offset time of SampleTimeOffset relative to the
+	// transform of the bone Schema.BoneReferences[SchemaOriginBoneIdx] at an offset time of time OriginTimeOffset 
+	// Times will be processed by GetPermutationTimeOffsets(PermutationTimeType, ...)
+	// if bUseHistoryRoot is true, eventually required root transforms will be gathered from the pose history node data, otherwise the context trajectory will be used
+	FVector GetSamplePosition(float SampleTimeOffset, float OriginTimeOffset, const UPoseSearchSchema* Schema, int8 SchemaSampleBoneIdx = RootSchemaBoneIdx, int8 SchemaOriginBoneIdx = RootSchemaBoneIdx, bool bUseHistoryRoot = false, EPermutationTimeType PermutationTimeType = EPermutationTimeType::UseSampleTime);
+	
+	// Returns the delta velocity of the velocity of the bone Schema.BoneReferences[SchemaSampleBoneIdx] at an offset time of SampleTimeOffset minus
+	// the velocity of the bone Schema.BoneReferences[SchemaOriginBoneIdx] at an offset time of time OriginTimeOffset 
+	// Times will be processed by GetPermutationTimeOffsets(PermutationTimeType, ...)
+	// if bUseCharacterSpaceVelocities is true, velocities will be computed in root bone space, rather than world space
+	// if bUseHistoryRoot is true, eventually required root transforms will be gathered from the pose history node data, otherwise the context trajectory will be used
+	FVector GetSampleVelocity(float SampleTimeOffset, float OriginTimeOffset, const UPoseSearchSchema* Schema, int8 SchemaSampleBoneIdx = RootSchemaBoneIdx, int8 SchemaOriginBoneIdx = RootSchemaBoneIdx, bool bUseCharacterSpaceVelocities = true, bool bUseHistoryRoot = false, EPermutationTimeType PermutationTimeType = EPermutationTimeType::UseSampleTime);
 
 	void ClearCachedEntries();
 
