motion matching - support for database index pca space values pruning

diff --git a/Engine/Plugins/Experimental/Animation/PoseSearch/Source/Editor/Private/PoseSearchDatabaseEditorReflection.cpp b/Engine/Plugins/Experimental/Animation/PoseSearch/Source/Editor/Private/PoseSearchDatabaseEditorReflection.cpp
--- a/Engine/Plugins/Experimental/Animation/PoseSearch/Source/Editor/Private/PoseSearchDatabaseEditorReflection.cpp
+++ b/Engine/Plugins/Experimental/Animation/PoseSearch/Source/Editor/Private/PoseSearchDatabaseEditorReflection.cpp
@@ -135,7 +135,7 @@ void UPoseSearchDatabaseStatistics::Initialize(const UPoseSearchDatabase* PoseSe
 
 		if (ConfigCardinality > 0)
 		{
-			const int32 TotalAnimationFeatureVectors = SearchIndex.GetValues().Num() / ConfigCardinality;
+			const int32 TotalAnimationFeatureVectors = SearchIndex.Values.Num() / ConfigCardinality;
 			PrunedFrames = TotalAnimationPosesInFrames - TotalAnimationFeatureVectors;
 		}
 		else
@@ -143,6 +143,16 @@ void UPoseSearchDatabaseStatistics::Initialize(const UPoseSearchDatabase* PoseSe
 			PrunedFrames = 0;
 		}
 
+		if (PoseSearchDatabase->GetNumberOfPrincipalComponents() > 0)
+		{
+			const int32 TotalAnimationPCAFeatureVectors = SearchIndex.PCAValues.Num() / PoseSearchDatabase->GetNumberOfPrincipalComponents();
+			PrunedPCAFrames = TotalAnimationPosesInFrames - TotalAnimationPCAFeatureVectors;
+		}
+		else
+		{
+			PrunedPCAFrames = 0;
+		}
+
 		// Kinematic Information
 	
 		// using FText instead of meta = (ForceUnits = "cm/s") to keep properties consistent
@@ -158,13 +168,15 @@ void UPoseSearchDatabaseStatistics::Initialize(const UPoseSearchDatabase* PoseSe
 		// Memory Information
 			
 		{
-			const uint32 ValuesBytesSize = SearchIndex.GetValues().GetAllocatedSize();
+			const uint32 ValuesBytesSize = SearchIndex.Values.GetAllocatedSize();
 			const uint32 PCAValuesBytesSize = SearchIndex.PCAValues.GetAllocatedSize();
 			const uint32 KDTreeBytesSize = SearchIndex.KDTree.GetAllocatedSize();
+			const uint32 PCAValuesVectorToPoseIndexesBytesSize = SearchIndex.PCAValuesVectorToPoseIndexes.GetAllocatedSize();
+
 			const uint32 PoseMetadataBytesSize = SearchIndex.PoseMetadata.GetAllocatedSize();
 			const uint32 AssetsBytesSize = SearchIndex.Assets.GetAllocatedSize();
 			const uint32 OtherBytesSize = SearchIndex.PCAProjectionMatrix.GetAllocatedSize() + SearchIndex.Mean.GetAllocatedSize() + SearchIndex.WeightsSqrt.GetAllocatedSize();
-			const uint32 EstimatedDatabaseBytesSize = ValuesBytesSize + PCAValuesBytesSize + KDTreeBytesSize + PoseMetadataBytesSize + AssetsBytesSize + OtherBytesSize;
+			const uint32 EstimatedDatabaseBytesSize = ValuesBytesSize + PCAValuesBytesSize + KDTreeBytesSize + PCAValuesVectorToPoseIndexesBytesSize + PoseMetadataBytesSize + AssetsBytesSize + OtherBytesSize;
 				
 			ValuesSize = FText::AsMemory(ValuesBytesSize);
 			PCAValuesSize = FText::AsMemory(PCAValuesBytesSize);
diff --git a/Engine/Plugins/Experimental/Animation/PoseSearch/Source/Editor/Public/PoseSearchDatabaseEditorReflection.h b/Engine/Plugins/Experimental/Animation/PoseSearch/Source/Editor/Public/PoseSearchDatabaseEditorReflection.h
--- a/Engine/Plugins/Experimental/Animation/PoseSearch/Source/Editor/Public/PoseSearchDatabaseEditorReflection.h
+++ b/Engine/Plugins/Experimental/Animation/PoseSearch/Source/Editor/Public/PoseSearchDatabaseEditorReflection.h
@@ -198,10 +198,14 @@ public:
 	UPROPERTY(VisibleAnywhere, Category = "General Information")
 	FText SearchableTime;
 
-	// Amount of animation frames (database poses) pose feature data that has been pruned out because close to each other more than PosePruningSimilartyThreshod
+	// Amount of animation frames (database poses) of pose features data that has been pruned out because close to each other more than PosePruningSimilarityThreshold
 	UPROPERTY(VisibleAnywhere, Category = "General Information")
 	uint32 PrunedFrames;
 
+	// Amount of animation frames projected in PCA space (database poses) of pose PCA feature data (used to construct the kdtree) that has been pruned out because close to each other more than PCAValuesPruningSimilarityThreshold
+	UPROPERTY(VisibleAnywhere, Category = "General Information")
+	uint32 PrunedPCAFrames;
+
 	// Cardinality for the database config (how many floats per pose to store the pose features data)
 	UPROPERTY(VisibleAnywhere, Category = "General Information")
 	uint32 ConfigCardinality;
diff --git a/Engine/Plugins/Experimental/Animation/PoseSearch/Source/Runtime/Private/KDTree.cpp b/Engine/Plugins/Experimental/Animation/PoseSearch/Source/Runtime/Private/KDTree.cpp
--- a/Engine/Plugins/Experimental/Animation/PoseSearch/Source/Runtime/Private/KDTree.cpp
+++ b/Engine/Plugins/Experimental/Animation/PoseSearch/Source/Runtime/Private/KDTree.cpp
@@ -148,7 +148,24 @@ void FKDTree::Construct(int32 Count, int32 Dim, const float* Data, int32 MaxLeaf
 	new(this)FKDTree(Count, Dim, Data, MaxLeafSize);
 }
 
-bool FKDTree::FindNeighbors(KNNResultSet& Result, const float* Query) const
+bool FKDTree::FindNeighbors(FKNNResultSet& Result, const float* Query) const
+{
+#if UE_POSE_SEARCH_USE_NANOFLANN
+	QUICK_SCOPE_CYCLE_COUNTER(STAT_FKDTree_FindNeighbors);
+
+	check(Query && Impl->root_node);
+	const nanoflann::SearchParams SearchParams(
+		32,			// Ignored parameter (Kept for compatibility with the FLANN interface).
+		0.f,		// search for eps-approximate neighbours (default: 0)
+		false);		// only for radius search, require neighbours sorted by
+	return Impl->findNeighbors(Result, Query, SearchParams);
+#else
+	checkNoEntry(); // unimplemented
+	return false;
+#endif
+}
+
+bool FKDTree::FindNeighbors(FRadiusResultSet& Result, const float* Query) const
 {
 #if UE_POSE_SEARCH_USE_NANOFLANN
 	QUICK_SCOPE_CYCLE_COUNTER(STAT_FKDTree_FindNeighbors);
diff --git a/Engine/Plugins/Experimental/Animation/PoseSearch/Source/Runtime/Private/PoseSearchDatabase.cpp b/Engine/Plugins/Experimental/Animation/PoseSearch/Source/Runtime/Private/PoseSearchDatabase.cpp
--- a/Engine/Plugins/Experimental/Animation/PoseSearch/Source/Runtime/Private/PoseSearchDatabase.cpp
+++ b/Engine/Plugins/Experimental/Animation/PoseSearch/Source/Runtime/Private/PoseSearchDatabase.cpp
@@ -526,11 +526,12 @@ int32 UPoseSearchDatabase::GetNumberOfPrincipalComponents() const
 
 bool UPoseSearchDatabase::GetSkipSearchIfPossible() const
 {
+#if WITH_EDITOR
 	if (PoseSearchMode == EPoseSearchMode::PCAKDTree_Validate || PoseSearchMode == EPoseSearchMode::PCAKDTree_Compare)
 	{
 		return false;
 	}
-
+#endif // WITH_EDITOR
 	return bSkipSearchIfPossible;
 }
 
@@ -653,9 +654,13 @@ UE::PoseSearch::FSearchResult UPoseSearchDatabase::Search(UE::PoseSearch::FSearc
 	{
 		return Result;
 	}
-#endif
+#endif // WITH_EDITOR
 
-	if (PoseSearchMode == EPoseSearchMode::BruteForce || PoseSearchMode == EPoseSearchMode::PCAKDTree_Compare)
+	if (PoseSearchMode == EPoseSearchMode::BruteForce
+#if WITH_EDITOR
+		|| PoseSearchMode == EPoseSearchMode::PCAKDTree_Compare
+#endif // WITH_EDITOR
+		)
 	{
 		Result = SearchBruteForce(SearchContext);
 	}
@@ -672,7 +677,14 @@ UE::PoseSearch::FSearchResult UPoseSearchDatabase::Search(UE::PoseSearch::FSearc
 		Result.BruteForcePoseCost = BruteForcePoseCost;
 		if (PoseSearchMode == EPoseSearchMode::PCAKDTree_Compare)
 		{
-			check(Result.BruteForcePoseCost.GetTotalCost() <= Result.PoseCost.GetTotalCost());
+			const float BruteForceTotalCost = Result.BruteForcePoseCost.GetTotalCost();
+			const float PCAKDTreeTotalCost = Result.PoseCost.GetTotalCost();
+			check(BruteForceTotalCost <= PCAKDTreeTotalCost);
+
+			if (!FMath::IsNearlyEqual(BruteForceTotalCost, PCAKDTreeTotalCost))
+			{
+				UE_LOG(LogPoseSearch, Warning, TEXT("PCAKDTree_Compare %f (PCAKDTreeTotalCost %f, BruteForceTotalCost %f)"), PCAKDTreeTotalCost - BruteForceTotalCost, PCAKDTreeTotalCost, BruteForceTotalCost);
+			}
 		}
 #endif // UE_POSE_SEARCH_TRACE_ENABLED
 	}
@@ -804,42 +816,19 @@ UE::PoseSearch::FSearchResult UPoseSearchDatabase::SearchPCAKDTree(UE::PoseSearc
 	const int32 NumDimensions = Schema->SchemaCardinality;
 	const FSearchIndex& SearchIndex = GetSearchIndex();
 
-	const uint32 ClampedNumberOfPrincipalComponents = GetNumberOfPrincipalComponents();
-	const uint32 ClampedKDTreeQueryNumNeighbors = FMath::Clamp<uint32>(KDTreeQueryNumNeighbors, 1, SearchIndex.GetNumPoses());
-
-	//stack allocated temporaries
-	TArrayView<size_t> ResultIndexes((size_t*)FMemory_Alloca((ClampedKDTreeQueryNumNeighbors + 1) * sizeof(size_t)), ClampedKDTreeQueryNumNeighbors + 1);
-	TArrayView<float> ResultDistanceSqr((float*)FMemory_Alloca((ClampedKDTreeQueryNumNeighbors + 1) * sizeof(float)), ClampedKDTreeQueryNumNeighbors + 1);
-	TArrayView<float> ProjectedQueryValues((float*)FMemory_Alloca(ClampedNumberOfPrincipalComponents * sizeof(float)), ClampedNumberOfPrincipalComponents);
-
-#if DO_CHECK
-	// KDTree in PCA space search
-	if (PoseSearchMode == EPoseSearchMode::PCAKDTree_Validate)
-	{
-		// testing the KDTree is returning the proper searches for all the original points transformed in pca space
-		for (int32 PoseIdx = 0; PoseIdx < SearchIndex.GetNumPoses(); ++PoseIdx)
-		{
-			FKDTree::KNNResultSet ResultSet(ClampedKDTreeQueryNumNeighbors, ResultIndexes, ResultDistanceSqr);
-			SearchIndex.KDTree.FindNeighbors(ResultSet, SearchIndex.PCAProject(SearchIndex.GetPoseValues(PoseIdx), ProjectedQueryValues).GetData());
-
-			size_t ResultIndex = 0;
-			for (; ResultIndex < ResultSet.Num(); ++ResultIndex)
-			{
-				if (PoseIdx == ResultIndexes[ResultIndex])
-				{
-					check(ResultDistanceSqr[ResultIndex] < UE_KINDA_SMALL_NUMBER);
-					break;
-				}
-			}
-			check(ResultIndex < ResultSet.Num());
-		}
-	}
-#endif //DO_CHECK
-
 	// since any PoseCost calculated here is at least SearchIndex.MinCostAddend,
 	// there's no point in performing the search if CurrentBestTotalCost is already better than that
 	if (!GetSkipSearchIfPossible() || SearchContext.GetCurrentBestTotalCost() > SearchIndex.MinCostAddend)
 	{
+		const uint32 ClampedNumberOfPrincipalComponents = GetNumberOfPrincipalComponents();
+		const uint32 ClampedKDTreeQueryNumNeighbors = FMath::Clamp<uint32>(KDTreeQueryNumNeighbors, 1, SearchIndex.GetNumPoses());
+		const bool bArePCAValuesPruned = SearchIndex.PCAValuesVectorToPoseIndexes.Num() > 0;
+
+		//stack allocated temporaries
+		TArrayView<size_t> ResultIndexes((size_t*)FMemory_Alloca((ClampedKDTreeQueryNumNeighbors + 1) * sizeof(size_t)), ClampedKDTreeQueryNumNeighbors + 1);
+		TArrayView<float> ResultDistanceSqr((float*)FMemory_Alloca((ClampedKDTreeQueryNumNeighbors + 1) * sizeof(float)), ClampedKDTreeQueryNumNeighbors + 1);
+		TArrayView<float> ProjectedQueryValues((float*)FMemory_Alloca(ClampedNumberOfPrincipalComponents * sizeof(float)), ClampedNumberOfPrincipalComponents);
+	
 		TConstArrayView<float> QueryValues = SearchContext.GetOrBuildQuery(Schema).GetValues();
 
 		FNonSelectableIdx NonSelectableIdx;
@@ -848,17 +837,52 @@ UE::PoseSearch::FSearchResult UPoseSearchDatabase::SearchPCAKDTree(UE::PoseSearc
 			, QueryValues
 #endif // UE_POSE_SEARCH_TRACE_ENABLED
 		);
-		FKDTree::KNNResultSet ResultSet(ClampedKDTreeQueryNumNeighbors, ResultIndexes, ResultDistanceSqr, NonSelectableIdx);
+
+		// if bArePCAValuesPruned we filter out the NonSelectableIdx after kdtree search, otherwise during kdtree search
+		FKDTree::FKNNResultSet ResultSet(ClampedKDTreeQueryNumNeighbors, ResultIndexes, ResultDistanceSqr, bArePCAValuesPruned ? TConstArrayView<size_t>() : NonSelectableIdx);
 
 		check(QueryValues.Num() == NumDimensions);
 		// projecting QueryValues into the PCA space ProjectedQueryValues and query the KDTree
 		SearchIndex.KDTree.FindNeighbors(ResultSet, SearchIndex.PCAProject(QueryValues, ProjectedQueryValues).GetData());
 
 		// NonSelectableIdx are already filtered out inside the kdtree search
-		const FSearchFilters SearchFilters(Schema, TConstArrayView<size_t>(), SearchIndex.bAnyBlockTransition);
+		const FSearchFilters SearchFilters(Schema, bArePCAValuesPruned ? NonSelectableIdx : TConstArrayView<size_t>(), SearchIndex.bAnyBlockTransition);
 		
+		// are the PCAValues pruned out of duplicates (multiple poses are associated with the same PCAValuesVector)
+		if (bArePCAValuesPruned)
+		{
+			// @todo: reconstruction is not yet supported with pruned PCAValues
+			check(!SearchIndex.IsValuesEmpty());
+			
+			const int32 MaxNumEvaluatePoseKernelCalls = KDTreeQueryNumNeighborsWithDuplicates > 0 ? KDTreeQueryNumNeighborsWithDuplicates : INT32_MAX;
+
+			if (NumDimensions % 4 == 0)
+			{
+				int32 NumEvaluatePoseKernelCalls = 0;
+				for (size_t ResultIndex = 0; ResultIndex < ResultSet.Num(); ++ResultIndex)
+				{
+					const TConstArrayView<uint32> PoseIndexes = SearchIndex.PCAValuesVectorToPoseIndexes[ResultIndexes[ResultIndex]];
+					for (int Index = 0; Index < PoseIndexes.Num() && NumEvaluatePoseKernelCalls < MaxNumEvaluatePoseKernelCalls; ++Index, ++NumEvaluatePoseKernelCalls)
+					{
+						EvaluatePoseKernel<false, true>(Result, SearchIndex, QueryValues, TArrayView<float>(), PoseIndexes[Index], SearchFilters, SearchContext, this, true, ResultIndex);
+					}
+				}
+			}
+			else
+			{
+				int32 NumEvaluatePoseKernelCalls = 0;
+				for (size_t ResultIndex = 0; ResultIndex < ResultSet.Num(); ++ResultIndex)
+				{
+					const TConstArrayView<uint32> PoseIndexes = SearchIndex.PCAValuesVectorToPoseIndexes[ResultIndexes[ResultIndex]];
+					for (int Index = 0; Index < PoseIndexes.Num() && NumEvaluatePoseKernelCalls < MaxNumEvaluatePoseKernelCalls; ++Index, ++NumEvaluatePoseKernelCalls)
+					{
+						EvaluatePoseKernel<false, false>(Result, SearchIndex, QueryValues, TArrayView<float>(), PoseIndexes[Index], SearchFilters, SearchContext, this, true, ResultIndex);
+					}
+				}
+			}
+		}
 		// do we need to reconstruct pose values?
-		if (SearchIndex.IsValuesEmpty())
+		else if (SearchIndex.IsValuesEmpty())
 		{
 			// FMemory_Alloca is forced 16 bytes aligned
 			TArrayView<float> ReconstructedPoseValuesBuffer((float*)FMemory_Alloca(NumDimensions * sizeof(float)), NumDimensions);
diff --git a/Engine/Plugins/Experimental/Animation/PoseSearch/Source/Runtime/Private/PoseSearchDatabaseIndexingContext.cpp b/Engine/Plugins/Experimental/Animation/PoseSearch/Source/Runtime/Private/PoseSearchDatabaseIndexingContext.cpp
--- a/Engine/Plugins/Experimental/Animation/PoseSearch/Source/Runtime/Private/PoseSearchDatabaseIndexingContext.cpp
+++ b/Engine/Plugins/Experimental/Animation/PoseSearch/Source/Runtime/Private/PoseSearchDatabaseIndexingContext.cpp
@@ -129,7 +129,7 @@ bool FDatabaseIndexingContext::IndexDatabase(FSearchIndexBase& SearchIndexBase,
 	TotalPoses = 0;
 	for (int32 AssetIdx = 0; AssetIdx != SearchIndexBase.Assets.Num(); ++AssetIdx)
 	{
-		Indexers[AssetIdx].AssignWorkingData(TotalPoses, SearchIndexBase.EditValues(), SearchIndexBase.PoseMetadata);
+		Indexers[AssetIdx].AssignWorkingData(TotalPoses, SearchIndexBase.Values, SearchIndexBase.PoseMetadata);
 		TotalPoses += Indexers[AssetIdx].GetNumIndexedPoses();
 	}
 
diff --git a/Engine/Plugins/Experimental/Animation/PoseSearch/Source/Runtime/Private/PoseSearchDerivedData.cpp b/Engine/Plugins/Experimental/Animation/PoseSearch/Source/Runtime/Private/PoseSearchDerivedData.cpp
--- a/Engine/Plugins/Experimental/Animation/PoseSearch/Source/Runtime/Private/PoseSearchDerivedData.cpp
+++ b/Engine/Plugins/Experimental/Animation/PoseSearch/Source/Runtime/Private/PoseSearchDerivedData.cpp
@@ -144,7 +144,7 @@ private:
 			const int32 NumValuesVectors = SearchIndexBase.GetNumValuesVectors(Schema->SchemaCardinality);
 
 			// Map input buffer with NumValuesVectors as rows and NumDimensions	as cols
-			RowMajorMatrixMapConst PoseMatrixSourceMap(SearchIndexBase.GetValues().GetData(), NumValuesVectors, Schema->SchemaCardinality);
+			RowMajorMatrixMapConst PoseMatrixSourceMap(SearchIndexBase.Values.GetData(), NumValuesVectors, Schema->SchemaCardinality);
 
 			// Given the sub matrix for the features, find the average distance to the feature's centroid.
 			CenteredSubPoseMatrix.block(AccumulatedNumValuesVectors, 0, NumValuesVectors, Cardinality) = PoseMatrixSourceMap.block(0, Entry.Channel->GetChannelDataOffset(), NumValuesVectors, Cardinality);
@@ -450,11 +450,11 @@ static void PreprocessSearchIndexPCAData(FSearchIndex& SearchIndex, int32 NumDim
 		RowMajorMatrixMap MapPCAValues(SearchIndex.PCAValues.GetData(), NumPoses, NumberOfPrincipalComponents);
 
 		// calculating the mean
-		RowMajorVectorMap Mean(SearchIndex.Mean.GetData(), 1, NumDimensions);
-		Mean = WeightedValues.colwise().mean();
+		RowMajorVectorMap MapMean(SearchIndex.Mean.GetData(), 1, NumDimensions);
+		MapMean = WeightedValues.colwise().mean();
 
 		// use the mean to center the data points
-		const RowMajorMatrix CenteredValues = WeightedValues.rowwise() - Mean;
+		const RowMajorMatrix CenteredValues = WeightedValues.rowwise() - MapMean;
 
 		// estimating the covariance matrix (with dimensionality of NumDimensions, NumDimensions)
 		// formula: https://en.wikipedia.org/wiki/Covariance_matrix#Estimation
@@ -473,7 +473,7 @@ static void PreprocessSearchIndexPCAData(FSearchIndex& SearchIndex, int32 NumDim
 			const RowMajorMatrix ProjectedValues = CenteredValues * EigenVectors;
 			for (Eigen::Index RowIndex = 0; RowIndex < MapValues.rows(); ++RowIndex)
 			{
-				const RowMajorVector WeightedReconstructedPoint = ProjectedValues.row(RowIndex) * EigenVectors.transpose() + Mean;
+				const RowMajorVector WeightedReconstructedPoint = ProjectedValues.row(RowIndex) * EigenVectors.transpose() + MapMean;
 				const RowMajorVector ReconstructedPoint = WeightedReconstructedPoint.array() * ReciprocalWeightsSqrt.array();
 				const float Error = (ReconstructedPoint - MapValues.row(RowIndex)).squaredNorm();
 				check(Error < UE_KINDA_SMALL_NUMBER);
@@ -489,9 +489,9 @@ static void PreprocessSearchIndexPCAData(FSearchIndex& SearchIndex, int32 NumDim
 			Indexer.Push(DimensionIndex);
 		}
 		Indexer.Sort([&EigenValues](size_t a, size_t b)
-			{
-				return EigenValues[a] > EigenValues[b];
-			});
+		{
+			return EigenValues[a] > EigenValues[b];
+		});
 
 		// composing the PCA projection matrix with the PCANumComponents most significant EigenVectors
 		ColMajorMatrixMap PCAProjectionMatrix(SearchIndex.PCAProjectionMatrix.GetData(), NumDimensions, NumberOfPrincipalComponents);
@@ -514,7 +514,7 @@ static void PreprocessSearchIndexPCAData(FSearchIndex& SearchIndex, int32 NumDim
 			const RowMajorVector ReciprocalWeightsSqrt = MapWeightsSqrt.cwiseInverse();
 			for (Eigen::Index RowIndex = 0; RowIndex < MapValues.rows(); ++RowIndex)
 			{
-				const RowMajorVector WeightedReconstructedValues = MapPCAValues.row(RowIndex) * PCAProjectionMatrix.transpose() + Mean;
+				const RowMajorVector WeightedReconstructedValues = MapPCAValues.row(RowIndex) * PCAProjectionMatrix.transpose() + MapMean;
 				const RowMajorVector ReconstructedValues = WeightedReconstructedValues.array() * ReciprocalWeightsSqrt.array();
 				const float Error = (ReconstructedValues - MapValues.row(RowIndex)).squaredNorm();
 				check(Error < UE_KINDA_SMALL_NUMBER);
@@ -538,64 +538,69 @@ static void PreprocessSearchIndexPCAData(FSearchIndex& SearchIndex, int32 NumDim
 	}
 }
 
-static void PreprocessSearchIndexKDTree(FSearchIndex& SearchIndex, int32 NumDimensions, uint32 NumberOfPrincipalComponents, EPoseSearchMode PoseSearchMode, int32 KDTreeMaxLeafSize, int32 KDTreeQueryNumNeighbors)
+static void PreprocessSearchIndexKDTree(FSearchIndex& SearchIndex, const UPoseSearchDatabase* Database)
 {
+	const int32 NumDimensions = Database->Schema->SchemaCardinality;
+	const EPoseSearchMode PoseSearchMode = Database->PoseSearchMode;
+
 	SearchIndex.KDTree.Reset();
 	if (PoseSearchMode != EPoseSearchMode::BruteForce && NumDimensions > 0)
 	{
-		const int32 NumPoses = SearchIndex.GetNumPoses();
-		SearchIndex.KDTree.Construct(NumPoses, NumberOfPrincipalComponents, SearchIndex.PCAValues.GetData(), KDTreeMaxLeafSize);
+		const uint32 NumberOfPrincipalComponents = Database->GetNumberOfPrincipalComponents();
+		const int32 KDTreeMaxLeafSize = Database->KDTreeMaxLeafSize;
+
+		check(SearchIndex.PCAValues.Num() % NumberOfPrincipalComponents == 0);
 
-		if (PoseSearchMode == EPoseSearchMode::PCAKDTree_Validate)
+		const int32 NumPCAValuesVectors = SearchIndex.PCAValues.Num() / NumberOfPrincipalComponents;
+		SearchIndex.KDTree.Construct(NumPCAValuesVectors, NumberOfPrincipalComponents, SearchIndex.PCAValues.GetData(), KDTreeMaxLeafSize);
+
+		// testing the KDTree is returning the proper searches for all the points in pca space
+		const int32 KDTreeQueryNumNeighbors = Database->KDTreeQueryNumNeighbors;
+
+		TArray<size_t> ResultIndexes;
+		TArray<float> ResultDistanceSqr;
+		ResultIndexes.SetNum(NumPCAValuesVectors + 1);
+		ResultDistanceSqr.SetNum(NumPCAValuesVectors + 1);
+
+		size_t MaxNumNeighborToFindAPoint = 0;
+		for (size_t PointIndex = 0; PointIndex < NumPCAValuesVectors; ++PointIndex)
 		{
-			// testing the KDTree is returning the proper searches for all the points in pca space
-			int32 NumberOfFailingPoints = 0;
-			for (size_t PointIndex = 0; PointIndex < NumPoses; ++PointIndex)
-			{
-				TArray<size_t> ResultIndexes;
-				TArray<float> ResultDistanceSqr;
-				ResultIndexes.SetNum(KDTreeQueryNumNeighbors + 1);
-				ResultDistanceSqr.SetNum(KDTreeQueryNumNeighbors + 1);
-				FKDTree::KNNResultSet ResultSet(KDTreeQueryNumNeighbors, ResultIndexes, ResultDistanceSqr);
-				SearchIndex.KDTree.FindNeighbors(ResultSet, &SearchIndex.PCAValues[PointIndex * NumberOfPrincipalComponents]);
+			// searching the kdtree for PointIndex
+			FKDTree::FRadiusResultSet ResultSet(UE_SMALL_NUMBER, NumPCAValuesVectors, ResultIndexes, ResultDistanceSqr);
+			SearchIndex.KDTree.FindNeighbors(ResultSet, &SearchIndex.PCAValues[PointIndex * NumberOfPrincipalComponents]);
 
-				size_t ResultIndex = 0;
-				for (; ResultIndex < ResultSet.Num(); ++ResultIndex)
-				{
-					if (PointIndex == ResultIndexes[ResultIndex])
-					{
-						check(ResultDistanceSqr[ResultIndex] < UE_KINDA_SMALL_NUMBER);
-						break;
-					}
-				}
-				if (ResultIndex == ResultSet.Num())
+			bool bFound = false;
+			for (size_t ResultIndex = 0; ResultIndex < ResultSet.Num(); ++ResultIndex)
+			{
+				if (PointIndex == ResultIndexes[ResultIndex])
 				{
-					++NumberOfFailingPoints;
+					// PointIndex is the ResultIndex-th candidates out of the kdtree. if ResultIndex-th is greater than KDTreeQueryNumNeighbors,
+					// we wouldn't have found it in a runtime search, so we log the errot (later on only once, with the worst case scenario)
+					check(ResultDistanceSqr[ResultIndex] < UE_KINDA_SMALL_NUMBER);
+					MaxNumNeighborToFindAPoint = FMath::Max(MaxNumNeighborToFindAPoint, ResultIndex);
+					bFound = true;
+					break;
 				}
 			}
+			check(bFound);
+		}
 
-			check(NumberOfFailingPoints == 0);
+		if (MaxNumNeighborToFindAPoint >= KDTreeQueryNumNeighbors)
+		{
+			UE_LOG(LogPoseSearch, Warning, TEXT("Not enough 'KDTreeQueryNumNeighbors' (%d) for database '%s'. Pose values projected in PCA space have too many duplicates, so try to prune duplicates by tuning 'PCAValuesPruningSimilarityThreshold' or increase 'KDTreeQueryNumNeighbors' at least to %d"), KDTreeQueryNumNeighbors, *Database->GetName(), MaxNumNeighborToFindAPoint);
+		}
 
+		// if bArePCAValuesPruned PointIndex is the index of the point in the kdtree, NOT necessary the pose index, so doing the PCAProject would lead to the wrong data
+		const bool bArePCAValuesPruned = SearchIndex.PCAValuesVectorToPoseIndexes.Num() > 0;
+		if (!bArePCAValuesPruned && PoseSearchMode == EPoseSearchMode::PCAKDTree_Validate)
+		{
 			// testing the KDTree is returning the proper searches for all the original points transformed in pca space
-			NumberOfFailingPoints = 0;
-			for (size_t PointIndex = 0; PointIndex < NumPoses; ++PointIndex)
+			int32 NumberOfFailingPoints = 0;
+			TArrayView<float> ProjectedValues((float*)FMemory_Alloca(NumberOfPrincipalComponents * sizeof(float)), NumberOfPrincipalComponents);
+			for (size_t PointIndex = 0; PointIndex < NumPCAValuesVectors; ++PointIndex)
 			{
-				TArray<size_t> ResultIndexes;
-				TArray<float> ResultDistanceSqr;
-				ResultIndexes.SetNum(KDTreeQueryNumNeighbors + 1);
-				ResultDistanceSqr.SetNum(KDTreeQueryNumNeighbors + 1);
-				FKDTree::KNNResultSet ResultSet(KDTreeQueryNumNeighbors, ResultIndexes, ResultDistanceSqr);
-
-				const RowMajorVectorMapConst MapValues(SearchIndex.GetPoseValuesBase(PointIndex, NumDimensions).GetData(), 1, NumDimensions);
-				const RowMajorVectorMapConst MapWeightsSqrt(SearchIndex.WeightsSqrt.GetData(), 1, NumDimensions);
-				const RowMajorVectorMapConst Mean(SearchIndex.Mean.GetData(), 1, NumDimensions);
-				const ColMajorMatrixMapConst PCAProjectionMatrix(SearchIndex.PCAProjectionMatrix.GetData(), NumDimensions, NumberOfPrincipalComponents);
-
-				const RowMajorMatrix WeightedValues = MapValues.array() * MapWeightsSqrt.array();
-				const RowMajorMatrix CenteredValues = WeightedValues - Mean;
-				const RowMajorVector ProjectedValues = CenteredValues * PCAProjectionMatrix;
-
-				SearchIndex.KDTree.FindNeighbors(ResultSet, ProjectedValues.data());
+				FKDTree::FKNNResultSet ResultSet(KDTreeQueryNumNeighbors, ResultIndexes, ResultDistanceSqr);
+				SearchIndex.KDTree.FindNeighbors(ResultSet, SearchIndex.PCAProject(SearchIndex.GetPoseValuesBase(PointIndex, NumDimensions), ProjectedValues).GetData());
 
 				size_t ResultIndex = 0;
 				for (; ResultIndex < ResultSet.Num(); ++ResultIndex)
@@ -1010,7 +1015,15 @@ void FPoseSearchDatabaseAsyncCacheTask::OnGetComplete(UE::DerivedData::FCacheGet
 					return;
 				}
 
-				PreprocessSearchIndexKDTree(SearchIndex, Database->Schema->SchemaCardinality, Database->GetNumberOfPrincipalComponents(), Database->PoseSearchMode, Database->KDTreeMaxLeafSize, Database->KDTreeQueryNumNeighbors);
+				SearchIndex.PruneDuplicatePCAValues(Database->PCAValuesPruningSimilarityThreshold, Database->GetNumberOfPrincipalComponents());
+				if (Owner.IsCanceled())
+				{
+					UE_LOG(LogPoseSearch, Log, TEXT("%s - %s BuildIndex Cancelled"), *LexToString(FullIndexKey.Hash), *Database->GetName());
+					SearchIndex.Reset();
+					return;
+				}
+
+				PreprocessSearchIndexKDTree(SearchIndex, Database.Get());
 				if (Owner.IsCanceled())
 				{
 					UE_LOG(LogPoseSearch, Log, TEXT("%s - %s BuildIndex Cancelled"), *LexToString(FullIndexKey.Hash), *Database->GetName());
@@ -1166,8 +1179,6 @@ void FAsyncPoseSearchDatabasesManagement::Tick(float DeltaTime)
 		{
 			Tasks[TaskIndex]->Update(Mutex);
 		}
-			
-		// @todo: check key validity every few ticks, or perhaps delete unused for a long time Tasks
 	}
 }
 
diff --git a/Engine/Plugins/Experimental/Animation/PoseSearch/Source/Runtime/Private/PoseSearchDerivedDataKey.cpp b/Engine/Plugins/Experimental/Animation/PoseSearch/Source/Runtime/Private/PoseSearchDerivedDataKey.cpp
--- a/Engine/Plugins/Experimental/Animation/PoseSearch/Source/Runtime/Private/PoseSearchDerivedDataKey.cpp
+++ b/Engine/Plugins/Experimental/Animation/PoseSearch/Source/Runtime/Private/PoseSearchDerivedDataKey.cpp
@@ -41,7 +41,7 @@ FKeyBuilder::FKeyBuilder(const UObject* Object, bool bUseDataVer)
 	if (bUseDataVer)
 	{
 		// used to invalidate the key without having to change POSESEARCHDB_DERIVEDDATA_VER all the times
-		int32 POSESEARCHDB_DERIVEDDATA_VER_SMALL = 167;
+		int32 POSESEARCHDB_DERIVEDDATA_VER_SMALL = 173;
 		FGuid VersionGuid = FDevSystemGuids::GetSystemGuid(FDevSystemGuids::Get().POSESEARCHDB_DERIVEDDATA_VER);
 
 		*this << VersionGuid;
diff --git a/Engine/Plugins/Experimental/Animation/PoseSearch/Source/Runtime/Private/PoseSearchIndex.cpp b/Engine/Plugins/Experimental/Animation/PoseSearch/Source/Runtime/Private/PoseSearchIndex.cpp
--- a/Engine/Plugins/Experimental/Animation/PoseSearch/Source/Runtime/Private/PoseSearchIndex.cpp
+++ b/Engine/Plugins/Experimental/Animation/PoseSearch/Source/Runtime/Private/PoseSearchIndex.cpp
@@ -85,6 +85,223 @@ float CompareFeatureVectors(TConstArrayView<float> A, TConstArrayView<float> B)
 	return (VA - VB).square().sum();
 }
 
+// pruning utils
+struct FPosePair
+{
+	int32 PoseIdxA = 0;
+	int32 PoseIdxB = 0;
+};
+struct FPosePairSimilarity : public FPosePair
+{
+	float Similarity = 0.f;
+};
+
+static bool CalculateSimilaritiesKDTree(TArray<FPosePairSimilarity>& PosePairSimilarities, float SimilarityThreshold, 
+	int32 DataCardinality, int32 NumPoses, const TAlignedArray<float>& Values, TFunctionRef<TConstArrayView<float>(int32, int32)> GetValuesVector)
+{
+	PosePairSimilarities.Reserve(1024 * 64);
+
+	
+
+	return !PosePairSimilarities.IsEmpty();
+}
+
+static bool CalculateSimilarities(TArray<FPosePairSimilarity>& PosePairSimilarities, float SimilarityThreshold, 
+	int32 DataCardinality, int32 NumPoses, const TAlignedArray<float>& Values,
+	TFunctionRef<TConstArrayView<float>(int32, int32)> GetValuesVector)
+{
+	enum EEvalMode
+	{
+		EUseKDTreeEvaluation,
+		EParalleEvaluation,
+		ESerialEvaluation,
+	};
+
+	static EEvalMode EvalMode = EEvalMode::EUseKDTreeEvaluation;
+
+	PosePairSimilarities.Reserve(1024 * 64);
+
+	if (EvalMode == EEvalMode::EUseKDTreeEvaluation)
+	{
+		check(Values.Num() == NumPoses * DataCardinality);
+		FKDTree KDTree(NumPoses, DataCardinality, Values.GetData());
+
+		TArray<size_t> ResultIndexes;
+		TArray<float> ResultDistanceSqr;
+		ResultIndexes.SetNum(NumPoses + 1);
+		ResultDistanceSqr.SetNum(NumPoses + 1);
+
+		for (int32 PoseIdx = 0; PoseIdx < NumPoses; ++PoseIdx)
+		{
+			TConstArrayView<float> ValuesA = GetValuesVector(PoseIdx, DataCardinality);
+
+			// searching for duplicates within a radius of SimilarityThreshold
+			FKDTree::FRadiusResultSet ResultSet(SimilarityThreshold, NumPoses, ResultIndexes, ResultDistanceSqr);
+			KDTree.FindNeighbors(ResultSet, ValuesA.GetData());
+
+			for (size_t ResultIndex = 0; ResultIndex < ResultSet.Num(); ++ResultIndex)
+			{
+				const int32 ResultPoseIdx = ResultIndexes[ResultIndex];
+				if (PoseIdx != ResultPoseIdx)
+				{
+					FPosePairSimilarity PosePair;
+					PosePair.PoseIdxA = PoseIdx;
+					PosePair.PoseIdxB = ResultPoseIdx;
+					PosePair.Similarity = ResultDistanceSqr[ResultIndex];
+					PosePairSimilarities.Emplace(PosePair);
+				}
+			}
+		}
+	}
+	// calculating a sparse proximity matrix with Similarity up to SimilarityThreshold (to perform something similar to hierarchical clustering)
+	else if (EvalMode == EEvalMode::EParalleEvaluation)
+	{
+		// @todo: we should use a smarter approach, knowing we need NumPoses * (NumPoses - 1) / 2 comparisons, and reconstruct the pose indexes from the comparison index...
+
+		// doing it in batches of items to avoid using too much memory
+		constexpr int32 BatchSize = 1024 * 1024;
+
+		TArray<FPosePair> PosePairs;
+		PosePairs.Reserve(BatchSize);
+		FCriticalSection Mutex;
+
+		auto EvaluatePosePairs = [GetValuesVector, DataCardinality, SimilarityThreshold, &PosePairs, &PosePairSimilarities, &Mutex]()
+		{
+			ParallelFor(PosePairs.Num(), [GetValuesVector, DataCardinality, SimilarityThreshold, &PosePairs, &PosePairSimilarities, &Mutex](int32 ComparisonIndex)
+			{
+				const int32 PoseIdxA = PosePairs[ComparisonIndex].PoseIdxA;
+				const int32 PoseIdxB = PosePairs[ComparisonIndex].PoseIdxB;
+				TConstArrayView<float> ValuesA = GetValuesVector(PoseIdxA, DataCardinality);
+				TConstArrayView<float> ValuesB = GetValuesVector(PoseIdxB, DataCardinality);
+				const float Similarity = CompareFeatureVectors(ValuesA, ValuesB);
+				if (Similarity < SimilarityThreshold)
+				{
+					// since this condition doesn't happen often, we're not gonna spend too much time in mutex contention
+					FScopeLock Lock(&Mutex);
+					FPosePairSimilarity PosePairSimilarity;
+					PosePairSimilarity.PoseIdxA = PoseIdxA;
+					PosePairSimilarity.PoseIdxB = PoseIdxB;
+					PosePairSimilarity.Similarity = Similarity;
+					PosePairSimilarities.Emplace(PosePairSimilarity);
+				}
+			});
+		};
+
+		FPosePair PosePair;
+		for (PosePair.PoseIdxA = 0; PosePair.PoseIdxA < NumPoses; ++PosePair.PoseIdxA)
+		{
+			for (PosePair.PoseIdxB = PosePair.PoseIdxA + 1; PosePair.PoseIdxB < NumPoses; ++PosePair.PoseIdxB)
+			{
+				PosePairs.Add(PosePair);
+				if (PosePairs.Num() == BatchSize)
+				{
+					EvaluatePosePairs();
+					PosePairs.Reset();
+				}
+			}
+		}
+
+		EvaluatePosePairs();
+	}
+	else if (EvalMode == EEvalMode::ESerialEvaluation)
+	{
+		FPosePairSimilarity PosePair;
+		for (PosePair.PoseIdxA = 0; PosePair.PoseIdxA < NumPoses; ++PosePair.PoseIdxA)
+		{
+			TConstArrayView<float> ValuesA = GetValuesVector(PosePair.PoseIdxA, DataCardinality);
+			for (PosePair.PoseIdxB = PosePair.PoseIdxA + 1; PosePair.PoseIdxB < NumPoses; ++PosePair.PoseIdxB)
+			{
+				TConstArrayView<float> ValuesB = GetValuesVector(PosePair.PoseIdxB, DataCardinality);
+				PosePair.Similarity = CompareFeatureVectors(ValuesA, ValuesB);
+				if (PosePair.Similarity < SimilarityThreshold)
+				{
+					PosePairSimilarities.Emplace(PosePair);
+				}
+			}
+		}
+	}
+	else
+	{
+		checkNoEntry();
+	}
+
+	if (!PosePairSimilarities.IsEmpty())
+	{
+		PosePairSimilarities.Sort([](const FPosePairSimilarity& A, const FPosePairSimilarity& B)
+		{
+			return A.Similarity < B.Similarity;
+		});
+		return true;
+	}
+	return false;
+}
+
+static bool PruneValues(int32 DataCardinality, int32 NumPoses, const TArray<FPosePairSimilarity>& PosePairSimilarities, TAlignedArray<float>& Values,
+	TFunctionRef<uint32(int32)> GetValueOffset, TFunctionRef<void(int32, uint32)> SetValueOffset)
+{
+	// mapping between the one value offset and all the poses sharing it
+	TMap<uint32, TArray<int32>> ValueOffsetToPoses;
+	for (int32 PoseIdx = 0; PoseIdx < NumPoses; ++PoseIdx)
+	{
+		const uint32 ValueOffset = GetValueOffset(PoseIdx);
+		// FindOrAdd to support the eventuality of having multiple poses already sharing the same value offset
+		ValueOffsetToPoses.FindOrAdd(ValueOffset).Add(PoseIdx);
+	}
+
+	// at this point ValueOffsetToPoses is fully populated with all the possible value offset, and since we're not adding, but eventually removing entries we can just use the [] operator
+	uint32 ValueOffsetLast = Values.Num() - DataCardinality;
+	for (int32 PosePairSimilarityIdx = 0; PosePairSimilarityIdx < PosePairSimilarities.Num(); ++PosePairSimilarityIdx)
+	{
+		const FPosePairSimilarity& PosePairSimilarity = PosePairSimilarities[PosePairSimilarityIdx];
+		const uint32 ValueOffsetA = GetValueOffset(PosePairSimilarity.PoseIdxA);
+		const uint32 ValueOffsetB = GetValueOffset(PosePairSimilarity.PoseIdxB);
+				
+		// if the two poses don't point already to the same value offset, we can remove one of them
+		if (ValueOffsetA != ValueOffsetB)
+		{
+			// transferring all the poses associated to ValueOffsetB to ValueOffsetA
+			TArray<int32>& PosesAtValueOffsetA = ValueOffsetToPoses[ValueOffsetA];
+			TArray<int32>& PosesAtValueOffsetB = ValueOffsetToPoses[ValueOffsetB];
+
+			for (int32 PoseAtValueOffsetB : PosesAtValueOffsetB)
+			{
+				SetValueOffset(PoseAtValueOffsetB, ValueOffsetA);
+				PosesAtValueOffsetA.Add(PoseAtValueOffsetB);
+			}
+
+			// moving the ValueOffsetLast values into the location ValueOffsetB, that we just free up
+			if (ValueOffsetB != ValueOffsetLast)
+			{
+				FMemory::Memcpy(&Values[ValueOffsetB], &Values[ValueOffsetLast], DataCardinality * sizeof(float));
+				TArray<int32>& PosesAtValueOffsetLast = ValueOffsetToPoses[ValueOffsetLast];
+						
+				for (int32 PoseAtValueOffsetLast : PosesAtValueOffsetLast)
+				{
+					SetValueOffset(PoseAtValueOffsetLast, ValueOffsetB);
+				}
+
+				PosesAtValueOffsetB = PosesAtValueOffsetLast;
+				PosesAtValueOffsetLast.Reset();
+			}
+			else
+			{
+				PosesAtValueOffsetB.Reset();
+			}
+
+			ValueOffsetLast -= DataCardinality;
+		}
+	}
+
+	if (ValueOffsetLast + DataCardinality != Values.Num())
+	{
+		// resizing the Values array  
+		Values.SetNum(ValueOffsetLast + DataCardinality);
+		return true;
+	}
+
+	return false;
+}
+
 //////////////////////////////////////////////////////////////////////////
 // FPoseMetadata
 FArchive& operator<<(FArchive& Ar, FPoseMetadata& Metadata)
@@ -162,158 +379,16 @@ void FSearchIndexBase::Reset()
 
 void FSearchIndexBase::PruneDuplicateValues(float SimilarityThreshold, int32 DataCardinality)
 {
-	constexpr bool bParalleEvaluation = true;
-
 	const int32 NumPoses = GetNumPoses();
 	if (SimilarityThreshold > 0.f && NumPoses >= 2)
 	{
-		struct FPosePair
-		{
-			int32 PoseIdxA = 0;
-			int32 PoseIdxB = 0;
-		};
-		struct FPosePairSimilarity : public FPosePair
-		{
-			float Similarity = 0.f;
-		};
 		TArray<FPosePairSimilarity> PosePairSimilarities;
-		PosePairSimilarities.Reserve(1024 * 64);
-
-		// calculating a sparse proximity matrix with Similarity up to SimilarityThreshold (to perform something similar to hierarchical clustering)
-		if (bParalleEvaluation)
-		{
-			// @todo: we should use a smarter approach, knowing we need NumPoses * (NumPoses - 1) / 2 comparisons, and reconstruct the pose indexes from the comparison index...
-			
-			// doing it in batches of items to avoid using too much memory
-			constexpr int32 BatchSize = 1024 * 1024;
-
-			TArray<FPosePair> PosePairs;
-			PosePairs.Reserve(BatchSize);
-			FPosePair PosePair;
-			FCriticalSection Mutex;
-
-			auto EvaluatePosePairs = [this, DataCardinality, SimilarityThreshold, &PosePairs, &PosePairSimilarities, &Mutex]()
-			{
-				ParallelFor(PosePairs.Num(), [this, DataCardinality, SimilarityThreshold, &PosePairs, &PosePairSimilarities, &Mutex](int32 ComparisonIndex)
-				{
-					const int PoseIdxA = PosePairs[ComparisonIndex].PoseIdxA;
-					const int PoseIdxB = PosePairs[ComparisonIndex].PoseIdxB;
-					TConstArrayView<float> ValuesA = GetPoseValuesBase(PoseIdxA, DataCardinality);
-					TConstArrayView<float> ValuesB = GetPoseValuesBase(PoseIdxB, DataCardinality);
-					const float Similarity = CompareFeatureVectors(ValuesA, ValuesB);
-					if (Similarity < SimilarityThreshold)
-					{
-						// since this condition doesn't happen often, we're not gonna spend too much time in mutex contention
-						FScopeLock Lock(&Mutex);
-						FPosePairSimilarity PosePairSimilarity;
-						PosePairSimilarity.PoseIdxA = PoseIdxA;
-						PosePairSimilarity.PoseIdxB = PoseIdxB;
-						PosePairSimilarity.Similarity = Similarity;
-						PosePairSimilarities.Emplace(PosePairSimilarity);
-					}
-				});
-			};
-
-			for (PosePair.PoseIdxA = 0; PosePair.PoseIdxA < NumPoses; ++PosePair.PoseIdxA)
-			{
-				for (PosePair.PoseIdxB = PosePair.PoseIdxA + 1; PosePair.PoseIdxB < NumPoses; ++PosePair.PoseIdxB)
-				{
-					PosePairs.Add(PosePair);
-					if (PosePairs.Num() == BatchSize)
-					{
-						EvaluatePosePairs();
-
-						PosePairs.Reset();
-					}
-				}
-			}
-
-			EvaluatePosePairs();
-		}
-		else
-		{
-			for (int32 PoseIdxA = 0; PoseIdxA < NumPoses; ++PoseIdxA)
-			{
-				TConstArrayView<float> ValuesA = GetPoseValuesBase(PoseIdxA, DataCardinality);
-				for (int32 PoseIdxB = PoseIdxA + 1; PoseIdxB < NumPoses; ++PoseIdxB)
-				{
-					TConstArrayView<float> ValuesB = GetPoseValuesBase(PoseIdxB, DataCardinality);
-					const float Similarity = CompareFeatureVectors(ValuesA, ValuesB);
-					if (Similarity < SimilarityThreshold)
-					{
-						FPosePairSimilarity PosePairSimilarity;
-						PosePairSimilarity.PoseIdxA = PoseIdxA;
-						PosePairSimilarity.PoseIdxB = PoseIdxB;
-						PosePairSimilarity.Similarity = Similarity;
-						PosePairSimilarities.Emplace(PosePairSimilarity);
-					}
-				}
-			}
-		}
-
-		if (!PosePairSimilarities.IsEmpty())
+		if (CalculateSimilarities(PosePairSimilarities, SimilarityThreshold, DataCardinality, NumPoses, Values,
+			[this](int32 PoseIdx, int32 DataCardinality) { return GetPoseValuesBase(PoseIdx, DataCardinality); }))
 		{
-			PosePairSimilarities.Sort([](const FPosePairSimilarity& A, const FPosePairSimilarity& B)
-			{
-				return A.Similarity < B.Similarity;
-			});
-
-			// mapping between the one value offset and all the poses sharing it
-			TMap<uint32, TArray<int32>> ValueOffsetToPoses;
-			for (int32 PoseIdx = 0; PoseIdx < NumPoses; ++PoseIdx)
-			{
-				const uint32 ValueOffset = PoseMetadata[PoseIdx].GetValueOffset();
-				// FindOrAdd to support the eventuality of having multiple poses already sharing the same value offset
-				ValueOffsetToPoses.FindOrAdd(ValueOffset).Add(PoseIdx);
-			}
-
-			// at this point ValueOffsetToPoses is fully populated with all the possible value offset, and since we're not adding, but eventually removing entries we can just use the [] operator
-			int32 LastValidPoseIdx = NumPoses - 1;
-			for (int32 PosePairSimilarityIdx = 0; PosePairSimilarityIdx < PosePairSimilarities.Num(); ++PosePairSimilarityIdx)
-			{
-				const FPosePairSimilarity& PosePairSimilarity = PosePairSimilarities[PosePairSimilarityIdx];
-				const uint32 ValueOffsetA = PoseMetadata[PosePairSimilarity.PoseIdxA].GetValueOffset();
-				const uint32 ValueOffsetB = PoseMetadata[PosePairSimilarity.PoseIdxB].GetValueOffset();
-				
-				// if the two poses don't point already to the same value offset, we can remove one of them
-				if (ValueOffsetA != ValueOffsetB)
-				{
-					// transferring all the poses associated to ValueOffsetB to ValueOffsetA
-					TArray<int32>& PosesAtValueOffsetA = ValueOffsetToPoses[ValueOffsetA];
-					TArray<int32>& PosesAtValueOffsetB = ValueOffsetToPoses[ValueOffsetB];
-
-					for (int32 PoseAtValueOffsetB : PosesAtValueOffsetB)
-					{
-						PoseMetadata[PoseAtValueOffsetB].SetValueOffset(ValueOffsetA);
-						PosesAtValueOffsetA.Add(PoseAtValueOffsetB);
-					}
-
-					// moving the LastValidPoseIdx values into the location ValueOffsetB, that we just free up
-					const uint32 ValueOffsetLast = LastValidPoseIdx * DataCardinality;
-					if (ValueOffsetB != ValueOffsetLast)
-					{
-						FMemory::Memcpy(&Values[ValueOffsetB], &Values[ValueOffsetLast], DataCardinality * sizeof(float));
-						TArray<int32>& PosesAtValueOffsetLast = ValueOffsetToPoses[ValueOffsetLast];
-						
-						for (int32 PoseAtValueOffsetLast : PosesAtValueOffsetLast)
-						{
-							PoseMetadata[PoseAtValueOffsetLast].SetValueOffset(ValueOffsetB);
-						}
-
-						PosesAtValueOffsetB = PosesAtValueOffsetLast;
-						PosesAtValueOffsetLast.Reset();
-					}
-					else
-					{
-						PosesAtValueOffsetB.Reset();
-					}
-
-					--LastValidPoseIdx;
-				}
-			}
-
-			// resizing the Values array  
-			Values.SetNum((LastValidPoseIdx + 1) * DataCardinality);
+			PruneValues(DataCardinality, NumPoses, PosePairSimilarities, Values,
+			[this](int32 PoseIdx) {	return PoseMetadata[PoseIdx].GetValueOffset(); },
+			[this](int32 PoseIdx, uint32 ValueOffset) {	PoseMetadata[PoseIdx].SetValueOffset(ValueOffset); });
 		}
 	}
 }
@@ -344,6 +419,7 @@ FSearchIndex::FSearchIndex(const FSearchIndex& Other)
 	: FSearchIndexBase(Other)
 	, WeightsSqrt(Other.WeightsSqrt)
 	, PCAValues(Other.PCAValues)
+	, PCAValuesVectorToPoseIndexes(Other.PCAValuesVectorToPoseIndexes)
 	, PCAProjectionMatrix(Other.PCAProjectionMatrix)
 	, Mean(Other.Mean)
 	, KDTree(Other.KDTree)
@@ -378,13 +454,20 @@ TConstArrayView<float> FSearchIndex::GetReconstructedPoseValues(int32 PoseIdx, T
 {
 	QUICK_SCOPE_CYCLE_COUNTER(STAT_PoseSearch_PCAReconstruct);
 
+	// @todo: reconstruction is not yet supported with pruned PCAValues
+	check(PCAValuesVectorToPoseIndexes.Num() == 0);
+
 	const int32 NumDimensions = WeightsSqrt.Num();
 	const int32 NumPoses = GetNumPoses();
-	check(PoseIdx >= 0 && PoseIdx < NumPoses&& NumDimensions > 0);
+	check(PoseIdx >= 0 && PoseIdx < NumPoses && NumDimensions > 0);
 	check(BufferUsedForReconstruction.Num() == NumDimensions);
 
 	const int32 NumberOfPrincipalComponents = PCAValues.Num() / NumPoses;
+	
+	// @todo: if one of these checks trigger, most likely PCAValuesPruningSimilarityThreshold > 0.f and we pruned some PCAValues.
+	// currently GetReconstructedPoseValues is not supported with PCAValues pruning
 	check(NumPoses * NumberOfPrincipalComponents == PCAValues.Num());
+	check(PCAProjectionMatrix.Num() == NumDimensions * NumberOfPrincipalComponents);
 
 	const RowMajorVectorMapConst MapWeightsSqrt(WeightsSqrt.GetData(), 1, NumDimensions);
 	const ColMajorMatrixMapConst MapPCAProjectionMatrix(PCAProjectionMatrix.GetData(), NumDimensions, NumberOfPrincipalComponents);
@@ -407,6 +490,7 @@ TConstArrayView<float> FSearchIndex::PCAProject(TConstArrayView<float> PoseValue
 	const int32 NumDimensions = WeightsSqrt.Num();
 	const int32 NumberOfPrincipalComponents = PCAProjectionMatrix.Num() / NumDimensions;
 
+	check(PoseValues.Num() == NumDimensions);
 	check(PCAProjectionMatrix.Num() > 0 && PCAProjectionMatrix.Num() % NumDimensions == 0);
 	check(BufferUsedForProjection.Num() == NumberOfPrincipalComponents);
 
@@ -427,6 +511,78 @@ TConstArrayView<float> FSearchIndex::PCAProject(TConstArrayView<float> PoseValue
 	return BufferUsedForProjection;
 }
 
+void FSearchIndex::PruneDuplicatePCAValues(float SimilarityThreshold, int32 NumberOfPrincipalComponents)
+{
+	PCAValuesVectorToPoseIndexes = FSparsePoseMultiMap<uint32>();
+
+	const uint32 NumPoses = GetNumPoses();
+	if (SimilarityThreshold > 0.f && NumPoses >= 2 && NumberOfPrincipalComponents > 0)
+	{
+		check(PCAValues.Num() % NumberOfPrincipalComponents == 0);
+		const int32 NumPCAValuesVectors = PCAValues.Num() / NumberOfPrincipalComponents;
+		// so far we support only pruning an original PCAValues set, where there's a 1:1 mapping between PCAValuesVectors and Poses
+		check(NumPCAValuesVectors == NumPoses);
+
+		TArray<uint32> PoseToPCAValueOffset;
+		PoseToPCAValueOffset.AddUninitialized(NumPoses);
+		for (uint32 PoseIdx = 0; PoseIdx < NumPoses; ++PoseIdx)
+		{
+			PoseToPCAValueOffset[PoseIdx] = PoseIdx * NumberOfPrincipalComponents;
+		}
+
+		TArray<FPosePairSimilarity> PosePairSimilarities;
+		if (CalculateSimilarities(PosePairSimilarities, SimilarityThreshold, NumberOfPrincipalComponents, NumPoses, PCAValues,
+			[this, &PoseToPCAValueOffset](int32 PoseIdx, int32 NumberOfPrincipalComponents) { return MakeArrayView(&PCAValues[PoseToPCAValueOffset[PoseIdx]], NumberOfPrincipalComponents); }))
+		{
+			if (PruneValues(NumberOfPrincipalComponents, NumPoses, PosePairSimilarities, PCAValues,
+				[&PoseToPCAValueOffset](int32 PoseIdx) { return PoseToPCAValueOffset[PoseIdx]; },
+				[&PoseToPCAValueOffset](int32 PoseIdx, uint32 ValueOffset) { PoseToPCAValueOffset[PoseIdx] = ValueOffset; }))
+			{
+				// we pruned some PCAValues: we need to construct a mapping between PCAValuesVector to PoseIdx(s)
+				TMap<uint32, TArray<uint32>> PCAValuesVectorToPoseIndexesMap;
+				for (uint32 PoseIdx = 0; PoseIdx < NumPoses; ++PoseIdx)
+				{
+					check(PoseToPCAValueOffset[PoseIdx] % NumberOfPrincipalComponents == 0);
+					const uint32 PCAValuesVector = PoseToPCAValueOffset[PoseIdx] / NumberOfPrincipalComponents;
+					TArray<uint32>& PoseIndexes = PCAValuesVectorToPoseIndexesMap.FindOrAdd(PCAValuesVector);
+					check(!PoseIndexes.Contains(PoseIdx));
+					PoseIndexes.Add(PoseIdx);
+				}
+
+				int32 MoreThanOne = 0;
+				int32 MaxDuplicates = 1;
+				for (const TPair<uint32, TArray<uint32>>& Pair : PCAValuesVectorToPoseIndexesMap)
+				{
+					if (Pair.Value.Num() > 1)
+					{
+						++MoreThanOne;
+						MaxDuplicates = FMath::Max(MaxDuplicates, Pair.Value.Num());
+					}
+				}
+
+				UE_LOG(LogPoseSearch, Log, TEXT("MoreThanOne %d, MaxDuplicates %d"), MoreThanOne, MaxDuplicates);
+
+				FSparsePoseMultiMap<uint32> SparsePoseMultiMap(PCAValuesVectorToPoseIndexesMap.Num(), NumPoses - 1);
+				for (const TPair<uint32, TArray<uint32>>& Pair : PCAValuesVectorToPoseIndexesMap)
+				{
+					const uint32 PCAValuesVector = Pair.Key;
+					const TArray<uint32>& PoseIndexes = Pair.Value;
+					SparsePoseMultiMap.Insert(PCAValuesVector, PoseIndexes);
+				}
+
+				for (uint32 PCAValuesVector = 0; PCAValuesVector < SparsePoseMultiMap.Num(); ++PCAValuesVector)
+				{
+					const TConstArrayView<uint32> PoseIndexes = SparsePoseMultiMap[PCAValuesVector];
+					const TArray<uint32>& TestPoseIndexes = PCAValuesVectorToPoseIndexesMap[PCAValuesVector];
+					check(PoseIndexes == TestPoseIndexes);
+				}
+
+				PCAValuesVectorToPoseIndexes = SparsePoseMultiMap;
+			}
+		}
+	}
+}
+
 TArray<float> FSearchIndex::GetPoseValuesSafe(int32 PoseIdx) const
 {
 	TArray<float> PoseValues;
@@ -489,6 +645,7 @@ FArchive& operator<<(FArchive& Ar, FSearchIndex& Index)
 
 	Ar << Index.WeightsSqrt;
 	Ar << Index.PCAValues;
+	Ar << Index.PCAValuesVectorToPoseIndexes;
 	Ar << Index.PCAProjectionMatrix;
 	Ar << Index.Mean;
 	Ar << Index.PCAExplainedVariance;
diff --git a/Engine/Plugins/Experimental/Animation/PoseSearch/Source/Runtime/Public/PoseSearch/KDTree.h b/Engine/Plugins/Experimental/Animation/PoseSearch/Source/Runtime/Public/PoseSearch/KDTree.h
--- a/Engine/Plugins/Experimental/Animation/PoseSearch/Source/Runtime/Public/PoseSearch/KDTree.h
+++ b/Engine/Plugins/Experimental/Animation/PoseSearch/Source/Runtime/Public/PoseSearch/KDTree.h
@@ -54,9 +54,9 @@ struct POSESEARCH_API FKDTree
 		const float* Data = nullptr;
 	};
 
-	struct KNNResultSet
+	struct FKNNResultSet
 	{
-		inline KNNResultSet(size_t InNumNeighbors, TArrayView<size_t> InIndexes, TArrayView<float> InDistances, TConstArrayView<size_t> InExcludeFromSearchIndexes = TConstArrayView<size_t>())
+		inline FKNNResultSet(size_t InNumNeighbors, TArrayView<size_t> InIndexes, TArrayView<float> InDistances, TConstArrayView<size_t> InExcludeFromSearchIndexes = TConstArrayView<size_t>())
 		: Indexes(InIndexes)
 		, Distances(InDistances)
 		, NumNeighbors(InNumNeighbors)
@@ -121,9 +121,72 @@ struct POSESEARCH_API FKDTree
 		size_t Count;
 		TConstArrayView<size_t> ExcludeFromSearchIndexes; // sorted array view
 	};
+
+	struct FRadiusResultSet
+	{
+		inline FRadiusResultSet(float Radius, size_t MaxNumNeighbors, TArrayView<size_t> InIndexes, TArrayView<float> InDistances)
+		: Indexes(InIndexes)
+		, Distances(InDistances)
+		, NumNeighbors(MaxNumNeighbors)
+		, Count(0)
+		{
+			// by having IndexesView and DistancesView cardinality bigger than NumNeighbors, we can skip some if statements in the addPoint method
+			check(NumNeighbors > 0);
+			check(InIndexes.Num() > NumNeighbors);
+			check(InDistances.Num() > NumNeighbors);
+
+			Distances[NumNeighbors - 1] = Radius;
+		}
+
+		inline size_t Num() const
+		{
+			return Count;
+		}
+
+		inline bool full() const
+		{
+			return Count == NumNeighbors;
+		}
+
+		inline bool addPoint(float dist, size_t index)
+		{
+			if (dist < worstDist())
+			{
+				// shifting Distances[i] and Indexes[i] to make space for "dist" and "index" at the right "i"th slot
+				size_t i;
+				for (i = Count; (i > 0) && (Distances[i - 1] > dist); --i)
+				{
+					// no need to check "if (i < capacity)" since dists and indices can contains more items than capacity_ 
+					Distances[i] = Distances[i - 1];
+					Indexes[i] = Indexes[i - 1];
+				}
+
+				// inserting "dist" and "index" in a sorted manner
+				// no need to check "if (i < capacity)" since dists and indices can contains more items than capacity_ 
+				Distances[i] = dist;
+				Indexes[i] = index;
+
+				if (Count < NumNeighbors)
+				{
+					Count++;
+				}
+			}
+			
+			// tell caller that the search shall continue
+			return true;
+		}
+
+		inline float worstDist() const { return Distances[NumNeighbors - 1]; }
+
+	private:
+		TArrayView<size_t> Indexes;
+		TArrayView<float> Distances;
+		size_t NumNeighbors;
+		size_t Count;
+	};
 	
 
-	FKDTree(int32 Count, int32 Dim, const float* Data, int32 MaxLeafSize);
+	FKDTree(int32 Count, int32 Dim, const float* Data, int32 MaxLeafSize = 16);
 	FKDTree();
 	FKDTree(const FKDTree& Other);
 	FKDTree(FKDTree&& Other) = delete;
@@ -134,8 +197,9 @@ struct POSESEARCH_API FKDTree
 	FKDTree& operator=(FKDTree&& Other) = delete;
 
 	void Reset();
-	void Construct(int32 Count, int32 Dim, const float* Data, int32 MaxLeafSize);
-	bool FindNeighbors(KNNResultSet& Result, const float* Query) const;
+	void Construct(int32 Count, int32 Dim, const float* Data, int32 MaxLeafSize = 16);
+	bool FindNeighbors(FKNNResultSet& Result, const float* Query) const;
+	bool FindNeighbors(FRadiusResultSet& Result, const float* Query) const;
 	SIZE_T GetAllocatedSize() const;
 
 	FDataSource DataSource;
diff --git a/Engine/Plugins/Experimental/Animation/PoseSearch/Source/Runtime/Public/PoseSearch/PoseSearchDatabase.h b/Engine/Plugins/Experimental/Animation/PoseSearch/Source/Runtime/Public/PoseSearch/PoseSearchDatabase.h
--- a/Engine/Plugins/Experimental/Animation/PoseSearch/Source/Runtime/Public/PoseSearch/PoseSearchDatabase.h
+++ b/Engine/Plugins/Experimental/Animation/PoseSearch/Source/Runtime/Public/PoseSearch/PoseSearchDatabase.h
@@ -294,7 +294,7 @@ public:
 	int32 KDTreeMaxLeafSize = 16;
 	
 	// Out of a kdtree search, results will have only an approximate cost, so the database search will select the best “KDTree Query Num Neighbors” poses to perform the full cost analysis, and be able to elect the best pose.
-	UPROPERTY(EditAnywhere, Category = "Performance", meta = (EditCondition = "PoseSearchMode != EPoseSearchMode::BruteForce", EditConditionHides, ClampMin = "1", ClampMax = "600", UIMin = "1", UIMax = "600"))
+	UPROPERTY(EditAnywhere, Category = "Performance", meta = (EditCondition = "PoseSearchMode != EPoseSearchMode::BruteForce", EditConditionHides, ClampMin = "1", ClampMax = "600", UIMin = "1"))
 	int32 KDTreeQueryNumNeighbors = 200;
 
 	// When evaluating multiple searches, including the continuing pose search, the system keeps track of the best pose and associated cost.
@@ -303,10 +303,20 @@ public:
 	bool bSkipSearchIfPossible = true;
 
 	// if two poses values (multi dimensional point with the schema cardinality) are closer than PosePruningSimilarityThreshold,
-	// only one will be saved into the database (to save memory) and accessed by the two different pose indexes
+	// only one will be saved into the database FSearchIndexBase (to save memory) and accessed by the two different pose indexes
 	UPROPERTY(EditAnywhere, Category = "Performance")
 	float PosePruningSimilarityThreshold = 0.f;
 
+	// if two PCA values (multi dimensional point with the GetNumberOfPrincipalComponents cardinality) are closer than PCAValuesPruningSimilarityThreshold,
+	// only one will be saved into the database FSearchIndex (to save memory).
+	UPROPERTY(EditAnywhere, Category = "Performance")
+	float PCAValuesPruningSimilarityThreshold = 0.f;
+
+	// if PCAValuesPruningSimilarityThreshold > 0 the kdtree will remove duplicates, every result out of the KDTreeQueryNumNeighbors could potentially references multiple poses.
+	// KDTreeQueryNumNeighborsWithDuplicates is the upper bound number of poses the system will perform the full cost evaluation. if KDTreeQueryNumNeighborsWithDuplicates is zero then there's no upper bound
+	UPROPERTY(EditAnywhere, Category = "Performance", meta = (EditCondition = "PoseSearchMode != EPoseSearchMode::BruteForce && PCAValuesPruningSimilarityThreshold > 0", EditConditionHides, ClampMin = "0", ClampMax = "600", UIMin = "1"))
+	int32 KDTreeQueryNumNeighborsWithDuplicates = 0;
+	
 #if WITH_EDITORONLY_DATA
 	// This optional asset defines a list of databases you want to normalize together. Without it, it would be difficult to compare costs from separately normalized databases containing different types of animation,
 	// like only idles versus only runs animations, given that the range of movement would be dramatically different.
diff --git a/Engine/Plugins/Experimental/Animation/PoseSearch/Source/Runtime/Public/PoseSearch/PoseSearchIndex.h b/Engine/Plugins/Experimental/Animation/PoseSearch/Source/Runtime/Public/PoseSearch/PoseSearchIndex.h
--- a/Engine/Plugins/Experimental/Animation/PoseSearch/Source/Runtime/Public/PoseSearch/PoseSearchIndex.h
+++ b/Engine/Plugins/Experimental/Animation/PoseSearch/Source/Runtime/Public/PoseSearch/PoseSearchIndex.h
@@ -174,15 +174,119 @@ struct FSearchStats
 	friend FArchive& operator<<(FArchive& Ar, FSearchStats& Stats);
 };
 
+// compact representation of an array of arrays
+template <typename Type = uint32>
+struct FSparsePoseMultiMap
+{
+	FSparsePoseMultiMap(Type InMaxKey = Type(0), Type InMaxValue = Type(0))
+	: MaxKey(InMaxKey)
+	, MaxValue(InMaxValue)
+	, DeltaKeyValue(InMaxValue >= InMaxKey ? InMaxValue - InMaxKey + 1 : 0)
+	{
+		// @todo: maybe expose this initial allocation budget
+		DataValues.Reserve(InMaxKey * 2);
+		for (Type Index = 0; Index < InMaxKey; ++Index)
+		{
+			DataValues.Add(Type(INDEX_NONE));
+		}
+	}
+
+	void Insert(Type Key, TConstArrayView<Type> Values)
+	{
+#if DO_CHECK
+		// key must be valid..
+		check(Key != Type(INDEX_NONE));
+
+		// ..and within range of acceptance
+		check(Key >= 0 && Key < MaxKey);
+
+		// DataValues[Key] should be empty - inserting the same key multiple times is not allowed
+		check(DataValues[Key] == Type(INDEX_NONE));
+
+		// Values must contains at least one element..
+		check(!Values.IsEmpty());
+
+		// ..and none of the elements should be an invalid value (or it'll confuse the key/value decoding)
+		for (Type Value : Values)
+		{
+			check(Value <= MaxValue && Value != Type(INDEX_NONE));
+		}
+#endif //DO_CHECK
+
+		// if Values contains only one element we store it directly at the location referenced by key
+		if (Values.Num() == 1)
+		{
+			DataValues[Key] = Values[0];
+		}
+		// else we store the offset of the beginning of the encoded array (where the first element is the array size, followed by all the Values[i] elements
+		else
+		{
+			// checking for overflow
+			check((DataValues.Num() + 1 + Values.Num()) < (1 << (sizeof(Type) * 4 - 1)));
+			check(int(MaxKey) <= DataValues.Num());
+
+			// adding DeltaKeyValue to DataValues.Num() to making sure DataValues[Key] > MaxValue
+			DataValues[Key] = DataValues.Num() + DeltaKeyValue;
+			check(DataValues[Key] > MaxValue);
+
+			// encoding Values at the end of DataValues, by storing its size.. 
+			DataValues.Add(Values.Num());
+			// ..and its data right after
+			DataValues.Append(Values);
+		}
+	}
+
+	TConstArrayView<Type> operator [](Type Key) const
+	{
+		check(Key != Type(INDEX_NONE) && Key < MaxKey);
+		const Type Value = DataValues[Key];
+		if (Value <= MaxValue)
+		{
+			return MakeArrayView(DataValues.GetData() + Key, 1);
+		}
+
+		check(Value >= DeltaKeyValue);
+		const Type DecodedArrayStartLocation = Value - DeltaKeyValue;
+
+		// decoding the array at location DecodedArrayStartLocation: its size is stored at DecodedArrayStartLocation offset..
+		const Type Size = DataValues[DecodedArrayStartLocation];
+		// ..and it's data starts at the next location DecodedArrayStartLocation + 1
+		const Type DataOffset = DecodedArrayStartLocation + 1;
+		check(int32(DataOffset + Size) <= DataValues.Num());
+		return MakeArrayView(DataValues.GetData() + DataOffset, Size);
+	}
+
+	Type Num() const
+	{
+		return MaxKey;
+	}
+
+	SIZE_T GetAllocatedSize() const
+	{
+		return sizeof(MaxKey) + sizeof(MaxValue) + sizeof(DeltaKeyValue) + DataValues.GetAllocatedSize();
+	}
+
+	friend FArchive& operator<<(FArchive& Ar, FSparsePoseMultiMap& SparsePoseMultiMap)
+	{
+		Ar << SparsePoseMultiMap.MaxKey;
+		Ar << SparsePoseMultiMap.MaxValue;
+		Ar << SparsePoseMultiMap.DeltaKeyValue;
+		Ar << SparsePoseMultiMap.DataValues;
+		return Ar;
+	}
+
+	Type MaxKey = Type(0);
+	Type MaxValue = Type(0);
+	Type DeltaKeyValue = Type(0);
+	TArray<Type> DataValues;
+};
+
 /**
 * case class for FSearchIndex. building block used to gather data for data mining and calculate weights, pca, kdtree stuff
 */
 struct FSearchIndexBase
 {
-private:
 	TAlignedArray<float> Values;
-
-public:
 	TAlignedArray<FPoseMetadata> PoseMetadata;
 	bool bAnyBlockTransition = false;
 	TAlignedArray<FSearchIndexAsset> Assets;
@@ -208,9 +312,6 @@ public:
 	void ResetValues() { Values.Reset(); }
 	void AllocateData(int32 DataCardinality, int32 NumPoses);
 	
-	const TAlignedArray<float>& GetValues() const { return Values; }
-	TAlignedArray<float>& EditValues() { return Values; }
-
 	const FSearchIndexAsset& GetAssetForPose(int32 PoseIdx) const;
 	POSESEARCH_API const FSearchIndexAsset* GetAssetForPoseSafe(int32 PoseIdx) const;
 
@@ -242,6 +343,7 @@ struct FSearchIndex : public FSearchIndexBase
 	// since (VA - VB).square() could lead to big numbers, and VW being multiplied by the variance of the dataset
 	TAlignedArray<float> WeightsSqrt;
 	TAlignedArray<float> PCAValues;
+	FSparsePoseMultiMap<uint32> PCAValuesVectorToPoseIndexes;
 	TAlignedArray<float> PCAProjectionMatrix;
 	TAlignedArray<float> Mean;
 
@@ -267,6 +369,8 @@ struct FSearchIndex : public FSearchIndexBase
 	POSESEARCH_API FPoseSearchCost ComparePoses(int32 PoseIdx, float ContinuingPoseCostBias, TConstArrayView<float> PoseValues, TConstArrayView<float> QueryValues) const;
 	POSESEARCH_API FPoseSearchCost CompareAlignedPoses(int32 PoseIdx, float ContinuingPoseCostBias, TConstArrayView<float> PoseValues, TConstArrayView<float> QueryValues) const;
 
+	void PruneDuplicatePCAValues(float SimilarityThreshold, int32 NumberOfPrincipalComponents);
+
 	friend FArchive& operator<<(FArchive& Ar, FSearchIndex& Index);
 };
 
