motion matching - fix for SDebuggerDatabaseView::Update crashing when trying to AddUnfilteredDatabaseRow while using databases with pruned out PCA values indexes

diff --git a/Engine/Plugins/Experimental/Animation/PoseSearch/Source/Editor/Private/PoseSearchDebuggerDatabaseView.cpp b/Engine/Plugins/Experimental/Animation/PoseSearch/Source/Editor/Private/PoseSearchDebuggerDatabaseView.cpp
--- a/Engine/Plugins/Experimental/Animation/PoseSearch/Source/Editor/Private/PoseSearchDebuggerDatabaseView.cpp
+++ b/Engine/Plugins/Experimental/Animation/PoseSearch/Source/Editor/Private/PoseSearchDebuggerDatabaseView.cpp
@@ -155,7 +155,7 @@ private:
 
 static void AddUnfilteredDatabaseRow(const UPoseSearchDatabase* Database, 
 	TArray<TSharedRef<FDebuggerDatabaseRowData>>& UnfilteredDatabaseRows, TSharedRef<FDebuggerDatabaseSharedData> SharedData,
-	int32 DbPoseIdx, EPoseCandidateFlags PoseCandidateFlags, const FPoseSearchCost& Cost = FPoseSearchCost())
+	int32 DbPoseIdx, EPoseCandidateFlags PoseCandidateFlags, TConstArrayView<uint32> PoseToPCAValuesVectorIndexes, const FPoseSearchCost& Cost = FPoseSearchCost())
 {
 	const FSearchIndex& SearchIndex = Database->GetSearchIndex();
 	if (const FSearchIndexAsset* SearchIndexAsset = SearchIndex.GetAssetForPoseSafe(DbPoseIdx))
@@ -203,7 +203,8 @@ static void AddUnfilteredDatabaseRow(const UPoseSearchDatabase* Database,
 
 		if (!SharedData->PCAQueryVector.IsEmpty())
 		{
-			TConstArrayView<float> PCAPoseValues = SearchIndex.GetPCAPoseValues(DbPoseIdx);
+			const int32 PCAValuesVectorIdx = PoseToPCAValuesVectorIndexes.IsEmpty() ? DbPoseIdx : PoseToPCAValuesVectorIndexes[DbPoseIdx];
+			TConstArrayView<float> PCAPoseValues = SearchIndex.GetPCAPoseValues(PCAValuesVectorIdx);
 			if (SharedData->PCAQueryVector.Num() == PCAPoseValues.Num())
 			{
 				Row->PosePCACost = CompareFeatureVectors(SharedData->PCAQueryVector, PCAPoseValues);
@@ -253,11 +254,15 @@ void SDebuggerDatabaseView::Update(const FTraceMotionMatchingStateMessage& State
 	UnfilteredDatabaseRows.Reset();
 
 	bool bAddPCACost = false;
+	TArray<uint32> PoseToPCAValuesVectorIndexes;
 	for (const FTraceMotionMatchingStateDatabaseEntry& DbEntry : State.DatabaseEntries)
 	{
 		const UPoseSearchDatabase* Database = FTraceMotionMatchingState::GetObjectFromId<UPoseSearchDatabase>(DbEntry.DatabaseId);
 		if (FAsyncPoseSearchDatabasesManagement::RequestAsyncBuildIndex(Database, ERequestAsyncBuildFlag::ContinueRequest))
 		{
+			const FSearchIndex& SearchIndex = Database->GetSearchIndex();
+			SearchIndex.GetPoseToPCAValuesVectorIndexes(PoseToPCAValuesVectorIndexes);
+
 			TSharedRef<FDebuggerDatabaseSharedData> SharedData = MakeShared<FDebuggerDatabaseSharedData>();
 			SharedData->SourceDatabase = Database;
 			SharedData->DatabaseName = Database->GetName();
@@ -273,7 +278,7 @@ void SDebuggerDatabaseView::Update(const FTraceMotionMatchingStateMessage& State
 
 			for (const FTraceMotionMatchingStatePoseEntry& PoseEntry : DbEntry.PoseEntries)
 			{
-				AddUnfilteredDatabaseRow(Database, UnfilteredDatabaseRows, SharedData, PoseEntry.DbPoseIdx, PoseEntry.PoseCandidateFlags, PoseEntry.Cost);
+				AddUnfilteredDatabaseRow(Database, UnfilteredDatabaseRows, SharedData, PoseEntry.DbPoseIdx, PoseEntry.PoseCandidateFlags, PoseToPCAValuesVectorIndexes, PoseEntry.Cost);
 			}
 
 			if (bShowAllPoses)
@@ -284,12 +289,11 @@ void SDebuggerDatabaseView::Update(const FTraceMotionMatchingStateMessage& State
 					PoseEntriesIdx.Add(PoseEntry.DbPoseIdx);
 				}
 
-				const FSearchIndex& SearchIndex = Database->GetSearchIndex();
 				for (int32 DbPoseIdx = 0; DbPoseIdx < SearchIndex.GetNumPoses(); ++DbPoseIdx)
 				{
 					if (!PoseEntriesIdx.Find(DbPoseIdx))
 					{
-						AddUnfilteredDatabaseRow(Database, UnfilteredDatabaseRows, SharedData, DbPoseIdx, EPoseCandidateFlags::DiscardedBy_Search);
+						AddUnfilteredDatabaseRow(Database, UnfilteredDatabaseRows, SharedData, DbPoseIdx, EPoseCandidateFlags::DiscardedBy_Search, PoseToPCAValuesVectorIndexes);
 					}
 				}
 			}
diff --git a/Engine/Plugins/Experimental/Animation/PoseSearch/Source/Runtime/Private/PoseSearchDatabase.cpp b/Engine/Plugins/Experimental/Animation/PoseSearch/Source/Runtime/Private/PoseSearchDatabase.cpp
--- a/Engine/Plugins/Experimental/Animation/PoseSearch/Source/Runtime/Private/PoseSearchDatabase.cpp
+++ b/Engine/Plugins/Experimental/Animation/PoseSearch/Source/Runtime/Private/PoseSearchDatabase.cpp
@@ -848,7 +848,7 @@ UE::PoseSearch::FSearchResult UPoseSearchDatabase::SearchPCAKDTree(UE::PoseSearc
 		// NonSelectableIdx are already filtered out inside the kdtree search
 		const FSearchFilters SearchFilters(Schema, bArePCAValuesPruned ? NonSelectableIdx : TConstArrayView<size_t>(), SearchIndex.bAnyBlockTransition);
 		
-		// are the PCAValues pruned out of duplicates (multiple poses are associated with the same PCAValuesVector)
+		// are the PCAValues pruned out of duplicates (multiple poses are associated with the same PCAValuesVectorIdx)
 		if (bArePCAValuesPruned)
 		{
 			// @todo: reconstruction is not yet supported with pruned PCAValues
diff --git a/Engine/Plugins/Experimental/Animation/PoseSearch/Source/Runtime/Private/PoseSearchIndex.cpp b/Engine/Plugins/Experimental/Animation/PoseSearch/Source/Runtime/Private/PoseSearchIndex.cpp
--- a/Engine/Plugins/Experimental/Animation/PoseSearch/Source/Runtime/Private/PoseSearchIndex.cpp
+++ b/Engine/Plugins/Experimental/Animation/PoseSearch/Source/Runtime/Private/PoseSearchIndex.cpp
@@ -538,13 +538,13 @@ void FSearchIndex::PruneDuplicatePCAValues(float SimilarityThreshold, int32 Numb
 				[&PoseToPCAValueOffset](int32 PoseIdx) { return PoseToPCAValueOffset[PoseIdx]; },
 				[&PoseToPCAValueOffset](int32 PoseIdx, uint32 ValueOffset) { PoseToPCAValueOffset[PoseIdx] = ValueOffset; }))
 			{
-				// we pruned some PCAValues: we need to construct a mapping between PCAValuesVector to PoseIdx(s)
+				// we pruned some PCAValues: we need to construct a mapping between PCAValuesVectorIdx to PoseIdx(s)
 				TMap<uint32, TArray<uint32>> PCAValuesVectorToPoseIndexesMap;
 				for (uint32 PoseIdx = 0; PoseIdx < NumPoses; ++PoseIdx)
 				{
 					check(PoseToPCAValueOffset[PoseIdx] % NumberOfPrincipalComponents == 0);
-					const uint32 PCAValuesVector = PoseToPCAValueOffset[PoseIdx] / NumberOfPrincipalComponents;
-					TArray<uint32>& PoseIndexes = PCAValuesVectorToPoseIndexesMap.FindOrAdd(PCAValuesVector);
+					const uint32 PCAValuesVectorIdx = PoseToPCAValueOffset[PoseIdx] / NumberOfPrincipalComponents;
+					TArray<uint32>& PoseIndexes = PCAValuesVectorToPoseIndexesMap.FindOrAdd(PCAValuesVectorIdx);
 					check(!PoseIndexes.Contains(PoseIdx));
 					PoseIndexes.Add(PoseIdx);
 				}
@@ -565,15 +565,15 @@ void FSearchIndex::PruneDuplicatePCAValues(float SimilarityThreshold, int32 Numb
 				FSparsePoseMultiMap<uint32> SparsePoseMultiMap(PCAValuesVectorToPoseIndexesMap.Num(), NumPoses - 1);
 				for (const TPair<uint32, TArray<uint32>>& Pair : PCAValuesVectorToPoseIndexesMap)
 				{
-					const uint32 PCAValuesVector = Pair.Key;
+					const uint32 PCAValuesVectorIdx = Pair.Key;
 					const TArray<uint32>& PoseIndexes = Pair.Value;
-					SparsePoseMultiMap.Insert(PCAValuesVector, PoseIndexes);
+					SparsePoseMultiMap.Insert(PCAValuesVectorIdx, PoseIndexes);
 				}
 
-				for (uint32 PCAValuesVector = 0; PCAValuesVector < SparsePoseMultiMap.Num(); ++PCAValuesVector)
+				for (uint32 PCAValuesVectorIdx = 0; PCAValuesVectorIdx < SparsePoseMultiMap.Num(); ++PCAValuesVectorIdx)
 				{
-					const TConstArrayView<uint32> PoseIndexes = SparsePoseMultiMap[PCAValuesVector];
-					const TArray<uint32>& TestPoseIndexes = PCAValuesVectorToPoseIndexesMap[PCAValuesVector];
+					const TConstArrayView<uint32> PoseIndexes = SparsePoseMultiMap[PCAValuesVectorIdx];
+					const TArray<uint32>& TestPoseIndexes = PCAValuesVectorToPoseIndexesMap[PCAValuesVectorIdx];
 					check(PoseIndexes == TestPoseIndexes);
 				}
 
@@ -602,7 +602,7 @@ TArray<float> FSearchIndex::GetPoseValuesSafe(int32 PoseIdx) const
 	return PoseValues;
 }
 
-TConstArrayView<float> FSearchIndex::GetPCAPoseValues(int32 PoseIdx) const
+TConstArrayView<float> FSearchIndex::GetPCAPoseValues(int32 PCAValuesVectorIdx) const
 {
 	if (PCAValues.IsEmpty())
 	{
@@ -612,10 +612,14 @@ TConstArrayView<float> FSearchIndex::GetPCAPoseValues(int32 PoseIdx) const
 	const int32 NumDimensions = WeightsSqrt.Num();
 	const int32 NumberOfPrincipalComponents = PCAProjectionMatrix.Num() / NumDimensions;
 
-	check(PCAProjectionMatrix.Num() > 0 && PCAProjectionMatrix.Num() % NumDimensions == 0);
-	check(PoseIdx >= 0 && PoseIdx < GetNumPoses() && NumDimensions > 0);
+#if DO_CHECK
+	check(NumDimensions > 0 && PCAProjectionMatrix.Num() > 0 && PCAProjectionMatrix.Num() % NumDimensions == 0);
+	check(PCAValues.Num() % NumberOfPrincipalComponents == 0);
+	const int32 NumPCAValuesVectors = PCAValues.Num() / NumberOfPrincipalComponents;
+	check(PCAValuesVectorIdx >= 0 && PCAValuesVectorIdx < NumPCAValuesVectors );
+#endif // DO_CHECK
 
-	const int32 ValueOffset = PoseIdx * NumberOfPrincipalComponents;
+	const int32 ValueOffset = PCAValuesVectorIdx * NumberOfPrincipalComponents;
 	return MakeArrayView(&PCAValues[ValueOffset], NumberOfPrincipalComponents);
 }
 
@@ -639,6 +643,27 @@ FPoseSearchCost FSearchIndex::CompareAlignedPoses(int32 PoseIdx, float Continuin
 	return FPoseSearchCost(DissimilarityCost, NotifyAddend, ContinuingPoseCostBias);
 }
 
+void FSearchIndex::GetPoseToPCAValuesVectorIndexes(TArray<uint32>& PoseToPCAValuesVectorIndexes) const
+{
+	if (PCAValuesVectorToPoseIndexes.Num() > 0)
+	{
+		PoseToPCAValuesVectorIndexes.Init(uint32(INDEX_NONE), PCAValuesVectorToPoseIndexes.MaxValue + 1);
+		for (uint32 PCAValuesVectorIdx = 0; PCAValuesVectorIdx < PCAValuesVectorToPoseIndexes.Num(); ++PCAValuesVectorIdx)
+		{
+			for (uint32 PoseIdx : PCAValuesVectorToPoseIndexes[PCAValuesVectorIdx])
+			{
+				PoseToPCAValuesVectorIndexes[PoseIdx] = PCAValuesVectorIdx;
+			}
+		}
+
+		check(!PoseToPCAValuesVectorIndexes.Contains(uint32(INDEX_NONE)));
+	}
+	else
+	{
+		PoseToPCAValuesVectorIndexes.Reset();
+	}
+}
+
 FArchive& operator<<(FArchive& Ar, FSearchIndex& Index)
 {
 	Ar << static_cast<FSearchIndexBase&>(Index);
diff --git a/Engine/Plugins/Experimental/Animation/PoseSearch/Source/Runtime/Public/PoseSearch/PoseSearchIndex.h b/Engine/Plugins/Experimental/Animation/PoseSearch/Source/Runtime/Public/PoseSearch/PoseSearchIndex.h
--- a/Engine/Plugins/Experimental/Animation/PoseSearch/Source/Runtime/Public/PoseSearch/PoseSearchIndex.h
+++ b/Engine/Plugins/Experimental/Animation/PoseSearch/Source/Runtime/Public/PoseSearch/PoseSearchIndex.h
@@ -365,12 +365,19 @@ struct FSearchIndex : public FSearchIndexBase
 	POSESEARCH_API TConstArrayView<float> PCAProject(TConstArrayView<float> PoseValues, TArrayView<float> BufferUsedForProjection) const;
 
 	POSESEARCH_API TArray<float> GetPoseValuesSafe(int32 PoseIdx) const;
-	POSESEARCH_API TConstArrayView<float> GetPCAPoseValues(int32 PoseIdx) const;
+
+	// since PCAValues (as well as Values can be pruned out from duplicate data, we lose the 1:1 mapping between PoseIdx and PCAValuesVectorIdx
+	// that in the case of GetPoseValuesSafe it's stored in PoseMetadata[PoseIdx].GetValueOffset(), but missing for the PCAValues, so this API input is NOT a PoseIdx
+	// mapping between PoseIdx to PCAValuesVectorIdx can be reconstructed by inverting the PCAValuesVectorToPoseIndexes via GetPoseToPCAValuesVectorIndexes
+	POSESEARCH_API TConstArrayView<float> GetPCAPoseValues(int32 PCAValuesVectorIdx) const;
 	POSESEARCH_API FPoseSearchCost ComparePoses(int32 PoseIdx, float ContinuingPoseCostBias, TConstArrayView<float> PoseValues, TConstArrayView<float> QueryValues) const;
 	POSESEARCH_API FPoseSearchCost CompareAlignedPoses(int32 PoseIdx, float ContinuingPoseCostBias, TConstArrayView<float> PoseValues, TConstArrayView<float> QueryValues) const;
 
 	void PruneDuplicatePCAValues(float SimilarityThreshold, int32 NumberOfPrincipalComponents);
 
+	// returns the inverse mapping of PCAValuesVectorToPoseIndexes
+	POSESEARCH_API void GetPoseToPCAValuesVectorIndexes(TArray<uint32>& PoseToPCAValuesVectorIndexes) const;
+
 	friend FArchive& operator<<(FArchive& Ar, FSearchIndex& Index);
 };
 
