motion matching - resetting blend stack when UpdateAssetPlayer bNeedsReset (fix for stale animation being held in blend stack when MM node is becoming  relevant)

diff --git a/Engine/Plugins/Experimental/Animation/PoseSearch/Source/Runtime/Private/AnimNode_BlendStack.cpp b/Engine/Plugins/Experimental/Animation/PoseSearch/Source/Runtime/Private/AnimNode_BlendStack.cpp
--- a/Engine/Plugins/Experimental/Animation/PoseSearch/Source/Runtime/Private/AnimNode_BlendStack.cpp
+++ b/Engine/Plugins/Experimental/Animation/PoseSearch/Source/Runtime/Private/AnimNode_BlendStack.cpp
@@ -404,6 +404,8 @@ void FAnimNode_BlendStack_Standalone::Initialize_AnyThread(const FAnimationIniti
 {
 	Super::Initialize_AnyThread(Context);
 
+	Reset();
+
 	IAnimClassInterface* AnimBlueprintClass = Context.GetAnimClass();
 	if (SampleGraphPoseLinks.IsEmpty() == false)
 	{
@@ -565,6 +567,11 @@ void FAnimNode_BlendStack_Standalone::BlendTo(const FAnimationUpdateContext& Con
 	InitializeSample(InitContext, AnimPlayer);
 }
 
+void FAnimNode_BlendStack_Standalone::Reset()
+{
+	AnimPlayers.Reset();
+}
+
 int32 FAnimNode_BlendStack_Standalone::GetNextPoseLinkIndex()
 {
 	if (SampleGraphPoseLinks.IsEmpty())
@@ -612,6 +619,18 @@ void FAnimNode_BlendStack_Standalone::GatherDebugData(FNodeDebugData& DebugData)
 
 void FAnimNode_BlendStack::UpdateAssetPlayer(const FAnimationUpdateContext& Context)
 {
+	const bool bNeedsReset =
+		bResetOnBecomingRelevant &&
+		UpdateCounter.HasEverBeenUpdated() &&
+		!UpdateCounter.WasSynchronizedCounter(Context.AnimInstanceProxy->GetUpdateCounter());
+
+	if (bNeedsReset)
+	{
+		Reset();
+	}
+
+	UpdateCounter.SynchronizeWith(Context.AnimInstanceProxy->GetUpdateCounter());
+
 	GetEvaluateGraphExposedInputs().Execute(Context);
 
 	if (AnimationAsset)
diff --git a/Engine/Plugins/Experimental/Animation/PoseSearch/Source/Runtime/Private/AnimNode_MotionMatching.cpp b/Engine/Plugins/Experimental/Animation/PoseSearch/Source/Runtime/Private/AnimNode_MotionMatching.cpp
--- a/Engine/Plugins/Experimental/Animation/PoseSearch/Source/Runtime/Private/AnimNode_MotionMatching.cpp
+++ b/Engine/Plugins/Experimental/Animation/PoseSearch/Source/Runtime/Private/AnimNode_MotionMatching.cpp
@@ -35,12 +35,12 @@ void FAnimNode_MotionMatching::Initialize_AnyThread(const FAnimationInitializeCo
 
 	GetEvaluateGraphExposedInputs().Execute(Context);
 
-	BlendStackNode.Initialize_AnyThread(Context);
-
 	Source.SetLinkNode(&BlendStackNode);
+
+	// calling BlendStackNode.Initialize_AnyThread(Context) that will reset the BlendStackNode.AnimPlayers
 	Source.Initialize(Context);
 
-	MotionMatchingState.UpdateRootBoneControl(Context.AnimInstanceProxy, 0.f);
+	MotionMatchingState.Reset(Context.AnimInstanceProxy->GetComponentTransform());
 }
 
 void FAnimNode_MotionMatching::Evaluate_AnyThread(FPoseContext& Output)
@@ -111,6 +111,7 @@ void FAnimNode_MotionMatching::UpdateAssetPlayer(const FAnimationUpdateContext&
 	if (bNeedsReset)
 	{
 		MotionMatchingState.Reset(Context.AnimInstanceProxy->GetComponentTransform());
+		BlendStackNode.Reset();
 	}
 	else
 	{
diff --git a/Engine/Plugins/Experimental/Animation/PoseSearch/Source/Runtime/Public/PoseSearch/AnimNode_BlendStack.h b/Engine/Plugins/Experimental/Animation/PoseSearch/Source/Runtime/Public/PoseSearch/AnimNode_BlendStack.h
--- a/Engine/Plugins/Experimental/Animation/PoseSearch/Source/Runtime/Public/PoseSearch/AnimNode_BlendStack.h
+++ b/Engine/Plugins/Experimental/Animation/PoseSearch/Source/Runtime/Public/PoseSearch/AnimNode_BlendStack.h
@@ -103,6 +103,7 @@ struct POSESEARCH_API FAnimNode_BlendStack_Standalone : public FAnimNode_AssetPl
 
 	void BlendTo(const FAnimationUpdateContext& Context, UAnimationAsset* AnimationAsset, float AccumulatedTime = 0.f, bool bLoop = false, bool bMirrored = false, UMirrorDataTable* MirrorDataTable = nullptr, float BlendTime = 0.2f, float RootBoneBlendTime = -1.f, const UBlendProfile* BlendProfile = nullptr, EAlphaBlendOption BlendOption = EAlphaBlendOption::Linear, FVector BlendParameters = FVector::Zero(), float PlayRate = 1.f);
 	void UpdatePlayRate(float PlayRate);
+	void Reset();
 
 	// FAnimNode_AssetPlayerBase interface
 	virtual float GetAccumulatedTime() const override;
@@ -178,6 +179,13 @@ struct POSESEARCH_API FAnimNode_BlendStack : public FAnimNode_BlendStack_Standal
 	UPROPERTY(EditAnywhere, Category = Settings, meta = (PinHiddenByDefault))
 	TObjectPtr<UMirrorDataTable> MirrorDataTable;
 
+	// Reset the blend stack if it has become relevant to the graph after not being updated on previous frames.
+	UPROPERTY(EditAnywhere, Category = Settings)
+	bool bResetOnBecomingRelevant = true;
+
+	// Update Counter for detecting being relevant
+	FGraphTraversalCounter UpdateCounter;
+
 	// FAnimNode_Base interface
 	virtual void UpdateAssetPlayer(const FAnimationUpdateContext& Context) override;
 	// End of FAnimNode_Base interface
diff --git a/Engine/Plugins/Experimental/Animation/PoseSearch/Source/Runtime/Public/PoseSearch/AnimNode_MotionMatching.h b/Engine/Plugins/Experimental/Animation/PoseSearch/Source/Runtime/Public/PoseSearch/AnimNode_MotionMatching.h
--- a/Engine/Plugins/Experimental/Animation/PoseSearch/Source/Runtime/Public/PoseSearch/AnimNode_MotionMatching.h
+++ b/Engine/Plugins/Experimental/Animation/PoseSearch/Source/Runtime/Public/PoseSearch/AnimNode_MotionMatching.h
@@ -96,7 +96,7 @@ private:
 	FFloatInterval PlayRate = FFloatInterval(1.f, 1.f);
 
 	// Reset the motion matching selection state if it has become relevant to the graph after not being updated on previous frames.
-	UPROPERTY(EditAnywhere, Category = Settings, meta = (PinHiddenByDefault))
+	UPROPERTY(EditAnywhere, Category = Settings)
 	bool bResetOnBecomingRelevant = true;
 
 	// If set to false, the motion matching node will perform a search only if the continuing pose is invalid. This is useful if you want to stagger searches of different nodes for performance reasons
