motion matching - preventing MM db editor removing preview actors when editing database or dependent assets properties

diff --git a/Engine/Plugins/Experimental/Animation/PoseSearch/Source/Editor/Private/PoseSearchDatabaseEditor.cpp b/Engine/Plugins/Experimental/Animation/PoseSearch/Source/Editor/Private/PoseSearchDatabaseEditor.cpp
--- a/Engine/Plugins/Experimental/Animation/PoseSearch/Source/Editor/Private/PoseSearchDatabaseEditor.cpp
+++ b/Engine/Plugins/Experimental/Animation/PoseSearch/Source/Editor/Private/PoseSearchDatabaseEditor.cpp
@@ -219,6 +219,7 @@ namespace UE::PoseSearch
 				}
 
 				// Ensure any database changes are reflected
+				DatabaseAsset->UnregisterOnDerivedDataRebuild(this);
 				DatabaseAsset->RegisterOnDerivedDataRebuild(UPoseSearchDatabase::FOnDerivedDataRebuild::CreateSP(this, &FDatabaseEditor::RefreshStatisticsWidgetInformation));
 			}
 		}
@@ -287,6 +288,17 @@ namespace UE::PoseSearch
 		RegenerateMenusAndToolbars();
 	}
 
+	FDatabaseEditor::~FDatabaseEditor()
+	{
+		if (ViewModel)
+		{
+			if (UPoseSearchDatabase* DatabaseAsset = ViewModel->GetPoseSearchDatabase())
+			{
+				DatabaseAsset->UnregisterOnDerivedDataRebuild(this);
+			}
+		}
+	}
+
 	void FDatabaseEditor::RegisterTabSpawners(const TSharedRef<FTabManager>& InTabManager)
 	{
 		WorkspaceMenuCategory = InTabManager->AddLocalWorkspaceMenuCategory(
diff --git a/Engine/Plugins/Experimental/Animation/PoseSearch/Source/Editor/Private/PoseSearchDatabaseViewModel.cpp b/Engine/Plugins/Experimental/Animation/PoseSearch/Source/Editor/Private/PoseSearchDatabaseViewModel.cpp
--- a/Engine/Plugins/Experimental/Animation/PoseSearch/Source/Editor/Private/PoseSearchDatabaseViewModel.cpp
+++ b/Engine/Plugins/Experimental/Animation/PoseSearch/Source/Editor/Private/PoseSearchDatabaseViewModel.cpp
@@ -63,11 +63,6 @@ namespace UE::PoseSearch
 	{
 	}
 
-	FDatabaseViewModel::~FDatabaseViewModel()
-	{
-		PoseSearchDatabase->UnregisterOnDerivedDataRebuild(this);
-	}
-
 	void FDatabaseViewModel::AddReferencedObjects(FReferenceCollector& Collector)
 	{
 		Collector.AddReferencedObject(PoseSearchDatabase);
@@ -80,8 +75,6 @@ namespace UE::PoseSearch
 		DatabaseDataDetails = InDatabaseDataDetails;
 
 		RemovePreviewActors();
-
-		PoseSearchDatabase->RegisterOnDerivedDataRebuild(UPoseSearchDatabase::FOnDerivedDataRebuild::CreateSP(this, &FDatabaseViewModel::RemovePreviewActors));
 	}
 
 	void FDatabaseViewModel::BuildSearchIndex()
diff --git a/Engine/Plugins/Experimental/Animation/PoseSearch/Source/Editor/Public/PoseSearchDatabaseEditor.h b/Engine/Plugins/Experimental/Animation/PoseSearch/Source/Editor/Public/PoseSearchDatabaseEditor.h
--- a/Engine/Plugins/Experimental/Animation/PoseSearch/Source/Editor/Public/PoseSearchDatabaseEditor.h
+++ b/Engine/Plugins/Experimental/Animation/PoseSearch/Source/Editor/Public/PoseSearchDatabaseEditor.h
@@ -27,6 +27,7 @@ namespace UE::PoseSearch
 			const EToolkitMode::Type Mode,
 			const TSharedPtr<IToolkitHost>& InitToolkitHost,
 			UPoseSearchDatabase* DatabaseAsset);
+		virtual ~FDatabaseEditor();
 
 		virtual void RegisterTabSpawners(const TSharedRef<FTabManager>& TabManager) override;
 		virtual void UnregisterTabSpawners(const TSharedRef<FTabManager>& TabManager) override;
diff --git a/Engine/Plugins/Experimental/Animation/PoseSearch/Source/Editor/Public/PoseSearchDatabaseViewModel.h b/Engine/Plugins/Experimental/Animation/PoseSearch/Source/Editor/Public/PoseSearchDatabaseViewModel.h
--- a/Engine/Plugins/Experimental/Animation/PoseSearch/Source/Editor/Public/PoseSearchDatabaseViewModel.h
+++ b/Engine/Plugins/Experimental/Animation/PoseSearch/Source/Editor/Public/PoseSearchDatabaseViewModel.h
@@ -60,9 +60,7 @@ namespace UE::PoseSearch
 	class FDatabaseViewModel : public TSharedFromThis<FDatabaseViewModel>, public FGCObject
 	{
 	public:
-
 		FDatabaseViewModel();
-		virtual ~FDatabaseViewModel();
 
 		// ~ FGCObject interface
 		virtual void AddReferencedObjects(FReferenceCollector& Collector) override;
