motion matching - heading channel debug lines improvements

diff --git a/Engine/Plugins/Experimental/Animation/PoseSearch/Source/Runtime/Private/PoseSearchContext.cpp b/Engine/Plugins/Experimental/Animation/PoseSearch/Source/Runtime/Private/PoseSearchContext.cpp
--- a/Engine/Plugins/Experimental/Animation/PoseSearch/Source/Runtime/Private/PoseSearchContext.cpp
+++ b/Engine/Plugins/Experimental/Animation/PoseSearch/Source/Runtime/Private/PoseSearchContext.cpp
@@ -8,6 +8,7 @@
 #include "PoseSearch/PoseSearchSchema.h"
 #include "PoseSearch/PoseSearchTrajectoryTypes.h"
 #include "PoseSearchFeatureChannel_Position.h"
+#include "PoseSearchFeatureChannel_Heading.h"
 
 namespace UE::PoseSearch
 {
@@ -62,25 +63,26 @@ FVector FDebugDrawParams::ExtractPosition(TConstArrayView<float> PoseVector, flo
 {
 	// we don't wanna ask for a SchemaOriginBoneIdx in the future or past
 	check(PermutationTimeType != EPermutationTimeType::UsePermutationTime);
-	if (const UPoseSearchSchema* Schema = GetSchema())
+	const UPoseSearchSchema* Schema = GetSchema();
+	if (!FMath::IsNearlyZero(SampleTimeOffset) && Schema)
 	{
 		// looking for a UPoseSearchFeatureChannel_Position that matches the TimeOffset and SchemaBoneIdx,
 		// with SchemaOriginBoneIdx to be the root bone and the appropriate PermutationTimeType 
 		if (const UPoseSearchFeatureChannel_Position* FoundPosition = static_cast<const UPoseSearchFeatureChannel_Position*>(
-				Schema->FindChannel([SampleTimeOffset, SchemaBoneIdx, PermutationTimeType, Schema](const UPoseSearchFeatureChannel* Channel) -> const UPoseSearchFeatureChannel_Position*
-			{
-				if (const UPoseSearchFeatureChannel_Position* Position = Cast<UPoseSearchFeatureChannel_Position>(Channel))
+			Schema->FindChannel([SampleTimeOffset, SchemaBoneIdx, PermutationTimeType, Schema](const UPoseSearchFeatureChannel* Channel) -> const UPoseSearchFeatureChannel_Position*
 				{
-					if (Position->SchemaBoneIdx == SchemaBoneIdx &&
-						Position->SampleTimeOffset == SampleTimeOffset &&
-						Position->PermutationTimeType == PermutationTimeType &&
-						Schema->IsRootBone(Position->SchemaOriginBoneIdx))
+					if (const UPoseSearchFeatureChannel_Position* Position = Cast<UPoseSearchFeatureChannel_Position>(Channel))
 					{
-						return Position;
+						if (Position->SchemaBoneIdx == SchemaBoneIdx &&
+							Position->SampleTimeOffset == SampleTimeOffset &&
+							Position->PermutationTimeType == PermutationTimeType &&
+							Schema->IsRootBone(Position->SchemaOriginBoneIdx))
+						{
+							return Position;
+						}
 					}
-				}
-		return nullptr;
-			})))
+					return nullptr;
+				})))
 		{
 			return ExtractPosition(PoseVector, FoundPosition);
 		}
@@ -93,6 +95,111 @@ FVector FDebugDrawParams::ExtractPosition(TConstArrayView<float> PoseVector, flo
 	return GetRootTransform().GetTranslation();
 }
 
+FQuat FDebugDrawParams::ExtractRotation(TConstArrayView<float> PoseVector, float SampleTimeOffset, int8 SchemaBoneIdx) const
+{
+	const UPoseSearchSchema* Schema = GetSchema();
+	if (!FMath::IsNearlyZero(SampleTimeOffset) && Schema)
+	{
+		int32 HeadingAxisFoundNum = 0;
+		const UPoseSearchFeatureChannel_Heading* FoundHeading[int32(EHeadingAxis::Num)];
+		FVector DecodedHeading[int32(EHeadingAxis::Num)];
+		for (int32 HeadingAxis = 0; HeadingAxis < int32(EHeadingAxis::Num); ++HeadingAxis)
+		{
+			// looking for a UPoseSearchFeatureChannel_Heading that matches the SampleTimeOffset, SchemaBoneIdx, and with OriginTimeOffset as zero.
+			// the features data associated to this channel would be a heading vector in GetRootTransform space (since OriginTimeOffset is zero)), 
+			// so by finding at least two with differnt axis we'll be able to compose a delta rotation from OriginTimeOffset (zero) to SampleTimeOffset
+			FoundHeading[HeadingAxis] = static_cast<const UPoseSearchFeatureChannel_Heading*>(
+				Schema->FindChannel([SampleTimeOffset, SchemaBoneIdx, HeadingAxis](const UPoseSearchFeatureChannel* Channel) -> const UPoseSearchFeatureChannel_Heading*
+					{
+						if (const UPoseSearchFeatureChannel_Heading* Heading = Cast<UPoseSearchFeatureChannel_Heading>(Channel))
+						{
+							if (Heading->OriginTimeOffset == 0.f &&
+								Heading->SchemaBoneIdx == SchemaBoneIdx &&
+								Heading->SampleTimeOffset == SampleTimeOffset &&
+								int32(Heading->HeadingAxis) == HeadingAxis)
+							{
+								return Heading;
+							}
+						}
+						return nullptr;
+					}));
+			if (FoundHeading[HeadingAxis])
+			{
+				DecodedHeading[HeadingAxis] = FFeatureVectorHelper::DecodeVector(PoseVector, FoundHeading[HeadingAxis]->GetChannelDataOffset(), FoundHeading[HeadingAxis]->ComponentStripping);
+
+				++HeadingAxisFoundNum;
+				if (HeadingAxisFoundNum == 2)
+				{
+					// we've found enough heading axis to compose a rotation
+					break;
+				}
+			}
+		}
+
+		if (HeadingAxisFoundNum > 0)
+		{
+			bool bAbleToReconstructMissingAxis = true;
+			if (HeadingAxisFoundNum == 2)
+			{
+				// reconstructing the missing axis
+				if (!FoundHeading[int32(EHeadingAxis::X)])
+				{
+					DecodedHeading[int32(EHeadingAxis::X)] = FVector::CrossProduct(DecodedHeading[int32(EHeadingAxis::Y)], DecodedHeading[int32(EHeadingAxis::Z)]);
+				}
+				else if (!FoundHeading[int32(EHeadingAxis::Y)])
+				{
+					DecodedHeading[int32(EHeadingAxis::Y)] = FVector::CrossProduct(DecodedHeading[int32(EHeadingAxis::Z)], DecodedHeading[int32(EHeadingAxis::X)]);
+				}
+				else // if (!FoundHeading[int32(EHeadingAxis::Z)])
+				{
+					DecodedHeading[int32(EHeadingAxis::Z)] = FVector::CrossProduct(DecodedHeading[int32(EHeadingAxis::X)], DecodedHeading[int32(EHeadingAxis::Y)]);
+				}
+			}
+			else 
+			{
+				check(HeadingAxisFoundNum == 1);
+			
+				// reconstructing the two missing axis
+				if (FoundHeading[int32(EHeadingAxis::X)])
+				{
+					DecodedHeading[int32(EHeadingAxis::Y)] = FVector::CrossProduct(FVector::ZAxisVector, DecodedHeading[int32(EHeadingAxis::X)]);
+					bAbleToReconstructMissingAxis &= DecodedHeading[int32(EHeadingAxis::Y)].Normalize();
+					DecodedHeading[int32(EHeadingAxis::Z)] = FVector::CrossProduct(DecodedHeading[int32(EHeadingAxis::X)], DecodedHeading[int32(EHeadingAxis::Y)]);
+				}
+				else if (FoundHeading[int32(EHeadingAxis::Y)])
+				{
+					DecodedHeading[int32(EHeadingAxis::X)] = FVector::CrossProduct(DecodedHeading[int32(EHeadingAxis::Y)], FVector::ZAxisVector);
+					bAbleToReconstructMissingAxis &= DecodedHeading[int32(EHeadingAxis::X)].Normalize();
+					DecodedHeading[int32(EHeadingAxis::Z)] = FVector::CrossProduct(DecodedHeading[int32(EHeadingAxis::X)], DecodedHeading[int32(EHeadingAxis::Y)]);
+				}
+				else // if (FoundHeading[int32(EHeadingAxis::Z)])
+				{
+					DecodedHeading[int32(EHeadingAxis::X)] = FVector::CrossProduct(FVector::YAxisVector, DecodedHeading[int32(EHeadingAxis::Z)]);
+					bAbleToReconstructMissingAxis &= DecodedHeading[int32(EHeadingAxis::X)].Normalize();
+					DecodedHeading[int32(EHeadingAxis::Y)] = FVector::CrossProduct(DecodedHeading[int32(EHeadingAxis::Z)], DecodedHeading[int32(EHeadingAxis::X)]);
+				}
+			}
+
+			if (bAbleToReconstructMissingAxis)
+			{
+				// RotMatrix is the rotation matrix from time zero (OriginTimeOffset) to time SampleTimeOffset, so by composing it with GetRootTransform().GetRotation(),
+				// world rotation associated to the time zero, we can calcualte the world rotation at time SampleTimeOffset
+				const FMatrix RotMatrix(DecodedHeading[int32(EHeadingAxis::X)], DecodedHeading[int32(EHeadingAxis::Y)], DecodedHeading[int32(EHeadingAxis::Z)], FVector::ZeroVector);
+				const FQuat RotQuat(RotMatrix);
+				const FQuat RotQuatWorld = RotQuat * GetRootTransform().GetRotation();
+				return RotQuatWorld;
+			}
+		}
+
+		if (Mesh && SchemaBoneIdx > RootSchemaBoneIdx)
+		{
+			return Mesh->GetSocketTransform(Schema->BoneReferences[SchemaBoneIdx].BoneName).GetRotation();
+		}
+	}
+
+	return GetRootTransform().GetRotation();
+}
+
 const FTransform& FDebugDrawParams::GetRootTransform() const
 {
 	check(RootMotionTransform);
diff --git a/Engine/Plugins/Experimental/Animation/PoseSearch/Source/Runtime/Private/PoseSearchFeatureChannel_Heading.cpp b/Engine/Plugins/Experimental/Animation/PoseSearch/Source/Runtime/Private/PoseSearchFeatureChannel_Heading.cpp
--- a/Engine/Plugins/Experimental/Animation/PoseSearch/Source/Runtime/Private/PoseSearchFeatureChannel_Heading.cpp
+++ b/Engine/Plugins/Experimental/Animation/PoseSearch/Source/Runtime/Private/PoseSearchFeatureChannel_Heading.cpp
@@ -8,6 +8,31 @@
 #include "PoseSearch/PoseSearchSchema.h"
 #include "PoseSearchFeatureChannel_Position.h"
 
+void UPoseSearchFeatureChannel_Heading::FindOrAddToSchema(UPoseSearchSchema* Schema, float SampleTimeOffset, const FName& BoneName, EHeadingAxis HeadingAxis)
+{
+	if (!Schema->FindChannel([SampleTimeOffset, &BoneName, HeadingAxis](const UPoseSearchFeatureChannel* Channel) -> const UPoseSearchFeatureChannel_Heading*
+		{
+			if (const UPoseSearchFeatureChannel_Heading* Heading = Cast<UPoseSearchFeatureChannel_Heading>(Channel))
+			{
+				if (Heading->Bone.BoneName == BoneName && Heading->SampleTimeOffset == SampleTimeOffset && Heading->OriginTimeOffset == 0.f && Heading->HeadingAxis == HeadingAxis)
+				{
+					return Heading;
+				}
+			}
+			return nullptr;
+		}))
+	{
+		UPoseSearchFeatureChannel_Heading* Heading = NewObject<UPoseSearchFeatureChannel_Heading>(Schema, NAME_None, RF_Transient);
+		Heading->Bone.BoneName = BoneName;
+		Heading->Weight = 0.f;
+		Heading->SampleTimeOffset = SampleTimeOffset;
+		Heading->HeadingAxis = HeadingAxis;
+		// @todo: perhaps add a tunable color for injected channels
+		Heading->DebugColor = FLinearColor::Gray;
+		Schema->AddTemporaryChannel(Heading);
+	}
+}
+
 void UPoseSearchFeatureChannel_Heading::Finalize(UPoseSearchSchema* Schema)
 {
 	ChannelDataOffset = Schema->SchemaCardinality;
@@ -21,6 +46,15 @@ void UPoseSearchFeatureChannel_Heading::AddDependentChannels(UPoseSearchSchema*
 	if (Schema->bInjectAdditionalDebugChannels)
 	{
 		UPoseSearchFeatureChannel_Position::FindOrAddToSchema(Schema, SampleTimeOffset, Bone.BoneName);
+		if (!FMath::IsNearlyZero(OriginTimeOffset))
+		{
+			// adding the position OriginTimeOffset seconds ahead
+			UPoseSearchFeatureChannel_Position::FindOrAddToSchema(Schema, OriginTimeOffset, Bone.BoneName);
+			
+			// adding the rotation (X, Y axis) OriginTimeOffset seconds ahead
+			UPoseSearchFeatureChannel_Heading::FindOrAddToSchema(Schema, OriginTimeOffset, Bone.BoneName, EHeadingAxis::X);
+			UPoseSearchFeatureChannel_Heading::FindOrAddToSchema(Schema, OriginTimeOffset, Bone.BoneName, EHeadingAxis::Y);
+		}
 	}
 }
 
@@ -73,7 +107,7 @@ void UPoseSearchFeatureChannel_Heading::DebugDraw(const UE::PoseSearch::FDebugDr
 	using namespace UE::PoseSearch;
 
 	const FColor Color = DebugColor.ToFColor(true);
-	const FVector BoneHeading = DrawParams.GetRootTransform().GetRotation().RotateVector(FFeatureVectorHelper::DecodeVector(PoseVector, ChannelDataOffset, ComponentStripping));
+	const FVector BoneHeading = DrawParams.ExtractRotation(PoseVector, OriginTimeOffset).RotateVector(FFeatureVectorHelper::DecodeVector(PoseVector, ChannelDataOffset, ComponentStripping));
 	const FVector BonePos = DrawParams.ExtractPosition(PoseVector, SampleTimeOffset, SchemaBoneIdx);
 
 	DrawParams.DrawPoint(BonePos, Color, 3.f);
diff --git a/Engine/Plugins/Experimental/Animation/PoseSearch/Source/Runtime/Private/PoseSearchFeatureChannel_Heading.h b/Engine/Plugins/Experimental/Animation/PoseSearch/Source/Runtime/Private/PoseSearchFeatureChannel_Heading.h
--- a/Engine/Plugins/Experimental/Animation/PoseSearch/Source/Runtime/Private/PoseSearchFeatureChannel_Heading.h
+++ b/Engine/Plugins/Experimental/Animation/PoseSearch/Source/Runtime/Private/PoseSearchFeatureChannel_Heading.h
@@ -73,4 +73,6 @@ public:
 #endif
 
 	FVector GetAxis(const FQuat& Rotation) const;
+
+	static void FindOrAddToSchema(UPoseSearchSchema* Schema, float SampleTimeOffset, const FName& BoneName = NAME_None, EHeadingAxis HeadingAxis = EHeadingAxis::X);
 };
\ No newline at end of file
diff --git a/Engine/Plugins/Experimental/Animation/PoseSearch/Source/Runtime/Public/PoseSearch/PoseSearchContext.h b/Engine/Plugins/Experimental/Animation/PoseSearch/Source/Runtime/Public/PoseSearch/PoseSearchContext.h
--- a/Engine/Plugins/Experimental/Animation/PoseSearch/Source/Runtime/Public/PoseSearch/PoseSearchContext.h
+++ b/Engine/Plugins/Experimental/Animation/PoseSearch/Source/Runtime/Public/PoseSearch/PoseSearchContext.h
@@ -114,6 +114,9 @@ struct POSESEARCH_API FDebugDrawParams
 
 	FVector ExtractPosition(TConstArrayView<float> PoseVector, const UPoseSearchFeatureChannel_Position* Position) const;
 	FVector ExtractPosition(TConstArrayView<float> PoseVector, float SampleTimeOffset, int8 SchemaBoneIdx = RootSchemaBoneIdx, EPermutationTimeType PermutationTimeType = EPermutationTimeType::UseSampleTime) const;
+
+	FQuat ExtractRotation(TConstArrayView<float> PoseVector, float SampleTimeOffset, int8 SchemaBoneIdx = RootSchemaBoneIdx) const;
+
 	const FTransform& GetRootTransform() const;
 
 	void DrawLine(const FVector& LineStart, const FVector& LineEnd, const FColor& Color, float Thickness = 0.f) const;
