Enable inertialization for blend stack graph node. Also moves inertialization out of motion matching/pose search.

diff --git a/Engine/Plugins/Experimental/Animation/PoseSearch/Source/Runtime/Private/AnimNode_BlendStack.cpp b/Engine/Plugins/Experimental/Animation/PoseSearch/Source/Runtime/Private/AnimNode_BlendStack.cpp
--- a/Engine/Plugins/Experimental/Animation/PoseSearch/Source/Runtime/Private/AnimNode_BlendStack.cpp
+++ b/Engine/Plugins/Experimental/Animation/PoseSearch/Source/Runtime/Private/AnimNode_BlendStack.cpp
@@ -9,6 +9,7 @@
 #include "Animation/AnimMontage.h"
 #include "PoseSearch/PoseSearchDefines.h"
 #include "PoseSearch/AnimNode_BlendStackInput.h"
+#include "Animation/AnimNode_Inertialization.h"
 
 #include UE_INLINE_GENERATED_CPP_BY_NAME(AnimNode_BlendStack)
 
@@ -17,6 +18,8 @@ TAutoConsoleVariable<int32> CVarAnimBlendStackEnable(TEXT("a.AnimNode.BlendStack
 TAutoConsoleVariable<int32> CVarAnimBlendStackPruningEnable(TEXT("a.AnimNode.BlendStack.Pruning.Enable"), 1, TEXT("Enable / Disable Blend Stack Pruning"));
 #endif
 
+#define LOCTEXT_NAMESPACE "AnimNode_BlendStack"
+
 /////////////////////////////////////////////////////
 // FPoseSearchAnimPlayer
 void FPoseSearchAnimPlayer::Initialize(UAnimationAsset* AnimationAsset, float AccumulatedTime, bool bLoop, bool bMirrored, UMirrorDataTable* MirrorDataTable, float BlendTime, float RootBoneBlendTime, const UBlendProfile* BlendProfile, EAlphaBlendOption InBlendOption, FVector BlendParameters, float PlayRate, int32 InPoseLinkIdx)
@@ -557,8 +560,41 @@ float FAnimNode_BlendStack_Standalone::GetAccumulatedTime() const
 	return AnimPlayers.IsEmpty() ? 0.f : AnimPlayers.First().GetAccumulatedTime();
 }
 
-void FAnimNode_BlendStack_Standalone::BlendTo(const FAnimationUpdateContext& Context, UAnimationAsset* AnimationAsset, float AccumulatedTime, bool bLoop, bool bMirrored, UMirrorDataTable* MirrorDataTable, float BlendTime, float RootBoneBlendTime, const UBlendProfile* BlendProfile, EAlphaBlendOption BlendOption, FVector BlendParameters, float PlayRate)
+static void RequestInertialBlend(const FAnimationUpdateContext& Context, float BlendTime, const UBlendProfile* BlendProfile, EAlphaBlendOption BlendOption)
 {
+	if (BlendTime > 0.0f)
+	{
+		UE::Anim::IInertializationRequester* InertializationRequester = Context.GetMessage<UE::Anim::IInertializationRequester>();
+		if (InertializationRequester)
+		{
+			FInertializationRequest Request;
+			Request.Duration = BlendTime;
+			Request.BlendProfile = BlendProfile;
+			Request.bUseBlendMode = true;
+			Request.BlendMode = BlendOption;
+#if ANIM_TRACE_ENABLED
+			Request.Description = LOCTEXT("InertializationRequestDescription", "Blend Stack");
+			Request.NodeId = Context.GetCurrentNodeId();
+			Request.AnimInstance = Context.AnimInstanceProxy->GetAnimInstanceObject();
+#endif
+
+			InertializationRequester->RequestInertialization(Request);
+		}
+		else
+		{
+			FAnimNode_Inertialization::LogRequestError(Context, Context.GetCurrentNodeId());
+		}
+	}
+}
+
+void FAnimNode_BlendStack_Standalone::BlendTo(const FAnimationUpdateContext& Context, UAnimationAsset* AnimationAsset, float AccumulatedTime, bool bLoop, bool bMirrored, UMirrorDataTable* MirrorDataTable, float BlendTime, float RootBoneBlendTime, const UBlendProfile* BlendProfile, EAlphaBlendOption BlendOption, bool bUseInertialBlend, FVector BlendParameters, float PlayRate)
+{
+	if (bUseInertialBlend)
+	{
+		RequestInertialBlend(Context, BlendTime, BlendProfile, BlendOption);
+		BlendTime = 0.0f;
+	}
+
 	AnimPlayers.PushFirst(FPoseSearchAnimPlayer());
 	FPoseSearchAnimPlayer& AnimPlayer = AnimPlayers.First();
 	AnimPlayer.Initialize(AnimationAsset, AccumulatedTime, bLoop, bMirrored, MirrorDataTable, BlendTime, RootBoneBlendTime, BlendProfile, BlendOption, BlendParameters, PlayRate, GetNextPoseLinkIndex());
@@ -666,7 +702,7 @@ void FAnimNode_BlendStack::UpdateAssetPlayer(const FAnimationUpdateContext& Cont
 
 		if (bExecuteBlendTo)
 		{
-			BlendTo(Context, AnimationAsset, AnimationTime, bLoop, bMirrored, MirrorDataTable.Get(), BlendTime, RootBoneBlendTime, BlendProfile, BlendOption, BlendParameters, WantedPlayRate);
+			BlendTo(Context, AnimationAsset, AnimationTime, bLoop, bMirrored, MirrorDataTable.Get(), BlendTime, RootBoneBlendTime, BlendProfile, BlendOption, bUseInertialBlend, BlendParameters, WantedPlayRate);
 		}
 	}
 	
@@ -684,3 +720,5 @@ void FBlendStack_SampleGraphPoseLink::SetInputPosePlayer(FPoseSearchAnimPlayer&
 	// Link our anim player to the input pose node.
 	InputNode->Player = &Player;
 }
+
+#undef LOCTEXT_NAMESPACE
\ No newline at end of file
diff --git a/Engine/Plugins/Experimental/Animation/PoseSearch/Source/Runtime/Private/AnimNode_MotionMatching.cpp b/Engine/Plugins/Experimental/Animation/PoseSearch/Source/Runtime/Private/AnimNode_MotionMatching.cpp
--- a/Engine/Plugins/Experimental/Animation/PoseSearch/Source/Runtime/Private/AnimNode_MotionMatching.cpp
+++ b/Engine/Plugins/Experimental/Animation/PoseSearch/Source/Runtime/Private/AnimNode_MotionMatching.cpp
@@ -179,7 +179,7 @@ void FAnimNode_MotionMatching::UpdateAssetPlayer(const FAnimationUpdateContext&
 				const float RootBoneBlendTime = YawFromAnimationBlendRate < 0.f ? BlendTime : 0.f;
 				BlendStackNode.BlendTo(Context, DatabaseAsset->GetAnimationAsset(), MotionMatchingState.CurrentSearchResult.AssetTime,
 					DatabaseAsset->IsLooping(), SearchIndexAsset->bMirrored, CurrentResultDatabase->Schema->MirrorDataTable.Get(),
-					BlendTime, RootBoneBlendTime, BlendProfile, BlendOption, SearchIndexAsset->BlendParameters, MotionMatchingState.WantedPlayRate);
+					BlendTime, RootBoneBlendTime, BlendProfile, BlendOption, bUseInertialBlend, SearchIndexAsset->BlendParameters, MotionMatchingState.WantedPlayRate);
 			}
 		}
 	}
diff --git a/Engine/Plugins/Experimental/Animation/PoseSearch/Source/Runtime/Private/PoseSearchLibrary.cpp b/Engine/Plugins/Experimental/Animation/PoseSearch/Source/Runtime/Private/PoseSearchLibrary.cpp
--- a/Engine/Plugins/Experimental/Animation/PoseSearch/Source/Runtime/Private/PoseSearchLibrary.cpp
+++ b/Engine/Plugins/Experimental/Animation/PoseSearch/Source/Runtime/Private/PoseSearchLibrary.cpp
@@ -57,40 +57,8 @@ void FMotionMatchingState::AdjustAssetTime(float AssetTime)
 	CurrentSearchResult.Update(AssetTime);
 }
 
-static void RequestInertialBlend(const FAnimationUpdateContext& Context, float BlendTime)
-{
-	// Use inertial blending to smooth over the transition
-	// It would be cool in the future to adjust the blend time by amount of dissimilarity, but we'll need a standardized distance metric first.
-	if (BlendTime > 0.0f)
-	{
-		UE::Anim::IInertializationRequester* InertializationRequester = Context.GetMessage<UE::Anim::IInertializationRequester>();
-		if (InertializationRequester)
-		{
-			FInertializationRequest Request;
-			Request.Duration = BlendTime;
-#if ANIM_TRACE_ENABLED
-			Request.Description = LOCTEXT("InertializationRequestDescription", "Motion Matching");
-			Request.NodeId = Context.GetCurrentNodeId();
-			Request.AnimInstance = Context.AnimInstanceProxy->GetAnimInstanceObject();
-#endif
-
-			InertializationRequester->RequestInertialization(Request);
-		}
-		else
-		{
-			FAnimNode_Inertialization::LogRequestError(Context, Context.GetCurrentNodeId());
-		}
-	}
-}
-
 void FMotionMatchingState::JumpToPose(const FAnimationUpdateContext& Context, const UE::PoseSearch::FSearchResult& Result, int32 MaxActiveBlends, float BlendTime)
 {
-	// requesting inertial blending only if blendstack is disabled
-	if (MaxActiveBlends <= 0)
-	{
-		RequestInertialBlend(Context, BlendTime);
-	}
-
 	// Remember which pose and sequence we're playing from the database
 	CurrentSearchResult = Result;
 
diff --git a/Engine/Plugins/Experimental/Animation/PoseSearch/Source/Runtime/Public/PoseSearch/AnimNode_BlendStack.h b/Engine/Plugins/Experimental/Animation/PoseSearch/Source/Runtime/Public/PoseSearch/AnimNode_BlendStack.h
--- a/Engine/Plugins/Experimental/Animation/PoseSearch/Source/Runtime/Public/PoseSearch/AnimNode_BlendStack.h
+++ b/Engine/Plugins/Experimental/Animation/PoseSearch/Source/Runtime/Public/PoseSearch/AnimNode_BlendStack.h
@@ -101,7 +101,7 @@ struct POSESEARCH_API FAnimNode_BlendStack_Standalone : public FAnimNode_AssetPl
 	virtual void GatherDebugData(FNodeDebugData& DebugData) override;
 	// End of FAnimNode_Base interface
 
-	void BlendTo(const FAnimationUpdateContext& Context, UAnimationAsset* AnimationAsset, float AccumulatedTime = 0.f, bool bLoop = false, bool bMirrored = false, UMirrorDataTable* MirrorDataTable = nullptr, float BlendTime = 0.2f, float RootBoneBlendTime = -1.f, const UBlendProfile* BlendProfile = nullptr, EAlphaBlendOption BlendOption = EAlphaBlendOption::Linear, FVector BlendParameters = FVector::Zero(), float PlayRate = 1.f);
+	void BlendTo(const FAnimationUpdateContext& Context, UAnimationAsset* AnimationAsset, float AccumulatedTime = 0.f, bool bLoop = false, bool bMirrored = false, UMirrorDataTable* MirrorDataTable = nullptr, float BlendTime = 0.2f, float RootBoneBlendTime = -1.f, const UBlendProfile* BlendProfile = nullptr, EAlphaBlendOption BlendOption = EAlphaBlendOption::Linear, bool bUseInertialBlend = false, FVector BlendParameters = FVector::Zero(), float PlayRate = 1.f);
 	void UpdatePlayRate(float PlayRate);
 	void Reset();
 
@@ -112,6 +112,8 @@ struct POSESEARCH_API FAnimNode_BlendStack_Standalone : public FAnimNode_AssetPl
 
 	int32 GetNextPoseLinkIndex();
 
+	int32 GetMaxActiveBlends() const { return MaxActiveBlends; }
+
 private:
 
 	void InitializeSample(const FAnimationInitializeContext& Context, FPoseSearchAnimPlayer& SamplePlayer);
@@ -179,6 +181,10 @@ struct POSESEARCH_API FAnimNode_BlendStack : public FAnimNode_BlendStack_Standal
 	UPROPERTY(EditAnywhere, Category = Settings, meta = (PinHiddenByDefault))
 	TObjectPtr<UMirrorDataTable> MirrorDataTable;
 
+	// tunable animation transition blend time 
+	UPROPERTY(EditAnywhere, Category = Settings)
+	bool bUseInertialBlend = false;
+
 	// Reset the blend stack if it has become relevant to the graph after not being updated on previous frames.
 	UPROPERTY(EditAnywhere, Category = Settings)
 	bool bResetOnBecomingRelevant = true;
diff --git a/Engine/Plugins/Experimental/Animation/PoseSearch/Source/Runtime/Public/PoseSearch/AnimNode_MotionMatching.h b/Engine/Plugins/Experimental/Animation/PoseSearch/Source/Runtime/Public/PoseSearch/AnimNode_MotionMatching.h
--- a/Engine/Plugins/Experimental/Animation/PoseSearch/Source/Runtime/Public/PoseSearch/AnimNode_MotionMatching.h
+++ b/Engine/Plugins/Experimental/Animation/PoseSearch/Source/Runtime/Public/PoseSearch/AnimNode_MotionMatching.h
@@ -95,6 +95,9 @@ private:
 	UPROPERTY(EditAnywhere, Category = Settings, meta = (PinHiddenByDefault, ClampMin = "0.2", ClampMax = "3.0", UIMin = "0.2", UIMax = "3.0"))
 	FFloatInterval PlayRate = FFloatInterval(1.f, 1.f);
 
+	UPROPERTY(EditAnywhere, Category = Settings, Category = Settings)
+	bool bUseInertialBlend = false;
+
 	// Reset the motion matching selection state if it has become relevant to the graph after not being updated on previous frames.
 	UPROPERTY(EditAnywhere, Category = Settings)
 	bool bResetOnBecomingRelevant = true;
