Allow multiple blend stack input pose nodes for layering.

diff --git a/Engine/Plugins/Experimental/Animation/PoseSearch/Source/Editor/Private/AnimGraphNode_BlendStack.cpp b/Engine/Plugins/Experimental/Animation/PoseSearch/Source/Editor/Private/AnimGraphNode_BlendStack.cpp
--- a/Engine/Plugins/Experimental/Animation/PoseSearch/Source/Editor/Private/AnimGraphNode_BlendStack.cpp
+++ b/Engine/Plugins/Experimental/Animation/PoseSearch/Source/Editor/Private/AnimGraphNode_BlendStack.cpp
@@ -50,9 +50,9 @@ void UAnimGraphNode_BlendStack::BakeDataDuringCompilation(class FCompilerResults
 
 void UAnimGraphNode_BlendStack_Base::ExpandGraphAndProcessNodes(
 	UEdGraph* SourceGraph, 
-	UAnimGraphNode_Base* SourceRootNode, UAnimGraphNode_Base* SourceInputNode,
+	UAnimGraphNode_Base* SourceRootNode, TArrayView<UAnimGraphNode_BlendStackInput*> InputNodes,
 	IAnimBlueprintCompilationContext& InCompilationContext, IAnimBlueprintGeneratedClassCompiledData& OutCompiledData,
-	UAnimGraphNode_Base*& OutRootNode, UAnimGraphNode_Base*& OutInputNode)
+	UAnimGraphNode_Base*& OutRootNode, TArrayView<UAnimGraphNode_BlendStackInput*> OutInputNodes)
 {
 	// Note: This is mostly copied from UAnimGraphNode_BlendSpaceGraphBase::ExpandGraphAndProcessNodes
 	
@@ -62,33 +62,40 @@ void UAnimGraphNode_BlendStack_Base::ExpandGraphAndProcessNodes(
 	UEdGraph* ClonedGraph = FEdGraphUtilities::CloneGraph(SourceGraph, InCompilationContext.GetConsolidatedEventGraph(), &InCompilationContext.GetMessageLog(), true);
 
 	// Grab all the animation nodes and find the corresponding 
-	// root node and input pose node in the cloned set
+	// root node and input pose nodes in the cloned set
 	TArray<UAnimGraphNode_Base*> AnimNodeList;
 
 	const UObject* SourceRootObject = InCompilationContext.GetMessageLog().FindSourceObject(SourceRootNode);
-	const UObject* SourceInputObject = InCompilationContext.GetMessageLog().FindSourceObject(SourceInputNode);
+	TArray<UObject*> SourceInputObjects;
+	SourceInputObjects.SetNum(InputNodes.Num());
+	for (int32 Index = 0; Index < InputNodes.Num(); ++Index)
+	{
+		SourceInputObjects[Index] = InCompilationContext.GetMessageLog().FindSourceObject(InputNodes[Index]);
+	}
+
 	for (auto NodeIt = ClonedGraph->Nodes.CreateIterator(); NodeIt; ++NodeIt)
 	{
 		UEdGraphNode* ClonedNode = *NodeIt;
-
 		if (UAnimGraphNode_Base* TestNode = Cast<UAnimGraphNode_Base>(ClonedNode))
 		{
 			AnimNodeList.Add(TestNode);
 
 			//@TODO: There ought to be a better way to determine this
-			const UObject* TestObject = InCompilationContext.GetMessageLog().FindSourceObject(TestNode);
+			UObject* TestObject = InCompilationContext.GetMessageLog().FindSourceObject(TestNode);
 			if (TestObject == SourceRootObject)
 			{
 				OutRootNode = TestNode;
 			}
-			if (TestObject == SourceInputObject)
+
+			int32 FoundIndex;
+			if (SourceInputObjects.Find(TestObject, FoundIndex))
 			{
-				OutInputNode = TestNode;
+				OutInputNodes[FoundIndex] = (UAnimGraphNode_BlendStackInput*)TestNode;
 			}
 		}
 	}
 
-	check(OutRootNode && OutInputNode);
+	check(OutRootNode && !OutInputNodes.Contains(nullptr));
 
 	// Run another expansion pass to catch the graph we just added (this is slightly wasteful)
 	InCompilationContext.ExpansionStep(ClonedGraph, false);
@@ -135,14 +142,26 @@ void UAnimGraphNode_BlendStack_Base::OnProcessDuringCompilation(IAnimBlueprintCo
 	// Allocate one sample graph per-active blend plus an extra one for the stored pose.
 	AnimNode->SampleGraphPoseLinks.SetNum(MaxBlendsNum + 1);
 
-	for(FBlendStack_SampleGraphPoseLink& GraphPoseLink : AnimNode->SampleGraphPoseLinks)
+	TArray<UAnimGraphNode_BlendStackInput*> InputNodes;
+	BoundGraph->GetNodesOfClass<UAnimGraphNode_BlendStackInput>(InputNodes);
+
+	int32 BlendStackAllocationIndex = InCompilationContext.GetAllocationIndexOfNode(this);
+	for(int32 Index = 0; Index < AnimNode->SampleGraphPoseLinks.Num(); ++ Index)
 	{
-		UAnimGraphNode_Base *RootNode, *InputNode;
-		ExpandGraphAndProcessNodes(SampleGraph, SampleGraph->ResultNode, SampleGraph->InputNode, InCompilationContext, OutCompiledData, RootNode, InputNode);
+		FBlendStack_SampleGraphPoseLink& GraphPoseLink = AnimNode->SampleGraphPoseLinks[Index];
+		UAnimGraphNode_Base* ClonedRootNode;
+		TArray<UAnimGraphNode_BlendStackInput*> ClonedInputNodes;
+		ClonedInputNodes.SetNum(InputNodes.Num());
+
+		ExpandGraphAndProcessNodes(SampleGraph, SampleGraph->ResultNode, InputNodes, InCompilationContext, OutCompiledData, ClonedRootNode, ClonedInputNodes);
 
 		// Blend stack node is potentially nested in the struct, so we can't use FPoseLinkMappingRecord. Patch at runtime instead.
-		GraphPoseLink.RootNodeIndex = InCompilationContext.GetAllocationIndexOfNode(RootNode);
-		GraphPoseLink.InputPoseNodeIndex = InCompilationContext.GetAllocationIndexOfNode(InputNode);
+		for (UAnimGraphNode_BlendStackInput* InputNode : ClonedInputNodes)
+		{
+			InputNode->Node.BlendStackAllocationIndex = BlendStackAllocationIndex;
+			InputNode->Node.SampleIndex = Index;
+		}
+		GraphPoseLink.RootNodeIndex = InCompilationContext.GetAllocationIndexOfNode(ClonedRootNode);
 	}
 }
 
diff --git a/Engine/Plugins/Experimental/Animation/PoseSearch/Source/Editor/Private/AnimGraphNode_BlendStackInput.cpp b/Engine/Plugins/Experimental/Animation/PoseSearch/Source/Editor/Private/AnimGraphNode_BlendStackInput.cpp
--- a/Engine/Plugins/Experimental/Animation/PoseSearch/Source/Editor/Private/AnimGraphNode_BlendStackInput.cpp
+++ b/Engine/Plugins/Experimental/Animation/PoseSearch/Source/Editor/Private/AnimGraphNode_BlendStackInput.cpp
@@ -1,6 +1,7 @@
 // Copyright Epic Games, Inc. All Rights Reserved.
 
 #include "AnimGraphNode_BlendStackInput.h"
+#include "AnimGraphNode_BlendStack.h"
 #include "AnimationBlendStackGraph.h"
 
 #define LOCTEXT_NAMESPACE "BlendStackInput"
@@ -18,12 +19,12 @@ FText UAnimGraphNode_BlendStackInput::GetMenuCategory() const
 
 bool UAnimGraphNode_BlendStackInput::CanUserDeleteNode() const
 {
-	return false;
+	return true;
 }
 
 bool UAnimGraphNode_BlendStackInput::CanDuplicateNode() const
 {
-	return false;
+	return true;
 }
 
 bool UAnimGraphNode_BlendStackInput::IsCompatibleWithGraph(UEdGraph const* Graph) const
diff --git a/Engine/Plugins/Experimental/Animation/PoseSearch/Source/Editor/Private/AnimationBlendStackGraphSchema.cpp b/Engine/Plugins/Experimental/Animation/PoseSearch/Source/Editor/Private/AnimationBlendStackGraphSchema.cpp
--- a/Engine/Plugins/Experimental/Animation/PoseSearch/Source/Editor/Private/AnimationBlendStackGraphSchema.cpp
+++ b/Engine/Plugins/Experimental/Animation/PoseSearch/Source/Editor/Private/AnimationBlendStackGraphSchema.cpp
@@ -32,8 +32,7 @@ void UAnimationBlendStackGraphSchema::CreateDefaultNodesForGraph(UEdGraph& Graph
 		SetNodeMetaData(InputNode, FNodeMetadata::DefaultGraphNode);
 
 		// Move input node to the left of the output pose.
-		InputNode->NodePosX -= 200; 
-		TypedGraph->InputNode = InputNode;
+		InputNode->NodePosX -= 200;
 	}
 
 	// Connect input node to output pose.
diff --git a/Engine/Plugins/Experimental/Animation/PoseSearch/Source/Editor/Public/AnimGraphNode_BlendStack.h b/Engine/Plugins/Experimental/Animation/PoseSearch/Source/Editor/Public/AnimGraphNode_BlendStack.h
--- a/Engine/Plugins/Experimental/Animation/PoseSearch/Source/Editor/Public/AnimGraphNode_BlendStack.h
+++ b/Engine/Plugins/Experimental/Animation/PoseSearch/Source/Editor/Public/AnimGraphNode_BlendStack.h
@@ -6,6 +6,7 @@
 #include "PoseSearch/AnimNode_BlendStack.h"
 #include "AnimGraphNode_BlendStack.generated.h"
 
+class UAnimGraphNode_BlendStackInput;
 
 UCLASS(MinimalAPI, Abstract, Experimental)
 class UAnimGraphNode_BlendStack_Base : public UAnimGraphNode_AssetPlayerBase
@@ -30,9 +31,9 @@ class UAnimGraphNode_BlendStack_Base : public UAnimGraphNode_AssetPlayerBase
 	// Helper function for compilation
 	void ExpandGraphAndProcessNodes(
 		UEdGraph* SourceGraph, 
-		UAnimGraphNode_Base* SourceRootNode, UAnimGraphNode_Base* SourceInputNode,
+		UAnimGraphNode_Base* SourceRootNode, TArrayView<UAnimGraphNode_BlendStackInput*> SourceInputNode,
 		IAnimBlueprintCompilationContext& InCompilationContext, IAnimBlueprintGeneratedClassCompiledData& OutCompiledData,
-		UAnimGraphNode_Base*& OutRootNode, UAnimGraphNode_Base*& OutInputNode);
+		UAnimGraphNode_Base*& OutRootNode, TArrayView<UAnimGraphNode_BlendStackInput*> OutInputNodes);
 
 	virtual TArray<UEdGraph*> GetSubGraphs() const override;
 
diff --git a/Engine/Plugins/Experimental/Animation/PoseSearch/Source/Editor/Public/AnimationBlendStackGraph.h b/Engine/Plugins/Experimental/Animation/PoseSearch/Source/Editor/Public/AnimationBlendStackGraph.h
--- a/Engine/Plugins/Experimental/Animation/PoseSearch/Source/Editor/Public/AnimationBlendStackGraph.h
+++ b/Engine/Plugins/Experimental/Animation/PoseSearch/Source/Editor/Public/AnimationBlendStackGraph.h
@@ -21,9 +21,5 @@ public:
 	// Result node within the state's animation graph
 	UPROPERTY()
 	TObjectPtr<class UAnimGraphNode_Root> ResultNode = nullptr;
-
-	// Result node within the state's animation graph
-	UPROPERTY()
-	TObjectPtr<class UAnimGraphNode_BlendStackInput> InputNode = nullptr;
 };
 
diff --git a/Engine/Plugins/Experimental/Animation/PoseSearch/Source/Runtime/Private/AnimNode_BlendStack.cpp b/Engine/Plugins/Experimental/Animation/PoseSearch/Source/Runtime/Private/AnimNode_BlendStack.cpp
--- a/Engine/Plugins/Experimental/Animation/PoseSearch/Source/Runtime/Private/AnimNode_BlendStack.cpp
+++ b/Engine/Plugins/Experimental/Animation/PoseSearch/Source/Runtime/Private/AnimNode_BlendStack.cpp
@@ -410,17 +410,12 @@ void FAnimNode_BlendStack_Standalone::Initialize_AnyThread(const FAnimationIniti
 	Reset();
 
 	IAnimClassInterface* AnimBlueprintClass = Context.GetAnimClass();
+	check(AnimBlueprintClass);
 	if (SampleGraphPoseLinks.IsEmpty() == false)
 	{
 		// Patch our pose links
 		for (FBlendStack_SampleGraphPoseLink& GraphPoseLink : SampleGraphPoseLinks)
 		{
-			if (GraphPoseLink.InputPoseNodeIndex != INDEX_NONE)
-			{
-				GraphPoseLink.InputPose.LinkID = AnimBlueprintClass->GetAnimNodeProperties().Num() - 1 - GraphPoseLink.InputPoseNodeIndex;
-				GraphPoseLink.InputPose.AttemptRelink(Context);
-			}
-
 			if (GraphPoseLink.RootNodeIndex != INDEX_NONE)
 			{
 				GraphPoseLink.Root.LinkID = AnimBlueprintClass->GetAnimNodeProperties().Num() - 1 - GraphPoseLink.RootNodeIndex;
@@ -711,14 +706,11 @@ void FAnimNode_BlendStack::UpdateAssetPlayer(const FAnimationUpdateContext& Cont
 	Super::UpdateAssetPlayer(Context);
 }
 
-void FBlendStack_SampleGraphPoseLink::SetInputPosePlayer(FPoseSearchAnimPlayer& Player)
+void FBlendStack_SampleGraphPoseLink::SetInputPosePlayer(FPoseSearchAnimPlayer& InPlayer)
 {
 	// Because our anim players may get reallocated, or change indices due to push/pops,
-	// we must call this before every operation that might end up needing the anim player through the graph's input node.
-	// @todo: There's probably a better way to do this.
-	FAnimNode_BlendStackInput* InputNode = static_cast<FAnimNode_BlendStackInput*>(InputPose.GetLinkNode());
-	// Link our anim player to the input pose node.
-	InputNode->Player = &Player;
+	// we must call this before every operation that might end up needing the anim player through the graph's input nodes.
+	Player = &InPlayer;
 }
 
 #undef LOCTEXT_NAMESPACE
\ No newline at end of file
diff --git a/Engine/Plugins/Experimental/Animation/PoseSearch/Source/Runtime/Private/AnimNode_BlendStackInput.cpp b/Engine/Plugins/Experimental/Animation/PoseSearch/Source/Runtime/Private/AnimNode_BlendStackInput.cpp
--- a/Engine/Plugins/Experimental/Animation/PoseSearch/Source/Runtime/Private/AnimNode_BlendStackInput.cpp
+++ b/Engine/Plugins/Experimental/Animation/PoseSearch/Source/Runtime/Private/AnimNode_BlendStackInput.cpp
@@ -8,13 +8,36 @@
 
 void FAnimNode_BlendStackInput::Evaluate_AnyThread(FPoseContext& Output)
 {
-	check(Player);
-	Player->Evaluate_AnyThread(Output);
+	check(Player && *Player);
+	(*Player)->Evaluate_AnyThread(Output);
+}
+
+static FAnimNode_BlendStack_Standalone* GetBlendStackNodeFromIndex(const FAnimationBaseContext& Context, IAnimClassInterface* AnimBlueprintClass, int32 NodeIndex)
+{
+	const TArray<FStructProperty*>& AnimNodeProperties = AnimBlueprintClass->GetAnimNodeProperties();
+	check(AnimNodeProperties.IsValidIndex(NodeIndex));
+
+	FStructProperty* LinkedProperty = AnimNodeProperties[NodeIndex];
+	void* LinkedNodePtr = LinkedProperty->ContainerPtrToValuePtr<void>(Context.AnimInstanceProxy->GetAnimInstanceObject());
+	return (FAnimNode_BlendStack_Standalone*)LinkedNodePtr;
 }
 
 void FAnimNode_BlendStackInput::Update_AnyThread(const FAnimationUpdateContext& Context)
 {
-	check(Player);
-	Player->Update_AnyThread(Context);
+	// If there's no player, use the allocated blend stack index to get a reference to it.
+	// This should only happen on the first ever update of this node.
+	if (Player == nullptr)
+	{
+		IAnimClassInterface* AnimBlueprintClass = Context.GetAnimClass();
+		check(AnimBlueprintClass);
+		// Allocation index is the inverse of the node index.
+		const int32 BlendStackNodeIndex = AnimBlueprintClass->GetAnimNodeProperties().Num() - 1 - BlendStackAllocationIndex;
+		check(Context.AnimInstanceProxy);
+		FAnimNode_BlendStack_Standalone* BlendStackNode = GetBlendStackNodeFromIndex(Context, AnimBlueprintClass, BlendStackNodeIndex);
+		Player = &BlendStackNode->SampleGraphPoseLinks[SampleIndex].Player;
+	}
+
+	check(Player && *Player);
+	(*Player)->Update_AnyThread(Context);
 }
 
diff --git a/Engine/Plugins/Experimental/Animation/PoseSearch/Source/Runtime/Public/PoseSearch/AnimNode_BlendStack.h b/Engine/Plugins/Experimental/Animation/PoseSearch/Source/Runtime/Public/PoseSearch/AnimNode_BlendStack.h
--- a/Engine/Plugins/Experimental/Animation/PoseSearch/Source/Runtime/Public/PoseSearch/AnimNode_BlendStack.h
+++ b/Engine/Plugins/Experimental/Animation/PoseSearch/Source/Runtime/Public/PoseSearch/AnimNode_BlendStack.h
@@ -70,16 +70,14 @@ struct FBlendStack_SampleGraphPoseLink
 {
 	GENERATED_BODY()
 
-	UPROPERTY()
-	int32 InputPoseNodeIndex = INDEX_NONE;
 	UPROPERTY()
 	int32 RootNodeIndex = INDEX_NONE;
 
-	FPoseLink InputPose;
 	FPoseLink Root;
+	FPoseSearchAnimPlayer* Player = nullptr;
 	FGraphTraversalCounter CacheBoneCounter;
 
-	void SetInputPosePlayer(FPoseSearchAnimPlayer& Player);
+	void SetInputPosePlayer(FPoseSearchAnimPlayer& InPlayer);
 	void EvaluatePlayer(FPoseContext& Output, FPoseSearchAnimPlayer& SamplePlayer);
 	void ConditionalCacheBones(const FAnimationBaseContext& Output);
 };
diff --git a/Engine/Plugins/Experimental/Animation/PoseSearch/Source/Runtime/Public/PoseSearch/AnimNode_BlendStackInput.h b/Engine/Plugins/Experimental/Animation/PoseSearch/Source/Runtime/Public/PoseSearch/AnimNode_BlendStackInput.h
--- a/Engine/Plugins/Experimental/Animation/PoseSearch/Source/Runtime/Public/PoseSearch/AnimNode_BlendStackInput.h
+++ b/Engine/Plugins/Experimental/Animation/PoseSearch/Source/Runtime/Public/PoseSearch/AnimNode_BlendStackInput.h
@@ -18,7 +18,13 @@ public:
 	virtual void Update_AnyThread(const FAnimationUpdateContext& Context) override;
 	virtual void Evaluate_AnyThread(FPoseContext& Output) override;
 
+	UPROPERTY()
+	int32 SampleIndex = INDEX_NONE;
+
+	UPROPERTY()
+	int32 BlendStackAllocationIndex = INDEX_NONE;
+
 	// The player is guaranteed to be valid for the whole duration of update/eval.
-	struct FPoseSearchAnimPlayer* Player;
+	struct FPoseSearchAnimPlayer** Player;
 };
 
