motion matching - promoting BuildIndex Failed to UE_LOG Error(s)

diff --git a/Engine/Plugins/Experimental/Animation/PoseSearch/Source/Runtime/Private/PoseSearchDerivedData.cpp b/Engine/Plugins/Experimental/Animation/PoseSearch/Source/Runtime/Private/PoseSearchDerivedData.cpp
--- a/Engine/Plugins/Experimental/Animation/PoseSearch/Source/Runtime/Private/PoseSearchDerivedData.cpp
+++ b/Engine/Plugins/Experimental/Animation/PoseSearch/Source/Runtime/Private/PoseSearchDerivedData.cpp
@@ -926,11 +926,11 @@ void FPoseSearchDatabaseAsyncCacheTask::OnGetComplete(UE::DerivedData::FCacheGet
 					{
 						if (IndexBaseDatabase == Database)
 						{
-							UE_LOG(LogPoseSearch, Log, TEXT("%s - %s BuildIndex Failed"), *LexToString(FullIndexKey.Hash), *Database->GetName());
+							UE_LOG(LogPoseSearch, Error, TEXT("%s - %s BuildIndex Failed becasue of invalid Schema"), *LexToString(FullIndexKey.Hash), *Database->GetName());
 						}
 						else
 						{
-							UE_LOG(LogPoseSearch, Log, TEXT("%s - %s BuildIndex Failed because of dependent database fail '%s'"), *LexToString(FullIndexKey.Hash), *Database->GetName(), *IndexBaseDatabase->GetName());
+							UE_LOG(LogPoseSearch, Error, TEXT("%s - %s BuildIndex Failed because dependent database '%s' has an invalid Schema"), *LexToString(FullIndexKey.Hash), *Database->GetName(), *IndexBaseDatabase->GetName());
 						}
 						SearchIndex.Reset();
 						return;
@@ -946,7 +946,7 @@ void FPoseSearchDatabaseAsyncCacheTask::OnGetComplete(UE::DerivedData::FCacheGet
 								if (DatabaseAsset->GetMirrorOption() == EPoseSearchMirrorOption::MirroredOnly || DatabaseAsset->GetMirrorOption() == EPoseSearchMirrorOption::UnmirroredAndMirrored)
 								{
 									// want to sample a mirrored asset
-									UE_LOG(LogPoseSearch, Log, TEXT("%s - %s BuildIndex Failed because '%s' requires a MirrorDataTable to sample mirrored animation assets"), *LexToString(FullIndexKey.Hash), *Database->GetName(), *IndexBaseDatabase->Schema->GetName());
+									UE_LOG(LogPoseSearch, Error, TEXT("%s - %s BuildIndex Failed because '%s' requires a MirrorDataTable to sample mirrored animation assets"), *LexToString(FullIndexKey.Hash), *Database->GetName(), *IndexBaseDatabase->Schema->GetName());
 									SearchIndex.Reset();
 									return;
 								}
@@ -974,19 +974,7 @@ void FPoseSearchDatabaseAsyncCacheTask::OnGetComplete(UE::DerivedData::FCacheGet
 					FDatabaseIndexingContext DbIndexingContext;
 					if (!DbIndexingContext.IndexDatabase(SearchIndexBase, *IndexBaseDatabase, Owner))
 					{
-						if (Owner.IsCanceled())
-						{
-							UE_LOG(LogPoseSearch, Log, TEXT("%s - %s BuildIndex Cancelled"), *LexToString(FullIndexKey.Hash), *Database->GetName());
-						}
-						else if (IndexBaseDatabase == Database)
-						{
-							UE_LOG(LogPoseSearch, Log, TEXT("%s - %s BuildIndex Failed"), *LexToString(FullIndexKey.Hash), *Database->GetName());
-						}
-						else
-						{
-							UE_LOG(LogPoseSearch, Log, TEXT("%s - %s BuildIndex Failed because of dependent database fail '%s'"), *LexToString(FullIndexKey.Hash), *Database->GetName(), *IndexBaseDatabase->GetName());
-						}
-
+						UE_LOG(LogPoseSearch, Log, TEXT("%s - %s BuildIndex Cancelled"), *LexToString(FullIndexKey.Hash), *Database->GetName());
 						SearchIndex.Reset();
 						return;
 					}
