motion matching - making FPoseSearchDatabaseAsyncCacheTask::OnGetComplete launched task relying on its enclosed TObjectPtr rather than the object TWeakObjectPtr

diff --git a/Engine/Plugins/Experimental/Animation/PoseSearch/Source/Runtime/Private/PoseSearchDerivedData.cpp b/Engine/Plugins/Experimental/Animation/PoseSearch/Source/Runtime/Private/PoseSearchDerivedData.cpp
--- a/Engine/Plugins/Experimental/Animation/PoseSearch/Source/Runtime/Private/PoseSearchDerivedData.cpp
+++ b/Engine/Plugins/Experimental/Animation/PoseSearch/Source/Runtime/Private/PoseSearchDerivedData.cpp
@@ -871,7 +871,7 @@ void FPoseSearchDatabaseAsyncCacheTask::OnGetComplete(UE::DerivedData::FCacheGet
 		FMemoryReaderView Reader(RawData);
 		Reader << SearchIndex;
 
-		check(Database->Schema && Database->Schema->IsValid());
+		check(Database != nullptr && Database->Schema && Database->Schema->IsValid());
 		// cache can be corrupted in case the version of the derived data cache has not being updated while 
 		// developing channels that changes their cardinality without impacting any asset properties
 		// so to account for this, we just reindex the database and update the associated DDC 
@@ -904,9 +904,16 @@ void FPoseSearchDatabaseAsyncCacheTask::OnGetComplete(UE::DerivedData::FCacheGet
 				// collecting all the databases that need to be built to gather their FSearchIndexBase
 				TArray<TObjectPtr<const UPoseSearchDatabase>> IndexBaseDatabases;
 				IndexBaseDatabases.Add(Database.Get()); // the first one is always this Database
-				if (Database->NormalizationSet)
+				if (!IndexBaseDatabases[0])
 				{
-					Database->NormalizationSet->AddUniqueDatabases(IndexBaseDatabases);
+					UE_LOG(LogPoseSearch, Log, TEXT("%s - BuildIndex Cancelled because associated Database weak pointer has been released."), *LexToString(FullIndexKey.Hash), *IndexBaseDatabases[0]->GetName());
+					SearchIndex.Reset();
+					return;
+				}
+
+				if (IndexBaseDatabases[0]->NormalizationSet)
+				{
+					IndexBaseDatabases[0]->NormalizationSet->AddUniqueDatabases(IndexBaseDatabases);
 				}
 
 				// @todo: DDC or parallelize this code
@@ -924,13 +931,13 @@ void FPoseSearchDatabaseAsyncCacheTask::OnGetComplete(UE::DerivedData::FCacheGet
 					// early out for invalid indexing conditions
 					if (!IndexBaseDatabase->Schema || !IndexBaseDatabase->Schema->IsValid() || IndexBaseDatabase->Schema->SchemaCardinality <= 0)
 					{
-						if (IndexBaseDatabase == Database)
+						if (IndexBaseIdx == 0)
 						{
-							UE_LOG(LogPoseSearch, Error, TEXT("%s - %s BuildIndex Failed becasue of invalid Schema"), *LexToString(FullIndexKey.Hash), *Database->GetName());
+							UE_LOG(LogPoseSearch, Error, TEXT("%s - %s BuildIndex Failed becasue of invalid Schema"), *LexToString(FullIndexKey.Hash), *IndexBaseDatabases[0]->GetName());
 						}
 						else
 						{
-							UE_LOG(LogPoseSearch, Error, TEXT("%s - %s BuildIndex Failed because dependent database '%s' has an invalid Schema"), *LexToString(FullIndexKey.Hash), *Database->GetName(), *IndexBaseDatabase->GetName());
+							UE_LOG(LogPoseSearch, Error, TEXT("%s - %s BuildIndex Failed because dependent database '%s' has an invalid Schema"), *LexToString(FullIndexKey.Hash), *IndexBaseDatabases[0]->GetName(), *IndexBaseDatabase->GetName());
 						}
 						SearchIndex.Reset();
 						return;
@@ -946,7 +953,7 @@ void FPoseSearchDatabaseAsyncCacheTask::OnGetComplete(UE::DerivedData::FCacheGet
 								if (DatabaseAsset->GetMirrorOption() == EPoseSearchMirrorOption::MirroredOnly || DatabaseAsset->GetMirrorOption() == EPoseSearchMirrorOption::UnmirroredAndMirrored)
 								{
 									// want to sample a mirrored asset
-									UE_LOG(LogPoseSearch, Error, TEXT("%s - %s BuildIndex Failed because '%s' requires a MirrorDataTable to sample mirrored animation assets"), *LexToString(FullIndexKey.Hash), *Database->GetName(), *IndexBaseDatabase->Schema->GetName());
+									UE_LOG(LogPoseSearch, Error, TEXT("%s - %s BuildIndex Failed because '%s' requires a MirrorDataTable to sample mirrored animation assets"), *LexToString(FullIndexKey.Hash), *IndexBaseDatabases[0]->GetName(), *IndexBaseDatabase->Schema->GetName());
 									SearchIndex.Reset();
 									return;
 								}
@@ -956,7 +963,7 @@ void FPoseSearchDatabaseAsyncCacheTask::OnGetComplete(UE::DerivedData::FCacheGet
 
 					if (Owner.IsCanceled())
 					{
-						UE_LOG(LogPoseSearch, Log, TEXT("%s - %s BuildIndex Cancelled"), *LexToString(FullIndexKey.Hash), *Database->GetName());
+						UE_LOG(LogPoseSearch, Log, TEXT("%s - %s BuildIndex Cancelled"), *LexToString(FullIndexKey.Hash), *IndexBaseDatabases[0]->GetName());
 						SearchIndex.Reset();
 						return;
 					}
@@ -966,7 +973,7 @@ void FPoseSearchDatabaseAsyncCacheTask::OnGetComplete(UE::DerivedData::FCacheGet
 
 					if (Owner.IsCanceled())
 					{
-						UE_LOG(LogPoseSearch, Log, TEXT("%s - %s BuildIndex Cancelled"), *LexToString(FullIndexKey.Hash), *Database->GetName());
+						UE_LOG(LogPoseSearch, Log, TEXT("%s - %s BuildIndex Cancelled"), *LexToString(FullIndexKey.Hash), *IndexBaseDatabases[0]->GetName());
 						SearchIndex.Reset();
 						return;
 					}
@@ -974,7 +981,7 @@ void FPoseSearchDatabaseAsyncCacheTask::OnGetComplete(UE::DerivedData::FCacheGet
 					FDatabaseIndexingContext DbIndexingContext;
 					if (!DbIndexingContext.IndexDatabase(SearchIndexBase, *IndexBaseDatabase, Owner))
 					{
-						UE_LOG(LogPoseSearch, Log, TEXT("%s - %s BuildIndex Cancelled"), *LexToString(FullIndexKey.Hash), *Database->GetName());
+						UE_LOG(LogPoseSearch, Log, TEXT("%s - %s BuildIndex Cancelled"), *LexToString(FullIndexKey.Hash), *IndexBaseDatabases[0]->GetName());
 						SearchIndex.Reset();
 						return;
 					}
@@ -982,50 +989,50 @@ void FPoseSearchDatabaseAsyncCacheTask::OnGetComplete(UE::DerivedData::FCacheGet
 
 				static_cast<FSearchIndexBase&>(SearchIndex) = SearchIndexBases[0];
 				
-				SearchIndex.PruneDuplicateValues(Database->PosePruningSimilarityThreshold, Database->Schema->SchemaCardinality);
+				SearchIndex.PruneDuplicateValues(IndexBaseDatabases[0]->PosePruningSimilarityThreshold, IndexBaseDatabases[0]->Schema->SchemaCardinality);
 
 				TArray<float> Deviation = FMeanDeviationCalculator::Calculate(SearchIndexBases, Schemas);
 
 				// Building FSearchIndex
-				PreprocessSearchIndexWeights(SearchIndex, Database->Schema, Deviation);
+				PreprocessSearchIndexWeights(SearchIndex, IndexBaseDatabases[0]->Schema, Deviation);
 				if (Owner.IsCanceled())
 				{
-					UE_LOG(LogPoseSearch, Log, TEXT("%s - %s BuildIndex Cancelled"), *LexToString(FullIndexKey.Hash), *Database->GetName());
+					UE_LOG(LogPoseSearch, Log, TEXT("%s - %s BuildIndex Cancelled"), *LexToString(FullIndexKey.Hash), *IndexBaseDatabases[0]->GetName());
 					SearchIndex.Reset();
 					return;
 				}
 
-				PreprocessSearchIndexPCAData(SearchIndex, Database->Schema->SchemaCardinality, Database->GetNumberOfPrincipalComponents(), Database->PoseSearchMode);
+				PreprocessSearchIndexPCAData(SearchIndex, IndexBaseDatabases[0]->Schema->SchemaCardinality, IndexBaseDatabases[0]->GetNumberOfPrincipalComponents(), IndexBaseDatabases[0]->PoseSearchMode);
 				if (Owner.IsCanceled())
 				{
-					UE_LOG(LogPoseSearch, Log, TEXT("%s - %s BuildIndex Cancelled"), *LexToString(FullIndexKey.Hash), *Database->GetName());
+					UE_LOG(LogPoseSearch, Log, TEXT("%s - %s BuildIndex Cancelled"), *LexToString(FullIndexKey.Hash), *IndexBaseDatabases[0]->GetName());
 					SearchIndex.Reset();
 					return;
 				}
 
-				SearchIndex.PruneDuplicatePCAValues(Database->PCAValuesPruningSimilarityThreshold, Database->GetNumberOfPrincipalComponents());
+				SearchIndex.PruneDuplicatePCAValues(IndexBaseDatabases[0]->PCAValuesPruningSimilarityThreshold, IndexBaseDatabases[0]->GetNumberOfPrincipalComponents());
 				if (Owner.IsCanceled())
 				{
-					UE_LOG(LogPoseSearch, Log, TEXT("%s - %s BuildIndex Cancelled"), *LexToString(FullIndexKey.Hash), *Database->GetName());
+					UE_LOG(LogPoseSearch, Log, TEXT("%s - %s BuildIndex Cancelled"), *LexToString(FullIndexKey.Hash), *IndexBaseDatabases[0]->GetName());
 					SearchIndex.Reset();
 					return;
 				}
 
-				PreprocessSearchIndexKDTree(SearchIndex, Database.Get());
+				PreprocessSearchIndexKDTree(SearchIndex, IndexBaseDatabases[0].Get());
 				if (Owner.IsCanceled())
 				{
-					UE_LOG(LogPoseSearch, Log, TEXT("%s - %s BuildIndex Cancelled"), *LexToString(FullIndexKey.Hash), *Database->GetName());
+					UE_LOG(LogPoseSearch, Log, TEXT("%s - %s BuildIndex Cancelled"), *LexToString(FullIndexKey.Hash), *IndexBaseDatabases[0]->GetName());
 					SearchIndex.Reset();
 					return;
 				}
 
 				// removing SearchIndex.Values and relying on FSearchIndex::GetReconstructedPoseValues to reconstruct the Values data from the PCAValues
-				if (Database->PoseSearchMode == EPoseSearchMode::PCAKDTree && Database->KDTreeQueryNumNeighbors <= 1)
+				if (IndexBaseDatabases[0]->PoseSearchMode == EPoseSearchMode::PCAKDTree && IndexBaseDatabases[0]->KDTreeQueryNumNeighbors <= 1)
 				{
 					SearchIndex.ResetValues();
 				}
 
-				UE_LOG(LogPoseSearch, Log, TEXT("%s - %s BuildIndex Succeeded"), *LexToString(FullIndexKey.Hash), *Database->GetName());
+				UE_LOG(LogPoseSearch, Log, TEXT("%s - %s BuildIndex Succeeded"), *LexToString(FullIndexKey.Hash), *IndexBaseDatabases[0]->GetName());
 
 				// putting SearchIndex to DDC
 				TArray<uint8> RawBytes;
@@ -1036,11 +1043,11 @@ void FPoseSearchDatabaseAsyncCacheTask::OnGetComplete(UE::DerivedData::FCacheGet
 
 				FCacheRecordBuilder Builder(FullIndexKey);
 				Builder.AddValue(Id, RawData);
-				GetCache().Put({ { { Database->GetPathName() }, Builder.Build() } }, Owner, [this, FullIndexKey](FCachePutResponse&& Response)
+				GetCache().Put({ { { IndexBaseDatabases[0]->GetPathName() }, Builder.Build() } }, Owner, [IndexBaseDatabases, FullIndexKey](FCachePutResponse&& Response)
 					{
 						if (Response.Status == EStatus::Error)
 						{
-							UE_LOG(LogPoseSearch, Log, TEXT("%s - %s Failed to store DDC"), *LexToString(FullIndexKey.Hash), *Database->GetName());
+							UE_LOG(LogPoseSearch, Log, TEXT("%s - %s Failed to store DDC"), *LexToString(FullIndexKey.Hash), *IndexBaseDatabases[0]->GetName());
 						}
 					});
 
