motion matching - fix for pose channels debug drawing (velocities where defaulting on character position rather than their relative bone position)

diff --git a/Engine/Plugins/Experimental/Animation/PoseSearch/Source/Runtime/Private/PoseSearchContext.cpp b/Engine/Plugins/Experimental/Animation/PoseSearch/Source/Runtime/Private/PoseSearchContext.cpp
--- a/Engine/Plugins/Experimental/Animation/PoseSearch/Source/Runtime/Private/PoseSearchContext.cpp
+++ b/Engine/Plugins/Experimental/Animation/PoseSearch/Source/Runtime/Private/PoseSearchContext.cpp
@@ -63,8 +63,7 @@ FVector FDebugDrawParams::ExtractPosition(TConstArrayView<float> PoseVector, flo
 {
 	// we don't wanna ask for a SchemaOriginBoneIdx in the future or past
 	check(PermutationTimeType != EPermutationTimeType::UsePermutationTime);
-	const UPoseSearchSchema* Schema = GetSchema();
-	if (!FMath::IsNearlyZero(SampleTimeOffset) && Schema)
+	if (const UPoseSearchSchema* Schema = GetSchema())
 	{
 		// looking for a UPoseSearchFeatureChannel_Position that matches the TimeOffset and SchemaBoneIdx,
 		// with SchemaOriginBoneIdx to be the root bone and the appropriate PermutationTimeType 
@@ -97,8 +96,7 @@ FVector FDebugDrawParams::ExtractPosition(TConstArrayView<float> PoseVector, flo
 
 FQuat FDebugDrawParams::ExtractRotation(TConstArrayView<float> PoseVector, float SampleTimeOffset, int8 SchemaBoneIdx) const
 {
-	const UPoseSearchSchema* Schema = GetSchema();
-	if (!FMath::IsNearlyZero(SampleTimeOffset) && Schema)
+	if (const UPoseSearchSchema* Schema = GetSchema())
 	{
 		int32 HeadingAxisFoundNum = 0;
 		const UPoseSearchFeatureChannel_Heading* FoundHeading[int32(EHeadingAxis::Num)];
