motion matching - database support for malformated blend spaces (FBlendSampleData with null animations)

diff --git a/Engine/Plugins/Experimental/Animation/PoseSearch/Source/Editor/Private/PoseSearchMeshComponent.cpp b/Engine/Plugins/Experimental/Animation/PoseSearch/Source/Editor/Private/PoseSearchMeshComponent.cpp
--- a/Engine/Plugins/Experimental/Animation/PoseSearch/Source/Editor/Private/PoseSearchMeshComponent.cpp
+++ b/Engine/Plugins/Experimental/Animation/PoseSearch/Source/Editor/Private/PoseSearchMeshComponent.cpp
@@ -81,32 +81,28 @@ void UPoseSearchMeshComponent::UpdatePose(const FUpdateContext& UpdateContext)
 		int32 TriangulationIndex = 0;
 		UpdateContext.BlendSpace->GetSamplesFromBlendInput(UpdateContext.BlendParameters, BlendSamples, TriangulationIndex, true);
 		
-		float PlayLength = UpdateContext.BlendSpace->GetAnimationLengthFromSampleData(BlendSamples);
-		
-		float PreviousTime = UpdateContext.StartTime * PlayLength;
+		const float PlayLength = UpdateContext.BlendSpace->GetAnimationLengthFromSampleData(BlendSamples);
+		const float PreviousTime = UpdateContext.StartTime * PlayLength;
 		float CurrentTime = UpdateContext.Time * PlayLength;
 
 		float AdvancedTime = PreviousTime;
-		FAnimationRuntime::AdvanceTime(
-			UpdateContext.bLoop,
-			CurrentTime - PreviousTime,
-			CurrentTime,
-			PlayLength);
-		
+		FAnimationRuntime::AdvanceTime(UpdateContext.bLoop, CurrentTime - PreviousTime, CurrentTime, PlayLength);
+
 		FDeltaTimeRecord DeltaTimeRecord;
 		DeltaTimeRecord.Set(PreviousTime, AdvancedTime - PreviousTime);
 		FAnimExtractContext ExtractionCtx(static_cast<double>(AdvancedTime), true, DeltaTimeRecord, UpdateContext.bLoop);
 
 		for (int32 BlendSampleIdex = 0; BlendSampleIdex < BlendSamples.Num(); BlendSampleIdex++)
 		{
-			float Scale = BlendSamples[BlendSampleIdex].Animation->GetPlayLength() / PlayLength;
+			FBlendSampleData& BlendSample = BlendSamples[BlendSampleIdex];
+			const float Scale = BlendSample.Animation && PlayLength > UE_KINDA_SMALL_NUMBER ? BlendSample.Animation->GetPlayLength() / PlayLength : 1.f;
 
 			FDeltaTimeRecord BlendSampleDeltaTimeRecord;
 			BlendSampleDeltaTimeRecord.Set(DeltaTimeRecord.GetPrevious() * Scale, DeltaTimeRecord.Delta * Scale);
 
-			BlendSamples[BlendSampleIdex].DeltaTimeRecord = BlendSampleDeltaTimeRecord;
-			BlendSamples[BlendSampleIdex].PreviousTime = PreviousTime * Scale;
-			BlendSamples[BlendSampleIdex].Time = AdvancedTime * Scale;
+			BlendSample.DeltaTimeRecord = BlendSampleDeltaTimeRecord;
+			BlendSample.PreviousTime = PreviousTime * Scale;
+			BlendSample.Time = AdvancedTime * Scale;
 		}
 
 		UpdateContext.BlendSpace->GetAnimationPose(BlendSamples, ExtractionCtx, PoseData);
diff --git a/Engine/Plugins/Experimental/Animation/PoseSearch/Source/Runtime/Private/PoseSearchAssetSampler.cpp b/Engine/Plugins/Experimental/Animation/PoseSearch/Source/Runtime/Private/PoseSearchAssetSampler.cpp
--- a/Engine/Plugins/Experimental/Animation/PoseSearch/Source/Runtime/Private/PoseSearchAssetSampler.cpp
+++ b/Engine/Plugins/Experimental/Animation/PoseSearch/Source/Runtime/Private/PoseSearchAssetSampler.cpp
@@ -172,14 +172,15 @@ static void ProcessRootTransform(const UBlendSpace* BlendSpace, const FVector& B
 
 			for (int32 BlendSampleIdex = 0; BlendSampleIdex < BlendSamples.Num(); BlendSampleIdex++)
 			{
-				float Scale = BlendSamples[BlendSampleIdex].Animation->GetPlayLength() / CachedPlayLength;
+				FBlendSampleData& BlendSample = BlendSamples[BlendSampleIdex];
+				const float Scale = BlendSample.Animation && CachedPlayLength > UE_KINDA_SMALL_NUMBER ? BlendSample.Animation->GetPlayLength() / CachedPlayLength : 1.f;
 
 				FDeltaTimeRecord BlendSampleDeltaTimeRecord;
 				BlendSampleDeltaTimeRecord.Set(DeltaTimeRecord.GetPrevious() * Scale, DeltaTimeRecord.Delta * Scale);
 
-				BlendSamples[BlendSampleIdex].DeltaTimeRecord = BlendSampleDeltaTimeRecord;
-				BlendSamples[BlendSampleIdex].PreviousTime = PreviousTime * Scale;
-				BlendSamples[BlendSampleIdex].Time = CurrentTime * Scale;
+				BlendSample.DeltaTimeRecord = BlendSampleDeltaTimeRecord;
+				BlendSample.PreviousTime = PreviousTime * Scale;
+				BlendSample.Time = CurrentTime * Scale;
 			}
 
 			FCompactPose Pose;
@@ -380,14 +381,15 @@ void FAnimationAssetSampler::ExtractPose(const FAnimExtractContext& ExtractionCt
 		{
 			for (int32 BlendSampleIdex = 0; BlendSampleIdex < BlendSamples.Num(); BlendSampleIdex++)
 			{
-				float Scale = BlendSamples[BlendSampleIdex].Animation->GetPlayLength() / CachedPlayLength;
+				FBlendSampleData& BlendSample = BlendSamples[BlendSampleIdex];
+				const float Scale = BlendSample.Animation && CachedPlayLength > UE_KINDA_SMALL_NUMBER ? BlendSample.Animation->GetPlayLength() / CachedPlayLength : 1.f;
 
 				FDeltaTimeRecord BlendSampleDeltaTimeRecord;
 				BlendSampleDeltaTimeRecord.Set(ExtractionCtx.DeltaTimeRecord.GetPrevious() * Scale, ExtractionCtx.DeltaTimeRecord.Delta * Scale);
 
-				BlendSamples[BlendSampleIdex].DeltaTimeRecord = BlendSampleDeltaTimeRecord;
-				BlendSamples[BlendSampleIdex].PreviousTime = ExtractionCtx.DeltaTimeRecord.GetPrevious() * Scale;
-				BlendSamples[BlendSampleIdex].Time = ExtractionCtx.CurrentTime * Scale;
+				BlendSample.DeltaTimeRecord = BlendSampleDeltaTimeRecord;
+				BlendSample.PreviousTime = ExtractionCtx.DeltaTimeRecord.GetPrevious() * Scale;
+				BlendSample.Time = ExtractionCtx.CurrentTime * Scale;
 			}
 
 			BlendSpace->GetAnimationPose(BlendSamples, ExtractionCtx, OutAnimPoseData);
@@ -598,15 +600,18 @@ void FAnimationAssetSampler::ExtractPoseSearchNotifyStates(float Time, TArray<UA
 			{
 				// Find highest weighted
 				const int32 HighestWeightIndex = GetHighestWeightSample(BlendSamples);
-
-				// getting pose search notifies in an interval of size ExtractionInterval, centered on Time
-				SampleTime = Time * (BlendSamples[HighestWeightIndex].Animation->GetPlayLength() / CachedPlayLength);
-
-				// Get notifies for highest weighted
-				BlendSamples[HighestWeightIndex].Animation->GetAnimNotifies(
-					(SampleTime - (ExtractionInterval * 0.5f)),
-					ExtractionInterval,
-					NotifyContext);
+				const FBlendSampleData& BlendSample = BlendSamples[HighestWeightIndex];
+				if (BlendSample.Animation)
+				{
+					// getting pose search notifies in an interval of size ExtractionInterval, centered on Time
+					if (CachedPlayLength > UE_KINDA_SMALL_NUMBER)
+					{
+						SampleTime = Time * (BlendSample.Animation->GetPlayLength() / CachedPlayLength);
+					}
+
+					// Get notifies for highest weighted
+					BlendSample.Animation->GetAnimNotifies((SampleTime - (ExtractionInterval * 0.5f)), ExtractionInterval, NotifyContext);
+				}
 			}
 		}
 	}
