motion matching - setting appropriate ClampMin for FAnimNode_PoseSearchHistoryCollector_Base to avoid possible divizion by zero

diff --git a/Engine/Plugins/Experimental/Animation/PoseSearch/Source/Runtime/Private/PoseSearchHistory.cpp b/Engine/Plugins/Experimental/Animation/PoseSearch/Source/Runtime/Private/PoseSearchHistory.cpp
--- a/Engine/Plugins/Experimental/Animation/PoseSearch/Source/Runtime/Private/PoseSearchHistory.cpp
+++ b/Engine/Plugins/Experimental/Animation/PoseSearch/Source/Runtime/Private/PoseSearchHistory.cpp
@@ -110,6 +110,7 @@ void FPoseHistoryEntry::Update(float InTime, FCSPose<FCompactPose>& ComponentSpa
 // FPoseHistory
 void FPoseHistory::Init(int32 InNumPoses, float InTimeHorizon, const TArray<FBoneIndexType>& RequiredBones)
 {
+	check(InNumPoses >= 2 && InTimeHorizon > UE_KINDA_SMALL_NUMBER);
 	TimeHorizon = InTimeHorizon;
 
 	BoneToTransformMap.Reset();
@@ -278,8 +279,11 @@ void FPoseHistory::Update(float SecondsElapsed, FCSPose<FCompactPose>& Component
 
 float FPoseHistory::GetSampleTimeInterval() const
 {
+	const int32 EntriesMax = Entries.Max();
+	check(EntriesMax >= 2 && TimeHorizon > UE_KINDA_SMALL_NUMBER);
+
 	// Reserve one pose for computing derivatives at the time horizon
-	return TimeHorizon / (Entries.Max() - 1);
+	return TimeHorizon / (EntriesMax - 1);
 }
 
 #if ENABLE_DRAW_DEBUG && ENABLE_ANIM_DEBUG
diff --git a/Engine/Plugins/Experimental/Animation/PoseSearch/Source/Runtime/Public/PoseSearch/AnimNode_PoseSearchHistoryCollector.h b/Engine/Plugins/Experimental/Animation/PoseSearch/Source/Runtime/Public/PoseSearch/AnimNode_PoseSearchHistoryCollector.h
--- a/Engine/Plugins/Experimental/Animation/PoseSearch/Source/Runtime/Public/PoseSearch/AnimNode_PoseSearchHistoryCollector.h
+++ b/Engine/Plugins/Experimental/Animation/PoseSearch/Source/Runtime/Public/PoseSearch/AnimNode_PoseSearchHistoryCollector.h
@@ -14,11 +14,11 @@ struct POSESEARCH_API FAnimNode_PoseSearchHistoryCollector_Base : public FAnimNo
 public:
 	
 	// The maximum amount of poses that can be stored
-	UPROPERTY(EditAnywhere, BlueprintReadWrite, Category = Settings, meta = (ClampMin="0"))
+	UPROPERTY(EditAnywhere, BlueprintReadWrite, Category = Settings, meta = (ClampMin="2"))
 	int32 PoseCount = 64;
 	
 	// The time horizon for how long a pose will be stored in seconds
-	UPROPERTY(EditAnywhere, BlueprintReadWrite, Category = Settings, meta = (ClampMin="0"))
+	UPROPERTY(EditAnywhere, BlueprintReadWrite, Category = Settings, meta = (ClampMin="0.001"))
 	float PoseDuration = 1.5f;
 
 	UPROPERTY(EditAnywhere, Category = Settings)
