motion matching - fix for ContinuingPoseCandidate

diff --git a/Engine/Plugins/Experimental/Animation/PoseSearch/Source/Runtime/Public/PoseSearch/PoseSearchContext.h b/Engine/Plugins/Experimental/Animation/PoseSearch/Source/Runtime/Public/PoseSearch/PoseSearchContext.h
--- a/Engine/Plugins/Experimental/Animation/PoseSearch/Source/Runtime/Public/PoseSearch/PoseSearchContext.h
+++ b/Engine/Plugins/Experimental/Animation/PoseSearch/Source/Runtime/Public/PoseSearch/PoseSearchContext.h
@@ -260,11 +260,11 @@ public:
 
 				if (bPoppedContinuingPoseCandidate)
 				{
-					// if we popped the continuing pose candidate, we make some space fir it and push it back
+					// if we popped the continuing pose candidate, we make some space for it and push it back
 					FPoseCandidate PoppedPoseCandidate;
 					Pop(PoppedPoseCandidate);
 					PoseCandidateHeap.HeapPush(ContinuingPoseCandidate);
-					PoseIdxToFlags.Add(ContinuingPoseCandidate.PoseIdx, PoppedPoseCandidate.PoseCandidateFlags);
+					PoseIdxToFlags.Add(ContinuingPoseCandidate.PoseIdx, ContinuingPoseCandidate.PoseCandidateFlags);
 				}
 
 				FPoseCandidate PoseCandidate;
