Fix for crash when compiling motion matching sample graph.
Reparented AnimNode_MotionMatching to AnimNode_BlendStack_Standalone

diff --git a/Engine/Plugins/Experimental/Animation/PoseSearch/Source/Editor/Public/AnimGraphNode_BlendStack.h b/Engine/Plugins/Experimental/Animation/PoseSearch/Source/Editor/Public/AnimGraphNode_BlendStack.h
--- a/Engine/Plugins/Experimental/Animation/PoseSearch/Source/Editor/Public/AnimGraphNode_BlendStack.h
+++ b/Engine/Plugins/Experimental/Animation/PoseSearch/Source/Editor/Public/AnimGraphNode_BlendStack.h
@@ -39,7 +39,7 @@ class UAnimGraphNode_BlendStack_Base : public UAnimGraphNode_AssetPlayerBase
 
 protected:
 	virtual FAnimNode_BlendStack_Standalone* GetBlendStackNode() const PURE_VIRTUAL(UAnimGraphNode_BlendStack_Base::GetBlendStackNode, return nullptr;);
-	virtual int32 GetMaxActiveBlends() const;
+	int32 GetMaxActiveBlends() const;
 
 private:
 
diff --git a/Engine/Plugins/Experimental/Animation/PoseSearch/Source/Editor/Public/AnimGraphNode_MotionMatching.h b/Engine/Plugins/Experimental/Animation/PoseSearch/Source/Editor/Public/AnimGraphNode_MotionMatching.h
--- a/Engine/Plugins/Experimental/Animation/PoseSearch/Source/Editor/Public/AnimGraphNode_MotionMatching.h
+++ b/Engine/Plugins/Experimental/Animation/PoseSearch/Source/Editor/Public/AnimGraphNode_MotionMatching.h
@@ -24,6 +24,5 @@ class UAnimGraphNode_MotionMatching : public UAnimGraphNode_BlendStack_Base
 	virtual UScriptStruct* GetTimePropertyStruct() const override;
 	virtual void BakeDataDuringCompilation(class FCompilerResultsLog& MessageLog) override;
 
-	virtual int32 GetMaxActiveBlends() const override { return Node.MaxActiveBlends; }
-	virtual FAnimNode_BlendStack_Standalone* GetBlendStackNode() const override { return (FAnimNode_BlendStack_Standalone*)(&Node.BlendStackNode); }
+	virtual FAnimNode_BlendStack_Standalone* GetBlendStackNode() const override { return (FAnimNode_BlendStack_Standalone*)(&Node); }
 };
diff --git a/Engine/Plugins/Experimental/Animation/PoseSearch/Source/Runtime/Private/AnimNode_BlendStackInput.cpp b/Engine/Plugins/Experimental/Animation/PoseSearch/Source/Runtime/Private/AnimNode_BlendStackInput.cpp
--- a/Engine/Plugins/Experimental/Animation/PoseSearch/Source/Runtime/Private/AnimNode_BlendStackInput.cpp
+++ b/Engine/Plugins/Experimental/Animation/PoseSearch/Source/Runtime/Private/AnimNode_BlendStackInput.cpp
@@ -18,8 +18,7 @@ static FAnimNode_BlendStack_Standalone* GetBlendStackNodeFromIndex(const FAnimat
 	check(AnimNodeProperties.IsValidIndex(NodeIndex));
 
 	FStructProperty* LinkedProperty = AnimNodeProperties[NodeIndex];
-	void* LinkedNodePtr = LinkedProperty->ContainerPtrToValuePtr<void>(Context.AnimInstanceProxy->GetAnimInstanceObject());
-	return (FAnimNode_BlendStack_Standalone*)LinkedNodePtr;
+	return LinkedProperty->ContainerPtrToValuePtr<FAnimNode_BlendStack_Standalone>(Context.AnimInstanceProxy->GetAnimInstanceObject());
 }
 
 void FAnimNode_BlendStackInput::Update_AnyThread(const FAnimationUpdateContext& Context)
diff --git a/Engine/Plugins/Experimental/Animation/PoseSearch/Source/Runtime/Private/AnimNode_MotionMatching.cpp b/Engine/Plugins/Experimental/Animation/PoseSearch/Source/Runtime/Private/AnimNode_MotionMatching.cpp
--- a/Engine/Plugins/Experimental/Animation/PoseSearch/Source/Runtime/Private/AnimNode_MotionMatching.cpp
+++ b/Engine/Plugins/Experimental/Animation/PoseSearch/Source/Runtime/Private/AnimNode_MotionMatching.cpp
@@ -35,11 +35,7 @@ void FAnimNode_MotionMatching::Initialize_AnyThread(const FAnimationInitializeCo
 
 	GetEvaluateGraphExposedInputs().Execute(Context);
 
-	Source.SetLinkNode(&BlendStackNode);
-
-	// calling BlendStackNode.Initialize_AnyThread(Context) that will reset the BlendStackNode.AnimPlayers
-	Source.Initialize(Context);
-
+	FAnimNode_BlendStack_Standalone::Initialize_AnyThread(Context);
 	MotionMatchingState.Reset(Context.AnimInstanceProxy->GetComponentTransform());
 }
 
@@ -48,7 +44,7 @@ void FAnimNode_MotionMatching::Evaluate_AnyThread(FPoseContext& Output)
 	DECLARE_SCOPE_HIERARCHICAL_COUNTER_ANIMNODE(Evaluate_AnyThread);
 	ANIM_MT_SCOPE_CYCLE_COUNTER_VERBOSE(MotionMatching, !IsInGameThread());
 
-	Source.Evaluate(Output);
+	FAnimNode_BlendStack_Standalone::Evaluate_AnyThread(Output);
 
 	const UE::Anim::IAnimRootMotionProvider* RootMotionProvider = UE::Anim::IAnimRootMotionProvider::Get();
 	FTransform RootMotionTransformDelta;
@@ -111,14 +107,14 @@ void FAnimNode_MotionMatching::UpdateAssetPlayer(const FAnimationUpdateContext&
 	if (bNeedsReset)
 	{
 		MotionMatchingState.Reset(Context.AnimInstanceProxy->GetComponentTransform());
-		BlendStackNode.Reset();
+		FAnimNode_BlendStack_Standalone::Reset();
 	}
 	else
 	{
 		// We adjust the motion matching state asset time to the current player node's asset time. This is done 
 		// because the player node may have ticked more or less time than we expected due to variable dt or the 
 		// dynamic playback rate adjustment and as such the motion matching state does not update by itself
-		MotionMatchingState.AdjustAssetTime(BlendStackNode.GetAccumulatedTime());
+		MotionMatchingState.AdjustAssetTime(GetAccumulatedTime());
 	}
 	UpdateCounter.SynchronizeWith(Context.AnimInstanceProxy->GetUpdateCounter());
 
@@ -177,24 +173,19 @@ void FAnimNode_MotionMatching::UpdateAssetPlayer(const FAnimationUpdateContext&
 			{
 				// root bone blending needs to be immediate if MM node controls the offset between mesh component and root bone
 				const float RootBoneBlendTime = YawFromAnimationBlendRate < 0.f ? BlendTime : 0.f;
-				BlendStackNode.BlendTo(Context, DatabaseAsset->GetAnimationAsset(), MotionMatchingState.CurrentSearchResult.AssetTime,
+				FAnimNode_BlendStack_Standalone::BlendTo(Context, DatabaseAsset->GetAnimationAsset(), MotionMatchingState.CurrentSearchResult.AssetTime,
 					DatabaseAsset->IsLooping(), SearchIndexAsset->bMirrored, CurrentResultDatabase->Schema->MirrorDataTable.Get(),
 					BlendTime, RootBoneBlendTime, BlendProfile, BlendOption, bUseInertialBlend, SearchIndexAsset->BlendParameters, MotionMatchingState.WantedPlayRate);
 			}
 		}
 	}
-	BlendStackNode.UpdatePlayRate(MotionMatchingState.WantedPlayRate);
-
-	Source.Update(Context);
+	
+	FAnimNode_BlendStack_Standalone::UpdatePlayRate(MotionMatchingState.WantedPlayRate);
+	FAnimNode_BlendStack_Standalone::UpdateAssetPlayer(Context);
 
 	bForceInterruptNextUpdate = false;
 }
 
-void FAnimNode_MotionMatching::GatherDebugData(FNodeDebugData& DebugData)
-{
-	Source.GatherDebugData(DebugData);
-}
-
 void FAnimNode_MotionMatching::SetDatabaseToSearch(UPoseSearchDatabase* InDatabase, bool bForceInterruptIfNew)
 {
 	if (DatabasesToSearch.Num() == 1 && DatabasesToSearch[0] == InDatabase)
@@ -274,31 +265,6 @@ void FAnimNode_MotionMatching::ForceInterruptNextUpdate()
 }
 
 // FAnimNode_AssetPlayerBase interface
-float FAnimNode_MotionMatching::GetAccumulatedTime() const
-{
-	return BlendStackNode.GetAccumulatedTime();
-}
-
-UAnimationAsset* FAnimNode_MotionMatching::GetAnimAsset() const
-{
-	return BlendStackNode.GetAnimAsset();
-}
-
-float FAnimNode_MotionMatching::GetCurrentAssetLength() const
-{
-	return BlendStackNode.GetCurrentAssetLength();
-}
-
-float FAnimNode_MotionMatching::GetCurrentAssetTime() const
-{
-	return BlendStackNode.GetCurrentAssetLength();
-}
-
-float FAnimNode_MotionMatching::GetCurrentAssetTimePlayRateAdjusted() const
-{
-	return BlendStackNode.GetCurrentAssetTimePlayRateAdjusted();
-}
-
 bool FAnimNode_MotionMatching::GetIgnoreForRelevancyTest() const
 {
 	return GET_ANIM_NODE_DATA(bool, bIgnoreForRelevancyTest);
diff --git a/Engine/Plugins/Experimental/Animation/PoseSearch/Source/Runtime/Public/PoseSearch/AnimNode_BlendStack.h b/Engine/Plugins/Experimental/Animation/PoseSearch/Source/Runtime/Public/PoseSearch/AnimNode_BlendStack.h
--- a/Engine/Plugins/Experimental/Animation/PoseSearch/Source/Runtime/Public/PoseSearch/AnimNode_BlendStack.h
+++ b/Engine/Plugins/Experimental/Animation/PoseSearch/Source/Runtime/Public/PoseSearch/AnimNode_BlendStack.h
@@ -112,6 +112,11 @@ struct POSESEARCH_API FAnimNode_BlendStack_Standalone : public FAnimNode_AssetPl
 
 	int32 GetMaxActiveBlends() const { return MaxActiveBlends; }
 
+protected:
+	// Number of max active blending animation in the blend stack. If MaxActiveBlends is zero then blend stack is disabled
+	UPROPERTY(EditAnywhere, Category = Settings, meta = (ClampMin = "0"))
+	int32 MaxActiveBlends = 4;
+
 private:
 
 	void InitializeSample(const FAnimationInitializeContext& Context, FPoseSearchAnimPlayer& SamplePlayer);
@@ -119,10 +124,6 @@ private:
 	void UpdateSample(const FAnimationUpdateContext& Context, const int32 PlayerIndex);
 	bool IsSampleGraphAvailableForPlayer(const int32 PlayerIndex);
 
-	// Number of max active blending animation in the blend stack. If MaxActiveBlends is zero then blend stack is disabled
-	UPROPERTY(EditAnywhere, Category = Settings, meta = (ClampMin = "0"))
-	int32 MaxActiveBlends = 4;
-
 	friend class UAnimGraphNode_BlendStack_Base;
 };
 
diff --git a/Engine/Plugins/Experimental/Animation/PoseSearch/Source/Runtime/Public/PoseSearch/AnimNode_MotionMatching.h b/Engine/Plugins/Experimental/Animation/PoseSearch/Source/Runtime/Public/PoseSearch/AnimNode_MotionMatching.h
--- a/Engine/Plugins/Experimental/Animation/PoseSearch/Source/Runtime/Public/PoseSearch/AnimNode_MotionMatching.h
+++ b/Engine/Plugins/Experimental/Animation/PoseSearch/Source/Runtime/Public/PoseSearch/AnimNode_MotionMatching.h
@@ -12,7 +12,7 @@
 class UPoseSearchDatabase;
 
 USTRUCT(BlueprintInternalUseOnly)
-struct POSESEARCH_API FAnimNode_MotionMatching : public FAnimNode_AssetPlayerBase
+struct POSESEARCH_API FAnimNode_MotionMatching : public FAnimNode_BlendStack_Standalone
 {
 	GENERATED_BODY()
 
@@ -34,23 +34,14 @@ private:
 	// @todo: implement CacheBones_AnyThread to rebind the schema bones
 	virtual void Initialize_AnyThread(const FAnimationInitializeContext& Context) override;
 	virtual void Evaluate_AnyThread(FPoseContext& Output) override;
-	virtual void GatherDebugData(FNodeDebugData& DebugData) override;
 	// End of FAnimNode_Base interface
 
 	// FAnimNode_AssetPlayerBase interface
-	virtual float GetAccumulatedTime() const override;
-	virtual UAnimationAsset* GetAnimAsset() const override;
 	virtual void UpdateAssetPlayer(const FAnimationUpdateContext& Context) override;
-	virtual float GetCurrentAssetLength() const override;
-	virtual float GetCurrentAssetTime() const override;
-	virtual float GetCurrentAssetTimePlayRateAdjusted() const override;
 	virtual bool GetIgnoreForRelevancyTest() const override;
 	virtual bool SetIgnoreForRelevancyTest(bool bInIgnoreForRelevancyTest) override;
 	// End of FAnimNode_AssetPlayerBase interface
 
-	UPROPERTY()
-	FPoseLink Source;
-
 	// The database to search. This can be overridden by Anim Node Functions such as "On Become Relevant" and "On Update" via SetDatabaseToSearch/SetDatabasesToSearch.
 	UPROPERTY(EditAnywhere, Category = Settings, meta = (PinShownByDefault))
 	TObjectPtr<const UPoseSearchDatabase> Database = nullptr;
@@ -67,10 +58,6 @@ private:
 	UPROPERTY(EditAnywhere, Category = Settings, meta = (PinHiddenByDefault, ClampMin="0"))
 	float BlendTime = 0.2f;
 
-	// Number of max active animation segments being blended together in the blend stack. If MaxActiveBlends is zero then the blend stack is disabled.
-	UPROPERTY(EditAnywhere, Category = Settings, meta = (ClampMin="0"))
-	int32 MaxActiveBlends = 4;
-
 	// Set Blend Profiles (editable in the skeleton) to determine how the blending is distributed among your character's bones. It could be used to differentiate between upper body and lower body to blend timing.
 	UPROPERTY(EditAnywhere, Category = Settings, meta = (PinHiddenByDefault, UseAsBlendProfile = true))
 	TObjectPtr<UBlendProfile> BlendProfile;
