Fix for cached bones on AnimNode_BlendStack not being propagated.
#rnx

diff --git a/Engine/Plugins/Experimental/Animation/PoseSearch/Source/Runtime/Private/AnimNode_BlendStack.cpp b/Engine/Plugins/Experimental/Animation/PoseSearch/Source/Runtime/Private/AnimNode_BlendStack.cpp
--- a/Engine/Plugins/Experimental/Animation/PoseSearch/Source/Runtime/Private/AnimNode_BlendStack.cpp
+++ b/Engine/Plugins/Experimental/Animation/PoseSearch/Source/Runtime/Private/AnimNode_BlendStack.cpp
@@ -427,6 +427,17 @@ void FAnimNode_BlendStack_Standalone::Initialize_AnyThread(const FAnimationIniti
 	}
 }
 
+void FAnimNode_BlendStack_Standalone::CacheBones_AnyThread(const FAnimationCacheBonesContext& Context) 
+{
+	const int32 BlendStackSize = AnimPlayers.Num();
+	for (int32 AnimPlayerIndex = 0; AnimPlayerIndex < BlendStackSize; ++AnimPlayerIndex)
+	{
+		// Cache bones for all active anim players.
+		// There's no need to check for weight since all unneeded anim players 
+		// would have been pruned during the last evaluation.
+		CacheBonesForSample(Context, AnimPlayerIndex);
+	}
+}
 
 void FAnimNode_BlendStack_Standalone::UpdateAssetPlayer(const FAnimationUpdateContext& Context)
 {
@@ -542,12 +553,25 @@ void FAnimNode_BlendStack_Standalone::UpdateSample(const FAnimationUpdateContext
 	SamplePlayer.AdvanceBlendInTime(Context.GetDeltaTime());
 }
 
+void FAnimNode_BlendStack_Standalone::CacheBonesForSample(const FAnimationCacheBonesContext& Context, const int32 PlayerIndex)
+{
+	FPoseSearchAnimPlayer& SamplePlayer = AnimPlayers[PlayerIndex];
+	const bool bHasSampleGraph = !SampleGraphPoseLinks.IsEmpty() && (PlayerIndex <= MaxActiveBlends);
+	if (bHasSampleGraph)
+	{
+		FBlendStack_SampleGraphPoseLink& PoseLink = SampleGraphPoseLinks[SamplePlayer.GetPoseLinkIndex()];
+		PoseLink.ConditionalCacheBones(Context);
+	}
+}
+
 void FAnimNode_BlendStack_Standalone::InitializeSample(const FAnimationInitializeContext& Context, FPoseSearchAnimPlayer& SamplePlayer)
 {
 	if (SamplePlayer.GetPoseLinkIndex() != INDEX_NONE)
 	{
 		FBlendStack_SampleGraphPoseLink& PoseLink = SampleGraphPoseLinks[SamplePlayer.GetPoseLinkIndex()];
+		PoseLink.SetInputPosePlayer(SamplePlayer);
 		PoseLink.Root.Initialize(Context);
+		PoseLink.ConditionalCacheBones(Context);
 	}
 }
 
diff --git a/Engine/Plugins/Experimental/Animation/PoseSearch/Source/Runtime/Public/PoseSearch/AnimNode_BlendStack.h b/Engine/Plugins/Experimental/Animation/PoseSearch/Source/Runtime/Public/PoseSearch/AnimNode_BlendStack.h
--- a/Engine/Plugins/Experimental/Animation/PoseSearch/Source/Runtime/Public/PoseSearch/AnimNode_BlendStack.h
+++ b/Engine/Plugins/Experimental/Animation/PoseSearch/Source/Runtime/Public/PoseSearch/AnimNode_BlendStack.h
@@ -96,6 +96,7 @@ struct POSESEARCH_API FAnimNode_BlendStack_Standalone : public FAnimNode_AssetPl
 	// FAnimNode_Base interface
 	virtual void Evaluate_AnyThread(FPoseContext& Output) override;
 	virtual void Initialize_AnyThread(const FAnimationInitializeContext& Context) override;
+	virtual void CacheBones_AnyThread(const FAnimationCacheBonesContext& Context) override;
 	virtual void GatherDebugData(FNodeDebugData& DebugData) override;
 	// End of FAnimNode_Base interface
 
@@ -122,6 +123,7 @@ private:
 	void InitializeSample(const FAnimationInitializeContext& Context, FPoseSearchAnimPlayer& SamplePlayer);
 	void EvaluateSample(FPoseContext& Output, const int32 PlayerIndex);
 	void UpdateSample(const FAnimationUpdateContext& Context, const int32 PlayerIndex);
+	void CacheBonesForSample(const FAnimationCacheBonesContext& Context, const int32 PlayerIndex);
 	bool IsSampleGraphAvailableForPlayer(const int32 PlayerIndex);
 
 	friend class UAnimGraphNode_BlendStack_Base;
