Fix for blend stack swapping LODs during active blends.
#rnx

diff --git a/Engine/Plugins/Experimental/Animation/PoseSearch/Source/Runtime/Private/AnimNode_BlendStack.cpp b/Engine/Plugins/Experimental/Animation/PoseSearch/Source/Runtime/Private/AnimNode_BlendStack.cpp
--- a/Engine/Plugins/Experimental/Animation/PoseSearch/Source/Runtime/Private/AnimNode_BlendStack.cpp
+++ b/Engine/Plugins/Experimental/Animation/PoseSearch/Source/Runtime/Private/AnimNode_BlendStack.cpp
@@ -132,6 +132,7 @@ void FPoseSearchAnimPlayer::StorePoseContext(const FPoseContext& PoseContext)
 	if (PoseContext.Pose.IsValid())
 	{
 		StoredPose.CopyBonesFrom(PoseContext.Pose);
+		StoredPose.CopyAndAssignBoneContainer(StoredBoneContainer);
 	}
 
 	StoredCurve.CopyFrom(PoseContext.Curve);
@@ -142,9 +143,33 @@ void FPoseSearchAnimPlayer::RestorePoseContext(FPoseContext& PoseContext) const
 {
 	check(!SequencePlayerNode.GetSequence() && !BlendSpacePlayerNode.GetBlendSpace());
 
-	if (StoredPose.IsValid() && PoseContext.Pose.GetNumBones() == StoredPose.GetNumBones())
+	if (StoredPose.IsValid())
 	{
-		PoseContext.Pose.CopyBonesFrom(StoredPose);
+		// Serial number mismatch means a potential bone LOD mismatch, even if we have the same number of bones.
+		// Remap the pose manually in those cases.
+		if (PoseContext.Pose.GetBoneContainer().GetSerialNumber() == StoredBoneContainer.GetSerialNumber())
+		{
+			PoseContext.Pose.CopyBonesFrom(StoredPose);
+		}
+		else
+		{
+			const FBoneContainer CurrentBoneContainer = PoseContext.Pose.GetBoneContainer();
+			for (FCompactPoseBoneIndex CompactPoseIndex : PoseContext.Pose.ForEachBoneIndex())
+			{
+				// Map the current compact pose index to skeleton index, and map this back to the stored compact pose index.
+				const FSkeletonPoseBoneIndex SkeletonPoseIndex =  CurrentBoneContainer.GetSkeletonPoseIndexFromCompactPoseIndex(CompactPoseIndex);
+				const FCompactPoseBoneIndex StoredCompactPoseIndex =  StoredBoneContainer.GetCompactPoseIndexFromSkeletonPoseIndex(SkeletonPoseIndex);
+				if (StoredCompactPoseIndex == INDEX_NONE)
+				{
+					// If our stored pose doesn't have the bone, reset to ref pose.
+					PoseContext.Pose[CompactPoseIndex] = CurrentBoneContainer.GetRefPoseTransform(CompactPoseIndex);
+				}
+				else
+				{
+					PoseContext.Pose[CompactPoseIndex] = StoredPose[StoredCompactPoseIndex];
+				}
+			}
+		}
 	}
 	else
 	{
diff --git a/Engine/Plugins/Experimental/Animation/PoseSearch/Source/Runtime/Public/PoseSearch/AnimNode_BlendStack.h b/Engine/Plugins/Experimental/Animation/PoseSearch/Source/Runtime/Public/PoseSearch/AnimNode_BlendStack.h
--- a/Engine/Plugins/Experimental/Animation/PoseSearch/Source/Runtime/Public/PoseSearch/AnimNode_BlendStack.h
+++ b/Engine/Plugins/Experimental/Animation/PoseSearch/Source/Runtime/Public/PoseSearch/AnimNode_BlendStack.h
@@ -56,6 +56,8 @@ public:
 	FCompactHeapPose StoredPose;
 	FBlendedHeapCurve StoredCurve;
 	UE::Anim::FHeapAttributeContainer StoredAttributes;
+	// We need to store the bone container, in case we have a LOD swap during a blend that uses the stored pose.
+	FBoneContainer StoredBoneContainer;
 
 	EAlphaBlendOption BlendOption = EAlphaBlendOption::Linear;
 	int32 PoseLinkIndex = INDEX_NONE;
