motion matching - fix for MM db index loading when there isn't an EditorEngine (ex. Standalone Game via -game argument)

diff --git a/Engine/Plugins/Experimental/Animation/PoseSearch/Source/Runtime/Private/PoseSearchDerivedData.cpp b/Engine/Plugins/Experimental/Animation/PoseSearch/Source/Runtime/Private/PoseSearchDerivedData.cpp
--- a/Engine/Plugins/Experimental/Animation/PoseSearch/Source/Runtime/Private/PoseSearchDerivedData.cpp
+++ b/Engine/Plugins/Experimental/Animation/PoseSearch/Source/Runtime/Private/PoseSearchDerivedData.cpp
@@ -11,6 +11,7 @@
 #include "AssetRegistry/AssetRegistryModule.h"
 #include "DerivedDataCache.h"
 #include "DerivedDataRequestOwner.h"
+#include "Editor/EditorEngine.h"
 #include "InstancedStruct.h"
 #include "Misc/CoreDelegates.h"
 #include "PoseSearchDatabaseIndexingContext.h"
@@ -636,8 +637,8 @@ struct FPoseSearchDatabaseAsyncCacheTask
 	};
 
 	// these methods MUST be protected by FPoseSearchDatabaseAsyncCacheTask::Mutex! and to make sure we pass the mutex as input param
-	FPoseSearchDatabaseAsyncCacheTask(UPoseSearchDatabase* InDatabase, FCriticalSection& OuterMutex);
-	void StartNewRequestIfNeeded(FCriticalSection& OuterMutex);
+	FPoseSearchDatabaseAsyncCacheTask(UPoseSearchDatabase* InDatabase, FCriticalSection& OuterMutex, bool bPerformConditionalPostLoadIfRequired);
+	void StartNewRequestIfNeeded(FCriticalSection& OuterMutex, bool bPerformConditionalPostLoadIfRequired);
 	bool CancelIfDependsOn(const UObject* Object, FCriticalSection& OuterMutex);
 	void Update(FCriticalSection& OuterMutex);
 	void Wait(FCriticalSection& OuterMutex);
@@ -673,7 +674,7 @@ private:
 
 class FPoseSearchDatabaseAsyncCacheTasks : public TArray<TUniquePtr<FPoseSearchDatabaseAsyncCacheTask>> {};
 
-FPoseSearchDatabaseAsyncCacheTask::FPoseSearchDatabaseAsyncCacheTask(UPoseSearchDatabase* InDatabase, FCriticalSection& OuterMutex)
+FPoseSearchDatabaseAsyncCacheTask::FPoseSearchDatabaseAsyncCacheTask(UPoseSearchDatabase* InDatabase, FCriticalSection& OuterMutex, bool bPerformConditionalPostLoadIfRequired)
 	: Database(InDatabase)
 	, Owner(UE::DerivedData::EPriority::Normal)
 	, DerivedDataKey(FIoHash::Zero)
@@ -681,7 +682,7 @@ FPoseSearchDatabaseAsyncCacheTask::FPoseSearchDatabaseAsyncCacheTask(UPoseSearch
 	if (IsInGameThread())
 	{
 		// it is safe to compose DDC key only on the game thread, since assets can modified in this thread execution
-		StartNewRequestIfNeeded(OuterMutex);
+		StartNewRequestIfNeeded(OuterMutex, bPerformConditionalPostLoadIfRequired);
 	}
 	else
 	{
@@ -698,7 +699,7 @@ FPoseSearchDatabaseAsyncCacheTask::~FPoseSearchDatabaseAsyncCacheTask()
 	DatabaseDependencies.Reset();
 }
 
-void FPoseSearchDatabaseAsyncCacheTask::StartNewRequestIfNeeded(FCriticalSection& OuterMutex)
+void FPoseSearchDatabaseAsyncCacheTask::StartNewRequestIfNeeded(FCriticalSection& OuterMutex, bool bPerformConditionalPostLoadIfRequired)
 {
 	using namespace UE::DerivedData;
 
@@ -710,7 +711,7 @@ void FPoseSearchDatabaseAsyncCacheTask::StartNewRequestIfNeeded(FCriticalSection
 	Owner.Cancel();
 
 	// composing the key
-	const FKeyBuilder KeyBuilder(Database.Get(), true);
+	const FKeyBuilder KeyBuilder(Database.Get(), true, bPerformConditionalPostLoadIfRequired);
 	if (KeyBuilder.AnyAssetNotReady())
 	{
 		DerivedDataKey = FIoHash::Zero;
@@ -789,7 +790,7 @@ void FPoseSearchDatabaseAsyncCacheTask::Update(FCriticalSection& OuterMutex)
 
 	if (GetState() == EState::Notstarted)
 	{
-		StartNewRequestIfNeeded(OuterMutex);
+		StartNewRequestIfNeeded(OuterMutex, false);
 	}
 
 	if (GetState() == EState::Prestarted && Poll(OuterMutex))
@@ -1207,11 +1208,20 @@ bool FAsyncPoseSearchDatabasesManagement::RequestAsyncBuildIndex(const UPoseSear
 
 	FScopeLock Lock(&Mutex);
 
-	check(Database);
+#if WITH_ENGINE
+	//If there isn't an EditorEngine (ex. Standalone Game via -game argument) the FAsyncPoseSearchDatabasesManagement task doesn't get ticked, need to force completion here
+	if (Cast<UEditorEngine>(GEngine) == nullptr)
+	{
+		Flag |= ERequestAsyncBuildFlag::WaitForCompletion;
+	}
+#endif
+
 	check(EnumHasAnyFlags(Flag, ERequestAsyncBuildFlag::NewRequest | ERequestAsyncBuildFlag::ContinueRequest));
 
 	FAsyncPoseSearchDatabasesManagement& This = FAsyncPoseSearchDatabasesManagement::Get();
 
+	const bool bWaitForCompletion = EnumHasAnyFlags(Flag, ERequestAsyncBuildFlag::WaitForCompletion);
+
 	FPoseSearchDatabaseAsyncCacheTask* Task = nullptr;
 	for (TUniquePtr<FPoseSearchDatabaseAsyncCacheTask>& TaskPtr : This.Tasks)
 	{
@@ -1225,7 +1235,7 @@ bool FAsyncPoseSearchDatabasesManagement::RequestAsyncBuildIndex(const UPoseSear
 				{
 					Task->Cancel(Mutex);
 				}
-				Task->StartNewRequestIfNeeded(Mutex);
+				Task->StartNewRequestIfNeeded(Mutex, bWaitForCompletion);
 			}
 			break;
 		}
@@ -1234,11 +1244,11 @@ bool FAsyncPoseSearchDatabasesManagement::RequestAsyncBuildIndex(const UPoseSear
 	if (!Task)
 	{
 		// we didn't find the Task, so we Emplace a new one
-		This.Tasks.Emplace(MakeUnique<FPoseSearchDatabaseAsyncCacheTask>(const_cast<UPoseSearchDatabase*>(Database), Mutex));
+		This.Tasks.Emplace(MakeUnique<FPoseSearchDatabaseAsyncCacheTask>(const_cast<UPoseSearchDatabase*>(Database), Mutex, bWaitForCompletion));
 		Task = This.Tasks.Last().Get();
 	}
 
-	if (EnumHasAnyFlags(Flag, ERequestAsyncBuildFlag::WaitForCompletion))
+	if (bWaitForCompletion)
 	{
 		check(Task->GetState() != FPoseSearchDatabaseAsyncCacheTask::EState::Notstarted);
 		if (Task->GetState() == FPoseSearchDatabaseAsyncCacheTask::EState::Prestarted)
diff --git a/Engine/Plugins/Experimental/Animation/PoseSearch/Source/Runtime/Private/PoseSearchDerivedDataKey.cpp b/Engine/Plugins/Experimental/Animation/PoseSearch/Source/Runtime/Private/PoseSearchDerivedDataKey.cpp
--- a/Engine/Plugins/Experimental/Animation/PoseSearch/Source/Runtime/Private/PoseSearchDerivedDataKey.cpp
+++ b/Engine/Plugins/Experimental/Animation/PoseSearch/Source/Runtime/Private/PoseSearchDerivedDataKey.cpp
@@ -22,7 +22,6 @@ namespace UE::PoseSearch
 
 FKeyBuilder::FKeyBuilder()
 {
-	bAnyAssetNotReady = false;
 	ArIgnoreOuterRef = true;
 
 	// Set FDerivedDataKeyBuilder to be a saving archive instead of a reference collector.
@@ -32,16 +31,16 @@ FKeyBuilder::FKeyBuilder()
 	SetIsSaving(true);
 }
 
-FKeyBuilder::FKeyBuilder(const UObject* Object, bool bUseDataVer)
+FKeyBuilder::FKeyBuilder(const UObject* Object, bool bUseDataVer, bool bPerformConditionalPostLoadIfRequired)
 : FKeyBuilder()
 {
 	check(Object);
-	bAnyAssetNotReady = false;
+	bPerformConditionalPostLoad = bPerformConditionalPostLoadIfRequired;
 
 	if (bUseDataVer)
 	{
 		// used to invalidate the key without having to change POSESEARCHDB_DERIVEDDATA_VER all the times
-		int32 POSESEARCHDB_DERIVEDDATA_VER_SMALL = 173;
+		int32 POSESEARCHDB_DERIVEDDATA_VER_SMALL = 179;
 		FGuid VersionGuid = FDevSystemGuids::GetSystemGuid(FDevSystemGuids::Get().POSESEARCHDB_DERIVEDDATA_VER);
 
 		*this << VersionGuid;
@@ -129,9 +128,17 @@ FArchive& FKeyBuilder::operator<<(class UObject*& Object)
 	{
 		if (Object->HasAnyFlags(RF_NeedPostLoad))
 		{
-			bAnyAssetNotReady = true;
+			if (bPerformConditionalPostLoad)
+			{
+				Object->ConditionalPostLoad();
+			}
+			else
+			{
+				bAnyAssetNotReady = true;
+			}
 		}
-		else
+		
+		if (!bAnyAssetNotReady)
 		{
 			#if UE_POSE_SEARCH_DERIVED_DATA_LOGGING
 			++Indentation;
diff --git a/Engine/Plugins/Experimental/Animation/PoseSearch/Source/Runtime/Public/PoseSearch/PoseSearchDerivedDataKey.h b/Engine/Plugins/Experimental/Animation/PoseSearch/Source/Runtime/Public/PoseSearch/PoseSearchDerivedDataKey.h
--- a/Engine/Plugins/Experimental/Animation/PoseSearch/Source/Runtime/Public/PoseSearch/PoseSearchDerivedDataKey.h
+++ b/Engine/Plugins/Experimental/Animation/PoseSearch/Source/Runtime/Public/PoseSearch/PoseSearchDerivedDataKey.h
@@ -33,7 +33,7 @@ public:
 	inline static const FName NeverInHashName = FName(TEXT("NeverInHash"));
 
 	FKeyBuilder();
-	FKeyBuilder(const UObject* Object, bool bUseDataVer = false);
+	FKeyBuilder(const UObject* Object, bool bUseDataVer, bool bPerformConditionalPostLoadIfRequired);
 
 	using Super::IsSaving;
 	using Super::operator<<;
@@ -70,6 +70,9 @@ protected:
 
 	// true if some dependent assets are not ready (fully loaded)
 	bool bAnyAssetNotReady = false;
+
+	// if true ConditionalPostLoad will be performed on the dependant assets requiring it
+	bool bPerformConditionalPostLoad = false;
 };
 
 } // namespace UE::PoseSearch
