motion matching
- introducing UAnimNotifyState_PoseSearchSamplingAttribute as way of populating position, heading, and velocity channel during database indexing
- introducing ability to extend channels via Blueprint to be able to populate MM runtime query in a scriptable manner
- removing slow Schema IsRootBone method since now root is always RootSchemaBoneIdx

diff --git a/Engine/Plugins/Experimental/Animation/PoseSearch/Source/Runtime/Private/PoseSearchAssetIndexer.cpp b/Engine/Plugins/Experimental/Animation/PoseSearch/Source/Runtime/Private/PoseSearchAssetIndexer.cpp
--- a/Engine/Plugins/Experimental/Animation/PoseSearch/Source/Runtime/Private/PoseSearchAssetIndexer.cpp
+++ b/Engine/Plugins/Experimental/Animation/PoseSearch/Source/Runtime/Private/PoseSearchAssetIndexer.cpp
@@ -101,6 +101,8 @@ void FAssetIndexer::Process(int32 AssetIdx)
 
 	FMemMark Mark(FMemStack::Get());
 
+	bProcessFailed = false;
+
 	// Generate pose metadata
 	const float SequenceLength = AssetSampler.GetPlayLength();
 	for (int32 SampleIdx = GetBeginSampleIdx(); SampleIdx != GetEndSampleIdx(); ++SampleIdx)
@@ -139,12 +141,19 @@ void FAssetIndexer::Process(int32 AssetIdx)
 	{
 		for (const TObjectPtr<UPoseSearchFeatureChannel>& ChannelPtr : Schema.GetChannels())
 		{
-			ChannelPtr->IndexAsset(*this);
+			if (!ChannelPtr->IndexAsset(*this))
+			{
+				bProcessFailed = true;
+				break;
+			}
 		}
 	}
 
 	// Computing stats
-	ComputeStats();
+	if (!bProcessFailed)
+	{
+		ComputeStats();
+	}
 }
 
 void FAssetIndexer::ComputeStats()
@@ -321,7 +330,7 @@ FTransform FAssetIndexer::GetComponentSpaceTransform(float SampleTime, bool& bCl
 	CachedEntry& Entry = GetEntry(SampleTime);
 	bClamped = Entry.bClamped;
 
-	if (!Schema.IsRootBone(SchemaBoneIdx))
+	if (SchemaBoneIdx != RootSchemaBoneIdx)
 	{
 		return CalculateComponentSpaceTransform(Entry, SchemaBoneIdx);
 	}
@@ -336,7 +345,7 @@ FTransform FAssetIndexer::GetTransform(float SampleTime, bool& bClamped, int8 Sc
 	bClamped = Entry.bClamped;
 
 	const FTransform MirroredRootTransform = MirrorTransform(Entry.RootTransform);
-	if (!Schema.IsRootBone(SchemaBoneIdx))
+	if (SchemaBoneIdx != RootSchemaBoneIdx)
 	{
 		return CalculateComponentSpaceTransform(Entry, SchemaBoneIdx) * MirroredRootTransform;
 	}
@@ -344,9 +353,23 @@ FTransform FAssetIndexer::GetTransform(float SampleTime, bool& bClamped, int8 Sc
 	return MirroredRootTransform;
 }
 
+// returns the transform in animation space for the BoneReference at SampleTime seconds
+FTransform FAssetIndexer::GetTransform(float SampleTime, bool& bClamped, const FBoneReference& BoneReference)
+{
+	CachedEntry& Entry = GetEntry(SampleTime);
+	bClamped = Entry.bClamped;
+
+	const FTransform MirroredRootTransform = MirrorTransform(Entry.RootTransform);
+	return CalculateComponentSpaceTransform(Entry, BoneReference) * MirroredRootTransform;
+}
+
 FTransform FAssetIndexer::CalculateComponentSpaceTransform(FAssetIndexer::CachedEntry& Entry, int8 SchemaBoneIdx)
 {
-	const FBoneReference& BoneReference = Schema.BoneReferences[SchemaBoneIdx];
+	return CalculateComponentSpaceTransform(Entry, Schema.BoneReferences[SchemaBoneIdx]);
+}
+
+FTransform FAssetIndexer::CalculateComponentSpaceTransform(FAssetIndexer::CachedEntry& Entry, const FBoneReference& BoneReference)
+{
 	const FCompactPoseBoneIndex CompactBoneIndex = BoneContainer.MakeCompactPoseIndex(FMeshPoseBoneIndex(BoneReference.BoneIndex));
 	return Entry.ComponentSpacePose.GetComponentSpaceTransform(CompactBoneIndex);
 }
@@ -356,7 +379,7 @@ float FAssetIndexer::CalculateSampleTime(int32 SampleIdx) const
 	return SampleIdx / float(Schema.SampleRate);
 }
 
-FQuat FAssetIndexer::GetSampleRotation(float SampleTimeOffset, float OriginTimeOffset, int32 SampleIdx, int8 SchemaSampleBoneIdx, int8 SchemaOriginBoneIdx, EPermutationTimeType PermutationTimeType)
+bool FAssetIndexer::GetSampleRotation(FQuat& OutSampleRotation, float SampleTimeOffset, float OriginTimeOffset, int32 SampleIdx, int8 SchemaSampleBoneIdx, int8 SchemaOriginBoneIdx, EPermutationTimeType PermutationTimeType, int32 SamplingAttributeId)
 {
 	float PermutationSampleTimeOffset = 0.f;
 	float PermutationOriginTimeOffset = 0.f;
@@ -367,7 +390,7 @@ FQuat FAssetIndexer::GetSampleRotation(float SampleTimeOffset, float OriginTimeO
 	const float OriginTime = Time + OriginTimeOffset + PermutationOriginTimeOffset;
 
 	// @todo: add support for SchemaSampleBoneIdx
-	if (!Schema.IsRootBone(SchemaOriginBoneIdx))
+	if (SchemaOriginBoneIdx != RootSchemaBoneIdx)
 	{
 		UE_LOG(LogPoseSearch,
 			Error,
@@ -377,18 +400,62 @@ FQuat FAssetIndexer::GetSampleRotation(float SampleTimeOffset, float OriginTimeO
 	}
 
 	bool bUnused;
+	if (SamplingAttributeId >= 0)
+	{
+		const TConstArrayView<FAnimNotifyEvent> AnimNotifyEvents = AssetSampler.GetAllAnimNotifyEvents();
+		if (!AnimNotifyEvents.IsEmpty())
+		{
+			for (const FAnimNotifyEvent& AnimNotifyEvent : AnimNotifyEvents)
+			{
+				if (const UAnimNotifyState_PoseSearchSamplingAttribute* SamplingAttribute = Cast<UAnimNotifyState_PoseSearchSamplingAttribute>(AnimNotifyEvent.NotifyStateClass))
+				{
+					if (SamplingAttribute->SamplingAttributeId == SamplingAttributeId)
+					{
+						if (SamplingAttribute->Bone.BoneName != NAME_None)
+						{
+							// @todo perhaps cache the initialized bone to speed up this method
+							FBoneReference TempBoneReference = SamplingAttribute->Bone;
+							TempBoneReference.Initialize(BoneContainer.GetSkeletonAsset());
+							if (TempBoneReference.HasValidSetup())
+							{
+								const float SamplingAttributeTime = AnimNotifyEvent.GetTriggerTime();
+								const FTransform RootBoneTransform = GetTransform(OriginTime, bUnused, RootSchemaBoneIdx);
+								const FTransform SamplingAttributeBoneTransform = GetTransform(SamplingAttributeTime, bUnused, TempBoneReference);
+								OutSampleRotation = RootBoneTransform.InverseTransformRotation(SamplingAttributeBoneTransform.GetRotation());
+								return true;
+							}
+
+							UE_LOG(LogPoseSearch, Error, TEXT("FAssetIndexer::GetSampleRotation: required UAnimNotifyState_PoseSearchSamplingAttribute in '%s' has an invalid Bone"), *GetNameSafe(AssetSampler.GetAsset()));
+							OutSampleRotation = FQuat::Identity;
+							return false;
+						}
+						
+						const FTransform RootBoneTransform = GetTransform(OriginTime, bUnused, RootSchemaBoneIdx);
+						OutSampleRotation = RootBoneTransform.InverseTransformRotation(SamplingAttribute->Rotation);
+						return true;
+					}
+				}
+			}
+		}
+
+		UE_LOG(LogPoseSearch, Error, TEXT("FAssetIndexer::GetSamplePositionInternal: required UAnimNotifyState_PoseSearchSamplingAttribute not found in '%s'"), *GetNameSafe(AssetSampler.GetAsset()));
+		OutSampleRotation = FQuat::Identity;
+		return false;
+	}
+
 	if (SampleTime == OriginTime)
 	{
-		return GetComponentSpaceTransform(SampleTime, bUnused, SchemaSampleBoneIdx).GetRotation();
+		OutSampleRotation = GetComponentSpaceTransform(SampleTime, bUnused, SchemaSampleBoneIdx).GetRotation();
+		return true;
 	}
 
 	const FTransform RootBoneTransform = GetTransform(OriginTime, bUnused, RootSchemaBoneIdx);
-	FTransform BoneTransform = GetTransform(SampleTime, bUnused, SchemaSampleBoneIdx);
-	BoneTransform.SetToRelativeTransform(RootBoneTransform);
-	return BoneTransform.GetRotation();
+	const FTransform BoneTransform = GetTransform(SampleTime, bUnused, SchemaSampleBoneIdx);
+	OutSampleRotation = RootBoneTransform.InverseTransformRotation(BoneTransform.GetRotation());
+	return true;
 }
 
-FVector FAssetIndexer::GetSamplePosition(float SampleTimeOffset, float OriginTimeOffset, int32 SampleIdx, int8 SchemaSampleBoneIdx, int8 SchemaOriginBoneIdx, EPermutationTimeType PermutationTimeType)
+bool FAssetIndexer::GetSamplePosition(FVector& OutSamplePosition, float SampleTimeOffset, float OriginTimeOffset, int32 SampleIdx, int8 SchemaSampleBoneIdx, int8 SchemaOriginBoneIdx, EPermutationTimeType PermutationTimeType, int32 SamplingAttributeId)
 {
 	float PermutationSampleTimeOffset = 0.f;
 	float PermutationOriginTimeOffset = 0.f;
@@ -399,41 +466,99 @@ FVector FAssetIndexer::GetSamplePosition(float SampleTimeOffset, float OriginTim
 	const float OriginTime = Time + OriginTimeOffset + PermutationOriginTimeOffset;
 	
 	bool bUnused;
-	return GetSamplePositionInternal(SampleTime, OriginTime, bUnused, SchemaSampleBoneIdx, SchemaOriginBoneIdx);
+	return GetSamplePositionInternal(OutSamplePosition, SampleTime, OriginTime, bUnused, SchemaSampleBoneIdx, SchemaOriginBoneIdx, SamplingAttributeId);
 }
 
-FVector FAssetIndexer::GetSamplePositionInternal(float SampleTime, float OriginTime, bool& bClamped, int8 SchemaSampleBoneIdx, int8 SchemaOriginBoneIdx)
+bool FAssetIndexer::GetSamplePositionInternal(FVector& OutSamplePosition, float SampleTime, float OriginTime, bool& bClamped, int8 SchemaSampleBoneIdx, int8 SchemaOriginBoneIdx, int32 SamplingAttributeId)
 {
+	bool bUnused;
+	if (SamplingAttributeId >= 0)
+	{
+		const TConstArrayView<FAnimNotifyEvent> AnimNotifyEvents = AssetSampler.GetAllAnimNotifyEvents();
+		if (!AnimNotifyEvents.IsEmpty())
+		{
+			for (const FAnimNotifyEvent& AnimNotifyEvent : AnimNotifyEvents)
+			{
+				if (const UAnimNotifyState_PoseSearchSamplingAttribute* SamplingAttribute = Cast<UAnimNotifyState_PoseSearchSamplingAttribute>(AnimNotifyEvent.NotifyStateClass))
+				{
+					if (SamplingAttribute->SamplingAttributeId == SamplingAttributeId)
+					{
+						if (SamplingAttribute->Bone.BoneName != NAME_None)
+						{
+							// @todo perhaps cache the initialized bone to speed up this method
+							FBoneReference TempBoneReference = SamplingAttribute->Bone;
+							TempBoneReference.Initialize(BoneContainer.GetSkeletonAsset());
+							if (TempBoneReference.HasValidSetup())
+							{
+								const float SamplingAttributeTime = AnimNotifyEvent.GetTriggerTime();
+								const FTransform RootBoneTransform = GetTransform(OriginTime, bUnused, RootSchemaBoneIdx);
+								const FTransform SamplingAttributeBoneTransform = GetTransform(SamplingAttributeTime, bClamped, TempBoneReference);
+								if (SchemaOriginBoneIdx == RootSchemaBoneIdx)
+								{
+									OutSamplePosition = RootBoneTransform.InverseTransformPosition(SamplingAttributeBoneTransform.GetTranslation());
+								}
+								else
+								{
+									bool bOriginClamped;
+									const FTransform OriginBoneTransform = GetTransform(OriginTime, bOriginClamped, SchemaOriginBoneIdx);
+									bClamped |= bOriginClamped;
+									const FVector DeltaBoneTranslation = SamplingAttributeBoneTransform.GetTranslation() - OriginBoneTransform.GetTranslation();
+									OutSamplePosition = RootBoneTransform.InverseTransformVector(DeltaBoneTranslation);
+								}
+								return true;
+							}
+							
+							UE_LOG(LogPoseSearch, Error, TEXT("FAssetIndexer::GetSamplePositionInternal: required UAnimNotifyState_PoseSearchSamplingAttribute in '%s' has an invalid Bone"), *GetNameSafe(AssetSampler.GetAsset()));
+							OutSamplePosition = FVector::ZeroVector;
+							return false;
+						}
+						
+						const FTransform RootBoneTransform = GetTransform(OriginTime, bUnused, RootSchemaBoneIdx);
+						OutSamplePosition = RootBoneTransform.InverseTransformPosition(SamplingAttribute->Position);
+						return true;
+					}
+				}
+			}
+		}
+
+		UE_LOG(LogPoseSearch, Error, TEXT("FAssetIndexer::GetSamplePositionInternal: required UAnimNotifyState_PoseSearchSamplingAttribute not found in '%s'"), *GetNameSafe(AssetSampler.GetAsset()));
+		OutSamplePosition = FVector::ZeroVector;
+		return false;
+	}
+
 	if (SampleTime == OriginTime)
 	{
-		if (Schema.IsRootBone(SchemaOriginBoneIdx))
+		if (SchemaOriginBoneIdx == RootSchemaBoneIdx)
 		{
-			return GetComponentSpaceTransform(SampleTime, bClamped, SchemaSampleBoneIdx).GetTranslation();
+			OutSamplePosition = GetComponentSpaceTransform(SampleTime, bClamped, SchemaSampleBoneIdx).GetTranslation();
+			return true;
 		}
 
 		bool bOriginClamped;
 		const FVector SampleBonePosition = GetComponentSpaceTransform(SampleTime, bClamped, SchemaSampleBoneIdx).GetTranslation();
 		const FVector OriginBonePosition = GetComponentSpaceTransform(OriginTime, bOriginClamped, SchemaOriginBoneIdx).GetTranslation();
 		bClamped |= bOriginClamped;
-		return SampleBonePosition - OriginBonePosition;
+		OutSamplePosition = SampleBonePosition - OriginBonePosition;
+		return true;
 	}
 
-	bool Unused;
-	const FTransform RootBoneTransform = GetTransform(OriginTime, Unused, RootSchemaBoneIdx);
+	const FTransform RootBoneTransform = GetTransform(OriginTime, bUnused, RootSchemaBoneIdx);
 	const FTransform SampleBoneTransform = GetTransform(SampleTime, bClamped, SchemaSampleBoneIdx);
-	if (Schema.IsRootBone(SchemaOriginBoneIdx))
+	if (SchemaOriginBoneIdx == RootSchemaBoneIdx)
 	{
-		return RootBoneTransform.InverseTransformPosition(SampleBoneTransform.GetTranslation());
+		OutSamplePosition = RootBoneTransform.InverseTransformPosition(SampleBoneTransform.GetTranslation());
+		return true;
 	}
 
 	bool bOriginClamped;
 	const FTransform OriginBoneTransform = GetTransform(OriginTime, bOriginClamped, SchemaOriginBoneIdx);
 	bClamped |= bOriginClamped;
 	const FVector DeltaBoneTranslation = SampleBoneTransform.GetTranslation() - OriginBoneTransform.GetTranslation();
-	return RootBoneTransform.InverseTransformVector(DeltaBoneTranslation);
+	OutSamplePosition = RootBoneTransform.InverseTransformVector(DeltaBoneTranslation);
+	return true;
 }
 
-FVector FAssetIndexer::GetSampleVelocity(float SampleTimeOffset, float OriginTimeOffset, int32 SampleIdx, int8 SchemaSampleBoneIdx, int8 SchemaOriginBoneIdx, bool bUseCharacterSpaceVelocities, EPermutationTimeType PermutationTimeType)
+bool FAssetIndexer::GetSampleVelocity(FVector& OutSampleVelocity, float SampleTimeOffset, float OriginTimeOffset, int32 SampleIdx, int8 SchemaSampleBoneIdx, int8 SchemaOriginBoneIdx, bool bUseCharacterSpaceVelocities, EPermutationTimeType PermutationTimeType, int32 SamplingAttributeId)
 {
 	float PermutationSampleTimeOffset = 0.f;
 	float PermutationOriginTimeOffset = 0.f;
@@ -444,22 +569,88 @@ FVector FAssetIndexer::GetSampleVelocity(float SampleTimeOffset, float OriginTim
 	const float OriginTime = Time + OriginTimeOffset + PermutationOriginTimeOffset;
 	const float FiniteDelta = SamplingContext.FiniteDelta;
 
-	bool bClampedPast, bClampedPresent;
-	const FVector BonePositionPast = GetSamplePositionInternal(SampleTime - FiniteDelta, bUseCharacterSpaceVelocities ? OriginTime - FiniteDelta : OriginTime, bClampedPast, SchemaSampleBoneIdx, SchemaOriginBoneIdx);
-	const FVector BonePositionPresent = GetSamplePositionInternal(SampleTime, OriginTime, bClampedPresent, SchemaSampleBoneIdx, SchemaOriginBoneIdx);
-
-	FVector LinearVelocity;
-	if (!bClampedPast)
+	bool bUnused, bClampedPast;
+	FVector BonePositionPast, BonePositionPresent;
+	if (SamplingAttributeId >= 0)
 	{
-		LinearVelocity = (BonePositionPresent - BonePositionPast) / FiniteDelta;
+		const TConstArrayView<FAnimNotifyEvent> AnimNotifyEvents = AssetSampler.GetAllAnimNotifyEvents();
+		if (!AnimNotifyEvents.IsEmpty())
+		{
+			for (const FAnimNotifyEvent& AnimNotifyEvent : AnimNotifyEvents)
+			{
+				if (const UAnimNotifyState_PoseSearchSamplingAttribute* SamplingAttribute = Cast<UAnimNotifyState_PoseSearchSamplingAttribute>(AnimNotifyEvent.NotifyStateClass))
+				{
+					if (SamplingAttribute->SamplingAttributeId == SamplingAttributeId)
+					{
+						if (SamplingAttribute->Bone.BoneName != NAME_None)
+						{
+							// @todo perhaps cache the initialized bone to speed up this method
+							FBoneReference TempBoneReference = SamplingAttribute->Bone;
+							TempBoneReference.Initialize(BoneContainer.GetSkeletonAsset());
+							if (TempBoneReference.HasValidSetup())
+							{
+								const float SamplingAttributeTime = AnimNotifyEvent.GetTriggerTime();
+
+								if (GetSamplePositionInternal(BonePositionPast, SamplingAttributeTime - FiniteDelta, bUseCharacterSpaceVelocities ? OriginTime - FiniteDelta : OriginTime, bClampedPast, SchemaSampleBoneIdx, SchemaOriginBoneIdx, SamplingAttributeId) &&
+									GetSamplePositionInternal(BonePositionPresent, SamplingAttributeTime, OriginTime, bUnused, SchemaSampleBoneIdx, SchemaOriginBoneIdx, SamplingAttributeId))
+								{
+									if (!bClampedPast)
+									{
+										OutSampleVelocity = (BonePositionPresent - BonePositionPast) / FiniteDelta;
+										return true;
+									}
+
+									FVector BonePositionFuture;
+									if (GetSamplePositionInternal(BonePositionFuture, SamplingAttributeTime + FiniteDelta, bUseCharacterSpaceVelocities ? OriginTime + FiniteDelta : OriginTime, bUnused, SchemaSampleBoneIdx, SchemaOriginBoneIdx, SamplingAttributeId))
+									{
+										OutSampleVelocity = (BonePositionFuture - BonePositionPresent) / FiniteDelta;
+										return true;
+									}
+								}
+
+								// @todo: should we log some error message?
+								OutSampleVelocity = FVector::ZeroVector;
+								return false;
+							}
+
+							UE_LOG(LogPoseSearch, Error, TEXT("FAssetIndexer::GetSampleVelocity: required UAnimNotifyState_PoseSearchSamplingAttribute in '%s' has an invalid Bone"), *GetNameSafe(AssetSampler.GetAsset()));
+							OutSampleVelocity = FVector::ZeroVector;
+							return false;
+						}
+						
+						const FTransform RootBoneTransform = GetTransform(OriginTime, bUnused, RootSchemaBoneIdx);
+						OutSampleVelocity = RootBoneTransform.InverseTransformPosition(SamplingAttribute->LinearVelocity);
+						return true;
+					}
+				}
+			}
+		}
+
+		UE_LOG(LogPoseSearch, Error, TEXT("FAssetIndexer::GetSampleVelocity: required UAnimNotifyState_PoseSearchSamplingAttribute not found in '%s'"), *GetNameSafe(AssetSampler.GetAsset()));
+		OutSampleVelocity = FVector::ZeroVector;
+		return false;
 	}
-	else
+
+	if (GetSamplePositionInternal(BonePositionPast, SampleTime - FiniteDelta, bUseCharacterSpaceVelocities ? OriginTime - FiniteDelta : OriginTime, bClampedPast, SchemaSampleBoneIdx, SchemaOriginBoneIdx, -1) &&
+		GetSamplePositionInternal(BonePositionPresent, SampleTime, OriginTime, bUnused, SchemaSampleBoneIdx, SchemaOriginBoneIdx, -1))
 	{
-		bool Unused;
-		const FVector BonePositionFuture = GetSamplePositionInternal(SampleTime + FiniteDelta, bUseCharacterSpaceVelocities ? OriginTime + FiniteDelta : OriginTime, Unused, SchemaSampleBoneIdx, SchemaOriginBoneIdx);
-		LinearVelocity = (BonePositionFuture - BonePositionPresent) / FiniteDelta;
+		if (!bClampedPast)
+		{
+			OutSampleVelocity = (BonePositionPresent - BonePositionPast) / FiniteDelta;
+			return true;
+		}
+
+		FVector BonePositionFuture;
+		if (GetSamplePositionInternal(BonePositionFuture, SampleTime + FiniteDelta, bUseCharacterSpaceVelocities ? OriginTime + FiniteDelta : OriginTime, bUnused, SchemaSampleBoneIdx, SchemaOriginBoneIdx, -1))
+		{
+			OutSampleVelocity = (BonePositionFuture - BonePositionPresent) / FiniteDelta;
+			return true;
+		}
 	}
-	return LinearVelocity;
+
+	// @todo: should we log some error message?
+	OutSampleVelocity = FVector::ZeroVector;
+	return false;
 }
 
 int32 FAssetIndexer::GetBeginSampleIdx() const
diff --git a/Engine/Plugins/Experimental/Animation/PoseSearch/Source/Runtime/Private/PoseSearchAssetSampler.cpp b/Engine/Plugins/Experimental/Animation/PoseSearch/Source/Runtime/Private/PoseSearchAssetSampler.cpp
--- a/Engine/Plugins/Experimental/Animation/PoseSearch/Source/Runtime/Private/PoseSearchAssetSampler.cpp
+++ b/Engine/Plugins/Experimental/Animation/PoseSearch/Source/Runtime/Private/PoseSearchAssetSampler.cpp
@@ -614,6 +614,10 @@ void FAnimationAssetSampler::ExtractPoseSearchNotifyStates(float Time, TArray<UA
 				}
 			}
 		}
+		else
+		{
+			UE_LOG(LogPoseSearch, Error, TEXT("FAnimationAssetSampler::ExtractPoseSearchNotifyStates: Unsupported BlendSpace NotifyTriggerMode for '%s'"), *BlendSpace->GetName());
+		}
 	}
 	else if (const UAnimSequenceBase* SequenceBase = Cast<UAnimSequenceBase>(AnimationAsset.Get()))
 	{
@@ -647,4 +651,40 @@ void FAnimationAssetSampler::ExtractPoseSearchNotifyStates(float Time, TArray<UA
 	}
 }
 
+TConstArrayView<FAnimNotifyEvent> FAnimationAssetSampler::GetAllAnimNotifyEvents() const
+{
+	if (const UBlendSpace* BlendSpace = Cast<UBlendSpace>(AnimationAsset.Get()))
+	{
+		if (BlendSpace->NotifyTriggerMode == ENotifyTriggerMode::HighestWeightedAnimation)
+		{
+			TArray<FBlendSampleData> BlendSamples;
+			int32 TriangulationIndex = 0;
+			if (BlendSpace->GetSamplesFromBlendInput(BlendParameters, BlendSamples, TriangulationIndex, true))
+			{
+				// Find highest weighted
+				const int32 HighestWeightIndex = GetHighestWeightSample(BlendSamples);
+				const FBlendSampleData& BlendSample = BlendSamples[HighestWeightIndex];
+				if (BlendSample.Animation)
+				{
+					return BlendSample.Animation->Notifies;
+				}
+			}
+		}
+		else
+		{
+			UE_LOG(LogPoseSearch, Error, TEXT("FAnimationAssetSampler::ExtractPoseSearchNotifyStates: Unsupported BlendSpace NotifyTriggerMode for '%s'"), *BlendSpace->GetName());
+		}
+	}
+	else if (const UAnimSequenceBase* SequenceBase = Cast<UAnimSequenceBase>(AnimationAsset.Get()))
+	{
+		return SequenceBase->Notifies;
+	}
+	else
+	{
+		checkNoEntry();
+	}
+
+	return TConstArrayView<FAnimNotifyEvent>(); 
+}
+
 } // namespace UE::PoseSearch
diff --git a/Engine/Plugins/Experimental/Animation/PoseSearch/Source/Runtime/Private/PoseSearchContext.cpp b/Engine/Plugins/Experimental/Animation/PoseSearch/Source/Runtime/Private/PoseSearchContext.cpp
--- a/Engine/Plugins/Experimental/Animation/PoseSearch/Source/Runtime/Private/PoseSearchContext.cpp
+++ b/Engine/Plugins/Experimental/Animation/PoseSearch/Source/Runtime/Private/PoseSearchContext.cpp
@@ -59,7 +59,7 @@ FVector FDebugDrawParams::ExtractPosition(TConstArrayView<float> PoseVector, con
 	return WorldBonePosition;
 }
 
-FVector FDebugDrawParams::ExtractPosition(TConstArrayView<float> PoseVector, float SampleTimeOffset, int8 SchemaBoneIdx, EPermutationTimeType PermutationTimeType) const
+FVector FDebugDrawParams::ExtractPosition(TConstArrayView<float> PoseVector, float SampleTimeOffset, int8 SchemaBoneIdx, EPermutationTimeType PermutationTimeType, int32 SamplingAttributeId) const
 {
 	// we don't wanna ask for a SchemaOriginBoneIdx in the future or past
 	check(PermutationTimeType != EPermutationTimeType::UsePermutationTime);
@@ -68,14 +68,15 @@ FVector FDebugDrawParams::ExtractPosition(TConstArrayView<float> PoseVector, flo
 		// looking for a UPoseSearchFeatureChannel_Position that matches the TimeOffset and SchemaBoneIdx,
 		// with SchemaOriginBoneIdx to be the root bone and the appropriate PermutationTimeType 
 		if (const UPoseSearchFeatureChannel_Position* FoundPosition = static_cast<const UPoseSearchFeatureChannel_Position*>(
-			Schema->FindChannel([SampleTimeOffset, SchemaBoneIdx, PermutationTimeType, Schema](const UPoseSearchFeatureChannel* Channel) -> const UPoseSearchFeatureChannel_Position*
+			Schema->FindChannel([SampleTimeOffset, SchemaBoneIdx, PermutationTimeType, SamplingAttributeId, Schema](const UPoseSearchFeatureChannel* Channel) -> const UPoseSearchFeatureChannel_Position*
 				{
 					if (const UPoseSearchFeatureChannel_Position* Position = Cast<UPoseSearchFeatureChannel_Position>(Channel))
 					{
 						if (Position->SchemaBoneIdx == SchemaBoneIdx &&
 							Position->SampleTimeOffset == SampleTimeOffset &&
 							Position->PermutationTimeType == PermutationTimeType &&
-							Schema->IsRootBone(Position->SchemaOriginBoneIdx))
+							Position->SamplingAttributeId == SamplingAttributeId &&
+							Position->SchemaOriginBoneIdx == RootSchemaBoneIdx)
 						{
 							return Position;
 						}
@@ -350,9 +351,10 @@ void FDebugDrawParams::DrawFeatureVector(int32 PoseIdx)
 
 //////////////////////////////////////////////////////////////////////////
 // FSearchContext
-FSearchContext::FSearchContext(const FPoseSearchQueryTrajectory* InTrajectory, const IPoseHistory* InHistory, float InDesiredPermutationTimeOffset,
+FSearchContext::FSearchContext(const UAnimInstance* InAnimInstance, const FPoseSearchQueryTrajectory* InTrajectory, const IPoseHistory* InHistory, float InDesiredPermutationTimeOffset,
 	const FPoseIndicesHistory* InPoseIndicesHistory, const FSearchResult& InCurrentResult, float InPoseJumpThresholdTime, bool bInForceInterrupt)
-: Trajectory(InTrajectory)
+: AnimInstance(InAnimInstance)
+, Trajectory(InTrajectory)
 , History(InHistory)
 , DesiredPermutationTimeOffset(InDesiredPermutationTimeOffset)
 , PoseIndicesHistory(InPoseIndicesHistory)
@@ -360,6 +362,7 @@ FSearchContext::FSearchContext(const FPoseSearchQueryTrajectory* InTrajectory, c
 , PoseJumpThresholdTime(InPoseJumpThresholdTime)
 , bForceInterrupt(bInForceInterrupt)
 {
+	check(AnimInstance);
 	if (CurrentResult.IsValid())
 	{
 		const FSearchIndex& SearchIndex = CurrentResult.Database->GetSearchIndex();
@@ -376,7 +379,7 @@ FSearchContext::FSearchContext(const FPoseSearchQueryTrajectory* InTrajectory, c
 	}
 }
 
-FQuat FSearchContext::GetSampleRotation(float SampleTimeOffset, float OriginTimeOffset, const UPoseSearchSchema* Schema, int8 SchemaSampleBoneIdx, int8 SchemaOriginBoneIdx, bool bUseHistoryRoot, EPermutationTimeType PermutationTimeType)
+FQuat FSearchContext::GetSampleRotation(float SampleTimeOffset, float OriginTimeOffset, const UPoseSearchSchema* Schema, int8 SchemaSampleBoneIdx, int8 SchemaOriginBoneIdx, bool bUseHistoryRoot, EPermutationTimeType PermutationTimeType, const FQuat* SampleBoneRotationWorldOverride)
 {
 	float PermutationSampleTimeOffset = 0.f;
 	float PermutationOriginTimeOffset = 0.f;
@@ -385,10 +388,10 @@ FQuat FSearchContext::GetSampleRotation(float SampleTimeOffset, float OriginTime
 	const float SampleTime = SampleTimeOffset + PermutationSampleTimeOffset;
 	const float OriginTime = OriginTimeOffset + PermutationOriginTimeOffset;
 
-	return GetSampleRotationInternal(SampleTime, OriginTime, Schema, SchemaSampleBoneIdx, SchemaOriginBoneIdx, bUseHistoryRoot);
+	return GetSampleRotationInternal(SampleTime, OriginTime, Schema, SchemaSampleBoneIdx, SchemaOriginBoneIdx, bUseHistoryRoot, SampleBoneRotationWorldOverride);
 }
 
-FVector FSearchContext::GetSamplePosition(float SampleTimeOffset, float OriginTimeOffset, const UPoseSearchSchema* Schema, int8 SchemaSampleBoneIdx, int8 SchemaOriginBoneIdx, bool bUseHistoryRoot, EPermutationTimeType PermutationTimeType)
+FVector FSearchContext::GetSamplePosition(float SampleTimeOffset, float OriginTimeOffset, const UPoseSearchSchema* Schema, int8 SchemaSampleBoneIdx, int8 SchemaOriginBoneIdx, bool bUseHistoryRoot, EPermutationTimeType PermutationTimeType, const FVector* SampleBonePositionWorldOverride)
 {
 	float PermutationSampleTimeOffset = 0.f;
 	float PermutationOriginTimeOffset = 0.f;
@@ -396,10 +399,10 @@ FVector FSearchContext::GetSamplePosition(float SampleTimeOffset, float OriginTi
 
 	const float SampleTime = SampleTimeOffset + PermutationSampleTimeOffset;
 	const float OriginTime = OriginTimeOffset + PermutationOriginTimeOffset;
-	return GetSamplePositionInternal(SampleTime, OriginTime, Schema, SchemaSampleBoneIdx, SchemaOriginBoneIdx, bUseHistoryRoot);
+	return GetSamplePositionInternal(SampleTime, OriginTime, Schema, SchemaSampleBoneIdx, SchemaOriginBoneIdx, bUseHistoryRoot, SampleBonePositionWorldOverride);
 }
 
-FVector FSearchContext::GetSampleVelocity(float SampleTimeOffset, float OriginTimeOffset, const UPoseSearchSchema* Schema, int8 SchemaSampleBoneIdx, int8 SchemaOriginBoneIdx, bool bUseCharacterSpaceVelocities, bool bUseHistoryRoot, EPermutationTimeType PermutationTimeType)
+FVector FSearchContext::GetSampleVelocity(float SampleTimeOffset, float OriginTimeOffset, const UPoseSearchSchema* Schema, int8 SchemaSampleBoneIdx, int8 SchemaOriginBoneIdx, bool bUseCharacterSpaceVelocities, bool bUseHistoryRoot, EPermutationTimeType PermutationTimeType, const FVector* SampleBoneVelocityWorldOverride)
 {
 	float PermutationSampleTimeOffset = 0.f;
 	float PermutationOriginTimeOffset = 0.f;
@@ -410,6 +413,12 @@ FVector FSearchContext::GetSampleVelocity(float SampleTimeOffset, float OriginTi
 	const float FiniteDelta = History ? History->GetSampleTimeInterval() : 1 / 60.0f;
 	check(FiniteDelta > UE_KINDA_SMALL_NUMBER);
 
+	if (SampleBoneVelocityWorldOverride)
+	{
+		const FTransform RootBoneTransform = GetTransform(OriginTime, Schema, RootSchemaBoneIdx, bUseHistoryRoot);
+		return RootBoneTransform.InverseTransformVector(*SampleBoneVelocityWorldOverride);
+	}
+
 	// calculating the Position in component space for the bone indexed by SchemaSampleBoneIdx
 	const FVector PreviousTranslation = GetSamplePositionInternal(SampleTime - FiniteDelta, bUseCharacterSpaceVelocities ? OriginTime - FiniteDelta : OriginTime, Schema, SchemaSampleBoneIdx, SchemaOriginBoneIdx, bUseHistoryRoot);
 	const FVector CurrentTranslation = GetSamplePositionInternal(SampleTime, OriginTime, Schema, SchemaSampleBoneIdx, SchemaOriginBoneIdx, bUseHistoryRoot);
@@ -487,11 +496,24 @@ FTransform FSearchContext::GetComponentSpaceTransform(float SampleTime, const UP
 	return FTransform::Identity;
 }
 
-FVector FSearchContext::GetSamplePositionInternal(float SampleTime, float OriginTime, const UPoseSearchSchema* Schema, int8 SchemaSampleBoneIdx, int8 SchemaOriginBoneIdx, bool bUseHistoryRoot)
+FVector FSearchContext::GetSamplePositionInternal(float SampleTime, float OriginTime, const UPoseSearchSchema* Schema, int8 SchemaSampleBoneIdx, int8 SchemaOriginBoneIdx, bool bUseHistoryRoot, const FVector* SampleBonePositionWorldOverride)
 {
+	if (SampleBonePositionWorldOverride)
+	{
+		const FTransform RootBoneTransform = GetTransform(OriginTime, Schema, RootSchemaBoneIdx, bUseHistoryRoot);
+		if (SchemaOriginBoneIdx == RootSchemaBoneIdx)
+		{
+			return RootBoneTransform.InverseTransformPosition(*SampleBonePositionWorldOverride);
+		}
+
+		const FTransform OriginBoneTransform = GetTransform(OriginTime, Schema, SchemaOriginBoneIdx, bUseHistoryRoot);
+		const FVector DeltaBoneTranslation = *SampleBonePositionWorldOverride - OriginBoneTransform.GetTranslation();
+		return RootBoneTransform.InverseTransformVector(DeltaBoneTranslation);
+	}
+
 	if (SampleTime == OriginTime)
 	{
-		if (Schema->IsRootBone(SchemaOriginBoneIdx))
+		if (SchemaOriginBoneIdx == RootSchemaBoneIdx)
 		{
 			return GetComponentSpaceTransform(SampleTime, Schema, SchemaSampleBoneIdx).GetTranslation();
 		}
@@ -503,7 +525,7 @@ FVector FSearchContext::GetSamplePositionInternal(float SampleTime, float Origin
 
 	const FTransform RootBoneTransform = GetTransform(OriginTime, Schema, RootSchemaBoneIdx, bUseHistoryRoot);
 	const FTransform SampleBoneTransform = GetTransform(SampleTime, Schema, SchemaSampleBoneIdx, bUseHistoryRoot);
-	if (Schema->IsRootBone(SchemaOriginBoneIdx))
+	if (SchemaOriginBoneIdx == RootSchemaBoneIdx)
 	{
 		return RootBoneTransform.InverseTransformPosition(SampleBoneTransform.GetTranslation());
 	}
@@ -513,11 +535,24 @@ FVector FSearchContext::GetSamplePositionInternal(float SampleTime, float Origin
 	return RootBoneTransform.InverseTransformVector(DeltaBoneTranslation);
 }
 
-FQuat FSearchContext::GetSampleRotationInternal(float SampleTime, float OriginTime, const UPoseSearchSchema* Schema, int8 SchemaSampleBoneIdx, int8 SchemaOriginBoneIdx, bool bUseHistoryRoot)
+FQuat FSearchContext::GetSampleRotationInternal(float SampleTime, float OriginTime, const UPoseSearchSchema* Schema, int8 SchemaSampleBoneIdx, int8 SchemaOriginBoneIdx, bool bUseHistoryRoot, const FQuat* SampleBoneRotationWorldOverride)
 {
+	if (SampleBoneRotationWorldOverride)
+	{
+		const FTransform RootBoneTransform = GetTransform(OriginTime, Schema, RootSchemaBoneIdx, bUseHistoryRoot);
+		if (SchemaOriginBoneIdx == RootSchemaBoneIdx)
+		{
+			return RootBoneTransform.InverseTransformRotation(*SampleBoneRotationWorldOverride);
+		}
+
+		const FTransform OriginBoneTransform = GetTransform(OriginTime, Schema, SchemaOriginBoneIdx, bUseHistoryRoot);
+		const FQuat DeltaBoneRotation = OriginBoneTransform.InverseTransformRotation(*SampleBoneRotationWorldOverride);
+		return RootBoneTransform.InverseTransformRotation(DeltaBoneRotation);
+	}
+
 	if (SampleTime == OriginTime)
 	{
-		if (Schema->IsRootBone(SchemaOriginBoneIdx))
+		if (SchemaOriginBoneIdx == RootSchemaBoneIdx)
 		{
 			return GetComponentSpaceTransform(SampleTime, Schema, SchemaSampleBoneIdx).GetRotation();
 		}
@@ -529,7 +564,7 @@ FQuat FSearchContext::GetSampleRotationInternal(float SampleTime, float OriginTi
 
 	const FTransform RootBoneTransform = GetTransform(OriginTime, Schema, RootSchemaBoneIdx, bUseHistoryRoot);
 	const FTransform SampleBoneTransform = GetTransform(SampleTime, Schema, SchemaSampleBoneIdx, bUseHistoryRoot);
-	if (Schema->IsRootBone(SchemaOriginBoneIdx))
+	if (SchemaOriginBoneIdx == RootSchemaBoneIdx)
 	{
 		return RootBoneTransform.InverseTransformRotation(SampleBoneTransform.GetRotation());
 	}
diff --git a/Engine/Plugins/Experimental/Animation/PoseSearch/Source/Runtime/Private/PoseSearchDatabaseIndexingContext.cpp b/Engine/Plugins/Experimental/Animation/PoseSearch/Source/Runtime/Private/PoseSearchDatabaseIndexingContext.cpp
--- a/Engine/Plugins/Experimental/Animation/PoseSearch/Source/Runtime/Private/PoseSearchDatabaseIndexingContext.cpp
+++ b/Engine/Plugins/Experimental/Animation/PoseSearch/Source/Runtime/Private/PoseSearchDatabaseIndexingContext.cpp
@@ -141,6 +141,14 @@ bool FDatabaseIndexingContext::IndexDatabase(FSearchIndexBase& SearchIndexBase,
 	// Index asset data
 	ParallelFor(Indexers.Num(), [this](int32 AssetIdx) { Indexers[AssetIdx].Process(AssetIdx); }, ParallelForFlags);
 
+	for (const FAssetIndexer& Indexer : Indexers)
+	{
+		if (Indexer.IsProcessFailed())
+		{
+			return false;
+		}
+	}
+
 	if (Owner.IsCanceled())
 	{
 		return false;
diff --git a/Engine/Plugins/Experimental/Animation/PoseSearch/Source/Runtime/Private/PoseSearchFeatureChannel.cpp b/Engine/Plugins/Experimental/Animation/PoseSearch/Source/Runtime/Private/PoseSearchFeatureChannel.cpp
--- a/Engine/Plugins/Experimental/Animation/PoseSearch/Source/Runtime/Private/PoseSearchFeatureChannel.cpp
+++ b/Engine/Plugins/Experimental/Animation/PoseSearch/Source/Runtime/Private/PoseSearchFeatureChannel.cpp
@@ -182,7 +182,6 @@ const UPoseSearchSchema* UPoseSearchFeatureChannel::GetSchema() const
 		}
 		Outer = Outer->GetOuter();
 	}
-	checkNoEntry();
 	return nullptr;
 }
 #endif // WITH_EDITOR
@@ -191,11 +190,15 @@ USkeleton* UPoseSearchFeatureChannel::GetSkeleton(bool& bInvalidSkeletonIsError,
 {
 	bInvalidSkeletonIsError = false;
 #if WITH_EDITOR
-	return GetSchema()->Skeleton;
+	// blueprint generated classes don't have a schema, until they're instanced by the schema
+	if (const UPoseSearchSchema* Schema = GetSchema())
+	{
+		return Schema->Skeleton;
+	}
 #else
 	checkNoEntry();
-	return nullptr;
 #endif
+	return nullptr;
 }
 
 
diff --git a/Engine/Plugins/Experimental/Animation/PoseSearch/Source/Runtime/Private/PoseSearchFeatureChannel_FilterCrashingLegs.cpp b/Engine/Plugins/Experimental/Animation/PoseSearch/Source/Runtime/Private/PoseSearchFeatureChannel_FilterCrashingLegs.cpp
--- a/Engine/Plugins/Experimental/Animation/PoseSearch/Source/Runtime/Private/PoseSearchFeatureChannel_FilterCrashingLegs.cpp
+++ b/Engine/Plugins/Experimental/Animation/PoseSearch/Source/Runtime/Private/PoseSearchFeatureChannel_FilterCrashingLegs.cpp
@@ -149,21 +149,26 @@ void UPoseSearchFeatureChannel_FilterCrashingLegs::FillWeights(TArrayView<float>
 	}
 }
 
-void UPoseSearchFeatureChannel_FilterCrashingLegs::IndexAsset(UE::PoseSearch::FAssetIndexer& Indexer) const
+bool UPoseSearchFeatureChannel_FilterCrashingLegs::IndexAsset(UE::PoseSearch::FAssetIndexer& Indexer) const
 {
 	using namespace UE::PoseSearch;
-
+	FVector RightThighPosition, LeftThighPosition, RightFootPosition, LeftFootPosition;
 	for (int32 SampleIdx = Indexer.GetBeginSampleIdx(); SampleIdx != Indexer.GetEndSampleIdx(); ++SampleIdx)
 	{
-		const FVector RightThighPosition = Indexer.GetSamplePosition(0.f, 0.f, SampleIdx, RightThighIdx);
-		const FVector LeftThighPosition = Indexer.GetSamplePosition(0.f, 0.f, SampleIdx, LeftThighIdx);
-		const FVector RightFootPosition = Indexer.GetSamplePosition(0.f, 0.f, SampleIdx, RightFootIdx);
-		const FVector LeftFootPosition = Indexer.GetSamplePosition(0.f, 0.f, SampleIdx, LeftFootIdx);
-
-		const float CrashingLegsValue = ComputeCrashingLegsValue(RightThighPosition, LeftThighPosition, RightFootPosition, LeftFootPosition);
-
-		FFeatureVectorHelper::EncodeFloat(Indexer.GetPoseVector(SampleIdx), ChannelDataOffset, CrashingLegsValue);
+		if (Indexer.GetSamplePosition(RightThighPosition, 0.f, 0.f, SampleIdx, RightThighIdx) &&
+			Indexer.GetSamplePosition(LeftThighPosition, 0.f, 0.f, SampleIdx, LeftThighIdx) &&
+			Indexer.GetSamplePosition(RightFootPosition, 0.f, 0.f, SampleIdx, RightFootIdx) &&
+			Indexer.GetSamplePosition(LeftFootPosition, 0.f, 0.f, SampleIdx, LeftFootIdx))
+		{
+			const float CrashingLegsValue = ComputeCrashingLegsValue(RightThighPosition, LeftThighPosition, RightFootPosition, LeftFootPosition);
+			FFeatureVectorHelper::EncodeFloat(Indexer.GetPoseVector(SampleIdx), ChannelDataOffset, CrashingLegsValue);
+		}
+		else
+		{
+			return false;
+		}
 	}
+	return true;
 }
 
 FString UPoseSearchFeatureChannel_FilterCrashingLegs::GetLabel() const
diff --git a/Engine/Plugins/Experimental/Animation/PoseSearch/Source/Runtime/Private/PoseSearchFeatureChannel_FilterCrashingLegs.h b/Engine/Plugins/Experimental/Animation/PoseSearch/Source/Runtime/Private/PoseSearchFeatureChannel_FilterCrashingLegs.h
--- a/Engine/Plugins/Experimental/Animation/PoseSearch/Source/Runtime/Private/PoseSearchFeatureChannel_FilterCrashingLegs.h
+++ b/Engine/Plugins/Experimental/Animation/PoseSearch/Source/Runtime/Private/PoseSearchFeatureChannel_FilterCrashingLegs.h
@@ -64,7 +64,7 @@ public:
 
 #if WITH_EDITOR
 	virtual void FillWeights(TArrayView<float> Weights) const override;
-	virtual void IndexAsset(UE::PoseSearch::FAssetIndexer& Indexer) const override;
+	virtual bool IndexAsset(UE::PoseSearch::FAssetIndexer& Indexer) const override;
 	virtual FString GetLabel() const override;
 #endif
 
diff --git a/Engine/Plugins/Experimental/Animation/PoseSearch/Source/Runtime/Private/PoseSearchFeatureChannel_Group.cpp b/Engine/Plugins/Experimental/Animation/PoseSearch/Source/Runtime/Private/PoseSearchFeatureChannel_Group.cpp
--- a/Engine/Plugins/Experimental/Animation/PoseSearch/Source/Runtime/Private/PoseSearchFeatureChannel_Group.cpp
+++ b/Engine/Plugins/Experimental/Animation/PoseSearch/Source/Runtime/Private/PoseSearchFeatureChannel_Group.cpp
@@ -39,15 +39,19 @@ void UPoseSearchFeatureChannel_GroupBase::FillWeights(TArrayView<float> Weights)
 	}
 }
 
-void UPoseSearchFeatureChannel_GroupBase::IndexAsset(UE::PoseSearch::FAssetIndexer& Indexer) const
+bool UPoseSearchFeatureChannel_GroupBase::IndexAsset(UE::PoseSearch::FAssetIndexer& Indexer) const
 {
 	for (const TObjectPtr<UPoseSearchFeatureChannel>& SubChannelPtr : GetSubChannels())
 	{
 		if (const UPoseSearchFeatureChannel* SubChannel = SubChannelPtr.Get())
 		{
-			SubChannel->IndexAsset(Indexer);
+			if (!SubChannel->IndexAsset(Indexer))
+			{
+				return false;
+			}
 		}
 	}
+	return true;
 }
 #endif // WITH_EDITOR
 
diff --git a/Engine/Plugins/Experimental/Animation/PoseSearch/Source/Runtime/Private/PoseSearchFeatureChannel_Group.h b/Engine/Plugins/Experimental/Animation/PoseSearch/Source/Runtime/Private/PoseSearchFeatureChannel_Group.h
--- a/Engine/Plugins/Experimental/Animation/PoseSearch/Source/Runtime/Private/PoseSearchFeatureChannel_Group.h
+++ b/Engine/Plugins/Experimental/Animation/PoseSearch/Source/Runtime/Private/PoseSearchFeatureChannel_Group.h
@@ -24,7 +24,7 @@ public:
 
 #if WITH_EDITOR
 	virtual void FillWeights(TArrayView<float> Weights) const override;
-	virtual void IndexAsset(UE::PoseSearch::FAssetIndexer& Indexer) const override;
+	virtual bool IndexAsset(UE::PoseSearch::FAssetIndexer& Indexer) const override;
 #endif //WITH_EDITOR
 
 	// IPoseSearchFilter interface
diff --git a/Engine/Plugins/Experimental/Animation/PoseSearch/Source/Runtime/Private/PoseSearchFeatureChannel_Heading.cpp b/Engine/Plugins/Experimental/Animation/PoseSearch/Source/Runtime/Private/PoseSearchFeatureChannel_Heading.cpp
--- a/Engine/Plugins/Experimental/Animation/PoseSearch/Source/Runtime/Private/PoseSearchFeatureChannel_Heading.cpp
+++ b/Engine/Plugins/Experimental/Animation/PoseSearch/Source/Runtime/Private/PoseSearchFeatureChannel_Heading.cpp
@@ -7,6 +7,12 @@
 #include "PoseSearch/PoseSearchDatabase.h"
 #include "PoseSearch/PoseSearchSchema.h"
 #include "PoseSearchFeatureChannel_Position.h"
+#include "Engine/BlueprintGeneratedClass.h"
+
+UPoseSearchFeatureChannel_Heading::UPoseSearchFeatureChannel_Heading()
+{
+	bUseBlueprintQueryOverride = Cast<UBlueprintGeneratedClass>(GetClass()) != nullptr;
+}
 
 void UPoseSearchFeatureChannel_Heading::FindOrAddToSchema(UPoseSearchSchema* Schema, float SampleTimeOffset, const FName& BoneName, EHeadingAxis HeadingAxis)
 {
@@ -78,27 +84,37 @@ void UPoseSearchFeatureChannel_Heading::BuildQuery(UE::PoseSearch::FSearchContex
 {
 	using namespace UE::PoseSearch;
 
-	const bool bIsCurrentResultValid = SearchContext.GetCurrentResult().IsValid() && SearchContext.GetCurrentResult().Database->Schema == InOutQuery.GetSchema();
-	const bool bSkip = InputQueryPose != EInputQueryPose::UseCharacterPose && bIsCurrentResultValid;
-	const bool bIsRootBone = InOutQuery.GetSchema()->IsRootBone(SchemaBoneIdx);
-	if (bSkip || (!SearchContext.IsHistoryValid() && !bIsRootBone))
+	check(InOutQuery.GetSchema());
+	const bool bIsRootBone = SchemaBoneIdx == RootSchemaBoneIdx;
+	if (bUseBlueprintQueryOverride)
+	{
+		const FQuat BoneRotationWorld = BP_GetWorldRotation(SearchContext.GetAnimInstance());
+		const FQuat BoneRotation = SearchContext.GetSampleRotation(SampleTimeOffset, OriginTimeOffset, InOutQuery.GetSchema(), SchemaBoneIdx, RootSchemaBoneIdx, /*!bIsRootBone*/ true, EPermutationTimeType::UseSampleTime, &BoneRotationWorld);
+		FFeatureVectorHelper::EncodeVector(InOutQuery.EditValues(), ChannelDataOffset, GetAxis(BoneRotation), ComponentStripping);
+	}
+	else
 	{
-		if (bIsCurrentResultValid)
+		const bool bIsCurrentResultValid = SearchContext.GetCurrentResult().IsValid() && SearchContext.GetCurrentResult().Database->Schema == InOutQuery.GetSchema();
+		const bool bSkip = InputQueryPose != EInputQueryPose::UseCharacterPose && bIsCurrentResultValid;
+		if (bSkip || (!SearchContext.IsHistoryValid() && !bIsRootBone))
 		{
-			FFeatureVectorHelper::Copy(InOutQuery.EditValues(), ChannelDataOffset, ChannelCardinality, SearchContext.GetCurrentResultPoseVector());
+			if (bIsCurrentResultValid)
+			{
+				FFeatureVectorHelper::Copy(InOutQuery.EditValues(), ChannelDataOffset, ChannelCardinality, SearchContext.GetCurrentResultPoseVector());
+			}
+			else
+			{
+				// we leave the InOutQuery set to zero since the SearchContext.History is invalid and it'll fail if we continue
+				UE_LOG(LogPoseSearch, Error, TEXT("UPoseSearchFeatureChannel_Heading::BuildQuery - Failed because Pose History Node is missing."));
+			}
 		}
 		else
 		{
-			// we leave the InOutQuery set to zero since the SearchContext.History is invalid and it'll fail if we continue
-			UE_LOG(LogPoseSearch, Error, TEXT("UPoseSearchFeatureChannel_Heading::BuildQuery - Failed because Pose History Node is missing."));
+			// calculating the BoneRotation in component space for the bone indexed by SchemaBoneIdx
+			const FQuat BoneRotation = SearchContext.GetSampleRotation(SampleTimeOffset, OriginTimeOffset, InOutQuery.GetSchema(), SchemaBoneIdx, RootSchemaBoneIdx, !bIsRootBone);
+			FFeatureVectorHelper::EncodeVector(InOutQuery.EditValues(), ChannelDataOffset, GetAxis(BoneRotation), ComponentStripping);
 		}
 	}
-	else
-	{
-		// calculating the BoneRotation in component space for the bone indexed by SchemaBoneIdx
-		const FQuat BoneRotation = SearchContext.GetSampleRotation(SampleTimeOffset, OriginTimeOffset, InOutQuery.GetSchema(), SchemaBoneIdx, RootSchemaBoneIdx, !bIsRootBone);
-		FFeatureVectorHelper::EncodeVector(InOutQuery.EditValues(), ChannelDataOffset, GetAxis(BoneRotation), ComponentStripping);
-	}
 }
 
 #if ENABLE_DRAW_DEBUG
@@ -108,7 +124,7 @@ void UPoseSearchFeatureChannel_Heading::DebugDraw(const UE::PoseSearch::FDebugDr
 
 	const FColor Color = DebugColor.ToFColor(true);
 	const FVector BoneHeading = DrawParams.ExtractRotation(PoseVector, OriginTimeOffset).RotateVector(FFeatureVectorHelper::DecodeVector(PoseVector, ChannelDataOffset, ComponentStripping));
-	const FVector BonePos = DrawParams.ExtractPosition(PoseVector, SampleTimeOffset, SchemaBoneIdx);
+	const FVector BonePos = DrawParams.ExtractPosition(PoseVector, SampleTimeOffset, SchemaBoneIdx, EPermutationTimeType::UseSampleTime, SamplingAttributeId);
 
 	DrawParams.DrawPoint(BonePos, Color, 3.f);
 	DrawParams.DrawLine(BonePos + BoneHeading * 4.f, BonePos + BoneHeading * 15.f, Color);
@@ -124,15 +140,23 @@ void UPoseSearchFeatureChannel_Heading::FillWeights(TArrayView<float> Weights) c
 	}
 }
 
-void UPoseSearchFeatureChannel_Heading::IndexAsset(UE::PoseSearch::FAssetIndexer& Indexer) const
+bool UPoseSearchFeatureChannel_Heading::IndexAsset(UE::PoseSearch::FAssetIndexer& Indexer) const
 {
 	using namespace UE::PoseSearch;
 
+	FQuat SampleRotation = FQuat::Identity;
 	for (int32 SampleIdx = Indexer.GetBeginSampleIdx(); SampleIdx != Indexer.GetEndSampleIdx(); ++SampleIdx)
 	{
-		const FVector Heading = GetAxis(Indexer.GetSampleRotation(SampleTimeOffset, OriginTimeOffset, SampleIdx, SchemaBoneIdx));
-		FFeatureVectorHelper::EncodeVector(Indexer.GetPoseVector(SampleIdx), ChannelDataOffset, Heading, ComponentStripping);
+		if (Indexer.GetSampleRotation(SampleRotation, SampleTimeOffset, OriginTimeOffset, SampleIdx, SchemaBoneIdx, RootSchemaBoneIdx, EPermutationTimeType::UseSampleTime, SamplingAttributeId))
+		{
+			FFeatureVectorHelper::EncodeVector(Indexer.GetPoseVector(SampleIdx), ChannelDataOffset, GetAxis(SampleRotation), ComponentStripping);
+		}
+		else
+		{
+			return false;
+		}
 	}
+	return true;
 }
 
 FString UPoseSearchFeatureChannel_Heading::GetLabel() const
@@ -169,7 +193,7 @@ FString UPoseSearchFeatureChannel_Heading::GetLabel() const
 
 	const UPoseSearchSchema* Schema = GetSchema();
 	check(Schema);
-	if (!Schema->IsRootBone(SchemaBoneIdx))
+	if (SchemaBoneIdx != RootSchemaBoneIdx)
 	{
 		Label.Append(TEXT("_"));
 		Label.Append(Schema->BoneReferences[SchemaBoneIdx].BoneName.ToString());
diff --git a/Engine/Plugins/Experimental/Animation/PoseSearch/Source/Runtime/Private/PoseSearchFeatureChannel_Heading.h b/Engine/Plugins/Experimental/Animation/PoseSearch/Source/Runtime/Private/PoseSearchFeatureChannel_Heading.h
--- a/Engine/Plugins/Experimental/Animation/PoseSearch/Source/Runtime/Private/PoseSearchFeatureChannel_Heading.h
+++ b/Engine/Plugins/Experimental/Animation/PoseSearch/Source/Runtime/Private/PoseSearchFeatureChannel_Heading.h
@@ -17,7 +17,7 @@ enum class EHeadingAxis : uint8
 	Invalid = Num UMETA(Hidden)
 };
 
-UCLASS(EditInlineNew, meta = (DisplayName = "Heading Channel"), CollapseCategories)
+UCLASS(EditInlineNew, Blueprintable, meta = (DisplayName = "Heading Channel"), CollapseCategories)
 class POSESEARCH_API UPoseSearchFeatureChannel_Heading : public UPoseSearchFeatureChannel
 {
 	GENERATED_BODY()
@@ -29,6 +29,12 @@ public:
 	UPROPERTY(EditAnywhere, Category = "Settings")
 	float Weight = 1.f;
 
+	// if SamplingAttributeId >= 0, ALL the animations contained in the pose search database referencing the schema containing this channel are expected to have 
+	// UAnimNotifyState_PoseSearchSamplingAttribute notify state with a matching SamplingAttributeId, and the UAnimNotifyState_PoseSearchSamplingAttribute properties
+	// will be used as source of data instead of this channel "Bone".
+	UPROPERTY(EditAnywhere, Category = "Settings")
+	int32 SamplingAttributeId = -1;
+
 	// the data relative to the sampling time associated to this channel will be offsetted by SampleTimeOffset seconds.
 	// For example, if Bone is the head bone, and SampleTimeOffset is 0.5, this channel will try to match the future heading of the character head bone 0.5 seconds ahead
 	UPROPERTY(EditAnywhere, Category = "Settings")
@@ -56,6 +62,13 @@ public:
 	UPROPERTY(EditAnywhere, Category = "Settings")
 	EComponentStrippingVector ComponentStripping = EComponentStrippingVector::None;
 
+	UFUNCTION(BlueprintPure, BlueprintImplementableEvent, meta=(BlueprintThreadSafe, DisplayName = "Get World Rotation"), Category = "Settings")
+	FQuat BP_GetWorldRotation(const UAnimInstance* AnimInstance) const;
+
+	bool bUseBlueprintQueryOverride = false;
+
+	UPoseSearchFeatureChannel_Heading();
+
 	// UPoseSearchFeatureChannel interface
 	virtual void Finalize(UPoseSearchSchema* Schema) override;
 	virtual void BuildQuery(UE::PoseSearch::FSearchContext& SearchContext, UE::PoseSearch::FFeatureVectorBuilder& InOutQuery) const override;
@@ -68,7 +81,7 @@ public:
 
 #if WITH_EDITOR
 	virtual void FillWeights(TArrayView<float> Weights) const override;
-	virtual void IndexAsset(UE::PoseSearch::FAssetIndexer& Indexer) const override;
+	virtual bool IndexAsset(UE::PoseSearch::FAssetIndexer& Indexer) const override;
 	virtual FString GetLabel() const override;
 #endif
 
diff --git a/Engine/Plugins/Experimental/Animation/PoseSearch/Source/Runtime/Private/PoseSearchFeatureChannel_Padding.cpp b/Engine/Plugins/Experimental/Animation/PoseSearch/Source/Runtime/Private/PoseSearchFeatureChannel_Padding.cpp
--- a/Engine/Plugins/Experimental/Animation/PoseSearch/Source/Runtime/Private/PoseSearchFeatureChannel_Padding.cpp
+++ b/Engine/Plugins/Experimental/Animation/PoseSearch/Source/Runtime/Private/PoseSearchFeatureChannel_Padding.cpp
@@ -33,9 +33,10 @@ void UPoseSearchFeatureChannel_Padding::FillWeights(TArrayView<float> Weights) c
 	}
 }
 
-void UPoseSearchFeatureChannel_Padding::IndexAsset(UE::PoseSearch::FAssetIndexer& Indexer) const
+bool UPoseSearchFeatureChannel_Padding::IndexAsset(UE::PoseSearch::FAssetIndexer& Indexer) const
 {
 	UE_LOG(LogPoseSearch, Log, TEXT("UPoseSearchFeatureChannel_Padding::IndexAsset: padding data with '%d' additional floats per pose in UPoseSearchSchema '%s'"), PaddingSize, *GetNameSafe(GetSchema()));
+	return true;
 }
 
 FString UPoseSearchFeatureChannel_Padding::GetLabel() const
diff --git a/Engine/Plugins/Experimental/Animation/PoseSearch/Source/Runtime/Private/PoseSearchFeatureChannel_Padding.h b/Engine/Plugins/Experimental/Animation/PoseSearch/Source/Runtime/Private/PoseSearchFeatureChannel_Padding.h
--- a/Engine/Plugins/Experimental/Animation/PoseSearch/Source/Runtime/Private/PoseSearchFeatureChannel_Padding.h
+++ b/Engine/Plugins/Experimental/Animation/PoseSearch/Source/Runtime/Private/PoseSearchFeatureChannel_Padding.h
@@ -20,7 +20,7 @@ public:
 
 #if WITH_EDITOR
 	virtual void FillWeights(TArrayView<float> Weights) const override;
-	virtual void IndexAsset(UE::PoseSearch::FAssetIndexer& Indexer) const override;
+	virtual bool IndexAsset(UE::PoseSearch::FAssetIndexer& Indexer) const override;
 	virtual FString GetLabel() const override;
 #endif
 
diff --git a/Engine/Plugins/Experimental/Animation/PoseSearch/Source/Runtime/Private/PoseSearchFeatureChannel_PermutationTime.cpp b/Engine/Plugins/Experimental/Animation/PoseSearch/Source/Runtime/Private/PoseSearchFeatureChannel_PermutationTime.cpp
--- a/Engine/Plugins/Experimental/Animation/PoseSearch/Source/Runtime/Private/PoseSearchFeatureChannel_PermutationTime.cpp
+++ b/Engine/Plugins/Experimental/Animation/PoseSearch/Source/Runtime/Private/PoseSearchFeatureChannel_PermutationTime.cpp
@@ -37,7 +37,7 @@ void UPoseSearchFeatureChannel_PermutationTime::FillWeights(TArrayView<float> We
 	}
 }
 
-void UPoseSearchFeatureChannel_PermutationTime::IndexAsset(UE::PoseSearch::FAssetIndexer& Indexer) const
+bool UPoseSearchFeatureChannel_PermutationTime::IndexAsset(UE::PoseSearch::FAssetIndexer& Indexer) const
 {
 	using namespace UE::PoseSearch;
 
@@ -46,6 +46,7 @@ void UPoseSearchFeatureChannel_PermutationTime::IndexAsset(UE::PoseSearch::FAsse
 	{
 		FFeatureVectorHelper::EncodeFloat(Indexer.GetPoseVector(SampleIdx), ChannelDataOffset, PermutationTimeOffset);
 	}
+	return true;
 }
 
 FString UPoseSearchFeatureChannel_PermutationTime::GetLabel() const
diff --git a/Engine/Plugins/Experimental/Animation/PoseSearch/Source/Runtime/Private/PoseSearchFeatureChannel_PermutationTime.h b/Engine/Plugins/Experimental/Animation/PoseSearch/Source/Runtime/Private/PoseSearchFeatureChannel_PermutationTime.h
--- a/Engine/Plugins/Experimental/Animation/PoseSearch/Source/Runtime/Private/PoseSearchFeatureChannel_PermutationTime.h
+++ b/Engine/Plugins/Experimental/Animation/PoseSearch/Source/Runtime/Private/PoseSearchFeatureChannel_PermutationTime.h
@@ -20,7 +20,7 @@ public:
 
 #if WITH_EDITOR
 	virtual void FillWeights(TArrayView<float> Weights) const override;
-	virtual void IndexAsset(UE::PoseSearch::FAssetIndexer& Indexer) const override;
+	virtual bool IndexAsset(UE::PoseSearch::FAssetIndexer& Indexer) const override;
 	virtual FString GetLabel() const override;
 #endif
 
diff --git a/Engine/Plugins/Experimental/Animation/PoseSearch/Source/Runtime/Private/PoseSearchFeatureChannel_Phase.cpp b/Engine/Plugins/Experimental/Animation/PoseSearch/Source/Runtime/Private/PoseSearchFeatureChannel_Phase.cpp
--- a/Engine/Plugins/Experimental/Animation/PoseSearch/Source/Runtime/Private/PoseSearchFeatureChannel_Phase.cpp
+++ b/Engine/Plugins/Experimental/Animation/PoseSearch/Source/Runtime/Private/PoseSearchFeatureChannel_Phase.cpp
@@ -41,7 +41,7 @@ namespace UE::PoseSearch
 		return (Values[Num - 1] - Values[Num - 2]) * (Sample - (Num - 1)) + Values[Num - 1];
 	}
 
-	static void CollectBonePositions(TArray<FVector>& BonePositions, FAssetIndexer& Indexer, int8 SchemaBoneIdx)
+	static bool CollectBonePositions(TArray<FVector>& BonePositions, FAssetIndexer& Indexer, int8 SchemaBoneIdx)
 	{
 		const int32 NumSamples = Indexer.GetEndSampleIdx() - Indexer.GetBeginSampleIdx();
 
@@ -50,8 +50,12 @@ namespace UE::PoseSearch
 		BonePositions.AddDefaulted(NumSamples);
 		for (int32 SampleIdx = 0; SampleIdx != NumSamples; ++SampleIdx)
 		{
-			BonePositions[SampleIdx] = Indexer.GetSamplePosition(0.f, 0.f, SampleIdx, SchemaBoneIdx, RootSchemaBoneIdx, EPermutationTimeType::UseSampleTime);
+			if (!Indexer.GetSamplePosition(BonePositions[SampleIdx], 0.f, 0.f, SampleIdx, SchemaBoneIdx, RootSchemaBoneIdx, EPermutationTimeType::UseSampleTime))
+			{
+				return false;
+			}
 		}
+		return true;
 	}
 
 	static void CalculateSignal(const TArray<FVector>& BonePositions, TArray<float>& Signal, int32 offset = 1)
@@ -357,7 +361,7 @@ void UPoseSearchFeatureChannel_Phase::FillWeights(TArrayView<float> Weights) con
 	}
 }
 
-void UPoseSearchFeatureChannel_Phase::IndexAsset(UE::PoseSearch::FAssetIndexer& Indexer) const
+bool UPoseSearchFeatureChannel_Phase::IndexAsset(UE::PoseSearch::FAssetIndexer& Indexer) const
 {
 	using namespace UE::PoseSearch;
 
@@ -374,7 +378,10 @@ void UPoseSearchFeatureChannel_Phase::IndexAsset(UE::PoseSearch::FAssetIndexer&
 	TArray<LocalMinMax> LocalMinMax;
 	TArray<FVector> BonePositions;
 
-	CollectBonePositions(BonePositions, Indexer, SchemaBoneIdx);
+	if (!CollectBonePositions(BonePositions, Indexer, SchemaBoneIdx))
+	{
+		return false;
+	}
 
 	// @todo: have different way of calculating signals, for example: height of the bone transform, acceleration, etc?
 	const int32 BoneSamplingCentralDifferencesOffset = FMath::Max(FMath::CeilToInt(BoneSamplingCentralDifferencesTime * Schema->SampleRate), 1);
@@ -394,6 +401,7 @@ void UPoseSearchFeatureChannel_Phase::IndexAsset(UE::PoseSearch::FAssetIndexer&
 	{
 		FFeatureVectorHelper::EncodeVector2D(Indexer.GetPoseVector(SampleIdx), ChannelDataOffset, Phases[SampleIdx - Indexer.GetBeginSampleIdx()]);
 	}
+	return true;
 }
 
 FString UPoseSearchFeatureChannel_Phase::GetLabel() const
@@ -409,7 +417,7 @@ FString UPoseSearchFeatureChannel_Phase::GetLabel() const
 
 	const UPoseSearchSchema* Schema = GetSchema();
 	check(Schema);
-	if (!Schema->IsRootBone(SchemaBoneIdx))
+	if (SchemaBoneIdx != RootSchemaBoneIdx)
 	{
 		Label.Append(TEXT("_"));
 		Label.Append(Schema->BoneReferences[SchemaBoneIdx].BoneName.ToString());
diff --git a/Engine/Plugins/Experimental/Animation/PoseSearch/Source/Runtime/Private/PoseSearchFeatureChannel_Phase.h b/Engine/Plugins/Experimental/Animation/PoseSearch/Source/Runtime/Private/PoseSearchFeatureChannel_Phase.h
--- a/Engine/Plugins/Experimental/Animation/PoseSearch/Source/Runtime/Private/PoseSearchFeatureChannel_Phase.h
+++ b/Engine/Plugins/Experimental/Animation/PoseSearch/Source/Runtime/Private/PoseSearchFeatureChannel_Phase.h
@@ -40,7 +40,7 @@ public:
 
 #if WITH_EDITOR
 	virtual void FillWeights(TArrayView<float> Weights) const override;
-	virtual void IndexAsset(UE::PoseSearch::FAssetIndexer& Indexer) const override;
+	virtual bool IndexAsset(UE::PoseSearch::FAssetIndexer& Indexer) const override;
 	virtual FString GetLabel() const override;
 #endif
 };
\ No newline at end of file
diff --git a/Engine/Plugins/Experimental/Animation/PoseSearch/Source/Runtime/Private/PoseSearchFeatureChannel_Position.cpp b/Engine/Plugins/Experimental/Animation/PoseSearch/Source/Runtime/Private/PoseSearchFeatureChannel_Position.cpp
--- a/Engine/Plugins/Experimental/Animation/PoseSearch/Source/Runtime/Private/PoseSearchFeatureChannel_Position.cpp
+++ b/Engine/Plugins/Experimental/Animation/PoseSearch/Source/Runtime/Private/PoseSearchFeatureChannel_Position.cpp
@@ -6,6 +6,12 @@
 #include "PoseSearch/PoseSearchContext.h"
 #include "PoseSearch/PoseSearchDatabase.h"
 #include "PoseSearch/PoseSearchSchema.h"
+#include "Engine/BlueprintGeneratedClass.h"
+
+UPoseSearchFeatureChannel_Position::UPoseSearchFeatureChannel_Position()
+{
+	bUseBlueprintQueryOverride = Cast<UBlueprintGeneratedClass>(GetClass()) != nullptr;
+}
 
 void UPoseSearchFeatureChannel_Position::FindOrAddToSchema(UPoseSearchSchema* Schema, float SampleTimeOffset, const FName& BoneName, EPermutationTimeType PermutationTimeType)
 {
@@ -46,7 +52,7 @@ void UPoseSearchFeatureChannel_Position::AddDependentChannels(UPoseSearchSchema*
 {
 	if (Schema->bInjectAdditionalDebugChannels)
 	{
-		if (!Schema->IsRootBone(SchemaOriginBoneIdx))
+		if (SchemaOriginBoneIdx != RootSchemaBoneIdx)
 		{
 			const EPermutationTimeType DependentChannelsPermutationTimeType = PermutationTimeType == EPermutationTimeType::UsePermutationTime ? EPermutationTimeType::UseSampleToPermutationTime : EPermutationTimeType::UseSampleTime;
 			UPoseSearchFeatureChannel_Position::FindOrAddToSchema(Schema, 0.f, OriginBone.BoneName, DependentChannelsPermutationTimeType);
@@ -59,27 +65,36 @@ void UPoseSearchFeatureChannel_Position::BuildQuery(UE::PoseSearch::FSearchConte
 	using namespace UE::PoseSearch;
 
 	check(InOutQuery.GetSchema());
-	const bool bIsCurrentResultValid = SearchContext.GetCurrentResult().IsValid() && SearchContext.GetCurrentResult().Database->Schema == InOutQuery.GetSchema();
-	const bool bSkip = InputQueryPose != EInputQueryPose::UseCharacterPose && bIsCurrentResultValid;
-	const bool bIsRootBone = InOutQuery.GetSchema()->IsRootBone(SchemaBoneIdx);
-	if (bSkip || (!SearchContext.IsHistoryValid() && !bIsRootBone))
+	const bool bIsRootBone = SchemaBoneIdx == RootSchemaBoneIdx;
+	if (bUseBlueprintQueryOverride)
 	{
-		if (bIsCurrentResultValid)
+		const FVector BonePositionWorld = BP_GetWorldPosition(SearchContext.GetAnimInstance());
+		const FVector BonePosition = SearchContext.GetSamplePosition(SampleTimeOffset, 0.f, InOutQuery.GetSchema(), SchemaBoneIdx, SchemaOriginBoneIdx, /*!bIsRootBone*/ true, PermutationTimeType, &BonePositionWorld);
+  		FFeatureVectorHelper::EncodeVector(InOutQuery.EditValues(), ChannelDataOffset, BonePosition, ComponentStripping);
+	}
+	else
+	{
+		const bool bIsCurrentResultValid = SearchContext.GetCurrentResult().IsValid() && SearchContext.GetCurrentResult().Database->Schema == InOutQuery.GetSchema();
+		const bool bSkip = InputQueryPose != EInputQueryPose::UseCharacterPose && bIsCurrentResultValid;
+		if (bSkip || (!SearchContext.IsHistoryValid() && !bIsRootBone))
 		{
-			FFeatureVectorHelper::Copy(InOutQuery.EditValues(), ChannelDataOffset, ChannelCardinality, SearchContext.GetCurrentResultPoseVector());
+			if (bIsCurrentResultValid)
+			{
+				FFeatureVectorHelper::Copy(InOutQuery.EditValues(), ChannelDataOffset, ChannelCardinality, SearchContext.GetCurrentResultPoseVector());
+			}
+			else
+			{
+				// we leave the InOutQuery set to zero since the SearchContext.History is invalid and it'll fail if we continue
+				UE_LOG(LogPoseSearch, Error, TEXT("UPoseSearchFeatureChannel_Position::BuildQuery - Failed because Pose History Node is missing."));
+			}
 		}
 		else
 		{
-			// we leave the InOutQuery set to zero since the SearchContext.History is invalid and it'll fail if we continue
-			UE_LOG(LogPoseSearch, Error, TEXT("UPoseSearchFeatureChannel_Position::BuildQuery - Failed because Pose History Node is missing."));
+			// calculating the BonePosition in root bone space for the bone indexed by SchemaBoneIdx
+			const FVector BonePosition = SearchContext.GetSamplePosition(SampleTimeOffset, 0.f, InOutQuery.GetSchema(), SchemaBoneIdx, SchemaOriginBoneIdx, !bIsRootBone, PermutationTimeType);
+			FFeatureVectorHelper::EncodeVector(InOutQuery.EditValues(), ChannelDataOffset, BonePosition, ComponentStripping);
 		}
 	}
-	else
-	{
-		// calculating the BonePosition in component space for the bone indexed by SchemaBoneIdx
-		const FVector BonePosition = SearchContext.GetSamplePosition(SampleTimeOffset, 0.f, InOutQuery.GetSchema(), SchemaBoneIdx, SchemaOriginBoneIdx, !bIsRootBone, PermutationTimeType);
-		FFeatureVectorHelper::EncodeVector(InOutQuery.EditValues(), ChannelDataOffset, BonePosition, ComponentStripping);
-	}
 }
 
 #if ENABLE_DRAW_DEBUG
@@ -90,7 +105,7 @@ void UPoseSearchFeatureChannel_Position::DebugDraw(const UE::PoseSearch::FDebugD
 	const FColor Color = DebugColor.ToFColor(true);
 
 	const FVector FeaturesVector = FFeatureVectorHelper::DecodeVector(PoseVector, ChannelDataOffset, ComponentStripping);
-	if (DrawParams.GetSchema()->IsRootBone(SchemaOriginBoneIdx))
+	if (SchemaOriginBoneIdx == RootSchemaBoneIdx)
 	{
 		const FVector BonePos = DrawParams.GetRootTransform().TransformPosition(FeaturesVector);
 		DrawParams.DrawPoint(BonePos, Color);
@@ -98,7 +113,7 @@ void UPoseSearchFeatureChannel_Position::DebugDraw(const UE::PoseSearch::FDebugD
 	else
 	{
 		const EPermutationTimeType TimeType = PermutationTimeType == EPermutationTimeType::UsePermutationTime ? EPermutationTimeType::UseSampleToPermutationTime : EPermutationTimeType::UseSampleTime;
-		const FVector OriginBonePos = DrawParams.ExtractPosition(PoseVector, SampleTimeOffset, SchemaOriginBoneIdx, TimeType);
+		const FVector OriginBonePos = DrawParams.ExtractPosition(PoseVector, SampleTimeOffset, SchemaOriginBoneIdx, TimeType, SamplingAttributeId);
 		const FVector DeltaPos = DrawParams.GetRootTransform().TransformVector(FeaturesVector);
 		const FVector BonePos = OriginBonePos + DeltaPos;
 		DrawParams.DrawLine(OriginBonePos, BonePos, Color);
@@ -117,15 +132,23 @@ void UPoseSearchFeatureChannel_Position::FillWeights(TArrayView<float> Weights)
 	}
 }
 
-void UPoseSearchFeatureChannel_Position::IndexAsset(UE::PoseSearch::FAssetIndexer& Indexer) const
+bool UPoseSearchFeatureChannel_Position::IndexAsset(UE::PoseSearch::FAssetIndexer& Indexer) const
 {
 	using namespace UE::PoseSearch;
 
+	FVector BonePosition;
 	for (int32 SampleIdx = Indexer.GetBeginSampleIdx(); SampleIdx != Indexer.GetEndSampleIdx(); ++SampleIdx)
 	{
-		const FVector BonePosition = Indexer.GetSamplePosition(SampleTimeOffset, 0.f, SampleIdx, SchemaBoneIdx, SchemaOriginBoneIdx, PermutationTimeType);
-		FFeatureVectorHelper::EncodeVector(Indexer.GetPoseVector(SampleIdx), ChannelDataOffset, BonePosition, ComponentStripping);
+		if (Indexer.GetSamplePosition(BonePosition, SampleTimeOffset, 0.f, SampleIdx, SchemaBoneIdx, SchemaOriginBoneIdx, PermutationTimeType, SamplingAttributeId))
+		{
+			FFeatureVectorHelper::EncodeVector(Indexer.GetPoseVector(SampleIdx), ChannelDataOffset, BonePosition, ComponentStripping);
+		}
+		else
+		{
+			return false;
+		}
 	}
+	return true;
 }
 
 FString UPoseSearchFeatureChannel_Position::GetLabel() const
@@ -150,13 +173,13 @@ FString UPoseSearchFeatureChannel_Position::GetLabel() const
 
 	const UPoseSearchSchema* Schema = GetSchema();
 	check(Schema);
-	if (!Schema->IsRootBone(SchemaBoneIdx))
+	if (SchemaBoneIdx != RootSchemaBoneIdx)
 	{
 		Label.Append(TEXT("_"));
 		Label.Append(Schema->BoneReferences[SchemaBoneIdx].BoneName.ToString());
 	}
 
-	if (!Schema->IsRootBone(SchemaOriginBoneIdx))
+	if (SchemaOriginBoneIdx != RootSchemaBoneIdx)
 	{
 		Label.Append(TEXT("_"));
 		Label.Append(Schema->BoneReferences[SchemaOriginBoneIdx].BoneName.ToString());
diff --git a/Engine/Plugins/Experimental/Animation/PoseSearch/Source/Runtime/Private/PoseSearchFeatureChannel_Position.h b/Engine/Plugins/Experimental/Animation/PoseSearch/Source/Runtime/Private/PoseSearchFeatureChannel_Position.h
--- a/Engine/Plugins/Experimental/Animation/PoseSearch/Source/Runtime/Private/PoseSearchFeatureChannel_Position.h
+++ b/Engine/Plugins/Experimental/Animation/PoseSearch/Source/Runtime/Private/PoseSearchFeatureChannel_Position.h
@@ -6,7 +6,7 @@
 #include "BoneContainer.h"
 #include "PoseSearchFeatureChannel_Position.generated.h"
 
-UCLASS(EditInlineNew, meta = (DisplayName = "Position Channel"), CollapseCategories)
+UCLASS(EditInlineNew, Blueprintable, meta = (DisplayName = "Position Channel"), CollapseCategories)
 class POSESEARCH_API UPoseSearchFeatureChannel_Position : public UPoseSearchFeatureChannel
 {
 	GENERATED_BODY()
@@ -21,6 +21,13 @@ public:
 	UPROPERTY(EditAnywhere, Category = "Settings")
 	float Weight = 1.f;
 
+
+	// if SamplingAttributeId >= 0, ALL the animations contained in the pose search database referencing the schema containing this channel are expected to have 
+	// UAnimNotifyState_PoseSearchSamplingAttribute notify state with a matching SamplingAttributeId, and the UAnimNotifyState_PoseSearchSamplingAttribute properties
+	// will be used as source of data instead of this channel "Bone". UAnimNotifyState_PoseSearchSamplingAttribute properties will be then converted into OriginBone space
+	UPROPERTY(EditAnywhere, Category = "Settings")
+	int32 SamplingAttributeId = -1;
+	
 	// the data relative to the sampling time associated to this channel will be offsetted by SampleTimeOffset seconds.
 	// For example, if Bone is the head bone, and SampleTimeOffset is 0.5, this channel will try to match the future position of the character head bone 0.5 seconds ahead
 	UPROPERTY(EditAnywhere, Category = "Settings")
@@ -45,6 +52,13 @@ public:
 	UPROPERTY(EditAnywhere, Category = "Settings")
 	EPermutationTimeType PermutationTimeType = EPermutationTimeType::UseSampleTime;
 
+	UFUNCTION(BlueprintPure, BlueprintImplementableEvent, meta=(BlueprintThreadSafe, DisplayName = "Get World Position"), Category = "Settings")
+	FVector BP_GetWorldPosition(const UAnimInstance* AnimInstance) const;
+
+	bool bUseBlueprintQueryOverride = false;
+
+	UPoseSearchFeatureChannel_Position();
+
 	// UPoseSearchFeatureChannel interface
 	virtual void Finalize(UPoseSearchSchema* Schema) override;
 	virtual void BuildQuery(UE::PoseSearch::FSearchContext& SearchContext, UE::PoseSearch::FFeatureVectorBuilder& InOutQuery) const override;
@@ -58,7 +72,7 @@ public:
 
 #if WITH_EDITOR
 	virtual void FillWeights(TArrayView<float> Weights) const override;
-	virtual void IndexAsset(UE::PoseSearch::FAssetIndexer& Indexer) const override;
+	virtual bool IndexAsset(UE::PoseSearch::FAssetIndexer& Indexer) const override;
 	virtual FString GetLabel() const override;
 #endif
 
diff --git a/Engine/Plugins/Experimental/Animation/PoseSearch/Source/Runtime/Private/PoseSearchFeatureChannel_SamplingTime.cpp b/Engine/Plugins/Experimental/Animation/PoseSearch/Source/Runtime/Private/PoseSearchFeatureChannel_SamplingTime.cpp
--- a/Engine/Plugins/Experimental/Animation/PoseSearch/Source/Runtime/Private/PoseSearchFeatureChannel_SamplingTime.cpp
+++ b/Engine/Plugins/Experimental/Animation/PoseSearch/Source/Runtime/Private/PoseSearchFeatureChannel_SamplingTime.cpp
@@ -27,7 +27,7 @@ void UPoseSearchFeatureChannel_SamplingTime::FillWeights(TArrayView<float> Weigh
 	}
 }
 
-void UPoseSearchFeatureChannel_SamplingTime::IndexAsset(UE::PoseSearch::FAssetIndexer& Indexer) const
+bool UPoseSearchFeatureChannel_SamplingTime::IndexAsset(UE::PoseSearch::FAssetIndexer& Indexer) const
 {
 	using namespace UE::PoseSearch;
 
@@ -35,6 +35,7 @@ void UPoseSearchFeatureChannel_SamplingTime::IndexAsset(UE::PoseSearch::FAssetIn
 	{
 		FFeatureVectorHelper::EncodeFloat(Indexer.GetPoseVector(SampleIdx), ChannelDataOffset, Indexer.CalculateSampleTime(SampleIdx));
 	}
+	return true;
 }
 
 FString UPoseSearchFeatureChannel_SamplingTime::GetLabel() const
diff --git a/Engine/Plugins/Experimental/Animation/PoseSearch/Source/Runtime/Private/PoseSearchFeatureChannel_SamplingTime.h b/Engine/Plugins/Experimental/Animation/PoseSearch/Source/Runtime/Private/PoseSearchFeatureChannel_SamplingTime.h
--- a/Engine/Plugins/Experimental/Animation/PoseSearch/Source/Runtime/Private/PoseSearchFeatureChannel_SamplingTime.h
+++ b/Engine/Plugins/Experimental/Animation/PoseSearch/Source/Runtime/Private/PoseSearchFeatureChannel_SamplingTime.h
@@ -24,7 +24,7 @@ public:
 
 #if WITH_EDITOR
 	virtual void FillWeights(TArrayView<float> Weights) const override;
-	virtual void IndexAsset(UE::PoseSearch::FAssetIndexer& Indexer) const override;
+	virtual bool IndexAsset(UE::PoseSearch::FAssetIndexer& Indexer) const override;
 	virtual FString GetLabel() const override;
 #endif
 };
diff --git a/Engine/Plugins/Experimental/Animation/PoseSearch/Source/Runtime/Private/PoseSearchFeatureChannel_Velocity.cpp b/Engine/Plugins/Experimental/Animation/PoseSearch/Source/Runtime/Private/PoseSearchFeatureChannel_Velocity.cpp
--- a/Engine/Plugins/Experimental/Animation/PoseSearch/Source/Runtime/Private/PoseSearchFeatureChannel_Velocity.cpp
+++ b/Engine/Plugins/Experimental/Animation/PoseSearch/Source/Runtime/Private/PoseSearchFeatureChannel_Velocity.cpp
@@ -8,6 +8,12 @@
 #include "PoseSearch/PoseSearchHistory.h"
 #include "PoseSearch/PoseSearchSchema.h"
 #include "PoseSearchFeatureChannel_Position.h"
+#include "Engine/BlueprintGeneratedClass.h"
+
+UPoseSearchFeatureChannel_Velocity::UPoseSearchFeatureChannel_Velocity()
+{
+	bUseBlueprintQueryOverride = Cast<UBlueprintGeneratedClass>(GetClass()) != nullptr;
+}
 
 void UPoseSearchFeatureChannel_Velocity::Finalize(UPoseSearchSchema* Schema)
 {
@@ -30,31 +36,46 @@ void UPoseSearchFeatureChannel_Velocity::BuildQuery(UE::PoseSearch::FSearchConte
 {
 	using namespace UE::PoseSearch;
 
-	const bool bIsCurrentResultValid = SearchContext.GetCurrentResult().IsValid() && SearchContext.GetCurrentResult().Database->Schema == InOutQuery.GetSchema();
-	const bool bSkip = InputQueryPose != EInputQueryPose::UseCharacterPose && bIsCurrentResultValid;
-	const bool bIsRootBone = InOutQuery.GetSchema()->IsRootBone(SchemaBoneIdx);
-	if (bSkip || (!SearchContext.IsHistoryValid() && !bIsRootBone))
+	check(InOutQuery.GetSchema());
+	const bool bIsRootBone = SchemaBoneIdx == RootSchemaBoneIdx;
+	if (bUseBlueprintQueryOverride)
 	{
-		if (bIsCurrentResultValid)
-		{
-			FFeatureVectorHelper::Copy(InOutQuery.EditValues(), ChannelDataOffset, ChannelCardinality, SearchContext.GetCurrentResultPoseVector());
-		}
-		else
+		const FVector LinearVelocityWorld = BP_GetWorldVelocity(SearchContext.GetAnimInstance());
+
+		FVector LinearVelocity = SearchContext.GetSampleVelocity(SampleTimeOffset, 0.f, InOutQuery.GetSchema(), SchemaBoneIdx, RootSchemaBoneIdx, bUseCharacterSpaceVelocities, /*!bIsRootBone*/ true, EPermutationTimeType::UseSampleTime, &LinearVelocityWorld);
+		if (bNormalize)
 		{
-			// we leave the InOutQuery set to zero since the SearchContext.History is invalid and it'll fail if we continue
-			UE_LOG(LogPoseSearch, Error, TEXT("UPoseSearchFeatureChannel_Velocity::BuildQuery - Failed because Pose History Node is missing."));
+			LinearVelocity = LinearVelocity.GetClampedToMaxSize(1.f);
 		}
+		FFeatureVectorHelper::EncodeVector(InOutQuery.EditValues(), ChannelDataOffset, LinearVelocity, ComponentStripping);
 	}
 	else
 	{
-		// calculating the LinearVelocity for the bone indexed by SchemaBoneIdx
-		FVector LinearVelocity = SearchContext.GetSampleVelocity(SampleTimeOffset, 0.f, InOutQuery.GetSchema(), SchemaBoneIdx, RootSchemaBoneIdx, bUseCharacterSpaceVelocities, !bIsRootBone, EPermutationTimeType::UseSampleTime);
-		if (bNormalize)
+		const bool bIsCurrentResultValid = SearchContext.GetCurrentResult().IsValid() && SearchContext.GetCurrentResult().Database->Schema == InOutQuery.GetSchema();
+		const bool bSkip = InputQueryPose != EInputQueryPose::UseCharacterPose && bIsCurrentResultValid;
+		if (bSkip || (!SearchContext.IsHistoryValid() && !bIsRootBone))
 		{
-			LinearVelocity = LinearVelocity.GetClampedToMaxSize(1.f);
+			if (bIsCurrentResultValid)
+			{
+				FFeatureVectorHelper::Copy(InOutQuery.EditValues(), ChannelDataOffset, ChannelCardinality, SearchContext.GetCurrentResultPoseVector());
+			}
+			else
+			{
+				// we leave the InOutQuery set to zero since the SearchContext.History is invalid and it'll fail if we continue
+				UE_LOG(LogPoseSearch, Error, TEXT("UPoseSearchFeatureChannel_Velocity::BuildQuery - Failed because Pose History Node is missing."));
+			}
+		}
+		else
+		{
+			// calculating the LinearVelocity for the bone indexed by SchemaBoneIdx
+			FVector LinearVelocity = SearchContext.GetSampleVelocity(SampleTimeOffset, 0.f, InOutQuery.GetSchema(), SchemaBoneIdx, RootSchemaBoneIdx, bUseCharacterSpaceVelocities, !bIsRootBone, EPermutationTimeType::UseSampleTime);
+			if (bNormalize)
+			{
+				LinearVelocity = LinearVelocity.GetClampedToMaxSize(1.f);
+			}
+
+			FFeatureVectorHelper::EncodeVector(InOutQuery.EditValues(), ChannelDataOffset, LinearVelocity, ComponentStripping);
 		}
-
-		FFeatureVectorHelper::EncodeVector(InOutQuery.EditValues(), ChannelDataOffset, LinearVelocity, ComponentStripping);
 	}
 }
 
@@ -68,7 +89,7 @@ void UPoseSearchFeatureChannel_Velocity::DebugDraw(const UE::PoseSearch::FDebugD
 
 	const FVector LinearVelocity = DrawParams.GetRootTransform().TransformVector(FFeatureVectorHelper::DecodeVector(PoseVector, ChannelDataOffset, ComponentStripping));
 	const FVector BoneVelDirection = LinearVelocity.GetSafeNormal();
-	const FVector BonePos = DrawParams.ExtractPosition(PoseVector, SampleTimeOffset, SchemaBoneIdx);
+	const FVector BonePos = DrawParams.ExtractPosition(PoseVector, SampleTimeOffset, SchemaBoneIdx, EPermutationTimeType::UseSampleTime, SamplingAttributeId);
 
 	DrawParams.DrawLine(BonePos, BonePos + LinearVelocity * LinearVelocityScale, Color);
 }
@@ -83,19 +104,27 @@ void UPoseSearchFeatureChannel_Velocity::FillWeights(TArrayView<float> Weights)
 	}
 }
 
-void UPoseSearchFeatureChannel_Velocity::IndexAsset(UE::PoseSearch::FAssetIndexer& Indexer) const
+bool UPoseSearchFeatureChannel_Velocity::IndexAsset(UE::PoseSearch::FAssetIndexer& Indexer) const
 {
 	using namespace UE::PoseSearch;
 
+	FVector LinearVelocity;
 	for (int32 SampleIdx = Indexer.GetBeginSampleIdx(); SampleIdx != Indexer.GetEndSampleIdx(); ++SampleIdx)
 	{
-		FVector LinearVelocity = Indexer.GetSampleVelocity(SampleTimeOffset, 0.f, SampleIdx, SchemaBoneIdx, RootSchemaBoneIdx, bUseCharacterSpaceVelocities, EPermutationTimeType::UseSampleTime);
-		if (bNormalize)
+		if (Indexer.GetSampleVelocity(LinearVelocity, SampleTimeOffset, 0.f, SampleIdx, SchemaBoneIdx, RootSchemaBoneIdx, bUseCharacterSpaceVelocities, EPermutationTimeType::UseSampleTime, SamplingAttributeId))
 		{
-			LinearVelocity = LinearVelocity.GetClampedToMaxSize(1.f);
+			if (bNormalize)
+			{
+				LinearVelocity = LinearVelocity.GetClampedToMaxSize(1.f);
+			}
+			FFeatureVectorHelper::EncodeVector(Indexer.GetPoseVector(SampleIdx), ChannelDataOffset, LinearVelocity, ComponentStripping);
+		}
+		else
+		{
+			return false;
 		}
-		FFeatureVectorHelper::EncodeVector(Indexer.GetPoseVector(SampleIdx), ChannelDataOffset, LinearVelocity, ComponentStripping);
 	}
+	return true;
 }
 
 FString UPoseSearchFeatureChannel_Velocity::GetLabel() const
@@ -124,7 +153,7 @@ FString UPoseSearchFeatureChannel_Velocity::GetLabel() const
 
 	const UPoseSearchSchema* Schema = GetSchema();
 	check(Schema);
-	if (!Schema->IsRootBone(SchemaBoneIdx))
+	if (SchemaBoneIdx != RootSchemaBoneIdx)
 	{
 		Label.Append(TEXT("_"));
 		Label.Append(Schema->BoneReferences[SchemaBoneIdx].BoneName.ToString());
diff --git a/Engine/Plugins/Experimental/Animation/PoseSearch/Source/Runtime/Private/PoseSearchFeatureChannel_Velocity.h b/Engine/Plugins/Experimental/Animation/PoseSearch/Source/Runtime/Private/PoseSearchFeatureChannel_Velocity.h
--- a/Engine/Plugins/Experimental/Animation/PoseSearch/Source/Runtime/Private/PoseSearchFeatureChannel_Velocity.h
+++ b/Engine/Plugins/Experimental/Animation/PoseSearch/Source/Runtime/Private/PoseSearchFeatureChannel_Velocity.h
@@ -6,7 +6,7 @@
 #include "BoneContainer.h"
 #include "PoseSearchFeatureChannel_Velocity.generated.h"
 
-UCLASS(EditInlineNew, meta = (DisplayName = "Velocity Channel"), CollapseCategories)
+UCLASS(EditInlineNew, Blueprintable, meta = (DisplayName = "Velocity Channel"), CollapseCategories)
 class POSESEARCH_API UPoseSearchFeatureChannel_Velocity : public UPoseSearchFeatureChannel
 {
 	GENERATED_BODY()
@@ -18,6 +18,12 @@ public:
 	UPROPERTY(EditAnywhere, Category = "Settings")
 	float Weight = 1.f;
 
+	// if SamplingAttributeId >= 0, ALL the animations contained in the pose search database referencing the schema containing this channel are expected to have 
+	// UAnimNotifyState_PoseSearchSamplingAttribute notify state with a matching SamplingAttributeId, and the UAnimNotifyState_PoseSearchSamplingAttribute properties
+	// will be used as source of data instead of this channel "Bone".
+	UPROPERTY(EditAnywhere, Category = "Settings")
+	int32 SamplingAttributeId = -1;
+
 	// the data relative to the sampling time associated to this channel will be offsetted by SampleTimeOffset seconds.
 	// For example, if Bone is the head bone, and SampleTimeOffset is 0.5, this channel will try to match the future velocity of the character head bone 0.5 seconds ahead
 	UPROPERTY(EditAnywhere, Category = "Settings")
@@ -44,6 +50,13 @@ public:
 	UPROPERTY(EditAnywhere, Category = "Settings")
 	EComponentStrippingVector ComponentStripping = EComponentStrippingVector::None;
 
+	UFUNCTION(BlueprintPure, BlueprintImplementableEvent, meta=(BlueprintThreadSafe, DisplayName = "Get World Velocity"), Category = "Settings")
+	FVector BP_GetWorldVelocity(const UAnimInstance* AnimInstance) const;
+
+	bool bUseBlueprintQueryOverride = false;
+
+	UPoseSearchFeatureChannel_Velocity();
+
 	// UPoseSearchFeatureChannel interface
 	virtual void Finalize(UPoseSearchSchema* Schema) override;
 	virtual void BuildQuery(UE::PoseSearch::FSearchContext& SearchContext, UE::PoseSearch::FFeatureVectorBuilder& InOutQuery) const override;
@@ -56,7 +69,7 @@ public:
 
 #if WITH_EDITOR
 	virtual void FillWeights(TArrayView<float> Weights) const override;
-	virtual void IndexAsset(UE::PoseSearch::FAssetIndexer& Indexer) const override;
+	virtual bool IndexAsset(UE::PoseSearch::FAssetIndexer& Indexer) const override;
 	virtual FString GetLabel() const override;
 #endif
 };
diff --git a/Engine/Plugins/Experimental/Animation/PoseSearch/Source/Runtime/Private/PoseSearchLibrary.cpp b/Engine/Plugins/Experimental/Animation/PoseSearch/Source/Runtime/Private/PoseSearchLibrary.cpp
--- a/Engine/Plugins/Experimental/Animation/PoseSearch/Source/Runtime/Private/PoseSearchLibrary.cpp
+++ b/Engine/Plugins/Experimental/Animation/PoseSearch/Source/Runtime/Private/PoseSearchLibrary.cpp
@@ -283,7 +283,9 @@ void UPoseSearchLibrary::UpdateMotionMatchingState(
 	const FPoseSearchQueryTrajectory TrajectoryRootSpace = ProcessTrajectory(Trajectory, Context.AnimInstanceProxy->GetComponentTransform(), InOutMotionMatchingState.ComponentDeltaYaw, YawFromAnimationTrajectoryBlendTime, TrajectorySpeedMultiplier);
 	
 	FMemMark Mark(FMemStack::Get());
-	FSearchContext SearchContext(&TrajectoryRootSpace, History, 0.f, &InOutMotionMatchingState.PoseIndicesHistory, InOutMotionMatchingState.CurrentSearchResult, PoseJumpThresholdTime, bForceInterrupt);
+	const UAnimInstance* AnimInstance = Cast<const UAnimInstance>(Context.AnimInstanceProxy->GetAnimInstanceObject());
+	check(AnimInstance);
+	FSearchContext SearchContext(AnimInstance, &TrajectoryRootSpace, History, 0.f, &InOutMotionMatchingState.PoseIndicesHistory, InOutMotionMatchingState.CurrentSearchResult, PoseJumpThresholdTime, bForceInterrupt);
 
 	const bool bCanAdvance = InOutMotionMatchingState.CurrentSearchResult.CanAdvance(DeltaTime);
 
@@ -349,8 +351,6 @@ void UPoseSearchLibrary::UpdateMotionMatchingState(
 	// Record debugger details
 	if (IsTracing(Context))
 	{
-		const UAnimInstance* AnimInstance = Cast<const UAnimInstance>(Context.AnimInstanceProxy->GetAnimInstanceObject());
-
 		const float SearchBestCost = InOutMotionMatchingState.CurrentSearchResult.PoseCost.GetTotalCost();
 		const float SearchBruteForceCost = InOutMotionMatchingState.CurrentSearchResult.BruteForcePoseCost.GetTotalCost();
 		TraceMotionMatchingState(Trajectory, SearchContext, InOutMotionMatchingState.CurrentSearchResult, InOutMotionMatchingState.ElapsedPoseSearchTime,
@@ -540,7 +540,7 @@ void UPoseSearchLibrary::MotionMatch(
 		}
 
 		FMemMark Mark(FMemStack::Get());
-		FSearchContext SearchContext(&TrajectoryRootSpace, ExtendedPoseHistory.IsInitialized() ? &ExtendedPoseHistory : nullptr, TimeToFutureAnimationStart);
+		FSearchContext SearchContext(AnimInstance, &TrajectoryRootSpace, ExtendedPoseHistory.IsInitialized() ? &ExtendedPoseHistory : nullptr, TimeToFutureAnimationStart);
 
 		FSearchResult SearchResult = Database->Search(SearchContext);
 		if (SearchResult.IsValid())
diff --git a/Engine/Plugins/Experimental/Animation/PoseSearch/Source/Runtime/Private/PoseSearchSchema.cpp b/Engine/Plugins/Experimental/Animation/PoseSearch/Source/Runtime/Private/PoseSearchSchema.cpp
--- a/Engine/Plugins/Experimental/Animation/PoseSearch/Source/Runtime/Private/PoseSearchSchema.cpp
+++ b/Engine/Plugins/Experimental/Animation/PoseSearch/Source/Runtime/Private/PoseSearchSchema.cpp
@@ -39,25 +39,10 @@ void UPoseSearchSchema::BuildQuery(UE::PoseSearch::FSearchContext& SearchContext
 
 FBoneIndexType UPoseSearchSchema::GetBoneIndexType(int8 SchemaBoneIdx) const
 {
-	if (SchemaBoneIdx < 0)
-	{
-		return RootBoneIndexType;
-	}
 	check(BoneReferences[SchemaBoneIdx].HasValidSetup());
 	return BoneReferences[SchemaBoneIdx].BoneIndex;
 }
 
-bool UPoseSearchSchema::IsRootBone(int8 SchemaBoneIdx) const
-{
-	if (SchemaBoneIdx < 0)
-	{
-		return true;
-	}
-
-	check(BoneReferences[SchemaBoneIdx].HasValidSetup());
-	return BoneReferences[SchemaBoneIdx].BoneIndex == int32(RootBoneIndexType);
-}
-
 int8 UPoseSearchSchema::AddBoneReference(const FBoneReference& BoneReference)
 {
 	int32 SchemaBoneIdx = 0;
diff --git a/Engine/Plugins/Experimental/Animation/PoseSearch/Source/Runtime/Public/PoseSearch/PoseSearchAnimNotifies.h b/Engine/Plugins/Experimental/Animation/PoseSearch/Source/Runtime/Public/PoseSearch/PoseSearchAnimNotifies.h
--- a/Engine/Plugins/Experimental/Animation/PoseSearch/Source/Runtime/Public/PoseSearch/PoseSearchAnimNotifies.h
+++ b/Engine/Plugins/Experimental/Animation/PoseSearch/Source/Runtime/Public/PoseSearch/PoseSearchAnimNotifies.h
@@ -3,7 +3,7 @@
 #pragma once
 
 #include "Animation/AnimNotifies/AnimNotifyState.h"
-
+#include "BoneContainer.h"
 #include "PoseSearchAnimNotifies.generated.h"
 
 // Base class for pose search anim notify states
@@ -58,4 +58,33 @@ public:
 	// makes the segment less likely to continuing playing
 	UPROPERTY(EditAnywhere, Category = Config, meta = (DisplayName = "Modifier"))
 	float CostAddend = -1.0f;
+};
+
+// UPoseSearchFeatureChannel(s) can use this UAnimNotifyState_PoseSearchSamplingAttribute as animation space position, rotation, and linear velocity provider 
+// during database indexing by specifying their SamplingAttributeId property to match UAnimNotifyState_PoseSearchSamplingAttribute::SamplingAttributeId
+UCLASS(Blueprintable, meta = (DisplayName = "Pose Matching: Sampling Attribute"))
+class POSESEARCH_API UAnimNotifyState_PoseSearchSamplingAttribute : public UAnimNotifyState_PoseSearchBase
+{
+	GENERATED_BODY()
+
+public:
+
+#if WITH_EDITORONLY_DATA
+	
+	UPROPERTY(EditAnywhere, Category = Config, meta = (ClampMin=0))
+	int32 SamplingAttributeId = 0;
+
+	UPROPERTY(EditAnywhere, Category = Config)
+	FBoneReference Bone;
+
+	UPROPERTY(EditAnywhere, Category = Config)
+	FVector Position = FVector::ZeroVector;
+
+	UPROPERTY(EditAnywhere, Category = Config)
+	FQuat Rotation = FQuat::Identity;
+
+	UPROPERTY(EditAnywhere, Category = Config)
+	FVector LinearVelocity = FVector::ZeroVector;
+
+#endif // WITH_EDITORONLY_DATA
 };
\ No newline at end of file
diff --git a/Engine/Plugins/Experimental/Animation/PoseSearch/Source/Runtime/Public/PoseSearch/PoseSearchAssetIndexer.h b/Engine/Plugins/Experimental/Animation/PoseSearch/Source/Runtime/Public/PoseSearch/PoseSearchAssetIndexer.h
--- a/Engine/Plugins/Experimental/Animation/PoseSearch/Source/Runtime/Public/PoseSearch/PoseSearchAssetIndexer.h
+++ b/Engine/Plugins/Experimental/Animation/PoseSearch/Source/Runtime/Public/PoseSearch/PoseSearchAssetIndexer.h
@@ -36,21 +36,21 @@ public:
 	void Process(int32 AssetIdx);
 	const FStats& GetStats() const { return Stats; }
 
-	// Returns the rotation of the bone Schema.BoneReferences[SchemaSampleBoneIdx] at time CalculateSampleTime(SampleIdx) + SampleTimeOffset relative to the
+	// Returns OutSampleRotation as the rotation of the bone Schema.BoneReferences[SchemaSampleBoneIdx] at time CalculateSampleTime(SampleIdx) + SampleTimeOffset relative to the
 	// transform of the bone Schema.BoneReferences[SchemaOriginBoneIdx] at time CalculateSampleTime(SampleIdx) + OriginTimeOffset 
 	// Times will be processed by GetPermutationTimeOffsets(PermutationTimeType, ...)
-	FQuat GetSampleRotation(float SampleTimeOffset, float OriginTimeOffset, int32 SampleIdx, int8 SchemaSampleBoneIdx = RootSchemaBoneIdx, int8 SchemaOriginBoneIdx = RootSchemaBoneIdx, EPermutationTimeType PermutationTimeType = EPermutationTimeType::UseSampleTime);
+	bool GetSampleRotation(FQuat& OutSampleRotation, float SampleTimeOffset, float OriginTimeOffset, int32 SampleIdx, int8 SchemaSampleBoneIdx = RootSchemaBoneIdx, int8 SchemaOriginBoneIdx = RootSchemaBoneIdx, EPermutationTimeType PermutationTimeType = EPermutationTimeType::UseSampleTime, int32 SamplingAttributeId = -1);
 
-	// Returns the position of the bone Schema.BoneReferences[SchemaSampleBoneIdx] at time CalculateSampleTime(SampleIdx) + SampleTimeOffset relative to the
+	// Returns OutSamplePosition as the position of the bone Schema.BoneReferences[SchemaSampleBoneIdx] at time CalculateSampleTime(SampleIdx) + SampleTimeOffset relative to the
 	// transform of the bone Schema.BoneReferences[SchemaOriginBoneIdx] at time CalculateSampleTime(SampleIdx) + OriginTimeOffset.
 	// Times will be processed by GetPermutationTimeOffsets(PermutationTimeType, ...)
-	FVector GetSamplePosition(float SampleTimeOffset, float OriginTimeOffset, int32 SampleIdx, int8 SchemaSampleBoneIdx = RootSchemaBoneIdx, int8 SchemaOriginBoneIdx = RootSchemaBoneIdx, EPermutationTimeType PermutationTimeType = EPermutationTimeType::UseSampleTime);
+	bool GetSamplePosition(FVector& OutSamplePosition, float SampleTimeOffset, float OriginTimeOffset, int32 SampleIdx, int8 SchemaSampleBoneIdx = RootSchemaBoneIdx, int8 SchemaOriginBoneIdx = RootSchemaBoneIdx, EPermutationTimeType PermutationTimeType = EPermutationTimeType::UseSampleTime, int32 SamplingAttributeId = -1);
 
-	// Returns the delta velocity of the bone Schema.BoneReferences[SchemaSampleBoneIdx] velocity at time CalculateSampleTime(SampleIdx) + SampleTimeOffset minus
+	// Returns OutSampleVelocity as the delta velocity of the bone Schema.BoneReferences[SchemaSampleBoneIdx] velocity at time CalculateSampleTime(SampleIdx) + SampleTimeOffset minus
 	// the velocity of the bone Schema.BoneReferences[SchemaOriginBoneIdx] at time CalculateSampleTime(SampleIdx) + OriginTimeOffset.
 	// Times will be processed by GetPermutationTimeOffsets(PermutationTimeType, ...)
 	// if bUseCharacterSpaceVelocities is true, velocities will be computed in root bone space, rather than animation (world) space
-	FVector GetSampleVelocity(float SampleTimeOffset, float OriginTimeOffset, int32 SampleIdx, int8 SchemaSampleBoneIdx = RootSchemaBoneIdx, int8 SchemaOriginBoneIdx = RootSchemaBoneIdx, bool bUseCharacterSpaceVelocities = true, EPermutationTimeType PermutationTimeType = EPermutationTimeType::UseSampleTime);
+	bool GetSampleVelocity(FVector& OutSampleVelocity, float SampleTimeOffset, float OriginTimeOffset, int32 SampleIdx, int8 SchemaSampleBoneIdx = RootSchemaBoneIdx, int8 SchemaOriginBoneIdx = RootSchemaBoneIdx, bool bUseCharacterSpaceVelocities = true, EPermutationTimeType PermutationTimeType = EPermutationTimeType::UseSampleTime, int32 SamplingAttributeId = -1);
 
 	int32 GetBeginSampleIdx() const;
 	int32 GetEndSampleIdx() const;
@@ -61,21 +61,24 @@ public:
 	float CalculatePermutationTimeOffset() const;
 	float CalculateSampleTime(int32 SampleIdx) const;
 
+	bool IsProcessFailed() const { return bProcessFailed; }
+
 private:
 	int32 GetVectorIdx(int32 SampleIdx) const;
 
 	// Returns the animation (world) space transform of the bone Schema.BoneReferences[SchemaBoneIdx] at time SampleTime
 	// bClamped will be true if SampleTime is outside the animation duration boundaries
 	FTransform GetTransform(float SampleTime, bool& bClamped, int8 SchemaBoneIdx = RootSchemaBoneIdx);
+	FTransform GetTransform(float SampleTime, bool& bClamped, const FBoneReference& BoneReference);
 
 	// Returns the component space transform of the bone Schema.BoneReferences[SchemaBoneIdx] at time SampleTime
 	// bClamped will be true if SampleTime is outside the animation duration boundaries
 	FTransform GetComponentSpaceTransform(float SampleTime, bool& bClamped, int8 SchemaBoneIdx = RootSchemaBoneIdx);
 
-	// Returns the position of the bone Schema.BoneReferences[SchemaSampleBoneIdx] at time SampleTime relative to the
+	// Returns OutSamplePosition as the position of the bone Schema.BoneReferences[SchemaSampleBoneIdx] at time SampleTime relative to the
 	// transform of the bone Schema.BoneReferences[SchemaOriginBoneIdx] at time OriginTime.
 	// bClamped will be true if SampleTime or OriginTime are outside the animation duration boundaries
-	FVector GetSamplePositionInternal(float SampleTime, float OriginTime, bool& bClamped, int8 SchemaSampleBoneIdx, int8 SchemaOriginBoneIdx);
+	bool GetSamplePositionInternal(FVector& OutSamplePosition, float SampleTime, float OriginTime, bool& bClamped, int8 SchemaSampleBoneIdx, int8 SchemaOriginBoneIdx, int32 SamplingAttributeId);
 
 	struct FSampleInfo
 	{
@@ -97,6 +100,8 @@ private:
 	FTransform MirrorTransform(const FTransform& Transform) const;
 	CachedEntry& GetEntry(float SampleTime);
 	FTransform CalculateComponentSpaceTransform(CachedEntry& Entry, int8 SchemaBoneIdx);
+	FTransform CalculateComponentSpaceTransform(CachedEntry& Entry, const FBoneReference& BoneReference);
+
 	void ComputeStats();
 
 	FBoneContainer BoneContainer;
@@ -111,6 +116,8 @@ private:
 	TArrayView<FPoseMetadata> PoseMetadata;
 
 	FStats Stats;
+
+	bool bProcessFailed = false;
 };
 
 } // namespace UE::PoseSearch
diff --git a/Engine/Plugins/Experimental/Animation/PoseSearch/Source/Runtime/Public/PoseSearch/PoseSearchAssetSampler.h b/Engine/Plugins/Experimental/Animation/PoseSearch/Source/Runtime/Public/PoseSearch/PoseSearchAssetSampler.h
--- a/Engine/Plugins/Experimental/Animation/PoseSearch/Source/Runtime/Public/PoseSearch/PoseSearchAssetSampler.h
+++ b/Engine/Plugins/Experimental/Animation/PoseSearch/Source/Runtime/Public/PoseSearch/PoseSearchAssetSampler.h
@@ -42,6 +42,8 @@ struct POSESEARCH_API FAnimationAssetSampler
 	// Extracts notify states inheriting from UAnimNotifyState_PoseSearchBase present in the sequence at Time.
 	// The function does not empty NotifyStates before adding new notifies!
 	void ExtractPoseSearchNotifyStates(float Time, TArray<UAnimNotifyState_PoseSearchBase*>& NotifyStates) const;
+	
+	TConstArrayView<FAnimNotifyEvent> GetAllAnimNotifyEvents() const;
 
 	const UAnimationAsset* GetAsset() const;
 
diff --git a/Engine/Plugins/Experimental/Animation/PoseSearch/Source/Runtime/Public/PoseSearch/PoseSearchContext.h b/Engine/Plugins/Experimental/Animation/PoseSearch/Source/Runtime/Public/PoseSearch/PoseSearchContext.h
--- a/Engine/Plugins/Experimental/Animation/PoseSearch/Source/Runtime/Public/PoseSearch/PoseSearchContext.h
+++ b/Engine/Plugins/Experimental/Animation/PoseSearch/Source/Runtime/Public/PoseSearch/PoseSearchContext.h
@@ -113,7 +113,7 @@ struct POSESEARCH_API FDebugDrawParams
 	const UPoseSearchSchema* GetSchema() const;
 
 	FVector ExtractPosition(TConstArrayView<float> PoseVector, const UPoseSearchFeatureChannel_Position* Position) const;
-	FVector ExtractPosition(TConstArrayView<float> PoseVector, float SampleTimeOffset, int8 SchemaBoneIdx = RootSchemaBoneIdx, EPermutationTimeType PermutationTimeType = EPermutationTimeType::UseSampleTime) const;
+	FVector ExtractPosition(TConstArrayView<float> PoseVector, float SampleTimeOffset, int8 SchemaBoneIdx = RootSchemaBoneIdx, EPermutationTimeType PermutationTimeType = EPermutationTimeType::UseSampleTime, int32 SamplingAttributeId = -1) const;
 
 	FQuat ExtractRotation(TConstArrayView<float> PoseVector, float SampleTimeOffset, int8 SchemaBoneIdx = RootSchemaBoneIdx) const;
 
@@ -142,27 +142,27 @@ private:
 
 struct POSESEARCH_API FSearchContext
 {
-	FSearchContext(const FPoseSearchQueryTrajectory* InTrajectory, const IPoseHistory* InHistory, float InDesiredPermutationTimeOffset, const FPoseIndicesHistory* InPoseIndicesHistory = nullptr,
+	FSearchContext(const UAnimInstance* InAnimInstance, const FPoseSearchQueryTrajectory* InTrajectory, const IPoseHistory* InHistory, float InDesiredPermutationTimeOffset, const FPoseIndicesHistory* InPoseIndicesHistory = nullptr,
 		const FSearchResult& InCurrentResult = FSearchResult(), float InPoseJumpThresholdTime = 0.f, bool bInForceInterrupt = false);
 
 	// Returns the rotation of the bone Schema.BoneReferences[SchemaSampleBoneIdx] at an offset time of SampleTimeOffset relative to the
 	// transform of the bone Schema.BoneReferences[SchemaOriginBoneIdx] at an offset time of time OriginTimeOffset 
 	// Times will be processed by GetPermutationTimeOffsets(PermutationTimeType, ...)
 	// if bUseHistoryRoot is true, eventually required root transforms will be gathered from the pose history node data, otherwise the context trajectory will be used
-	FQuat GetSampleRotation(float SampleTimeOffset, float OriginTimeOffset, const UPoseSearchSchema* Schema, int8 SchemaSampleBoneIdx = RootSchemaBoneIdx, int8 SchemaOriginBoneIdx = RootSchemaBoneIdx, bool bUseHistoryRoot = false, EPermutationTimeType PermutationTimeType = EPermutationTimeType::UseSampleTime);
+	FQuat GetSampleRotation(float SampleTimeOffset, float OriginTimeOffset, const UPoseSearchSchema* Schema, int8 SchemaSampleBoneIdx = RootSchemaBoneIdx, int8 SchemaOriginBoneIdx = RootSchemaBoneIdx, bool bUseHistoryRoot = false, EPermutationTimeType PermutationTimeType = EPermutationTimeType::UseSampleTime, const FQuat* SampleBoneRotationWorldOverride = nullptr);
 	
 	// Returns the position of the bone Schema.BoneReferences[SchemaSampleBoneIdx] at an offset time of SampleTimeOffset relative to the
 	// transform of the bone Schema.BoneReferences[SchemaOriginBoneIdx] at an offset time of time OriginTimeOffset 
 	// Times will be processed by GetPermutationTimeOffsets(PermutationTimeType, ...)
 	// if bUseHistoryRoot is true, eventually required root transforms will be gathered from the pose history node data, otherwise the context trajectory will be used
-	FVector GetSamplePosition(float SampleTimeOffset, float OriginTimeOffset, const UPoseSearchSchema* Schema, int8 SchemaSampleBoneIdx = RootSchemaBoneIdx, int8 SchemaOriginBoneIdx = RootSchemaBoneIdx, bool bUseHistoryRoot = false, EPermutationTimeType PermutationTimeType = EPermutationTimeType::UseSampleTime);
+	FVector GetSamplePosition(float SampleTimeOffset, float OriginTimeOffset, const UPoseSearchSchema* Schema, int8 SchemaSampleBoneIdx = RootSchemaBoneIdx, int8 SchemaOriginBoneIdx = RootSchemaBoneIdx, bool bUseHistoryRoot = false, EPermutationTimeType PermutationTimeType = EPermutationTimeType::UseSampleTime, const FVector* SampleBonePositionWorldOverride = nullptr);
 	
 	// Returns the delta velocity of the velocity of the bone Schema.BoneReferences[SchemaSampleBoneIdx] at an offset time of SampleTimeOffset minus
 	// the velocity of the bone Schema.BoneReferences[SchemaOriginBoneIdx] at an offset time of time OriginTimeOffset 
 	// Times will be processed by GetPermutationTimeOffsets(PermutationTimeType, ...)
 	// if bUseCharacterSpaceVelocities is true, velocities will be computed in root bone space, rather than world space
 	// if bUseHistoryRoot is true, eventually required root transforms will be gathered from the pose history node data, otherwise the context trajectory will be used
-	FVector GetSampleVelocity(float SampleTimeOffset, float OriginTimeOffset, const UPoseSearchSchema* Schema, int8 SchemaSampleBoneIdx = RootSchemaBoneIdx, int8 SchemaOriginBoneIdx = RootSchemaBoneIdx, bool bUseCharacterSpaceVelocities = true, bool bUseHistoryRoot = false, EPermutationTimeType PermutationTimeType = EPermutationTimeType::UseSampleTime);
+	FVector GetSampleVelocity(float SampleTimeOffset, float OriginTimeOffset, const UPoseSearchSchema* Schema, int8 SchemaSampleBoneIdx = RootSchemaBoneIdx, int8 SchemaOriginBoneIdx = RootSchemaBoneIdx, bool bUseCharacterSpaceVelocities = true, bool bUseHistoryRoot = false, EPermutationTimeType PermutationTimeType = EPermutationTimeType::UseSampleTime, const FVector* SampleBoneVelocityWorldOverride = nullptr);
 
 	void ClearCachedEntries();
 
@@ -185,13 +185,15 @@ struct POSESEARCH_API FSearchContext
 	bool IsTrajectoryValid() const { return Trajectory != nullptr; }
 	bool IsForceInterrupt() const { return bForceInterrupt; }
 	FTransform GetRootAtTime(float Time, bool bUseHistoryRoot = false, bool bExtrapolate = true) const;
-
+	const UAnimInstance* GetAnimInstance() const { return AnimInstance; }
+	
 private:
 	FTransform GetTransform(float SampleTime, const UPoseSearchSchema* Schema, int8 SchemaBoneIdx = RootSchemaBoneIdx, bool bUseHistoryRoot = false);
 	FTransform GetComponentSpaceTransform(float SampleTime, const UPoseSearchSchema* Schema, int8 SchemaSampleBoneIdx = RootSchemaBoneIdx);
-	FVector GetSamplePositionInternal(float SampleTime, float OriginTime, const UPoseSearchSchema* Schema, int8 SchemaSampleBoneIdx = RootSchemaBoneIdx, int8 SchemaOriginBoneIdx = RootSchemaBoneIdx, bool bUseHistoryRoot = false);
-	FQuat GetSampleRotationInternal(float SampleTime, float OriginTime, const UPoseSearchSchema* Schema, int8 SchemaSampleBoneIdx = RootSchemaBoneIdx, int8 SchemaOriginBoneIdx = RootSchemaBoneIdx, bool bUseHistoryRoot = false);
+	FVector GetSamplePositionInternal(float SampleTime, float OriginTime, const UPoseSearchSchema* Schema, int8 SchemaSampleBoneIdx = RootSchemaBoneIdx, int8 SchemaOriginBoneIdx = RootSchemaBoneIdx, bool bUseHistoryRoot = false, const FVector* SampleBonePositionWorldOverride = nullptr);
+	FQuat GetSampleRotationInternal(float SampleTime, float OriginTime, const UPoseSearchSchema* Schema, int8 SchemaSampleBoneIdx = RootSchemaBoneIdx, int8 SchemaOriginBoneIdx = RootSchemaBoneIdx, bool bUseHistoryRoot = false, const FQuat* SampleBoneRotationWorldOverride = nullptr);
 
+	const UAnimInstance* AnimInstance = nullptr;
 	const FPoseSearchQueryTrajectory* Trajectory = nullptr;
 	const IPoseHistory* History = nullptr;
 	float DesiredPermutationTimeOffset = 0.f;
diff --git a/Engine/Plugins/Experimental/Animation/PoseSearch/Source/Runtime/Public/PoseSearch/PoseSearchFeatureChannel.h b/Engine/Plugins/Experimental/Animation/PoseSearch/Source/Runtime/Public/PoseSearch/PoseSearchFeatureChannel.h
--- a/Engine/Plugins/Experimental/Animation/PoseSearch/Source/Runtime/Public/PoseSearch/PoseSearchFeatureChannel.h
+++ b/Engine/Plugins/Experimental/Animation/PoseSearch/Source/Runtime/Public/PoseSearch/PoseSearchFeatureChannel.h
@@ -136,7 +136,7 @@ public:
 	virtual void FillWeights(TArrayView<float> Weights) const PURE_VIRTUAL(UPoseSearchFeatureChannel::FillWeights, );
 
 	// Called at database build time to populate pose vectors with this channel's data
-	virtual void IndexAsset(UE::PoseSearch::FAssetIndexer& Indexer) const PURE_VIRTUAL(UPoseSearchFeatureChannel::IndexAsset, );
+	virtual bool IndexAsset(UE::PoseSearch::FAssetIndexer& Indexer) const PURE_VIRTUAL(UPoseSearchFeatureChannel::IndexAsset, return false;);
 
 	// returns the FString used editor side to identify this UPoseSearchFeatureChannel (for instance in the pose search debugger)
 	virtual FString GetLabel() const;
diff --git a/Engine/Plugins/Experimental/Animation/PoseSearch/Source/Runtime/Public/PoseSearch/PoseSearchSchema.h b/Engine/Plugins/Experimental/Animation/PoseSearch/Source/Runtime/Public/PoseSearch/PoseSearchSchema.h
--- a/Engine/Plugins/Experimental/Animation/PoseSearch/Source/Runtime/Public/PoseSearch/PoseSearchSchema.h
+++ b/Engine/Plugins/Experimental/Animation/PoseSearch/Source/Runtime/Public/PoseSearch/PoseSearchSchema.h
@@ -146,8 +146,6 @@ public:
 
 	FBoneIndexType GetBoneIndexType(int8 SchemaBoneIdx) const;
 
-	bool IsRootBone(int8 SchemaBoneIdx) const;
-	
 private:
 	template <typename FindPredicateType>
 	static const UPoseSearchFeatureChannel* FindChannelRecursive(TConstArrayView<TObjectPtr<UPoseSearchFeatureChannel>> Channels, FindPredicateType FindPredicate)
