motion matching - adding BlueprintPure UFUNCTION UMotionMatchingAnimNodeLibrary::GetMotionMatchingSearchResult

diff --git a/Engine/Plugins/Experimental/Animation/PoseSearch/Source/Runtime/Private/MotionMatchingAnimNodeLibrary.cpp b/Engine/Plugins/Experimental/Animation/PoseSearch/Source/Runtime/Private/MotionMatchingAnimNodeLibrary.cpp
--- a/Engine/Plugins/Experimental/Animation/PoseSearch/Source/Runtime/Private/MotionMatchingAnimNodeLibrary.cpp
+++ b/Engine/Plugins/Experimental/Animation/PoseSearch/Source/Runtime/Private/MotionMatchingAnimNodeLibrary.cpp
@@ -1,60 +1,90 @@
 // Copyright Epic Games, Inc. All Rights Reserved.
 
 #include "PoseSearch/MotionMatchingAnimNodeLibrary.h"
-
 #include "PoseSearch/AnimNode_MotionMatching.h"
-#include "PoseSearch/PoseSearchDefines.h"
+#include "PoseSearch/PoseSearchDatabase.h"
 
 FMotionMatchingAnimNodeReference UMotionMatchingAnimNodeLibrary::ConvertToMotionMatchingNode(const FAnimNodeReference& Node, EAnimNodeReferenceConversionResult& Result)
 {
 	return FAnimNodeReference::ConvertToType<FMotionMatchingAnimNodeReference>(Node, Result);
 }
 
-void UMotionMatchingAnimNodeLibrary::SetDatabaseToSearch(const FMotionMatchingAnimNodeReference& MotionMatchingNode, UPoseSearchDatabase* Database, bool bForceInterruptIfNew)
+void UMotionMatchingAnimNodeLibrary::GetMotionMatchingSearchResult(const FMotionMatchingAnimNodeReference& MotionMatchingNode, FPoseSearchBlueprintResult& Result, bool& bIsResultValid)
 {
-	FAnimNode_MotionMatching* MotionMatchingNodePtr = MotionMatchingNode.GetAnimNodePtr<FAnimNode_MotionMatching>();
-	if (MotionMatchingNodePtr == nullptr)
+	using namespace UE::PoseSearch;
+
+	bIsResultValid = false;
+	if (FAnimNode_MotionMatching* MotionMatchingNodePtr = MotionMatchingNode.GetAnimNodePtr<FAnimNode_MotionMatching>())
+	{
+		const FMotionMatchingState& MotionMatchingState = MotionMatchingNodePtr->GetMotionMatchingState();
+		if (const FSearchIndexAsset* SearchIndexAsset = MotionMatchingState.CurrentSearchResult.GetSearchIndexAsset())
+		{
+			const UPoseSearchDatabase* CurrentResultDatabase = MotionMatchingState.CurrentSearchResult.Database.Get();
+			if (CurrentResultDatabase && CurrentResultDatabase->Schema)
+			{
+				if (const FPoseSearchDatabaseAnimationAssetBase* DatabaseAsset = CurrentResultDatabase->GetAnimationAssetBase(*SearchIndexAsset))
+				{
+					Result.SelectedAnimation = DatabaseAsset->GetAnimationAsset();
+					Result.SelectedTime = MotionMatchingState.CurrentSearchResult.AssetTime;
+					Result.bLoop = DatabaseAsset->IsLooping();
+					Result.bIsMirrored = SearchIndexAsset->bMirrored;
+					Result.BlendParameters = SearchIndexAsset->BlendParameters;
+					Result.SelectedDatabase = CurrentResultDatabase;
+					bIsResultValid = true;
+				}
+			}
+		}
+	}
+	else
 	{
 		UE_LOG(LogPoseSearch, Warning, TEXT("UMotionMatchingAnimNodeLibrary::SetDatabase called on an invalid context or with an invalid type"));
-		return;
 	}
+}
 
-	MotionMatchingNodePtr->SetDatabaseToSearch(Database, bForceInterruptIfNew);
+void UMotionMatchingAnimNodeLibrary::SetDatabaseToSearch(const FMotionMatchingAnimNodeReference& MotionMatchingNode, UPoseSearchDatabase* Database, bool bForceInterruptIfNew)
+{
+	if (FAnimNode_MotionMatching* MotionMatchingNodePtr = MotionMatchingNode.GetAnimNodePtr<FAnimNode_MotionMatching>())
+	{
+		MotionMatchingNodePtr->SetDatabaseToSearch(Database, bForceInterruptIfNew);
+	}
+	else
+	{
+		UE_LOG(LogPoseSearch, Warning, TEXT("UMotionMatchingAnimNodeLibrary::SetDatabase called on an invalid context or with an invalid type"));
+	}
 }
 
 void UMotionMatchingAnimNodeLibrary::SetDatabasesToSearch(const FMotionMatchingAnimNodeReference& MotionMatchingNode, const TArray<UPoseSearchDatabase*>& Databases, bool bForceInterruptIfNew)
 {
-	FAnimNode_MotionMatching* MotionMatchingNodePtr = MotionMatchingNode.GetAnimNodePtr<FAnimNode_MotionMatching>();
-	if (MotionMatchingNodePtr == nullptr)
+	if (FAnimNode_MotionMatching* MotionMatchingNodePtr = MotionMatchingNode.GetAnimNodePtr<FAnimNode_MotionMatching>())
+	{
+		MotionMatchingNodePtr->SetDatabasesToSearch(Databases, bForceInterruptIfNew);
+	}
+	else
 	{
 		UE_LOG(LogPoseSearch, Warning, TEXT("UMotionMatchingAnimNodeLibrary::SetDatabases called on an invalid context or with an invalid type"));
-		return;
 	}
-
-	MotionMatchingNodePtr->SetDatabasesToSearch(Databases, bForceInterruptIfNew);
 }
 
 void UMotionMatchingAnimNodeLibrary::ResetDatabasesToSearch(const FMotionMatchingAnimNodeReference& MotionMatchingNode, bool bForceInterrupt)
 {
-	FAnimNode_MotionMatching* MotionMatchingNodePtr = MotionMatchingNode.GetAnimNodePtr<FAnimNode_MotionMatching>();
-	if (MotionMatchingNodePtr == nullptr)
+	if (FAnimNode_MotionMatching* MotionMatchingNodePtr = MotionMatchingNode.GetAnimNodePtr<FAnimNode_MotionMatching>())
+	{
+		MotionMatchingNodePtr->ResetDatabasesToSearch(bForceInterrupt);
+	}
+	else
 	{
 		UE_LOG(LogPoseSearch, Warning, TEXT("UMotionMatchingAnimNodeLibrary::ResetDatabasesToSearch called on an invalid context or with an invalid type"));
-		return;
 	}
-
-	MotionMatchingNodePtr->ResetDatabasesToSearch(bForceInterrupt);
-
 }
 
 void UMotionMatchingAnimNodeLibrary::ForceInterruptNextUpdate(const FMotionMatchingAnimNodeReference& MotionMatchingNode)
 {
-	FAnimNode_MotionMatching* MotionMatchingNodePtr = MotionMatchingNode.GetAnimNodePtr<FAnimNode_MotionMatching>();
-	if (MotionMatchingNodePtr == nullptr)
+	if (FAnimNode_MotionMatching* MotionMatchingNodePtr = MotionMatchingNode.GetAnimNodePtr<FAnimNode_MotionMatching>())
+	{
+		MotionMatchingNodePtr->ForceInterruptNextUpdate();
+	}
+	else
 	{
 		UE_LOG(LogPoseSearch, Warning, TEXT("UMotionMatchingAnimNodeLibrary::ForceInterruptOneFrame called on an invalid context or with an invalid type"));
-		return;
 	}
-
-	MotionMatchingNodePtr->ForceInterruptNextUpdate();
 }
\ No newline at end of file
diff --git a/Engine/Plugins/Experimental/Animation/PoseSearch/Source/Runtime/Private/PoseSearchLibrary.cpp b/Engine/Plugins/Experimental/Animation/PoseSearch/Source/Runtime/Private/PoseSearchLibrary.cpp
--- a/Engine/Plugins/Experimental/Animation/PoseSearch/Source/Runtime/Private/PoseSearchLibrary.cpp
+++ b/Engine/Plugins/Experimental/Animation/PoseSearch/Source/Runtime/Private/PoseSearchLibrary.cpp
@@ -423,12 +423,7 @@ void UPoseSearchLibrary::MotionMatch(
 	const FPoseSearchQueryTrajectory Trajectory,
 	float TrajectorySpeedMultiplier,
 	const FName PoseHistoryName,
-	UAnimationAsset*& SelectedAnimation,
-	float& SelectedTime,
-	bool& bLoop,
-	bool& bIsMirrored,
-	FVector& BlendParameters,
-	float& SearchCost,
+	FPoseSearchBlueprintResult& Result,
 	const UAnimationAsset* FutureAnimation,
 	float FutureAnimationStartTime,
 	float TimeToFutureAnimationStart,
@@ -454,12 +449,13 @@ void UPoseSearchLibrary::MotionMatch(
 
 	static constexpr float FiniteDelta = 1 / 60.0f;
 
-	SelectedAnimation = nullptr;
-	SelectedTime = 0.f;
-	bLoop = false;
-	bIsMirrored = false;
-	BlendParameters = FVector::ZeroVector;
-	SearchCost = MAX_flt;
+	Result.SelectedAnimation = nullptr;
+	Result.SelectedTime = 0.f;
+	Result.bLoop = false;
+	Result.bIsMirrored = false;
+	Result.BlendParameters = FVector::ZeroVector;
+	Result.SelectedDatabase = nullptr;
+	Result.SearchCost = MAX_flt;
 
 	if (Database && AnimInstance)
 	{
@@ -548,12 +544,13 @@ void UPoseSearchLibrary::MotionMatch(
 			const FSearchIndexAsset* SearchIndexAsset = SearchResult.GetSearchIndexAsset();
 			if (const FPoseSearchDatabaseAnimationAssetBase* DatabaseAsset = SearchResult.Database->GetAnimationAssetBase(*SearchIndexAsset))
 			{
-				SelectedAnimation = DatabaseAsset->GetAnimationAsset();
-				SelectedTime = SearchResult.AssetTime;
-				bLoop = DatabaseAsset->IsLooping();
-				bIsMirrored = SearchIndexAsset->bMirrored;
-				BlendParameters = SearchIndexAsset->BlendParameters;
-				SearchCost = SearchResult.PoseCost.GetTotalCost();
+				Result.SelectedAnimation = DatabaseAsset->GetAnimationAsset();
+				Result.SelectedTime = SearchResult.AssetTime;
+				Result.bLoop = DatabaseAsset->IsLooping();
+				Result.bIsMirrored = SearchIndexAsset->bMirrored;
+				Result.BlendParameters = SearchIndexAsset->BlendParameters;
+				Result.SelectedDatabase = Database;
+				Result.SearchCost = SearchResult.PoseCost.GetTotalCost();
 			}
 		}
 
diff --git a/Engine/Plugins/Experimental/Animation/PoseSearch/Source/Runtime/Public/PoseSearch/AnimNode_MotionMatching.h b/Engine/Plugins/Experimental/Animation/PoseSearch/Source/Runtime/Public/PoseSearch/AnimNode_MotionMatching.h
--- a/Engine/Plugins/Experimental/Animation/PoseSearch/Source/Runtime/Public/PoseSearch/AnimNode_MotionMatching.h
+++ b/Engine/Plugins/Experimental/Animation/PoseSearch/Source/Runtime/Public/PoseSearch/AnimNode_MotionMatching.h
@@ -29,6 +29,8 @@ public:
 	// Ignore the continuing pose on the next update and force a search.
 	void ForceInterruptNextUpdate();
 
+	const FMotionMatchingState& GetMotionMatchingState() const { return MotionMatchingState; }
+
 private:
 	// FAnimNode_Base interface
 	// @todo: implement CacheBones_AnyThread to rebind the schema bones
diff --git a/Engine/Plugins/Experimental/Animation/PoseSearch/Source/Runtime/Public/PoseSearch/MotionMatchingAnimNodeLibrary.h b/Engine/Plugins/Experimental/Animation/PoseSearch/Source/Runtime/Public/PoseSearch/MotionMatchingAnimNodeLibrary.h
--- a/Engine/Plugins/Experimental/Animation/PoseSearch/Source/Runtime/Public/PoseSearch/MotionMatchingAnimNodeLibrary.h
+++ b/Engine/Plugins/Experimental/Animation/PoseSearch/Source/Runtime/Public/PoseSearch/MotionMatchingAnimNodeLibrary.h
@@ -38,6 +38,9 @@ public:
 		Result = (ConversionResult == EAnimNodeReferenceConversionResult::Succeeded);
 	}
 
+	UFUNCTION(BlueprintPure, Category = "Animation|MotionMatching", meta = (BlueprintThreadSafe, DisplayName = "Get Motion Matching Search Result"))
+	static void GetMotionMatchingSearchResult(const FMotionMatchingAnimNodeReference& MotionMatchingNode, FPoseSearchBlueprintResult& Result, bool& bIsResultValid);
+
 	/**
 	 * Set the database to search on the motion matching node. This overrides the Database property on the motion matching node.
 	 * @param MotionMatchingNode - The motion matching node to operate on.
diff --git a/Engine/Plugins/Experimental/Animation/PoseSearch/Source/Runtime/Public/PoseSearch/PoseSearchCost.h b/Engine/Plugins/Experimental/Animation/PoseSearch/Source/Runtime/Public/PoseSearch/PoseSearchCost.h
--- a/Engine/Plugins/Experimental/Animation/PoseSearch/Source/Runtime/Public/PoseSearch/PoseSearchCost.h
+++ b/Engine/Plugins/Experimental/Animation/PoseSearch/Source/Runtime/Public/PoseSearch/PoseSearchCost.h
@@ -10,12 +10,11 @@ struct POSESEARCH_API FPoseSearchCost
 	GENERATED_BODY()
 public:
 	FPoseSearchCost() = default;
-	FPoseSearchCost(float InDissimilarityCost, float InNotifyCostAddend, float InContinuingPoseCostAddend)
-	: TotalCost(InDissimilarityCost + InNotifyCostAddend + InContinuingPoseCostAddend)
+	FPoseSearchCost(float DissimilarityCost, float NotifyCostAddend, float ContinuingPoseCostAddend)
+	: TotalCost(DissimilarityCost + NotifyCostAddend + ContinuingPoseCostAddend)
 	{
 #if WITH_EDITORONLY_DATA
-		NotifyCostAddend = InNotifyCostAddend;
-		ContinuingPoseCostAddend = InContinuingPoseCostAddend;
+		CostAddend = NotifyCostAddend + ContinuingPoseCostAddend;
 #endif // WITH_EDITORONLY_DATA
 	}
 
@@ -30,14 +29,11 @@ protected:
 #if WITH_EDITORONLY_DATA
 public:
 
-	float GetCostAddend() const { return NotifyCostAddend + ContinuingPoseCostAddend; }
+	float GetCostAddend() const { return CostAddend; }
 
-	// Contribution from ModifyCost anim notify
+	// Contribution from ModifyCost anim notify, and ContinuingPoseCostAddend
 	UPROPERTY()
-	float NotifyCostAddend = 0.f;
-
-	UPROPERTY()
-	float ContinuingPoseCostAddend = 0.f;
+	float CostAddend = 0.f;
 
 #endif // WITH_EDITORONLY_DATA
 };
\ No newline at end of file
diff --git a/Engine/Plugins/Experimental/Animation/PoseSearch/Source/Runtime/Public/PoseSearch/PoseSearchLibrary.h b/Engine/Plugins/Experimental/Animation/PoseSearch/Source/Runtime/Public/PoseSearch/PoseSearchLibrary.h
--- a/Engine/Plugins/Experimental/Animation/PoseSearch/Source/Runtime/Public/PoseSearch/PoseSearchLibrary.h
+++ b/Engine/Plugins/Experimental/Animation/PoseSearch/Source/Runtime/Public/PoseSearch/PoseSearchLibrary.h
@@ -23,11 +23,8 @@ namespace UE::PoseSearch
 struct FAnimationUpdateContext;
 struct FPoseSearchQueryTrajectory;
 
-USTRUCT(BlueprintType, Category="Animation|Pose Search")
-struct POSESEARCH_API FMotionMatchingState
+struct FMotionMatchingState
 {
-	GENERATED_BODY()
-
 	// Reset the state to a default state using the current Database
 	void Reset(const FTransform& ComponentTransform);
 
@@ -46,15 +43,12 @@ struct POSESEARCH_API FMotionMatchingState
 	UE::PoseSearch::FSearchResult CurrentSearchResult;
 
 	// Time since the last pose jump
-	UPROPERTY(Transient)
 	float ElapsedPoseSearchTime = 0.f;
 
 	// wanted PlayRate to have the selected animation playing at the estimated requested speed from the query.
-	UPROPERTY(Transient)
 	float WantedPlayRate = 1.f;
 
 	// true if a new animation has been selected
-	UPROPERTY(Transient, VisibleAnywhere, BlueprintReadOnly, Category=State)
 	bool bJumpedToPose = false;
 
 	UE::PoseSearch::FPoseIndicesHistory PoseIndicesHistory;
@@ -167,12 +161,7 @@ public:
 		const FPoseSearchQueryTrajectory Trajectory,
 		float TrajectorySpeedMultiplier,
 		const FName PoseHistoryName,
-		UAnimationAsset*& SelectedAnimation,
-		float& SelectedTime,
-		bool& bLoop,
-		bool& bIsMirrored,
-		FVector& BlendParameters,
-		float& SearchCost,
+		FPoseSearchBlueprintResult& Result,
 		const UAnimationAsset* FutureAnimation = nullptr,
 		float FutureAnimationStartTime = 0.f,
 		float TimeToFutureAnimationStart = 0.f,
diff --git a/Engine/Plugins/Experimental/Animation/PoseSearch/Source/Runtime/Public/PoseSearch/PoseSearchResult.h b/Engine/Plugins/Experimental/Animation/PoseSearch/Source/Runtime/Public/PoseSearch/PoseSearchResult.h
--- a/Engine/Plugins/Experimental/Animation/PoseSearch/Source/Runtime/Public/PoseSearch/PoseSearchResult.h
+++ b/Engine/Plugins/Experimental/Animation/PoseSearch/Source/Runtime/Public/PoseSearch/PoseSearchResult.h
@@ -4,6 +4,7 @@
 
 #include "PoseSearch/PoseSearchCost.h"
 #include "PoseSearch/PoseSearchDefines.h"
+#include "PoseSearchResult.generated.h"
 
 class UPoseSearchDatabase;
 class UPoseSearchSchema;
@@ -59,4 +60,33 @@ struct FSearchResult
 	bool CanAdvance(float DeltaTime) const;
 };
 
-} // namespace UE::PoseSearch
\ No newline at end of file
+} // namespace UE::PoseSearch
+
+USTRUCT(BlueprintType, Category="Animation|Pose Search")
+struct POSESEARCH_API FPoseSearchBlueprintResult
+{
+	GENERATED_BODY()
+public:
+
+	UPROPERTY(Transient, VisibleAnywhere, BlueprintReadOnly, Category=State)
+	TWeakObjectPtr<const UAnimationAsset> SelectedAnimation = nullptr;
+	
+	UPROPERTY(Transient, VisibleAnywhere, BlueprintReadOnly, Category=State)
+	float SelectedTime = 0.f;
+	
+	UPROPERTY(Transient, VisibleAnywhere, BlueprintReadOnly, Category=State)
+	bool bLoop = false;
+	
+	UPROPERTY(Transient, VisibleAnywhere, BlueprintReadOnly, Category=State)
+	bool bIsMirrored = false;
+	
+	UPROPERTY(Transient, VisibleAnywhere, BlueprintReadOnly, Category=State)
+	FVector BlendParameters = FVector::ZeroVector;
+
+	UPROPERTY(Transient, VisibleAnywhere, BlueprintReadOnly, Category=State)
+	TWeakObjectPtr<const UPoseSearchDatabase> SelectedDatabase = nullptr;
+
+	UPROPERTY(Transient, VisibleAnywhere, BlueprintReadOnly, Category=State)
+	float SearchCost = MAX_flt;
+};
+
