motion matching - PoseJumpThresholdTime as float range

diff --git a/Engine/Plugins/Experimental/Animation/PoseSearch/Source/Runtime/Private/PoseSearchContext.cpp b/Engine/Plugins/Experimental/Animation/PoseSearch/Source/Runtime/Private/PoseSearchContext.cpp
--- a/Engine/Plugins/Experimental/Animation/PoseSearch/Source/Runtime/Private/PoseSearchContext.cpp
+++ b/Engine/Plugins/Experimental/Animation/PoseSearch/Source/Runtime/Private/PoseSearchContext.cpp
@@ -352,7 +352,7 @@ void FDebugDrawParams::DrawFeatureVector(int32 PoseIdx)
 //////////////////////////////////////////////////////////////////////////
 // FSearchContext
 FSearchContext::FSearchContext(const UAnimInstance* InAnimInstance, const FPoseSearchQueryTrajectory* InTrajectory, const IPoseHistory* InHistory, float InDesiredPermutationTimeOffset,
-	const FPoseIndicesHistory* InPoseIndicesHistory, const FSearchResult& InCurrentResult, float InPoseJumpThresholdTime, bool bInForceInterrupt)
+	const FPoseIndicesHistory* InPoseIndicesHistory, const FSearchResult& InCurrentResult, const FFloatInterval& InPoseJumpThresholdTime, bool bInForceInterrupt)
 : AnimInstance(InAnimInstance)
 , Trajectory(InTrajectory)
 , History(InHistory)
diff --git a/Engine/Plugins/Experimental/Animation/PoseSearch/Source/Runtime/Private/PoseSearchDatabase.cpp b/Engine/Plugins/Experimental/Animation/PoseSearch/Source/Runtime/Private/PoseSearchDatabase.cpp
--- a/Engine/Plugins/Experimental/Animation/PoseSearch/Source/Runtime/Private/PoseSearchDatabase.cpp
+++ b/Engine/Plugins/Experimental/Animation/PoseSearch/Source/Runtime/Private/PoseSearchDatabase.cpp
@@ -39,31 +39,23 @@ static void PopulateNonSelectableIdx(FNonSelectableIdx& NonSelectableIdx, FSearc
 
 	NonSelectableIdx.Reset();
 	const FSearchIndexAsset* CurrentIndexAsset = SearchContext.GetCurrentResult().GetSearchIndexAsset();
-	if (CurrentIndexAsset && SearchContext.IsCurrentResultFromDatabase(Database) && SearchContext.GetPoseJumpThresholdTime() > 0.f)
+	if (CurrentIndexAsset && SearchContext.IsCurrentResultFromDatabase(Database) && !FMath::IsNearlyEqual(SearchContext.GetPoseJumpThresholdTime().Min, SearchContext.GetPoseJumpThresholdTime().Max))
 	{
-		const int32 PoseJumpIndexThreshold = FMath::FloorToInt(SearchContext.GetPoseJumpThresholdTime() * Database->Schema->SampleRate);
+		const int32 CurrentResultPoseIdx = SearchContext.GetCurrentResult().PoseIdx;
+		const int32 UnboundMinPoseIdx = CurrentResultPoseIdx + FMath::FloorToInt(SearchContext.GetPoseJumpThresholdTime().Min * Database->Schema->SampleRate);
+		const int32 UnboundMaxPoseIdx = CurrentResultPoseIdx + FMath::CeilToInt(SearchContext.GetPoseJumpThresholdTime().Max * Database->Schema->SampleRate);
+		const int32 CurrentIndexAssetFirstPoseIdx = CurrentIndexAsset->FirstPoseIdx;
+		const int32 CurrentIndexAssetNumPoses = CurrentIndexAsset->GetNumPoses();
 		const bool IsLooping = Database->IsSourceAssetLooping(*CurrentIndexAsset);
 
-		for (int32 i = -PoseJumpIndexThreshold; i <= -1; ++i)
+		if (IsLooping)
 		{
-			int32 PoseIdx = SearchContext.GetCurrentResult().PoseIdx + i;
-			bool bIsPoseInRange = false;
-			if (IsLooping)
+			for (int32 UnboundPoseIdx = UnboundMinPoseIdx; UnboundPoseIdx < UnboundMaxPoseIdx; ++UnboundPoseIdx)
 			{
-				bIsPoseInRange = true;
+				const int32 Modulo = (UnboundPoseIdx - CurrentIndexAssetFirstPoseIdx) % CurrentIndexAssetNumPoses;
+				const int32 CurrentIndexAssetFirstPoseIdxPlusModulo = CurrentIndexAssetFirstPoseIdx + Modulo;
+				const int32 PoseIdx = Modulo >= 0 ? CurrentIndexAssetFirstPoseIdxPlusModulo : CurrentIndexAssetFirstPoseIdxPlusModulo + CurrentIndexAssetNumPoses;
 
-				while (PoseIdx < CurrentIndexAsset->FirstPoseIdx)
-				{
-					PoseIdx += CurrentIndexAsset->GetNumPoses();
-				}
-			}
-			else if (CurrentIndexAsset->IsPoseInRange(PoseIdx))
-			{
-				bIsPoseInRange = true;
-			}
-
-			if (bIsPoseInRange)
-			{
 				NonSelectableIdx.AddUnique(PoseIdx);
 
 #if UE_POSE_SEARCH_TRACE_ENABLED
@@ -72,43 +64,22 @@ static void PopulateNonSelectableIdx(FNonSelectableIdx& NonSelectableIdx, FSearc
 				SearchContext.BestCandidates.Add(PoseCost, PoseIdx, Database, EPoseCandidateFlags::DiscardedBy_PoseJumpThresholdTime);
 #endif // UE_POSE_SEARCH_TRACE_ENABLED
 			}
-			else
-			{
-				break;
-			}
 		}
-
-		for (int32 i = 0; i <= PoseJumpIndexThreshold; ++i)
+		else
 		{
-			int32 PoseIdx = SearchContext.GetCurrentResult().PoseIdx + i;
-			bool bIsPoseInRange = false;
-			if (IsLooping)
-			{
-				bIsPoseInRange = true;
-
-				while (PoseIdx >= CurrentIndexAsset->FirstPoseIdx + CurrentIndexAsset->GetNumPoses())
-				{
-					PoseIdx -= CurrentIndexAsset->GetNumPoses();
-				}
-			}
-			else if (CurrentIndexAsset->IsPoseInRange(PoseIdx))
-			{
-				bIsPoseInRange = true;
-			}
+			const int32 MinPoseIdx = FMath::Max(CurrentIndexAssetFirstPoseIdx, UnboundMinPoseIdx);
+			const int32 MaxPoseIdx = FMath::Min(CurrentIndexAssetFirstPoseIdx + CurrentIndexAssetNumPoses, UnboundMaxPoseIdx);
 
-			if (bIsPoseInRange)
+			for (int32 PoseIdx = MinPoseIdx; PoseIdx < MaxPoseIdx; ++PoseIdx)
 			{
 				NonSelectableIdx.AddUnique(PoseIdx);
 
 #if UE_POSE_SEARCH_TRACE_ENABLED
-				const FPoseSearchCost PoseCost = SearchIndex.ComparePoses(PoseIdx, 0.f, SearchIndex.GetPoseValuesSafe(PoseIdx), QueryValues);
+				const TArray<float> PoseValues = SearchIndex.GetPoseValuesSafe(PoseIdx);
+				const FPoseSearchCost PoseCost = SearchIndex.ComparePoses(PoseIdx, 0.f, PoseValues, QueryValues);
 				SearchContext.BestCandidates.Add(PoseCost, PoseIdx, Database, EPoseCandidateFlags::DiscardedBy_PoseJumpThresholdTime);
 #endif // UE_POSE_SEARCH_TRACE_ENABLED
 			}
-			else
-			{
-				break;
-			}
 		}
 	}
 
diff --git a/Engine/Plugins/Experimental/Animation/PoseSearch/Source/Runtime/Private/PoseSearchLibrary.cpp b/Engine/Plugins/Experimental/Animation/PoseSearch/Source/Runtime/Private/PoseSearchLibrary.cpp
--- a/Engine/Plugins/Experimental/Animation/PoseSearch/Source/Runtime/Private/PoseSearchLibrary.cpp
+++ b/Engine/Plugins/Experimental/Animation/PoseSearch/Source/Runtime/Private/PoseSearchLibrary.cpp
@@ -244,7 +244,7 @@ void UPoseSearchLibrary::UpdateMotionMatchingState(
 	float TrajectorySpeedMultiplier,
 	float BlendTime,
 	int32 MaxActiveBlends,
-	float PoseJumpThresholdTime,
+	const FFloatInterval& PoseJumpThresholdTime,
 	float PoseReselectHistory,
 	float SearchThrottleTime,
 	const FFloatInterval& PlayRate,
diff --git a/Engine/Plugins/Experimental/Animation/PoseSearch/Source/Runtime/Public/PoseSearch/AnimNode_MotionMatching.h b/Engine/Plugins/Experimental/Animation/PoseSearch/Source/Runtime/Public/PoseSearch/AnimNode_MotionMatching.h
--- a/Engine/Plugins/Experimental/Animation/PoseSearch/Source/Runtime/Public/PoseSearch/AnimNode_MotionMatching.h
+++ b/Engine/Plugins/Experimental/Animation/PoseSearch/Source/Runtime/Public/PoseSearch/AnimNode_MotionMatching.h
@@ -68,9 +68,9 @@ private:
 	UPROPERTY(EditAnywhere, Category = Settings, meta = (PinHiddenByDefault))
 	EAlphaBlendOption BlendOption = EAlphaBlendOption::Linear;
 
-	// Don't jump to poses of the same segment that are less than this many seconds away.
-	UPROPERTY(EditAnywhere, Category = Settings, meta = (PinHiddenByDefault, ClampMin="0"))
-	float PoseJumpThresholdTime = 0.f;
+	// Don't jump to poses of the same segment that are within the interval this many seconds away from the continuing pose.
+	UPROPERTY(EditAnywhere, Category = Settings, meta = (PinHiddenByDefault))
+	FFloatInterval PoseJumpThresholdTime = FFloatInterval(0.f, 0.f);
 
 	// Prevent re-selection of poses that have been selected previously within this much time (in seconds) in the past. This is across all animation segments that have been selected within this time range.
 	UPROPERTY(EditAnywhere, Category = Settings, meta = (PinHiddenByDefault, ClampMin = "0"))
diff --git a/Engine/Plugins/Experimental/Animation/PoseSearch/Source/Runtime/Public/PoseSearch/PoseSearchContext.h b/Engine/Plugins/Experimental/Animation/PoseSearch/Source/Runtime/Public/PoseSearch/PoseSearchContext.h
--- a/Engine/Plugins/Experimental/Animation/PoseSearch/Source/Runtime/Public/PoseSearch/PoseSearchContext.h
+++ b/Engine/Plugins/Experimental/Animation/PoseSearch/Source/Runtime/Public/PoseSearch/PoseSearchContext.h
@@ -143,7 +143,7 @@ private:
 struct POSESEARCH_API FSearchContext
 {
 	FSearchContext(const UAnimInstance* InAnimInstance, const FPoseSearchQueryTrajectory* InTrajectory, const IPoseHistory* InHistory, float InDesiredPermutationTimeOffset, const FPoseIndicesHistory* InPoseIndicesHistory = nullptr,
-		const FSearchResult& InCurrentResult = FSearchResult(), float InPoseJumpThresholdTime = 0.f, bool bInForceInterrupt = false);
+		const FSearchResult& InCurrentResult = FSearchResult(), const FFloatInterval& InPoseJumpThresholdTime = FFloatInterval(0.f, 0.f), bool bInForceInterrupt = false);
 
 	// Returns the rotation of the bone Schema.BoneReferences[SchemaSampleBoneIdx] at an offset time of SampleTimeOffset relative to the
 	// transform of the bone Schema.BoneReferences[SchemaOriginBoneIdx] at an offset time of time OriginTimeOffset 
@@ -178,7 +178,7 @@ struct POSESEARCH_API FSearchContext
 	TConstArrayView<float> GetCurrentResultPoseVector() const { return CurrentResultPoseVector; }
 
 	const FSearchResult& GetCurrentResult() const { return CurrentResult; }
-	float GetPoseJumpThresholdTime() const { return PoseJumpThresholdTime; }
+	const FFloatInterval& GetPoseJumpThresholdTime() const { return PoseJumpThresholdTime; }
 	const FPoseIndicesHistory* GetPoseIndicesHistory() const { return PoseIndicesHistory; }
 	bool IsHistoryValid() const { return History != nullptr; }
 	float GetDesiredPermutationTimeOffset() const { return DesiredPermutationTimeOffset; }
@@ -199,7 +199,7 @@ private:
 	float DesiredPermutationTimeOffset = 0.f;
 	const FPoseIndicesHistory* PoseIndicesHistory = nullptr;
 	FSearchResult CurrentResult;
-	float PoseJumpThresholdTime = 0.f;
+	FFloatInterval PoseJumpThresholdTime = FFloatInterval(0.f, 0.f);
 	bool bForceInterrupt = false;
 
 	TConstArrayView<float> CurrentResultPoseVector;
diff --git a/Engine/Plugins/Experimental/Animation/PoseSearch/Source/Runtime/Public/PoseSearch/PoseSearchLibrary.h b/Engine/Plugins/Experimental/Animation/PoseSearch/Source/Runtime/Public/PoseSearch/PoseSearchLibrary.h
--- a/Engine/Plugins/Experimental/Animation/PoseSearch/Source/Runtime/Public/PoseSearch/PoseSearchLibrary.h
+++ b/Engine/Plugins/Experimental/Animation/PoseSearch/Source/Runtime/Public/PoseSearch/PoseSearchLibrary.h
@@ -106,7 +106,7 @@ public:
 	* @param Trajectory						Input motion trajectory samples for pose search queries. Expected to be in the space of the SkeletalMeshComponent. This is provided with the CharacterMovementTrajectory Component output.
 	* @param BlendTime						Input time in seconds to blend out to the new pose. Uses either inertial blending, requiring an Inertialization node after this node, or the internal blend stack, if MaxActiveBlends is greater than zero.
 	* @param MaxActiveBlends				Input number of max active animation segments being blended together in the blend stack. If MaxActiveBlends is zero then the blend stack is disabled.
-	* @param PoseJumpThresholdTime			Input don't jump to poses of the same segment that are less than this many seconds away.
+	* @param PoseJumpThresholdTime			Input don't jump to poses of the same segment that are within the interval this many seconds away from the continuing pose.
 	* @param PoseReselectHistory			Input prevent re-selection of poses that have been selected previously within this much time (in seconds) in the past. This is across all animation segments that have been selected within this time range.
 	* @param SearchThrottleTime				Input minimum amount of time to wait between searching for a new pose segment. It allows users to define how often the system searches, default for locomotion is searching every update, but you may only want to search once for other situations, like jump.
 	* @param PlayRate						Input effective range of play rate that can be applied to the animations to account for discrepancies in estimated velocity between the movement modeland the animation.
@@ -123,7 +123,7 @@ public:
 		float TrajectorySpeedMultiplier,
 		float BlendTime,
 		int32 MaxActiveBlends,
-		float PoseJumpThresholdTime,
+		const FFloatInterval& PoseJumpThresholdTime,
 		float PoseReselectHistory,
 		float SearchThrottleTime,
 		const FFloatInterval& PlayRate,
