motion matching - preallocating memory in FBestPoseCandidates constructor to avoid reallocations and rehashing

diff --git a/Engine/Plugins/Experimental/Animation/PoseSearch/Source/Runtime/Public/PoseSearch/PoseSearchContext.h b/Engine/Plugins/Experimental/Animation/PoseSearch/Source/Runtime/Public/PoseSearch/PoseSearchContext.h
--- a/Engine/Plugins/Experimental/Animation/PoseSearch/Source/Runtime/Public/PoseSearch/PoseSearchContext.h
+++ b/Engine/Plugins/Experimental/Animation/PoseSearch/Source/Runtime/Public/PoseSearch/PoseSearchContext.h
@@ -231,6 +231,14 @@ public:
 
 	struct FBestPoseCandidates
 	{
+		FBestPoseCandidates()
+		{
+			// preallocating memory to avoid multiple reallocations / rehashing
+			SearchedDatabases.Empty(8);
+			PoseCandidateHeap.Reserve(MaxPoseCandidates);
+			PoseIdxToFlags.Empty(MaxPoseCandidates);
+		}
+
 		void Add(const UPoseSearchDatabase* Database)
 		{
 			SearchedDatabases.Add(Database);
