motion matching - removing no longer supported EInputQueryPose::UseInterpolatedContinuingPose

diff --git a/Engine/Plugins/Experimental/Animation/PoseSearch/Source/Runtime/Public/PoseSearch/PoseSearchFeatureChannel.h b/Engine/Plugins/Experimental/Animation/PoseSearch/Source/Runtime/Public/PoseSearch/PoseSearchFeatureChannel.h
--- a/Engine/Plugins/Experimental/Animation/PoseSearch/Source/Runtime/Public/PoseSearch/PoseSearchFeatureChannel.h
+++ b/Engine/Plugins/Experimental/Animation/PoseSearch/Source/Runtime/Public/PoseSearch/PoseSearchFeatureChannel.h
@@ -30,9 +30,6 @@ enum class EInputQueryPose : uint8
 
 	// If available reuse continuing pose from the database to compose the query, or else UseCharacterPose.
 	UseContinuingPose,
-
-	// If available reuse and interpolate continuing pose from the database to compose the query, or else UseCharacterPose.
-	UseInterpolatedContinuingPose,
 };
 
 // this enumeration controls the channel sampling time:
