Add playrate to PoseSearchBlueprintResult. Make animation asset a strong reference.
#rnx

diff --git a/Engine/Plugins/Experimental/Animation/PoseSearch/Source/Runtime/Private/MotionMatchingAnimNodeLibrary.cpp b/Engine/Plugins/Experimental/Animation/PoseSearch/Source/Runtime/Private/MotionMatchingAnimNodeLibrary.cpp
--- a/Engine/Plugins/Experimental/Animation/PoseSearch/Source/Runtime/Private/MotionMatchingAnimNodeLibrary.cpp
+++ b/Engine/Plugins/Experimental/Animation/PoseSearch/Source/Runtime/Private/MotionMatchingAnimNodeLibrary.cpp
@@ -30,6 +30,7 @@ void UMotionMatchingAnimNodeLibrary::GetMotionMatchingSearchResult(const FMotion
 					Result.bIsMirrored = SearchIndexAsset->bMirrored;
 					Result.BlendParameters = SearchIndexAsset->BlendParameters;
 					Result.SelectedDatabase = CurrentResultDatabase;
+					Result.WantedPlayRate = MotionMatchingState.WantedPlayRate;
 					bIsResultValid = true;
 				}
 			}
diff --git a/Engine/Plugins/Experimental/Animation/PoseSearch/Source/Runtime/Public/PoseSearch/PoseSearchResult.h b/Engine/Plugins/Experimental/Animation/PoseSearch/Source/Runtime/Public/PoseSearch/PoseSearchResult.h
--- a/Engine/Plugins/Experimental/Animation/PoseSearch/Source/Runtime/Public/PoseSearch/PoseSearchResult.h
+++ b/Engine/Plugins/Experimental/Animation/PoseSearch/Source/Runtime/Public/PoseSearch/PoseSearchResult.h
@@ -69,11 +69,14 @@ struct POSESEARCH_API FPoseSearchBlueprintResult
 public:
 
 	UPROPERTY(Transient, VisibleAnywhere, BlueprintReadOnly, Category=State)
-	TWeakObjectPtr<const UAnimationAsset> SelectedAnimation = nullptr;
+	TObjectPtr<const UAnimationAsset> SelectedAnimation = nullptr;
 	
 	UPROPERTY(Transient, VisibleAnywhere, BlueprintReadOnly, Category=State)
 	float SelectedTime = 0.f;
 	
+	UPROPERTY(Transient, VisibleAnywhere, BlueprintReadOnly, Category=State)
+	float WantedPlayRate = 0.f;
+
 	UPROPERTY(Transient, VisibleAnywhere, BlueprintReadOnly, Category=State)
 	bool bLoop = false;
 	
