Blend Stack: Added anim node functions to force blend, and blend to
#rnx

diff --git a/Engine/Plugins/Experimental/Animation/PoseSearch/Source/Runtime/Private/AnimNode_BlendStack.cpp b/Engine/Plugins/Experimental/Animation/PoseSearch/Source/Runtime/Private/AnimNode_BlendStack.cpp
--- a/Engine/Plugins/Experimental/Animation/PoseSearch/Source/Runtime/Private/AnimNode_BlendStack.cpp
+++ b/Engine/Plugins/Experimental/Animation/PoseSearch/Source/Runtime/Private/AnimNode_BlendStack.cpp
@@ -727,7 +727,12 @@ void FAnimNode_BlendStack::UpdateAssetPlayer(const FAnimationUpdateContext& Cont
 			const UAnimationAsset* PlayingAnimationAsset = MainAnimPlayer.GetAnimationAsset();
 			check(PlayingAnimationAsset);
 
-			if (AnimationAsset != PlayingAnimationAsset)
+			if (bForceBlendNextUpdate)
+			{
+				bForceBlendNextUpdate = false;
+				bExecuteBlendTo = true;
+			}
+			else if (AnimationAsset != PlayingAnimationAsset)
 			{
 				bExecuteBlendTo = true;
 			}
@@ -756,6 +761,11 @@ void FAnimNode_BlendStack::UpdateAssetPlayer(const FAnimationUpdateContext& Cont
 	Super::UpdateAssetPlayer(Context);
 }
 
+void FAnimNode_BlendStack::ForceBlendNextUpdate()
+{
+	bForceBlendNextUpdate = true;
+}
+
 void FBlendStack_SampleGraphPoseLink::SetInputPosePlayer(FPoseSearchAnimPlayer& InPlayer)
 {
 	// Because our anim players may get reallocated, or change indices due to push/pops,
diff --git a/Engine/Plugins/Experimental/Animation/PoseSearch/Source/Runtime/Private/BlendStackAnimNodeLibrary.cpp b/Engine/Plugins/Experimental/Animation/PoseSearch/Source/Runtime/Private/BlendStackAnimNodeLibrary.cpp
new file mode 100644
--- /dev/null
+++ b/Engine/Plugins/Experimental/Animation/PoseSearch/Source/Runtime/Private/BlendStackAnimNodeLibrary.cpp
@@ -0,0 +1,67 @@
+// Copyright Epic Games, Inc. All Rights Reserved.
+
+#include "PoseSearch/BlendStackAnimNodeLibrary.h"
+#include "PoseSearch/AnimNode_BlendStack.h"
+
+FBlendStackAnimNodeReference UBlendStackAnimNodeLibrary::ConvertToBlendStackNode(const FAnimNodeReference& Node, EAnimNodeReferenceConversionResult& Result)
+{
+	return FAnimNodeReference::ConvertToType<FBlendStackAnimNodeReference>(Node, Result);
+}
+
+void UBlendStackAnimNodeLibrary::BlendTo(const FAnimUpdateContext& Context, 
+										const FBlendStackAnimNodeReference& BlendStackNode, 
+										UAnimationAsset* AnimationAsset,
+										float AnimationTime,
+										bool bLoop,
+										bool bMirrored,
+										float BlendTime,
+										FVector BlendParameters,
+										float WantedPlayRate)
+{
+	if (FAnimNode_BlendStack* BlendStackNodePtr = BlendStackNode.GetAnimNodePtr<FAnimNode_BlendStack>())
+	{
+		if (const FAnimationUpdateContext* AnimationUpdateContext = Context.GetContext())
+		{
+			if (AnimationAsset == nullptr)
+			{
+				UE_LOG(LogPoseSearch, Warning, TEXT("UBlendStackAnimNodeLibrary::BlendTo called with null animation asset."));
+				return;
+			}
+
+			BlendStackNodePtr->BlendTo(
+				*AnimationUpdateContext, 
+				AnimationAsset, 
+				AnimationTime, 
+				bLoop,
+				bMirrored,
+				BlendStackNodePtr->MirrorDataTable,
+				BlendTime,
+				BlendStackNodePtr->RootBoneBlendTime,
+				BlendStackNodePtr->BlendProfile,
+				BlendStackNodePtr->BlendOption,
+				BlendStackNodePtr->bUseInertialBlend,
+				BlendParameters,
+				WantedPlayRate);
+		}
+		else
+		{
+			UE_LOG(LogPoseSearch, Warning, TEXT("UBlendStackAnimNodeLibrary::BlendTo called with an invalid context."));
+		}
+	}
+	else
+	{
+		UE_LOG(LogPoseSearch, Warning, TEXT("UBlendStackAnimNodeLibrary::BlendTo called with an invalid type."));
+	}
+}
+
+void UBlendStackAnimNodeLibrary::ForceBlendNextUpdate(const FBlendStackAnimNodeReference& BlendStackNode)
+{
+	if (FAnimNode_BlendStack* BlendStackNodePtr = BlendStackNode.GetAnimNodePtr<FAnimNode_BlendStack>())
+	{
+		BlendStackNodePtr->ForceBlendNextUpdate();
+	}
+	else
+	{
+		UE_LOG(LogPoseSearch, Warning, TEXT("UBlendStackAnimNodeLibrary::ForceBlendNextUpdate called with an invalid type."));
+	}
+}
\ No newline at end of file
diff --git a/Engine/Plugins/Experimental/Animation/PoseSearch/Source/Runtime/Public/PoseSearch/AnimNode_BlendStack.h b/Engine/Plugins/Experimental/Animation/PoseSearch/Source/Runtime/Public/PoseSearch/AnimNode_BlendStack.h
--- a/Engine/Plugins/Experimental/Animation/PoseSearch/Source/Runtime/Public/PoseSearch/AnimNode_BlendStack.h
+++ b/Engine/Plugins/Experimental/Animation/PoseSearch/Source/Runtime/Public/PoseSearch/AnimNode_BlendStack.h
@@ -198,5 +198,13 @@ struct POSESEARCH_API FAnimNode_BlendStack : public FAnimNode_BlendStack_Standal
 	// FAnimNode_Base interface
 	virtual void UpdateAssetPlayer(const FAnimationUpdateContext& Context) override;
 	// End of FAnimNode_Base interface
+
+	// Force a blend on the next update, even if the anim sequence has not changed.
+	void ForceBlendNextUpdate();
+
+private:
+
+	bool bForceBlendNextUpdate = false;
+
 };
 
diff --git a/Engine/Plugins/Experimental/Animation/PoseSearch/Source/Runtime/Public/PoseSearch/BlendStackAnimNodeLibrary.h b/Engine/Plugins/Experimental/Animation/PoseSearch/Source/Runtime/Public/PoseSearch/BlendStackAnimNodeLibrary.h
new file mode 100644
--- /dev/null
+++ b/Engine/Plugins/Experimental/Animation/PoseSearch/Source/Runtime/Public/PoseSearch/BlendStackAnimNodeLibrary.h
@@ -0,0 +1,53 @@
+// Copyright Epic Games, Inc. All Rights Reserved.
+
+#pragma once
+
+#include "Animation/AnimExecutionContext.h"
+#include "Animation/AnimNodeReference.h"
+#include "Kismet/BlueprintFunctionLibrary.h"
+#include "BlendStackAnimNodeLibrary.generated.h"
+
+struct FAnimNode_BlendStack;
+
+USTRUCT(BlueprintType)
+struct FBlendStackAnimNodeReference : public FAnimNodeReference
+{
+	GENERATED_BODY()
+
+	typedef FAnimNode_BlendStack FInternalNodeType;
+};
+
+// Exposes operations that can be run on a Blend Stack node via Anim Node Functions such as "On Become Relevant" and "On Update".
+UCLASS(Experimental)
+class POSESEARCH_API UBlendStackAnimNodeLibrary : public UBlueprintFunctionLibrary
+{
+	GENERATED_BODY()
+
+public:
+	/** Get a blend stack node context from an anim node context */
+	UFUNCTION(BlueprintCallable, Category = "Animation|BlendStack", meta = (BlueprintThreadSafe, ExpandEnumAsExecs = "Result"))
+	static FBlendStackAnimNodeReference ConvertToBlendStackNode(const FAnimNodeReference& Node, EAnimNodeReferenceConversionResult& Result);
+
+	/** Get a blend stack node context from an anim node context (pure) */
+	UFUNCTION(BlueprintPure, Category = "Animation|BlendStack", meta = (BlueprintThreadSafe, DisplayName = "Convert to Blend Stack Node"))
+	static void ConvertToBlendStackNodePure(const FAnimNodeReference& Node, FBlendStackAnimNodeReference& BlendStackNode, bool& Result)
+	{
+		EAnimNodeReferenceConversionResult ConversionResult;
+		BlendStackNode = ConvertToBlendStackNode(Node, ConversionResult);
+		Result = (ConversionResult == EAnimNodeReferenceConversionResult::Succeeded);
+	}
+
+	UFUNCTION(BlueprintCallable, Category = "Animation|BlendStack", meta = (BlueprintThreadSafe, DisplayName = "Force Blend On Next Update"))
+	static void ForceBlendNextUpdate(const FBlendStackAnimNodeReference& BlendStackNode);
+
+	UFUNCTION(BlueprintCallable, Category = "Animation|BlendStack", meta = (BlueprintThreadSafe, DisplayName = "Blend To"))
+	static void BlendTo(const FAnimUpdateContext& Context, 
+						const FBlendStackAnimNodeReference& BlendStackNode, 
+						UAnimationAsset* AnimationAsset = nullptr,
+						float AnimationTime = 0.f,
+						bool bLoop = false,
+						bool bMirrored = false,
+						float BlendTime = 0.2f,
+						FVector BlendParameters = FVector::ZeroVector,
+						float WantedPlayRate = 1.f);
+};
\ No newline at end of file
