motion matching - fix for checking for overflowing inputs in FPoseMetadata::FPoseMetadata

diff --git a/Engine/Plugins/Experimental/Animation/PoseSearch/Source/Runtime/Public/PoseSearch/PoseSearchIndex.h b/Engine/Plugins/Experimental/Animation/PoseSearch/Source/Runtime/Public/PoseSearch/PoseSearchIndex.h
--- a/Engine/Plugins/Experimental/Animation/PoseSearch/Source/Runtime/Public/PoseSearch/PoseSearchIndex.h
+++ b/Engine/Plugins/Experimental/Animation/PoseSearch/Source/Runtime/Public/PoseSearch/PoseSearchIndex.h
@@ -36,7 +36,7 @@ public:
 	{
 		// checking for overflowing inputs
 		check(InValueOffset < (1 << ValueOffsetNumBits));
-		check(InValueOffset < (1 << AssetIndexNumBits));
+		check(InAssetIndex < (1 << AssetIndexNumBits));
 	}
 
 	bool IsBlockTransition() const
