Show current playing asset with "a.AnimNode.MotionMatching.DebugDrawInfo". Also show blend stack when "a.AnimNode.MotionMatching.DebugDrawInfoVerbose" is true (it's on by default).

diff --git a/Engine/Plugins/Experimental/Animation/PoseSearch/Source/Runtime/Private/AnimNode_BlendStack.cpp b/Engine/Plugins/Experimental/Animation/PoseSearch/Source/Runtime/Private/AnimNode_BlendStack.cpp
--- a/Engine/Plugins/Experimental/Animation/PoseSearch/Source/Runtime/Private/AnimNode_BlendStack.cpp
+++ b/Engine/Plugins/Experimental/Animation/PoseSearch/Source/Runtime/Private/AnimNode_BlendStack.cpp
@@ -235,6 +235,21 @@ float FPoseSearchAnimPlayer::GetAccumulatedTime() const
 	return 0.f;
 }
 
+float FPoseSearchAnimPlayer::GetPlayRate() const
+{
+	if (SequencePlayerNode.GetSequence())
+	{
+		return SequencePlayerNode.GetPlayRate();
+	}
+
+	if (BlendSpacePlayerNode.GetBlendSpace())
+	{
+		return BlendSpacePlayerNode.GetPlayRate();
+	}
+
+	return 0.f;
+}
+
 FVector FPoseSearchAnimPlayer::GetBlendParameters() const
 {
 	if (BlendSpacePlayerNode.GetBlendSpace())
diff --git a/Engine/Plugins/Experimental/Animation/PoseSearch/Source/Runtime/Private/AnimNode_MotionMatching.cpp b/Engine/Plugins/Experimental/Animation/PoseSearch/Source/Runtime/Private/AnimNode_MotionMatching.cpp
--- a/Engine/Plugins/Experimental/Animation/PoseSearch/Source/Runtime/Private/AnimNode_MotionMatching.cpp
+++ b/Engine/Plugins/Experimental/Animation/PoseSearch/Source/Runtime/Private/AnimNode_MotionMatching.cpp
@@ -22,7 +22,8 @@
 #if ENABLE_ANIM_DEBUG
 static TAutoConsoleVariable<int32> CVarAnimNodeMotionMatchingDrawQuery(TEXT("a.AnimNode.MotionMatching.DebugDrawQuery"), 0, TEXT("Draw input query"));
 static TAutoConsoleVariable<int32> CVarAnimNodeMotionMatchingDrawCurResult(TEXT("a.AnimNode.MotionMatching.DebugDrawCurResult"), 0, TEXT("Draw current result"));
-static TAutoConsoleVariable<int32> CVarAnimNodeMotionMatchingDrawInfo(TEXT("a.AnimNode.MotionMatching.DebugDrawInfo"), 0, TEXT("Draw info like current databases to search"));
+static TAutoConsoleVariable<int32> CVarAnimNodeMotionMatchingDrawInfo(TEXT("a.AnimNode.MotionMatching.DebugDrawInfo"), 0, TEXT("Draw info like current databases and asset"));
+static TAutoConsoleVariable<int32> CVarAnimNodeMotionMatchingDrawInfoVerbose(TEXT("a.AnimNode.MotionMatching.DebugDrawInfoVerbose"), 1, TEXT("Draw additional info like blend stack"));
 static TAutoConsoleVariable<float> CVarAnimNodeMotionMatchingDrawInfoHeight(TEXT("a.AnimNode.MotionMatching.DebugDrawInfoHeight"), 50.f, TEXT("Vertical offset for DebugDrawInfo"));
 #endif
 
@@ -128,12 +129,24 @@ void FAnimNode_MotionMatching::UpdateAssetPlayer(const FAnimationUpdateContext&
 #if ENABLE_ANIM_DEBUG
 	if (CVarAnimNodeMotionMatchingDrawInfo.GetValueOnAnyThread() > 0)
 	{
+		const UPoseSearchDatabase* CurrentDatabase = MotionMatchingState.CurrentSearchResult.Database;
+		const UAnimationAsset* CurrentAnimationAsset = AnimPlayers.IsEmpty() ? nullptr : AnimPlayers.First().GetAnimationAsset();
+
 		FString DebugInfo = FString::Printf(TEXT("bForceInterruptNextUpdate(%d)\n"), bForceInterruptNextUpdate);
-		DebugInfo += FString::Printf(TEXT("Current Database(%s)\n"), *GetNameSafe(MotionMatchingState.CurrentSearchResult.Database.Get()));
-		DebugInfo += FString::Printf(TEXT("Databases to search:\n"));
-		for (const UPoseSearchDatabase* DatabaseToSearch : DatabasesToSearch)
+		DebugInfo += FString::Printf(TEXT("Current Database(%s)\n"), *GetNameSafe(CurrentDatabase));
+		DebugInfo += FString::Printf(TEXT("Current Asset(%s)\n"), *GetNameSafe(CurrentAnimationAsset));
+		if (CVarAnimNodeMotionMatchingDrawInfoVerbose.GetValueOnAnyThread() > 0)
 		{
-			DebugInfo += FString::Printf(TEXT("  %s\n"), *GetNameSafe(DatabaseToSearch));
+			DebugInfo += FString::Printf(TEXT("Databases to search:\n"));
+			for (const UPoseSearchDatabase* DatabaseToSearch : DatabasesToSearch)
+			{
+				DebugInfo += FString::Printf(TEXT("  %s\n"), *GetNameSafe(DatabaseToSearch));
+			}
+			DebugInfo += FString::Printf(TEXT("Blend Stack:\n"));
+			for (const FPoseSearchAnimPlayer& AnimPlayer : AnimPlayers)
+			{
+				DebugInfo += FString::Printf(TEXT("  %s [time:%.2f|playrate:%.2f]\n"), *GetNameSafe(AnimPlayer.GetAnimationAsset()), AnimPlayer.GetAccumulatedTime(), AnimPlayer.GetPlayRate());
+			}
 		}
 		Context.AnimInstanceProxy->AnimDrawDebugInWorldMessage(DebugInfo, FVector::UpVector * CVarAnimNodeMotionMatchingDrawInfoHeight.GetValueOnAnyThread(), FColor::Yellow, 1.f /*TextScale*/);
 	}
diff --git a/Engine/Plugins/Experimental/Animation/PoseSearch/Source/Runtime/Public/PoseSearch/AnimNode_BlendStack.h b/Engine/Plugins/Experimental/Animation/PoseSearch/Source/Runtime/Public/PoseSearch/AnimNode_BlendStack.h
--- a/Engine/Plugins/Experimental/Animation/PoseSearch/Source/Runtime/Public/PoseSearch/AnimNode_BlendStack.h
+++ b/Engine/Plugins/Experimental/Animation/PoseSearch/Source/Runtime/Public/PoseSearch/AnimNode_BlendStack.h
@@ -20,6 +20,7 @@ struct FPoseSearchAnimPlayer
 	void Evaluate_AnyThread(FPoseContext& Output);
 	void Update_AnyThread(const FAnimationUpdateContext& Context);
 	float GetAccumulatedTime() const;
+	float GetPlayRate() const;
 	
 	float GetBlendInPercentage() const;
 	bool GetBlendInWeights(TArray<float>& Weights) const;
