motion matching
- adding motion database indexing stress test via cvars
- fix for non deterministic serialization of FSearchIndex::PCAValuesVectorToPoseIndexes
- bug fix for Error in Test Deterministic Cooking Windows Client: UPoseSearchFeatureChannel_Position::IndexAsset()

diff --git a/Engine/Plugins/Experimental/Animation/PoseSearch/Source/Runtime/Private/PoseSearchDerivedData.cpp b/Engine/Plugins/Experimental/Animation/PoseSearch/Source/Runtime/Private/PoseSearchDerivedData.cpp
--- a/Engine/Plugins/Experimental/Animation/PoseSearch/Source/Runtime/Private/PoseSearchDerivedData.cpp
+++ b/Engine/Plugins/Experimental/Animation/PoseSearch/Source/Runtime/Private/PoseSearchDerivedData.cpp
@@ -29,6 +29,11 @@
 #include "UObject/NoExportTypes.h"
 #include "UObject/PackageReload.h"
 
+#if ENABLE_ANIM_DEBUG
+static TAutoConsoleVariable<bool> CVarMotionMatchInvalidateIndexingCache(TEXT("a.MotionMatch.InvalidateIndexingCache"), false, TEXT("MotionMatch Invalidate Indexing Cache"));
+static TAutoConsoleVariable<bool> CVarMotionMatchForceIndexing(TEXT("a.MotionMatch.ForceIndexing"), false, TEXT("MotionMatch Force Indexing"));
+#endif
+
 namespace UE::PoseSearch
 {
 static const UE::DerivedData::FValueId Id(UE::DerivedData::FValueId::FromName("Data"));
@@ -636,17 +641,15 @@ struct FPoseSearchDatabaseAsyncCacheTask
 		Failed		// the task has ended unsuccessfully
 	};
 
-	// these methods MUST be protected by FPoseSearchDatabaseAsyncCacheTask::Mutex! and to make sure we pass the mutex as input param
-	FPoseSearchDatabaseAsyncCacheTask(UPoseSearchDatabase* InDatabase, FCriticalSection& OuterMutex, bool bPerformConditionalPostLoadIfRequired);
-	void StartNewRequestIfNeeded(FCriticalSection& OuterMutex, bool bPerformConditionalPostLoadIfRequired);
-	bool CancelIfDependsOn(const UObject* Object, FCriticalSection& OuterMutex);
+	FPoseSearchDatabaseAsyncCacheTask(UPoseSearchDatabase* InDatabase, bool bPerformConditionalPostLoadIfRequired);
+	void StartNewRequestIfNeeded(bool bPerformConditionalPostLoadIfRequired);
+	bool CancelIfDependsOn(const UObject* Object);
 	void Update(FCriticalSection& OuterMutex);
 	void Wait(FCriticalSection& OuterMutex);
-	void Cancel(FCriticalSection& OuterMutex);
-	bool Poll(FCriticalSection& OuterMutex) const;
-	void AddReferencedObjects(FReferenceCollector& Collector, FCriticalSection& OuterMutex);
-	bool ContainsDatabase(const UPoseSearchDatabase* OtherDatabase, FCriticalSection& OuterMutex) const;
-	bool IsValid(FCriticalSection& OuterMutex) const;
+	void Cancel();
+	bool Poll() const;
+	void AddReferencedObjects(FReferenceCollector& Collector);
+	bool IsValid() const;
 	const FIoHash& GetDerivedDataKey() const { return DerivedDataKey; }
 	const UPoseSearchDatabase* GetDatabase() const { return Database.Get(); }
 
@@ -664,6 +667,11 @@ private:
 
 	TWeakObjectPtr<UPoseSearchDatabase> Database;
 	FSearchIndex SearchIndex;
+
+#if ENABLE_ANIM_DEBUG
+	FSearchIndex SearchIndexCompare;
+#endif //ENABLE_ANIM_DEBUG
+
 	UE::DerivedData::FRequestOwner Owner;
 	FIoHash DerivedDataKey = FIoHash::Zero;
 	TSet<TWeakObjectPtr<const UObject>> DatabaseDependencies;
@@ -674,7 +682,7 @@ private:
 
 class FPoseSearchDatabaseAsyncCacheTasks : public TArray<TUniquePtr<FPoseSearchDatabaseAsyncCacheTask>> {};
 
-FPoseSearchDatabaseAsyncCacheTask::FPoseSearchDatabaseAsyncCacheTask(UPoseSearchDatabase* InDatabase, FCriticalSection& OuterMutex, bool bPerformConditionalPostLoadIfRequired)
+FPoseSearchDatabaseAsyncCacheTask::FPoseSearchDatabaseAsyncCacheTask(UPoseSearchDatabase* InDatabase, bool bPerformConditionalPostLoadIfRequired)
 	: Database(InDatabase)
 	, Owner(UE::DerivedData::EPriority::Normal)
 	, DerivedDataKey(FIoHash::Zero)
@@ -682,7 +690,7 @@ FPoseSearchDatabaseAsyncCacheTask::FPoseSearchDatabaseAsyncCacheTask(UPoseSearch
 	if (IsInGameThread())
 	{
 		// it is safe to compose DDC key only on the game thread, since assets can modified in this thread execution
-		StartNewRequestIfNeeded(OuterMutex, bPerformConditionalPostLoadIfRequired);
+		StartNewRequestIfNeeded(bPerformConditionalPostLoadIfRequired);
 	}
 	else
 	{
@@ -692,22 +700,28 @@ FPoseSearchDatabaseAsyncCacheTask::FPoseSearchDatabaseAsyncCacheTask(UPoseSearch
 
 FPoseSearchDatabaseAsyncCacheTask::~FPoseSearchDatabaseAsyncCacheTask()
 {
+	// Owner.Cancel must be performed before SearchIndex.Reset() in case any task is flying (launched by Owner.LaunchTask)
+	Owner.Cancel();
+
 	Database = nullptr;
+
 	SearchIndex.Reset();
-	Owner.Cancel();
+#if ENABLE_ANIM_DEBUG
+	SearchIndexCompare.Reset();
+#endif //ENABLE_ANIM_DEBUG
+
 	DerivedDataKey = FIoHash::Zero;
 	DatabaseDependencies.Reset();
 }
 
-void FPoseSearchDatabaseAsyncCacheTask::StartNewRequestIfNeeded(FCriticalSection& OuterMutex, bool bPerformConditionalPostLoadIfRequired)
+void FPoseSearchDatabaseAsyncCacheTask::StartNewRequestIfNeeded(bool bPerformConditionalPostLoadIfRequired)
 {
 	using namespace UE::DerivedData;
 
 	check(IsInGameThread());
 
-	FScopeLock Lock(&OuterMutex);
-
 	// making sure there are no active requests
+	// Owner.Cancel must be performed before SearchIndex.Reset() in case any task is flying (launched by Owner.LaunchTask)
 	Owner.Cancel();
 
 	// composing the key
@@ -751,29 +765,31 @@ void FPoseSearchDatabaseAsyncCacheTask::StartNewRequestIfNeeded(FCriticalSection
 }
 
 // it cancels and waits for the task to be done and reset the local SearchIndex. SetState to Cancelled
-void FPoseSearchDatabaseAsyncCacheTask::Cancel(FCriticalSection& OuterMutex)
+void FPoseSearchDatabaseAsyncCacheTask::Cancel()
 {
 	check(IsInGameThread());
 
-	FScopeLock Lock(&OuterMutex);
-
+	// Owner.Cancel must be performed before SearchIndex.Reset() in case any task is flying (launched by Owner.LaunchTask)
 	Owner.Cancel();
+
 	SearchIndex.Reset();
+#if ENABLE_ANIM_DEBUG
+	SearchIndexCompare.Reset();
+#endif //ENABLE_ANIM_DEBUG
+
 	DerivedDataKey = FIoHash::Zero;
 	SetState(EState::Cancelled);
 }
 
-bool FPoseSearchDatabaseAsyncCacheTask::CancelIfDependsOn(const UObject* Object, FCriticalSection& OuterMutex)
+bool FPoseSearchDatabaseAsyncCacheTask::CancelIfDependsOn(const UObject* Object)
 {
 	if (Object)
 	{
-		FScopeLock Lock(&OuterMutex);
-
 		// DatabaseDependencies is updated only in StartNewRequestIfNeeded when there are no active requests, so it's thread safe to access it 
 		if (DatabaseDependencies.Contains(Object))
 		{
 			UE_LOG(LogPoseSearch, Log, TEXT("%s - %s Cancelled because of %s"), *LexToString(DerivedDataKey), *Database->GetName(), *Object->GetName());
-			Cancel(OuterMutex);
+			Cancel();
 			return true;
 		}
 	}
@@ -784,16 +800,14 @@ void FPoseSearchDatabaseAsyncCacheTask::Update(FCriticalSection& OuterMutex)
 {
 	check(IsInGameThread());
 
-	FScopeLock Lock(&OuterMutex);
-
 	check(GetState() != EState::Cancelled); // otherwise FPoseSearchDatabaseAsyncCacheTask should have been already removed
 
 	if (GetState() == EState::Notstarted)
 	{
-		StartNewRequestIfNeeded(OuterMutex, false);
+		StartNewRequestIfNeeded(false);
 	}
 
-	if (GetState() == EState::Prestarted && Poll(OuterMutex))
+	if (GetState() == EState::Prestarted && Poll())
 	{
 		// task is done: we need to update the state form Prestarted to Ended/Failed
 		Wait(OuterMutex);
@@ -811,9 +825,8 @@ void FPoseSearchDatabaseAsyncCacheTask::Wait(FCriticalSection& OuterMutex)
 {
 	check(GetState() == EState::Prestarted);
 
-	Owner.Wait();
-
 	FScopeLock Lock(&OuterMutex);
+	Owner.Wait();
 
 	const bool bFailedIndexing = SearchIndex.IsEmpty();
 	if (!bFailedIndexing)
@@ -831,23 +844,19 @@ void FPoseSearchDatabaseAsyncCacheTask::Wait(FCriticalSection& OuterMutex)
 		SetState(EState::Failed);
 	}
 	SearchIndex.Reset();
+#if ENABLE_ANIM_DEBUG
+	SearchIndexCompare.Reset();
+#endif //ENABLE_ANIM_DEBUG
 }
 
 // true is the task is done executing
-bool FPoseSearchDatabaseAsyncCacheTask::Poll(FCriticalSection& OuterMutex) const
+bool FPoseSearchDatabaseAsyncCacheTask::Poll() const
 {
 	return Owner.Poll();
 }
 
-bool FPoseSearchDatabaseAsyncCacheTask::ContainsDatabase(const UPoseSearchDatabase* OtherDatabase, FCriticalSection& OuterMutex) const
-{
-	FScopeLock Lock(&OuterMutex);
-	return Database.Get() == OtherDatabase;
-}
-
-bool FPoseSearchDatabaseAsyncCacheTask::IsValid(FCriticalSection& OuterMutex) const
+bool FPoseSearchDatabaseAsyncCacheTask::IsValid() const
 {
-	FScopeLock Lock(&OuterMutex);
 	return Database.IsValid();
 }
 
@@ -868,6 +877,10 @@ void FPoseSearchDatabaseAsyncCacheTask::OnGetComplete(UE::DerivedData::FCacheGet
 
 		// we found the cached data associated to the PendingDerivedDataKey: we'll deserialized into SearchIndex
 		SearchIndex.Reset();
+#if ENABLE_ANIM_DEBUG
+		SearchIndexCompare.Reset();
+#endif //ENABLE_ANIM_DEBUG
+
 		FSharedBuffer RawData = Response.Record.GetValue(Id).GetData().Decompress();
 		FMemoryReaderView Reader(RawData);
 		Reader << SearchIndex;
@@ -892,13 +905,32 @@ void FPoseSearchDatabaseAsyncCacheTask::OnGetComplete(UE::DerivedData::FCacheGet
 	if (Response.Status == EStatus::Canceled)
 	{
 		SearchIndex.Reset();
+#if ENABLE_ANIM_DEBUG
+		SearchIndexCompare.Reset();
+#endif //ENABLE_ANIM_DEBUG
 		UE_LOG(LogPoseSearch, Log, TEXT("%s - %s BuildIndex Cancelled"), *LexToString(FullIndexKey.Hash), *Database->GetName());
 	}
 	
-	if (Response.Status == EStatus::Error || bCacheCorrupted)
+	bool bForceBuildIndex = false;
+#if ENABLE_ANIM_DEBUG
+	if (CVarMotionMatchForceIndexing.GetValueOnAnyThread())
+	{
+		bForceBuildIndex = true;
+	}
+#endif // ENABLE_ANIM_DEBUG
+
+	if (Response.Status == EStatus::Error || bCacheCorrupted || bForceBuildIndex)
 	{
+#if ENABLE_ANIM_DEBUG
+		const bool bCompareSearchIndex = Response.Status != EStatus::Error && !bCacheCorrupted && bForceBuildIndex;
+		if (bCompareSearchIndex)
+		{
+			SearchIndexCompare = SearchIndex;
+		}
+#endif // ENABLE_ANIM_DEBUG
+
 		// we didn't find the cached data associated to the PendingDerivedDataKey: we'll BuildIndex to update SearchIndex and "Put" the data over the DDC
-		Owner.LaunchTask(TEXT("PoseSearchDatabaseBuild"), [this, FullIndexKey]
+		Owner.LaunchTask(TEXT("PoseSearchDatabaseBuild"), [this, FullIndexKey, bCompareSearchIndex]
 			{
 				COOK_STAT(auto Timer = UsageStats.TimeSyncWork());
 
@@ -907,8 +939,11 @@ void FPoseSearchDatabaseAsyncCacheTask::OnGetComplete(UE::DerivedData::FCacheGet
 				IndexBaseDatabases.Add(Database.Get()); // the first one is always this Database
 				if (!IndexBaseDatabases[0])
 				{
-					UE_LOG(LogPoseSearch, Log, TEXT("%s - BuildIndex Cancelled because associated Database weak pointer has been released."), *LexToString(FullIndexKey.Hash), *IndexBaseDatabases[0]->GetName());
+					UE_LOG(LogPoseSearch, Log, TEXT("%s - BuildIndex Cancelled because associated Database weak pointer has been released."), *LexToString(FullIndexKey.Hash));
 					SearchIndex.Reset();
+#if ENABLE_ANIM_DEBUG
+					SearchIndexCompare.Reset();
+#endif //ENABLE_ANIM_DEBUG
 					return;
 				}
 
@@ -934,13 +969,16 @@ void FPoseSearchDatabaseAsyncCacheTask::OnGetComplete(UE::DerivedData::FCacheGet
 					{
 						if (IndexBaseIdx == 0)
 						{
-							UE_LOG(LogPoseSearch, Error, TEXT("%s - %s BuildIndex Failed becasue of invalid Schema"), *LexToString(FullIndexKey.Hash), *IndexBaseDatabases[0]->GetName());
+							UE_LOG(LogPoseSearch, Error, TEXT("%s - %s BuildIndex Failed because of invalid Schema"), *LexToString(FullIndexKey.Hash), *IndexBaseDatabases[0]->GetName());
 						}
 						else
 						{
 							UE_LOG(LogPoseSearch, Error, TEXT("%s - %s BuildIndex Failed because dependent database '%s' has an invalid Schema"), *LexToString(FullIndexKey.Hash), *IndexBaseDatabases[0]->GetName(), *IndexBaseDatabase->GetName());
 						}
 						SearchIndex.Reset();
+#if ENABLE_ANIM_DEBUG
+						SearchIndexCompare.Reset();
+#endif //ENABLE_ANIM_DEBUG
 						return;
 					}
 
@@ -956,6 +994,9 @@ void FPoseSearchDatabaseAsyncCacheTask::OnGetComplete(UE::DerivedData::FCacheGet
 									// want to sample a mirrored asset
 									UE_LOG(LogPoseSearch, Error, TEXT("%s - %s BuildIndex Failed because '%s' requires a MirrorDataTable to sample mirrored animation assets"), *LexToString(FullIndexKey.Hash), *IndexBaseDatabases[0]->GetName(), *IndexBaseDatabase->Schema->GetName());
 									SearchIndex.Reset();
+#if ENABLE_ANIM_DEBUG
+									SearchIndexCompare.Reset();
+#endif //ENABLE_ANIM_DEBUG
 									return;
 								}
 							}
@@ -966,6 +1007,9 @@ void FPoseSearchDatabaseAsyncCacheTask::OnGetComplete(UE::DerivedData::FCacheGet
 					{
 						UE_LOG(LogPoseSearch, Log, TEXT("%s - %s BuildIndex Cancelled"), *LexToString(FullIndexKey.Hash), *IndexBaseDatabases[0]->GetName());
 						SearchIndex.Reset();
+#if ENABLE_ANIM_DEBUG
+						SearchIndexCompare.Reset();
+#endif //ENABLE_ANIM_DEBUG
 						return;
 					}
 
@@ -976,6 +1020,9 @@ void FPoseSearchDatabaseAsyncCacheTask::OnGetComplete(UE::DerivedData::FCacheGet
 					{
 						UE_LOG(LogPoseSearch, Log, TEXT("%s - %s BuildIndex Cancelled"), *LexToString(FullIndexKey.Hash), *IndexBaseDatabases[0]->GetName());
 						SearchIndex.Reset();
+#if ENABLE_ANIM_DEBUG
+						SearchIndexCompare.Reset();
+#endif //ENABLE_ANIM_DEBUG
 						return;
 					}
 
@@ -984,6 +1031,9 @@ void FPoseSearchDatabaseAsyncCacheTask::OnGetComplete(UE::DerivedData::FCacheGet
 					{
 						UE_LOG(LogPoseSearch, Log, TEXT("%s - %s BuildIndex Cancelled"), *LexToString(FullIndexKey.Hash), *IndexBaseDatabases[0]->GetName());
 						SearchIndex.Reset();
+#if ENABLE_ANIM_DEBUG
+						SearchIndexCompare.Reset();
+#endif //ENABLE_ANIM_DEBUG
 						return;
 					}
 				}
@@ -1000,6 +1050,9 @@ void FPoseSearchDatabaseAsyncCacheTask::OnGetComplete(UE::DerivedData::FCacheGet
 				{
 					UE_LOG(LogPoseSearch, Log, TEXT("%s - %s BuildIndex Cancelled"), *LexToString(FullIndexKey.Hash), *IndexBaseDatabases[0]->GetName());
 					SearchIndex.Reset();
+#if ENABLE_ANIM_DEBUG
+					SearchIndexCompare.Reset();
+#endif //ENABLE_ANIM_DEBUG
 					return;
 				}
 
@@ -1008,6 +1061,9 @@ void FPoseSearchDatabaseAsyncCacheTask::OnGetComplete(UE::DerivedData::FCacheGet
 				{
 					UE_LOG(LogPoseSearch, Log, TEXT("%s - %s BuildIndex Cancelled"), *LexToString(FullIndexKey.Hash), *IndexBaseDatabases[0]->GetName());
 					SearchIndex.Reset();
+#if ENABLE_ANIM_DEBUG
+					SearchIndexCompare.Reset();
+#endif //ENABLE_ANIM_DEBUG
 					return;
 				}
 
@@ -1016,6 +1072,9 @@ void FPoseSearchDatabaseAsyncCacheTask::OnGetComplete(UE::DerivedData::FCacheGet
 				{
 					UE_LOG(LogPoseSearch, Log, TEXT("%s - %s BuildIndex Cancelled"), *LexToString(FullIndexKey.Hash), *IndexBaseDatabases[0]->GetName());
 					SearchIndex.Reset();
+#if ENABLE_ANIM_DEBUG
+					SearchIndexCompare.Reset();
+#endif //ENABLE_ANIM_DEBUG
 					return;
 				}
 
@@ -1024,6 +1083,9 @@ void FPoseSearchDatabaseAsyncCacheTask::OnGetComplete(UE::DerivedData::FCacheGet
 				{
 					UE_LOG(LogPoseSearch, Log, TEXT("%s - %s BuildIndex Cancelled"), *LexToString(FullIndexKey.Hash), *IndexBaseDatabases[0]->GetName());
 					SearchIndex.Reset();
+#if ENABLE_ANIM_DEBUG
+					SearchIndexCompare.Reset();
+#endif //ENABLE_ANIM_DEBUG
 					return;
 				}
 
@@ -1035,6 +1097,16 @@ void FPoseSearchDatabaseAsyncCacheTask::OnGetComplete(UE::DerivedData::FCacheGet
 
 				UE_LOG(LogPoseSearch, Log, TEXT("%s - %s BuildIndex Succeeded"), *LexToString(FullIndexKey.Hash), *IndexBaseDatabases[0]->GetName());
 
+#if ENABLE_ANIM_DEBUG
+				if (bCompareSearchIndex)
+				{
+					if (!SearchIndexCompare.Compare(SearchIndex))
+					{
+						UE_LOG(LogPoseSearch, Warning, TEXT("%s - %s BuildIndex mismatch with DDC Index"), *LexToString(FullIndexKey.Hash), *IndexBaseDatabases[0]->GetName());
+					}
+				}
+#endif // ENABLE_ANIM_DEBUG
+
 				// putting SearchIndex to DDC
 				TArray<uint8> RawBytes;
 				FMemoryWriter Writer(RawBytes);
@@ -1044,6 +1116,7 @@ void FPoseSearchDatabaseAsyncCacheTask::OnGetComplete(UE::DerivedData::FCacheGet
 
 				FCacheRecordBuilder Builder(FullIndexKey);
 				Builder.AddValue(Id, RawData);
+
 				GetCache().Put({ { { IndexBaseDatabases[0]->GetPathName() }, Builder.Build() } }, Owner, [IndexBaseDatabases, FullIndexKey](FCachePutResponse&& Response)
 					{
 						if (Response.Status == EStatus::Error)
@@ -1057,10 +1130,8 @@ void FPoseSearchDatabaseAsyncCacheTask::OnGetComplete(UE::DerivedData::FCacheGet
 	}
 }
 
-void FPoseSearchDatabaseAsyncCacheTask::AddReferencedObjects(FReferenceCollector& Collector, FCriticalSection& OuterMutex)
+void FPoseSearchDatabaseAsyncCacheTask::AddReferencedObjects(FReferenceCollector& Collector)
 {
-	FScopeLock Lock(&OuterMutex);
-
 	const EState State = GetState();
 	if (State != EState::Ended && State != EState::Failed)
 	{
@@ -1117,7 +1188,7 @@ void FAsyncPoseSearchDatabasesManagement::OnObjectModified(UObject* Object)
 	// iterating backwards because of the possible RemoveAtSwap
 	for (int32 TaskIndex = Tasks.Num() - 1; TaskIndex >= 0; --TaskIndex)
 	{
-		if (Tasks[TaskIndex]->CancelIfDependsOn(Object, Mutex))
+		if (Tasks[TaskIndex]->CancelIfDependsOn(Object))
 		{
 			Tasks.RemoveAtSwap(TaskIndex, 1, false);
 		}
@@ -1164,10 +1235,21 @@ void FAsyncPoseSearchDatabasesManagement::Tick(float DeltaTime)
 
 	check(IsInGameThread());
 
+#if ENABLE_ANIM_DEBUG
+	if (CVarMotionMatchInvalidateIndexingCache.GetValueOnAnyThread())
+	{
+		for (TUniquePtr<FPoseSearchDatabaseAsyncCacheTask>& TaskPtr : Tasks)
+		{
+			UE_LOG(LogPoseSearch, Log, TEXT("%s - %s Cancelled because of CVarMotionMatchInvalidateIndexingCache"), *LexToString(TaskPtr->GetDerivedDataKey()), *TaskPtr->GetDatabase()->GetName());
+		}
+		Tasks.Reset();
+	}
+#endif // ENABLE_ANIM_DEBUG
+	
 	// iterating backwards because of the possible RemoveAtSwap 
 	for (int32 TaskIndex = Tasks.Num() - 1; TaskIndex >= 0; --TaskIndex)
 	{
-		if (!Tasks[TaskIndex]->IsValid(Mutex))
+		if (!Tasks[TaskIndex]->IsValid())
 		{
 			Tasks.RemoveAtSwap(TaskIndex, 1, false);
 		}
@@ -1194,7 +1276,7 @@ void FAsyncPoseSearchDatabasesManagement::AddReferencedObjects(FReferenceCollect
 
 	for (TUniquePtr<FPoseSearchDatabaseAsyncCacheTask>& TaskPtr : Tasks)
 	{
-		TaskPtr->AddReferencedObjects(Collector, Mutex);
+		TaskPtr->AddReferencedObjects(Collector);
 	}
 }
 
@@ -1225,7 +1307,7 @@ bool FAsyncPoseSearchDatabasesManagement::RequestAsyncBuildIndex(const UPoseSear
 	FPoseSearchDatabaseAsyncCacheTask* Task = nullptr;
 	for (TUniquePtr<FPoseSearchDatabaseAsyncCacheTask>& TaskPtr : This.Tasks)
 	{
-		if (TaskPtr->ContainsDatabase(Database, Mutex))
+		if (TaskPtr->GetDatabase() == Database)
 		{
 			Task = TaskPtr.Get();
 
@@ -1233,9 +1315,9 @@ bool FAsyncPoseSearchDatabasesManagement::RequestAsyncBuildIndex(const UPoseSear
 			{
 				if (Task->GetState() == FPoseSearchDatabaseAsyncCacheTask::EState::Prestarted)
 				{
-					Task->Cancel(Mutex);
+					Task->Cancel();
 				}
-				Task->StartNewRequestIfNeeded(Mutex, bWaitForCompletion);
+				Task->StartNewRequestIfNeeded(bWaitForCompletion);
 			}
 			break;
 		}
@@ -1244,7 +1326,7 @@ bool FAsyncPoseSearchDatabasesManagement::RequestAsyncBuildIndex(const UPoseSear
 	if (!Task)
 	{
 		// we didn't find the Task, so we Emplace a new one
-		This.Tasks.Emplace(MakeUnique<FPoseSearchDatabaseAsyncCacheTask>(const_cast<UPoseSearchDatabase*>(Database), Mutex, bWaitForCompletion));
+		This.Tasks.Emplace(MakeUnique<FPoseSearchDatabaseAsyncCacheTask>(const_cast<UPoseSearchDatabase*>(Database), bWaitForCompletion));
 		Task = This.Tasks.Last().Get();
 	}
 
diff --git a/Engine/Plugins/Experimental/Animation/PoseSearch/Source/Runtime/Private/PoseSearchDerivedDataKey.cpp b/Engine/Plugins/Experimental/Animation/PoseSearch/Source/Runtime/Private/PoseSearchDerivedDataKey.cpp
--- a/Engine/Plugins/Experimental/Animation/PoseSearch/Source/Runtime/Private/PoseSearchDerivedDataKey.cpp
+++ b/Engine/Plugins/Experimental/Animation/PoseSearch/Source/Runtime/Private/PoseSearchDerivedDataKey.cpp
@@ -40,7 +40,7 @@ FKeyBuilder::FKeyBuilder(const UObject* Object, bool bUseDataVer, bool bPerformC
 	if (bUseDataVer)
 	{
 		// used to invalidate the key without having to change POSESEARCHDB_DERIVEDDATA_VER all the times
-		int32 POSESEARCHDB_DERIVEDDATA_VER_SMALL = 179;
+		int32 POSESEARCHDB_DERIVEDDATA_VER_SMALL = 187;
 		FGuid VersionGuid = FDevSystemGuids::GetSystemGuid(FDevSystemGuids::Get().POSESEARCHDB_DERIVEDDATA_VER);
 
 		*this << VersionGuid;
diff --git a/Engine/Plugins/Experimental/Animation/PoseSearch/Source/Runtime/Private/PoseSearchIndex.cpp b/Engine/Plugins/Experimental/Animation/PoseSearch/Source/Runtime/Private/PoseSearchIndex.cpp
--- a/Engine/Plugins/Experimental/Animation/PoseSearch/Source/Runtime/Private/PoseSearchIndex.cpp
+++ b/Engine/Plugins/Experimental/Animation/PoseSearch/Source/Runtime/Private/PoseSearchIndex.cpp
@@ -323,6 +323,18 @@ FArchive& operator<<(FArchive& Ar, FPoseMetadata& Metadata)
 
 //////////////////////////////////////////////////////////////////////////
 // FSearchIndexAsset
+bool FSearchIndexAsset::operator==(const FSearchIndexAsset& Other) const
+{
+	return
+		SourceAssetIdx == Other.SourceAssetIdx &&
+		bMirrored == Other.bMirrored &&
+		PermutationIdx == Other.PermutationIdx &&
+		BlendParameters == Other.BlendParameters &&
+		FirstPoseIdx == Other.FirstPoseIdx &&
+		FirstSampleIdx == Other.FirstSampleIdx &&
+		LastSampleIdx == Other.LastSampleIdx;
+}
+
 FArchive& operator<<(FArchive& Ar, FSearchIndexAsset& IndexAsset)
 {
 	Ar << IndexAsset.SourceAssetIdx;
@@ -337,6 +349,15 @@ FArchive& operator<<(FArchive& Ar, FSearchIndexAsset& IndexAsset)
 
 //////////////////////////////////////////////////////////////////////////
 // FSearchStats
+bool FSearchStats::operator==(const FSearchStats& Other) const
+{
+	return
+		AverageSpeed == Other.AverageSpeed &&
+		MaxSpeed == Other.MaxSpeed &&
+		AverageAcceleration == Other.AverageAcceleration &&
+		MaxAcceleration == Other.MaxAcceleration;
+}
+
 FArchive& operator<<(FArchive& Ar, FSearchStats& Stats)
 {
 	Ar << Stats.AverageSpeed;
@@ -402,6 +423,45 @@ void FSearchIndexBase::AllocateData(int32 DataCardinality, int32 NumPoses)
 	PoseMetadata.SetNumZeroed(NumPoses);
 }
 
+#if ENABLE_ANIM_DEBUG
+bool FSearchIndexBase::Compare(const FSearchIndexBase& Other) const
+{
+	bool bResult = true;
+
+	if (Values != Other.Values)
+	{
+		UE_LOG(LogPoseSearch, Warning, TEXT("FSearchIndexBase::Compare - Values mismatch"));
+		bResult = false;
+	}
+
+	if (bAnyBlockTransition != Other.bAnyBlockTransition)
+	{
+		UE_LOG(LogPoseSearch, Warning, TEXT("FSearchIndexBase::Compare - bAnyBlockTransition mismatch"));
+		bResult = false;
+	}
+
+	if (Assets != Other.Assets)
+	{
+		UE_LOG(LogPoseSearch, Warning, TEXT("FSearchIndexBase::Compare - Assets mismatch"));
+		bResult = false;
+	}
+	 
+	if (MinCostAddend != Other.MinCostAddend)
+	{
+		UE_LOG(LogPoseSearch, Warning, TEXT("FSearchIndexBase::Compare - MinCostAddend mismatch"));
+		bResult = false;
+	}
+
+	if (Stats != Other.Stats)
+	{
+		UE_LOG(LogPoseSearch, Warning, TEXT("FSearchIndexBase::Compare - Stats mismatch"));
+		bResult = false;
+	}
+
+	return true;
+}
+#endif // ENABLE_ANIM_DEBUG
+
 FArchive& operator<<(FArchive& Ar, FSearchIndexBase& Index)
 {
 	Ar << Index.Values;
@@ -549,24 +609,21 @@ void FSearchIndex::PruneDuplicatePCAValues(float SimilarityThreshold, int32 Numb
 					PoseIndexes.Add(PoseIdx);
 				}
 
-				int32 MoreThanOne = 0;
-				int32 MaxDuplicates = 1;
+				// sorting PCAValuesVectorToPoseIndexesMap keys to create a deterministic FSparsePoseMultiMap later on
+				// we're not using TSortedMap for performance reasons, because PCAValuesVectorToPoseIndexesMap can be quite big
+				TArray<uint32> SortedKeys;
+				SortedKeys.Reserve(PCAValuesVectorToPoseIndexesMap.Num());
 				for (const TPair<uint32, TArray<uint32>>& Pair : PCAValuesVectorToPoseIndexesMap)
 				{
-					if (Pair.Value.Num() > 1)
-					{
-						++MoreThanOne;
-						MaxDuplicates = FMath::Max(MaxDuplicates, Pair.Value.Num());
-					}
+					SortedKeys.Add(Pair.Key);
 				}
-
-				UE_LOG(LogPoseSearch, Log, TEXT("MoreThanOne %d, MaxDuplicates %d"), MoreThanOne, MaxDuplicates);
+				SortedKeys.Sort();
 
 				FSparsePoseMultiMap<uint32> SparsePoseMultiMap(PCAValuesVectorToPoseIndexesMap.Num(), NumPoses - 1);
-				for (const TPair<uint32, TArray<uint32>>& Pair : PCAValuesVectorToPoseIndexesMap)
+				for (const uint32& Key : SortedKeys)
 				{
-					const uint32 PCAValuesVectorIdx = Pair.Key;
-					const TArray<uint32>& PoseIndexes = Pair.Value;
+					const uint32 PCAValuesVectorIdx = Key;
+					const TArray<uint32>& PoseIndexes = PCAValuesVectorToPoseIndexesMap[Key];
 					SparsePoseMultiMap.Insert(PCAValuesVectorIdx, PoseIndexes);
 				}
 
@@ -664,6 +721,59 @@ void FSearchIndex::GetPoseToPCAValuesVectorIndexes(TArray<uint32>& PoseToPCAValu
 	}
 }
 
+#if ENABLE_ANIM_DEBUG
+bool FSearchIndex::Compare(const FSearchIndex& Other) const
+{
+	bool bResult = FSearchIndexBase::Compare(static_cast<const FSearchIndexBase&>(Other));
+
+	if (WeightsSqrt != Other.WeightsSqrt)
+	{
+		UE_LOG(LogPoseSearch, Warning, TEXT("FSearchIndex::Compare - WeightsSqrt mismatch"));
+		bResult = false;
+	}
+	
+	if (PCAValues != Other.PCAValues)
+	{
+		UE_LOG(LogPoseSearch, Warning, TEXT("FSearchIndex::Compare - PCAValues mismatch"));
+		bResult = false;
+	}
+
+	if (PCAValuesVectorToPoseIndexes != Other.PCAValuesVectorToPoseIndexes)
+	{
+		UE_LOG(LogPoseSearch, Warning, TEXT("FSearchIndex::Compare - PCAValuesVectorToPoseIndexes mismatch"));
+		bResult = false;
+	}
+
+	if (PCAProjectionMatrix != Other.PCAProjectionMatrix)
+	{
+		UE_LOG(LogPoseSearch, Warning, TEXT("FSearchIndex::Compare - PCAProjectionMatrix mismatch"));
+		bResult = false;
+	}
+
+	if (Mean != Other.Mean)
+	{
+		UE_LOG(LogPoseSearch, Warning, TEXT("FSearchIndex::Compare - Mean mismatch"));
+		bResult = false;
+	}
+
+	if (PCAExplainedVariance != Other.PCAExplainedVariance)
+	{
+		UE_LOG(LogPoseSearch, Warning, TEXT("FSearchIndex::Compare - PCAExplainedVariance mismatch"));
+		bResult = false;
+	}
+
+	// @todo: implement me
+	//if (KDTree != Other.KDTree)
+	//{
+	//	UE_LOG(LogPoseSearch, Warning, TEXT("FSearchIndex::Compare - KDTree mismatch"));
+	//	bResult = false;
+	//}
+	
+	return bResult;
+}
+#endif // ENABLE_ANIM_DEBUG
+
+
 FArchive& operator<<(FArchive& Ar, FSearchIndex& Index)
 {
 	Ar << static_cast<FSearchIndexBase&>(Index);
diff --git a/Engine/Plugins/Experimental/Animation/PoseSearch/Source/Runtime/Public/PoseSearch/PoseSearchIndex.h b/Engine/Plugins/Experimental/Animation/PoseSearch/Source/Runtime/Public/PoseSearch/PoseSearchIndex.h
--- a/Engine/Plugins/Experimental/Animation/PoseSearch/Source/Runtime/Public/PoseSearch/PoseSearchIndex.h
+++ b/Engine/Plugins/Experimental/Animation/PoseSearch/Source/Runtime/Public/PoseSearch/PoseSearchIndex.h
@@ -104,7 +104,8 @@ struct FSearchIndexAsset
 	int32 LastSampleIdx = INDEX_NONE;
 
 	bool IsPoseInRange(int32 PoseIdx) const { return (PoseIdx >= FirstPoseIdx) && (PoseIdx < FirstPoseIdx + GetNumPoses()); }
-
+	bool operator==(const FSearchIndexAsset& Other) const;
+	
 #if DO_CHECK
 	bool IsInitialized() const 
 	{
@@ -171,6 +172,8 @@ struct FSearchStats
 	float MaxSpeed = 0.f;
 	float AverageAcceleration = 0.f;
 	float MaxAcceleration = 0.f;
+
+	bool operator==(const FSearchStats& Other) const;
 	friend FArchive& operator<<(FArchive& Ar, FSearchStats& Stats);
 };
 
@@ -266,6 +269,15 @@ struct FSparsePoseMultiMap
 		return sizeof(MaxKey) + sizeof(MaxValue) + sizeof(DeltaKeyValue) + DataValues.GetAllocatedSize();
 	}
 
+	bool operator==(const FSparsePoseMultiMap& Other) const
+	{
+		return
+			MaxKey == Other.MaxKey &&
+			MaxValue == Other.MaxValue &&
+			DeltaKeyValue == Other.DeltaKeyValue &&
+			DataValues == Other.DataValues;
+	}
+
 	friend FArchive& operator<<(FArchive& Ar, FSparsePoseMultiMap& SparsePoseMultiMap)
 	{
 		Ar << SparsePoseMultiMap.MaxKey;
@@ -327,6 +339,10 @@ struct FSearchIndexBase
 		return MakeArrayView(&Values[ValueOffset], DataCardinality);
 	}
 
+#if ENABLE_ANIM_DEBUG
+	bool Compare(const FSearchIndexBase& Other) const;
+#endif // ENABLE_ANIM_DEBUG
+
 	friend FArchive& operator<<(FArchive& Ar, FSearchIndexBase& Index);
 };
 
@@ -379,6 +395,10 @@ struct FSearchIndex : public FSearchIndexBase
 	POSESEARCH_API void GetPoseToPCAValuesVectorIndexes(TArray<uint32>& PoseToPCAValuesVectorIndexes) const;
 
 	friend FArchive& operator<<(FArchive& Ar, FSearchIndex& Index);
+
+#if ENABLE_ANIM_DEBUG
+	bool Compare(const FSearchIndex& Other) const;
+#endif // ENABLE_ANIM_DEBUG
 };
 
 } // namespace UE::PoseSearch
