motion matching - fix for PoseSearchDerivedData.cpp(933): error C3481: 'bCompareSearchIndex': lambda capture variable not found

diff --git a/Engine/Plugins/Experimental/Animation/PoseSearch/Source/Runtime/Private/PoseSearchDerivedData.cpp b/Engine/Plugins/Experimental/Animation/PoseSearch/Source/Runtime/Private/PoseSearchDerivedData.cpp
--- a/Engine/Plugins/Experimental/Animation/PoseSearch/Source/Runtime/Private/PoseSearchDerivedData.cpp
+++ b/Engine/Plugins/Experimental/Animation/PoseSearch/Source/Runtime/Private/PoseSearchDerivedData.cpp
@@ -921,8 +921,9 @@ void FPoseSearchDatabaseAsyncCacheTask::OnGetComplete(UE::DerivedData::FCacheGet
 
 	if (Response.Status == EStatus::Error || bCacheCorrupted || bForceBuildIndex)
 	{
+		bool bCompareSearchIndex = false;
 #if ENABLE_ANIM_DEBUG
-		const bool bCompareSearchIndex = Response.Status != EStatus::Error && !bCacheCorrupted && bForceBuildIndex;
+		bCompareSearchIndex = Response.Status != EStatus::Error && !bCacheCorrupted && bForceBuildIndex;
 		if (bCompareSearchIndex)
 		{
 			SearchIndexCompare = SearchIndex;
