motion matching - fix for non deterministic serialization of FKDTree

diff --git a/Engine/Plugins/Experimental/Animation/PoseSearch/Source/Runtime/Private/KDTree.cpp b/Engine/Plugins/Experimental/Animation/PoseSearch/Source/Runtime/Private/KDTree.cpp
--- a/Engine/Plugins/Experimental/Animation/PoseSearch/Source/Runtime/Private/KDTree.cpp
+++ b/Engine/Plugins/Experimental/Animation/PoseSearch/Source/Runtime/Private/KDTree.cpp
@@ -23,6 +23,125 @@ using FKDTreeImplementationBase = nanoflann::KDTreeSingleIndexAdaptor<nanoflann:
 struct FKDTreeImplementation : FKDTreeImplementationBase
 {
 	using FKDTreeImplementationBase::FKDTreeImplementationBase;
+
+	bool operator==(const FKDTreeImplementation& Other) const
+	{
+		if (m_size != Other.m_size)
+		{
+			return false;
+		}
+
+		if (dim != Other.dim)
+		{
+			return false;
+		}
+
+		const uint32 RootBBoxSize = root_bbox.size();
+		if (RootBBoxSize != Other.root_bbox.size())
+		{
+			return false;
+		}
+
+		for (uint32 Index = 0; Index < RootBBoxSize; ++Index)
+		{
+			const Interval& ThisInterval = root_bbox[Index];
+			const Interval& OtherInterval = Other.root_bbox[Index];
+
+			if (ThisInterval.high != OtherInterval.high)
+			{
+				return false;
+			}
+
+			if (ThisInterval.low != OtherInterval.low)
+			{
+				return false;
+			}
+		}
+
+		if (m_leaf_max_size != Other.m_leaf_max_size)
+		{
+			return false;
+		}
+
+		if (vAcc != Other.vAcc)
+		{
+			return false;
+		}
+
+		if (!CompareNodes(root_node, Other.root_node))
+		{
+			return false;
+		}
+
+		return true;
+	}
+
+private:
+	static bool CompareNodes(const NodePtr& NodeA, const NodePtr& NodeB)
+	{
+		const bool bAnyNodeAChild1 = NodeA->child1 != nullptr;
+		const bool bAnyNodeBChild1 = NodeA->child1 != nullptr;
+		if (bAnyNodeAChild1 != bAnyNodeBChild1)
+		{
+			return false;
+		}
+
+		const bool bAnyNodeAChild2 = NodeA->child2 != nullptr;
+		const bool bAnyNodeBChild2 = NodeA->child2 != nullptr;
+		if (bAnyNodeAChild2 != bAnyNodeBChild2)
+		{
+			return false;
+		}
+
+		const bool bIsLeafNode = !bAnyNodeAChild1 && !bAnyNodeAChild2;
+		if (bIsLeafNode)
+		{
+			if (NodeA->node_type.lr.left != NodeB->node_type.lr.left)
+			{
+				return false;
+			}
+
+			if (NodeA->node_type.lr.right != NodeB->node_type.lr.right)
+			{
+				return false;
+			}
+		}
+		else
+		{
+			if (NodeA->node_type.sub.divfeat != NodeB->node_type.sub.divfeat)
+			{
+				return false;
+			}
+
+			if (NodeA->node_type.sub.divhigh != NodeB->node_type.sub.divhigh)
+			{
+				return false;
+			}
+
+			if (NodeA->node_type.sub.divlow != NodeB->node_type.sub.divlow)
+			{
+				return false;
+			}
+		}
+
+		if (bAnyNodeAChild1)
+		{
+			if (!CompareNodes(NodeA->child1, NodeB->child1))
+			{
+				return false;
+			}
+		}
+
+		if (bAnyNodeAChild2)
+		{
+			if (!CompareNodes(NodeA->child2, NodeB->child2))
+			{
+				return false;
+			}
+		}
+
+		return true;
+	}
 };
 #endif
 
@@ -35,7 +154,7 @@ FKDTree::FKDTree(int32 Count, int32 Dim, const float* Data, int32 MaxLeafSize)
 	{
 		Impl = new FKDTreeImplementation(Dim, DataSource, nanoflann::KDTreeSingleIndexAdaptorParams(MaxLeafSize));
 	}
-#endif
+#endif // UE_POSE_SEARCH_USE_NANOFLANN
 }
 
 FKDTree::FKDTree()
@@ -55,7 +174,7 @@ void CopySubTree(FKDTree& KDTree, FKDTreeImplementation::NodePtr& ThisNode, cons
 	check(KDTree.Impl);
 
 	ThisNode = KDTree.Impl->pool.template allocate<FKDTreeImplementation::Node>();
-
+	
 	ThisNode->node_type = OtherNode->node_type;
 
 	if (OtherNode->child1 != nullptr)
@@ -77,7 +196,7 @@ void CopySubTree(FKDTree& KDTree, FKDTreeImplementation::NodePtr& ThisNode, cons
 	}
 }
 
-#endif
+#endif // UE_POSE_SEARCH_USE_NANOFLANN
 
 FKDTree::FKDTree(const FKDTree& Other)
 {
@@ -88,15 +207,15 @@ FKDTree::FKDTree(const FKDTree& Other)
 
 		DataSource = Other.DataSource;
 
-		check(Other.Impl->m_size < UINT_MAX);
+		check(Other.Impl->m_size <= UINT_MAX);
 		Impl->m_size = Other.Impl->m_size;
 
 		if (Impl->m_size > 0)
 		{
 			Impl->dim = Other.Impl->dim;
 
+			check(Other.Impl->root_bbox.size() <= UINT_MAX);
 			const uint32 root_bbox_size = Other.Impl->root_bbox.size();
-			check(root_bbox_size < UINT_MAX);
 			Impl->root_bbox.resize(root_bbox_size);
 
 			for (uint32 i = 0; i < root_bbox_size; ++i)
@@ -104,12 +223,12 @@ FKDTree::FKDTree(const FKDTree& Other)
 				Impl->root_bbox[i] = Other.Impl->root_bbox[i];
 			}
 
-			check(Other.Impl->m_leaf_max_size < UINT_MAX);
+			check(Other.Impl->m_leaf_max_size <= UINT_MAX);
 			const uint32 KDTreeLeafMaxSize = Other.Impl->m_leaf_max_size;
 			Impl->m_leaf_max_size = KDTreeLeafMaxSize;
 
+			check(Other.Impl->vAcc.size() <= UINT_MAX);
 			const uint32 VAccSize = Other.Impl->vAcc.size();
-			check(VAccSize < UINT_MAX);
 			Impl->vAcc.resize(VAccSize);
 			
 			for (uint32 i = 0; i < VAccSize; ++i)
@@ -120,7 +239,7 @@ FKDTree::FKDTree(const FKDTree& Other)
 			CopySubTree(*this, Impl->root_node, Other.Impl->root_node);
 		}
 	}
-#endif
+#endif // UE_POSE_SEARCH_USE_NANOFLANN
 }
 
 FKDTree& FKDTree::operator=(const FKDTree& Other)
@@ -133,12 +252,39 @@ FKDTree& FKDTree::operator=(const FKDTree& Other)
 	return *this;
 }
 
+bool FKDTree::operator==(const FKDTree& Other) const
+{
+	const bool bAnyImpl = Impl != nullptr;
+	const bool bAnyOtherImpl = Other.Impl != nullptr;
+	if (bAnyImpl != bAnyOtherImpl)
+	{
+		return false;
+	}
+
+#if UE_POSE_SEARCH_USE_NANOFLANN
+	if (bAnyImpl && bAnyOtherImpl)
+	{
+		if (*Impl != *Other.Impl)
+		{
+			return false;
+		}
+	}
+#endif // UE_POSE_SEARCH_USE_NANOFLANN
+
+	if (DataSource != Other.DataSource)
+	{
+		return false;
+	}
+
+	return true;
+}
+
 void FKDTree::Reset()
 {
 #if UE_POSE_SEARCH_USE_NANOFLANN
 	delete Impl;
 	Impl = nullptr;
-#endif
+#endif // UE_POSE_SEARCH_USE_NANOFLANN
 	DataSource = FDataSource();
 }
 
@@ -151,6 +297,7 @@ void FKDTree::Construct(int32 Count, int32 Dim, const float* Data, int32 MaxLeaf
 bool FKDTree::FindNeighbors(FKNNResultSet& Result, const float* Query) const
 {
 #if UE_POSE_SEARCH_USE_NANOFLANN
+
 	QUICK_SCOPE_CYCLE_COUNTER(STAT_FKDTree_FindNeighbors);
 
 	check(Query && Impl->root_node);
@@ -159,15 +306,19 @@ bool FKDTree::FindNeighbors(FKNNResultSet& Result, const float* Query) const
 		0.f,		// search for eps-approximate neighbours (default: 0)
 		false);		// only for radius search, require neighbours sorted by
 	return Impl->findNeighbors(Result, Query, SearchParams);
-#else
+
+#else // UE_POSE_SEARCH_USE_NANOFLANN
+
 	checkNoEntry(); // unimplemented
 	return false;
-#endif
+
+#endif // UE_POSE_SEARCH_USE_NANOFLANN
 }
 
 bool FKDTree::FindNeighbors(FRadiusResultSet& Result, const float* Query) const
 {
 #if UE_POSE_SEARCH_USE_NANOFLANN
+
 	QUICK_SCOPE_CYCLE_COUNTER(STAT_FKDTree_FindNeighbors);
 
 	check(Query && Impl->root_node);
@@ -176,10 +327,13 @@ bool FKDTree::FindNeighbors(FRadiusResultSet& Result, const float* Query) const
 		0.f,		// search for eps-approximate neighbours (default: 0)
 		false);		// only for radius search, require neighbours sorted by
 	return Impl->findNeighbors(Result, Query, SearchParams);
-#else
+
+#else // UE_POSE_SEARCH_USE_NANOFLANN
+
 	checkNoEntry(); // unimplemented
 	return false;
-#endif
+
+#endif // UE_POSE_SEARCH_USE_NANOFLANN
 }
 
 SIZE_T FKDTree::GetAllocatedSize() const
@@ -192,12 +346,11 @@ SIZE_T FKDTree::GetAllocatedSize() const
 		AllocatedSize += sizeof(FKDTreeImplementation);
 		AllocatedSize += Impl->usedMemory(*Impl);
 	}
-#endif
+#endif // UE_POSE_SEARCH_USE_NANOFLANN
 
 	return AllocatedSize;
 }
 
-
 #if UE_POSE_SEARCH_USE_NANOFLANN
 FArchive& SerializeSubTree(FArchive& Ar, FKDTree& KDTree, FKDTreeImplementation::NodePtr& KDTreeNode)
 {
@@ -206,16 +359,38 @@ FArchive& SerializeSubTree(FArchive& Ar, FKDTree& KDTree, FKDTreeImplementation:
 	if (Ar.IsLoading())
 	{
 		KDTreeNode = KDTree.Impl->pool.template allocate<FKDTreeImplementation::Node>();
+		// zeroing FKDTreeImplementation::Node memory since it contains a union and doesn't have a constructor 
+		FMemory::Memzero(KDTreeNode, sizeof(FKDTreeImplementation::Node));
 	}
 
-	Ar.Serialize(&KDTreeNode->node_type, sizeof(KDTreeNode->node_type));
+	bool bAnyNodeChild1 = KDTreeNode->child1 != nullptr;
+	bool bAnyNodeChild2 = KDTreeNode->child2 != nullptr;
+	Ar << bAnyNodeChild1;
+	Ar << bAnyNodeChild2;
+
+	const bool bIsLeafNode = !bAnyNodeChild1 && !bAnyNodeChild2;
+	if (bIsLeafNode)
+	{
+		check(KDTreeNode->node_type.lr.left <= UINT_MAX);
+		check(KDTreeNode->node_type.lr.right <= UINT_MAX);
+
+		uint32 OffsetLeft = KDTreeNode->node_type.lr.left;
+		uint32 OffsetRight = KDTreeNode->node_type.lr.right;
 
-	bool child1 = KDTreeNode->child1 != nullptr;
-	bool child2 = KDTreeNode->child2 != nullptr;
-	Ar << child1;
-	Ar << child2;
+		Ar << OffsetLeft;
+		Ar << OffsetRight;
 
-	if (child1)
+		KDTreeNode->node_type.lr.left = OffsetLeft;
+		KDTreeNode->node_type.lr.right = OffsetRight;
+	}
+	else
+	{
+		Ar << KDTreeNode->node_type.sub.divfeat;
+		Ar << KDTreeNode->node_type.sub.divhigh;
+		Ar << KDTreeNode->node_type.sub.divlow;
+	}
+
+	if (bAnyNodeChild1)
 	{
 		SerializeSubTree(Ar, KDTree, KDTreeNode->child1);
 	}
@@ -224,7 +399,7 @@ FArchive& SerializeSubTree(FArchive& Ar, FKDTree& KDTree, FKDTreeImplementation:
 		KDTreeNode->child1 = nullptr;
 	}
 
-	if (child2)
+	if (bAnyNodeChild2)
 	{
 		SerializeSubTree(Ar, KDTree, KDTreeNode->child2);
 	}
@@ -235,17 +410,17 @@ FArchive& SerializeSubTree(FArchive& Ar, FKDTree& KDTree, FKDTreeImplementation:
 	return Ar;
 }
 
-#endif
+#endif // UE_POSE_SEARCH_USE_NANOFLANN
 
 FArchive& Serialize(FArchive& Ar, FKDTree& KDTree, const float* KDTreeData)
 {
 #if UE_POSE_SEARCH_USE_NANOFLANN
+	check(!KDTree.Impl || KDTree.Impl->m_size <= UINT_MAX);
+
 	uint32 KDTreeSize = KDTree.Impl ? KDTree.Impl->m_size : 0;
 
 	Ar << KDTreeSize;
 
-	check(KDTreeSize < UINT_MAX);
-
 	if (KDTreeSize > 0)
 	{
 		if (Ar.IsLoading() && !KDTree.Impl)
@@ -258,7 +433,7 @@ FArchive& Serialize(FArchive& Ar, FKDTree& KDTree, const float* KDTreeData)
 		Ar << KDTree.Impl->dim;
 
 		uint32 root_bbox_size = KDTree.Impl->root_bbox.size();
-		check(KDTree.Impl->root_bbox.size() < UINT_MAX);
+		check(KDTree.Impl->root_bbox.size() <= UINT_MAX);
 		Ar << root_bbox_size;
 
 		if (Ar.IsLoading())
@@ -275,12 +450,12 @@ FArchive& Serialize(FArchive& Ar, FKDTree& KDTree, const float* KDTreeData)
 			Ar.Serialize(&el, sizeof(FKDTreeImplementation::Interval));
 		}
 
-		check(KDTree.Impl->m_leaf_max_size < UINT_MAX);
+		check(KDTree.Impl->m_leaf_max_size <= UINT_MAX);
 		uint32 KDTreeLeafMaxSize = KDTree.Impl->m_leaf_max_size;
 		Ar << KDTreeLeafMaxSize;
 		KDTree.Impl->m_leaf_max_size = KDTreeLeafMaxSize;
 
-		check(KDTree.Impl->vAcc.size() < UINT_MAX);
+		check(KDTree.Impl->vAcc.size() <= UINT_MAX);
 		uint32 VAccSize = KDTree.Impl->vAcc.size();
 		Ar << VAccSize;
 		if (Ar.IsLoading())
@@ -297,7 +472,7 @@ FArchive& Serialize(FArchive& Ar, FKDTree& KDTree, const float* KDTreeData)
 	{
 		KDTree.Reset();
 	}
-#endif
+#endif // UE_POSE_SEARCH_USE_NANOFLANN
 
 	return Ar;
 }
diff --git a/Engine/Plugins/Experimental/Animation/PoseSearch/Source/Runtime/Private/PoseSearchDerivedData.cpp b/Engine/Plugins/Experimental/Animation/PoseSearch/Source/Runtime/Private/PoseSearchDerivedData.cpp
--- a/Engine/Plugins/Experimental/Animation/PoseSearch/Source/Runtime/Private/PoseSearchDerivedData.cpp
+++ b/Engine/Plugins/Experimental/Animation/PoseSearch/Source/Runtime/Private/PoseSearchDerivedData.cpp
@@ -30,9 +30,10 @@
 #include "UObject/PackageReload.h"
 
 #if ENABLE_ANIM_DEBUG
-static TAutoConsoleVariable<bool> CVarMotionMatchInvalidateIndexingCache(TEXT("a.MotionMatch.InvalidateIndexingCache"), false, TEXT("MotionMatch Invalidate Indexing Cache"));
-static TAutoConsoleVariable<bool> CVarMotionMatchForceIndexing(TEXT("a.MotionMatch.ForceIndexing"), false, TEXT("MotionMatch Force Indexing"));
-#endif
+static TAutoConsoleVariable<bool> CVarMotionMatchTestInvalidateIndexingCache(TEXT("a.MotionMatch.TestInvalidateIndexingCache"), false, TEXT("Test Invalidate Motion Matching Indexing Cache"));
+static TAutoConsoleVariable<bool> CVarMotionMatchTestForceIndexing(TEXT("a.MotionMatch.TestForceIndexing"), false, TEXT("Test Motion Matching Force Indexing"));
+static TAutoConsoleVariable<int32> CVarMotionMatchTestDeterministicKDTreeConstructIterations(TEXT("a.MotionMatch.TestDeterministicKDTreeConstructIterations"), 0, TEXT("Test Motion Matching Deterministic KDTree Construct Iterations"));
+#endif // ENABLE_ANIM_DEBUG
 
 namespace UE::PoseSearch
 {
@@ -45,7 +46,7 @@ static FCookStatsManager::FAutoRegisterCallback RegisterCookStats([](FCookStatsM
 	{
 		UsageStats.LogStats(AddStat, TEXT("MotionMatching.Usage"), TEXT(""));
 	});
-#endif
+#endif // ENABLE_COOK_STATS
 
 // helper struct to calculate mean deviations
 struct FMeanDeviationCalculator
@@ -560,6 +561,26 @@ static void PreprocessSearchIndexKDTree(FSearchIndex& SearchIndex, const UPoseSe
 		const int32 NumPCAValuesVectors = SearchIndex.PCAValues.Num() / NumberOfPrincipalComponents;
 		SearchIndex.KDTree.Construct(NumPCAValuesVectors, NumberOfPrincipalComponents, SearchIndex.PCAValues.GetData(), KDTreeMaxLeafSize);
 
+#if ENABLE_ANIM_DEBUG
+		// testing kdtree Construct determinism
+		const int32 TestDeterministicKDTreeConstructIterations = CVarMotionMatchTestDeterministicKDTreeConstructIterations.GetValueOnAnyThread();
+		if (TestDeterministicKDTreeConstructIterations > 0)
+		{
+			for (int32 Iteration = 0; Iteration < TestDeterministicKDTreeConstructIterations; ++Iteration)
+			{
+				// copy PCAValues in a different container to ensure input data has different memory addresses
+				TAlignedArray<float> PCAValuesTest = SearchIndex.PCAValues;
+
+				FKDTree KDTreeTest;
+				KDTreeTest.Construct(NumPCAValuesVectors, NumberOfPrincipalComponents, PCAValuesTest.GetData(), KDTreeMaxLeafSize);
+
+				if (KDTreeTest != SearchIndex.KDTree)
+				{
+					UE_LOG(LogPoseSearch, Warning, TEXT("PreprocessSearchIndexKDTree - FKDTree::Construct is not deterministic"));
+				}
+			}
+		}
+
 		// testing the KDTree is returning the proper searches for all the points in pca space
 		const int32 KDTreeQueryNumNeighbors = Database->KDTreeQueryNumNeighbors;
 
@@ -625,6 +646,8 @@ static void PreprocessSearchIndexKDTree(FSearchIndex& SearchIndex, const UPoseSe
 
 			check(NumberOfFailingPoints == 0);
 		}
+
+#endif // ENABLE_ANIM_DEBUG
 	}
 }
 
@@ -913,7 +936,7 @@ void FPoseSearchDatabaseAsyncCacheTask::OnGetComplete(UE::DerivedData::FCacheGet
 	
 	bool bForceBuildIndex = false;
 #if ENABLE_ANIM_DEBUG
-	if (CVarMotionMatchForceIndexing.GetValueOnAnyThread())
+	if (CVarMotionMatchTestForceIndexing.GetValueOnAnyThread())
 	{
 		bForceBuildIndex = true;
 	}
@@ -1237,11 +1260,11 @@ void FAsyncPoseSearchDatabasesManagement::Tick(float DeltaTime)
 	check(IsInGameThread());
 
 #if ENABLE_ANIM_DEBUG
-	if (CVarMotionMatchInvalidateIndexingCache.GetValueOnAnyThread())
+	if (CVarMotionMatchTestInvalidateIndexingCache.GetValueOnAnyThread())
 	{
 		for (TUniquePtr<FPoseSearchDatabaseAsyncCacheTask>& TaskPtr : Tasks)
 		{
-			UE_LOG(LogPoseSearch, Log, TEXT("%s - %s Cancelled because of CVarMotionMatchInvalidateIndexingCache"), *LexToString(TaskPtr->GetDerivedDataKey()), *TaskPtr->GetDatabase()->GetName());
+			UE_LOG(LogPoseSearch, Log, TEXT("%s - %s Cancelled because of CVarMotionMatchTestInvalidateIndexingCache"), *LexToString(TaskPtr->GetDerivedDataKey()), *TaskPtr->GetDatabase()->GetName());
 		}
 		Tasks.Reset();
 	}
@@ -1297,7 +1320,7 @@ bool FAsyncPoseSearchDatabasesManagement::RequestAsyncBuildIndex(const UPoseSear
 	{
 		Flag |= ERequestAsyncBuildFlag::WaitForCompletion;
 	}
-#endif
+#endif // WITH_ENGINE
 
 	check(EnumHasAnyFlags(Flag, ERequestAsyncBuildFlag::NewRequest | ERequestAsyncBuildFlag::ContinueRequest));
 
diff --git a/Engine/Plugins/Experimental/Animation/PoseSearch/Source/Runtime/Private/PoseSearchDerivedDataKey.cpp b/Engine/Plugins/Experimental/Animation/PoseSearch/Source/Runtime/Private/PoseSearchDerivedDataKey.cpp
--- a/Engine/Plugins/Experimental/Animation/PoseSearch/Source/Runtime/Private/PoseSearchDerivedDataKey.cpp
+++ b/Engine/Plugins/Experimental/Animation/PoseSearch/Source/Runtime/Private/PoseSearchDerivedDataKey.cpp
@@ -40,7 +40,7 @@ FKeyBuilder::FKeyBuilder(const UObject* Object, bool bUseDataVer, bool bPerformC
 	if (bUseDataVer)
 	{
 		// used to invalidate the key without having to change POSESEARCHDB_DERIVEDDATA_VER all the times
-		int32 POSESEARCHDB_DERIVEDDATA_VER_SMALL = 187;
+		int32 POSESEARCHDB_DERIVEDDATA_VER_SMALL = 189;
 		FGuid VersionGuid = FDevSystemGuids::GetSystemGuid(FDevSystemGuids::Get().POSESEARCHDB_DERIVEDDATA_VER);
 
 		*this << VersionGuid;
diff --git a/Engine/Plugins/Experimental/Animation/PoseSearch/Source/Runtime/Private/PoseSearchIndex.cpp b/Engine/Plugins/Experimental/Animation/PoseSearch/Source/Runtime/Private/PoseSearchIndex.cpp
--- a/Engine/Plugins/Experimental/Animation/PoseSearch/Source/Runtime/Private/PoseSearchIndex.cpp
+++ b/Engine/Plugins/Experimental/Animation/PoseSearch/Source/Runtime/Private/PoseSearchIndex.cpp
@@ -762,12 +762,11 @@ bool FSearchIndex::Compare(const FSearchIndex& Other) const
 		bResult = false;
 	}
 
-	// @todo: implement me
-	//if (KDTree != Other.KDTree)
-	//{
-	//	UE_LOG(LogPoseSearch, Warning, TEXT("FSearchIndex::Compare - KDTree mismatch"));
-	//	bResult = false;
-	//}
+	if (KDTree != Other.KDTree)
+	{
+		UE_LOG(LogPoseSearch, Warning, TEXT("FSearchIndex::Compare - KDTree mismatch"));
+		bResult = false;
+	}
 	
 	return bResult;
 }
diff --git a/Engine/Plugins/Experimental/Animation/PoseSearch/Source/Runtime/Public/PoseSearch/KDTree.h b/Engine/Plugins/Experimental/Animation/PoseSearch/Source/Runtime/Public/PoseSearch/KDTree.h
--- a/Engine/Plugins/Experimental/Animation/PoseSearch/Source/Runtime/Public/PoseSearch/KDTree.h
+++ b/Engine/Plugins/Experimental/Animation/PoseSearch/Source/Runtime/Public/PoseSearch/KDTree.h
@@ -35,6 +35,29 @@ struct POSESEARCH_API FKDTree
 		FDataSource(FDataSource&& Other) = default;
 		FDataSource& operator=(const FDataSource& Other) = default;
 		FDataSource& operator=(FDataSource&& Other) = default;
+		bool operator==(const FDataSource& Other) const
+		{
+			if (PointCount != Other.PointCount)
+			{
+				return false;
+			}
+
+			if (PointDim != Other.PointDim)
+			{
+				return false;
+			}
+
+			const int32 DataSize = PointCount * PointDim;
+			for (int32 DataIndex = 0; DataIndex < DataSize; ++DataIndex)
+			{
+				if (Data[DataIndex] != Other.Data[DataIndex])
+				{
+					return false;
+				}
+			}
+			
+			return true;
+		}
 
 		// Must return the number of data points
 		inline size_t kdtree_get_point_count() const { return PointCount; }
@@ -195,7 +218,8 @@ struct POSESEARCH_API FKDTree
 
 	FKDTree& operator=(const FKDTree& Other);
 	FKDTree& operator=(FKDTree&& Other) = delete;
-
+	bool operator==(const FKDTree& Other) const;
+	
 	void Reset();
 	void Construct(int32 Count, int32 Dim, const float* Data, int32 MaxLeafSize = 16);
 	bool FindNeighbors(FKNNResultSet& Result, const float* Query) const;
