motion matching
- adding determinism tests for IndexDatabase and PruneDuplicateValues
- added SearchIndex DDC validation

diff --git a/Engine/Plugins/Experimental/Animation/PoseSearch/Source/Runtime/Private/PoseSearchDerivedData.cpp b/Engine/Plugins/Experimental/Animation/PoseSearch/Source/Runtime/Private/PoseSearchDerivedData.cpp
--- a/Engine/Plugins/Experimental/Animation/PoseSearch/Source/Runtime/Private/PoseSearchDerivedData.cpp
+++ b/Engine/Plugins/Experimental/Animation/PoseSearch/Source/Runtime/Private/PoseSearchDerivedData.cpp
@@ -29,14 +29,41 @@
 #include "UObject/NoExportTypes.h"
 #include "UObject/PackageReload.h"
 
+namespace UE::PoseSearch
+{
 #if ENABLE_ANIM_DEBUG
-static TAutoConsoleVariable<bool> CVarMotionMatchTestInvalidateIndexingCache(TEXT("a.MotionMatch.TestInvalidateIndexingCache"), false, TEXT("Test Invalidate Motion Matching Indexing Cache"));
-static TAutoConsoleVariable<bool> CVarMotionMatchTestForceIndexing(TEXT("a.MotionMatch.TestForceIndexing"), false, TEXT("Test Motion Matching Force Indexing"));
-static TAutoConsoleVariable<int32> CVarMotionMatchTestDeterministicKDTreeConstructIterations(TEXT("a.MotionMatch.TestDeterministicKDTreeConstructIterations"), 0, TEXT("Test Motion Matching Deterministic KDTree Construct Iterations"));
-#endif // ENABLE_ANIM_DEBUG
 
-namespace UE::PoseSearch
+enum EMotionMatchTestFlags
 {
+	// no additional tests will be performed
+	None = 0x0,
+
+	// cache will be invalidated every frame (to stress test DDC cancellation while tasks are flying if !WaitForTaskCompletion)
+	InvalidateCache = 0x1,
+
+	// cache will be invalidated once the flying tasks are ended
+	WaitForTaskCompletion = 0x2,
+
+	// we'll force the database re-indexing and compare the result SearchIndex with the one retrieved via DDC
+	ForceIndexing = 0x4,
+
+	// test KDTree Construct determinism
+	TestKDTreeConstructDeterminism = 0x8,
+
+	// test IndexDatabase determinism
+	TestIndexDatabaseDeterminism = 0x10,
+
+	// test PruneDuplicateValues determinism
+	TestPruneDuplicateValuesDeterminism = 0x20,
+
+	// validating the data we gave to DDC is stored correctly
+	ValidateDDC = 0x40,
+};
+static TAutoConsoleVariable<int32> CVarMotionMatchTestFlags(TEXT("a.MotionMatch.TestFlags"), EMotionMatchTestFlags::None, TEXT("Test Motion Matching using EMotionMatchTestFlags"));
+static TAutoConsoleVariable<int32> CVarMotionMatchTestNumIterations(TEXT("a.MotionMatch.TestNumIterations"), 10, TEXT("Test Motion Matching Num Iterations"));
+static bool AnyTestFlags(int32 Flags) { return (CVarMotionMatchTestFlags.GetValueOnAnyThread() & Flags) != 0; }
+#endif // ENABLE_ANIM_DEBUG
+
 static const UE::DerivedData::FValueId Id(UE::DerivedData::FValueId::FromName("Data"));
 static const UE::DerivedData::FCacheBucket Bucket("PoseSearchDatabase");
 
@@ -563,10 +590,10 @@ static void PreprocessSearchIndexKDTree(FSearchIndex& SearchIndex, const UPoseSe
 
 #if ENABLE_ANIM_DEBUG
 		// testing kdtree Construct determinism
-		const int32 TestDeterministicKDTreeConstructIterations = CVarMotionMatchTestDeterministicKDTreeConstructIterations.GetValueOnAnyThread();
-		if (TestDeterministicKDTreeConstructIterations > 0)
+		if (AnyTestFlags(EMotionMatchTestFlags::TestKDTreeConstructDeterminism))
 		{
-			for (int32 Iteration = 0; Iteration < TestDeterministicKDTreeConstructIterations; ++Iteration)
+			const int32 NumIterations = CVarMotionMatchTestNumIterations.GetValueOnAnyThread();
+			for (int32 Iteration = 0; Iteration < NumIterations; ++Iteration)
 			{
 				// copy PCAValues in a different container to ensure input data has different memory addresses
 				TAlignedArray<float> PCAValuesTest = SearchIndex.PCAValues;
@@ -890,7 +917,7 @@ void FPoseSearchDatabaseAsyncCacheTask::OnGetComplete(UE::DerivedData::FCacheGet
 {
 	using namespace UE::DerivedData;
 
-	const FCacheKey FullIndexKey = Response.Record.GetKey();
+	const FCacheKey& FullIndexKey = Response.Record.GetKey();
 
 	// The database is part of the derived data cache and up to date, skip re-building it.
 	bool bCacheCorrupted = false;
@@ -936,10 +963,7 @@ void FPoseSearchDatabaseAsyncCacheTask::OnGetComplete(UE::DerivedData::FCacheGet
 	
 	bool bForceBuildIndex = false;
 #if ENABLE_ANIM_DEBUG
-	if (CVarMotionMatchTestForceIndexing.GetValueOnAnyThread())
-	{
-		bForceBuildIndex = true;
-	}
+	bForceBuildIndex = AnyTestFlags(EMotionMatchTestFlags::ForceIndexing);
 #endif // ENABLE_ANIM_DEBUG
 
 	if (Response.Status == EStatus::Error || bCacheCorrupted || bForceBuildIndex)
@@ -1060,10 +1084,50 @@ void FPoseSearchDatabaseAsyncCacheTask::OnGetComplete(UE::DerivedData::FCacheGet
 #endif //ENABLE_ANIM_DEBUG
 						return;
 					}
+
+#if ENABLE_ANIM_DEBUG
+					if (AnyTestFlags(EMotionMatchTestFlags::TestIndexDatabaseDeterminism))
+					{
+						const int32 NumIterations = CVarMotionMatchTestNumIterations.GetValueOnAnyThread();
+						for (int32 Iteration = 0; Iteration < NumIterations; ++Iteration)
+						{
+							FSearchIndexBase TestSearchIndexBase = SearchIndexBase;
+							FDatabaseIndexingContext TestDbIndexingContext;
+							if (TestDbIndexingContext.IndexDatabase(TestSearchIndexBase, *IndexBaseDatabase, Owner))
+							{
+								if (!TestSearchIndexBase.Compare(SearchIndexBase))
+								{
+									UE_LOG(LogPoseSearch, Warning, TEXT("OnGetComplete - IndexDatabase is not deterministic"));
+								}
+							}
+						}
+					}
+#endif //ENABLE_ANIM_DEBUG
 				}
 
 				static_cast<FSearchIndexBase&>(SearchIndex) = SearchIndexBases[0];
 				
+#if ENABLE_ANIM_DEBUG
+				// testing PruneDuplicateValues determinism
+				if (AnyTestFlags(EMotionMatchTestFlags::TestPruneDuplicateValuesDeterminism))
+				{
+					const int32 NumIterations = CVarMotionMatchTestNumIterations.GetValueOnAnyThread();
+
+					FSearchIndex TestSearchIndexA = SearchIndex;
+					TestSearchIndexA.PruneDuplicateValues(IndexBaseDatabases[0]->PosePruningSimilarityThreshold, IndexBaseDatabases[0]->Schema->SchemaCardinality);
+					for (int32 Iteration = 0; Iteration < NumIterations; ++Iteration)
+					{
+						FSearchIndex TestSearchIndexB = SearchIndex;
+						TestSearchIndexB.PruneDuplicateValues(IndexBaseDatabases[0]->PosePruningSimilarityThreshold, IndexBaseDatabases[0]->Schema->SchemaCardinality);
+
+						if (TestSearchIndexA.Values != TestSearchIndexB.Values)
+						{
+							UE_LOG(LogPoseSearch, Warning, TEXT("OnGetComplete - PruneDuplicateValues is not deterministic"));
+						}
+					}
+				}
+#endif // ENABLE_ANIM_DEBUG
+
 				SearchIndex.PruneDuplicateValues(IndexBaseDatabases[0]->PosePruningSimilarityThreshold, IndexBaseDatabases[0]->Schema->SchemaCardinality);
 
 				TArray<float> Deviation = FMeanDeviationCalculator::Calculate(SearchIndexBases, Schemas);
@@ -1141,11 +1205,46 @@ void FPoseSearchDatabaseAsyncCacheTask::OnGetComplete(UE::DerivedData::FCacheGet
 				FCacheRecordBuilder Builder(FullIndexKey);
 				Builder.AddValue(Id, RawData);
 
-				GetCache().Put({ { { IndexBaseDatabases[0]->GetPathName() }, Builder.Build() } }, Owner, [IndexBaseDatabases, FullIndexKey](FCachePutResponse&& Response)
+				GetCache().Put({ { { IndexBaseDatabases[0]->GetPathName() }, Builder.Build() } }, Owner, [this, IndexBaseDatabases, FullIndexKey](FCachePutResponse&& Response)
 					{
-						if (Response.Status == EStatus::Error)
+						switch (Response.Status)
 						{
+						case EStatus::Error:
 							UE_LOG(LogPoseSearch, Log, TEXT("%s - %s Failed to store DDC"), *LexToString(FullIndexKey.Hash), *IndexBaseDatabases[0]->GetName());
+							break;
+						case EStatus::Canceled:
+							UE_LOG(LogPoseSearch, Log, TEXT("%s - %s Canceled to store DDC"), *LexToString(FullIndexKey.Hash), *IndexBaseDatabases[0]->GetName());
+							break;
+						case EStatus::Ok:
+							UE_LOG(LogPoseSearch, Log, TEXT("%s - %s BuildIndex stored to DDC"), *LexToString(FullIndexKey.Hash), *IndexBaseDatabases[0]->GetName());
+
+#if ENABLE_ANIM_DEBUG
+							if (AnyTestFlags(EMotionMatchTestFlags::ValidateDDC))
+							{
+								TArray<FCacheGetRequest> CacheRequests;
+								const FCacheKey CacheKey{ Bucket, DerivedDataKey };
+								CacheRequests.Add({ { Database->GetPathName() }, CacheKey, ECachePolicy::Default });
+								GetCache().Get(CacheRequests, Owner, [this, IndexBaseDatabases](FCacheGetResponse&& Response)
+									{
+										const FCacheKey& FullIndexKey = Response.Record.GetKey();
+										check(FullIndexKey.Hash == DerivedDataKey);
+
+										if (Response.Status == EStatus::Ok)
+										{
+											FSharedBuffer RawData = Response.Record.GetValue(Id).GetData().Decompress();
+											FMemoryReaderView Reader(RawData);
+
+											FSearchIndex TestSearchIndex;
+											Reader << TestSearchIndex;
+
+											if (!TestSearchIndex.Compare(SearchIndex))
+											{
+												UE_LOG(LogPoseSearch, Warning, TEXT("%s - %s DDC Index mismatch with BuildIndex"), *LexToString(FullIndexKey.Hash), *IndexBaseDatabases[0]->GetName());
+											}
+										}
+									});
+							}
+#endif // ENABLE_ANIM_DEBUG
 						}
 					});
 
@@ -1260,13 +1359,29 @@ void FAsyncPoseSearchDatabasesManagement::Tick(float DeltaTime)
 	check(IsInGameThread());
 
 #if ENABLE_ANIM_DEBUG
-	if (CVarMotionMatchTestInvalidateIndexingCache.GetValueOnAnyThread())
+	if (AnyTestFlags(EMotionMatchTestFlags::InvalidateCache))
 	{
-		for (TUniquePtr<FPoseSearchDatabaseAsyncCacheTask>& TaskPtr : Tasks)
+		if (AnyTestFlags(EMotionMatchTestFlags::WaitForTaskCompletion))
 		{
-			UE_LOG(LogPoseSearch, Log, TEXT("%s - %s Cancelled because of CVarMotionMatchTestInvalidateIndexingCache"), *LexToString(TaskPtr->GetDerivedDataKey()), *TaskPtr->GetDatabase()->GetName());
+			// iterating backwards because of the possible RemoveAtSwap 
+			for (int32 TaskIndex = Tasks.Num() - 1; TaskIndex >= 0; --TaskIndex)
+			{
+				if (Tasks[TaskIndex]->GetState() == FPoseSearchDatabaseAsyncCacheTask::EState::Ended ||
+					Tasks[TaskIndex]->GetState() == FPoseSearchDatabaseAsyncCacheTask::EState::Failed)
+				{
+					UE_LOG(LogPoseSearch, Log, TEXT("%s - %s Removed because of InvalidateCache with WaitForTaskCompletion"), *LexToString(Tasks[TaskIndex]->GetDerivedDataKey()), *Tasks[TaskIndex]->GetDatabase()->GetName());
+					Tasks.RemoveAtSwap(TaskIndex, 1, false);
+				}
+			}
+		}
+		else
+		{
+			for (TUniquePtr<FPoseSearchDatabaseAsyncCacheTask>& TaskPtr : Tasks)
+			{
+				UE_LOG(LogPoseSearch, Log, TEXT("%s - %s Cancelled because of InvalidateCache"), *LexToString(TaskPtr->GetDerivedDataKey()), *TaskPtr->GetDatabase()->GetName());
+			}
+			Tasks.Reset();
 		}
-		Tasks.Reset();
 	}
 #endif // ENABLE_ANIM_DEBUG
 	
diff --git a/Engine/Plugins/Experimental/Animation/PoseSearch/Source/Runtime/Private/PoseSearchDerivedDataKey.cpp b/Engine/Plugins/Experimental/Animation/PoseSearch/Source/Runtime/Private/PoseSearchDerivedDataKey.cpp
--- a/Engine/Plugins/Experimental/Animation/PoseSearch/Source/Runtime/Private/PoseSearchDerivedDataKey.cpp
+++ b/Engine/Plugins/Experimental/Animation/PoseSearch/Source/Runtime/Private/PoseSearchDerivedDataKey.cpp
@@ -40,7 +40,7 @@ FKeyBuilder::FKeyBuilder(const UObject* Object, bool bUseDataVer, bool bPerformC
 	if (bUseDataVer)
 	{
 		// used to invalidate the key without having to change POSESEARCHDB_DERIVEDDATA_VER all the times
-		int32 POSESEARCHDB_DERIVEDDATA_VER_SMALL = 189;
+		int32 POSESEARCHDB_DERIVEDDATA_VER_SMALL = 191;
 		FGuid VersionGuid = FDevSystemGuids::GetSystemGuid(FDevSystemGuids::Get().POSESEARCHDB_DERIVEDDATA_VER);
 
 		*this << VersionGuid;
@@ -89,6 +89,7 @@ bool FKeyBuilder::ShouldSkipProperty(const FProperty* InProperty) const
 		return true;
 	}
 		
+	check(!InProperty->HasMetaData(IgnoreForMemberInitializationTestName));
 	check(!InProperty->HasMetaData(NeverInHashName));
 
 	#if UE_POSE_SEARCH_DERIVED_DATA_LOGGING
diff --git a/Engine/Plugins/Experimental/Animation/PoseSearch/Source/Runtime/Public/PoseSearch/PoseSearchDerivedDataKey.h b/Engine/Plugins/Experimental/Animation/PoseSearch/Source/Runtime/Public/PoseSearch/PoseSearchDerivedDataKey.h
--- a/Engine/Plugins/Experimental/Animation/PoseSearch/Source/Runtime/Public/PoseSearch/PoseSearchDerivedDataKey.h
+++ b/Engine/Plugins/Experimental/Animation/PoseSearch/Source/Runtime/Public/PoseSearch/PoseSearchDerivedDataKey.h
@@ -31,7 +31,8 @@ public:
 
 	inline static const FName ExcludeFromHashName = FName(TEXT("ExcludeFromHash"));
 	inline static const FName NeverInHashName = FName(TEXT("NeverInHash"));
-
+	inline static const FName IgnoreForMemberInitializationTestName = FName(TEXT("IgnoreForMemberInitializationTest"));
+	
 	FKeyBuilder();
 	FKeyBuilder(const UObject* Object, bool bUseDataVer, bool bPerformConditionalPostLoadIfRequired);
 
