motion matching - initializing FDatabasePreviewScene without physics and audio, and set as non transactional

diff --git a/Engine/Plugins/Experimental/Animation/PoseSearch/Source/Editor/Private/PoseSearchDatabaseEditor.cpp b/Engine/Plugins/Experimental/Animation/PoseSearch/Source/Editor/Private/PoseSearchDatabaseEditor.cpp
--- a/Engine/Plugins/Experimental/Animation/PoseSearch/Source/Editor/Private/PoseSearchDatabaseEditor.cpp
+++ b/Engine/Plugins/Experimental/Animation/PoseSearch/Source/Editor/Private/PoseSearchDatabaseEditor.cpp
@@ -122,8 +122,8 @@ namespace UE::PoseSearch
 			PreviewScene = MakeShareable(
 				new FDatabasePreviewScene(
 					FPreviewScene::ConstructionValues()
-					.AllowAudioPlayback(true)
-					.ShouldSimulatePhysics(false)
+					.SetCreatePhysicsScene(false)
+					.SetTransactional(false)
 					.ForceUseMovementComponentInNonGameWorld(true),
 					StaticCastSharedRef<FDatabaseEditor>(AsShared())));
 
