Blend stack: Allow t-pose when missing anim asset. Previous behavior is "do nothing" on missing asset which might produce unexpected stateful behavior.
#rnx

diff --git a/Engine/Plugins/Experimental/Animation/PoseSearch/Source/Runtime/Private/AnimNode_BlendStack.cpp b/Engine/Plugins/Experimental/Animation/PoseSearch/Source/Runtime/Private/AnimNode_BlendStack.cpp
--- a/Engine/Plugins/Experimental/Animation/PoseSearch/Source/Runtime/Private/AnimNode_BlendStack.cpp
+++ b/Engine/Plugins/Experimental/Animation/PoseSearch/Source/Runtime/Private/AnimNode_BlendStack.cpp
@@ -22,16 +22,18 @@ TAutoConsoleVariable<int32> CVarAnimBlendStackPruningEnable(TEXT("a.AnimNode.Ble
 
 /////////////////////////////////////////////////////
 // FPoseSearchAnimPlayer
-void FPoseSearchAnimPlayer::Initialize(UAnimationAsset* AnimationAsset, float AccumulatedTime, bool bLoop, bool bMirrored, UMirrorDataTable* MirrorDataTable, float BlendTime, float RootBoneBlendTime, const UBlendProfile* BlendProfile, EAlphaBlendOption InBlendOption, FVector BlendParameters, float PlayRate, int32 InPoseLinkIdx)
+void FPoseSearchAnimPlayer::Initialize(const FAnimationInitializeContext& Context, UAnimationAsset* AnimationAsset, float AccumulatedTime, bool bLoop, bool bMirrored, UMirrorDataTable* MirrorDataTable, float BlendTime, float RootBoneBlendTime, const UBlendProfile* BlendProfile, EAlphaBlendOption InBlendOption, FVector BlendParameters, float PlayRate, int32 InPoseLinkIdx)
 {
-	check(AnimationAsset);
-
 	if (bMirrored && !MirrorDataTable)
 	{
 		UE_LOG(LogPoseSearch, Error, TEXT("FPoseSearchAnimPlayer failed to Initialize for %s. Mirroring will not work becasue MirrorDataTable is missing"), *GetNameSafe(AnimationAsset));
 	}
+	
+	check(Context.AnimInstanceProxy);
+	USkeleton* Skeleton = Context.AnimInstanceProxy->GetSkeleton();
+	check(Skeleton);
 
-	const FReferenceSkeleton& RefSkeleton = AnimationAsset->GetSkeleton()->GetReferenceSkeleton();
+	const FReferenceSkeleton& RefSkeleton = Context.AnimInstanceProxy->GetSkeleton()->GetReferenceSkeleton();
 	const bool bApplyDifferentRootBoneBlendTime = RootBoneBlendTime >= 0.f && !FMath::IsNearlyEqual(RootBoneBlendTime, BlendTime);
 	const int32 NumSkeletonBones = RefSkeleton.GetNum();
 	if (NumSkeletonBones <= 0)
@@ -79,9 +81,10 @@ void FPoseSearchAnimPlayer::Initialize(UAnimationAsset* AnimationAsset, float Ac
 	MirrorNode.SetMirrorDataTable(MirrorDataTable);
 	MirrorNode.SetMirror(bMirrored);
 	
+	bool bUnsupportedAnimAsset = false;
 	if (Cast<UAnimMontage>(AnimationAsset))
 	{
-		UE_LOG(LogPoseSearch, Error, TEXT("FPoseSearchAnimPlayer unsupported AnimationAsset %s"), *GetNameSafe(AnimationAsset));
+		bUnsupportedAnimAsset = true;
 	}
 	else if (UAnimSequenceBase* SequenceBase = Cast<UAnimSequenceBase>(AnimationAsset))
 	{
@@ -102,9 +105,14 @@ void FPoseSearchAnimPlayer::Initialize(UAnimationAsset* AnimationAsset, float Ac
 		BlendSpacePlayerNode.SetPlayRate(PlayRate);
 		BlendSpacePlayerNode.SetPosition(BlendParameters);
 	}
-	else
+	else if (AnimationAsset)
 	{
-		checkNoEntry();
+		bUnsupportedAnimAsset = true;
+	}
+
+	if (bUnsupportedAnimAsset)
+	{
+		UE_LOG(LogPoseSearch, Error, TEXT("FPoseSearchAnimPlayer unsupported AnimationAsset %s"), *GetNameSafe(AnimationAsset));
 	}
 
 	UpdateSourceLinkNode();
@@ -657,9 +665,9 @@ void FAnimNode_BlendStack_Standalone::BlendTo(const FAnimationUpdateContext& Con
 
 	AnimPlayers.PushFirst(FPoseSearchAnimPlayer());
 	FPoseSearchAnimPlayer& AnimPlayer = AnimPlayers.First();
-	AnimPlayer.Initialize(AnimationAsset, AccumulatedTime, bLoop, bMirrored, MirrorDataTable, BlendTime, RootBoneBlendTime, BlendProfile, BlendOption, BlendParameters, PlayRate, GetNextPoseLinkIndex());
 
 	FAnimationInitializeContext InitContext(Context.AnimInstanceProxy, Context.SharedContext);
+	AnimPlayer.Initialize(InitContext, AnimationAsset, AccumulatedTime, bLoop, bMirrored, MirrorDataTable, BlendTime, RootBoneBlendTime, BlendProfile, BlendOption, BlendParameters, PlayRate, GetNextPoseLinkIndex());
 	InitializeSample(InitContext, AnimPlayer);
 }
 
@@ -729,46 +737,42 @@ void FAnimNode_BlendStack::UpdateAssetPlayer(const FAnimationUpdateContext& Cont
 
 	GetEvaluateGraphExposedInputs().Execute(Context);
 
-	if (AnimationAsset)
+	bool bExecuteBlendTo = false;
+	if (AnimPlayers.IsEmpty())
+	{
+		bExecuteBlendTo = true;
+	}
+	else
 	{
-		bool bExecuteBlendTo = false;
-		if (AnimPlayers.IsEmpty())
+		const FPoseSearchAnimPlayer& MainAnimPlayer = AnimPlayers.First();
+		const UAnimationAsset* PlayingAnimationAsset = MainAnimPlayer.GetAnimationAsset();
+
+		if (bForceBlendNextUpdate)
 		{
+			bForceBlendNextUpdate = false;
 			bExecuteBlendTo = true;
 		}
-		else
+		else if (AnimationAsset != PlayingAnimationAsset)
 		{
-			const FPoseSearchAnimPlayer& MainAnimPlayer = AnimPlayers.First();
-			const UAnimationAsset* PlayingAnimationAsset = MainAnimPlayer.GetAnimationAsset();
-			check(PlayingAnimationAsset);
-
-			if (bForceBlendNextUpdate)
-			{
-				bForceBlendNextUpdate = false;
-				bExecuteBlendTo = true;
-			}
-			else if (AnimationAsset != PlayingAnimationAsset)
-			{
-				bExecuteBlendTo = true;
-			}
-			else if (bMirrored != MainAnimPlayer.GetMirror())
-			{
-				bExecuteBlendTo = true;
-			}
-			else if (BlendParameters != MainAnimPlayer.GetBlendParameters())
-			{
-				bExecuteBlendTo = true;
-			}
-			else if (MaxAnimationDeltaTime >= 0.f && FMath::Abs(AnimationTime - MainAnimPlayer.GetAccumulatedTime()) > MaxAnimationDeltaTime)
-			{
-				bExecuteBlendTo = true;
-			}
+			bExecuteBlendTo = true;
 		}
-
-		if (bExecuteBlendTo)
+		else if (bMirrored != MainAnimPlayer.GetMirror())
+		{
+			bExecuteBlendTo = true;
+		}
+		else if (BlendParameters != MainAnimPlayer.GetBlendParameters())
 		{
-			BlendTo(Context, AnimationAsset, AnimationTime, bLoop, bMirrored, MirrorDataTable.Get(), BlendTime, RootBoneBlendTime, BlendProfile, BlendOption, bUseInertialBlend, BlendParameters, WantedPlayRate);
+			bExecuteBlendTo = true;
 		}
+		else if (MaxAnimationDeltaTime >= 0.f && FMath::Abs(AnimationTime - MainAnimPlayer.GetAccumulatedTime()) > MaxAnimationDeltaTime)
+		{
+			bExecuteBlendTo = true;
+		}
+	}
+
+	if (bExecuteBlendTo)
+	{
+		BlendTo(Context, AnimationAsset, AnimationTime, bLoop, bMirrored, MirrorDataTable.Get(), BlendTime, RootBoneBlendTime, BlendProfile, BlendOption, bUseInertialBlend, BlendParameters, WantedPlayRate);
 	}
 	
 	UpdatePlayRate(WantedPlayRate);
diff --git a/Engine/Plugins/Experimental/Animation/PoseSearch/Source/Runtime/Public/PoseSearch/AnimNode_BlendStack.h b/Engine/Plugins/Experimental/Animation/PoseSearch/Source/Runtime/Public/PoseSearch/AnimNode_BlendStack.h
--- a/Engine/Plugins/Experimental/Animation/PoseSearch/Source/Runtime/Public/PoseSearch/AnimNode_BlendStack.h
+++ b/Engine/Plugins/Experimental/Animation/PoseSearch/Source/Runtime/Public/PoseSearch/AnimNode_BlendStack.h
@@ -15,7 +15,7 @@ struct FPoseSearchAnimPlayer
 {
 	GENERATED_BODY()
 	
-	void Initialize(UAnimationAsset* AnimationAsset, float AccumulatedTime, bool bLoop, bool bMirrored, UMirrorDataTable* MirrorDataTable, float BlendTime, float RootBoneBlendTime, const UBlendProfile* BlendProfile, EAlphaBlendOption InBlendOption, FVector BlendParameters, float PlayRate, int32 InPoseLinkIdx);
+	void Initialize(const FAnimationInitializeContext& Context, UAnimationAsset* AnimationAsset, float AccumulatedTime, bool bLoop, bool bMirrored, UMirrorDataTable* MirrorDataTable, float BlendTime, float RootBoneBlendTime, const UBlendProfile* BlendProfile, EAlphaBlendOption InBlendOption, FVector BlendParameters, float PlayRate, int32 InPoseLinkIdx);
 	void UpdatePlayRate(float PlayRate);
 	void Evaluate_AnyThread(FPoseContext& Output);
 	void Update_AnyThread(const FAnimationUpdateContext& Context);
