motion matching - making FAsyncPoseSearchDatabasesManagement tickable in game (so CVarMotionMatchTestFlags can be used for testing)

diff --git a/Engine/Plugins/Experimental/Animation/PoseSearch/Source/Runtime/Private/PoseSearchAssetIndexer.cpp b/Engine/Plugins/Experimental/Animation/PoseSearch/Source/Runtime/Private/PoseSearchAssetIndexer.cpp
--- a/Engine/Plugins/Experimental/Animation/PoseSearch/Source/Runtime/Private/PoseSearchAssetIndexer.cpp
+++ b/Engine/Plugins/Experimental/Animation/PoseSearch/Source/Runtime/Private/PoseSearchAssetIndexer.cpp
@@ -15,7 +15,12 @@
 
 namespace UE::PoseSearch
 {
-	
+
+#if ENABLE_ANIM_DEBUG
+static TAutoConsoleVariable<bool> CVarMotionMatchTestDisableIndexerCaching(TEXT("a.MotionMatch.TestDisableIndexerCaching"), false, TEXT("Disable Motion Matching Indexer Caching"));
+static TAutoConsoleVariable<int> CVarMotionMatchTestExtractPoseDeterminismNumIterations(TEXT("a.MotionMatch.TestExtractPoseDeterminismNumIterations"), 0, TEXT("Test Motion Matching ExtractPose Determinism via this NumIterations retries"));
+#endif // ENABLE_ANIM_DEBUG
+
 //////////////////////////////////////////////////////////////////////////
 // FSamplingParam helpers
 struct FSamplingParam
@@ -271,8 +276,13 @@ FTransform FAssetIndexer::MirrorTransform(const FTransform& Transform) const
 FAssetIndexer::CachedEntry& FAssetIndexer::GetEntry(float SampleTime)
 {
 	using namespace UE::Anim;
+	
+	bool bDisableCaching = false;
+#if ENABLE_ANIM_DEBUG
+	bDisableCaching = CVarMotionMatchTestDisableIndexerCaching.GetValueOnAnyThread();
+#endif // ENABLE_ANIM_DEBUG
 
-	CachedEntry* Entry = CachedEntries.Find(SampleTime);
+	CachedEntry* Entry = bDisableCaching ? nullptr : CachedEntries.Find(SampleTime);
 	if (!Entry)
 	{
 		Entry = &CachedEntries.Add(SampleTime);
@@ -303,6 +313,34 @@ FAssetIndexer::CachedEntry& FAssetIndexer::GetEntry(float SampleTime)
 		FCompactPose Pose;
 		Pose.SetBoneContainer(&BoneContainer);
 		AssetSampler.ExtractPose(CurrentTime, Pose);
+
+#if ENABLE_ANIM_DEBUG
+		const int32 NumIterations = CVarMotionMatchTestExtractPoseDeterminismNumIterations.GetValueOnAnyThread();
+		for (int32 IterationIndex = 0; IterationIndex < NumIterations; ++IterationIndex)
+		{
+			FCompactPose TestPose;
+			TestPose.SetBoneContainer(&BoneContainer);
+			AssetSampler.ExtractPose(CurrentTime, TestPose);
+
+			const TConstArrayView<FTransform> Bones = Pose.GetBones();
+			const TConstArrayView<FTransform> TestBones = TestPose.GetBones();
+			if (Bones.Num() != TestBones.Num())
+			{
+				UE_LOG(LogPoseSearch, Warning, TEXT("FAssetIndexer::GetEntry - ExtractPose is not deterministic"));
+			}
+			else
+			{
+				for (int32 BoneIndex = 0; BoneIndex < Bones.Num(); ++BoneIndex)
+				{
+					if (FMemory::Memcmp(&Bones[BoneIndex], &TestBones[BoneIndex], sizeof(FTransform)) != 0)
+					{
+						UE_LOG(LogPoseSearch, Warning, TEXT("FAssetIndexer::GetEntry - ExtractPose is not deterministic"));
+					}
+				}
+			}
+		}
+#endif // ENABLE_ANIM_DEBUG
+
 		Pose[FCompactPoseBoneIndex(RootBoneIndexType)].SetIdentity();
 
 		if (SearchIndexAsset.bMirrored && Schema.MirrorDataTable)
@@ -690,5 +728,65 @@ float FAssetIndexer::CalculatePermutationTimeOffset() const
 	return PermutationTimeOffset;
 }
 
+#if ENABLE_ANIM_DEBUG
+void FAssetIndexer::CompareCachedEntries(const FAssetIndexer& Other) const
+{
+	if (CachedEntries.Num() != Other.CachedEntries.Num())
+	{
+		UE_LOG(LogPoseSearch, Warning, TEXT("CompareCachedEntries - FAssetIndexer::CachedEntries::Num is not deterministic"));
+	}
+	else
+	{
+		for (const TPair<float, CachedEntry>& Pair : CachedEntries)
+		{
+			if (const CachedEntry* OtherEntry = Other.CachedEntries.Find(Pair.Key))
+			{
+				const CachedEntry* Entry = &Pair.Value;
+				if (Entry->SampleTime != OtherEntry->SampleTime)
+				{
+					UE_LOG(LogPoseSearch, Warning, TEXT("CompareCachedEntries - FAssetIndexer::CachedEntries::SampleTime is not deterministic (%f, %f)"), Entry->SampleTime, OtherEntry->SampleTime);
+				}
+
+				if (Entry->bClamped != OtherEntry->bClamped)
+				{
+					UE_LOG(LogPoseSearch, Warning, TEXT("CompareCachedEntries - FAssetIndexer::CachedEntries::bClamped is not deterministic"));
+				}
+
+				if (FMemory::Memcmp(&Entry->RootTransform, &OtherEntry->RootTransform, sizeof(FTransform)) != 0)
+				{
+					UE_LOG(LogPoseSearch, Warning, TEXT("CompareCachedEntries - FAssetIndexer::CachedEntries::RootTransform is not deterministic"));
+				}
+
+				if (Entry->ComponentSpacePose.GetComponentSpaceFlags() != OtherEntry->ComponentSpacePose.GetComponentSpaceFlags())
+				{
+					UE_LOG(LogPoseSearch, Warning, TEXT("CompareCachedEntries - FAssetIndexer::CachedEntries::ComponentSpacePose::ComponentSpaceFlags is not deterministic"));
+				}
+
+				const TConstArrayView<FTransform> Bones = Entry->ComponentSpacePose.GetPose().GetBones();
+				const TConstArrayView<FTransform> OtherBones = OtherEntry->ComponentSpacePose.GetPose().GetBones();
+				if (Bones.Num() != OtherBones.Num())
+				{
+					UE_LOG(LogPoseSearch, Warning, TEXT("CompareCachedEntries - FAssetIndexer::CachedEntries::ComponentSpacePose::Bones is not deterministic"));
+				}
+				else
+				{
+					for (int32 Index = 0; Index < Bones.Num(); ++Index)
+					{
+						if (FMemory::Memcmp(&Bones[Index], &OtherBones[Index], sizeof(FTransform)) != 0)
+						{
+							UE_LOG(LogPoseSearch, Warning, TEXT("CompareCachedEntries - FAssetIndexer::CachedEntries::Bones[%d] is not deterministic"), Index);
+						}
+					}
+				}
+			}
+			else
+			{
+				UE_LOG(LogPoseSearch, Warning, TEXT("CompareCachedEntries - FAssetIndexer::CachedEntries is not deterministic. Missing CachedEntry at time %f"), Pair.Key);
+			}
+		}
+	}
+}
+#endif // ENABLE_ANIM_DEBUG
+
 } // namespace UE::PoseSearch
 #endif // WITH_EDITOR
diff --git a/Engine/Plugins/Experimental/Animation/PoseSearch/Source/Runtime/Private/PoseSearchDatabase.cpp b/Engine/Plugins/Experimental/Animation/PoseSearch/Source/Runtime/Private/PoseSearchDatabase.cpp
--- a/Engine/Plugins/Experimental/Animation/PoseSearch/Source/Runtime/Private/PoseSearchDatabase.cpp
+++ b/Engine/Plugins/Experimental/Animation/PoseSearch/Source/Runtime/Private/PoseSearchDatabase.cpp
@@ -16,6 +16,10 @@
 #include "PoseSearch/PoseSearchSchema.h"
 #include "UObject/ObjectSaveContext.h"
 
+#if WITH_EDITOR && WITH_ENGINE
+#include "Editor/EditorEngine.h"
+#endif //WITH_EDITOR && WITH_ENGINE
+
 DECLARE_STATS_GROUP(TEXT("PoseSearch"), STATGROUP_PoseSearch, STATCAT_Advanced);
 DECLARE_CYCLE_STAT_EXTERN(TEXT("Search Brute Force"), STAT_PoseSearch_BruteForce, STATGROUP_PoseSearch, );
 DECLARE_CYCLE_STAT_EXTERN(TEXT("Search PCA/KNN"), STAT_PoseSearch_PCAKNN, STATGROUP_PoseSearch, );
@@ -524,7 +528,17 @@ void UPoseSearchDatabase::PostLoad()
 
 #if WITH_EDITOR
 	using namespace UE::PoseSearch;
-	FAsyncPoseSearchDatabasesManagement::RequestAsyncBuildIndex(this, ERequestAsyncBuildFlag::NewRequest);
+
+	ERequestAsyncBuildFlag Flag = ERequestAsyncBuildFlag::NewRequest;
+#if WITH_ENGINE
+	// If there isn't an EditorEngine (ex. Standalone Game via -game argument) we WaitForCompletion
+	if (Cast<UEditorEngine>(GEngine) == nullptr)
+	{
+		Flag |= ERequestAsyncBuildFlag::WaitForCompletion;
+	}
+#endif // WITH_ENGINE
+
+	FAsyncPoseSearchDatabasesManagement::RequestAsyncBuildIndex(this, Flag);
 #endif
 
 	Super::PostLoad();
diff --git a/Engine/Plugins/Experimental/Animation/PoseSearch/Source/Runtime/Private/PoseSearchDatabaseIndexingContext.h b/Engine/Plugins/Experimental/Animation/PoseSearch/Source/Runtime/Private/PoseSearchDatabaseIndexingContext.h
--- a/Engine/Plugins/Experimental/Animation/PoseSearch/Source/Runtime/Private/PoseSearchDatabaseIndexingContext.h
+++ b/Engine/Plugins/Experimental/Animation/PoseSearch/Source/Runtime/Private/PoseSearchDatabaseIndexingContext.h
@@ -37,6 +37,10 @@ struct FDatabaseIndexingContext
 {
 	bool IndexDatabase(FSearchIndexBase& SearchIndexBase, const UPoseSearchDatabase& Database, UE::DerivedData::FRequestOwner& Owner);
 
+	const FAssetSamplingContext& GetSamplingContext() const { return SamplingContext; }
+	const TConstArrayView<FAnimationAssetSampler> GetSamplers() const { return Samplers; }
+	const TConstArrayView<FAssetIndexer> GetIndexers() const { return Indexers; }
+
 private:
 	FAssetSamplingContext SamplingContext;
 	TArray<FAnimationAssetSampler> Samplers;
diff --git a/Engine/Plugins/Experimental/Animation/PoseSearch/Source/Runtime/Private/PoseSearchDerivedData.cpp b/Engine/Plugins/Experimental/Animation/PoseSearch/Source/Runtime/Private/PoseSearchDerivedData.cpp
--- a/Engine/Plugins/Experimental/Animation/PoseSearch/Source/Runtime/Private/PoseSearchDerivedData.cpp
+++ b/Engine/Plugins/Experimental/Animation/PoseSearch/Source/Runtime/Private/PoseSearchDerivedData.cpp
@@ -11,7 +11,6 @@
 #include "AssetRegistry/AssetRegistryModule.h"
 #include "DerivedDataCache.h"
 #include "DerivedDataRequestOwner.h"
-#include "Editor/EditorEngine.h"
 #include "InstancedStruct.h"
 #include "Misc/CoreDelegates.h"
 #include "PoseSearchDatabaseIndexingContext.h"
@@ -678,6 +677,150 @@ static void PreprocessSearchIndexKDTree(FSearchIndex& SearchIndex, const UPoseSe
 	}
 }
 
+#if ENABLE_ANIM_DEBUG
+
+static void CompareIndexingContext(const FDatabaseIndexingContext& A, const FDatabaseIndexingContext& B)
+{
+	if (A.GetIndexers().Num() != B.GetIndexers().Num())
+	{
+		UE_LOG(LogPoseSearch, Warning, TEXT("CompareIndexers - FAssetIndexer is not deterministic"));
+	}
+	else
+	{
+		for (int32 Index = 0; Index < A.GetIndexers().Num(); ++Index)
+		{
+			A.GetIndexers()[Index].CompareCachedEntries(B.GetIndexers()[Index]);
+		}
+	}
+}
+
+static void CompareChannelValues(int32 RecursionIndex, int32 PoseIndex, const TConstArrayView<float>& PoseA, const TConstArrayView<float>& PoseB, const TConstArrayView<TObjectPtr<UPoseSearchFeatureChannel>>& Channels, FStringBuilderBase& StringBuilder)
+{
+	bool bPrintHeader = true;
+	for (const TObjectPtr<UPoseSearchFeatureChannel>& ChannelPtr : Channels)
+	{
+		const int32 ChannelCardinality = ChannelPtr->GetChannelCardinality();
+		const int32 ChannelDataOffset = ChannelPtr->GetChannelDataOffset();
+
+		for (int32 Index = 0; Index < ChannelCardinality; ++Index)
+		{
+			const int32 DataOffset = ChannelDataOffset + Index;
+			const float ValueA = PoseA[DataOffset];
+			const float ValueB = PoseB[DataOffset];
+			if (ValueA != ValueB)
+			{
+				if (bPrintHeader && RecursionIndex == 0)
+				{
+					StringBuilder.Appendf(TEXT("Values mismatch at pose %d\n"), PoseIndex);
+					bPrintHeader = false;
+				}
+
+				for (int32 Indentation = 0; Indentation < RecursionIndex; ++Indentation)
+				{
+					StringBuilder.Append(TEXT("    "));
+				}
+
+				StringBuilder.Appendf(TEXT("%s - %d (%f, %f)\n"), *ChannelPtr->GetName(), Index, ValueA, ValueB);
+			}
+		}
+
+		CompareChannelValues(RecursionIndex + 1, PoseIndex, PoseA, PoseB, ChannelPtr->GetSubChannels(), StringBuilder);
+	}
+}
+
+static void CompareSearchIndexBase(const FSearchIndexBase& A, const FSearchIndexBase& B, const UPoseSearchSchema* Schema, FStringBuilderBase& StringBuilder)
+{
+	check(Schema);
+
+	if (A.Values.Num() != B.Values.Num())
+	{
+		StringBuilder.Append(TEXT("Values.Num mismatch\n"));
+	}
+	else if ((A.Values.Num() % Schema->SchemaCardinality) != 0)
+	{
+		StringBuilder.Append(TEXT("Values.Num is not a multiple of Schema->SchemaCardinality!\n"));
+	}
+	else if (Schema->SchemaCardinality > 0)
+	{
+		// cannot use A.GetNumPoses() since A.Values can be pruned out from duplicates
+		int32 ValuesDataOffset = 0;
+		const int32 NumValuePoses = A.Values.Num() / Schema->SchemaCardinality;
+		for (int32 ValuePoseIndex = 0; ValuePoseIndex < NumValuePoses; ++ValuePoseIndex)
+		{
+			const TConstArrayView<float> PoseA = MakeArrayView(A.Values.GetData() + ValuePoseIndex * Schema->SchemaCardinality, Schema->SchemaCardinality);
+			const TConstArrayView<float> PoseB = MakeArrayView(B.Values.GetData() + ValuePoseIndex * Schema->SchemaCardinality, Schema->SchemaCardinality);
+			CompareChannelValues(0, ValuePoseIndex, PoseA, PoseB, Schema->GetChannels(), StringBuilder);
+		}
+	}
+
+	if (A.PoseMetadata != B.PoseMetadata)
+	{
+		StringBuilder.Append(TEXT("PoseMetadata mismatch\n"));
+	}
+
+	if (A.bAnyBlockTransition != B.bAnyBlockTransition)
+	{
+		StringBuilder.Append(TEXT("bAnyBlockTransition mismatch\n"));
+	}
+
+	if (A.Assets != B.Assets)
+	{
+		StringBuilder.Append(TEXT("Assets mismatch\n"));
+	}
+	 
+	if (A.MinCostAddend != B.MinCostAddend)
+	{
+		StringBuilder.Append(TEXT("MinCostAddend mismatch\n"));
+	}
+
+	if (A.Stats != B.Stats)
+	{
+		StringBuilder.Append(TEXT("Stats mismatch\n"));
+	}
+}
+
+static void CompareSearchIndex(const FSearchIndex& A, const FSearchIndex& B, const UPoseSearchSchema* Schema, FStringBuilderBase& StringBuilder)
+{
+	CompareSearchIndexBase(A, B, Schema, StringBuilder);
+
+	if (A.WeightsSqrt != B.WeightsSqrt)
+	{
+		StringBuilder.Append(TEXT("WeightsSqrt mismatch\n"));
+	}
+	
+	if (A.PCAValues != B.PCAValues)
+	{
+		StringBuilder.Append(TEXT("PCAValues mismatch\n"));
+	}
+
+	if (A.PCAValuesVectorToPoseIndexes != B.PCAValuesVectorToPoseIndexes)
+	{
+		StringBuilder.Append(TEXT("PCAValuesVectorToPoseIndexes mismatch\n"));
+	}
+
+	if (A.PCAProjectionMatrix != B.PCAProjectionMatrix)
+	{
+		StringBuilder.Append(TEXT("PCAProjectionMatrix mismatch\n"));
+	}
+
+	if (A.Mean != B.Mean)
+	{
+		StringBuilder.Append(TEXT("Mean mismatch\n"));
+	}
+
+	if (A.PCAExplainedVariance != B.PCAExplainedVariance)
+	{
+		StringBuilder.Append(TEXT("PCAExplainedVariance mismatch\n"));
+	}
+
+	if (A.KDTree != B.KDTree)
+	{
+		StringBuilder.Append(TEXT("KDTree mismatch\n"));
+	}
+}
+
+#endif // ENABLE_ANIM_DEBUG
+
 //////////////////////////////////////////////////////////////////////////
 // FPoseSearchDatabaseAsyncCacheTask
 struct FPoseSearchDatabaseAsyncCacheTask
@@ -1095,9 +1238,13 @@ void FPoseSearchDatabaseAsyncCacheTask::OnGetComplete(UE::DerivedData::FCacheGet
 							FDatabaseIndexingContext TestDbIndexingContext;
 							if (TestDbIndexingContext.IndexDatabase(TestSearchIndexBase, *IndexBaseDatabase, Owner))
 							{
-								if (!TestSearchIndexBase.Compare(SearchIndexBase))
+								if (TestSearchIndexBase != SearchIndexBase)
 								{
-									UE_LOG(LogPoseSearch, Warning, TEXT("OnGetComplete - IndexDatabase is not deterministic"));
+									FStringBuilderBase Message;
+									CompareSearchIndexBase(TestSearchIndexBase, SearchIndexBase, IndexBaseDatabase->Schema, Message);
+									UE_LOG(LogPoseSearch, Warning, TEXT("OnGetComplete - IndexDatabase is not deterministic\n%s"), *Message);
+
+									CompareIndexingContext(TestDbIndexingContext, DbIndexingContext);
 								}
 							}
 						}
@@ -1188,9 +1335,11 @@ void FPoseSearchDatabaseAsyncCacheTask::OnGetComplete(UE::DerivedData::FCacheGet
 #if ENABLE_ANIM_DEBUG
 				if (bCompareSearchIndex)
 				{
-					if (!SearchIndexCompare.Compare(SearchIndex))
+					if (SearchIndexCompare != SearchIndex)
 					{
-						UE_LOG(LogPoseSearch, Warning, TEXT("%s - %s BuildIndex mismatch with DDC Index"), *LexToString(FullIndexKey.Hash), *IndexBaseDatabases[0]->GetName());
+						FStringBuilderBase Message;
+						CompareSearchIndex(SearchIndexCompare, SearchIndex, IndexBaseDatabases[0]->Schema, Message);
+						UE_LOG(LogPoseSearch, Warning, TEXT("%s - %s BuildIndex mismatch with DDC Index\n%s"), *LexToString(FullIndexKey.Hash), *IndexBaseDatabases[0]->GetName(), *Message);
 					}
 				}
 #endif // ENABLE_ANIM_DEBUG
@@ -1237,9 +1386,11 @@ void FPoseSearchDatabaseAsyncCacheTask::OnGetComplete(UE::DerivedData::FCacheGet
 											FSearchIndex TestSearchIndex;
 											Reader << TestSearchIndex;
 
-											if (!TestSearchIndex.Compare(SearchIndex))
+											if (TestSearchIndex != SearchIndex)
 											{
-												UE_LOG(LogPoseSearch, Warning, TEXT("%s - %s DDC Index mismatch with BuildIndex"), *LexToString(FullIndexKey.Hash), *IndexBaseDatabases[0]->GetName());
+												FStringBuilderBase Message;
+												CompareSearchIndex(TestSearchIndex, SearchIndex, IndexBaseDatabases[0]->Schema, Message);
+												UE_LOG(LogPoseSearch, Warning, TEXT("%s - %s DDC Index mismatch with BuildIndex\n%s"), *LexToString(FullIndexKey.Hash), *IndexBaseDatabases[0]->GetName(), *Message);
 											}
 										}
 									});
@@ -1429,14 +1580,6 @@ bool FAsyncPoseSearchDatabasesManagement::RequestAsyncBuildIndex(const UPoseSear
 
 	FScopeLock Lock(&Mutex);
 
-#if WITH_ENGINE
-	//If there isn't an EditorEngine (ex. Standalone Game via -game argument) the FAsyncPoseSearchDatabasesManagement task doesn't get ticked, need to force completion here
-	if (Cast<UEditorEngine>(GEngine) == nullptr)
-	{
-		Flag |= ERequestAsyncBuildFlag::WaitForCompletion;
-	}
-#endif // WITH_ENGINE
-
 	check(EnumHasAnyFlags(Flag, ERequestAsyncBuildFlag::NewRequest | ERequestAsyncBuildFlag::ContinueRequest));
 
 	FAsyncPoseSearchDatabasesManagement& This = FAsyncPoseSearchDatabasesManagement::Get();
diff --git a/Engine/Plugins/Experimental/Animation/PoseSearch/Source/Runtime/Private/PoseSearchIndex.cpp b/Engine/Plugins/Experimental/Animation/PoseSearch/Source/Runtime/Private/PoseSearchIndex.cpp
--- a/Engine/Plugins/Experimental/Animation/PoseSearch/Source/Runtime/Private/PoseSearchIndex.cpp
+++ b/Engine/Plugins/Experimental/Animation/PoseSearch/Source/Runtime/Private/PoseSearchIndex.cpp
@@ -423,44 +423,16 @@ void FSearchIndexBase::AllocateData(int32 DataCardinality, int32 NumPoses)
 	PoseMetadata.SetNumZeroed(NumPoses);
 }
 
-#if ENABLE_ANIM_DEBUG
-bool FSearchIndexBase::Compare(const FSearchIndexBase& Other) const
+bool FSearchIndexBase::operator==(const FSearchIndexBase& Other) const
 {
-	bool bResult = true;
-
-	if (Values != Other.Values)
-	{
-		UE_LOG(LogPoseSearch, Warning, TEXT("FSearchIndexBase::Compare - Values mismatch"));
-		bResult = false;
-	}
-
-	if (bAnyBlockTransition != Other.bAnyBlockTransition)
-	{
-		UE_LOG(LogPoseSearch, Warning, TEXT("FSearchIndexBase::Compare - bAnyBlockTransition mismatch"));
-		bResult = false;
-	}
-
-	if (Assets != Other.Assets)
-	{
-		UE_LOG(LogPoseSearch, Warning, TEXT("FSearchIndexBase::Compare - Assets mismatch"));
-		bResult = false;
-	}
-	 
-	if (MinCostAddend != Other.MinCostAddend)
-	{
-		UE_LOG(LogPoseSearch, Warning, TEXT("FSearchIndexBase::Compare - MinCostAddend mismatch"));
-		bResult = false;
-	}
-
-	if (Stats != Other.Stats)
-	{
-		UE_LOG(LogPoseSearch, Warning, TEXT("FSearchIndexBase::Compare - Stats mismatch"));
-		bResult = false;
-	}
-
-	return true;
+	return
+		Values == Other.Values &&
+		PoseMetadata == Other.PoseMetadata &&
+		bAnyBlockTransition == Other.bAnyBlockTransition &&
+		Assets == Other.Assets &&
+		MinCostAddend == Other.MinCostAddend &&
+		Stats == Other.Stats;
 }
-#endif // ENABLE_ANIM_DEBUG
 
 FArchive& operator<<(FArchive& Ar, FSearchIndexBase& Index)
 {
@@ -721,57 +693,18 @@ void FSearchIndex::GetPoseToPCAValuesVectorIndexes(TArray<uint32>& PoseToPCAValu
 	}
 }
 
-#if ENABLE_ANIM_DEBUG
-bool FSearchIndex::Compare(const FSearchIndex& Other) const
+bool FSearchIndex::operator==(const FSearchIndex& Other) const
 {
-	bool bResult = FSearchIndexBase::Compare(static_cast<const FSearchIndexBase&>(Other));
-
-	if (WeightsSqrt != Other.WeightsSqrt)
-	{
-		UE_LOG(LogPoseSearch, Warning, TEXT("FSearchIndex::Compare - WeightsSqrt mismatch"));
-		bResult = false;
-	}
-	
-	if (PCAValues != Other.PCAValues)
-	{
-		UE_LOG(LogPoseSearch, Warning, TEXT("FSearchIndex::Compare - PCAValues mismatch"));
-		bResult = false;
-	}
-
-	if (PCAValuesVectorToPoseIndexes != Other.PCAValuesVectorToPoseIndexes)
-	{
-		UE_LOG(LogPoseSearch, Warning, TEXT("FSearchIndex::Compare - PCAValuesVectorToPoseIndexes mismatch"));
-		bResult = false;
-	}
-
-	if (PCAProjectionMatrix != Other.PCAProjectionMatrix)
-	{
-		UE_LOG(LogPoseSearch, Warning, TEXT("FSearchIndex::Compare - PCAProjectionMatrix mismatch"));
-		bResult = false;
-	}
-
-	if (Mean != Other.Mean)
-	{
-		UE_LOG(LogPoseSearch, Warning, TEXT("FSearchIndex::Compare - Mean mismatch"));
-		bResult = false;
-	}
-
-	if (PCAExplainedVariance != Other.PCAExplainedVariance)
-	{
-		UE_LOG(LogPoseSearch, Warning, TEXT("FSearchIndex::Compare - PCAExplainedVariance mismatch"));
-		bResult = false;
-	}
-
-	if (KDTree != Other.KDTree)
-	{
-		UE_LOG(LogPoseSearch, Warning, TEXT("FSearchIndex::Compare - KDTree mismatch"));
-		bResult = false;
-	}
-	
-	return bResult;
+	return
+		FSearchIndexBase::operator==(Other) &&
+		WeightsSqrt == Other.WeightsSqrt &&
+		PCAValues == Other.PCAValues &&
+		PCAValuesVectorToPoseIndexes == Other.PCAValuesVectorToPoseIndexes &&
+		PCAProjectionMatrix == Other.PCAProjectionMatrix &&
+		Mean == Other.Mean &&
+		PCAExplainedVariance == Other.PCAExplainedVariance &&
+		KDTree == Other.KDTree;
 }
-#endif // ENABLE_ANIM_DEBUG
-
 
 FArchive& operator<<(FArchive& Ar, FSearchIndex& Index)
 {
diff --git a/Engine/Plugins/Experimental/Animation/PoseSearch/Source/Runtime/Public/PoseSearch/PoseSearchAssetIndexer.h b/Engine/Plugins/Experimental/Animation/PoseSearch/Source/Runtime/Public/PoseSearch/PoseSearchAssetIndexer.h
--- a/Engine/Plugins/Experimental/Animation/PoseSearch/Source/Runtime/Public/PoseSearch/PoseSearchAssetIndexer.h
+++ b/Engine/Plugins/Experimental/Animation/PoseSearch/Source/Runtime/Public/PoseSearch/PoseSearchAssetIndexer.h
@@ -63,6 +63,10 @@ public:
 
 	bool IsProcessFailed() const { return bProcessFailed; }
 
+#if ENABLE_ANIM_DEBUG
+	void CompareCachedEntries(const FAssetIndexer& Other) const;
+#endif // ENABLE_ANIM_DEBUG
+
 private:
 	int32 GetVectorIdx(int32 SampleIdx) const;
 
diff --git a/Engine/Plugins/Experimental/Animation/PoseSearch/Source/Runtime/Public/PoseSearch/PoseSearchDerivedData.h b/Engine/Plugins/Experimental/Animation/PoseSearch/Source/Runtime/Public/PoseSearch/PoseSearchDerivedData.h
--- a/Engine/Plugins/Experimental/Animation/PoseSearch/Source/Runtime/Public/PoseSearch/PoseSearchDerivedData.h
+++ b/Engine/Plugins/Experimental/Animation/PoseSearch/Source/Runtime/Public/PoseSearch/PoseSearchDerivedData.h
@@ -22,7 +22,7 @@ namespace UE::PoseSearch
 	};
 	ENUM_CLASS_FLAGS(ERequestAsyncBuildFlag);
 
-	class POSESEARCH_API FAsyncPoseSearchDatabasesManagement : public FTickableEditorObject, public FTickableCookObject, public FGCObject
+	class POSESEARCH_API FAsyncPoseSearchDatabasesManagement : public FTickableGameObject, public FTickableCookObject, public FGCObject
 	{
 	public:
 		static bool RequestAsyncBuildIndex(const UPoseSearchDatabase* Database, ERequestAsyncBuildFlag Flag);
@@ -39,11 +39,13 @@ namespace UE::PoseSearch
 		void Shutdown();
 		void StartQueuedTasks(int32 MaxActiveTasks);
 
-		// Begin FTickableEditorObject
+		// Begin FTickableGameObject
 		virtual void Tick(float DeltaTime) override;
 		virtual ETickableTickType GetTickableTickType() const override { return ETickableTickType::Always; }
 		virtual TStatId GetStatId() const override;
-		// End FTickableEditorObject
+		virtual bool IsTickableWhenPaused() const override { return true; }
+		virtual bool IsTickableInEditor() const override { return true; }
+		// End FTickableGameObject
 
 		// Begin FTickableCookObject
 		virtual void TickCook(float DeltaTime, bool bCookCompete) override;
diff --git a/Engine/Plugins/Experimental/Animation/PoseSearch/Source/Runtime/Public/PoseSearch/PoseSearchIndex.h b/Engine/Plugins/Experimental/Animation/PoseSearch/Source/Runtime/Public/PoseSearch/PoseSearchIndex.h
--- a/Engine/Plugins/Experimental/Animation/PoseSearch/Source/Runtime/Public/PoseSearch/PoseSearchIndex.h
+++ b/Engine/Plugins/Experimental/Animation/PoseSearch/Source/Runtime/Public/PoseSearch/PoseSearchIndex.h
@@ -64,6 +64,15 @@ public:
 		ValueOffset = Value;
 	}
 
+	bool operator==(const FPoseMetadata& Other) const
+	{
+		return
+			ValueOffset == Other.ValueOffset &&
+			AssetIndex == Other.AssetIndex &&
+			bBlockTransition == Other.bBlockTransition &&
+			CostAddend == Other.CostAddend;
+	}
+
 	friend FArchive& operator<<(FArchive& Ar, FPoseMetadata& Metadata);
 };
 
@@ -339,9 +348,7 @@ struct FSearchIndexBase
 		return MakeArrayView(&Values[ValueOffset], DataCardinality);
 	}
 
-#if ENABLE_ANIM_DEBUG
-	bool Compare(const FSearchIndexBase& Other) const;
-#endif // ENABLE_ANIM_DEBUG
+	bool operator==(const FSearchIndexBase& Other) const;
 
 	friend FArchive& operator<<(FArchive& Ar, FSearchIndexBase& Index);
 };
@@ -394,11 +401,8 @@ struct FSearchIndex : public FSearchIndexBase
 	// returns the inverse mapping of PCAValuesVectorToPoseIndexes
 	POSESEARCH_API void GetPoseToPCAValuesVectorIndexes(TArray<uint32>& PoseToPCAValuesVectorIndexes) const;
 
+	bool operator==(const FSearchIndex& Other) const;
 	friend FArchive& operator<<(FArchive& Ar, FSearchIndex& Index);
-
-#if ENABLE_ANIM_DEBUG
-	bool Compare(const FSearchIndex& Other) const;
-#endif // ENABLE_ANIM_DEBUG
 };
 
 } // namespace UE::PoseSearch
