motion matching
- fix for FAsyncPoseSearchDatabasesManagement::Shutdown crash if Tasks were flying
- CachedEntry now storing FCSPose<FCompactHeapPose> instead of FCSPose<FCompactPose>
- FMemMark Mark(FMemStack::Get()) usage cleanup in pose search
- moved FDatabaseIndexingContext::IndexDatabase into static funtion IndexDatabase in PoseSearchDerivedData.cpp

diff --git a/Engine/Plugins/Experimental/Animation/PoseSearch/Source/Editor/Private/PoseSearchDatabasePreviewScene.cpp b/Engine/Plugins/Experimental/Animation/PoseSearch/Source/Editor/Private/PoseSearchDatabasePreviewScene.cpp
--- a/Engine/Plugins/Experimental/Animation/PoseSearch/Source/Editor/Private/PoseSearchDatabasePreviewScene.cpp
+++ b/Engine/Plugins/Experimental/Animation/PoseSearch/Source/Editor/Private/PoseSearchDatabasePreviewScene.cpp
@@ -66,11 +66,6 @@ namespace UE::PoseSearch
 		if (ViewModel->IsPoseFeaturesDrawMode(EFeaturesDrawMode::All | EFeaturesDrawMode::Detailed) && !ViewModel->GetPreviewActors().IsEmpty() &&
 			FAsyncPoseSearchDatabasesManagement::RequestAsyncBuildIndex(Database, ERequestAsyncBuildFlag::ContinueRequest))
 		{
-#if ENABLE_ANIM_DEBUG
-			// memory marker required for PreviewActor.Sampler.ExtractPose
-			FMemMark Mark(FMemStack::Get());
-#endif // ENABLE_ANIM_DEBUG
-
 			bool bDrawQueryVector = ViewModel->ShouldDrawQueryVector();
 			for (FDatabasePreviewActor& PreviewActor : ViewModel->GetPreviewActors())
 			{
@@ -92,6 +87,7 @@ namespace UE::PoseSearch
 						if (DebugDrawSamplerSize > UE_KINDA_SMALL_NUMBER)
 						{
 							// drawing the pose extracted from the Sampler to visually compare with the pose features and the mesh drawing
+							FMemMark Mark(FMemStack::Get());
 							FCompactPose Pose;
 							Pose.SetBoneContainer(&PreviewActor.GetAnimPreviewInstance()->GetRequiredBonesOnAnyThread());
 							PreviewActor.Sampler.ExtractPose(PreviewActor.CurrentTime, Pose);
@@ -99,9 +95,9 @@ namespace UE::PoseSearch
 							const FTransform RootTransform = PreviewActor.Sampler.ExtractRootTransform(PreviewActor.CurrentTime);
 
 							FCSPose<FCompactPose> ComponentSpacePose;
-							ComponentSpacePose.InitPose(Pose);
+							ComponentSpacePose.InitPose(MoveTemp(Pose));
 
-							for (int32 BoneIndex = 0; BoneIndex < Pose.GetNumBones(); ++BoneIndex)
+							for (int32 BoneIndex = 0; BoneIndex < ComponentSpacePose.GetPose().GetNumBones(); ++BoneIndex)
 							{
 								const FTransform BoneWorldTransforms = ComponentSpacePose.GetComponentSpaceTransform(FCompactPoseBoneIndex(BoneIndex)) * RootTransform;
 								DrawParams.DrawPoint(BoneWorldTransforms.GetTranslation(), FColor::Red, DebugDrawSamplerSize);
diff --git a/Engine/Plugins/Experimental/Animation/PoseSearch/Source/Editor/Private/PoseSearchMeshComponent.cpp b/Engine/Plugins/Experimental/Animation/PoseSearch/Source/Editor/Private/PoseSearchMeshComponent.cpp
--- a/Engine/Plugins/Experimental/Animation/PoseSearch/Source/Editor/Private/PoseSearchMeshComponent.cpp
+++ b/Engine/Plugins/Experimental/Animation/PoseSearch/Source/Editor/Private/PoseSearchMeshComponent.cpp
@@ -134,17 +134,13 @@ void UPoseSearchMeshComponent::UpdatePose(const FUpdateContext& UpdateContext)
 	}
 
 	FCSPose<FCompactPose> ComponentSpacePose;
-	ComponentSpacePose.InitPose(CompactPose);
+	ComponentSpacePose.InitPose(MoveTemp(CompactPose));
 
 	for (const FBoneIndexType BoneIndex : RequiredBones.GetBoneIndicesArray())
 	{
-		const FTransform BoneTransform = 
-			ComponentSpacePose.GetComponentSpaceTransform(FCompactPoseBoneIndex(BoneIndex));
-
-		FSkeletonPoseBoneIndex SkeletonBoneIndex =
-			RequiredBones.GetSkeletonPoseIndexFromCompactPoseIndex(FCompactPoseBoneIndex(BoneIndex));
-		FName BoneName = 
-			RequiredBones.GetSkeletonAsset()->GetReferenceSkeleton().GetBoneName(SkeletonBoneIndex.GetInt());
+		const FTransform BoneTransform = ComponentSpacePose.GetComponentSpaceTransform(FCompactPoseBoneIndex(BoneIndex));
+		FSkeletonPoseBoneIndex SkeletonBoneIndex = RequiredBones.GetSkeletonPoseIndexFromCompactPoseIndex(FCompactPoseBoneIndex(BoneIndex));
+		FName BoneName = RequiredBones.GetSkeletonAsset()->GetReferenceSkeleton().GetBoneName(SkeletonBoneIndex.GetInt());
 		SetBoneTransformByName(BoneName, BoneTransform, EBoneSpaces::ComponentSpace);
 	}
 
diff --git a/Engine/Plugins/Experimental/Animation/PoseSearch/Source/Runtime/Private/PoseSearchAssetIndexer.cpp b/Engine/Plugins/Experimental/Animation/PoseSearch/Source/Runtime/Private/PoseSearchAssetIndexer.cpp
--- a/Engine/Plugins/Experimental/Animation/PoseSearch/Source/Runtime/Private/PoseSearchAssetIndexer.cpp
+++ b/Engine/Plugins/Experimental/Animation/PoseSearch/Source/Runtime/Private/PoseSearchAssetIndexer.cpp
@@ -7,11 +7,11 @@
 #include "Animation/MirrorDataTable.h"
 #include "PoseSearch/PoseSearchAnimNotifies.h"
 #include "PoseSearch/PoseSearchAssetSampler.h"
+#include "PoseSearch/PoseSearchContext.h"
 #include "PoseSearch/PoseSearchDefines.h"
 #include "PoseSearch/PoseSearchFeatureChannel.h"
 #include "PoseSearch/PoseSearchIndex.h"
 #include "PoseSearch/PoseSearchSchema.h"
-#include "PoseSearchDatabaseIndexingContext.h"
 
 namespace UE::PoseSearch
 {
@@ -78,6 +78,28 @@ static FSamplingParam WrapOrClampSamplingParam(bool bCanWrap, float SamplingPara
 	return Result;
 }
 
+//////////////////////////////////////////////////////////////////////////
+// FAssetSamplingContext
+void FAssetSamplingContext::Init(const UMirrorDataTable* InMirrorDataTable, const FBoneContainer& BoneContainer)
+{
+	MirrorDataTable = InMirrorDataTable;
+
+	if (InMirrorDataTable)
+	{
+		InMirrorDataTable->FillCompactPoseAndComponentRefRotations(BoneContainer, CompactPoseMirrorBones, ComponentSpaceRefRotations);
+	}
+	else
+	{
+		CompactPoseMirrorBones.Reset();
+		ComponentSpaceRefRotations.Reset();
+	}
+}
+
+FTransform FAssetSamplingContext::MirrorTransform(const FTransform& InTransform) const
+{
+	return UE::PoseSearch::MirrorTransform(InTransform, MirrorDataTable->MirrorAxis, ComponentSpaceRefRotations[FCompactPoseBoneIndex(RootBoneIndexType)]);
+}
+
 //////////////////////////////////////////////////////////////////////////
 // FAssetIndexer
 FAssetIndexer::FAssetIndexer(const FBoneContainer& InBoneContainer, const FSearchIndexAsset& InSearchIndexAsset, 
@@ -104,8 +126,6 @@ void FAssetIndexer::Process(int32 AssetIdx)
 {
 	check(Schema.IsValid());
 
-	FMemMark Mark(FMemStack::Get());
-
 	bProcessFailed = false;
 
 	// Generate pose metadata
@@ -310,6 +330,7 @@ FAssetIndexer::CachedEntry& FAssetIndexer::GetEntry(float SampleTime)
 			CurrentTime = FMath::Clamp(CurrentTime, 0.f, PlayLength);
 		}
 
+		FMemMark Mark(FMemStack::Get());
 		FCompactPose Pose;
 		Pose.SetBoneContainer(&BoneContainer);
 		AssetSampler.ExtractPose(CurrentTime, Pose);
@@ -354,7 +375,10 @@ FAssetIndexer::CachedEntry& FAssetIndexer::GetEntry(float SampleTime)
 			// Note curves and attributes are not used during the indexing process and therefore don't need to be mirrored
 		}
 
-		Entry->ComponentSpacePose.InitPose(MoveTemp(Pose));
+		FCSPose<FCompactPose> StackComponentSpacePose;
+		StackComponentSpacePose.InitPose(MoveTemp(Pose));
+		Entry->ComponentSpacePose.CopyPose(StackComponentSpacePose);
+
 		Entry->RootTransform = Sample.RootTransform;
 		Entry->bClamped = Sample.bClamped;
 	}
diff --git a/Engine/Plugins/Experimental/Animation/PoseSearch/Source/Runtime/Private/PoseSearchAssetSampler.cpp b/Engine/Plugins/Experimental/Animation/PoseSearch/Source/Runtime/Private/PoseSearchAssetSampler.cpp
--- a/Engine/Plugins/Experimental/Animation/PoseSearch/Source/Runtime/Private/PoseSearchAssetSampler.cpp
+++ b/Engine/Plugins/Experimental/Animation/PoseSearch/Source/Runtime/Private/PoseSearchAssetSampler.cpp
@@ -183,13 +183,17 @@ static void ProcessRootTransform(const UBlendSpace* BlendSpace, const FVector& B
 				BlendSample.Time = CurrentTime * Scale;
 			}
 
+			FMemMark Mark(FMemStack::Get());
 			FCompactPose Pose;
 			FBlendedCurve BlendedCurve;
 			UE::Anim::FStackAttributeContainer StackAttributeContainer;
 			FAnimationPoseData AnimPoseData(Pose, BlendedCurve, StackAttributeContainer);
 
-			Pose.SetBoneContainer(&BoneContainer);
-			BlendedCurve.InitFrom(BoneContainer);
+			// copying bone container to be thread safe, since it has mutable members
+			FBoneContainer BoneContainerCopy = BoneContainer;
+
+			Pose.SetBoneContainer(&BoneContainerCopy);
+			BlendedCurve.InitFrom(BoneContainerCopy);
 
 			BlendSpace->GetAnimationPose(BlendSamples, ExtractionCtx, AnimPoseData);
 
@@ -261,8 +265,6 @@ float FAnimationAssetSampler::GetPlayLength(const UAnimationAsset* AnimAsset, co
 	{
 		if (const UBlendSpace* BlendSpace = Cast<UBlendSpace>(AnimAsset))
 		{
-			FMemMark Mark(FMemStack::Get());
-
 			TArray<FBlendSampleData> BlendSamples;
 			int32 TriangulationIndex = 0;
 			if (BlendSpace->GetSamplesFromBlendInput(BlendParameters, BlendSamples, TriangulationIndex, true))
@@ -419,16 +421,11 @@ void FAnimationAssetSampler::ExtractPose(const FAnimExtractContext& ExtractionCt
 
 void FAnimationAssetSampler::ExtractPose(float Time, FCompactPose& OutPose) const
 {
-	using namespace UE::Anim;
-
+	UE::Anim::FStackAttributeContainer UnusedAtrribute;
 	FBlendedCurve UnusedCurve;
-	FStackAttributeContainer UnusedAtrribute;
+	UnusedCurve.InitFrom(OutPose.GetBoneContainer());
 	FAnimationPoseData AnimPoseData = { OutPose, UnusedCurve, UnusedAtrribute };
 
-	check(OutPose.IsValid());
-	FBoneContainer& BoneContainer = OutPose.GetBoneContainer();
-	UnusedCurve.InitFrom(BoneContainer);
-
 	FDeltaTimeRecord DeltaTimeRecord;
 	DeltaTimeRecord.Set(Time, 0.f);
 	FAnimExtractContext ExtractionCtx(double(Time), false, DeltaTimeRecord, false);
@@ -580,7 +577,6 @@ void FAnimationAssetSampler::Process(const FBoneContainer& BoneContainer)
 {
 	if (const UBlendSpace* BlendSpace = Cast<UBlendSpace>(AnimationAsset.Get()))
 	{
-		FMemMark Mark(FMemStack::Get());
 		ProcessRootTransform(BlendSpace, BlendParameters, CachedPlayLength, BoneContainer, RootTransformSamplingRate, IsLoopable(), AccumulatedRootTransform);
 	}
 }
diff --git a/Engine/Plugins/Experimental/Animation/PoseSearch/Source/Runtime/Private/PoseSearchDatabaseIndexingContext.cpp b/Engine/Plugins/Experimental/Animation/PoseSearch/Source/Runtime/Private/PoseSearchDatabaseIndexingContext.cpp
deleted file mode 100644
diff --git a/Engine/Plugins/Experimental/Animation/PoseSearch/Source/Runtime/Private/PoseSearchDatabaseIndexingContext.h b/Engine/Plugins/Experimental/Animation/PoseSearch/Source/Runtime/Private/PoseSearchDatabaseIndexingContext.h
deleted file mode 100644
diff --git a/Engine/Plugins/Experimental/Animation/PoseSearch/Source/Runtime/Private/PoseSearchDerivedData.cpp b/Engine/Plugins/Experimental/Animation/PoseSearch/Source/Runtime/Private/PoseSearchDerivedData.cpp
--- a/Engine/Plugins/Experimental/Animation/PoseSearch/Source/Runtime/Private/PoseSearchDerivedData.cpp
+++ b/Engine/Plugins/Experimental/Animation/PoseSearch/Source/Runtime/Private/PoseSearchDerivedData.cpp
@@ -13,7 +13,6 @@
 #include "DerivedDataRequestOwner.h"
 #include "InstancedStruct.h"
 #include "Misc/CoreDelegates.h"
-#include "PoseSearchDatabaseIndexingContext.h"
 #include "PoseSearch/PoseSearchAnimNotifies.h"
 #include "PoseSearch/PoseSearchAssetIndexer.h"
 #include "PoseSearch/PoseSearchDatabase.h"
@@ -677,23 +676,179 @@ static void PreprocessSearchIndexKDTree(FSearchIndex& SearchIndex, const UPoseSe
 	}
 }
 
-#if ENABLE_ANIM_DEBUG
-
-static void CompareIndexingContext(const FDatabaseIndexingContext& A, const FDatabaseIndexingContext& B)
+static bool IndexDatabase(FSearchIndexBase& SearchIndexBase, const UPoseSearchDatabase& Database, UE::DerivedData::FRequestOwner& Owner)
 {
-	if (A.GetIndexers().Num() != B.GetIndexers().Num())
+	const UPoseSearchSchema* Schema = Database.Schema;
+	check(Schema);
+
+	FBoneContainer BoneContainer;
+	BoneContainer.InitializeTo(Schema->BoneIndicesWithParents, UE::Anim::FCurveFilterSettings(UE::Anim::ECurveFilterMode::DisallowAll), *Schema->Skeleton);
+
+	FAssetSamplingContext SamplingContext;
+	SamplingContext.Init(Schema->MirrorDataTable, BoneContainer);
+
+	if (Owner.IsCanceled())
 	{
-		UE_LOG(LogPoseSearch, Warning, TEXT("CompareIndexers - FAssetIndexer is not deterministic"));
+		return false;
 	}
-	else
+
+	// Prepare samplers for all animation assets.
+	TArray<FAnimationAssetSampler> Samplers;
+	Samplers.Reserve(256);
+	
+	TMap<TPair<const UAnimationAsset*, FVector>, int32> SamplerMap;
+	SamplerMap.Reserve(256);
+
+	for (const FInstancedStruct& DatabaseAssetStruct : Database.AnimationAssets)
 	{
-		for (int32 Index = 0; Index < A.GetIndexers().Num(); ++Index)
+		if (const FPoseSearchDatabaseBlendSpace* DatabaseBlendSpace = DatabaseAssetStruct.GetPtr<FPoseSearchDatabaseBlendSpace>())
+		{
+			if (DatabaseBlendSpace->BlendSpace)
+			{
+				int32 HorizontalBlendNum, VerticalBlendNum;
+				DatabaseBlendSpace->GetBlendSpaceParameterSampleRanges(HorizontalBlendNum, VerticalBlendNum);
+
+				for (int32 HorizontalIndex = 0; HorizontalIndex < HorizontalBlendNum; HorizontalIndex++)
+				{
+					for (int32 VerticalIndex = 0; VerticalIndex < VerticalBlendNum; VerticalIndex++)
+					{
+						const FVector BlendParameters = DatabaseBlendSpace->BlendParameterForSampleRanges(HorizontalIndex, VerticalIndex);
+
+						if (!SamplerMap.Contains({ DatabaseBlendSpace->BlendSpace, BlendParameters }))
+						{
+							SamplerMap.Add({ DatabaseBlendSpace->BlendSpace, BlendParameters }, Samplers.Num());
+							Samplers.Emplace(DatabaseBlendSpace->BlendSpace, BlendParameters);
+						}
+					}
+				}
+			}
+		}
+		else if (const FPoseSearchDatabaseAnimationAssetBase* DatabaseAnimationAssetBase = DatabaseAssetStruct.GetPtr<FPoseSearchDatabaseAnimationAssetBase>())
 		{
-			A.GetIndexers()[Index].CompareCachedEntries(B.GetIndexers()[Index]);
+			if (const UAnimationAsset* AnimationAsset = DatabaseAnimationAssetBase->GetAnimationAsset())
+			{
+				if (!SamplerMap.Contains({ AnimationAsset, FVector::ZeroVector }))
+				{
+					SamplerMap.Add({ AnimationAsset, FVector::ZeroVector }, Samplers.Num());
+					Samplers.Emplace(AnimationAsset);
+				}
+			}
 		}
 	}
+
+	ParallelFor(Samplers.Num(), [&Samplers, &BoneContainer](int32 SamplerIdx) { Samplers[SamplerIdx].Process(BoneContainer); }, ParallelForFlags);
+
+	if (Owner.IsCanceled())
+	{
+		return false;
+	}
+
+	// prepare indexers
+	TArray<FAssetIndexer> Indexers;
+	Indexers.Reserve(SearchIndexBase.Assets.Num());
+
+	int32 TotalPoses = 0;
+	for (int32 AssetIdx = 0; AssetIdx != SearchIndexBase.Assets.Num(); ++AssetIdx)
+	{
+		FSearchIndexAsset& SearchIndexAsset = SearchIndexBase.Assets[AssetIdx];
+		check(SearchIndexAsset.FirstPoseIdx == TotalPoses);
+
+		const FPoseSearchDatabaseAnimationAssetBase* DatabaseAnimationAssetBase = Database.GetAnimationAssetStruct(SearchIndexAsset).GetPtr<FPoseSearchDatabaseAnimationAssetBase>();
+		check(DatabaseAnimationAssetBase && DatabaseAnimationAssetBase->GetAnimationAsset());
+		const FAnimationAssetSampler& AssetSampler = Samplers[SamplerMap[{ DatabaseAnimationAssetBase->GetAnimationAsset(), SearchIndexAsset.BlendParameters }]];
+
+		Indexers.Emplace(BoneContainer, SearchIndexAsset, SamplingContext, *Schema, AssetSampler);
+		TotalPoses += SearchIndexAsset.GetNumPoses();
+	}
+
+	// allocating Values and PoseMetadata
+	SearchIndexBase.AllocateData(Schema->SchemaCardinality, TotalPoses);
+
+	// assigning local data to each Indexer
+	TotalPoses = 0;
+	for (int32 AssetIdx = 0; AssetIdx != SearchIndexBase.Assets.Num(); ++AssetIdx)
+	{
+		Indexers[AssetIdx].AssignWorkingData(TotalPoses, SearchIndexBase.Values, SearchIndexBase.PoseMetadata);
+		TotalPoses += Indexers[AssetIdx].GetNumIndexedPoses();
+	}
+
+	if (Owner.IsCanceled())
+	{
+		return false;
+	}
+	
+	// Index asset data
+	ParallelFor(Indexers.Num(), [&Indexers](int32 AssetIdx) { Indexers[AssetIdx].Process(AssetIdx); }, ParallelForFlags);
+
+	for (const FAssetIndexer& Indexer : Indexers)
+	{
+		if (Indexer.IsProcessFailed())
+		{
+			return false;
+		}
+	}
+
+	if (Owner.IsCanceled())
+	{
+		return false;
+	}
+
+	// Joining Metadata.Flags into OverallFlags
+	SearchIndexBase.bAnyBlockTransition = false;
+	for (const FPoseMetadata& Metadata : SearchIndexBase.PoseMetadata)
+	{
+		if (Metadata.IsBlockTransition())
+		{
+			SearchIndexBase.bAnyBlockTransition = true;
+			break;
+		}
+	}
+
+	// Joining Stats
+	int32 NumAccumulatedSamples = 0;
+	SearchIndexBase.Stats = FSearchStats();
+	for (int32 AssetIdx = 0; AssetIdx != SearchIndexBase.Assets.Num(); ++AssetIdx)
+	{
+		const FAssetIndexer::FStats& Stats = Indexers[AssetIdx].GetStats();
+		SearchIndexBase.Stats.AverageSpeed += Stats.AccumulatedSpeed;
+		SearchIndexBase.Stats.MaxSpeed = FMath::Max(SearchIndexBase.Stats.MaxSpeed, Stats.MaxSpeed);
+		SearchIndexBase.Stats.AverageAcceleration += Stats.AccumulatedAcceleration;
+		SearchIndexBase.Stats.MaxAcceleration = FMath::Max(SearchIndexBase.Stats.MaxAcceleration, Stats.MaxAcceleration);
+
+		NumAccumulatedSamples += Stats.NumAccumulatedSamples;
+	}
+
+	if (NumAccumulatedSamples > 0)
+	{
+		const float Denom = 1.f / float(NumAccumulatedSamples);
+		SearchIndexBase.Stats.AverageSpeed *= Denom;
+		SearchIndexBase.Stats.AverageAcceleration *= Denom;
+	}
+
+	// Calculate Min Cost Addend
+	SearchIndexBase.MinCostAddend = 0.f;
+	if (!SearchIndexBase.PoseMetadata.IsEmpty())
+	{
+		SearchIndexBase.MinCostAddend = MAX_FLT;
+		for (const FPoseMetadata& PoseMetadata : SearchIndexBase.PoseMetadata)
+		{
+			if (PoseMetadata.GetCostAddend() < SearchIndexBase.MinCostAddend)
+			{
+				SearchIndexBase.MinCostAddend = PoseMetadata.GetCostAddend();
+			}
+		}
+	}
+
+	if (Owner.IsCanceled())
+	{
+		return false;
+	}
+
+	return true;
 }
 
+#if ENABLE_ANIM_DEBUG
+
 static void CompareChannelValues(int32 RecursionIndex, int32 PoseIndex, const TConstArrayView<float>& PoseA, const TConstArrayView<float>& PoseB, const TConstArrayView<TObjectPtr<UPoseSearchFeatureChannel>>& Channels, FStringBuilderBase& StringBuilder)
 {
 	bool bPrintHeader = true;
@@ -1217,8 +1372,7 @@ void FPoseSearchDatabaseAsyncCacheTask::OnGetComplete(UE::DerivedData::FCacheGet
 						return;
 					}
 
-					FDatabaseIndexingContext DbIndexingContext;
-					if (!DbIndexingContext.IndexDatabase(SearchIndexBase, *IndexBaseDatabase, Owner))
+					if (!IndexDatabase(SearchIndexBase, *IndexBaseDatabase, Owner))
 					{
 						UE_LOG(LogPoseSearch, Log, TEXT("%s - %s BuildIndex Cancelled"), *LexToString(FullIndexKey.Hash), *IndexBaseDatabases[0]->GetName());
 						SearchIndex.Reset();
@@ -1235,16 +1389,13 @@ void FPoseSearchDatabaseAsyncCacheTask::OnGetComplete(UE::DerivedData::FCacheGet
 						for (int32 Iteration = 0; Iteration < NumIterations; ++Iteration)
 						{
 							FSearchIndexBase TestSearchIndexBase = SearchIndexBase;
-							FDatabaseIndexingContext TestDbIndexingContext;
-							if (TestDbIndexingContext.IndexDatabase(TestSearchIndexBase, *IndexBaseDatabase, Owner))
+							if (IndexDatabase(TestSearchIndexBase, *IndexBaseDatabase, Owner))
 							{
 								if (TestSearchIndexBase != SearchIndexBase)
 								{
 									FStringBuilderBase Message;
 									CompareSearchIndexBase(TestSearchIndexBase, SearchIndexBase, IndexBaseDatabase->Schema, Message);
 									UE_LOG(LogPoseSearch, Warning, TEXT("OnGetComplete - IndexDatabase is not deterministic\n%s"), *Message);
-
-									CompareIndexingContext(TestDbIndexingContext, DbIndexingContext);
 								}
 							}
 						}
@@ -1496,6 +1647,8 @@ void FAsyncPoseSearchDatabasesManagement::Shutdown()
 {
 	FScopeLock Lock(&Mutex);
 
+	Tasks.Reset();
+
 	FCoreUObjectDelegates::OnObjectModified.Remove(OnObjectModifiedHandle);
 	OnObjectModifiedHandle.Reset();
 
diff --git a/Engine/Plugins/Experimental/Animation/PoseSearch/Source/Runtime/Public/PoseSearch/PoseSearchAssetIndexer.h b/Engine/Plugins/Experimental/Animation/PoseSearch/Source/Runtime/Public/PoseSearch/PoseSearchAssetIndexer.h
--- a/Engine/Plugins/Experimental/Animation/PoseSearch/Source/Runtime/Public/PoseSearch/PoseSearchAssetIndexer.h
+++ b/Engine/Plugins/Experimental/Animation/PoseSearch/Source/Runtime/Public/PoseSearch/PoseSearchAssetIndexer.h
@@ -14,10 +14,28 @@ namespace UE::PoseSearch
 {
 
 struct FAnimationAssetSampler;
-struct FAssetSamplingContext;
 struct FSearchIndexAsset;
 struct FPoseMetadata;
 
+struct FAssetSamplingContext
+{
+	// Time delta used for computing pose derivatives
+	static constexpr float FiniteDelta = 1 / 60.0f;
+
+	// Mirror data table pointer copied from Schema for convenience
+	TObjectPtr<const UMirrorDataTable> MirrorDataTable;
+
+	// Compact pose format of Mirror Bone Map
+	TCustomBoneIndexArray<FCompactPoseBoneIndex, FCompactPoseBoneIndex> CompactPoseMirrorBones;
+
+	// Pre-calculated component space rotations of reference pose, which allows mirror to work with any joint orientation
+	// Only initialized and used when a mirroring table is specified
+	TCustomBoneIndexArray<FQuat, FCompactPoseBoneIndex> ComponentSpaceRefRotations;
+
+	void Init(const UMirrorDataTable* InMirrorDataTable, const FBoneContainer& BoneContainer);
+	FTransform MirrorTransform(const FTransform& Transform) const;
+};
+
 class FAssetIndexer
 {
 public:
@@ -97,7 +115,7 @@ private:
 		bool bClamped = false;
 
 		FTransform RootTransform;
-		FCSPose<FCompactPose> ComponentSpacePose;
+		FCSPose<FCompactHeapPose> ComponentSpacePose;
 	};
 
 	FSampleInfo GetSampleInfo(float SampleTime) const;
