motion matching - fix for invalid Database in FPoseSearchDatabaseAsyncCacheTask::CancelIfDependsOn when Database has been garbage collected, but the task hasn't been disposed by FAsyncPoseSearchDatabasesManagement::Tick

diff --git a/Engine/Plugins/Experimental/Animation/PoseSearch/Source/Runtime/Private/PoseSearchDerivedData.cpp b/Engine/Plugins/Experimental/Animation/PoseSearch/Source/Runtime/Private/PoseSearchDerivedData.cpp
--- a/Engine/Plugins/Experimental/Animation/PoseSearch/Source/Runtime/Private/PoseSearchDerivedData.cpp
+++ b/Engine/Plugins/Experimental/Animation/PoseSearch/Source/Runtime/Private/PoseSearchDerivedData.cpp
@@ -1136,7 +1136,10 @@ bool FPoseSearchDatabaseAsyncCacheTask::CancelIfDependsOn(const UObject* Object)
 		// DatabaseDependencies is updated only in StartNewRequestIfNeeded when there are no active requests, so it's thread safe to access it 
 		if (DatabaseDependencies.Contains(Object))
 		{
-			UE_LOG(LogPoseSearch, Log, TEXT("%s - %s Cancelled because of %s"), *LexToString(DerivedDataKey), *Database->GetName(), *Object->GetName());
+			// Database can be null if the task was Ended/Failed and Database was already garbage collected, but Tick hasn't been called yet
+			FString DatabaseName = IsValid() ? *Database->GetName() : TEXT("Garbage Collected Database");
+			UE_LOG(LogPoseSearch, Log, TEXT("%s - %s Cancelled because of %s"), *LexToString(DerivedDataKey), *DatabaseName, *Object->GetName());
+
 			Cancel();
 			return true;
 		}
