motion matching - matching in root bone space fix

diff --git a/Engine/Plugins/Experimental/Animation/PoseSearch/Source/Runtime/Private/AnimNode_MotionMatching.cpp b/Engine/Plugins/Experimental/Animation/PoseSearch/Source/Runtime/Private/AnimNode_MotionMatching.cpp
--- a/Engine/Plugins/Experimental/Animation/PoseSearch/Source/Runtime/Private/AnimNode_MotionMatching.cpp
+++ b/Engine/Plugins/Experimental/Animation/PoseSearch/Source/Runtime/Private/AnimNode_MotionMatching.cpp
@@ -64,9 +64,14 @@ void FAnimNode_MotionMatching::Evaluate_AnyThread(FPoseContext& Output)
 	{
 		const FQuat RootBoneDelta(FRotator(0.f, MotionMatchingState.ComponentDeltaYaw, 0.f));
 		FCompactPoseBoneIndex RootBoneIndex(RootBoneIndexType);
-		Output.Pose[RootBoneIndex].SetRotation(Output.Pose[RootBoneIndex].GetRotation() * RootBoneDelta);
-		Output.Pose[RootBoneIndex].NormalizeRotation();
 
+		// @todo: this code assumes the root bone from animation is always identity. Implement if not the case as
+		//		Output.Pose[RootBoneIndex].SetRotation(Output.Pose[RootBoneIndex].GetRotation() * RootBoneDelta);
+		//		Output.Pose[RootBoneIndex].NormalizeRotation();
+		// 		etc etc
+		check(Output.Pose[RootBoneIndex].GetRotation().IsIdentity());
+
+		Output.Pose[RootBoneIndex].SetRotation(RootBoneDelta);
 		RootMotionTransformDelta.SetTranslation(RootBoneDelta.RotateVector(RootMotionTransformDelta.GetTranslation()));
 
 		if (RootMotionProvider)
diff --git a/Engine/Plugins/Experimental/Animation/PoseSearch/Source/Runtime/Private/AnimNode_PoseSearchHistoryCollector.cpp b/Engine/Plugins/Experimental/Animation/PoseSearch/Source/Runtime/Private/AnimNode_PoseSearchHistoryCollector.cpp
--- a/Engine/Plugins/Experimental/Animation/PoseSearch/Source/Runtime/Private/AnimNode_PoseSearchHistoryCollector.cpp
+++ b/Engine/Plugins/Experimental/Animation/PoseSearch/Source/Runtime/Private/AnimNode_PoseSearchHistoryCollector.cpp
@@ -71,6 +71,17 @@ void FAnimNode_PoseSearchHistoryCollector_Base::CacheBones_AnyThread(const FAnim
 	}
 
 	PoseHistory.Init(PoseCount, PoseDuration, RequiredBones);
+
+	if (bInitializeWithRefPose)
+	{
+		// initializing PoseHistory with a ref pose at FAnimInstanceProxy location/facing
+		FMemMark Mark(FMemStack::Get());
+		FCompactPose Pose;
+		Pose.SetBoneContainer(&Context.AnimInstanceProxy->GetRequiredBones());
+		FCSPose<FCompactPose> ComponentSpacePose;
+		ComponentSpacePose.InitPose(Pose);
+		PoseHistory.Update(0.f, ComponentSpacePose, Context.AnimInstanceProxy->GetComponentTransform());
+	}
 }
 
 /////////////////////////////////////////////////////
diff --git a/Engine/Plugins/Experimental/Animation/PoseSearch/Source/Runtime/Private/PoseSearchAssetIndexer.cpp b/Engine/Plugins/Experimental/Animation/PoseSearch/Source/Runtime/Private/PoseSearchAssetIndexer.cpp
--- a/Engine/Plugins/Experimental/Animation/PoseSearch/Source/Runtime/Private/PoseSearchAssetIndexer.cpp
+++ b/Engine/Plugins/Experimental/Animation/PoseSearch/Source/Runtime/Private/PoseSearchAssetIndexer.cpp
@@ -397,6 +397,8 @@ FTransform FAssetIndexer::GetComponentSpaceTransform(float SampleTime, bool& bCl
 		return CalculateComponentSpaceTransform(Entry, SchemaBoneIdx);
 	}
 
+	// making sure the root bone transform is identity
+	check(CalculateComponentSpaceTransform(Entry, SchemaBoneIdx).Equals(FTransform::Identity));
 	return FTransform::Identity;
 }
 
diff --git a/Engine/Plugins/Experimental/Animation/PoseSearch/Source/Runtime/Private/PoseSearchContext.cpp b/Engine/Plugins/Experimental/Animation/PoseSearch/Source/Runtime/Private/PoseSearchContext.cpp
--- a/Engine/Plugins/Experimental/Animation/PoseSearch/Source/Runtime/Private/PoseSearchContext.cpp
+++ b/Engine/Plugins/Experimental/Animation/PoseSearch/Source/Runtime/Private/PoseSearchContext.cpp
@@ -415,11 +415,11 @@ FVector FSearchContext::GetSampleVelocity(float SampleTimeOffset, float OriginTi
 
 	if (SampleBoneVelocityWorldOverride)
 	{
-		const FTransform RootBoneTransform = GetTransform(OriginTime, Schema, RootSchemaBoneIdx, bUseHistoryRoot);
+		const FTransform RootBoneTransform = GetWorldBoneTransformAtTime(OriginTime, Schema, RootSchemaBoneIdx, bUseHistoryRoot);
 		return RootBoneTransform.InverseTransformVector(*SampleBoneVelocityWorldOverride);
 	}
 
-	// calculating the Position in component space for the bone indexed by SchemaSampleBoneIdx
+	// calculating the local Position for the bone indexed by SchemaSampleBoneIdx
 	const FVector PreviousTranslation = GetSamplePositionInternal(SampleTime - FiniteDelta, bUseCharacterSpaceVelocities ? OriginTime - FiniteDelta : OriginTime, Schema, SchemaSampleBoneIdx, SchemaOriginBoneIdx, bUseHistoryRoot);
 	const FVector CurrentTranslation = GetSamplePositionInternal(SampleTime, OriginTime, Schema, SchemaSampleBoneIdx, SchemaOriginBoneIdx, bUseHistoryRoot);
 
@@ -427,39 +427,40 @@ FVector FSearchContext::GetSampleVelocity(float SampleTimeOffset, float OriginTi
 	return LinearVelocity;
 }
 
-FTransform FSearchContext::GetRootAtTime(float Time, bool bUseHistoryRoot, bool bExtrapolate) const
+FTransform FSearchContext::GetWorldRootBoneTransformAtTime(float SampleTime, bool bUseHistoryRoot, bool bExtrapolate) const
 {
-	FTransform RootTransform = FTransform::Identity;
+	FTransform WorldRootBoneTransform = FTransform::Identity;
 	if (bUseHistoryRoot)
 	{
 		check(History);
-		History->GetRootTransformAtTime(Time, RootTransform, bExtrapolate);
+		// asking for root bone (RootBoneIndexType) in world space (WorldSpaceIndexType)
+		History->GetTransformAtTime(SampleTime, WorldRootBoneTransform, AnimInstance->CurrentSkeleton, RootBoneIndexType, WorldSpaceIndexType, bExtrapolate);
 	}
 	else
 	{
 		check(Trajectory);
-		RootTransform = Trajectory->GetSampleAtTime(Time, bExtrapolate).GetTransform();
+		// Trajectory is already in root bone world space (transformed in UPoseSearchLibrary::ProcessTrajectory), so we just ask for a sample at the proper SampleTime 
+		WorldRootBoneTransform = Trajectory->GetSampleAtTime(SampleTime, bExtrapolate).GetTransform();
 	}
 
-	return RootTransform;
+	return WorldRootBoneTransform;
 }
 
-FTransform FSearchContext::GetTransform(float SampleTime, const UPoseSearchSchema* Schema, int8 SchemaBoneIdx, bool bUseHistoryRoot)
+FTransform FSearchContext::GetWorldBoneTransformAtTime(float SampleTime, const UPoseSearchSchema* Schema, int8 SchemaBoneIdx, bool bUseHistoryRoot)
 {
-	const FTransform RootTransform = GetRootAtTime(SampleTime, bUseHistoryRoot);
+	const FTransform WorldRootBoneTransform = GetWorldRootBoneTransformAtTime(SampleTime, bUseHistoryRoot);
 
 	const FBoneIndexType BoneIndexType = Schema->GetBoneIndexType(SchemaBoneIdx);
 	if (BoneIndexType != RootBoneIndexType)
 	{
-		const FTransform BoneTransform = GetComponentSpaceTransform(SampleTime, Schema, SchemaBoneIdx);
-		return BoneTransform * RootTransform;
+		const FTransform BoneTransform = GetLocalBoneTransformAtTime(SampleTime, Schema, SchemaBoneIdx);
+		return BoneTransform * WorldRootBoneTransform;
 	}
 
-	return RootTransform;
+	return WorldRootBoneTransform;
 }
 
-// returns the component space transform of the bone SchemaBoneIdx at time SampleTime
-FTransform FSearchContext::GetComponentSpaceTransform(float SampleTime, const UPoseSearchSchema* Schema, int8 SchemaBoneIdx)
+FTransform FSearchContext::GetLocalBoneTransformAtTime(float SampleTime, const UPoseSearchSchema* Schema, int8 SchemaBoneIdx)
 {
 	check(Schema);
 
@@ -473,24 +474,24 @@ FTransform FSearchContext::GetComponentSpaceTransform(float SampleTime, const UP
 	
 		// collecting the local bone transforms from the IPoseHistory
 		check(History);
-		FTransform BoneComponentSpaceTransform;
-		if (!History->GetComponentSpaceTransformAtTime(SampleTime, BoneIndexType, Schema->Skeleton, BoneComponentSpaceTransform))
+		FTransform LocalBoneTransform;
+		if (!History->GetTransformAtTime(SampleTime, LocalBoneTransform, Schema->Skeleton, BoneIndexType, RootBoneIndexType))
 		{
 			if (const USkeleton* Skeleton = Schema->Skeleton)
 			{
 				if (!History->IsEmpty())
 				{
-					UE_LOG(LogPoseSearch, Warning, TEXT("FSearchContext::GetComponentSpaceTransform - Couldn't find BoneIndexType %d (%s) requested by %s"), BoneIndexType, *Skeleton->GetReferenceSkeleton().GetBoneName(BoneIndexType).ToString(), *Schema->GetName());
+					UE_LOG(LogPoseSearch, Warning, TEXT("FSearchContext::GetLocalBoneTransformAtTime - Couldn't find BoneIndexType %d (%s) requested by %s"), BoneIndexType, *Skeleton->GetReferenceSkeleton().GetBoneName(BoneIndexType).ToString(), *Schema->GetName());
 				}
 			}
 			else
 			{
-				UE_LOG(LogPoseSearch, Warning, TEXT("FSearchContext::GetComponentSpaceTransform - Schema '%s' Skeleton is not properly set"), *Schema->GetName());
+				UE_LOG(LogPoseSearch, Warning, TEXT("FSearchContext::GetLocalBoneTransformAtTime - Schema '%s' Skeleton is not properly set"), *Schema->GetName());
 			}
 		}
 
-		CachedTransforms.Add(SampleTime, BoneIndexType, BoneComponentSpaceTransform);
-		return BoneComponentSpaceTransform;
+		CachedTransforms.Add(SampleTime, BoneIndexType, LocalBoneTransform);
+		return LocalBoneTransform;
 	}
 
 	return FTransform::Identity;
@@ -500,13 +501,13 @@ FVector FSearchContext::GetSamplePositionInternal(float SampleTime, float Origin
 {
 	if (SampleBonePositionWorldOverride)
 	{
-		const FTransform RootBoneTransform = GetTransform(OriginTime, Schema, RootSchemaBoneIdx, bUseHistoryRoot);
+		const FTransform RootBoneTransform = GetWorldBoneTransformAtTime(OriginTime, Schema, RootSchemaBoneIdx, bUseHistoryRoot);
 		if (SchemaOriginBoneIdx == RootSchemaBoneIdx)
 		{
 			return RootBoneTransform.InverseTransformPosition(*SampleBonePositionWorldOverride);
 		}
 
-		const FTransform OriginBoneTransform = GetTransform(OriginTime, Schema, SchemaOriginBoneIdx, bUseHistoryRoot);
+		const FTransform OriginBoneTransform = GetWorldBoneTransformAtTime(OriginTime, Schema, SchemaOriginBoneIdx, bUseHistoryRoot);
 		const FVector DeltaBoneTranslation = *SampleBonePositionWorldOverride - OriginBoneTransform.GetTranslation();
 		return RootBoneTransform.InverseTransformVector(DeltaBoneTranslation);
 	}
@@ -515,22 +516,22 @@ FVector FSearchContext::GetSamplePositionInternal(float SampleTime, float Origin
 	{
 		if (SchemaOriginBoneIdx == RootSchemaBoneIdx)
 		{
-			return GetComponentSpaceTransform(SampleTime, Schema, SchemaSampleBoneIdx).GetTranslation();
+			return GetLocalBoneTransformAtTime(SampleTime, Schema, SchemaSampleBoneIdx).GetTranslation();
 		}
 
-		const FVector SampleBonePosition = GetComponentSpaceTransform(SampleTime, Schema, SchemaSampleBoneIdx).GetTranslation();
-		const FVector OriginBonePosition = GetComponentSpaceTransform(OriginTime, Schema, SchemaOriginBoneIdx).GetTranslation();
+		const FVector SampleBonePosition = GetLocalBoneTransformAtTime(SampleTime, Schema, SchemaSampleBoneIdx).GetTranslation();
+		const FVector OriginBonePosition = GetLocalBoneTransformAtTime(OriginTime, Schema, SchemaOriginBoneIdx).GetTranslation();
 		return SampleBonePosition - OriginBonePosition;
 	}
 
-	const FTransform RootBoneTransform = GetTransform(OriginTime, Schema, RootSchemaBoneIdx, bUseHistoryRoot);
-	const FTransform SampleBoneTransform = GetTransform(SampleTime, Schema, SchemaSampleBoneIdx, bUseHistoryRoot);
+	const FTransform RootBoneTransform = GetWorldBoneTransformAtTime(OriginTime, Schema, RootSchemaBoneIdx, bUseHistoryRoot);
+	const FTransform SampleBoneTransform = GetWorldBoneTransformAtTime(SampleTime, Schema, SchemaSampleBoneIdx, bUseHistoryRoot);
 	if (SchemaOriginBoneIdx == RootSchemaBoneIdx)
 	{
 		return RootBoneTransform.InverseTransformPosition(SampleBoneTransform.GetTranslation());
 	}
 
-	const FTransform OriginBoneTransform = GetTransform(OriginTime, Schema, SchemaOriginBoneIdx, bUseHistoryRoot);
+	const FTransform OriginBoneTransform = GetWorldBoneTransformAtTime(OriginTime, Schema, SchemaOriginBoneIdx, bUseHistoryRoot);
 	const FVector DeltaBoneTranslation = SampleBoneTransform.GetTranslation() - OriginBoneTransform.GetTranslation();
 	return RootBoneTransform.InverseTransformVector(DeltaBoneTranslation);
 }
@@ -539,13 +540,13 @@ FQuat FSearchContext::GetSampleRotationInternal(float SampleTime, float OriginTi
 {
 	if (SampleBoneRotationWorldOverride)
 	{
-		const FTransform RootBoneTransform = GetTransform(OriginTime, Schema, RootSchemaBoneIdx, bUseHistoryRoot);
+		const FTransform RootBoneTransform = GetWorldBoneTransformAtTime(OriginTime, Schema, RootSchemaBoneIdx, bUseHistoryRoot);
 		if (SchemaOriginBoneIdx == RootSchemaBoneIdx)
 		{
 			return RootBoneTransform.InverseTransformRotation(*SampleBoneRotationWorldOverride);
 		}
 
-		const FTransform OriginBoneTransform = GetTransform(OriginTime, Schema, SchemaOriginBoneIdx, bUseHistoryRoot);
+		const FTransform OriginBoneTransform = GetWorldBoneTransformAtTime(OriginTime, Schema, SchemaOriginBoneIdx, bUseHistoryRoot);
 		const FQuat DeltaBoneRotation = OriginBoneTransform.InverseTransformRotation(*SampleBoneRotationWorldOverride);
 		return RootBoneTransform.InverseTransformRotation(DeltaBoneRotation);
 	}
@@ -554,22 +555,22 @@ FQuat FSearchContext::GetSampleRotationInternal(float SampleTime, float OriginTi
 	{
 		if (SchemaOriginBoneIdx == RootSchemaBoneIdx)
 		{
-			return GetComponentSpaceTransform(SampleTime, Schema, SchemaSampleBoneIdx).GetRotation();
+			return GetLocalBoneTransformAtTime(SampleTime, Schema, SchemaSampleBoneIdx).GetRotation();
 		}
 
-		const FTransform SampleBoneTransform = GetComponentSpaceTransform(SampleTime, Schema, SchemaSampleBoneIdx);
-		const FTransform OriginBoneTransform = GetComponentSpaceTransform(OriginTime, Schema, SchemaOriginBoneIdx);
+		const FTransform SampleBoneTransform = GetLocalBoneTransformAtTime(SampleTime, Schema, SchemaSampleBoneIdx);
+		const FTransform OriginBoneTransform = GetLocalBoneTransformAtTime(OriginTime, Schema, SchemaOriginBoneIdx);
 		return OriginBoneTransform.InverseTransformRotation(SampleBoneTransform.GetRotation());
 	}
 
-	const FTransform RootBoneTransform = GetTransform(OriginTime, Schema, RootSchemaBoneIdx, bUseHistoryRoot);
-	const FTransform SampleBoneTransform = GetTransform(SampleTime, Schema, SchemaSampleBoneIdx, bUseHistoryRoot);
+	const FTransform RootBoneTransform = GetWorldBoneTransformAtTime(OriginTime, Schema, RootSchemaBoneIdx, bUseHistoryRoot);
+	const FTransform SampleBoneTransform = GetWorldBoneTransformAtTime(SampleTime, Schema, SchemaSampleBoneIdx, bUseHistoryRoot);
 	if (SchemaOriginBoneIdx == RootSchemaBoneIdx)
 	{
 		return RootBoneTransform.InverseTransformRotation(SampleBoneTransform.GetRotation());
 	}
 
-	const FTransform OriginBoneTransform = GetTransform(OriginTime, Schema, SchemaOriginBoneIdx, bUseHistoryRoot);
+	const FTransform OriginBoneTransform = GetWorldBoneTransformAtTime(OriginTime, Schema, SchemaOriginBoneIdx, bUseHistoryRoot);
 	const FQuat DeltaBoneRotation = OriginBoneTransform.InverseTransformRotation(SampleBoneTransform.GetRotation());
 	return RootBoneTransform.InverseTransformRotation(DeltaBoneRotation);
 }
diff --git a/Engine/Plugins/Experimental/Animation/PoseSearch/Source/Runtime/Private/PoseSearchHistory.cpp b/Engine/Plugins/Experimental/Animation/PoseSearch/Source/Runtime/Private/PoseSearchHistory.cpp
--- a/Engine/Plugins/Experimental/Animation/PoseSearch/Source/Runtime/Private/PoseSearchHistory.cpp
+++ b/Engine/Plugins/Experimental/Animation/PoseSearch/Source/Runtime/Private/PoseSearchHistory.cpp
@@ -68,12 +68,8 @@ FORCEINLINE auto LowerBound(IteratorType First, IteratorType Last, const ValueTy
 // FPoseHistoryEntry
 void FPoseHistoryEntry::Update(float InTime, FCSPose<FCompactPose>& ComponentSpacePose, const FTransform& ComponentTransform, const FBoneToTransformMap& BoneToTransformMap)
 {
-	// @todo: optimize this math by initializing FCSPose<FCompactPose> root to identity
-	const FTransform ComponentToRootTransform = ComponentSpacePose.GetComponentSpaceTransform(FCompactPoseBoneIndex(RootBoneIndexType));
-	const FTransform ComponentToRootTransformInv = ComponentToRootTransform.Inverse();
-
 	Time = InTime;
-	RootTransform = ComponentToRootTransform * ComponentTransform;
+	RootTransform = ComponentTransform;
 
 	const FBoneContainer& BoneContainer = ComponentSpacePose.GetPose().GetBoneContainer();
 	const USkeleton* SkeletonAsset = BoneContainer.GetSkeletonAsset();
@@ -89,8 +85,7 @@ void FPoseHistoryEntry::Update(float InTime, FCSPose<FCompactPose>& ComponentSpa
 		for (FSkeletonPoseBoneIndex SkeletonBoneIdx(0); SkeletonBoneIdx != NumSkeletonBones; ++SkeletonBoneIdx)
 		{
 			const FCompactPoseBoneIndex CompactBoneIdx = BoneContainer.GetCompactPoseIndexFromSkeletonPoseIndex(SkeletonBoneIdx);
-			ComponentSpaceTransforms[SkeletonBoneIdx.GetInt()] = 
-				(CompactBoneIdx.IsValid() ? ComponentSpacePose.GetComponentSpaceTransform(CompactBoneIdx) : RefBonePose[SkeletonBoneIdx.GetInt()]) * ComponentToRootTransformInv;
+			ComponentSpaceTransforms[SkeletonBoneIdx.GetInt()] = (CompactBoneIdx.IsValid() ? ComponentSpacePose.GetComponentSpaceTransform(CompactBoneIdx) : RefBonePose[SkeletonBoneIdx.GetInt()]);
 		}
 	}
 	else
@@ -100,8 +95,7 @@ void FPoseHistoryEntry::Update(float InTime, FCSPose<FCompactPose>& ComponentSpa
 		{
 			const FSkeletonPoseBoneIndex SkeletonBoneIdx(BoneToTransformPair.Key);
 			const FCompactPoseBoneIndex CompactBoneIdx = BoneContainer.GetCompactPoseIndexFromSkeletonPoseIndex(SkeletonBoneIdx);
-			ComponentSpaceTransforms[BoneToTransformPair.Value] = 
-				(CompactBoneIdx.IsValid() ? ComponentSpacePose.GetComponentSpaceTransform(CompactBoneIdx) : RefBonePose[SkeletonBoneIdx.GetInt()]) * ComponentToRootTransformInv;
+			ComponentSpaceTransforms[BoneToTransformPair.Value] = (CompactBoneIdx.IsValid() ? ComponentSpacePose.GetComponentSpaceTransform(CompactBoneIdx) : RefBonePose[SkeletonBoneIdx.GetInt()]);
 		}
 	}
 }
@@ -114,21 +108,31 @@ void FPoseHistory::Init(int32 InNumPoses, float InTimeHorizon, const TArray<FBon
 	TimeHorizon = InTimeHorizon;
 
 	BoneToTransformMap.Reset();
-	for (int32 i = 0; i < RequiredBones.Num(); ++i)
+	if (!RequiredBones.IsEmpty())
 	{
-		BoneToTransformMap.Add(RequiredBones[i]) = i;
+		// making sure we always collect the root bone transform (by construction BoneToTransformMap[0] = 0)
+		BoneToTransformMap.Add(RootBoneIndexType) = BoneToTransformMap.Num();
+
+		for (int32 i = 0; i < RequiredBones.Num(); ++i)
+		{
+			// adding only unique RequiredBones to avoid oversizing Entries::ComponentSpaceTransforms
+			if (!BoneToTransformMap.Find(RequiredBones[i]))
+			{
+				BoneToTransformMap.Add(RequiredBones[i]) = BoneToTransformMap.Num();
+			}
+		}
 	}
 
 	Entries.Reset();
 	Entries.Reserve(InNumPoses);
 }
 
-FBoneIndexType FPoseHistory::GetRemappedBoneIndexType(FBoneIndexType BoneIndexType, const USkeleton* BoneIndexSkeleton) const
+FBoneIndexType FPoseHistory::GetRemappedBoneIndexType(FBoneIndexType BoneIndexType, const USkeleton* BoneIndexSkeleton, const USkeleton* LastUpdateSkeleton)
 {
 	// remapping BoneIndexType in case the skeleton used to store history (LastUpdateSkeleton) is different from BoneIndexSkeleton
 	if (LastUpdateSkeleton != nullptr && LastUpdateSkeleton != BoneIndexSkeleton)
 	{
-		const FSkeletonRemapping& SkeletonRemapping = UE::Anim::FSkeletonRemappingRegistry::Get().GetRemapping(BoneIndexSkeleton, LastUpdateSkeleton.Get());
+		const FSkeletonRemapping& SkeletonRemapping = UE::Anim::FSkeletonRemappingRegistry::Get().GetRemapping(BoneIndexSkeleton, LastUpdateSkeleton);
 		if (SkeletonRemapping.IsValid())
 		{
 			BoneIndexType = SkeletonRemapping.GetTargetSkeletonBoneIndex(BoneIndexType);
@@ -138,93 +142,108 @@ FBoneIndexType FPoseHistory::GetRemappedBoneIndexType(FBoneIndexType BoneIndexTy
 	return BoneIndexType;
 }
 
-bool FPoseHistory::GetComponentSpaceTransformAtTime(float Time, FBoneIndexType BoneIndexType, const USkeleton* BoneIndexSkeleton, FTransform& OutBoneTransform, bool bExtrapolate) const
+FComponentSpaceTransformIndex FPoseHistory::GetRemappedComponentSpaceTransformIndex(const USkeleton* BoneIndexSkeleton, const USkeleton* LastUpdateSkeleton, const FBoneToTransformMap& BoneToTransformMap, FBoneIndexType BoneIndexType, bool& bSuccess)
 {
-	BoneIndexType = GetRemappedBoneIndexType(BoneIndexType, BoneIndexSkeleton);
-
-	FComponentSpaceTransformIndex TransformIndex = FComponentSpaceTransformIndex(BoneIndexType);
-	if (!BoneToTransformMap.IsEmpty())
+	FComponentSpaceTransformIndex BoneTransformIndex = FComponentSpaceTransformIndex(BoneIndexType);
+	if (BoneIndexType < WorldSpaceIndexType)
 	{
-		if (const FComponentSpaceTransformIndex* FoundTransformIndex = BoneToTransformMap.Find(BoneIndexType))
-		{
-			TransformIndex = *FoundTransformIndex;
-		}
-		else
+		BoneIndexType = GetRemappedBoneIndexType(BoneIndexType, BoneIndexSkeleton, LastUpdateSkeleton);
+
+		if (!BoneToTransformMap.IsEmpty())
 		{
-			GetRootTransformAtTime(Time, OutBoneTransform, bExtrapolate);
-			return false;
+			if (const FComponentSpaceTransformIndex* FoundBoneTransformIndex = BoneToTransformMap.Find(BoneTransformIndex))
+			{
+				BoneTransformIndex = *FoundBoneTransformIndex;
+			}
+			else
+			{
+				BoneTransformIndex = RootBoneIndexType;
+				bSuccess = false;
+			}
 		}
 	}
+	return BoneTransformIndex;
+}
 
-	const int32 Num = Entries.Num();
-	if (Num > 1)
-	{
-		const float SecondsAgo = -Time;
-		const int32 LowerBoundIdx = LowerBound(Entries.begin(), Entries.end(), SecondsAgo, [](const FPoseHistoryEntry& Entry, float Value) { return Value < Entry.Time; });
-		const int32 NextIdx = FMath::Clamp(LowerBoundIdx, 1, Num - 1);
-		const int32 PrevIdx = NextIdx - 1;
+bool FPoseHistory::LerpEntries(const FPoseHistoryEntry& PrevEntry, const FPoseHistoryEntry& NextEntry, float LerpValue, const USkeleton* BoneIndexSkeleton, const USkeleton* LastUpdateSkeleton, const FBoneToTransformMap& BoneToTransformMap, FBoneIndexType BoneIndexType, FBoneIndexType ReferenceBoneIndexType, FTransform& OutBoneTransform)
+{
+	bool bSuccess = true;
 
-		const FPoseHistoryEntry& PrevEntry = Entries[PrevIdx];
-		const FPoseHistoryEntry& NextEntry = Entries[NextIdx];
+	const FComponentSpaceTransformIndex BoneTransformIndex = GetRemappedComponentSpaceTransformIndex(BoneIndexSkeleton, LastUpdateSkeleton, BoneToTransformMap, BoneIndexType, bSuccess);
+	const FComponentSpaceTransformIndex ReferenceBoneTransformIndex = GetRemappedComponentSpaceTransformIndex(BoneIndexSkeleton, LastUpdateSkeleton, BoneToTransformMap, ReferenceBoneIndexType, bSuccess);
 
-		const float Denominator = NextEntry.Time - PrevEntry.Time;
-		if (!FMath::IsNearlyZero(Denominator))
-		{
-			const float Numerator = SecondsAgo - PrevEntry.Time;
-			const float LerpValue = bExtrapolate ? Numerator / Denominator : FMath::Clamp(Numerator / Denominator, 0.f, 1.f);
-			OutBoneTransform.Blend(PrevEntry.ComponentSpaceTransforms[TransformIndex], NextEntry.ComponentSpaceTransforms[TransformIndex], LerpValue);
-		}
-		else
+	if (BoneTransformIndex < WorldSpaceIndexType)
+	{
+		switch (ReferenceBoneTransformIndex)
 		{
-			OutBoneTransform = PrevEntry.ComponentSpaceTransforms[TransformIndex];
+		case ComponentSpaceIndexType:
+			OutBoneTransform.Blend(
+				PrevEntry.ComponentSpaceTransforms[BoneTransformIndex],
+				NextEntry.ComponentSpaceTransforms[BoneTransformIndex],
+				LerpValue);
+			break;
+		case WorldSpaceIndexType:
+			OutBoneTransform.Blend(
+				PrevEntry.ComponentSpaceTransforms[BoneTransformIndex] * PrevEntry.RootTransform,
+				NextEntry.ComponentSpaceTransforms[BoneTransformIndex] * NextEntry.RootTransform,
+				LerpValue);
+			break;
+		default:
+			OutBoneTransform.Blend(
+				PrevEntry.ComponentSpaceTransforms[BoneTransformIndex] * PrevEntry.ComponentSpaceTransforms[ReferenceBoneTransformIndex].Inverse(),
+				NextEntry.ComponentSpaceTransforms[BoneTransformIndex] * NextEntry.ComponentSpaceTransforms[ReferenceBoneTransformIndex].Inverse(),
+				LerpValue);
+			break;
 		}
 	}
-	else if (Num > 0)
-	{
-		OutBoneTransform = Entries[0].ComponentSpaceTransforms[TransformIndex];
-	}
 	else
 	{
-		GetRootTransformAtTime(Time, OutBoneTransform, bExtrapolate);
-		return false;
+		// @todo: implement if required
+		OutBoneTransform = FTransform::Identity;
+		bSuccess = false;
+		unimplemented();
 	}
 
-	return true;
+	return bSuccess;
 }
 
-void FPoseHistory::GetRootTransformAtTime(float Time, FTransform& OutRootTransform, bool bExtrapolate) const
+bool FPoseHistory::GetTransformAtTime(float Time, FTransform& OutBoneTransform, const USkeleton* BoneIndexSkeleton, FBoneIndexType BoneIndexType, FBoneIndexType ReferenceBoneIndexType, bool bExtrapolate) const
 {
+	static_assert(RootBoneIndexType == 0 && WorldSpaceIndexType < ComponentSpaceIndexType); // some assumptions
+	check(BoneIndexType < WorldSpaceIndexType);
+	
 	const int32 Num = Entries.Num();
-	if (Num > 1)
+	
+	if (Num > 0)
 	{
 		const float SecondsAgo = -Time;
-		const int32 LowerBoundIdx = LowerBound(Entries.begin(), Entries.end(), SecondsAgo, [](const FPoseHistoryEntry& Entry, float Value) { return Value < Entry.Time; });
-		const int32 NextIdx = FMath::Clamp(LowerBoundIdx, 1, Num - 1);
-		const int32 PrevIdx = NextIdx - 1;
 
+		int32 NextIdx = 0;
+		int32 PrevIdx = 0;
+
+		if (Num > 1)
+		{
+			const int32 LowerBoundIdx = LowerBound(Entries.begin(), Entries.end(), SecondsAgo, [](const FPoseHistoryEntry& Entry, float Value) { return Value < Entry.Time; });
+			NextIdx = FMath::Clamp(LowerBoundIdx, 1, Num - 1);
+			PrevIdx = NextIdx - 1;
+		}
+	
 		const FPoseHistoryEntry& PrevEntry = Entries[PrevIdx];
 		const FPoseHistoryEntry& NextEntry = Entries[NextIdx];
 
 		const float Denominator = NextEntry.Time - PrevEntry.Time;
+		float LerpValue = 0.f;
 		if (!FMath::IsNearlyZero(Denominator))
 		{
 			const float Numerator = SecondsAgo - PrevEntry.Time;
-			const float LerpValue = bExtrapolate ? Numerator / Denominator : FMath::Clamp(Numerator / Denominator, 0.f, 1.f);
-			OutRootTransform.Blend(PrevEntry.RootTransform, NextEntry.RootTransform, LerpValue);
-		}
-		else
-		{
-			OutRootTransform = PrevEntry.RootTransform;
+			LerpValue = bExtrapolate ? Numerator / Denominator : FMath::Clamp(Numerator / Denominator, 0.f, 1.f);
 		}
+
+		return LerpEntries(PrevEntry, NextEntry, LerpValue, BoneIndexSkeleton, GetLastUpdateSkeleton(), BoneToTransformMap, BoneIndexType, ReferenceBoneIndexType, OutBoneTransform);
 	}
-	else if (Num > 0)
-	{
-		OutRootTransform = Entries[0].RootTransform;
-	}
-	else
-	{
-		OutRootTransform = FTransform::Identity;
-	}
+	
+	OutBoneTransform = FTransform::Identity;
+	return false;
 }
 
 bool FPoseHistory::IsEmpty() const
@@ -351,66 +370,9 @@ float FExtendedPoseHistory::GetSampleTimeInterval() const
 	return PoseHistory->GetSampleTimeInterval();
 }
 
-bool FExtendedPoseHistory::GetComponentSpaceTransformAtTime(float Time, FBoneIndexType BoneIndexType, const USkeleton* BoneIndexSkeleton, FTransform& OutBoneTransform, bool bExtrapolate) const
-{
-	check(PoseHistory);
-
-	const FBoneIndexType OriginalBoneIndex = BoneIndexType;
-
-	BoneIndexType = PoseHistory->GetRemappedBoneIndexType(BoneIndexType, BoneIndexSkeleton);
-
-	if (Time > 0.f)
-	{
- 		FComponentSpaceTransformIndex TransformIndex = FComponentSpaceTransformIndex(BoneIndexType);
-		const FBoneToTransformMap& BoneToTransformMap = PoseHistory->GetBoneToTransformMap();
-		if (!BoneToTransformMap.IsEmpty())
-		{
-			if (const FComponentSpaceTransformIndex* FoundTransformIndex = BoneToTransformMap.Find(BoneIndexType))
-			{
-				TransformIndex = *FoundTransformIndex;
-			}
-			else
-			{
-				GetRootTransformAtTime(Time, OutBoneTransform, bExtrapolate);
-				return false;
-			}
-		}
-
-		const int32 Num = FutureEntries.Num();
-		if (Num > 0)
-		{
-			const float SecondsAgo = -Time;
-			const int32 LowerBoundIdx = Algo::LowerBound(FutureEntries, SecondsAgo, [](const FPoseHistoryEntry& Entry, float Value) { return Value < Entry.Time; });
-			const int32 NextIdx = FMath::Min(LowerBoundIdx, Num - 1);
-			const FPoseHistoryEntries& PastEntries = PoseHistory->GetEntries();
-			const FPoseHistoryEntry& NextEntry = FutureEntries[NextIdx];
-			const FPoseHistoryEntry& PrevEntry = NextIdx > 0 ? FutureEntries[NextIdx - 1] : !PastEntries.IsEmpty() ? PastEntries.First() : NextEntry;
-
-			const float Denominator = NextEntry.Time - PrevEntry.Time;
-			if (!FMath::IsNearlyZero(Denominator))
-			{
-				const float Numerator = SecondsAgo - PrevEntry.Time;
-				const float LerpValue = bExtrapolate ? Numerator / Denominator : FMath::Clamp(Numerator / Denominator, 0.f, 1.f);
-				OutBoneTransform.Blend(PrevEntry.ComponentSpaceTransforms[TransformIndex], NextEntry.ComponentSpaceTransforms[TransformIndex], LerpValue);
-			}
-			else
-			{
-				OutBoneTransform = NextEntry.ComponentSpaceTransforms[TransformIndex];
-			}
-			return true;
-		}
-
-		// we've no FutureEntries, so let's clamp the time to zero and query PoseHistory
-		Time = 0.f;
-	}
-
-	return PoseHistory->GetComponentSpaceTransformAtTime(Time, OriginalBoneIndex, BoneIndexSkeleton, OutBoneTransform, bExtrapolate);
-}
-
-void FExtendedPoseHistory::GetRootTransformAtTime(float Time, FTransform& OutRootTransform, bool bExtrapolate) const
+bool FExtendedPoseHistory::GetTransformAtTime(float Time, FTransform& OutBoneTransform, const USkeleton* BoneIndexSkeleton, FBoneIndexType BoneIndexType, FBoneIndexType ReferenceBoneIndexType, bool bExtrapolate) const
 {
 	check(PoseHistory);
-
 	if (Time > 0.f)
 	{
 		const int32 Num = FutureEntries.Num();
@@ -424,25 +386,17 @@ void FExtendedPoseHistory::GetRootTransformAtTime(float Time, FTransform& OutRoo
 			const FPoseHistoryEntry& PrevEntry = NextIdx > 0 ? FutureEntries[NextIdx - 1] : !PastEntries.IsEmpty() ? PastEntries.First() : NextEntry;
 
 			const float Denominator = NextEntry.Time - PrevEntry.Time;
+			float LerpValue = 0.f;
 			if (!FMath::IsNearlyZero(Denominator))
 			{
 				const float Numerator = SecondsAgo - PrevEntry.Time;
-				const float LerpValue = bExtrapolate ? Numerator / Denominator : FMath::Clamp(Numerator / Denominator, 0.f, 1.f);
-				OutRootTransform.Blend(PrevEntry.RootTransform, NextEntry.RootTransform, LerpValue);
+				LerpValue = bExtrapolate ? Numerator / Denominator : FMath::Clamp(Numerator / Denominator, 0.f, 1.f);
 			}
-			else
-			{
-				OutRootTransform = NextEntry.RootTransform;
-			}
-
-			return;
+			return FPoseHistory::LerpEntries(PrevEntry, NextEntry, LerpValue, BoneIndexSkeleton, PoseHistory->GetLastUpdateSkeleton(), PoseHistory->GetBoneToTransformMap(), BoneIndexType, ReferenceBoneIndexType, OutBoneTransform);
 		}
-
-		// we've no FutureEntries, so let's clamp the time to zero and query PoseHistory
-		Time = 0.f;
 	}
-
-	PoseHistory->GetRootTransformAtTime(Time, OutRootTransform, bExtrapolate);
+	
+	return PoseHistory->GetTransformAtTime(Time, OutBoneTransform, BoneIndexSkeleton, BoneIndexType, ReferenceBoneIndexType, bExtrapolate);
 }
 
 bool FExtendedPoseHistory::IsEmpty() const
diff --git a/Engine/Plugins/Experimental/Animation/PoseSearch/Source/Runtime/Private/PoseSearchLibrary.cpp b/Engine/Plugins/Experimental/Animation/PoseSearch/Source/Runtime/Private/PoseSearchLibrary.cpp
--- a/Engine/Plugins/Experimental/Animation/PoseSearch/Source/Runtime/Private/PoseSearchLibrary.cpp
+++ b/Engine/Plugins/Experimental/Animation/PoseSearch/Source/Runtime/Private/PoseSearchLibrary.cpp
@@ -207,11 +207,11 @@ void UPoseSearchLibrary::TraceMotionMatchingState(
 		}
 	}
 
-	if (DeltaTime > SMALL_NUMBER && SearchContext.IsTrajectoryValid())
+	if (DeltaTime > SMALL_NUMBER)
 	{
 		// simulation
-		const FTransform PrevRoot = SearchContext.GetRootAtTime(-DeltaTime);
-		const FTransform CurrRoot = SearchContext.GetRootAtTime(0.f);
+		const FTransform PrevRoot = SearchContext.GetWorldRootBoneTransformAtTime(-DeltaTime);
+		const FTransform CurrRoot = SearchContext.GetWorldRootBoneTransformAtTime(0.f);
 		const FTransform SimDelta = CurrRoot.GetRelativeTransform(PrevRoot);
 
 		TraceState.SimLinearVelocity = SimDelta.GetTranslation().Size() / DeltaTime;
@@ -280,8 +280,8 @@ void UPoseSearchLibrary::UpdateMotionMatchingState(
 		History = &PoseHistoryProvider->GetPoseHistory();
 	}
 	
-	const FPoseSearchQueryTrajectory TrajectoryRootSpace = ProcessTrajectory(Trajectory, Context.AnimInstanceProxy->GetComponentTransform(), InOutMotionMatchingState.ComponentDeltaYaw, YawFromAnimationTrajectoryBlendTime, TrajectorySpeedMultiplier);
-	
+	const FPoseSearchQueryTrajectory TrajectoryRootSpace = ProcessTrajectory(Trajectory, InOutMotionMatchingState.ComponentDeltaYaw, YawFromAnimationTrajectoryBlendTime, TrajectorySpeedMultiplier);
+
 	FMemMark Mark(FMemStack::Get());
 	const UAnimInstance* AnimInstance = Cast<const UAnimInstance>(Context.AnimInstanceProxy->GetAnimInstanceObject());
 	check(AnimInstance);
@@ -393,25 +393,24 @@ void UPoseSearchLibrary::UpdateMotionMatchingState(
 #endif
 }
 
-// transforms Trajectory from world space to root bone space, and scale it by TrajectorySpeedMultiplier
-FPoseSearchQueryTrajectory UPoseSearchLibrary::ProcessTrajectory(const FPoseSearchQueryTrajectory& Trajectory, const FTransform& ComponentWorldTransform, float RootBoneDeltaYaw, float YawFromAnimationTrajectoryBlendTime, float TrajectorySpeedMultiplier)
+// transforms Trajectory from SkeletalMeshComponent world space to root bone world space, and scale it by TrajectorySpeedMultiplier
+FPoseSearchQueryTrajectory UPoseSearchLibrary::ProcessTrajectory(const FPoseSearchQueryTrajectory& Trajectory, float RootBoneDeltaYaw, float YawFromAnimationTrajectoryBlendTime, float TrajectorySpeedMultiplier)
 {
 	QUICK_SCOPE_CYCLE_COUNTER(STAT_PoseSearch_ProcessTrajectory);
 
 	const float TrajectorySpeedMultiplierInv = FMath::IsNearlyZero(TrajectorySpeedMultiplier) ? 1.f : 1.f / TrajectorySpeedMultiplier;
 
 	FPoseSearchQueryTrajectory TrajectoryRootSpace = Trajectory;
-	const FTransform ToRootSpace = ComponentWorldTransform.Inverse();
 	for (FPoseSearchQueryTrajectorySample& Sample : TrajectoryRootSpace.Samples)
 	{
 		Sample.AccumulatedSeconds *= TrajectorySpeedMultiplierInv;
 
-		Sample.Position = ToRootSpace.TransformPosition(Sample.Position);
-
-		const float BlendParam = YawFromAnimationTrajectoryBlendTime < UE_KINDA_SMALL_NUMBER ? 1 : FMath::Clamp(1.f - (Sample.AccumulatedSeconds - YawFromAnimationTrajectoryBlendTime) / YawFromAnimationTrajectoryBlendTime, 0.f, 1.f);
-		const FQuat RootBoneDelta(FRotator(0.f, RootBoneDeltaYaw * BlendParam, 0.f));
-		const FQuat ToRootSpaceRotation = ToRootSpace.GetRotation() * RootBoneDelta;
-		Sample.Facing = ToRootSpaceRotation * Sample.Facing;
+		if (!FMath::IsNearlyZero(RootBoneDeltaYaw))
+		{
+			const float BlendParam = YawFromAnimationTrajectoryBlendTime < UE_KINDA_SMALL_NUMBER ? 1 : FMath::Clamp(1.f - (Sample.AccumulatedSeconds - YawFromAnimationTrajectoryBlendTime) / YawFromAnimationTrajectoryBlendTime, 0.f, 1.f);
+			const FQuat RootBoneDelta(FRotator(0.f, RootBoneDeltaYaw * BlendParam, 0.f));
+			Sample.Facing = RootBoneDelta * Sample.Facing;
+		}
 	}
 
 	return TrajectoryRootSpace;
@@ -460,7 +459,7 @@ void UPoseSearchLibrary::MotionMatch(
 	FMemMark Mark(FMemStack::Get());
 	if (Database && AnimInstance)
 	{
-		const FPoseSearchQueryTrajectory TrajectoryRootSpace = ProcessTrajectory(Trajectory, AnimInstance->GetOwningComponent()->GetComponentTransform(), 0.f, 0.f, TrajectorySpeedMultiplier);
+		const FPoseSearchQueryTrajectory TrajectoryRootSpace = ProcessTrajectory(Trajectory, 0.f, 0.f, TrajectorySpeedMultiplier);
 
 		// ExtendedPoseHistory will hold future poses to match AssetSamplerBase (at FutureAnimationStartTime) TimeToFutureAnimationStart seconds in the future
 		FExtendedPoseHistory ExtendedPoseHistory;
diff --git a/Engine/Plugins/Experimental/Animation/PoseSearch/Source/Runtime/Public/PoseSearch/AnimNode_MotionMatching.h b/Engine/Plugins/Experimental/Animation/PoseSearch/Source/Runtime/Public/PoseSearch/AnimNode_MotionMatching.h
--- a/Engine/Plugins/Experimental/Animation/PoseSearch/Source/Runtime/Public/PoseSearch/AnimNode_MotionMatching.h
+++ b/Engine/Plugins/Experimental/Animation/PoseSearch/Source/Runtime/Public/PoseSearch/AnimNode_MotionMatching.h
@@ -48,7 +48,7 @@ private:
 	UPROPERTY(EditAnywhere, Category = Settings, meta = (PinShownByDefault))
 	TObjectPtr<const UPoseSearchDatabase> Database = nullptr;
 
-	// Motion Trajectory samples for pose search queries in Motion Matching.These are expected to be in the space of the SkeletalMeshComponent.This is provided with the CharacterMovementTrajectory Component output.
+	// Motion Trajectory samples for pose search queries in Motion Matching. These are expected to be in the world space of the SkeletalMeshComponent. This is provided with the CharacterMovementTrajectory Component output.
 	UPROPERTY(EditAnywhere, Category = Settings, meta = (PinShownByDefault))
 	FPoseSearchQueryTrajectory Trajectory;
 
diff --git a/Engine/Plugins/Experimental/Animation/PoseSearch/Source/Runtime/Public/PoseSearch/AnimNode_PoseSearchHistoryCollector.h b/Engine/Plugins/Experimental/Animation/PoseSearch/Source/Runtime/Public/PoseSearch/AnimNode_PoseSearchHistoryCollector.h
--- a/Engine/Plugins/Experimental/Animation/PoseSearch/Source/Runtime/Public/PoseSearch/AnimNode_PoseSearchHistoryCollector.h
+++ b/Engine/Plugins/Experimental/Animation/PoseSearch/Source/Runtime/Public/PoseSearch/AnimNode_PoseSearchHistoryCollector.h
@@ -24,6 +24,10 @@ public:
 	UPROPERTY(EditAnywhere, Category = Settings)
 	TArray<FBoneReference> CollectedBones;
 
+	// if true, the pose history will be initialized with a ref pose at the location and orientation of the AnimInstance.
+	UPROPERTY(EditAnywhere, Category = Settings)
+	bool bInitializeWithRefPose = false;
+
 	// FAnimNode_Base interface
 	virtual void CacheBones_AnyThread(const FAnimationCacheBonesContext& Context) override;
 	// End of FAnimNode_Base interface
diff --git a/Engine/Plugins/Experimental/Animation/PoseSearch/Source/Runtime/Public/PoseSearch/PoseSearchContext.h b/Engine/Plugins/Experimental/Animation/PoseSearch/Source/Runtime/Public/PoseSearch/PoseSearchContext.h
--- a/Engine/Plugins/Experimental/Animation/PoseSearch/Source/Runtime/Public/PoseSearch/PoseSearchContext.h
+++ b/Engine/Plugins/Experimental/Animation/PoseSearch/Source/Runtime/Public/PoseSearch/PoseSearchContext.h
@@ -182,18 +182,22 @@ struct POSESEARCH_API FSearchContext
 	const FPoseIndicesHistory* GetPoseIndicesHistory() const { return PoseIndicesHistory; }
 	bool IsHistoryValid() const { return History != nullptr; }
 	float GetDesiredPermutationTimeOffset() const { return DesiredPermutationTimeOffset; }
-	bool IsTrajectoryValid() const { return Trajectory != nullptr; }
 	bool IsForceInterrupt() const { return bForceInterrupt; }
-	FTransform GetRootAtTime(float Time, bool bUseHistoryRoot = false, bool bExtrapolate = true) const;
+	FTransform GetWorldRootBoneTransformAtTime(float SampleTime, bool bUseHistoryRoot = false, bool bExtrapolate = true) const;
 	const UAnimInstance* GetAnimInstance() const { return AnimInstance; }
 	
 private:
-	FTransform GetTransform(float SampleTime, const UPoseSearchSchema* Schema, int8 SchemaBoneIdx = RootSchemaBoneIdx, bool bUseHistoryRoot = false);
-	FTransform GetComponentSpaceTransform(float SampleTime, const UPoseSearchSchema* Schema, int8 SchemaSampleBoneIdx = RootSchemaBoneIdx);
+	// returns the world space transform of the bone SchemaBoneIdx at time SampleTime
+	FTransform GetWorldBoneTransformAtTime(float SampleTime, const UPoseSearchSchema* Schema, int8 SchemaBoneIdx = RootSchemaBoneIdx, bool bUseHistoryRoot = false);
+	
+	// returns the local space transform relative to the root bone of the bone SchemaBoneIdx at time SampleTime
+	FTransform GetLocalBoneTransformAtTime(float SampleTime, const UPoseSearchSchema* Schema, int8 SchemaSampleBoneIdx = RootSchemaBoneIdx);
+	
 	FVector GetSamplePositionInternal(float SampleTime, float OriginTime, const UPoseSearchSchema* Schema, int8 SchemaSampleBoneIdx = RootSchemaBoneIdx, int8 SchemaOriginBoneIdx = RootSchemaBoneIdx, bool bUseHistoryRoot = false, const FVector* SampleBonePositionWorldOverride = nullptr);
 	FQuat GetSampleRotationInternal(float SampleTime, float OriginTime, const UPoseSearchSchema* Schema, int8 SchemaSampleBoneIdx = RootSchemaBoneIdx, int8 SchemaOriginBoneIdx = RootSchemaBoneIdx, bool bUseHistoryRoot = false, const FQuat* SampleBoneRotationWorldOverride = nullptr);
 
 	const UAnimInstance* AnimInstance = nullptr;
+	// Trajectory has been transformed in root bone world space reference system
 	const FPoseSearchQueryTrajectory* Trajectory = nullptr;
 	const IPoseHistory* History = nullptr;
 	float DesiredPermutationTimeOffset = 0.f;
diff --git a/Engine/Plugins/Experimental/Animation/PoseSearch/Source/Runtime/Public/PoseSearch/PoseSearchDefines.h b/Engine/Plugins/Experimental/Animation/PoseSearch/Source/Runtime/Public/PoseSearch/PoseSearchDefines.h
--- a/Engine/Plugins/Experimental/Animation/PoseSearch/Source/Runtime/Public/PoseSearch/PoseSearchDefines.h
+++ b/Engine/Plugins/Experimental/Animation/PoseSearch/Source/Runtime/Public/PoseSearch/PoseSearchDefines.h
@@ -23,6 +23,8 @@ DECLARE_LOG_CATEGORY_EXTERN(LogPoseSearch, Log, All);
 
 static constexpr int8 RootSchemaBoneIdx = 0;
 static constexpr FBoneIndexType RootBoneIndexType = 0;
+static constexpr FBoneIndexType ComponentSpaceIndexType = FBoneIndexType(-1);
+static constexpr FBoneIndexType WorldSpaceIndexType = FBoneIndexType(-2);
 
 static constexpr int32 PreallocatedCachedQueriesNum = 8;
 
diff --git a/Engine/Plugins/Experimental/Animation/PoseSearch/Source/Runtime/Public/PoseSearch/PoseSearchHistory.h b/Engine/Plugins/Experimental/Animation/PoseSearch/Source/Runtime/Public/PoseSearch/PoseSearchHistory.h
--- a/Engine/Plugins/Experimental/Animation/PoseSearch/Source/Runtime/Public/PoseSearch/PoseSearchHistory.h
+++ b/Engine/Plugins/Experimental/Animation/PoseSearch/Source/Runtime/Public/PoseSearch/PoseSearchHistory.h
@@ -6,6 +6,7 @@
 #include "BonePose.h"
 #include "Containers/RingBuffer.h"
 #include "DrawDebugHelpers.h"
+#include "PoseSearch/PoseSearchDefines.h"
 #include "UObject/ObjectKey.h"
 
 struct FAnimInstanceProxy;
@@ -24,14 +25,22 @@ struct POSESEARCH_API IPoseHistory
 {
 	virtual ~IPoseHistory() {}
 	virtual float GetSampleTimeInterval() const = 0;
-	virtual bool GetComponentSpaceTransformAtTime(float Time, FBoneIndexType BoneIndexType, const USkeleton* BoneIndexSkeleton, FTransform& OutBoneTransform, bool bExtrapolate = true) const = 0;
-	virtual void GetRootTransformAtTime(float Time, FTransform& OutRootTransform, bool bExtrapolate = true) const = 0;
+	
+	// returns the BoneIndexType transform relative to ReferenceBoneIndexType: 
+	// if ReferenceBoneIndexType is 0 (RootBoneIndexType), OutBoneTransform is in root bone space
+	// if ReferenceBoneIndexType is FBoneIndexType(-1) (ComponentSpaceIndexType), OutBoneTransform is in component space
+	// if ReferenceBoneIndexType is FBoneIndexType(-2) (WorldSpaceIndexType), OutBoneTransform is in world space
+	virtual bool GetTransformAtTime(float Time, FTransform& OutBoneTransform, const USkeleton* BoneIndexSkeleton = nullptr, FBoneIndexType BoneIndexType = RootBoneIndexType, FBoneIndexType ReferenceBoneIndexType = ComponentSpaceIndexType, bool bExtrapolate = true) const = 0;
+
 	virtual bool IsEmpty() const = 0;
 };
 
 struct FPoseHistoryEntry
 {
+	// collected root component transform in world space (not the root bone transform)
 	FTransform RootTransform;
+
+	// collected bones transforms in component space
 	TArray<FTransform> ComponentSpaceTransforms;
 	float Time = 0.f;
 
@@ -49,23 +58,30 @@ struct FPoseHistory : public IPoseHistory
 
 	const FBoneToTransformMap& GetBoneToTransformMap() const { return BoneToTransformMap; }
 	const FPoseHistoryEntries& GetEntries() const { return Entries; }
+	const USkeleton* GetLastUpdateSkeleton() const { return LastUpdateSkeleton.Get(); }
 
 	// IPoseHistory interface
 	virtual float GetSampleTimeInterval() const override;
-	virtual bool GetComponentSpaceTransformAtTime(float Time, FBoneIndexType BoneIndexType, const USkeleton* BoneIndexSkeleton, FTransform& OutBoneTransform, bool bExtrapolate = true) const override;
-	virtual void GetRootTransformAtTime(float Time, FTransform& OutRootTransform, bool bExtrapolate = true) const override;
+	virtual bool GetTransformAtTime(float Time, FTransform& OutBoneTransform, const USkeleton* BoneIndexSkeleton = nullptr, FBoneIndexType BoneIndexType = RootBoneIndexType, FBoneIndexType ReferenceBoneIndexType = ComponentSpaceIndexType, bool bExtrapolate = true) const;
 	virtual bool IsEmpty() const override;
 	// End of IPoseHistory interface
 
-	FBoneIndexType GetRemappedBoneIndexType(FBoneIndexType BoneIndexType, const USkeleton* BoneIndexSkeleton) const;
+	static FBoneIndexType GetRemappedBoneIndexType(FBoneIndexType BoneIndexType, const USkeleton* BoneIndexSkeleton, const USkeleton* LastUpdateSkeleton);
+	static FComponentSpaceTransformIndex GetRemappedComponentSpaceTransformIndex(const USkeleton* BoneIndexSkeleton, const USkeleton* LastUpdateSkeleton, const FBoneToTransformMap& BoneToTransformMap, FBoneIndexType BoneIndexType, bool& bSuccess);
+	static bool LerpEntries(const FPoseHistoryEntry& PrevEntry, const FPoseHistoryEntry& NextEntry, float LerpValue, const USkeleton* BoneIndexSkeleton, const USkeleton* LastUpdateSkeleton, const FBoneToTransformMap& BoneToTransformMap, FBoneIndexType BoneIndexType, FBoneIndexType ReferenceBoneIndexType, FTransform& OutBoneTransform);
 
 #if ENABLE_DRAW_DEBUG && ENABLE_ANIM_DEBUG
 	void DebugDraw(FAnimInstanceProxy& AnimInstanceProxy) const;
 #endif
 
 private:
+	// skeleton from the last Update, to keep tracking skeleton changes, and support compatible skeletons
 	TWeakObjectPtr<const USkeleton> LastUpdateSkeleton;
+
+	// map of FBoneIndexType(s) to collect. If Empty all the bones get collected
 	FBoneToTransformMap BoneToTransformMap;
+
+	// ring buffer of collected bones
 	FPoseHistoryEntries Entries;
 	float TimeHorizon = 0.f;
 };
@@ -78,8 +94,7 @@ struct FExtendedPoseHistory : public IPoseHistory
 
 	// IPoseHistory interface
 	virtual float GetSampleTimeInterval() const override;
-	virtual bool GetComponentSpaceTransformAtTime(float Time, FBoneIndexType BoneIndexType, const USkeleton* BoneIndexSkeleton, FTransform& OutBoneTransform, bool bExtrapolate = true) const override;
-	virtual void GetRootTransformAtTime(float Time, FTransform& OutRootTransform, bool bExtrapolate = true) const override;
+	virtual bool GetTransformAtTime(float Time, FTransform& OutBoneTransform, const USkeleton* BoneIndexSkeleton = nullptr, FBoneIndexType BoneIndexType = RootBoneIndexType, FBoneIndexType ReferenceBoneIndexType = ComponentSpaceIndexType, bool bExtrapolate = true) const;
 	virtual bool IsEmpty() const override;
 	// End of IPoseHistory interface
 
diff --git a/Engine/Plugins/Experimental/Animation/PoseSearch/Source/Runtime/Public/PoseSearch/PoseSearchLibrary.h b/Engine/Plugins/Experimental/Animation/PoseSearch/Source/Runtime/Public/PoseSearch/PoseSearchLibrary.h
--- a/Engine/Plugins/Experimental/Animation/PoseSearch/Source/Runtime/Public/PoseSearch/PoseSearchLibrary.h
+++ b/Engine/Plugins/Experimental/Animation/PoseSearch/Source/Runtime/Public/PoseSearch/PoseSearchLibrary.h
@@ -92,7 +92,6 @@ class POSESEARCH_API UPoseSearchLibrary : public UBlueprintFunctionLibrary
 
 	static FPoseSearchQueryTrajectory ProcessTrajectory(
 		const FPoseSearchQueryTrajectory& Trajectory,
-		const FTransform& ComponentWorldTransform,
 		float RootBoneDeltaYaw,
 		float YawFromAnimationTrajectoryBlendTime,
 		float TrajectorySpeedMultiplier);
