motion matching - vantage point tree WIP + general small improvements

diff --git a/Engine/Plugins/Experimental/Animation/PoseSearch/Source/Editor/Private/PoseSearchDatabaseEditorReflection.cpp b/Engine/Plugins/Experimental/Animation/PoseSearch/Source/Editor/Private/PoseSearchDatabaseEditorReflection.cpp
--- a/Engine/Plugins/Experimental/Animation/PoseSearch/Source/Editor/Private/PoseSearchDatabaseEditorReflection.cpp
+++ b/Engine/Plugins/Experimental/Animation/PoseSearch/Source/Editor/Private/PoseSearchDatabaseEditorReflection.cpp
@@ -135,7 +135,7 @@ void UPoseSearchDatabaseStatistics::Initialize(const UPoseSearchDatabase* PoseSe
 
 		if (ConfigCardinality > 0)
 		{
-			const int32 TotalAnimationFeatureVectors = SearchIndex.Values.Num() / ConfigCardinality;
+			const int32 TotalAnimationFeatureVectors = SearchIndex.GetNumValuesVectors(ConfigCardinality);
 			PrunedFrames = TotalAnimationPosesInFrames - TotalAnimationFeatureVectors;
 		}
 		else
@@ -145,7 +145,7 @@ void UPoseSearchDatabaseStatistics::Initialize(const UPoseSearchDatabase* PoseSe
 
 		if (PoseSearchDatabase->GetNumberOfPrincipalComponents() > 0)
 		{
-			const int32 TotalAnimationPCAFeatureVectors = SearchIndex.PCAValues.Num() / PoseSearchDatabase->GetNumberOfPrincipalComponents();
+			const int32 TotalAnimationPCAFeatureVectors = SearchIndex.GetNumPCAValuesVectors(PoseSearchDatabase->GetNumberOfPrincipalComponents());
 			PrunedPCAFrames = TotalAnimationPosesInFrames - TotalAnimationPCAFeatureVectors;
 		}
 		else
@@ -171,16 +171,19 @@ void UPoseSearchDatabaseStatistics::Initialize(const UPoseSearchDatabase* PoseSe
 			const uint32 ValuesBytesSize = SearchIndex.Values.GetAllocatedSize();
 			const uint32 PCAValuesBytesSize = SearchIndex.PCAValues.GetAllocatedSize();
 			const uint32 KDTreeBytesSize = SearchIndex.KDTree.GetAllocatedSize();
+			const uint32 VPTreeBytesSize = SearchIndex.VPTree.GetAllocatedSize();
+			const uint32 ValuesVectorToPoseIndexesBytesSize = SearchIndex.ValuesVectorToPoseIndexes.GetAllocatedSize();
 			const uint32 PCAValuesVectorToPoseIndexesBytesSize = SearchIndex.PCAValuesVectorToPoseIndexes.GetAllocatedSize();
 
 			const uint32 PoseMetadataBytesSize = SearchIndex.PoseMetadata.GetAllocatedSize();
 			const uint32 AssetsBytesSize = SearchIndex.Assets.GetAllocatedSize();
 			const uint32 OtherBytesSize = SearchIndex.PCAProjectionMatrix.GetAllocatedSize() + SearchIndex.Mean.GetAllocatedSize() + SearchIndex.WeightsSqrt.GetAllocatedSize();
-			const uint32 EstimatedDatabaseBytesSize = ValuesBytesSize + PCAValuesBytesSize + KDTreeBytesSize + PCAValuesVectorToPoseIndexesBytesSize + PoseMetadataBytesSize + AssetsBytesSize + OtherBytesSize;
+			const uint32 EstimatedDatabaseBytesSize = ValuesBytesSize + PCAValuesBytesSize + KDTreeBytesSize + VPTreeBytesSize + ValuesVectorToPoseIndexesBytesSize + PCAValuesVectorToPoseIndexesBytesSize + PoseMetadataBytesSize + AssetsBytesSize + OtherBytesSize;
 				
 			ValuesSize = FText::AsMemory(ValuesBytesSize);
 			PCAValuesSize = FText::AsMemory(PCAValuesBytesSize);
 			KDTreeSize = FText::AsMemory(KDTreeBytesSize);
+			VPTreeSize = FText::AsMemory(VPTreeBytesSize);
 			PoseMetadataSize = FText::AsMemory(PoseMetadataBytesSize);
 			AssetsSize = FText::AsMemory(AssetsBytesSize);
 			EstimatedDatabaseSize = FText::AsMemory(EstimatedDatabaseBytesSize);
diff --git a/Engine/Plugins/Experimental/Animation/PoseSearch/Source/Editor/Private/PoseSearchDatabaseViewModel.cpp b/Engine/Plugins/Experimental/Animation/PoseSearch/Source/Editor/Private/PoseSearchDatabaseViewModel.cpp
--- a/Engine/Plugins/Experimental/Animation/PoseSearch/Source/Editor/Private/PoseSearchDatabaseViewModel.cpp
+++ b/Engine/Plugins/Experimental/Animation/PoseSearch/Source/Editor/Private/PoseSearchDatabaseViewModel.cpp
@@ -381,7 +381,7 @@ namespace UE::PoseSearch
 
 	int32 FDatabaseViewModel::SetSelectedNode(int32 PoseIdx, bool bClearSelection, bool bDrawQuery, TConstArrayView<float> InQueryVector)
 	{
-		int32 SelectedSourceAssetIdx = -1;
+		int32 SelectedSourceAssetIdx = INDEX_NONE;
 
 		if (bClearSelection)
 		{
diff --git a/Engine/Plugins/Experimental/Animation/PoseSearch/Source/Editor/Private/PoseSearchDebugger.cpp b/Engine/Plugins/Experimental/Animation/PoseSearch/Source/Editor/Private/PoseSearchDebugger.cpp
--- a/Engine/Plugins/Experimental/Animation/PoseSearch/Source/Editor/Private/PoseSearchDebugger.cpp
+++ b/Engine/Plugins/Experimental/Animation/PoseSearch/Source/Editor/Private/PoseSearchDebugger.cpp
@@ -184,11 +184,11 @@ FReply SCostTimelineView::OnMouseMove(const FGeometry& MyGeometry, const FPointe
 
 		// Get curve value at given time
 		const TArray<FCurvePoint>& CurvePoints = BestCostData->Points;
-		const int NumPoints = CurvePoints.Num();
+		const int32 NumPoints = CurvePoints.Num();
 
 		if (NumPoints > 0)
 		{
-			for (int i = 1; i < NumPoints; ++i)
+			for (int32 i = 1; i < NumPoints; ++i)
 			{
 				const FCurvePoint& Point1 = CurvePoints[i - 1];
 				const FCurvePoint& Point2 = CurvePoints[i];
@@ -201,7 +201,7 @@ FReply SCostTimelineView::OnMouseMove(const FGeometry& MyGeometry, const FPointe
 					const float Delta2 = abs(TargetTime - Point2.Time);
 
 					// Get closest point index
-					const int TargetPointIndex = Delta1 < Delta2 ? i - 1 : i;
+					const int32 TargetPointIndex = Delta1 < Delta2 ? i - 1 : i;
 
 					const float Time = CurvePoints[TargetPointIndex].Time;
 					const float BestCost = CurvePoints[TargetPointIndex].Value;
@@ -331,7 +331,7 @@ void FDebugger::Update(float DeltaTime, IRewindDebugger* InRewindDebugger)
 void FDebugger::OnViewClosed(uint64 InAnimInstanceId)
 {
 	TArray<TSharedRef<FDebuggerViewModel>>& Models = Debugger->ViewModels;
-	for (int i = 0; i < Models.Num(); ++i)
+	for (int32 i = 0; i < Models.Num(); ++i)
 	{
 		if (Models[i]->AnimInstanceId == InAnimInstanceId)
 		{
@@ -346,7 +346,7 @@ void FDebugger::OnViewClosed(uint64 InAnimInstanceId)
 TSharedPtr<FDebuggerViewModel> FDebugger::GetViewModel(uint64 InAnimInstanceId)
 {
 	TArray<TSharedRef<FDebuggerViewModel>>& Models = Debugger->ViewModels;
-	for (int i = 0; i < Models.Num(); ++i)
+	for (int32 i = 0; i < Models.Num(); ++i)
 	{
 		if (Models[i]->AnimInstanceId == InAnimInstanceId)
 		{
diff --git a/Engine/Plugins/Experimental/Animation/PoseSearch/Source/Editor/Private/PoseSearchDebugger.h b/Engine/Plugins/Experimental/Animation/PoseSearch/Source/Editor/Private/PoseSearchDebugger.h
--- a/Engine/Plugins/Experimental/Animation/PoseSearch/Source/Editor/Private/PoseSearchDebugger.h
+++ b/Engine/Plugins/Experimental/Animation/PoseSearch/Source/Editor/Private/PoseSearchDebugger.h
@@ -77,7 +77,7 @@ private:
 	virtual FName GetNameInternal() const override;
 	virtual FText GetDisplayNameInternal() const override;
 	virtual uint64 GetObjectIdInternal() const override { return ObjectId; }
-	virtual int GetSortOrderPriorityInternal() const override { return 10; };
+	virtual int32 GetSortOrderPriorityInternal() const override { return 10; };
 	virtual bool UpdateInternal() override;
 
 	TSharedPtr<SCostTimelineView> CostTimelineView;
diff --git a/Engine/Plugins/Experimental/Animation/PoseSearch/Source/Editor/Private/PoseSearchDebuggerDatabaseView.cpp b/Engine/Plugins/Experimental/Animation/PoseSearch/Source/Editor/Private/PoseSearchDebuggerDatabaseView.cpp
--- a/Engine/Plugins/Experimental/Animation/PoseSearch/Source/Editor/Private/PoseSearchDebuggerDatabaseView.cpp
+++ b/Engine/Plugins/Experimental/Animation/PoseSearch/Source/Editor/Private/PoseSearchDebuggerDatabaseView.cpp
@@ -138,7 +138,7 @@ private:
 	{
 		float Sum = 0.f;
 		const int32 EndIndex = StartIndex + Offset;
-		for (int i = StartIndex; i < EndIndex; ++i)
+		for (int32 i = StartIndex; i < EndIndex; ++i)
 		{
 			Sum += View[i];
 		}
@@ -269,7 +269,7 @@ void SDebuggerDatabaseView::Update(const FTraceMotionMatchingStateMessage& State
 			SharedData->DatabasePath = Database->GetPathName();
 			SharedData->QueryVector = DbEntry.QueryVector;
 
-			if (Database->PoseSearchMode != EPoseSearchMode::BruteForce)
+			if (Database->PoseSearchMode == EPoseSearchMode::PCAKDTree)
 			{
 				SharedData->PCAQueryVector.SetNumZeroed(Database->GetNumberOfPrincipalComponents());
 				Database->GetSearchIndex().PCAProject(DbEntry.QueryVector, SharedData->PCAQueryVector);
@@ -331,7 +331,7 @@ void SDebuggerDatabaseView::Update(const FTraceMotionMatchingStateMessage& State
 				MinPCACost = FMath::Min(MinPCACost, PCACost);
 				MaxPCACost = FMath::Max(MaxPCACost, PCACost);
 
-				for (int Index = 0; Index < CostBreakdownsCardinality; ++Index)
+				for (int32 Index = 0; Index < CostBreakdownsCardinality; ++Index)
 				{
 					const float Value = UnfilteredRow->CostBreakdowns[Index];
 					if (Value != UE_MAX_FLT)
@@ -365,7 +365,7 @@ void SDebuggerDatabaseView::Update(const FTraceMotionMatchingStateMessage& State
 				UnfilteredRow->PCACostColor = BestScoreRowColor + (WorstScoreRowColor - BestScoreRowColor) * PCACostBlend;
 
 				UnfilteredRow->CostBreakdownsColors.SetNumUninitialized(CostBreakdownsCardinality);
-				for (int Index = 0; Index < CostBreakdownsCardinality; ++Index)
+				for (int32 Index = 0; Index < CostBreakdownsCardinality; ++Index)
 				{
 					const float CostBreakdownBlend = (UnfilteredRow->CostBreakdowns[Index] - MinCostBreakdowns[Index]) / DeltaCostBreakdowns[Index];
 					UnfilteredRow->CostBreakdownsColors[Index] = BestScoreRowColor + (WorstScoreRowColor - BestScoreRowColor) * CostBreakdownBlend;
diff --git a/Engine/Plugins/Experimental/Animation/PoseSearch/Source/Editor/Private/PoseSearchDebuggerView.h b/Engine/Plugins/Experimental/Animation/PoseSearch/Source/Editor/Private/PoseSearchDebuggerView.h
--- a/Engine/Plugins/Experimental/Animation/PoseSearch/Source/Editor/Private/PoseSearchDebuggerView.h
+++ b/Engine/Plugins/Experimental/Animation/PoseSearch/Source/Editor/Private/PoseSearchDebuggerView.h
@@ -113,7 +113,7 @@ private:
 	FOnViewClosed OnViewClosed;
 
 	/** Active node being debugged */
-	int32 SelectedNodeId = -1;
+	int32 SelectedNodeId = INDEX_NONE;
 
 	/** Database view of the motion matching node */
 	TSharedPtr<SDebuggerDatabaseView> DatabaseView;
diff --git a/Engine/Plugins/Experimental/Animation/PoseSearch/Source/Editor/Private/PoseSearchDebuggerViewModel.cpp b/Engine/Plugins/Experimental/Animation/PoseSearch/Source/Editor/Private/PoseSearchDebuggerViewModel.cpp
--- a/Engine/Plugins/Experimental/Animation/PoseSearch/Source/Editor/Private/PoseSearchDebuggerViewModel.cpp
+++ b/Engine/Plugins/Experimental/Animation/PoseSearch/Source/Editor/Private/PoseSearchDebuggerViewModel.cpp
@@ -127,7 +127,7 @@ void FDebuggerViewModel::OnUpdateNodeSelection(int32 InNodeId)
 
 	// Find node in all motion matching states this frame
 	const int32 NodesNum = NodeIds.Num();
-	for (int i = 0; i < NodesNum; ++i)
+	for (int32 i = 0; i < NodesNum; ++i)
 	{
 		if (NodeIds[i] == InNodeId)
 		{
diff --git a/Engine/Plugins/Experimental/Animation/PoseSearch/Source/Editor/Public/PoseSearchDatabaseEditorReflection.h b/Engine/Plugins/Experimental/Animation/PoseSearch/Source/Editor/Public/PoseSearchDatabaseEditorReflection.h
--- a/Engine/Plugins/Experimental/Animation/PoseSearch/Source/Editor/Public/PoseSearchDatabaseEditorReflection.h
+++ b/Engine/Plugins/Experimental/Animation/PoseSearch/Source/Editor/Public/PoseSearchDatabaseEditorReflection.h
@@ -247,6 +247,10 @@ public:
 	UPROPERTY(VisibleAnywhere, Category = "Memory Information")
 	FText KDTreeSize;
 
+	// partial memory size used by the vptree.
+	UPROPERTY(VisibleAnywhere, Category = "Memory Information")
+	FText VPTreeSize;
+		
 	// partial memory size used to store database metadata.
 	UPROPERTY(VisibleAnywhere, Category = "Memory Information")
 	FText PoseMetadataSize;
diff --git a/Engine/Plugins/Experimental/Animation/PoseSearch/Source/Editor/Public/PoseSearchDatabaseViewModel.h b/Engine/Plugins/Experimental/Animation/PoseSearch/Source/Editor/Public/PoseSearchDatabaseViewModel.h
--- a/Engine/Plugins/Experimental/Animation/PoseSearch/Source/Editor/Public/PoseSearchDatabaseViewModel.h
+++ b/Engine/Plugins/Experimental/Animation/PoseSearch/Source/Editor/Public/PoseSearchDatabaseViewModel.h
@@ -170,7 +170,7 @@ namespace UE::PoseSearch
 
 		UWorld* GetWorld();
 
-		FDatabasePreviewActor SpawnPreviewActor(int32 IndexAssetIndex, int32 PoseIdxForTimeOffset = -1);
+		FDatabasePreviewActor SpawnPreviewActor(int32 IndexAssetIndex, int32 PoseIdxForTimeOffset = INDEX_NONE);
 
 		void UpdatePreviewActors(bool bInTickPlayTime = false);
 
diff --git a/Engine/Plugins/Experimental/Animation/PoseSearch/Source/Runtime/Private/KDTree.cpp b/Engine/Plugins/Experimental/Animation/PoseSearch/Source/Runtime/Private/KDTree.cpp
--- a/Engine/Plugins/Experimental/Animation/PoseSearch/Source/Runtime/Private/KDTree.cpp
+++ b/Engine/Plugins/Experimental/Animation/PoseSearch/Source/Runtime/Private/KDTree.cpp
@@ -294,18 +294,19 @@ void FKDTree::Construct(int32 Count, int32 Dim, const float* Data, int32 MaxLeaf
 	new(this)FKDTree(Count, Dim, Data, MaxLeafSize);
 }
 
-bool FKDTree::FindNeighbors(FKNNResultSet& Result, const float* Query) const
+bool FKDTree::FindNeighbors(FKNNResultSet& Result, TConstArrayView<float> Query) const
 {
 #if UE_POSE_SEARCH_USE_NANOFLANN
 
 	QUICK_SCOPE_CYCLE_COUNTER(STAT_FKDTree_FindNeighbors);
 
-	check(Query && Impl->root_node);
+	check(Query.GetData() && Query.Num() == Impl->dim && Impl->root_node);
+
 	const nanoflann::SearchParams SearchParams(
 		32,			// Ignored parameter (Kept for compatibility with the FLANN interface).
 		0.f,		// search for eps-approximate neighbours (default: 0)
 		false);		// only for radius search, require neighbours sorted by
-	return Impl->findNeighbors(Result, Query, SearchParams);
+	return Impl->findNeighbors(Result, Query.GetData(), SearchParams);
 
 #else // UE_POSE_SEARCH_USE_NANOFLANN
 
@@ -315,18 +316,19 @@ bool FKDTree::FindNeighbors(FKNNResultSet& Result, const float* Query) const
 #endif // UE_POSE_SEARCH_USE_NANOFLANN
 }
 
-bool FKDTree::FindNeighbors(FRadiusResultSet& Result, const float* Query) const
+bool FKDTree::FindNeighbors(FRadiusResultSet& Result, TConstArrayView<float> Query) const
 {
 #if UE_POSE_SEARCH_USE_NANOFLANN
 
 	QUICK_SCOPE_CYCLE_COUNTER(STAT_FKDTree_FindNeighbors);
 
-	check(Query && Impl->root_node);
+	check(Query.GetData() && Query.Num() == Impl->dim && Impl->root_node);
+
 	const nanoflann::SearchParams SearchParams(
 		32,			// Ignored parameter (Kept for compatibility with the FLANN interface).
 		0.f,		// search for eps-approximate neighbours (default: 0)
 		false);		// only for radius search, require neighbours sorted by
-	return Impl->findNeighbors(Result, Query, SearchParams);
+	return Impl->findNeighbors(Result, Query.GetData(), SearchParams);
 
 #else // UE_POSE_SEARCH_USE_NANOFLANN
 
diff --git a/Engine/Plugins/Experimental/Animation/PoseSearch/Source/Runtime/Private/PoseSearchAssetIndexer.cpp b/Engine/Plugins/Experimental/Animation/PoseSearch/Source/Runtime/Private/PoseSearchAssetIndexer.cpp
--- a/Engine/Plugins/Experimental/Animation/PoseSearch/Source/Runtime/Private/PoseSearchAssetIndexer.cpp
+++ b/Engine/Plugins/Experimental/Animation/PoseSearch/Source/Runtime/Private/PoseSearchAssetIndexer.cpp
@@ -18,7 +18,7 @@ namespace UE::PoseSearch
 
 #if ENABLE_ANIM_DEBUG
 static TAutoConsoleVariable<bool> CVarMotionMatchTestDisableIndexerCaching(TEXT("a.MotionMatch.TestDisableIndexerCaching"), false, TEXT("Disable Motion Matching Indexer Caching"));
-static TAutoConsoleVariable<int> CVarMotionMatchTestExtractPoseDeterminismNumIterations(TEXT("a.MotionMatch.TestExtractPoseDeterminismNumIterations"), 0, TEXT("Test Motion Matching ExtractPose Determinism via this NumIterations retries"));
+static TAutoConsoleVariable<int32> CVarMotionMatchTestExtractPoseDeterminismNumIterations(TEXT("a.MotionMatch.TestExtractPoseDeterminismNumIterations"), 0, TEXT("Test Motion Matching ExtractPose Determinism via this NumIterations retries"));
 #endif // ENABLE_ANIM_DEBUG
 
 //////////////////////////////////////////////////////////////////////////
@@ -695,8 +695,8 @@ bool FAssetIndexer::GetSampleVelocity(FVector& OutSampleVelocity, float SampleTi
 		return false;
 	}
 
-	if (GetSamplePositionInternal(BonePositionPast, SampleTime - FiniteDelta, bUseCharacterSpaceVelocities ? OriginTime - FiniteDelta : OriginTime, bClampedPast, SchemaSampleBoneIdx, SchemaOriginBoneIdx, -1) &&
-		GetSamplePositionInternal(BonePositionPresent, SampleTime, OriginTime, bUnused, SchemaSampleBoneIdx, SchemaOriginBoneIdx, -1))
+	if (GetSamplePositionInternal(BonePositionPast, SampleTime - FiniteDelta, bUseCharacterSpaceVelocities ? OriginTime - FiniteDelta : OriginTime, bClampedPast, SchemaSampleBoneIdx, SchemaOriginBoneIdx, INDEX_NONE) &&
+		GetSamplePositionInternal(BonePositionPresent, SampleTime, OriginTime, bUnused, SchemaSampleBoneIdx, SchemaOriginBoneIdx, INDEX_NONE))
 	{
 		if (!bClampedPast)
 		{
@@ -705,7 +705,7 @@ bool FAssetIndexer::GetSampleVelocity(FVector& OutSampleVelocity, float SampleTi
 		}
 
 		FVector BonePositionFuture;
-		if (GetSamplePositionInternal(BonePositionFuture, SampleTime + FiniteDelta, bUseCharacterSpaceVelocities ? OriginTime + FiniteDelta : OriginTime, bUnused, SchemaSampleBoneIdx, SchemaOriginBoneIdx, -1))
+		if (GetSamplePositionInternal(BonePositionFuture, SampleTime + FiniteDelta, bUseCharacterSpaceVelocities ? OriginTime + FiniteDelta : OriginTime, bUnused, SchemaSampleBoneIdx, SchemaOriginBoneIdx, INDEX_NONE))
 		{
 			OutSampleVelocity = (BonePositionFuture - BonePositionPresent) / FiniteDelta;
 			return true;
diff --git a/Engine/Plugins/Experimental/Animation/PoseSearch/Source/Runtime/Private/PoseSearchAssetSampler.cpp b/Engine/Plugins/Experimental/Animation/PoseSearch/Source/Runtime/Private/PoseSearchAssetSampler.cpp
--- a/Engine/Plugins/Experimental/Animation/PoseSearch/Source/Runtime/Private/PoseSearchAssetSampler.cpp
+++ b/Engine/Plugins/Experimental/Animation/PoseSearch/Source/Runtime/Private/PoseSearchAssetSampler.cpp
@@ -35,7 +35,7 @@ static FTransform ExtrapolateRootMotion(FTransform SampleToExtrapolate, float Sa
 
 	// adding full samples to the extrapolated root motion
 	FTransform ExtrapolatedRootMotion = FTransform::Identity;
-	for (int i = 0; i < NumSamples; ++i)
+	for (int32 i = 0; i < NumSamples; ++i)
 	{
 		ExtrapolatedRootMotion = AbsTimeSampleToExtrapolate * ExtrapolatedRootMotion;
 	}
diff --git a/Engine/Plugins/Experimental/Animation/PoseSearch/Source/Runtime/Private/PoseSearchContext.cpp b/Engine/Plugins/Experimental/Animation/PoseSearch/Source/Runtime/Private/PoseSearchContext.cpp
--- a/Engine/Plugins/Experimental/Animation/PoseSearch/Source/Runtime/Private/PoseSearchContext.cpp
+++ b/Engine/Plugins/Experimental/Animation/PoseSearch/Source/Runtime/Private/PoseSearchContext.cpp
@@ -275,7 +275,7 @@ void FDebugDrawParams::DrawCentripetalCatmullRomSpline(TConstArrayView<FVector>
 		};
 
 		FVector PrevPoint = Points[0];
-		for (int i = 0; i < NumPoints - 1; ++i)
+		for (int32 i = 0; i < NumPoints - 1; ++i)
 		{
 			const FVector& P0 = Points[FMath::Max(i - 1, 0)];
 			const FVector& P1 = Points[i];
@@ -302,7 +302,7 @@ void FDebugDrawParams::DrawCentripetalCatmullRomSpline(TConstArrayView<FVector>
 			const FColor Color1 = Colors[FMath::Min(i, NumColors - 1)];
 			const FColor Color2 = Colors[FMath::Min(i + 1, NumColors - 1)];
 
-			for (int SampleIndex = 1; SampleIndex < NumSamplesPerSegment; ++SampleIndex)
+			for (int32 SampleIndex = 1; SampleIndex < NumSamplesPerSegment; ++SampleIndex)
 			{
 				const float ParametricDistance = float(SampleIndex) / float(NumSamplesPerSegment - 1);
 
diff --git a/Engine/Plugins/Experimental/Animation/PoseSearch/Source/Runtime/Private/PoseSearchDatabase.cpp b/Engine/Plugins/Experimental/Animation/PoseSearch/Source/Runtime/Private/PoseSearchDatabase.cpp
--- a/Engine/Plugins/Experimental/Animation/PoseSearch/Source/Runtime/Private/PoseSearchDatabase.cpp
+++ b/Engine/Plugins/Experimental/Animation/PoseSearch/Source/Runtime/Private/PoseSearchDatabase.cpp
@@ -23,11 +23,16 @@
 DECLARE_STATS_GROUP(TEXT("PoseSearch"), STATGROUP_PoseSearch, STATCAT_Advanced);
 DECLARE_CYCLE_STAT_EXTERN(TEXT("Search Brute Force"), STAT_PoseSearch_BruteForce, STATGROUP_PoseSearch, );
 DECLARE_CYCLE_STAT_EXTERN(TEXT("Search PCA/KNN"), STAT_PoseSearch_PCAKNN, STATGROUP_PoseSearch, );
+DECLARE_CYCLE_STAT_EXTERN(TEXT("Search VPTree"), STAT_PoseSearch_VPTree, STATGROUP_PoseSearch, );
 DEFINE_STAT(STAT_PoseSearch_BruteForce);
 DEFINE_STAT(STAT_PoseSearch_PCAKNN);
+DEFINE_STAT(STAT_PoseSearch_VPTree);
 
 namespace UE::PoseSearch
 {
+#if WITH_EDITOR && ENABLE_ANIM_DEBUG
+static TAutoConsoleVariable<bool> CVarMotionMatchCompareAgainstBruteForce(TEXT("a.MotionMatch.CompareAgainstBruteForce"), false, TEXT("Compare optimized search against brute force search"));
+#endif
 
 typedef TArray<size_t, TInlineAllocator<256>> FNonSelectableIdx;
 static void PopulateNonSelectableIdx(FNonSelectableIdx& NonSelectableIdx, FSearchContext& SearchContext, const UPoseSearchDatabase* Database
@@ -139,7 +144,7 @@ struct FSearchFilters
 		}
 	}
 
-	bool AreFiltersValid(const FSearchIndex& SearchIndex, TConstArrayView<float> PoseValues, TConstArrayView<float> QueryValues, int32 PoseIdx, const FPoseMetadata& Metadata
+	bool AreFiltersValid(const FSearchIndex& SearchIndex, TConstArrayView<float> PoseValues, TConstArrayView<float> QueryValues, int32 PoseIdx
 #if UE_POSE_SEARCH_TRACE_ENABLED
 		, UE::PoseSearch::FSearchContext& SearchContext, const UPoseSearchDatabase* Database
 #endif // UE_POSE_SEARCH_TRACE_ENABLED
@@ -147,7 +152,7 @@ struct FSearchFilters
 	{
 		for (const IPoseSearchFilter* Filter : Filters)
 		{
-			if (!Filter->IsFilterValid(PoseValues, QueryValues, PoseIdx, Metadata))
+			if (!Filter->IsFilterValid(PoseValues, QueryValues, PoseIdx, SearchIndex.PoseMetadata[PoseIdx]))
 			{
 #if UE_POSE_SEARCH_TRACE_ENABLED
 				if (Filter == &NonSelectableIdxFilter)
@@ -501,12 +506,12 @@ int32 UPoseSearchDatabase::GetNumberOfPrincipalComponents() const
 
 bool UPoseSearchDatabase::GetSkipSearchIfPossible() const
 {
-#if WITH_EDITOR
-	if (PoseSearchMode == EPoseSearchMode::PCAKDTree_Validate || PoseSearchMode == EPoseSearchMode::PCAKDTree_Compare)
+#if WITH_EDITOR && ENABLE_ANIM_DEBUG
+	if (UE::PoseSearch::CVarMotionMatchCompareAgainstBruteForce.GetValueOnAnyThread())
 	{
 		return false;
 	}
-#endif // WITH_EDITOR
+#endif // WITH_EDITOR && ENABLE_ANIM_DEBUG
 	return bSkipSearchIfPossible;
 }
 
@@ -642,25 +647,38 @@ UE::PoseSearch::FSearchResult UPoseSearchDatabase::Search(UE::PoseSearch::FSearc
 #endif // WITH_EDITOR
 
 	if (PoseSearchMode == EPoseSearchMode::BruteForce
-#if WITH_EDITOR
-		|| PoseSearchMode == EPoseSearchMode::PCAKDTree_Compare
-#endif // WITH_EDITOR
+#if WITH_EDITOR && ENABLE_ANIM_DEBUG
+		|| CVarMotionMatchCompareAgainstBruteForce.GetValueOnAnyThread()
+#endif // WITH_EDITOR && ENABLE_ANIM_DEBUG
 		)
 	{
 		Result = SearchBruteForce(SearchContext);
 	}
 
-	if (PoseSearchMode != EPoseSearchMode::BruteForce)
+#if WITH_EDITOR && ENABLE_ANIM_DEBUG
+	const FPoseSearchCost BruteForcePoseCost = Result.BruteForcePoseCost;
+	const int32 BruteForcePoseIdx = Result.PoseIdx;
+#endif // WITH_EDITOR && ENABLE_ANIM_DEBUG
+
+	if (PoseSearchMode == EPoseSearchMode::VPTree)
 	{
-#if UE_POSE_SEARCH_TRACE_ENABLED
-		FPoseSearchCost BruteForcePoseCost = Result.BruteForcePoseCost;
-#endif // UE_POSE_SEARCH_TRACE_ENABLED
+		Result = SearchVPTree(SearchContext);
 
+#if WITH_EDITOR && ENABLE_ANIM_DEBUG
+		Result.BruteForcePoseCost = BruteForcePoseCost;
+		if (Result.PoseIdx != BruteForcePoseIdx && CVarMotionMatchCompareAgainstBruteForce.GetValueOnAnyThread())
+		{
+			UE_LOG(LogPoseSearch, Error, TEXT("UPoseSearchDatabase::Search - VPTree search PoseIdx %d differs from BruteForce search PoseIdx %d"), Result.PoseIdx, BruteForcePoseIdx);
+		}
+#endif // UE_POSE_SEARCH_TRACE_ENABLED
+	}
+	else if (PoseSearchMode == EPoseSearchMode::PCAKDTree)
+	{
 		Result = SearchPCAKDTree(SearchContext);
 
-#if UE_POSE_SEARCH_TRACE_ENABLED
+#if WITH_EDITOR && ENABLE_ANIM_DEBUG
 		Result.BruteForcePoseCost = BruteForcePoseCost;
-		if (PoseSearchMode == EPoseSearchMode::PCAKDTree_Compare)
+		if (CVarMotionMatchCompareAgainstBruteForce.GetValueOnAnyThread())
 		{
 			const float BruteForceTotalCost = Result.BruteForcePoseCost.GetTotalCost();
 			const float PCAKDTreeTotalCost = Result.PoseCost.GetTotalCost();
@@ -668,12 +686,12 @@ UE::PoseSearch::FSearchResult UPoseSearchDatabase::Search(UE::PoseSearch::FSearc
 
 			if (!FMath::IsNearlyEqual(BruteForceTotalCost, PCAKDTreeTotalCost))
 			{
-				UE_LOG(LogPoseSearch, Warning, TEXT("PCAKDTree_Compare %f (PCAKDTreeTotalCost %f, BruteForceTotalCost %f)"), PCAKDTreeTotalCost - BruteForceTotalCost, PCAKDTreeTotalCost, BruteForceTotalCost);
+				UE_LOG(LogPoseSearch, Warning, TEXT("UPoseSearchDatabase::Search - PCAKDTree cost comparison %f (PCAKDTreeTotalCost %f, BruteForceTotalCost %f)"), PCAKDTreeTotalCost - BruteForceTotalCost, PCAKDTreeTotalCost, BruteForceTotalCost);
 			}
 		}
-#endif // UE_POSE_SEARCH_TRACE_ENABLED
+#endif // WITH_EDITOR && ENABLE_ANIM_DEBUG
 	}
-	
+
 	return Result;
 }
 
@@ -685,7 +703,7 @@ static inline void EvaluatePoseKernel(UE::PoseSearch::FSearchResult& Result, con
 
 	const TConstArrayView<float> PoseValues = bReconstructPoseValues ? SearchIndex.GetReconstructedPoseValues(PoseIdx, ReconstructedPoseValuesBuffer) : SearchIndex.GetPoseValues(PoseIdx);
 
-	if (SearchFilters.AreFiltersValid(SearchIndex, PoseValues, QueryValues, PoseIdx, SearchIndex.PoseMetadata[PoseIdx]
+	if (SearchFilters.AreFiltersValid(SearchIndex, PoseValues, QueryValues, PoseIdx
 #if UE_POSE_SEARCH_TRACE_ENABLED
 		, SearchContext, Database
 #endif // UE_POSE_SEARCH_TRACE_ENABLED
@@ -828,7 +846,7 @@ UE::PoseSearch::FSearchResult UPoseSearchDatabase::SearchPCAKDTree(UE::PoseSearc
 
 		check(QueryValues.Num() == NumDimensions);
 		// projecting QueryValues into the PCA space ProjectedQueryValues and query the KDTree
-		SearchIndex.KDTree.FindNeighbors(ResultSet, SearchIndex.PCAProject(QueryValues, ProjectedQueryValues).GetData());
+		SearchIndex.KDTree.FindNeighbors(ResultSet, SearchIndex.PCAProject(QueryValues, ProjectedQueryValues));
 
 		// NonSelectableIdx are already filtered out inside the kdtree search
 		const FSearchFilters SearchFilters(Schema, bArePCAValuesPruned ? NonSelectableIdx : TConstArrayView<size_t>(), SearchIndex.bAnyBlockTransition);
@@ -847,7 +865,7 @@ UE::PoseSearch::FSearchResult UPoseSearchDatabase::SearchPCAKDTree(UE::PoseSearc
 				for (size_t ResultIndex = 0; ResultIndex < ResultSet.Num(); ++ResultIndex)
 				{
 					const TConstArrayView<uint32> PoseIndexes = SearchIndex.PCAValuesVectorToPoseIndexes[ResultIndexes[ResultIndex]];
-					for (int Index = 0; Index < PoseIndexes.Num() && NumEvaluatePoseKernelCalls < MaxNumEvaluatePoseKernelCalls; ++Index, ++NumEvaluatePoseKernelCalls)
+					for (int32 Index = 0; Index < PoseIndexes.Num() && NumEvaluatePoseKernelCalls < MaxNumEvaluatePoseKernelCalls; ++Index, ++NumEvaluatePoseKernelCalls)
 					{
 						EvaluatePoseKernel<false, true>(Result, SearchIndex, QueryValues, TArrayView<float>(), PoseIndexes[Index], SearchFilters, SearchContext, this, true, ResultIndex);
 					}
@@ -859,7 +877,7 @@ UE::PoseSearch::FSearchResult UPoseSearchDatabase::SearchPCAKDTree(UE::PoseSearc
 				for (size_t ResultIndex = 0; ResultIndex < ResultSet.Num(); ++ResultIndex)
 				{
 					const TConstArrayView<uint32> PoseIndexes = SearchIndex.PCAValuesVectorToPoseIndexes[ResultIndexes[ResultIndex]];
-					for (int Index = 0; Index < PoseIndexes.Num() && NumEvaluatePoseKernelCalls < MaxNumEvaluatePoseKernelCalls; ++Index, ++NumEvaluatePoseKernelCalls)
+					for (int32 Index = 0; Index < PoseIndexes.Num() && NumEvaluatePoseKernelCalls < MaxNumEvaluatePoseKernelCalls; ++Index, ++NumEvaluatePoseKernelCalls)
 					{
 						EvaluatePoseKernel<false, false>(Result, SearchIndex, QueryValues, TArrayView<float>(), PoseIndexes[Index], SearchFilters, SearchContext, this, true, ResultIndex);
 					}
@@ -918,6 +936,95 @@ UE::PoseSearch::FSearchResult UPoseSearchDatabase::SearchPCAKDTree(UE::PoseSearc
 	return Result;
 }
 
+UE::PoseSearch::FSearchResult UPoseSearchDatabase::SearchVPTree(UE::PoseSearch::FSearchContext& SearchContext) const
+{
+	SCOPE_CYCLE_COUNTER(STAT_PoseSearch_VPTree);
+
+	using namespace UE::PoseSearch;
+
+	FSearchResult Result;
+
+	const FSearchIndex& SearchIndex = GetSearchIndex();
+
+	// since any PoseCost calculated here is at least SearchIndex.MinCostAddend,
+	// there's no point in performing the search if CurrentBestTotalCost is already better than that
+	if (!GetSkipSearchIfPossible() || SearchContext.GetCurrentBestTotalCost() > SearchIndex.MinCostAddend)
+	{
+		TConstArrayView<float> QueryValues = SearchContext.GetOrBuildQuery(Schema).GetValues();
+
+		// @todo: implement filtering within the VPTree as KDTree does
+		FNonSelectableIdx NonSelectableIdx;
+		PopulateNonSelectableIdx(NonSelectableIdx, SearchContext, this
+#if UE_POSE_SEARCH_TRACE_ENABLED
+			, QueryValues
+#endif // UE_POSE_SEARCH_TRACE_ENABLED
+		);
+		check(Algo::IsSorted(NonSelectableIdx));
+
+		const int32 NumDimensions = Schema->SchemaCardinality;
+		check(QueryValues.Num() == NumDimensions);
+		
+		const FSearchFilters SearchFilters(Schema, NonSelectableIdx, SearchIndex.bAnyBlockTransition);
+
+		// @todo: implement a FVPTreeDataSource for aligned and padded features vector like CompareAlignedPoses does 
+		FVPTreeDataSource DataSource(SearchIndex);
+		FVPTreeResultSet ResultSet(KDTreeQueryNumNeighbors);
+		SearchIndex.VPTree.FindNeighbors(QueryValues, ResultSet, DataSource);
+		
+		int32 NumEvaluatePoseKernelCalls = 0;
+		const TArray<FIndexDistance>& UnsortedResults = ResultSet.GetUnsortedResults();
+
+		const bool bAreValuesPruned = SearchIndex.ValuesVectorToPoseIndexes.Num() > 0;
+		if (bAreValuesPruned)
+		{
+			const int32 MaxNumEvaluatePoseKernelCalls = KDTreeQueryNumNeighborsWithDuplicates > 0 ? KDTreeQueryNumNeighborsWithDuplicates : INT32_MAX;
+			for (int32 ResultIndex = 0; ResultIndex < UnsortedResults.Num(); ++ResultIndex)
+			{
+				const FIndexDistance& IndexDistance = UnsortedResults[ResultIndex];
+
+				// @todo: IndexDistance.Distance is the Sqrt(DissimilarityCost), so there's no need to calculate it again in SearchIndex.ComparePoses
+				const TConstArrayView<uint32> PoseIndexes = SearchIndex.ValuesVectorToPoseIndexes[IndexDistance.Index];
+				for (int32 Index = 0; Index < PoseIndexes.Num() && NumEvaluatePoseKernelCalls < MaxNumEvaluatePoseKernelCalls; ++Index, ++NumEvaluatePoseKernelCalls)
+				{
+					EvaluatePoseKernel<false, false>(Result, SearchIndex, QueryValues, TArrayView<float>(), PoseIndexes[Index], SearchFilters, SearchContext, this, true, ResultIndex);
+				}
+			}
+		}
+		else
+		{
+			for (int32 ResultIndex = 0; ResultIndex < UnsortedResults.Num(); ++ResultIndex)
+			{
+				const FIndexDistance& IndexDistance = UnsortedResults[ResultIndex];
+
+				// @todo: IndexDistance.Distance is the Sqrt(DissimilarityCost), so there's no need to calculate it again in SearchIndex.ComparePoses
+				EvaluatePoseKernel<false, false>(Result, SearchIndex, QueryValues, TArrayView<float>(), IndexDistance.Index, SearchFilters, SearchContext, this, true, ResultIndex);
+			}
+		}
+	}
+	else
+	{
+#if UE_POSE_SEARCH_TRACE_ENABLED
+		// calling just for reporting non selectable poses
+		TConstArrayView<float> QueryValues = SearchContext.GetOrBuildQuery(Schema).GetValues();
+		FNonSelectableIdx NonSelectableIdx;
+		PopulateNonSelectableIdx(NonSelectableIdx, SearchContext, this, QueryValues);
+#endif // UE_POSE_SEARCH_TRACE_ENABLED
+	}
+
+#if UE_POSE_SEARCH_TRACE_ENABLED
+	SearchContext.BestCandidates.Add(this);
+#endif // UE_POSE_SEARCH_TRACE_ENABLED
+
+	// finalizing Result properties
+	if (Result.PoseIdx != INDEX_NONE)
+	{
+		Result.AssetTime = GetNormalizedAssetTime(Result.PoseIdx);
+		Result.Database = this;
+	}
+
+	return Result;
+}
+
 UE::PoseSearch::FSearchResult UPoseSearchDatabase::SearchBruteForce(UE::PoseSearch::FSearchContext& SearchContext) const
 {
 	SCOPE_CYCLE_COUNTER(STAT_PoseSearch_BruteForce);
diff --git a/Engine/Plugins/Experimental/Animation/PoseSearch/Source/Runtime/Private/PoseSearchDerivedData.cpp b/Engine/Plugins/Experimental/Animation/PoseSearch/Source/Runtime/Private/PoseSearchDerivedData.cpp
--- a/Engine/Plugins/Experimental/Animation/PoseSearch/Source/Runtime/Private/PoseSearchDerivedData.cpp
+++ b/Engine/Plugins/Experimental/Animation/PoseSearch/Source/Runtime/Private/PoseSearchDerivedData.cpp
@@ -48,14 +48,26 @@ enum EMotionMatchTestFlags
 	// test KDTree Construct determinism
 	TestKDTreeConstructDeterminism = 0x8,
 
+	// validating the kdtree construction
+	ValidateKDTreeConstruct = 0x10,
+
+	// test VPTree Construct determinism
+	TestVPTreeConstructDeterminism = 0x20,
+
+	// validating the vptree construction
+	ValidateVPTreeConstruct = 0x40,
+
 	// test IndexDatabase determinism
-	TestIndexDatabaseDeterminism = 0x10,
+	TestIndexDatabaseDeterminism = 0x80,
 
 	// test PruneDuplicateValues determinism
-	TestPruneDuplicateValuesDeterminism = 0x20,
+	TestPruneDuplicateValuesDeterminism = 0x100,
+
+	// test PruneDuplicatePCAValues determinism
+	TestPruneDuplicatePCAValuesDeterminism = 0x200,
 
 	// validating the data we gave to DDC is stored correctly
-	ValidateDDC = 0x40,
+	ValidateDDC = 0x400,
 };
 static TAutoConsoleVariable<int32> CVarMotionMatchTestFlags(TEXT("a.MotionMatch.TestFlags"), EMotionMatchTestFlags::None, TEXT("Test Motion Matching using EMotionMatchTestFlags"));
 static TAutoConsoleVariable<int32> CVarMotionMatchTestNumIterations(TEXT("a.MotionMatch.TestNumIterations"), 10, TEXT("Test Motion Matching Num Iterations"));
@@ -79,7 +91,7 @@ struct FMeanDeviationCalculator
 private:
 	struct FEntry
 	{
-		int32 SchemaIndex = -1; // index of the Channel associated Schemas / SearchIndexBases index
+		int32 SchemaIndex = INDEX_NONE; // index of the Channel associated Schemas / SearchIndexBases index
 		const UPoseSearchFeatureChannel* Channel;
 	};
 
@@ -293,7 +305,7 @@ static void FindValidSequenceIntervals(const UAnimSequenceBase* SequenceBase, FF
 
 				// Split every valid range based on the exclusion range just found. Because this might increase the 
 				// number of ranges in ValidRanges, the algorithm iterates from end to start.
-				for (int RangeIdx = ValidRanges.Num() - 1; RangeIdx >= 0; --RangeIdx)
+				for (int32 RangeIdx = ValidRanges.Num() - 1; RangeIdx >= 0; --RangeIdx)
 				{
 					FFloatRange EvaluatedRange = ValidRanges[RangeIdx];
 					ValidRanges.RemoveAt(RangeIdx);
@@ -462,7 +474,7 @@ static void PreprocessSearchIndexPCAData(FSearchIndex& SearchIndex, int32 NumDim
 	SearchIndex.Mean.Reset();
 	SearchIndex.PCAProjectionMatrix.Reset();
 
-	if (PoseSearchMode != EPoseSearchMode::BruteForce && NumDimensions > 0 && NumPoses > 0 && NumberOfPrincipalComponents > 0)
+	if (PoseSearchMode == EPoseSearchMode::PCAKDTree && NumDimensions > 0 && NumPoses > 0 && NumberOfPrincipalComponents > 0)
 	{
 		SearchIndex.PCAValues.AddZeroed(NumPoses * NumberOfPrincipalComponents);
 		SearchIndex.Mean.AddZeroed(NumDimensions);
@@ -499,7 +511,7 @@ static void PreprocessSearchIndexPCAData(FSearchIndex& SearchIndex, int32 NumDim
 		// validating EigenSolver results
 		const ColMajorMatrix EigenVectors = EigenSolver.eigenvectors().real();
 
-		if (PoseSearchMode == EPoseSearchMode::PCAKDTree_Validate && NumberOfPrincipalComponents == NumDimensions)
+		if (AnyTestFlags(EMotionMatchTestFlags::ValidateKDTreeConstruct) && NumberOfPrincipalComponents == NumDimensions)
 		{
 			const RowMajorVector ReciprocalWeightsSqrt = MapWeightsSqrt.cwiseInverse();
 			const RowMajorMatrix ProjectedValues = CenteredValues * EigenVectors;
@@ -541,7 +553,7 @@ static void PreprocessSearchIndexPCAData(FSearchIndex& SearchIndex, int32 NumDim
 
 		MapPCAValues = CenteredValues * PCAProjectionMatrix;
 
-		if (PoseSearchMode == EPoseSearchMode::PCAKDTree_Validate && NumberOfPrincipalComponents == NumDimensions)
+		if (AnyTestFlags(EMotionMatchTestFlags::ValidateKDTreeConstruct) && NumberOfPrincipalComponents == NumDimensions)
 		{
 			const RowMajorVector ReciprocalWeightsSqrt = MapWeightsSqrt.cwiseInverse();
 			for (Eigen::Index RowIndex = 0; RowIndex < MapValues.rows(); ++RowIndex)
@@ -576,14 +588,12 @@ static void PreprocessSearchIndexKDTree(FSearchIndex& SearchIndex, const UPoseSe
 	const EPoseSearchMode PoseSearchMode = Database->PoseSearchMode;
 
 	SearchIndex.KDTree.Reset();
-	if (PoseSearchMode != EPoseSearchMode::BruteForce && NumDimensions > 0)
+	if (NumDimensions > 0 && PoseSearchMode == EPoseSearchMode::PCAKDTree)
 	{
 		const uint32 NumberOfPrincipalComponents = Database->GetNumberOfPrincipalComponents();
 		const int32 KDTreeMaxLeafSize = Database->KDTreeMaxLeafSize;
 
-		check(SearchIndex.PCAValues.Num() % NumberOfPrincipalComponents == 0);
-
-		const int32 NumPCAValuesVectors = SearchIndex.PCAValues.Num() / NumberOfPrincipalComponents;
+		const int32 NumPCAValuesVectors = SearchIndex.GetNumPCAValuesVectors(NumberOfPrincipalComponents);
 		SearchIndex.KDTree.Construct(NumPCAValuesVectors, NumberOfPrincipalComponents, SearchIndex.PCAValues.GetData(), KDTreeMaxLeafSize);
 
 #if ENABLE_ANIM_DEBUG
@@ -606,72 +616,156 @@ static void PreprocessSearchIndexKDTree(FSearchIndex& SearchIndex, const UPoseSe
 			}
 		}
 
-		// testing the KDTree is returning the proper searches for all the points in pca space
-		const int32 KDTreeQueryNumNeighbors = Database->KDTreeQueryNumNeighbors;
+		if (AnyTestFlags(EMotionMatchTestFlags::ValidateKDTreeConstruct))
+		{
+			// testing the KDTree is returning the proper searches for all the points in pca space
+			const int32 KDTreeQueryNumNeighbors = Database->KDTreeQueryNumNeighbors;
 
-		TArray<size_t> ResultIndexes;
-		TArray<float> ResultDistanceSqr;
-		ResultIndexes.SetNum(NumPCAValuesVectors + 1);
-		ResultDistanceSqr.SetNum(NumPCAValuesVectors + 1);
+			TArray<size_t> ResultIndexes;
+			TArray<float> ResultDistanceSqr;
+			ResultIndexes.SetNum(NumPCAValuesVectors + 1);
+			ResultDistanceSqr.SetNum(NumPCAValuesVectors + 1);
 
-		size_t MaxNumNeighborToFindAPoint = 0;
-		for (size_t PointIndex = 0; PointIndex < NumPCAValuesVectors; ++PointIndex)
-		{
-			// searching the kdtree for PointIndex
-			FKDTree::FRadiusResultSet ResultSet(UE_SMALL_NUMBER, NumPCAValuesVectors, ResultIndexes, ResultDistanceSqr);
-			SearchIndex.KDTree.FindNeighbors(ResultSet, &SearchIndex.PCAValues[PointIndex * NumberOfPrincipalComponents]);
+			size_t MaxNumNeighborToFindAPoint = 0;
+			for (size_t PointIndex = 0; PointIndex < NumPCAValuesVectors; ++PointIndex)
+			{
+				// searching the kdtree for PointIndex
+				FKDTree::FRadiusResultSet ResultSet(UE_SMALL_NUMBER, NumPCAValuesVectors, ResultIndexes, ResultDistanceSqr);
+				SearchIndex.KDTree.FindNeighbors(ResultSet, MakeArrayView(&SearchIndex.PCAValues[PointIndex * NumberOfPrincipalComponents], NumberOfPrincipalComponents));
 
-			bool bFound = false;
-			for (size_t ResultIndex = 0; ResultIndex < ResultSet.Num(); ++ResultIndex)
+				bool bFound = false;
+				for (size_t ResultIndex = 0; ResultIndex < ResultSet.Num(); ++ResultIndex)
+				{
+					if (PointIndex == ResultIndexes[ResultIndex])
+					{
+						// PointIndex is the ResultIndex-th candidates out of the kdtree. if ResultIndex-th is greater than KDTreeQueryNumNeighbors,
+						// we wouldn't have found it in a runtime search, so we log the errot (later on only once, with the worst case scenario)
+						check(ResultDistanceSqr[ResultIndex] < UE_KINDA_SMALL_NUMBER);
+						MaxNumNeighborToFindAPoint = FMath::Max(MaxNumNeighborToFindAPoint, ResultIndex);
+						bFound = true;
+						break;
+					}
+				}
+				
+				if (!bFound)
+				{
+					UE_LOG(LogPoseSearch, Error, TEXT("PreprocessSearchIndexKDTree - kdtree for %s is not properly constructed! Couldn't find the Point %d in it"), *Database->GetName(), PointIndex);
+				}
+			}
+
+			if (MaxNumNeighborToFindAPoint >= KDTreeQueryNumNeighbors)
+			{
+				UE_LOG(LogPoseSearch, Warning, TEXT("Not enough 'KDTreeQueryNumNeighbors' (%d) for database '%s'. Pose values projected in PCA space have too many duplicates, so try to prune duplicates by tuning 'PCAValuesPruningSimilarityThreshold' or increase 'KDTreeQueryNumNeighbors' at least to %d"), KDTreeQueryNumNeighbors, *Database->GetName(), MaxNumNeighborToFindAPoint);
+			}
+
+			// if bArePCAValuesPruned PointIndex is the index of the point in the kdtree, NOT necessary the pose index, so doing the PCAProject would lead to the wrong data
+			const bool bArePCAValuesPruned = SearchIndex.PCAValuesVectorToPoseIndexes.Num() > 0;
+			if (!bArePCAValuesPruned)
 			{
-				if (PointIndex == ResultIndexes[ResultIndex])
+				// testing the KDTree is returning the proper searches for all the original points transformed in pca space
+				TArrayView<float> ProjectedValues((float*)FMemory_Alloca(NumberOfPrincipalComponents * sizeof(float)), NumberOfPrincipalComponents);
+				for (size_t PointIndex = 0; PointIndex < NumPCAValuesVectors; ++PointIndex)
 				{
-					// PointIndex is the ResultIndex-th candidates out of the kdtree. if ResultIndex-th is greater than KDTreeQueryNumNeighbors,
-					// we wouldn't have found it in a runtime search, so we log the errot (later on only once, with the worst case scenario)
-					check(ResultDistanceSqr[ResultIndex] < UE_KINDA_SMALL_NUMBER);
-					MaxNumNeighborToFindAPoint = FMath::Max(MaxNumNeighborToFindAPoint, ResultIndex);
-					bFound = true;
-					break;
+					FKDTree::FKNNResultSet ResultSet(KDTreeQueryNumNeighbors, ResultIndexes, ResultDistanceSqr);
+					SearchIndex.KDTree.FindNeighbors(ResultSet, SearchIndex.PCAProject(SearchIndex.GetPoseValuesBase(PointIndex, NumDimensions), ProjectedValues));
+
+					size_t ResultIndex = 0;
+					for (; ResultIndex < ResultSet.Num(); ++ResultIndex)
+					{
+						if (PointIndex == ResultIndexes[ResultIndex])
+						{
+							if (ResultDistanceSqr[ResultIndex] > UE_KINDA_SMALL_NUMBER)
+							{
+								UE_LOG(LogPoseSearch, Error, TEXT("PreprocessSearchIndexKDTree - kdtree for %s is not properly constructed! Couldn't find the Point %d in it within UE_KINDA_SMALL_NUMBER tolerance, after PCA projection"), *Database->GetName(), PointIndex);
+							}
+							break;
+						}
+					}
+					if (ResultIndex == ResultSet.Num())
+					{
+						UE_LOG(LogPoseSearch, Error, TEXT("PreprocessSearchIndexKDTree - kdtree for %s is not properly constructed! Couldn't find the Point %d in it, after PCA projection"), *Database->GetName(), PointIndex);
+					}
 				}
 			}
-			check(bFound);
 		}
+#endif // ENABLE_ANIM_DEBUG
+	}
+}
+
+// creating a vantage point tree
+static void PreprocessSearchIndexVPTree(FSearchIndex& SearchIndex, const UPoseSearchDatabase* Database, int32 RandomSeed)
+{
+	const int32 NumDimensions = Database->Schema->SchemaCardinality;
+	const int32 NumValuesVectors = SearchIndex.GetNumValuesVectors(NumDimensions);
+	const EPoseSearchMode PoseSearchMode = Database->PoseSearchMode;
+
+	SearchIndex.VPTree.Reset();
+	if (PoseSearchMode == EPoseSearchMode::VPTree && NumValuesVectors > 0)
+	{
+		const FVPTreeDataSource DataSource(SearchIndex);
+		FRandomStream RandStream(RandomSeed);
+		SearchIndex.VPTree.Construct(DataSource, RandStream);
 
-		if (MaxNumNeighborToFindAPoint >= KDTreeQueryNumNeighbors)
+#if ENABLE_ANIM_DEBUG
+
+		if (AnyTestFlags(EMotionMatchTestFlags::TestVPTreeConstructDeterminism))
 		{
-			UE_LOG(LogPoseSearch, Warning, TEXT("Not enough 'KDTreeQueryNumNeighbors' (%d) for database '%s'. Pose values projected in PCA space have too many duplicates, so try to prune duplicates by tuning 'PCAValuesPruningSimilarityThreshold' or increase 'KDTreeQueryNumNeighbors' at least to %d"), KDTreeQueryNumNeighbors, *Database->GetName(), MaxNumNeighborToFindAPoint);
+			if (!SearchIndex.VPTree.TestConstruct(DataSource))
+			{
+				UE_LOG(LogPoseSearch, Error, TEXT("PreprocessSearchIndexVPTree - FVPTree construction failed"));
+			}
+
+			const int32 NumIterations = CVarMotionMatchTestNumIterations.GetValueOnAnyThread();
+			for (int32 Iteration = 0; Iteration < NumIterations; ++Iteration)
+			{
+				FVPTree VPTreeTest;
+				FRandomStream RandStreamTest(RandomSeed);
+				VPTreeTest.Construct(DataSource, RandStreamTest);
+
+				if (VPTreeTest != SearchIndex.VPTree)
+				{
+					UE_LOG(LogPoseSearch, Warning, TEXT("PreprocessSearchIndexVPTree - FVPTree construction is not deterministic"));
+				}
+			}
 		}
 
-		// if bArePCAValuesPruned PointIndex is the index of the point in the kdtree, NOT necessary the pose index, so doing the PCAProject would lead to the wrong data
-		const bool bArePCAValuesPruned = SearchIndex.PCAValuesVectorToPoseIndexes.Num() > 0;
-		if (!bArePCAValuesPruned && PoseSearchMode == EPoseSearchMode::PCAKDTree_Validate)
+		if (AnyTestFlags(EMotionMatchTestFlags::ValidateVPTreeConstruct))
 		{
-			// testing the KDTree is returning the proper searches for all the original points transformed in pca space
-			int32 NumberOfFailingPoints = 0;
-			TArrayView<float> ProjectedValues((float*)FMemory_Alloca(NumberOfPrincipalComponents * sizeof(float)), NumberOfPrincipalComponents);
-			for (size_t PointIndex = 0; PointIndex < NumPCAValuesVectors; ++PointIndex)
+			// we can validate vantage point tree only if there are no duplicates poses (points)
+			if (Database->PosePruningSimilarityThreshold > 0.f)
 			{
-				FKDTree::FKNNResultSet ResultSet(KDTreeQueryNumNeighbors, ResultIndexes, ResultDistanceSqr);
-				SearchIndex.KDTree.FindNeighbors(ResultSet, SearchIndex.PCAProject(SearchIndex.GetPoseValuesBase(PointIndex, NumDimensions), ProjectedValues).GetData());
-
-				size_t ResultIndex = 0;
-				for (; ResultIndex < ResultSet.Num(); ++ResultIndex)
+				// testing the VPTree is returning the proper searches for all the points
+				for (int32 PointIndex = 0; PointIndex < NumValuesVectors; ++PointIndex)
 				{
-					if (PointIndex == ResultIndexes[ResultIndex])
+					FVPTreeResultSet ResultSet(Database->KDTreeQueryNumNeighbors);
+					SearchIndex.VPTree.FindNeighbors(SearchIndex.GetValuesVector(PointIndex, NumDimensions), ResultSet, DataSource);
+		
+					bool bFound = false;
+					for (const FIndexDistance& Result : ResultSet.GetUnsortedResults())
 					{
-						check(ResultDistanceSqr[ResultIndex] < UE_KINDA_SMALL_NUMBER);
-						break;
+						if (Result.Index == PointIndex)
+						{
+							if (!FMath::IsNearlyZero(Result.Distance))
+							{
+								UE_LOG(LogPoseSearch, Error, TEXT("PreprocessSearchIndexVPTree - VPTree for %s is malformed because foud PointIndex %d distance %f from itself (distance should be zero)!"), *Database->GetName(), PointIndex, Result.Distance);
+							}
+
+							bFound = true;
+							break;
+						}
+					}
+
+					if (!bFound)
+					{
+						UE_LOG(LogPoseSearch, Error, TEXT("PreprocessSearchIndexVPTree - VPTree for %s is malformed and couldn't find PointIndex %d"), *Database->GetName(), PointIndex);
 					}
-				}
-				if (ResultIndex == ResultSet.Num())
-				{
-					++NumberOfFailingPoints;
 				}
 			}
-
-			check(NumberOfFailingPoints == 0);
+			else
+			{
+				UE_LOG(LogPoseSearch, Warning, TEXT("PreprocessSearchIndexVPTree - cannot ValidateVPTreeConstruct for %s if there could be potential duplicate poses: set PosePruningSimilarityThreshold > 0 to enforce duplicate pruning"), *Database->GetName());
+			}
 		}
-
 #endif // ENABLE_ANIM_DEBUG
 	}
 }
@@ -908,6 +1002,11 @@ static void CompareSearchIndexBase(const FSearchIndexBase& A, const FSearchIndex
 		}
 	}
 
+	if (A.ValuesVectorToPoseIndexes != B.ValuesVectorToPoseIndexes)
+	{
+		StringBuilder.Append(TEXT("ValuesVectorToPoseIndexes mismatch\n"));
+	}
+
 	if (A.PoseMetadata != B.PoseMetadata)
 	{
 		StringBuilder.Append(TEXT("PoseMetadata mismatch\n"));
@@ -972,6 +1071,11 @@ static void CompareSearchIndex(const FSearchIndex& A, const FSearchIndex& B, con
 	{
 		StringBuilder.Append(TEXT("KDTree mismatch\n"));
 	}
+
+	if (A.VPTree != B.VPTree)
+	{
+		StringBuilder.Append(TEXT("VPTree mismatch\n"));
+	}
 }
 
 #endif // ENABLE_ANIM_DEBUG
@@ -1415,21 +1519,28 @@ void FPoseSearchDatabaseAsyncCacheTask::OnGetComplete(UE::DerivedData::FCacheGet
 					const int32 NumIterations = CVarMotionMatchTestNumIterations.GetValueOnAnyThread();
 
 					FSearchIndex TestSearchIndexA = SearchIndex;
-					TestSearchIndexA.PruneDuplicateValues(IndexBaseDatabases[0]->PosePruningSimilarityThreshold, IndexBaseDatabases[0]->Schema->SchemaCardinality);
+					TestSearchIndexA.PruneDuplicateValues(IndexBaseDatabases[0]->PosePruningSimilarityThreshold, IndexBaseDatabases[0]->Schema->SchemaCardinality, false);
 					for (int32 Iteration = 0; Iteration < NumIterations; ++Iteration)
 					{
 						FSearchIndex TestSearchIndexB = SearchIndex;
-						TestSearchIndexB.PruneDuplicateValues(IndexBaseDatabases[0]->PosePruningSimilarityThreshold, IndexBaseDatabases[0]->Schema->SchemaCardinality);
+						TestSearchIndexB.PruneDuplicateValues(IndexBaseDatabases[0]->PosePruningSimilarityThreshold, IndexBaseDatabases[0]->Schema->SchemaCardinality, false);
 
 						if (TestSearchIndexA.Values != TestSearchIndexB.Values)
 						{
 							UE_LOG(LogPoseSearch, Warning, TEXT("OnGetComplete - PruneDuplicateValues is not deterministic"));
 						}
+
+						if (TestSearchIndexA.ValuesVectorToPoseIndexes != TestSearchIndexB.ValuesVectorToPoseIndexes)
+						{
+							UE_LOG(LogPoseSearch, Warning, TEXT("OnGetComplete - PruneDuplicateValues ValuesVectorToPoseIndexes generation is not deterministic"));
+						}
 					}
 				}
 #endif // ENABLE_ANIM_DEBUG
 
-				SearchIndex.PruneDuplicateValues(IndexBaseDatabases[0]->PosePruningSimilarityThreshold, IndexBaseDatabases[0]->Schema->SchemaCardinality);
+				// VPTree requires ValuesVectorToPoseIndexes if there's any Values pruning
+				const bool bDoNotGenerateValuesVectorToPoseIndexes = IndexBaseDatabases[0]->PoseSearchMode != EPoseSearchMode::VPTree;
+				SearchIndex.PruneDuplicateValues(IndexBaseDatabases[0]->PosePruningSimilarityThreshold, IndexBaseDatabases[0]->Schema->SchemaCardinality, bDoNotGenerateValuesVectorToPoseIndexes);
 
 				TArray<float> Deviation = FMeanDeviationCalculator::Calculate(SearchIndexBases, Schemas);
 
@@ -1456,6 +1567,32 @@ void FPoseSearchDatabaseAsyncCacheTask::OnGetComplete(UE::DerivedData::FCacheGet
 					return;
 				}
 
+#if ENABLE_ANIM_DEBUG
+				// testing PruneDuplicatePCAValues determinism
+				if (AnyTestFlags(EMotionMatchTestFlags::TestPruneDuplicatePCAValuesDeterminism))
+				{
+					const int32 NumIterations = CVarMotionMatchTestNumIterations.GetValueOnAnyThread();
+
+					FSearchIndex TestSearchIndexA = SearchIndex;
+					TestSearchIndexA.PruneDuplicatePCAValues(IndexBaseDatabases[0]->PCAValuesPruningSimilarityThreshold, IndexBaseDatabases[0]->GetNumberOfPrincipalComponents());
+					for (int32 Iteration = 0; Iteration < NumIterations; ++Iteration)
+					{
+						FSearchIndex TestSearchIndexB = SearchIndex;
+						TestSearchIndexB.PruneDuplicatePCAValues(IndexBaseDatabases[0]->PCAValuesPruningSimilarityThreshold, IndexBaseDatabases[0]->GetNumberOfPrincipalComponents());
+
+						if (TestSearchIndexA.PCAValues != TestSearchIndexB.PCAValues)
+						{
+							UE_LOG(LogPoseSearch, Warning, TEXT("OnGetComplete - PruneDuplicatePCAValues is not deterministic"));
+						}
+
+						if (TestSearchIndexA.PCAValuesVectorToPoseIndexes != TestSearchIndexB.PCAValuesVectorToPoseIndexes)
+						{
+							UE_LOG(LogPoseSearch, Warning, TEXT("OnGetComplete - PruneDuplicatePCAValues PCAValuesVectorToPoseIndexes generation is not deterministic"));
+						}
+					}
+				}
+#endif // ENABLE_ANIM_DEBUG
+
 				SearchIndex.PruneDuplicatePCAValues(IndexBaseDatabases[0]->PCAValuesPruningSimilarityThreshold, IndexBaseDatabases[0]->GetNumberOfPrincipalComponents());
 				if (Owner.IsCanceled())
 				{
@@ -1484,6 +1621,18 @@ void FPoseSearchDatabaseAsyncCacheTask::OnGetComplete(UE::DerivedData::FCacheGet
 					SearchIndex.ResetValues();
 				}
 
+				const int32 RandomSeed = GetTypeHash(FullIndexKey.Hash);
+				PreprocessSearchIndexVPTree(SearchIndex, IndexBaseDatabases[0].Get(), RandomSeed);
+				if (Owner.IsCanceled())
+				{
+					UE_LOG(LogPoseSearch, Log, TEXT("%s - %s BuildIndex Cancelled"), *LexToString(FullIndexKey.Hash), *IndexBaseDatabases[0]->GetName());
+					SearchIndex.Reset();
+#if ENABLE_ANIM_DEBUG
+					SearchIndexCompare.Reset();
+#endif //ENABLE_ANIM_DEBUG
+					return;
+				}
+
 				UE_LOG(LogPoseSearch, Log, TEXT("%s - %s BuildIndex Succeeded"), *LexToString(FullIndexKey.Hash), *IndexBaseDatabases[0]->GetName());
 
 #if ENABLE_ANIM_DEBUG
diff --git a/Engine/Plugins/Experimental/Animation/PoseSearch/Source/Runtime/Private/PoseSearchDerivedDataKey.cpp b/Engine/Plugins/Experimental/Animation/PoseSearch/Source/Runtime/Private/PoseSearchDerivedDataKey.cpp
--- a/Engine/Plugins/Experimental/Animation/PoseSearch/Source/Runtime/Private/PoseSearchDerivedDataKey.cpp
+++ b/Engine/Plugins/Experimental/Animation/PoseSearch/Source/Runtime/Private/PoseSearchDerivedDataKey.cpp
@@ -40,7 +40,7 @@ FKeyBuilder::FKeyBuilder(const UObject* Object, bool bUseDataVer, bool bPerformC
 	if (bUseDataVer)
 	{
 		// used to invalidate the key without having to change POSESEARCHDB_DERIVEDDATA_VER all the times
-		int32 POSESEARCHDB_DERIVEDDATA_VER_SMALL = 191;
+		int32 POSESEARCHDB_DERIVEDDATA_VER_SMALL = 199;
 		FGuid VersionGuid = FDevSystemGuids::GetSystemGuid(FDevSystemGuids::Get().POSESEARCHDB_DERIVEDDATA_VER);
 
 		*this << VersionGuid;
diff --git a/Engine/Plugins/Experimental/Animation/PoseSearch/Source/Runtime/Private/PoseSearchFeatureChannel_Heading.h b/Engine/Plugins/Experimental/Animation/PoseSearch/Source/Runtime/Private/PoseSearchFeatureChannel_Heading.h
--- a/Engine/Plugins/Experimental/Animation/PoseSearch/Source/Runtime/Private/PoseSearchFeatureChannel_Heading.h
+++ b/Engine/Plugins/Experimental/Animation/PoseSearch/Source/Runtime/Private/PoseSearchFeatureChannel_Heading.h
@@ -33,7 +33,7 @@ public:
 	// UAnimNotifyState_PoseSearchSamplingAttribute notify state with a matching SamplingAttributeId, and the UAnimNotifyState_PoseSearchSamplingAttribute properties
 	// will be used as source of data instead of this channel "Bone".
 	UPROPERTY(EditAnywhere, Category = "Settings")
-	int32 SamplingAttributeId = -1;
+	int32 SamplingAttributeId = INDEX_NONE;
 
 	// the data relative to the sampling time associated to this channel will be offsetted by SampleTimeOffset seconds.
 	// For example, if Bone is the head bone, and SampleTimeOffset is 0.5, this channel will try to match the future heading of the character head bone 0.5 seconds ahead
diff --git a/Engine/Plugins/Experimental/Animation/PoseSearch/Source/Runtime/Private/PoseSearchFeatureChannel_Position.h b/Engine/Plugins/Experimental/Animation/PoseSearch/Source/Runtime/Private/PoseSearchFeatureChannel_Position.h
--- a/Engine/Plugins/Experimental/Animation/PoseSearch/Source/Runtime/Private/PoseSearchFeatureChannel_Position.h
+++ b/Engine/Plugins/Experimental/Animation/PoseSearch/Source/Runtime/Private/PoseSearchFeatureChannel_Position.h
@@ -26,7 +26,7 @@ public:
 	// UAnimNotifyState_PoseSearchSamplingAttribute notify state with a matching SamplingAttributeId, and the UAnimNotifyState_PoseSearchSamplingAttribute properties
 	// will be used as source of data instead of this channel "Bone". UAnimNotifyState_PoseSearchSamplingAttribute properties will be then converted into OriginBone space
 	UPROPERTY(EditAnywhere, Category = "Settings")
-	int32 SamplingAttributeId = -1;
+	int32 SamplingAttributeId = INDEX_NONE;
 	
 	// the data relative to the sampling time associated to this channel will be offsetted by SampleTimeOffset seconds.
 	// For example, if Bone is the head bone, and SampleTimeOffset is 0.5, this channel will try to match the future position of the character head bone 0.5 seconds ahead
diff --git a/Engine/Plugins/Experimental/Animation/PoseSearch/Source/Runtime/Private/PoseSearchFeatureChannel_Velocity.h b/Engine/Plugins/Experimental/Animation/PoseSearch/Source/Runtime/Private/PoseSearchFeatureChannel_Velocity.h
--- a/Engine/Plugins/Experimental/Animation/PoseSearch/Source/Runtime/Private/PoseSearchFeatureChannel_Velocity.h
+++ b/Engine/Plugins/Experimental/Animation/PoseSearch/Source/Runtime/Private/PoseSearchFeatureChannel_Velocity.h
@@ -22,7 +22,7 @@ public:
 	// UAnimNotifyState_PoseSearchSamplingAttribute notify state with a matching SamplingAttributeId, and the UAnimNotifyState_PoseSearchSamplingAttribute properties
 	// will be used as source of data instead of this channel "Bone".
 	UPROPERTY(EditAnywhere, Category = "Settings")
-	int32 SamplingAttributeId = -1;
+	int32 SamplingAttributeId = INDEX_NONE;
 
 	// the data relative to the sampling time associated to this channel will be offsetted by SampleTimeOffset seconds.
 	// For example, if Bone is the head bone, and SampleTimeOffset is 0.5, this channel will try to match the future velocity of the character head bone 0.5 seconds ahead
diff --git a/Engine/Plugins/Experimental/Animation/PoseSearch/Source/Runtime/Private/PoseSearchIndex.cpp b/Engine/Plugins/Experimental/Animation/PoseSearch/Source/Runtime/Private/PoseSearchIndex.cpp
--- a/Engine/Plugins/Experimental/Animation/PoseSearch/Source/Runtime/Private/PoseSearchIndex.cpp
+++ b/Engine/Plugins/Experimental/Animation/PoseSearch/Source/Runtime/Private/PoseSearchIndex.cpp
@@ -96,16 +96,6 @@ struct FPosePairSimilarity : public FPosePair
 	float Similarity = 0.f;
 };
 
-static bool CalculateSimilaritiesKDTree(TArray<FPosePairSimilarity>& PosePairSimilarities, float SimilarityThreshold, 
-	int32 DataCardinality, int32 NumPoses, const TAlignedArray<float>& Values, TFunctionRef<TConstArrayView<float>(int32, int32)> GetValuesVector)
-{
-	PosePairSimilarities.Reserve(1024 * 64);
-
-	
-
-	return !PosePairSimilarities.IsEmpty();
-}
-
 static bool CalculateSimilarities(TArray<FPosePairSimilarity>& PosePairSimilarities, float SimilarityThreshold, 
 	int32 DataCardinality, int32 NumPoses, const TAlignedArray<float>& Values,
 	TFunctionRef<TConstArrayView<float>(int32, int32)> GetValuesVector)
@@ -137,7 +127,7 @@ static bool CalculateSimilarities(TArray<FPosePairSimilarity>& PosePairSimilarit
 
 			// searching for duplicates within a radius of SimilarityThreshold
 			FKDTree::FRadiusResultSet ResultSet(SimilarityThreshold, NumPoses, ResultIndexes, ResultDistanceSqr);
-			KDTree.FindNeighbors(ResultSet, ValuesA.GetData());
+			KDTree.FindNeighbors(ResultSet, ValuesA);
 
 			for (size_t ResultIndex = 0; ResultIndex < ResultSet.Num(); ++ResultIndex)
 			{
@@ -398,8 +388,10 @@ void FSearchIndexBase::Reset()
 	*this = FSearchIndexBase();
 }
 
-void FSearchIndexBase::PruneDuplicateValues(float SimilarityThreshold, int32 DataCardinality)
+void FSearchIndexBase::PruneDuplicateValues(float SimilarityThreshold, int32 DataCardinality, bool bDoNotGenerateValuesVectorToPoseIndexes)
 {
+	ValuesVectorToPoseIndexes = FSparsePoseMultiMap<uint32>();
+
 	const int32 NumPoses = GetNumPoses();
 	if (SimilarityThreshold > 0.f && NumPoses >= 2)
 	{
@@ -411,6 +403,48 @@ void FSearchIndexBase::PruneDuplicateValues(float SimilarityThreshold, int32 Dat
 			[this](int32 PoseIdx) {	return PoseMetadata[PoseIdx].GetValueOffset(); },
 			[this](int32 PoseIdx, uint32 ValueOffset) {	PoseMetadata[PoseIdx].SetValueOffset(ValueOffset); });
 		}
+
+		if (!bDoNotGenerateValuesVectorToPoseIndexes)
+		{
+			TMap<uint32, TArray<uint32>> ValuesVectorToPoseIndexesMap;
+			ValuesVectorToPoseIndexesMap.Reserve(NumPoses);
+			for (int32 PoseIdx = 0; PoseIdx < NumPoses; ++PoseIdx)
+			{
+				const FPoseMetadata& Metadata = PoseMetadata[PoseIdx];
+				check(Metadata.GetValueOffset() % DataCardinality == 0);
+				const int32 ValuesVectorIdx = Metadata.GetValueOffset() / DataCardinality;
+
+				TArray<uint32>& PoseIndexes = ValuesVectorToPoseIndexesMap.FindOrAdd(ValuesVectorIdx);
+				PoseIndexes.Add(PoseIdx);
+			}
+
+			// sorting ValuesVectorToPoseIndexesMap keys to create a deterministic FSparsePoseMultiMap later on
+			// we're not using TSortedMap for performance reasons, because ValuesVectorToPoseIndexesMap can be quite big
+			TArray<uint32> SortedKeys;
+			SortedKeys.Reserve(ValuesVectorToPoseIndexesMap.Num());
+			for (const TPair<uint32, TArray<uint32>>& Pair : ValuesVectorToPoseIndexesMap)
+			{
+				SortedKeys.Add(Pair.Key);
+			}
+			SortedKeys.Sort();
+
+			FSparsePoseMultiMap<uint32> SparsePoseMultiMap(ValuesVectorToPoseIndexesMap.Num(), NumPoses - 1);
+			for (const uint32& Key : SortedKeys)
+			{
+				const uint32 PCAValuesVectorIdx = Key;
+				const TArray<uint32>& PoseIndexes = ValuesVectorToPoseIndexesMap[Key];
+				SparsePoseMultiMap.Insert(PCAValuesVectorIdx, PoseIndexes);
+			}
+
+			for (uint32 ValuesVectorIdx = 0; ValuesVectorIdx < SparsePoseMultiMap.Num(); ++ValuesVectorIdx)
+			{
+				const TConstArrayView<uint32> PoseIndexes = SparsePoseMultiMap[ValuesVectorIdx];
+				const TArray<uint32>& TestPoseIndexes = ValuesVectorToPoseIndexesMap[ValuesVectorIdx];
+				check(PoseIndexes == TestPoseIndexes);
+			}
+
+			ValuesVectorToPoseIndexes = SparsePoseMultiMap;
+		}
 	}
 }
 
@@ -427,6 +461,7 @@ bool FSearchIndexBase::operator==(const FSearchIndexBase& Other) const
 {
 	return
 		Values == Other.Values &&
+		ValuesVectorToPoseIndexes == Other.ValuesVectorToPoseIndexes &&
 		PoseMetadata == Other.PoseMetadata &&
 		bAnyBlockTransition == Other.bAnyBlockTransition &&
 		Assets == Other.Assets &&
@@ -437,6 +472,7 @@ bool FSearchIndexBase::operator==(const FSearchIndexBase& Other) const
 FArchive& operator<<(FArchive& Ar, FSearchIndexBase& Index)
 {
 	Ar << Index.Values;
+	Ar << Index.ValuesVectorToPoseIndexes;
 	Ar << Index.PoseMetadata;
 	Ar << Index.bAnyBlockTransition;
 	Ar << Index.Assets;
@@ -455,6 +491,7 @@ FSearchIndex::FSearchIndex(const FSearchIndex& Other)
 	, PCAProjectionMatrix(Other.PCAProjectionMatrix)
 	, Mean(Other.Mean)
 	, KDTree(Other.KDTree)
+	, VPTree(Other.VPTree)
 	, PCAExplainedVariance(Other.PCAExplainedVariance)
 {
 	check(!PCAValues.IsEmpty() || KDTree.DataSource.PointCount == 0);
@@ -572,6 +609,7 @@ void FSearchIndex::PruneDuplicatePCAValues(float SimilarityThreshold, int32 Numb
 			{
 				// we pruned some PCAValues: we need to construct a mapping between PCAValuesVectorIdx to PoseIdx(s)
 				TMap<uint32, TArray<uint32>> PCAValuesVectorToPoseIndexesMap;
+				PCAValuesVectorToPoseIndexesMap.Reserve(NumPoses);
 				for (uint32 PoseIdx = 0; PoseIdx < NumPoses; ++PoseIdx)
 				{
 					check(PoseToPCAValueOffset[PoseIdx] % NumberOfPrincipalComponents == 0);
@@ -703,7 +741,8 @@ bool FSearchIndex::operator==(const FSearchIndex& Other) const
 		PCAProjectionMatrix == Other.PCAProjectionMatrix &&
 		Mean == Other.Mean &&
 		PCAExplainedVariance == Other.PCAExplainedVariance &&
-		KDTree == Other.KDTree;
+		KDTree == Other.KDTree &&
+		VPTree == Other.VPTree;
 }
 
 FArchive& operator<<(FArchive& Ar, FSearchIndex& Index)
@@ -716,6 +755,7 @@ FArchive& operator<<(FArchive& Ar, FSearchIndex& Index)
 	Ar << Index.PCAProjectionMatrix;
 	Ar << Index.Mean;
 	Ar << Index.PCAExplainedVariance;
+	Ar << Index.VPTree;
 
 	Serialize(Ar, Index.KDTree, Index.PCAValues.GetData());
 	return Ar;
diff --git a/Engine/Plugins/Experimental/Animation/PoseSearch/Source/Runtime/Private/PoseSearchLibrary.cpp b/Engine/Plugins/Experimental/Animation/PoseSearch/Source/Runtime/Private/PoseSearchLibrary.cpp
--- a/Engine/Plugins/Experimental/Animation/PoseSearch/Source/Runtime/Private/PoseSearchLibrary.cpp
+++ b/Engine/Plugins/Experimental/Animation/PoseSearch/Source/Runtime/Private/PoseSearchLibrary.cpp
@@ -163,7 +163,7 @@ void UPoseSearchLibrary::TraceMotionMatchingState(
 				break;
 			}
 		}
-		if (DbEntryIdx == -1)
+		if (DbEntryIdx == INDEX_NONE)
 		{
 			DbEntryIdx = DatabaseEntries.Add({ DatabaseId });
 
@@ -174,7 +174,7 @@ void UPoseSearchLibrary::TraceMotionMatchingState(
 		return DbEntryIdx;
 	};
 
-	const int32 CurrentPoseIdx = bSearch && CurrentResult.PoseCost.IsValid() ? CurrentResult.PoseIdx : -1;
+	const int32 CurrentPoseIdx = bSearch && CurrentResult.PoseCost.IsValid() ? CurrentResult.PoseIdx : INDEX_NONE;
 	FTraceMotionMatchingState TraceState;
 
 	for (const UPoseSearchDatabase* SearchedDatabase : SearchContext.BestCandidates.GetSearchedDatabases())
@@ -426,7 +426,7 @@ void UPoseSearchLibrary::MotionMatch(
 	const UAnimationAsset* FutureAnimation,
 	float FutureAnimationStartTime,
 	float TimeToFutureAnimationStart,
-	const int DebugSessionUniqueIdentifier)
+	const int32 DebugSessionUniqueIdentifier)
 {
 #if ENABLE_DRAW_DEBUG && ENABLE_ANIM_DEBUG
 	class UAnimInstanceProxyProvider : public UAnimInstance
@@ -505,7 +505,7 @@ void UPoseSearchLibrary::MotionMatch(
 			}
 
 			// extracting 2 poses to be able to calculate velocities
-			for (int i = 0; i < 2; ++i)
+			for (int32 i = 0; i < 2; ++i)
 			{
 				const float ExtractionTime = FutureAnimationStartTime + (i - 1) * FiniteDelta;
 				const float FutureAnimationTime = TimeToFutureAnimationStart + (i - 1) * FiniteDelta;
diff --git a/Engine/Plugins/Experimental/Animation/PoseSearch/Source/Runtime/Private/VPTree.cpp b/Engine/Plugins/Experimental/Animation/PoseSearch/Source/Runtime/Private/VPTree.cpp
new file mode 100644
--- /dev/null
+++ b/Engine/Plugins/Experimental/Animation/PoseSearch/Source/Runtime/Private/VPTree.cpp
@@ -0,0 +1,29 @@
+// Copyright Epic Games, Inc. All Rights Reserved.
+
+#include "PoseSearch/VPTree.h"
+
+namespace UE::PoseSearch
+{
+
+void FVPTree::Reset()
+{
+	Nodes.Reset();
+}
+
+SIZE_T FVPTree::GetAllocatedSize() const
+{
+	return Nodes.GetAllocatedSize();
+}
+
+bool FVPTree::operator==(const FVPTree& Other) const
+{
+	return Nodes == Other.Nodes;
+}
+
+FArchive& operator<<(FArchive& Ar, FVPTree& VPTree)
+{
+	Ar << VPTree.Nodes;
+	return Ar;
+}
+
+} // namespace UE::PoseSearch
\ No newline at end of file
diff --git a/Engine/Plugins/Experimental/Animation/PoseSearch/Source/Runtime/Public/PoseSearch/KDTree.h b/Engine/Plugins/Experimental/Animation/PoseSearch/Source/Runtime/Public/PoseSearch/KDTree.h
--- a/Engine/Plugins/Experimental/Animation/PoseSearch/Source/Runtime/Public/PoseSearch/KDTree.h
+++ b/Engine/Plugins/Experimental/Animation/PoseSearch/Source/Runtime/Public/PoseSearch/KDTree.h
@@ -12,7 +12,7 @@ namespace UE::PoseSearch
 {
 struct FKDTreeImplementation;
 
-struct POSESEARCH_API FKDTree
+struct FKDTree
 {
 	struct FDataSource
 	{
@@ -222,9 +222,9 @@ struct POSESEARCH_API FKDTree
 	
 	void Reset();
 	void Construct(int32 Count, int32 Dim, const float* Data, int32 MaxLeafSize = 16);
-	bool FindNeighbors(FKNNResultSet& Result, const float* Query) const;
-	bool FindNeighbors(FRadiusResultSet& Result, const float* Query) const;
-	SIZE_T GetAllocatedSize() const;
+	bool FindNeighbors(FKNNResultSet& Result, TConstArrayView<float> Query) const;
+	bool FindNeighbors(FRadiusResultSet& Result, TConstArrayView<float> Query) const;
+	POSESEARCH_API SIZE_T GetAllocatedSize() const;
 
 	FDataSource DataSource;
 	FKDTreeImplementation* Impl = nullptr;
diff --git a/Engine/Plugins/Experimental/Animation/PoseSearch/Source/Runtime/Public/PoseSearch/PoseSearchAssetIndexer.h b/Engine/Plugins/Experimental/Animation/PoseSearch/Source/Runtime/Public/PoseSearch/PoseSearchAssetIndexer.h
--- a/Engine/Plugins/Experimental/Animation/PoseSearch/Source/Runtime/Public/PoseSearch/PoseSearchAssetIndexer.h
+++ b/Engine/Plugins/Experimental/Animation/PoseSearch/Source/Runtime/Public/PoseSearch/PoseSearchAssetIndexer.h
@@ -57,18 +57,18 @@ public:
 	// Returns OutSampleRotation as the rotation of the bone Schema.BoneReferences[SchemaSampleBoneIdx] at time CalculateSampleTime(SampleIdx) + SampleTimeOffset relative to the
 	// transform of the bone Schema.BoneReferences[SchemaOriginBoneIdx] at time CalculateSampleTime(SampleIdx) + OriginTimeOffset 
 	// Times will be processed by GetPermutationTimeOffsets(PermutationTimeType, ...)
-	bool GetSampleRotation(FQuat& OutSampleRotation, float SampleTimeOffset, float OriginTimeOffset, int32 SampleIdx, int8 SchemaSampleBoneIdx = RootSchemaBoneIdx, int8 SchemaOriginBoneIdx = RootSchemaBoneIdx, EPermutationTimeType PermutationTimeType = EPermutationTimeType::UseSampleTime, int32 SamplingAttributeId = -1);
+	bool GetSampleRotation(FQuat& OutSampleRotation, float SampleTimeOffset, float OriginTimeOffset, int32 SampleIdx, int8 SchemaSampleBoneIdx = RootSchemaBoneIdx, int8 SchemaOriginBoneIdx = RootSchemaBoneIdx, EPermutationTimeType PermutationTimeType = EPermutationTimeType::UseSampleTime, int32 SamplingAttributeId = INDEX_NONE);
 
 	// Returns OutSamplePosition as the position of the bone Schema.BoneReferences[SchemaSampleBoneIdx] at time CalculateSampleTime(SampleIdx) + SampleTimeOffset relative to the
 	// transform of the bone Schema.BoneReferences[SchemaOriginBoneIdx] at time CalculateSampleTime(SampleIdx) + OriginTimeOffset.
 	// Times will be processed by GetPermutationTimeOffsets(PermutationTimeType, ...)
-	bool GetSamplePosition(FVector& OutSamplePosition, float SampleTimeOffset, float OriginTimeOffset, int32 SampleIdx, int8 SchemaSampleBoneIdx = RootSchemaBoneIdx, int8 SchemaOriginBoneIdx = RootSchemaBoneIdx, EPermutationTimeType PermutationTimeType = EPermutationTimeType::UseSampleTime, int32 SamplingAttributeId = -1);
+	bool GetSamplePosition(FVector& OutSamplePosition, float SampleTimeOffset, float OriginTimeOffset, int32 SampleIdx, int8 SchemaSampleBoneIdx = RootSchemaBoneIdx, int8 SchemaOriginBoneIdx = RootSchemaBoneIdx, EPermutationTimeType PermutationTimeType = EPermutationTimeType::UseSampleTime, int32 SamplingAttributeId = INDEX_NONE);
 
 	// Returns OutSampleVelocity as the delta velocity of the bone Schema.BoneReferences[SchemaSampleBoneIdx] velocity at time CalculateSampleTime(SampleIdx) + SampleTimeOffset minus
 	// the velocity of the bone Schema.BoneReferences[SchemaOriginBoneIdx] at time CalculateSampleTime(SampleIdx) + OriginTimeOffset.
 	// Times will be processed by GetPermutationTimeOffsets(PermutationTimeType, ...)
 	// if bUseCharacterSpaceVelocities is true, velocities will be computed in root bone space, rather than animation (world) space
-	bool GetSampleVelocity(FVector& OutSampleVelocity, float SampleTimeOffset, float OriginTimeOffset, int32 SampleIdx, int8 SchemaSampleBoneIdx = RootSchemaBoneIdx, int8 SchemaOriginBoneIdx = RootSchemaBoneIdx, bool bUseCharacterSpaceVelocities = true, EPermutationTimeType PermutationTimeType = EPermutationTimeType::UseSampleTime, int32 SamplingAttributeId = -1);
+	bool GetSampleVelocity(FVector& OutSampleVelocity, float SampleTimeOffset, float OriginTimeOffset, int32 SampleIdx, int8 SchemaSampleBoneIdx = RootSchemaBoneIdx, int8 SchemaOriginBoneIdx = RootSchemaBoneIdx, bool bUseCharacterSpaceVelocities = true, EPermutationTimeType PermutationTimeType = EPermutationTimeType::UseSampleTime, int32 SamplingAttributeId = INDEX_NONE);
 
 	int32 GetBeginSampleIdx() const;
 	int32 GetEndSampleIdx() const;
diff --git a/Engine/Plugins/Experimental/Animation/PoseSearch/Source/Runtime/Public/PoseSearch/PoseSearchContext.h b/Engine/Plugins/Experimental/Animation/PoseSearch/Source/Runtime/Public/PoseSearch/PoseSearchContext.h
--- a/Engine/Plugins/Experimental/Animation/PoseSearch/Source/Runtime/Public/PoseSearch/PoseSearchContext.h
+++ b/Engine/Plugins/Experimental/Animation/PoseSearch/Source/Runtime/Public/PoseSearch/PoseSearchContext.h
@@ -113,7 +113,7 @@ struct POSESEARCH_API FDebugDrawParams
 	const UPoseSearchSchema* GetSchema() const;
 
 	FVector ExtractPosition(TConstArrayView<float> PoseVector, const UPoseSearchFeatureChannel_Position* Position) const;
-	FVector ExtractPosition(TConstArrayView<float> PoseVector, float SampleTimeOffset, int8 SchemaBoneIdx = RootSchemaBoneIdx, EPermutationTimeType PermutationTimeType = EPermutationTimeType::UseSampleTime, int32 SamplingAttributeId = -1) const;
+	FVector ExtractPosition(TConstArrayView<float> PoseVector, float SampleTimeOffset, int8 SchemaBoneIdx = RootSchemaBoneIdx, EPermutationTimeType PermutationTimeType = EPermutationTimeType::UseSampleTime, int32 SamplingAttributeId = INDEX_NONE) const;
 
 	FQuat ExtractRotation(TConstArrayView<float> PoseVector, float SampleTimeOffset, int8 SchemaBoneIdx = RootSchemaBoneIdx) const;
 
diff --git a/Engine/Plugins/Experimental/Animation/PoseSearch/Source/Runtime/Public/PoseSearch/PoseSearchDatabase.h b/Engine/Plugins/Experimental/Animation/PoseSearch/Source/Runtime/Public/PoseSearch/PoseSearchDatabase.h
--- a/Engine/Plugins/Experimental/Animation/PoseSearch/Source/Runtime/Public/PoseSearch/PoseSearchDatabase.h
+++ b/Engine/Plugins/Experimental/Animation/PoseSearch/Source/Runtime/Public/PoseSearch/PoseSearchDatabase.h
@@ -30,11 +30,8 @@ enum class EPoseSearchMode : int32
 	// Optimized search mode: the database projects the poses into a PCA space using only the most significant "NumberOfPrincipalComponents" dimensions, and construct a kdtree to facilitate the search.
 	PCAKDTree,
 
-	// Debug functionality performing diagnostics and validation during searches using PCAKDTree.
-	PCAKDTree_Validate,
-	
-	// Debug functionality to compare BruteForce vs PCAKDTree.
-	PCAKDTree_Compare,
+	// Optimized search mode using a vantage point tree
+	VPTree,
 
 	Num UMETA(Hidden),
 	Invalid = Num UMETA(Hidden)
@@ -287,14 +284,15 @@ public:
 	EPoseSearchMode PoseSearchMode = EPoseSearchMode::PCAKDTree;
 
 	// Number of dimensions used to create the kdtree. More dimensions allows a better explanation of the variance of the dataset that usually translates in better search results, but will imply more memory usage and worse performances.
-	UPROPERTY(EditAnywhere, Category = "Performance", meta = (EditCondition = "PoseSearchMode != EPoseSearchMode::BruteForce", EditConditionHides, ClampMin = "1", ClampMax = "64", UIMin = "1", UIMax = "64"))
+	UPROPERTY(EditAnywhere, Category = "Performance", meta = (EditCondition = "PoseSearchMode == EPoseSearchMode::PCAKDTree", EditConditionHides, ClampMin = "1", ClampMax = "64", UIMin = "1", UIMax = "64"))
 	int32 NumberOfPrincipalComponents = 4;
 
-	UPROPERTY(EditAnywhere, Category = "Performance", meta = (EditCondition = "PoseSearchMode != EPoseSearchMode::BruteForce", EditConditionHides, ClampMin = "1", ClampMax = "256", UIMin = "1", UIMax = "256"))
+	UPROPERTY(EditAnywhere, Category = "Performance", meta = (EditCondition = "PoseSearchMode == EPoseSearchMode::PCAKDTree", EditConditionHides, ClampMin = "1", ClampMax = "256", UIMin = "1", UIMax = "256"))
 	int32 KDTreeMaxLeafSize = 16;
 	
+	// @todo: rename to KNNQueryNumNeighbors to be usable with the VPTree as well
 	// Out of a kdtree search, results will have only an approximate cost, so the database search will select the best “KDTree Query Num Neighbors” poses to perform the full cost analysis, and be able to elect the best pose.
-	UPROPERTY(EditAnywhere, Category = "Performance", meta = (EditCondition = "PoseSearchMode != EPoseSearchMode::BruteForce", EditConditionHides, ClampMin = "1", ClampMax = "600", UIMin = "1"))
+	UPROPERTY(EditAnywhere, Category = "Performance", meta = (DisplayName = "KNNQueryNumNeighbors", EditCondition = "PoseSearchMode == EPoseSearchMode::PCAKDTree || PoseSearchMode == VPTree", EditConditionHides, ClampMin = "1", ClampMax = "600", UIMin = "1"))
 	int32 KDTreeQueryNumNeighbors = 200;
 
 	// When evaluating multiple searches, including the continuing pose search, the system keeps track of the best pose and associated cost.
@@ -309,12 +307,13 @@ public:
 
 	// if two PCA values (multi dimensional point with the GetNumberOfPrincipalComponents cardinality) are closer than PCAValuesPruningSimilarityThreshold,
 	// only one will be saved into the database FSearchIndex (to save memory).
-	UPROPERTY(EditAnywhere, Category = "Performance")
+	UPROPERTY(EditAnywhere, Category = "Performance", meta = (EditCondition = "PoseSearchMode == EPoseSearchMode::PCAKDTree", EditConditionHides))
 	float PCAValuesPruningSimilarityThreshold = 0.f;
 
+	// @todo: rename to KNNQueryNumNeighborsWithDuplicates to be usable with the VPTree as well
 	// if PCAValuesPruningSimilarityThreshold > 0 the kdtree will remove duplicates, every result out of the KDTreeQueryNumNeighbors could potentially references multiple poses.
 	// KDTreeQueryNumNeighborsWithDuplicates is the upper bound number of poses the system will perform the full cost evaluation. if KDTreeQueryNumNeighborsWithDuplicates is zero then there's no upper bound
-	UPROPERTY(EditAnywhere, Category = "Performance", meta = (EditCondition = "PoseSearchMode != EPoseSearchMode::BruteForce && PCAValuesPruningSimilarityThreshold > 0", EditConditionHides, ClampMin = "0", ClampMax = "600", UIMin = "1"))
+	UPROPERTY(EditAnywhere, Category = "Performance", meta = (DisplayName = "KNNQueryNumNeighborsWithDuplicates", EditCondition = "PoseSearchMode == EPoseSearchMode::PCAKDTree && PCAValuesPruningSimilarityThreshold > 0", EditConditionHides, ClampMin = "0", ClampMax = "600", UIMin = "1"))
 	int32 KDTreeQueryNumNeighborsWithDuplicates = 0;
 	
 #if WITH_EDITORONLY_DATA
@@ -377,5 +376,6 @@ public:
 
 private:
 	UE::PoseSearch::FSearchResult SearchPCAKDTree(UE::PoseSearch::FSearchContext& SearchContext) const;
+	UE::PoseSearch::FSearchResult SearchVPTree(UE::PoseSearch::FSearchContext& SearchContext) const;
 	UE::PoseSearch::FSearchResult SearchBruteForce(UE::PoseSearch::FSearchContext& SearchContext) const;
 };
\ No newline at end of file
diff --git a/Engine/Plugins/Experimental/Animation/PoseSearch/Source/Runtime/Public/PoseSearch/PoseSearchIndex.h b/Engine/Plugins/Experimental/Animation/PoseSearch/Source/Runtime/Public/PoseSearch/PoseSearchIndex.h
--- a/Engine/Plugins/Experimental/Animation/PoseSearch/Source/Runtime/Public/PoseSearch/PoseSearchIndex.h
+++ b/Engine/Plugins/Experimental/Animation/PoseSearch/Source/Runtime/Public/PoseSearch/PoseSearchIndex.h
@@ -5,6 +5,7 @@
 #include "PoseSearch/KDTree.h"
 #include "PoseSearch/PoseSearchCost.h"
 #include "PoseSearch/PoseSearchDefines.h"
+#include "PoseSearch/VPTree.h"
 
 namespace UE::PoseSearch
 {
@@ -235,7 +236,7 @@ struct FSparsePoseMultiMap
 		{
 			// checking for overflow
 			check((DataValues.Num() + 1 + Values.Num()) < (1 << (sizeof(Type) * 4 - 1)));
-			check(int(MaxKey) <= DataValues.Num());
+			check(int32(MaxKey) <= DataValues.Num());
 
 			// adding DeltaKeyValue to DataValues.Num() to making sure DataValues[Key] > MaxValue
 			DataValues[Key] = DataValues.Num() + DeltaKeyValue;
@@ -308,6 +309,7 @@ struct FSparsePoseMultiMap
 struct FSearchIndexBase
 {
 	TAlignedArray<float> Values;
+	FSparsePoseMultiMap<uint32> ValuesVectorToPoseIndexes;
 	TAlignedArray<FPoseMetadata> PoseMetadata;
 	bool bAnyBlockTransition = false;
 	TAlignedArray<FSearchIndexAsset> Assets;
@@ -338,7 +340,7 @@ struct FSearchIndexBase
 
 	void Reset();
 	
-	void PruneDuplicateValues(float SimilarityThreshold, int32 DataCardinality);
+	void PruneDuplicateValues(float SimilarityThreshold, int32 DataCardinality, bool bDoNotGenerateValuesVectorToPoseIndexes);
 
 	TConstArrayView<float> GetPoseValuesBase(int32 PoseIdx, int32 DataCardinality) const
 	{
@@ -348,6 +350,13 @@ struct FSearchIndexBase
 		return MakeArrayView(&Values[ValueOffset], DataCardinality);
 	}
 
+	TConstArrayView<float> GetValuesVector(int32 ValuesVectorIdx, int32 DataCardinality) const
+	{
+		check(!IsValuesEmpty() && ValuesVectorIdx >= 0 && ValuesVectorIdx < GetNumValuesVectors(DataCardinality));
+		const int32 ValueOffset = ValuesVectorIdx * DataCardinality;
+		return MakeArrayView(&Values[ValueOffset], DataCardinality);
+	}
+
 	bool operator==(const FSearchIndexBase& Other) const;
 
 	friend FArchive& operator<<(FArchive& Ar, FSearchIndexBase& Index);
@@ -371,6 +380,7 @@ struct FSearchIndex : public FSearchIndexBase
 	TAlignedArray<float> Mean;
 
 	FKDTree KDTree;
+	FVPTree VPTree;
 
 	// @todo: this property should be editor only
 	float PCAExplainedVariance = 0.f;
@@ -393,6 +403,14 @@ struct FSearchIndex : public FSearchIndexBase
 	// that in the case of GetPoseValuesSafe it's stored in PoseMetadata[PoseIdx].GetValueOffset(), but missing for the PCAValues, so this API input is NOT a PoseIdx
 	// mapping between PoseIdx to PCAValuesVectorIdx can be reconstructed by inverting the PCAValuesVectorToPoseIndexes via GetPoseToPCAValuesVectorIndexes
 	POSESEARCH_API TConstArrayView<float> GetPCAPoseValues(int32 PCAValuesVectorIdx) const;
+
+	int32 GetNumPCAValuesVectors(int32 DataCardinality) const
+	{
+		check(DataCardinality > 0);
+		check(PCAValues.Num() % DataCardinality == 0);
+		return PCAValues.Num() / DataCardinality;
+	}
+
 	POSESEARCH_API FPoseSearchCost ComparePoses(int32 PoseIdx, float ContinuingPoseCostBias, TConstArrayView<float> PoseValues, TConstArrayView<float> QueryValues) const;
 	POSESEARCH_API FPoseSearchCost CompareAlignedPoses(int32 PoseIdx, float ContinuingPoseCostBias, TConstArrayView<float> PoseValues, TConstArrayView<float> QueryValues) const;
 
@@ -405,4 +423,33 @@ struct FSearchIndex : public FSearchIndexBase
 	friend FArchive& operator<<(FArchive& Ar, FSearchIndex& Index);
 };
 
+struct FVPTreeDataSource
+{
+	explicit FVPTreeDataSource(const FSearchIndex& InSearchIndex)
+		: SearchIndex(InSearchIndex)
+	{
+	}
+
+    const TConstArrayView<float> operator[](int32 Index) const
+    {
+		const int32 DataCardinality = SearchIndex.WeightsSqrt.Num();
+        return SearchIndex.GetValuesVector(Index, DataCardinality);
+    }
+
+    int32 Num() const
+    {
+		const int32 DataCardinality = SearchIndex.WeightsSqrt.Num();
+        return SearchIndex.GetNumValuesVectors(DataCardinality);
+    }
+
+    static float GetDistance(const TConstArrayView<float> A, const TConstArrayView<float> B)
+    {
+		// Estracting the Sqrt to satisfy the triangle inequality metric space requirements, since a <= b+c doesn't imply a^2 <= b^2 + c^2
+		return FMath::Sqrt(CompareFeatureVectors(A, B));
+    }
+
+private:
+	const FSearchIndex& SearchIndex;
+};
+
 } // namespace UE::PoseSearch
diff --git a/Engine/Plugins/Experimental/Animation/PoseSearch/Source/Runtime/Public/PoseSearch/PoseSearchLibrary.h b/Engine/Plugins/Experimental/Animation/PoseSearch/Source/Runtime/Public/PoseSearch/PoseSearchLibrary.h
--- a/Engine/Plugins/Experimental/Animation/PoseSearch/Source/Runtime/Public/PoseSearch/PoseSearchLibrary.h
+++ b/Engine/Plugins/Experimental/Animation/PoseSearch/Source/Runtime/Public/PoseSearch/PoseSearchLibrary.h
@@ -164,6 +164,6 @@ public:
 		const UAnimationAsset* FutureAnimation = nullptr,
 		float FutureAnimationStartTime = 0.f,
 		float TimeToFutureAnimationStart = 0.f,
-		const int DebugSessionUniqueIdentifier = 6174);
+		const int32 DebugSessionUniqueIdentifier = 6174);
 };
 
diff --git a/Engine/Plugins/Experimental/Animation/PoseSearch/Source/Runtime/Public/PoseSearch/VPTree.h b/Engine/Plugins/Experimental/Animation/PoseSearch/Source/Runtime/Public/PoseSearch/VPTree.h
new file mode 100644
--- /dev/null
+++ b/Engine/Plugins/Experimental/Animation/PoseSearch/Source/Runtime/Public/PoseSearch/VPTree.h
@@ -0,0 +1,473 @@
+// Copyright Epic Games, Inc. All Rights Reserved.
+
+#pragma once
+
+// Vantage-point tree implementation https://en.wikipedia.org/wiki/Vantage-point_tree
+// it relies on templated methods to specify the
+// - data point type (typename T, usually a multidimensional point like TArray<float>) 
+// - data source struct (typename TDataSource) that needs to implement
+//		struct FDataSource
+//		{
+//			const T& operator[](int32 Index) const;
+//			int32 Num() const;
+//			static float GetDistance(const T& A, const T& B);
+//		};
+//   and can be specialized for runtime calls at FindNeighbors to laverage vectorized code with aliged and padded data
+// - the result set struct (typename TVPTreeResultSet), see struct FVPTreeResultSet below as example
+
+#include "Containers/ArrayView.h"
+#include "Math/RandomStream.h"
+#include "PoseSearch/PoseSearchDefines.h"
+
+namespace UE::PoseSearch
+{
+
+#define VALIDATE_FINDNEIGHBORS 0
+
+typedef float VPTreeScalar;
+
+struct FIndexDistance
+{
+    FIndexDistance(int32 InIndex = INDEX_NONE, VPTreeScalar InDistance = VPTreeScalar(0))
+    : Index(InIndex)
+    , Distance(InDistance)
+    {
+    }
+        
+    bool operator<(const FIndexDistance& Other) const
+    {
+        return Distance < Other.Distance;   
+    }
+
+	bool operator==(const FIndexDistance& Other) const
+	{
+		return Index == Other.Index && Distance == Other.Distance;
+	}
+
+	friend FArchive& operator<<(FArchive& Ar, FIndexDistance& IndexDistance)
+	{
+		Ar << IndexDistance.Index;
+		Ar << IndexDistance.Distance;
+		return Ar;
+	}
+
+    int32 Index;
+    VPTreeScalar Distance;
+};
+
+struct FVPTreeNode : public FIndexDistance
+{
+    int32 LeftIndex = INDEX_NONE;
+    int32 RightIndex = INDEX_NONE;
+
+	bool operator==(const FVPTreeNode& Other) const
+	{
+		return FIndexDistance::operator==(Other) && LeftIndex == Other.LeftIndex && RightIndex == Other.RightIndex;
+	}
+
+	friend FArchive& operator<<(FArchive& Ar, FVPTreeNode& Node)
+	{
+		Ar << static_cast<FIndexDistance&>(Node);
+		Ar << Node.LeftIndex;
+		Ar << Node.RightIndex;
+		return Ar;
+	}
+};
+
+struct FVPTreeResultSet
+{
+    FVPTreeResultSet(int32 InNumNeighbors)
+        : NumNeighbors(InNumNeighbors)
+    {
+    }
+
+    VPTreeScalar GetWorstDistance() const
+    {
+        return Heap.HeapTop().Distance;
+    }
+        
+    bool IsFull() const
+    {
+        return Heap.Num() == NumNeighbors;
+    }
+
+    void AddPoint(float Distance, int32 Index
+#if VALIDATE_FINDNEIGHBORS
+		, int32 BestIndex = INDEX_NONE
+#endif // VALIDATE_FINDNEIGHBORS
+	)
+    {
+		if (!IsFull())
+		{
+			Heap.HeapPush(FIndexDistance(Index, Distance), HeapCompare());
+
+#if VALIDATE_FINDNEIGHBORS
+			Heap.VerifyHeap(HeapCompare());
+#endif // VALIDATE_FINDNEIGHBORS
+		}
+        else if (Distance < GetWorstDistance())
+        {
+			// popping the worst FIndexDistance
+			FIndexDistance Removed;
+            Heap.HeapPop(Removed, HeapCompare());
+			
+#if VALIDATE_FINDNEIGHBORS
+			if (BestIndex != INDEX_NONE && Removed.Index == BestIndex)
+			{
+				UE_LOG(LogPoseSearch, Error, TEXT("FVPTreeResultSet::AddPoint - we removed the best node from the ResultSet..."));
+			}
+#endif // VALIDATE_FINDNEIGHBORS
+
+			// and replacing it with FIndexDistance(Index, Distance) 
+			Heap.HeapPush(FIndexDistance(Index, Distance), HeapCompare());
+
+#if VALIDATE_FINDNEIGHBORS
+			Heap.VerifyHeap(HeapCompare());
+#endif // VALIDATE_FINDNEIGHBORS
+        }
+    }
+
+    TArray<FIndexDistance> GetSortedResults()
+    {
+		TArray<FIndexDistance> Results;
+		Results.SetNumUninitialized(Heap.Num());
+
+		int32 Index = 0;
+        while (!Heap.IsEmpty())
+        {
+            Heap.HeapPop(Results[Index++], HeapCompare());
+        }
+		return Results;
+    }
+
+	const TArray<FIndexDistance>& GetUnsortedResults() const
+	{
+		return Heap;
+	}
+
+	bool operator==(const FVPTreeResultSet& Other) const
+	{
+		return NumNeighbors == Other.NumNeighbors && Heap == Other.Heap;
+	}
+
+private:
+	struct HeapCompare
+	{
+		bool operator()(const FIndexDistance& A, const FIndexDistance& B) const
+		{
+			// using > to create a heap where Heap.HeapTop().Distance is the greater Distance (to behave like the std::priority_queue)
+			return A.Distance > B.Distance;
+		}
+	};
+
+    const int32 NumNeighbors;
+    TArray<FIndexDistance> Heap;
+};
+
+struct FVPTree
+{
+    template<typename TDataSource>
+    void Construct(const TDataSource& DataSource, FRandomStream& RandStream)
+    {
+		Reset();
+
+		// initializing the DataSourceIndexMapping, array of indeces pointing to the original TDataSource
+		const int32 Count = DataSource.Num();
+        TArray<FIndexDistance> DataSourceIndexMapping;
+        DataSourceIndexMapping.SetNumUninitialized(Count);
+        for (int32 Index = 0; Index < DataSourceIndexMapping.Num(); ++Index)
+        {
+            DataSourceIndexMapping[Index].Index = Index;
+        }
+        ConstructRecursive(DataSource, DataSourceIndexMapping, RandStream);
+    }
+
+	template<typename TDataSource>
+	bool TestConstruct(const TDataSource& DataSource, int32 NodeIndex = 0) const
+	{
+		TArray<int32> SubTreeNodeIndexesTempBuffer;
+		SubTreeNodeIndexesTempBuffer.Reserve(Nodes.Num());
+		return TestConstructRecursive(DataSource, NodeIndex, SubTreeNodeIndexesTempBuffer);
+	}
+
+    template<typename T, typename TDataSource, typename TVPTreeResultSet>
+    void FindNeighbors(const T& Query, TVPTreeResultSet& ResultSet, const TDataSource& DataSource) const
+    {
+#if VALIDATE_FINDNEIGHBORS
+		// @todo: WIP validation code
+		int32 NumEvaluatedQueryDistances = 0;
+
+		// brute forcing to find the best Node
+		VPTreeScalar BestDistance = VPTreeScalar(UE_BIG_NUMBER);
+		int32 BestNodeIndex = INDEX_NONE;
+
+		TArray<bool> DoWeHaveAllNodes;
+		DoWeHaveAllNodes.SetNumZeroed(Nodes.Num());
+
+		for (int32 NodeIndex = 0; NodeIndex < Nodes.Num(); ++NodeIndex)
+		{
+			check(Nodes[NodeIndex].Index != INDEX_NONE);
+
+			DoWeHaveAllNodes[Nodes[NodeIndex].Index] = true;
+
+			const VPTreeScalar Distance = TDataSource::GetDistance(Query, DataSource[Nodes[NodeIndex].Index]);
+			if (Distance < BestDistance)
+			{
+				// @todo: handle eventual duplicates
+				BestDistance = Distance;
+				BestNodeIndex = NodeIndex;
+			}
+		}
+
+		check(!DoWeHaveAllNodes.Contains(false));
+
+		TArray<int32> SubTreeNodeIndexesTempBuffer;
+		SubTreeNodeIndexesTempBuffer.Reserve(Nodes.Num());
+
+		bool bBestNodeIndexFound = false;
+#endif // VALIDATE_FINDNEIGHBORS
+
+		// @todo: avoid allocations
+		TArray<int32> NodesToSearch;
+		NodesToSearch.Emplace(0);
+
+        VPTreeScalar Tau = VPTreeScalar(UE_BIG_NUMBER);
+		while (!NodesToSearch.IsEmpty())
+		{
+			const int32 NodeIndex = NodesToSearch.Pop();
+			const FVPTreeNode& Node = Nodes[NodeIndex];
+
+			// query distance from vantage point
+			const VPTreeScalar QueryDistance = TDataSource::GetDistance(DataSource[Node.Index], Query);
+
+#if VALIDATE_FINDNEIGHBORS
+			++NumEvaluatedQueryDistances;
+
+			bBestNodeIndexFound |= NodeIndex == BestNodeIndex;
+			bool bNodeIndexContainsBestNodeIndex = false;
+			if (!bBestNodeIndexFound)
+			{
+				SubTreeNodeIndexesTempBuffer.Reset();
+				GetSubTreeNodeIndexes(NodeIndex, SubTreeNodeIndexesTempBuffer);
+
+				bNodeIndexContainsBestNodeIndex = SubTreeNodeIndexesTempBuffer.Contains(BestNodeIndex);
+
+				if (!bNodeIndexContainsBestNodeIndex)
+				{
+					bool bOtherNodeIndexContainsBestNodeIndex = false;
+					for (int32 OtherNodeIndex : NodesToSearch)
+					{
+						SubTreeNodeIndexesTempBuffer.Reset();
+						GetSubTreeNodeIndexes(OtherNodeIndex, SubTreeNodeIndexesTempBuffer);
+						bOtherNodeIndexContainsBestNodeIndex = SubTreeNodeIndexesTempBuffer.Contains(BestNodeIndex);
+						if (bOtherNodeIndexContainsBestNodeIndex)
+						{
+							break;
+						}
+					}
+
+					if (!bOtherNodeIndexContainsBestNodeIndex)
+					{
+						UE_LOG(LogPoseSearch, Error, TEXT("FVPTree::FindNeighbors - we couldn't find the best node..."));
+					}
+				}
+			}
+#endif // VALIDATE_FINDNEIGHBORS
+
+			// Tau represents the worst distance of the points in the TVPTreeResultSet
+			if (QueryDistance <= Tau)
+			{
+#if VALIDATE_FINDNEIGHBORS
+				if (bBestNodeIndexFound)
+				{
+					ResultSet.AddPoint(QueryDistance, Node.Index, BestNodeIndex);
+				}
+				else
+#endif // VALIDATE_FINDNEIGHBORS
+				{
+					ResultSet.AddPoint(QueryDistance, Node.Index);
+				}
+
+				if (ResultSet.IsFull())
+				{
+					Tau = ResultSet.GetWorstDistance();
+				}
+			}
+#if VALIDATE_FINDNEIGHBORS
+			else if (NodeIndex == BestNodeIndex)
+			{
+				UE_LOG(LogPoseSearch, Error, TEXT("FVPTree::FindNeighbors - we had the best node in our hands and let it slipped away..."));
+			}
+
+			bool bIsBestNodeStillReachable = false;
+#endif // VALIDATE_FINDNEIGHBORS
+			if (Node.LeftIndex != INDEX_NONE && QueryDistance < Node.Distance + Tau)
+			{
+				NodesToSearch.Emplace(Node.LeftIndex);
+
+#if VALIDATE_FINDNEIGHBORS
+				if (bNodeIndexContainsBestNodeIndex)
+				{
+					SubTreeNodeIndexesTempBuffer.Reset();
+					GetSubTreeNodeIndexes(Node.LeftIndex, SubTreeNodeIndexesTempBuffer);
+					const bool bLeftIndexContainsBestNodeIndex = SubTreeNodeIndexesTempBuffer.Contains(BestNodeIndex);
+					bIsBestNodeStillReachable |= bLeftIndexContainsBestNodeIndex;
+				}
+#endif // VALIDATE_FINDNEIGHBORS
+			}
+
+			if (Node.RightIndex != INDEX_NONE && QueryDistance >= Node.Distance - Tau)
+			{
+				NodesToSearch.Emplace(Node.RightIndex);
+
+#if VALIDATE_FINDNEIGHBORS
+				if (bNodeIndexContainsBestNodeIndex)
+				{
+					SubTreeNodeIndexesTempBuffer.Reset();
+					GetSubTreeNodeIndexes(Node.RightIndex, SubTreeNodeIndexesTempBuffer);
+					const bool bRightIndexContainsBestNodeIndex = SubTreeNodeIndexesTempBuffer.Contains(BestNodeIndex);
+					bIsBestNodeStillReachable |= bRightIndexContainsBestNodeIndex;
+				}
+#endif // VALIDATE_FINDNEIGHBORS
+			}
+
+#if VALIDATE_FINDNEIGHBORS
+			if (bNodeIndexContainsBestNodeIndex != bIsBestNodeStillReachable)
+			{
+				UE_LOG(LogPoseSearch, Error, TEXT("FVPTree::FindNeighbors - we lost strack of the best node..."));
+			}
+#endif // VALIDATE_FINDNEIGHBORS
+		}
+    }
+
+	void Reset();
+
+	POSESEARCH_API SIZE_T GetAllocatedSize() const;
+
+	bool operator==(const FVPTree& Other) const;
+
+	friend FArchive& operator<<(FArchive& Ar, FVPTree& VPTree);
+
+private:
+    void GetSubTreeNodeIndexes(int32 RootNodeIndex, TArray<int32>& NodeIndexes) const
+    {
+		if (RootNodeIndex != INDEX_NONE)
+		{
+			NodeIndexes.Add(RootNodeIndex);
+			GetSubTreeNodeIndexes(Nodes[RootNodeIndex].LeftIndex, NodeIndexes);
+			GetSubTreeNodeIndexes(Nodes[RootNodeIndex].RightIndex, NodeIndexes);
+		}
+    }
+
+	// choose and swap vantage point into DataSourceIndexMapping[0]
+    template<typename TDataSource>
+    void ChooseVantagePointMappingIndex(const TDataSource& DataSource, TArrayView<FIndexDistance> DataSourceIndexMapping, FRandomStream& RandStream) const
+    {
+        // choosing a random, but deterministic, vantage point within DataSourceIndexMapping
+		const int32 VantagePointMappingIndex = FMath::FloorToInt32(RandStream.FRand() * DataSourceIndexMapping.Num());
+		
+		// swapping DataSourceIndexMapping[VantagePointMappingIndex] with DataSourceIndexMapping[0];
+		const int32 VantagePointDataIndex = DataSourceIndexMapping[VantagePointMappingIndex].Index;
+
+		DataSourceIndexMapping[VantagePointMappingIndex] = DataSourceIndexMapping[0].Index;
+		DataSourceIndexMapping[0].Index = VantagePointDataIndex;
+		DataSourceIndexMapping[0].Distance = VPTreeScalar(0);
+    }
+
+    template<typename TDataSource>
+    void CacheDistancesFromVantagePoint(const TDataSource& DataSource, TArrayView<FIndexDistance> DataSourceIndexMapping) const
+    {
+		DataSourceIndexMapping[0].Distance = VPTreeScalar(0);
+		const int32 VantagePointDataIndex = DataSourceIndexMapping[0].Index;
+
+		ParallelFor(DataSourceIndexMapping.Num() - 1, [&DataSource, &DataSourceIndexMapping, VantagePointDataIndex](int32 ParallelForIndex)
+			{
+				FIndexDistance& IndexDistance = DataSourceIndexMapping[ParallelForIndex + 1];
+				IndexDistance.Distance = TDataSource::GetDistance(DataSource[VantagePointDataIndex], DataSource[IndexDistance.Index]);
+			});
+    }
+
+	template<typename TDataSource>
+    int32 ConstructRecursive(const TDataSource& DataSource, TArrayView<FIndexDistance> DataSourceIndexMapping, FRandomStream& RandStream)
+    {
+		int32 NodeIndex = INDEX_NONE;
+		const int32 Num = DataSourceIndexMapping.Num();
+		if (Num > 0)
+		{
+			NodeIndex = Nodes.Num();
+			Nodes.SetNum(NodeIndex + 1);
+
+			if (Num > 1)
+			{
+				ChooseVantagePointMappingIndex(DataSource, DataSourceIndexMapping, RandStream);
+				CacheDistancesFromVantagePoint(DataSource, DataSourceIndexMapping);
+
+				Nodes[NodeIndex].Index = DataSourceIndexMapping[0].Index;
+
+				// @todo: use something like std::nth_element since we don't need to sort the MedianMappingIndex -> end section of the DataSourceIndexMapping array!
+				DataSourceIndexMapping.Slice(1, Num - 1).StableSort([](const FIndexDistance& A, const FIndexDistance& B)
+					{
+						return A.Distance < B.Distance;
+					});
+
+				// @todo: maybe use the average delta distance between MedianMappingIndex - 1 and MedianMappingIndex, to round up eventual search numerical errors
+				const int32 MedianMappingIndex = Num / 2;
+				Nodes[NodeIndex].Distance = DataSourceIndexMapping[MedianMappingIndex - 1].Distance;
+				Nodes[NodeIndex].LeftIndex = ConstructRecursive(DataSource, DataSourceIndexMapping.Slice(1, MedianMappingIndex - 1), RandStream);
+				Nodes[NodeIndex].RightIndex = ConstructRecursive(DataSource, DataSourceIndexMapping.Slice(MedianMappingIndex, Num - MedianMappingIndex), RandStream);
+			}
+			else
+			{
+				Nodes[NodeIndex].Index = DataSourceIndexMapping[0].Index;
+			}
+		}
+        return NodeIndex;
+    }
+
+	template<typename TDataSource>
+    bool TestConstructRecursive(const TDataSource& DataSource, int32 NodeIndex, TArray<int32>& SubTreeNodeIndexesTempBuffer) const
+    {
+		if (NodeIndex == INDEX_NONE)
+		{
+			return true;
+		}
+
+		const FVPTreeNode Node = Nodes[NodeIndex];
+
+		// making sure all the left sub tree nodes are within Nodes[NodeIndex].Distance to the vantage point Nodes[NodeIndex].Index
+		SubTreeNodeIndexesTempBuffer.Reset();
+		GetSubTreeNodeIndexes(Node.LeftIndex, SubTreeNodeIndexesTempBuffer);
+		for (int32 SubTreeNodeIndex : SubTreeNodeIndexesTempBuffer)
+		{
+			const VPTreeScalar Distance = TDataSource::GetDistance(DataSource[Nodes[NodeIndex].Index], DataSource[Nodes[SubTreeNodeIndex].Index]);
+			if (Distance > Nodes[NodeIndex].Distance)
+			{
+				return false;
+			}
+		}
+
+		// making sure all the right sub tree nodes are outside the hypersphere of radius Nodes[NodeIndex].Distance from the vantage point Nodes[NodeIndex].Index
+		SubTreeNodeIndexesTempBuffer.Reset();
+		GetSubTreeNodeIndexes(Node.RightIndex, SubTreeNodeIndexesTempBuffer);
+		for (int32 SubTreeNodeIndex : SubTreeNodeIndexesTempBuffer)
+		{
+			const VPTreeScalar Distance = TDataSource::GetDistance(DataSource[Nodes[NodeIndex].Index], DataSource[Nodes[SubTreeNodeIndex].Index]);
+			if (Distance <= Nodes[NodeIndex].Distance)
+			{
+				return false;
+			}
+		}
+
+		// test recursively left and right sub trees
+		if (!TestConstructRecursive(DataSource, Node.LeftIndex, SubTreeNodeIndexesTempBuffer))
+		{
+			return false;
+		}
+
+		return TestConstructRecursive(DataSource, Node.RightIndex, SubTreeNodeIndexesTempBuffer);
+    }
+
+	TArray<FVPTreeNode> Nodes;
+};
+
+} // namespace UE::PoseSearch
\ No newline at end of file
