motion matching - fix for PoseSearchDerivedData.cpp(514): error C2653: 'EMotionMatchTestFlags': is not a class or namespace name

diff --git a/Engine/Plugins/Experimental/Animation/PoseSearch/Source/Runtime/Private/PoseSearchDerivedData.cpp b/Engine/Plugins/Experimental/Animation/PoseSearch/Source/Runtime/Private/PoseSearchDerivedData.cpp
--- a/Engine/Plugins/Experimental/Animation/PoseSearch/Source/Runtime/Private/PoseSearchDerivedData.cpp
+++ b/Engine/Plugins/Experimental/Animation/PoseSearch/Source/Runtime/Private/PoseSearchDerivedData.cpp
@@ -511,6 +511,7 @@ static void PreprocessSearchIndexPCAData(FSearchIndex& SearchIndex, int32 NumDim
 		// validating EigenSolver results
 		const ColMajorMatrix EigenVectors = EigenSolver.eigenvectors().real();
 
+#if ENABLE_ANIM_DEBUG
 		if (AnyTestFlags(EMotionMatchTestFlags::ValidateKDTreeConstruct) && NumberOfPrincipalComponents == NumDimensions)
 		{
 			const RowMajorVector ReciprocalWeightsSqrt = MapWeightsSqrt.cwiseInverse();
@@ -523,6 +524,7 @@ static void PreprocessSearchIndexPCAData(FSearchIndex& SearchIndex, int32 NumDim
 				check(Error < UE_KINDA_SMALL_NUMBER);
 			}
 		}
+#endif // ENABLE_ANIM_DEBUG
 
 		// sorting EigenVectors by EigenValues, so we pick the most significant ones to compose our PCA projection matrix.
 		const RowMajorVector EigenValues = EigenSolver.eigenvalues().real();
@@ -553,6 +555,7 @@ static void PreprocessSearchIndexPCAData(FSearchIndex& SearchIndex, int32 NumDim
 
 		MapPCAValues = CenteredValues * PCAProjectionMatrix;
 
+#if ENABLE_ANIM_DEBUG
 		if (AnyTestFlags(EMotionMatchTestFlags::ValidateKDTreeConstruct) && NumberOfPrincipalComponents == NumDimensions)
 		{
 			const RowMajorVector ReciprocalWeightsSqrt = MapWeightsSqrt.cwiseInverse();
@@ -579,6 +582,7 @@ static void PreprocessSearchIndexPCAData(FSearchIndex& SearchIndex, int32 NumDim
 				check(Error < UE_KINDA_SMALL_NUMBER);
 			}
 		}
+#endif // ENABLE_ANIM_DEBUG
 	}
 }
 
