pose search - remove PrivateIncludePaths from PoseSearchEditor.Build.cs

diff --git a/Engine/Plugins/Experimental/Animation/PoseSearch/Source/Editor/PoseSearchEditor.Build.cs b/Engine/Plugins/Experimental/Animation/PoseSearch/Source/Editor/PoseSearchEditor.Build.cs
--- a/Engine/Plugins/Experimental/Animation/PoseSearch/Source/Editor/PoseSearchEditor.Build.cs
+++ b/Engine/Plugins/Experimental/Animation/PoseSearch/Source/Editor/PoseSearchEditor.Build.cs
@@ -53,11 +53,5 @@ public class PoseSearchEditor : ModuleRules
 				"ToolWidgets"
 			}
 		);
-
-		// TODO: Should not be including private headers from a different module
-		PrivateIncludePaths.AddRange(
-			new string[] {
-				Path.Combine(GetModuleDirectory("PoseSearch"), "Private"), // For PoseSearchTraceLogger.h
-			});
 	}
 }
diff --git a/Engine/Plugins/Experimental/Animation/PoseSearch/Source/Editor/Private/PoseSearchSchemaFactory.cpp b/Engine/Plugins/Experimental/Animation/PoseSearch/Source/Editor/Private/PoseSearchSchemaFactory.cpp
--- a/Engine/Plugins/Experimental/Animation/PoseSearch/Source/Editor/Private/PoseSearchSchemaFactory.cpp
+++ b/Engine/Plugins/Experimental/Animation/PoseSearch/Source/Editor/Private/PoseSearchSchemaFactory.cpp
@@ -2,8 +2,6 @@
 
 #include "PoseSearchSchemaFactory.h"
 #include "PoseSearch/PoseSearchSchema.h"
-#include "PoseSearchFeatureChannel_Pose.h"
-#include "PoseSearchFeatureChannel_Trajectory.h"
 
 #define LOCTEXT_NAMESPACE "PoseSearchEditor"
 
@@ -18,11 +16,7 @@ UPoseSearchSchemaFactory::UPoseSearchSchemaFactory(const FObjectInitializer& Obj
 UObject* UPoseSearchSchemaFactory::FactoryCreateNew(UClass* Class, UObject* InParent, FName Name, EObjectFlags Flags, UObject* Context, FFeedbackContext* Warn)
 {
 	UPoseSearchSchema* Schema = NewObject<UPoseSearchSchema>(InParent, Class, Name, Flags);
-
-	// defaulting UPoseSearchSchema for a meaningful locomotion setup
-	Schema->AddChannel(NewObject<UPoseSearchFeatureChannel_Trajectory>(Schema, NAME_None, RF_Transactional));
-	Schema->AddChannel(NewObject<UPoseSearchFeatureChannel_Pose>(Schema, NAME_None, RF_Transactional));
-
+	Schema->AddDefaultChannels();
 	return Schema;
 }
 
diff --git a/Engine/Plugins/Experimental/Animation/PoseSearch/Source/Editor/Private/Trace/PoseSearchTraceAnalyzer.cpp b/Engine/Plugins/Experimental/Animation/PoseSearch/Source/Editor/Private/Trace/PoseSearchTraceAnalyzer.cpp
--- a/Engine/Plugins/Experimental/Animation/PoseSearch/Source/Editor/Private/Trace/PoseSearchTraceAnalyzer.cpp
+++ b/Engine/Plugins/Experimental/Animation/PoseSearch/Source/Editor/Private/Trace/PoseSearchTraceAnalyzer.cpp
@@ -5,7 +5,7 @@
 #include "HAL/LowLevelMemTracker.h"
 #include "PoseSearchTraceProvider.h"
 #include "Serialization/MemoryReader.h"
-#include "Trace/PoseSearchTraceLogger.h"
+#include "PoseSearch/Trace/PoseSearchTraceLogger.h"
 #include "TraceServices/Model/AnalysisSession.h"
 
 namespace UE::PoseSearch
diff --git a/Engine/Plugins/Experimental/Animation/PoseSearch/Source/Editor/Private/Trace/PoseSearchTraceProvider.h b/Engine/Plugins/Experimental/Animation/PoseSearch/Source/Editor/Private/Trace/PoseSearchTraceProvider.h
--- a/Engine/Plugins/Experimental/Animation/PoseSearch/Source/Editor/Private/Trace/PoseSearchTraceProvider.h
+++ b/Engine/Plugins/Experimental/Animation/PoseSearch/Source/Editor/Private/Trace/PoseSearchTraceProvider.h
@@ -4,10 +4,9 @@
 
 #include "CoreMinimal.h"
 #include "Model/PointTimeline.h"
-#include "Trace/PoseSearchTraceLogger.h"
+#include "PoseSearch/Trace/PoseSearchTraceLogger.h"
 #include "TraceServices/Model/AnalysisSession.h"
 
-
 namespace UE::PoseSearch
 {
 
diff --git a/Engine/Plugins/Experimental/Animation/PoseSearch/Source/Runtime/Private/AnimNode_MotionMatching.cpp b/Engine/Plugins/Experimental/Animation/PoseSearch/Source/Runtime/Private/AnimNode_MotionMatching.cpp
--- a/Engine/Plugins/Experimental/Animation/PoseSearch/Source/Runtime/Private/AnimNode_MotionMatching.cpp
+++ b/Engine/Plugins/Experimental/Animation/PoseSearch/Source/Runtime/Private/AnimNode_MotionMatching.cpp
@@ -13,7 +13,7 @@
 #include "PoseSearch/PoseSearchDatabase.h"
 #include "PoseSearch/PoseSearchDerivedData.h"
 #include "PoseSearch/PoseSearchSchema.h"
-#include "Trace/PoseSearchTraceLogger.h"
+#include "PoseSearch/Trace/PoseSearchTraceLogger.h"
 
 #include UE_INLINE_GENERATED_CPP_BY_NAME(AnimNode_MotionMatching)
 
diff --git a/Engine/Plugins/Experimental/Animation/PoseSearch/Source/Runtime/Private/PoseSearchLibrary.cpp b/Engine/Plugins/Experimental/Animation/PoseSearch/Source/Runtime/Private/PoseSearchLibrary.cpp
--- a/Engine/Plugins/Experimental/Animation/PoseSearch/Source/Runtime/Private/PoseSearchLibrary.cpp
+++ b/Engine/Plugins/Experimental/Animation/PoseSearch/Source/Runtime/Private/PoseSearchLibrary.cpp
@@ -19,7 +19,7 @@
 #include "PoseSearch/PoseSearchDerivedData.h"
 #include "PoseSearch/PoseSearchSchema.h"
 #include "PoseSearchFeatureChannel_Trajectory.h"
-#include "Trace/PoseSearchTraceLogger.h"
+#include "PoseSearch/Trace/PoseSearchTraceLogger.h"
 
 #include UE_INLINE_GENERATED_CPP_BY_NAME(PoseSearchLibrary)
 
diff --git a/Engine/Plugins/Experimental/Animation/PoseSearch/Source/Runtime/Private/PoseSearchSchema.cpp b/Engine/Plugins/Experimental/Animation/PoseSearch/Source/Runtime/Private/PoseSearchSchema.cpp
--- a/Engine/Plugins/Experimental/Animation/PoseSearch/Source/Runtime/Private/PoseSearchSchema.cpp
+++ b/Engine/Plugins/Experimental/Animation/PoseSearch/Source/Runtime/Private/PoseSearchSchema.cpp
@@ -6,6 +6,8 @@
 #include "PoseSearch/PoseSearchResult.h"
 #include "PoseSearchFeatureChannel_Padding.h"
 #include "PoseSearchFeatureChannel_PermutationTime.h"
+#include "PoseSearchFeatureChannel_Pose.h"
+#include "PoseSearchFeatureChannel_Trajectory.h"
 #include "UObject/ObjectSaveContext.h"
 
 bool UPoseSearchSchema::IsValid() const
@@ -43,6 +45,13 @@ FBoneIndexType UPoseSearchSchema::GetBoneIndexType(int8 SchemaBoneIdx) const
 	return BoneReferences[SchemaBoneIdx].BoneIndex;
 }
 
+void UPoseSearchSchema::AddDefaultChannels()
+{
+	// defaulting UPoseSearchSchema for a meaningful locomotion setup
+	AddChannel(NewObject<UPoseSearchFeatureChannel_Trajectory>(this, NAME_None, RF_Transactional));
+	AddChannel(NewObject<UPoseSearchFeatureChannel_Pose>(this, NAME_None, RF_Transactional));
+}
+
 int8 UPoseSearchSchema::AddBoneReference(const FBoneReference& BoneReference)
 {
 	int32 SchemaBoneIdx = 0;
diff --git a/Engine/Plugins/Experimental/Animation/PoseSearch/Source/Runtime/Private/Trace/PoseSearchTraceLogger.cpp b/Engine/Plugins/Experimental/Animation/PoseSearch/Source/Runtime/Private/Trace/PoseSearchTraceLogger.cpp
--- a/Engine/Plugins/Experimental/Animation/PoseSearch/Source/Runtime/Private/Trace/PoseSearchTraceLogger.cpp
+++ b/Engine/Plugins/Experimental/Animation/PoseSearch/Source/Runtime/Private/Trace/PoseSearchTraceLogger.cpp
@@ -1,6 +1,6 @@
 ﻿// Copyright Epic Games, Inc. All Rights Reserved.
 
-#include "PoseSearchTraceLogger.h"
+#include "PoseSearch/Trace/PoseSearchTraceLogger.h"
 #include "Animation/AnimInstanceProxy.h"
 #include "Animation/AnimNodeBase.h"
 #include "Components/SkeletalMeshComponent.h"
diff --git a/Engine/Plugins/Experimental/Animation/PoseSearch/Source/Runtime/Private/Trace/PoseSearchTraceLogger.h b/Engine/Plugins/Experimental/Animation/PoseSearch/Source/Runtime/Private/Trace/PoseSearchTraceLogger.h
deleted file mode 100644
diff --git a/Engine/Plugins/Experimental/Animation/PoseSearch/Source/Runtime/Public/PoseSearch/PoseSearchDatabase.h b/Engine/Plugins/Experimental/Animation/PoseSearch/Source/Runtime/Public/PoseSearch/PoseSearchDatabase.h
--- a/Engine/Plugins/Experimental/Animation/PoseSearch/Source/Runtime/Public/PoseSearch/PoseSearchDatabase.h
+++ b/Engine/Plugins/Experimental/Animation/PoseSearch/Source/Runtime/Public/PoseSearch/PoseSearchDatabase.h
@@ -14,7 +14,10 @@ class UAnimationAsset;
 class UAnimComposite;
 class UAnimMontage;
 class UBlendSpace;
+
+#if WITH_EDITORONLY_DATA
 class UPoseSearchNormalizationSet;
+#endif // WITH_EDITORONLY_DATA
 
 namespace UE::PoseSearch
 {
diff --git a/Engine/Plugins/Experimental/Animation/PoseSearch/Source/Runtime/Public/PoseSearch/PoseSearchSchema.h b/Engine/Plugins/Experimental/Animation/PoseSearch/Source/Runtime/Public/PoseSearch/PoseSearchSchema.h
--- a/Engine/Plugins/Experimental/Animation/PoseSearch/Source/Runtime/Public/PoseSearch/PoseSearchSchema.h
+++ b/Engine/Plugins/Experimental/Animation/PoseSearch/Source/Runtime/Public/PoseSearch/PoseSearchSchema.h
@@ -146,6 +146,8 @@ public:
 
 	FBoneIndexType GetBoneIndexType(int8 SchemaBoneIdx) const;
 
+	void AddDefaultChannels();
+
 private:
 	template <typename FindPredicateType>
 	static const UPoseSearchFeatureChannel* FindChannelRecursive(TConstArrayView<TObjectPtr<UPoseSearchFeatureChannel>> Channels, FindPredicateType FindPredicate)
diff --git a/Engine/Plugins/Experimental/Animation/PoseSearch/Source/Runtime/Public/PoseSearch/Trace/PoseSearchTraceLogger.h b/Engine/Plugins/Experimental/Animation/PoseSearch/Source/Runtime/Public/PoseSearch/Trace/PoseSearchTraceLogger.h
new file mode 100644
--- /dev/null
+++ b/Engine/Plugins/Experimental/Animation/PoseSearch/Source/Runtime/Public/PoseSearch/Trace/PoseSearchTraceLogger.h
@@ -0,0 +1,131 @@
+﻿// Copyright Epic Games, Inc. All Rights Reserved.
+
+#pragma once
+
+#include "PoseSearch/PoseSearchContext.h"
+#include "PoseSearch/PoseSearchTrajectoryTypes.h"
+
+UE_TRACE_CHANNEL_EXTERN(PoseSearchChannel, POSESEARCH_API);
+
+namespace UE::PoseSearch
+{
+
+bool IsTracing(const FAnimationBaseContext& InContext);
+
+struct POSESEARCH_API FTraceLogger
+{
+	/** Used for reading trace data */
+    static const FName Name;
+};
+
+// Message types for appending / reading to the timeline
+/** Base message type for common data */
+struct POSESEARCH_API FTraceMessage
+{
+	uint64 Cycle = 0;
+
+	uint64 AnimInstanceId = 0;
+
+	/** Skeletal Mesh Component Id, outer class of the AnimInstance.
+	 *  Used for retrieval of traced root transforms from the animation provider.
+	 */
+	uint64 SkeletalMeshComponentId = 0;
+
+	/** Node Id of the motion matching node associated with this message */
+	int32 NodeId = 0;
+
+	uint16 FrameCounter = 0;
+};
+POSESEARCH_API FArchive& operator<<(FArchive& Ar, FTraceMessage& State);
+
+struct POSESEARCH_API FTraceMotionMatchingStatePoseEntry
+{
+	int32 DbPoseIdx = INDEX_NONE;
+	EPoseCandidateFlags PoseCandidateFlags = EPoseCandidateFlags::None;
+	FPoseSearchCost Cost;
+
+	bool operator==(const FTraceMotionMatchingStatePoseEntry& Other) const { return DbPoseIdx == Other.DbPoseIdx; }
+};
+POSESEARCH_API FArchive& operator<<(FArchive& Ar, FTraceMotionMatchingStatePoseEntry& Entry);
+
+struct POSESEARCH_API FTraceMotionMatchingStateDatabaseEntry
+{
+	// @todo: can we use UPoseSearchDatabase* instead of DatabaseId?
+	//UPoseSearchDatabase const* Database = nullptr;
+	uint64 DatabaseId = 0;
+	TArray<float> QueryVector;
+	TArray<FTraceMotionMatchingStatePoseEntry> PoseEntries;
+	
+	bool operator==(const FTraceMotionMatchingStateDatabaseEntry& Other) const { return DatabaseId == Other.DatabaseId; }
+};
+POSESEARCH_API FArchive& operator<<(FArchive& Ar, FTraceMotionMatchingStateDatabaseEntry& Entry);
+
+/**
+ * Used to trace motion matching state data via the logger, which is then placed into a timeline
+ */
+struct POSESEARCH_API FTraceMotionMatchingState
+{
+	/** Amount of time since the last pose switch */
+	float ElapsedPoseSearchTime = 0.f;
+
+	float AssetPlayerTime = 0.f;
+	float DeltaTime = 0.f;
+	float SimLinearVelocity = 0.f;
+	float SimAngularVelocity = 0.f;
+	float AnimLinearVelocity = 0.f;
+	float AnimAngularVelocity = 0.f;
+	
+	float RecordingTime = 0.f;
+	float SearchBestCost = 0.f;
+	float SearchBruteForceCost = 0.f;
+	int32 SearchBestPosePos = 0;
+
+	TArray<FTraceMotionMatchingStateDatabaseEntry> DatabaseEntries;
+
+	FPoseSearchQueryTrajectory Trajectory;
+
+	/** Index of the current database in DatabaseEntries */
+	int32 CurrentDbEntryIdx = INDEX_NONE;
+
+	/** Index of the current pose in DatabaseEntries[CurrentDbEntryIdx].PoseEntries */
+	int32 CurrentPoseEntryIdx = INDEX_NONE;
+
+	/** Output the current state info to the logger */
+	void Output(const UObject* AnimInstance, int32 NodeId);
+
+	const UPoseSearchDatabase* GetCurrentDatabase() const;
+	int32 GetCurrentDatabasePoseIndex() const;
+	const FTraceMotionMatchingStatePoseEntry* GetCurrentPoseEntry() const;
+
+	template<typename T>
+	static const T* GetObjectFromId(uint64 ObjectId)
+	{
+#if OBJECT_TRACE_ENABLED
+		if (ObjectId)
+		{
+			UObject* Object = FObjectTrace::GetObjectFromId(ObjectId);
+			if (Object)
+			{
+				return CastChecked<T>(Object);
+			}
+		}
+#endif
+
+		return nullptr;
+	}
+
+	static uint64 GetIdFromObject(const UObject* Object)
+	{
+#if OBJECT_TRACE_ENABLED
+		return FObjectTrace::GetObjectId(Object);
+#else
+		return 0;
+#endif
+	}
+	
+	static const FName Name;
+};
+
+POSESEARCH_API FArchive& operator<<(FArchive& Ar, FTraceMotionMatchingState& State);
+
+} // namespace UE::PoseSearch
