motion matching - disable reselection of the currently playing asset

diff --git a/Engine/Plugins/Experimental/Animation/PoseSearch/Source/Editor/Private/PoseSearchDatabaseAssetListItem.cpp b/Engine/Plugins/Experimental/Animation/PoseSearch/Source/Editor/Private/PoseSearchDatabaseAssetListItem.cpp
--- a/Engine/Plugins/Experimental/Animation/PoseSearch/Source/Editor/Private/PoseSearchDatabaseAssetListItem.cpp
+++ b/Engine/Plugins/Experimental/Animation/PoseSearch/Source/Editor/Private/PoseSearchDatabaseAssetListItem.cpp
@@ -301,6 +301,26 @@ namespace UE::PoseSearch
 						.Image(this, &SDatabaseAssetListItem::GetMirrorOptionSlateBrush)
 						.ToolTipText(this, &SDatabaseAssetListItem::GetMirrorOptionToolTip)
 					]
+
+					// Disable Reselection
+					+ SHorizontalBox::Slot()
+					.Padding(1.0f, 1.0f)
+					.AutoWidth()
+					.HAlign(HAlign_Center)
+					.VAlign(VAlign_Center)
+					[
+						SNew(SCheckBox)
+						.IsChecked(this, &SDatabaseAssetListItem::GetDisableReselectionChecked)
+						.OnCheckStateChanged(const_cast<SDatabaseAssetListItem*>(this), &SDatabaseAssetListItem::OnDisableReselectionChanged)
+						.ToolTipText(this, &SDatabaseAssetListItem::GetDisableReselectionToolTip)
+						// @todo: customize icon!
+						//.CheckedImage(FAppStyle::Get().GetBrush("FractureFlush"))
+						//.CheckedHoveredImage(FAppStyle::Get().GetBrush("FractureFlush"))
+						//.CheckedPressedImage(FAppStyle::Get().GetBrush("FractureFlush"))
+						//.UncheckedImage(FAppStyle::Get().GetBrush("FractureFlush"))
+						//.UncheckedHoveredImage(FAppStyle::Get().GetBrush("FractureFlush"))
+						//.UncheckedPressedImage(FAppStyle::Get().GetBrush("FractureFlush"))
+					]
 				]
 			]
 			
@@ -372,6 +392,49 @@ namespace UE::PoseSearch
 		return EVisibility::Hidden;
 	}
 
+	FText SDatabaseAssetListItem::GetDisableReselectionToolTip() const
+	{
+		if (GetDisableReselectionChecked() == ECheckBoxState::Checked)
+		{
+			return LOCTEXT("EnableReselectionToolTip", "Enable reselection of poses from the same asset.");
+		}
+		
+		return LOCTEXT("DisableReselectionToolTip", "Disable reselection of poses from the same asset.");
+	}
+
+	ECheckBoxState SDatabaseAssetListItem::GetDisableReselectionChecked() const
+	{
+		TSharedPtr<FDatabaseViewModel> ViewModelPtr = EditorViewModel.Pin();
+		TSharedPtr<FDatabaseAssetTreeNode> TreeNodePtr = WeakAssetTreeNode.Pin();
+		const UPoseSearchDatabase* Database = ViewModelPtr->GetPoseSearchDatabase();
+
+		if (Database->AnimationAssets.IsValidIndex(TreeNodePtr->SourceAssetIdx))
+		{
+			if (ViewModelPtr->IsDisableReselection(TreeNodePtr->SourceAssetIdx))
+			{
+				return ECheckBoxState::Checked;
+			}
+		}
+
+		return ECheckBoxState::Unchecked;
+	}
+
+	void SDatabaseAssetListItem::OnDisableReselectionChanged(ECheckBoxState NewCheckboxState)
+	{
+		const FScopedTransaction Transaction(LOCTEXT("EnableChangedForAssetInPoseSearchDatabase", "Update enabled flag for item from Pose Search Database"));
+
+		const TSharedPtr<FDatabaseViewModel> ViewModelPtr = EditorViewModel.Pin();
+		const TSharedPtr<FDatabaseAssetTreeNode> TreeNodePtr = WeakAssetTreeNode.Pin();
+
+		ViewModelPtr->GetPoseSearchDatabase()->Modify();
+		
+		ViewModelPtr->SetDisableReselection(TreeNodePtr->SourceAssetIdx, NewCheckboxState == ECheckBoxState::Checked ? true : false);
+
+		SkeletonView.Pin()->RefreshTreeView(false, true);
+		
+		// no need to rebuild the SearchIndex (ViewModelPtr->BuildSearchIndex()), since bDisableReselection is a runtime only parameter
+	}
+
 	ECheckBoxState SDatabaseAssetListItem::GetAssetEnabledChecked() const
 	{
 		TSharedPtr<FDatabaseViewModel> ViewModelPtr = EditorViewModel.Pin();
diff --git a/Engine/Plugins/Experimental/Animation/PoseSearch/Source/Editor/Private/PoseSearchDatabaseViewModel.cpp b/Engine/Plugins/Experimental/Animation/PoseSearch/Source/Editor/Private/PoseSearchDatabaseViewModel.cpp
--- a/Engine/Plugins/Experimental/Animation/PoseSearch/Source/Editor/Private/PoseSearchDatabaseViewModel.cpp
+++ b/Engine/Plugins/Experimental/Animation/PoseSearch/Source/Editor/Private/PoseSearchDatabaseViewModel.cpp
@@ -361,6 +361,24 @@ namespace UE::PoseSearch
 		PoseSearchDatabase->AnimationAssets.RemoveAt(AnimationAssetIndex);
 	}
 
+	void FDatabaseViewModel::SetDisableReselection(int32 AnimationAssetIndex, bool bEnabled)
+	{
+		if (FPoseSearchDatabaseAnimationAssetBase* DatabaseAnimationAsset = PoseSearchDatabase->GetMutableAnimationAssetBase(AnimationAssetIndex))
+		{
+			DatabaseAnimationAsset->SetDisableReselection(bEnabled);
+		}
+	}
+
+	bool FDatabaseViewModel::IsDisableReselection(int32 AnimationAssetIndex) const
+	{
+		if (const FPoseSearchDatabaseAnimationAssetBase* DatabaseAnimationAsset = PoseSearchDatabase->GetAnimationAssetBase(AnimationAssetIndex))
+		{
+			return DatabaseAnimationAsset->IsDisableReselection();
+		}
+
+		return false;
+	}
+
 	void FDatabaseViewModel::SetIsEnabled(int32 AnimationAssetIndex, bool bEnabled)
 	{
 		if (FPoseSearchDatabaseAnimationAssetBase* DatabaseAnimationAsset = PoseSearchDatabase->GetMutableAnimationAssetBase(AnimationAssetIndex))
diff --git a/Engine/Plugins/Experimental/Animation/PoseSearch/Source/Editor/Private/PoseSearchDebuggerDatabaseView.cpp b/Engine/Plugins/Experimental/Animation/PoseSearch/Source/Editor/Private/PoseSearchDebuggerDatabaseView.cpp
--- a/Engine/Plugins/Experimental/Animation/PoseSearch/Source/Editor/Private/PoseSearchDebuggerDatabaseView.cpp
+++ b/Engine/Plugins/Experimental/Animation/PoseSearch/Source/Editor/Private/PoseSearchDebuggerDatabaseView.cpp
@@ -269,7 +269,9 @@ void SDebuggerDatabaseView::Update(const FTraceMotionMatchingStateMessage& State
 			SharedData->DatabasePath = Database->GetPathName();
 			SharedData->QueryVector = DbEntry.QueryVector;
 
-			if (Database->PoseSearchMode == EPoseSearchMode::PCAKDTree)
+			// checking for DbEntry.QueryVector.Num() == Database->Schema->SchemaCardinality to make sure the stored DbEntry data is still aligned with the Database->Schema
+			// (this could happen when pausing PIE, change the schema cardinality, scrubbing back to previously recorded data)
+			if (Database->PoseSearchMode == EPoseSearchMode::PCAKDTree && DbEntry.QueryVector.Num() == Database->Schema->SchemaCardinality)
 			{
 				SharedData->PCAQueryVector.SetNumZeroed(Database->GetNumberOfPrincipalComponents());
 				Database->GetSearchIndex().PCAProject(DbEntry.QueryVector, SharedData->PCAQueryVector);
diff --git a/Engine/Plugins/Experimental/Animation/PoseSearch/Source/Editor/Public/PoseSearchDatabaseAssetListItem.h b/Engine/Plugins/Experimental/Animation/PoseSearch/Source/Editor/Public/PoseSearchDatabaseAssetListItem.h
--- a/Engine/Plugins/Experimental/Animation/PoseSearch/Source/Editor/Public/PoseSearchDatabaseAssetListItem.h
+++ b/Engine/Plugins/Experimental/Animation/PoseSearch/Source/Editor/Public/PoseSearchDatabaseAssetListItem.h
@@ -48,6 +48,10 @@ namespace UE::PoseSearch
 		const FSlateBrush* GetMirrorOptionSlateBrush() const;
 		FText GetMirrorOptionToolTip() const;
 		
+		FText GetDisableReselectionToolTip() const;
+		ECheckBoxState GetDisableReselectionChecked() const;
+		void OnDisableReselectionChanged(ECheckBoxState NewCheckboxState);
+
 		FText GetAssetEnabledToolTip() const;
 		ECheckBoxState GetAssetEnabledChecked() const;
 		void OnAssetIsEnabledChanged(ECheckBoxState NewCheckboxState);
diff --git a/Engine/Plugins/Experimental/Animation/PoseSearch/Source/Editor/Public/PoseSearchDatabaseViewModel.h b/Engine/Plugins/Experimental/Animation/PoseSearch/Source/Editor/Public/PoseSearchDatabaseViewModel.h
--- a/Engine/Plugins/Experimental/Animation/PoseSearch/Source/Editor/Public/PoseSearchDatabaseViewModel.h
+++ b/Engine/Plugins/Experimental/Animation/PoseSearch/Source/Editor/Public/PoseSearchDatabaseViewModel.h
@@ -108,6 +108,9 @@ namespace UE::PoseSearch
 		void AddAnimMontageToDatabase(UAnimMontage* AnimMontage);
 		void DeleteFromDatabase(int32 AnimationAssetIndex);
 
+		void SetDisableReselection(int32 AnimationAssetIndex, bool bEnabled);
+		bool IsDisableReselection(int32 AnimationAssetIndex) const;
+
 		void SetIsEnabled(int32 AnimationAssetIndex, bool bEnabled);
 		bool IsEnabled(int32 AnimationAssetIndex) const;
 
diff --git a/Engine/Plugins/Experimental/Animation/PoseSearch/Source/Runtime/Private/PoseSearchDatabase.cpp b/Engine/Plugins/Experimental/Animation/PoseSearch/Source/Runtime/Private/PoseSearchDatabase.cpp
--- a/Engine/Plugins/Experimental/Animation/PoseSearch/Source/Runtime/Private/PoseSearchDatabase.cpp
+++ b/Engine/Plugins/Experimental/Animation/PoseSearch/Source/Runtime/Private/PoseSearchDatabase.cpp
@@ -42,52 +42,81 @@ static void PopulateNonSelectableIdx(FNonSelectableIdx& NonSelectableIdx, FSearc
 )
 {
 	check(Database);
-#if UE_POSE_SEARCH_TRACE_ENABLED
 	const FSearchIndex& SearchIndex = Database->GetSearchIndex();
-#endif // UE_POSE_SEARCH_TRACE_ENABLED
 
 	NonSelectableIdx.Reset();
-	const FSearchIndexAsset* CurrentIndexAsset = SearchContext.GetCurrentResult().GetSearchIndexAsset();
-	if (CurrentIndexAsset && SearchContext.IsCurrentResultFromDatabase(Database) && !FMath::IsNearlyEqual(SearchContext.GetPoseJumpThresholdTime().Min, SearchContext.GetPoseJumpThresholdTime().Max))
+	if (SearchContext.IsCurrentResultFromDatabase(Database))
 	{
-		const int32 CurrentResultPoseIdx = SearchContext.GetCurrentResult().PoseIdx;
-		const int32 UnboundMinPoseIdx = CurrentResultPoseIdx + FMath::FloorToInt(SearchContext.GetPoseJumpThresholdTime().Min * Database->Schema->SampleRate);
-		const int32 UnboundMaxPoseIdx = CurrentResultPoseIdx + FMath::CeilToInt(SearchContext.GetPoseJumpThresholdTime().Max * Database->Schema->SampleRate);
-		const int32 CurrentIndexAssetFirstPoseIdx = CurrentIndexAsset->FirstPoseIdx;
-		const int32 CurrentIndexAssetNumPoses = CurrentIndexAsset->GetNumPoses();
-		const bool IsLooping = Database->IsSourceAssetLooping(*CurrentIndexAsset);
-
-		if (IsLooping)
+		if (const FSearchIndexAsset* CurrentIndexAsset = SearchContext.GetCurrentResult().GetSearchIndexAsset(true))
 		{
-			for (int32 UnboundPoseIdx = UnboundMinPoseIdx; UnboundPoseIdx < UnboundMaxPoseIdx; ++UnboundPoseIdx)
+			const FPoseSearchDatabaseAnimationAssetBase* DatabaseAnimationAssetBase = Database->GetAnimationAssetStruct(CurrentIndexAsset->SourceAssetIdx).GetPtr<FPoseSearchDatabaseAnimationAssetBase>();
+			check(DatabaseAnimationAssetBase);
+			if (!DatabaseAnimationAssetBase->bDisableReselection)
 			{
-				const int32 Modulo = (UnboundPoseIdx - CurrentIndexAssetFirstPoseIdx) % CurrentIndexAssetNumPoses;
-				const int32 CurrentIndexAssetFirstPoseIdxPlusModulo = CurrentIndexAssetFirstPoseIdx + Modulo;
-				const int32 PoseIdx = Modulo >= 0 ? CurrentIndexAssetFirstPoseIdxPlusModulo : CurrentIndexAssetFirstPoseIdxPlusModulo + CurrentIndexAssetNumPoses;
-
-				NonSelectableIdx.AddUnique(PoseIdx);
+				// excluding all the poses from DatabaseAnimationAssetBase
+				// @todo: optimize this code!
+				for (const FSearchIndexAsset& SearchIndexAsset : SearchIndex.Assets)
+				{
+					if (SearchIndexAsset.SourceAssetIdx == CurrentIndexAsset->SourceAssetIdx)
+					{
+						const int32 FirstPoseIdx = SearchIndexAsset.FirstPoseIdx;
+						const int32 LastPoseIdx = FirstPoseIdx + SearchIndexAsset.GetNumPoses();
+						for (int32 PoseIdx = FirstPoseIdx; PoseIdx < LastPoseIdx; ++PoseIdx)
+						{
+							// no need to AddUnique since there's no overlapping between pose indexes in the FSearchIndexAsset(s)
+							NonSelectableIdx.Add(PoseIdx);
 
 #if UE_POSE_SEARCH_TRACE_ENABLED
-				const TArray<float> PoseValues = SearchIndex.GetPoseValuesSafe(PoseIdx);
-				const FPoseSearchCost PoseCost = SearchIndex.ComparePoses(PoseIdx, 0.f, PoseValues, QueryValues);
-				SearchContext.BestCandidates.Add(PoseCost, PoseIdx, Database, EPoseCandidateFlags::DiscardedBy_PoseJumpThresholdTime);
+							const TArray<float> PoseValues = SearchIndex.GetPoseValuesSafe(PoseIdx);
+							const FPoseSearchCost PoseCost = SearchIndex.ComparePoses(PoseIdx, 0.f, PoseValues, QueryValues);
+							SearchContext.BestCandidates.Add(PoseCost, PoseIdx, Database, EPoseCandidateFlags::DiscardedBy_PoseJumpThresholdTime);
 #endif // UE_POSE_SEARCH_TRACE_ENABLED
+						}
+					}
+				}
 			}
-		}
-		else
-		{
-			const int32 MinPoseIdx = FMath::Max(CurrentIndexAssetFirstPoseIdx, UnboundMinPoseIdx);
-			const int32 MaxPoseIdx = FMath::Min(CurrentIndexAssetFirstPoseIdx + CurrentIndexAssetNumPoses, UnboundMaxPoseIdx);
-
-			for (int32 PoseIdx = MinPoseIdx; PoseIdx < MaxPoseIdx; ++PoseIdx)
+			else if (!FMath::IsNearlyEqual(SearchContext.GetPoseJumpThresholdTime().Min, SearchContext.GetPoseJumpThresholdTime().Max))
 			{
-				NonSelectableIdx.AddUnique(PoseIdx);
+				const int32 CurrentResultPoseIdx = SearchContext.GetCurrentResult().PoseIdx;
+				const int32 UnboundMinPoseIdx = CurrentResultPoseIdx + FMath::FloorToInt(SearchContext.GetPoseJumpThresholdTime().Min * Database->Schema->SampleRate);
+				const int32 UnboundMaxPoseIdx = CurrentResultPoseIdx + FMath::CeilToInt(SearchContext.GetPoseJumpThresholdTime().Max * Database->Schema->SampleRate);
+				const int32 CurrentIndexAssetFirstPoseIdx = CurrentIndexAsset->FirstPoseIdx;
+				const int32 CurrentIndexAssetNumPoses = CurrentIndexAsset->GetNumPoses();
+				const bool IsLooping = Database->IsSourceAssetLooping(*CurrentIndexAsset);
+
+				if (IsLooping)
+				{
+					for (int32 UnboundPoseIdx = UnboundMinPoseIdx; UnboundPoseIdx < UnboundMaxPoseIdx; ++UnboundPoseIdx)
+					{
+						const int32 Modulo = (UnboundPoseIdx - CurrentIndexAssetFirstPoseIdx) % CurrentIndexAssetNumPoses;
+						const int32 CurrentIndexAssetFirstPoseIdxPlusModulo = CurrentIndexAssetFirstPoseIdx + Modulo;
+						const int32 PoseIdx = Modulo >= 0 ? CurrentIndexAssetFirstPoseIdxPlusModulo : CurrentIndexAssetFirstPoseIdxPlusModulo + CurrentIndexAssetNumPoses;
+
+						NonSelectableIdx.AddUnique(PoseIdx);
+
+#if UE_POSE_SEARCH_TRACE_ENABLED
+						const TArray<float> PoseValues = SearchIndex.GetPoseValuesSafe(PoseIdx);
+						const FPoseSearchCost PoseCost = SearchIndex.ComparePoses(PoseIdx, 0.f, PoseValues, QueryValues);
+						SearchContext.BestCandidates.Add(PoseCost, PoseIdx, Database, EPoseCandidateFlags::DiscardedBy_PoseJumpThresholdTime);
+#endif // UE_POSE_SEARCH_TRACE_ENABLED
+					}
+				}
+				else
+				{
+					const int32 MinPoseIdx = FMath::Max(CurrentIndexAssetFirstPoseIdx, UnboundMinPoseIdx);
+					const int32 MaxPoseIdx = FMath::Min(CurrentIndexAssetFirstPoseIdx + CurrentIndexAssetNumPoses, UnboundMaxPoseIdx);
+
+					for (int32 PoseIdx = MinPoseIdx; PoseIdx < MaxPoseIdx; ++PoseIdx)
+					{
+						NonSelectableIdx.AddUnique(PoseIdx);
 
 #if UE_POSE_SEARCH_TRACE_ENABLED
-				const TArray<float> PoseValues = SearchIndex.GetPoseValuesSafe(PoseIdx);
-				const FPoseSearchCost PoseCost = SearchIndex.ComparePoses(PoseIdx, 0.f, PoseValues, QueryValues);
-				SearchContext.BestCandidates.Add(PoseCost, PoseIdx, Database, EPoseCandidateFlags::DiscardedBy_PoseJumpThresholdTime);
+						const TArray<float> PoseValues = SearchIndex.GetPoseValuesSafe(PoseIdx);
+						const FPoseSearchCost PoseCost = SearchIndex.ComparePoses(PoseIdx, 0.f, PoseValues, QueryValues);
+						SearchContext.BestCandidates.Add(PoseCost, PoseIdx, Database, EPoseCandidateFlags::DiscardedBy_PoseJumpThresholdTime);
 #endif // UE_POSE_SEARCH_TRACE_ENABLED
+					}
+				}
 			}
 		}
 	}
diff --git a/Engine/Plugins/Experimental/Animation/PoseSearch/Source/Runtime/Public/PoseSearch/PoseSearchDatabase.h b/Engine/Plugins/Experimental/Animation/PoseSearch/Source/Runtime/Public/PoseSearch/PoseSearchDatabase.h
--- a/Engine/Plugins/Experimental/Animation/PoseSearch/Source/Runtime/Public/PoseSearch/PoseSearchDatabase.h
+++ b/Engine/Plugins/Experimental/Animation/PoseSearch/Source/Runtime/Public/PoseSearch/PoseSearchDatabase.h
@@ -78,12 +78,27 @@ struct POSESEARCH_API FPoseSearchDatabaseAnimationAssetBase
 	virtual UClass* GetAnimationAssetStaticClass() const { return nullptr; }
 	virtual bool IsLooping() const { return false; }
 	virtual const FString GetName() const { return FString(); }
-	virtual bool IsEnabled() const { return false; }
-	virtual void SetIsEnabled(bool bInIsEnabled) {}
+	virtual bool IsDisableReselection() const { return bDisableReselection; }
+	virtual void SetDisableReselection(bool bValue) { bDisableReselection = bValue; }
+	virtual bool IsEnabled() const { return bEnabled; }
+	virtual void SetIsEnabled(bool bValue) { bEnabled = bValue; }
 	virtual bool IsRootMotionEnabled() const { return false; }
-	virtual EPoseSearchMirrorOption GetMirrorOption() const { return EPoseSearchMirrorOption::Invalid; }
+	virtual EPoseSearchMirrorOption GetMirrorOption() const { return MirrorOption; }
 	// [0, 0] represents the entire frame range of the original animation.
 	virtual FFloatInterval GetSamplingRange() const { return FFloatInterval(0.f, 0.f); }
+
+	// This allows users to enable or exclude animations from this database. Useful for debugging.
+	UPROPERTY(EditAnywhere, Category = "Settings", meta = (DisplayPriority = 1))
+	bool bEnabled = true;
+
+	// if bDisableReselection is true, poses from the same asset cannot be reselected. Useful to avoid jumping on frames on the same looping animations
+	UPROPERTY(EditAnywhere, Category = "Settings", meta = (ExcludeFromHash, DisplayPriority = 2))
+	bool bDisableReselection = false;
+
+	// This allows users to set if this animation is original only (no mirrored data), original and mirrored, or only the mirrored version of this animation.
+	// It requires the mirror table to be set up in the config file.
+	UPROPERTY(EditAnywhere, Category = "Settings", meta = (DisplayPriority = 3))
+	EPoseSearchMirrorOption MirrorOption = EPoseSearchMirrorOption::UnmirroredOnly;
 };
 
 /** A sequence entry in a UPoseSearchDatabase. */
@@ -93,32 +108,20 @@ struct POSESEARCH_API FPoseSearchDatabaseSequence : public FPoseSearchDatabaseAn
 	GENERATED_BODY()
 	virtual ~FPoseSearchDatabaseSequence() = default;
 
-	UPROPERTY(EditAnywhere, Category="Sequence", meta = (DisplayPriority = 0))
+	UPROPERTY(EditAnywhere, Category="Settings", meta = (DisplayPriority = 0))
 	TObjectPtr<UAnimSequence> Sequence;
 
-	// This allows users to enable or exclude animations from this database. Useful for debugging.
-	UPROPERTY(EditAnywhere, Category = "Sequence", meta = (DisplayPriority = 3))
-	bool bEnabled = true;
-
 	// It allows users to set a time range to an individual animation sequence in the database. 
 	// This is effectively trimming the beginning and end of the animation in the database (not in the original sequence).
 	// If set to [0, 0] it will be the entire frame range of the original sequence.
-	UPROPERTY(EditAnywhere, Category="Sequence", meta = (DisplayPriority = 1))
+	UPROPERTY(EditAnywhere, Category="Settings", meta = (DisplayPriority = 2))
 	FFloatInterval SamplingRange = FFloatInterval(0.f, 0.f);
 
-	// This allows users to set if this animation is original only (no mirrored data), original and mirrored, or only the mirrored version of this animation.
-	// It requires the mirror table to be set up in the config file.
-	UPROPERTY(EditAnywhere, Category = "Sequence", meta = (DisplayPriority = 2))
-	EPoseSearchMirrorOption MirrorOption = EPoseSearchMirrorOption::UnmirroredOnly;
-
 	UAnimationAsset* GetAnimationAsset() const override;
 	UClass* GetAnimationAssetStaticClass() const override;
 	bool IsLooping() const override;
 	const FString GetName() const override;
-	bool IsEnabled() const override { return bEnabled; }
-	void SetIsEnabled(bool bInIsEnabled) override { bEnabled = bInIsEnabled; }
 	bool IsRootMotionEnabled() const override;
-	EPoseSearchMirrorOption GetMirrorOption() const override { return MirrorOption; }
 	FFloatInterval GetSamplingRange() const override { return SamplingRange; }
 };
 
@@ -129,50 +132,38 @@ struct POSESEARCH_API FPoseSearchDatabaseBlendSpace : public FPoseSearchDatabase
 	GENERATED_BODY()
 	virtual ~FPoseSearchDatabaseBlendSpace() = default;
 
-	UPROPERTY(EditAnywhere, Category = "BlendSpace", meta = (DisplayPriority = 0))
+	UPROPERTY(EditAnywhere, Category = "Settings", meta = (DisplayPriority = 0))
 	TObjectPtr<UBlendSpace> BlendSpace;
 
-	// This allows users to set if this animation is original only (no mirrored data), original and mirrored, or only the mirrored version of this animation.
-	// It requires the mirror table to be set up in the config file.
-	UPROPERTY(EditAnywhere, Category = "BlendSpace", meta = (DisplayPriority = 1))
-	EPoseSearchMirrorOption MirrorOption = EPoseSearchMirrorOption::UnmirroredOnly;
-
 	// If true this BlendSpace will output a single segment in the database.
-	UPROPERTY(EditAnywhere, Category = "BlendSpace", meta = (DisplayPriority = 2))
+	UPROPERTY(EditAnywhere, Category = "Settings", meta = (DisplayPriority = 4))
 	bool bUseSingleSample = false;
 
 	// When turned on, this will use the set grid samples of the blend space asset for sampling. This will override the Number of Horizontal/Vertical Samples.
-	UPROPERTY(EditAnywhere, Category = "BlendSpace", meta = (EditCondition = "!bUseSingleSample", EditConditionHides, DisplayPriority = 3))
+	UPROPERTY(EditAnywhere, Category = "Settings", meta = (EditCondition = "!bUseSingleSample", EditConditionHides, DisplayPriority = 5))
 	bool bUseGridForSampling = false;
 
 	// Sets the number of horizontal samples in the blend space to pull the animation data coverage from. The larger the samples the more the data, but also the more memory and performance it takes.
-	UPROPERTY(EditAnywhere, Category = "BlendSpace", meta = (EditCondition = "!bUseSingleSample && !bUseGridForSampling", EditConditionHides, ClampMin = "1", UIMin = "1", UIMax = "25", DisplayPriority = 4))
+	UPROPERTY(EditAnywhere, Category = "Settings", meta = (EditCondition = "!bUseSingleSample && !bUseGridForSampling", EditConditionHides, ClampMin = "1", UIMin = "1", UIMax = "25", DisplayPriority = 6))
 	int32 NumberOfHorizontalSamples = 9;
 	
 	// Sets the number of vertical samples in the blend space to pull the animation data coverage from.The larger the samples the more the data, but also the more memory and performance it takes.
-	UPROPERTY(EditAnywhere, Category = "BlendSpace", meta = (EditCondition = "!bUseSingleSample && !bUseGridForSampling", EditConditionHides, ClampMin = "1", UIMin = "1", UIMax = "25", DisplayPriority = 5))
+	UPROPERTY(EditAnywhere, Category = "Settings", meta = (EditCondition = "!bUseSingleSample && !bUseGridForSampling", EditConditionHides, ClampMin = "1", UIMin = "1", UIMax = "25", DisplayPriority = 7))
 	int32 NumberOfVerticalSamples = 2;
 
 	// BlendParams used to sample this BlendSpace
-	UPROPERTY(EditAnywhere, Category = "BlendSpace", meta = (EditCondition = "bUseSingleSample", EditConditionHides, DisplayPriority = 6))
+	UPROPERTY(EditAnywhere, Category = "Settings", meta = (EditCondition = "bUseSingleSample", EditConditionHides, DisplayPriority = 8))
 	float BlendParamX = 0.f;
 
 	// BlendParams used to sample this BlendSpace
-	UPROPERTY(EditAnywhere, Category = "BlendSpace", meta = (EditCondition = "bUseSingleSample", EditConditionHides, DisplayPriority = 7))
+	UPROPERTY(EditAnywhere, Category = "Settings", meta = (EditCondition = "bUseSingleSample", EditConditionHides, DisplayPriority = 9))
 	float BlendParamY = 0.f;
 
-	// This allows users to enable or exclude animations from this database. Useful for debugging.
-	UPROPERTY(EditAnywhere, Category = "BlendSpace", meta = (DisplayPriority = 8))
-	bool bEnabled = true;
-
 	UAnimationAsset* GetAnimationAsset() const override;
 	UClass* GetAnimationAssetStaticClass() const override;
 	bool IsLooping() const override;
 	const FString GetName() const override;
-	bool IsEnabled() const override { return bEnabled; }
-	void SetIsEnabled(bool bInIsEnabled) override { bEnabled = bInIsEnabled; }
 	bool IsRootMotionEnabled() const override;
-	EPoseSearchMirrorOption GetMirrorOption() const override { return MirrorOption; }
 
 	void GetBlendSpaceParameterSampleRanges(int32& HorizontalBlendNum, int32& VerticalBlendNum) const;
 	FVector BlendParameterForSampleRanges(int32 HorizontalBlendIndex, int32 VerticalBlendIndex) const;
@@ -185,32 +176,20 @@ struct POSESEARCH_API FPoseSearchDatabaseAnimComposite : public FPoseSearchDatab
 	GENERATED_BODY()
 	virtual ~FPoseSearchDatabaseAnimComposite() = default;
 
-	UPROPERTY(EditAnywhere, Category = "AnimComposite", meta = (DisplayPriority = 0))
+	UPROPERTY(EditAnywhere, Category = "Settings", meta = (DisplayPriority = 0))
 	TObjectPtr<UAnimComposite> AnimComposite;
 
-	// This allows users to enable or exclude animations from this database. Useful for debugging.
-	UPROPERTY(EditAnywhere, Category = "AnimComposite", meta = (DisplayPriority = 3))
-	bool bEnabled = true;
-
 	// It allows users to set a time range to an individual animation sequence in the database. 
 	// This is effectively trimming the beginning and end of the animation in the database (not in the original sequence).
 	// If set to [0, 0] it will be the entire frame range of the original sequence.
-	UPROPERTY(EditAnywhere, Category = "AnimComposite", meta = (DisplayPriority = 1))
+	UPROPERTY(EditAnywhere, Category = "Settings", meta = (DisplayPriority = 3))
 	FFloatInterval SamplingRange = FFloatInterval(0.f, 0.f);
 
-	// This allows users to set if this animation is original only (no mirrored data), original and mirrored, or only the mirrored version of this animation.
-	// It requires the mirror table to be set up in the config file.
-	UPROPERTY(EditAnywhere, Category = "AnimComposite", meta = (DisplayPriority = 2))
-	EPoseSearchMirrorOption MirrorOption = EPoseSearchMirrorOption::UnmirroredOnly;
-
 	UAnimationAsset* GetAnimationAsset() const override;
 	UClass* GetAnimationAssetStaticClass() const override;
 	bool IsLooping() const override;
 	const FString GetName() const override;
-	bool IsEnabled() const override { return bEnabled; }
-	void SetIsEnabled(bool bInIsEnabled) override { bEnabled = bInIsEnabled; }
 	bool IsRootMotionEnabled() const override;
-	EPoseSearchMirrorOption GetMirrorOption() const override { return MirrorOption; }
 	FFloatInterval GetSamplingRange() const override { return SamplingRange; }
 };
 
@@ -221,32 +200,20 @@ struct POSESEARCH_API FPoseSearchDatabaseAnimMontage : public FPoseSearchDatabas
 	GENERATED_BODY()
 	virtual ~FPoseSearchDatabaseAnimMontage() = default;
 
-	UPROPERTY(EditAnywhere, Category="AnimMontage", meta = (DisplayPriority = 0))
+	UPROPERTY(EditAnywhere, Category="Settings", meta = (DisplayPriority = 0))
 	TObjectPtr<UAnimMontage> AnimMontage;
 
-	// This allows users to enable or exclude animations from this database. Useful for debugging.
-	UPROPERTY(EditAnywhere, Category = "AnimMontage", meta = (DisplayPriority = 3))
-	bool bEnabled = true;
-
 	// It allows users to set a time range to an individual animation sequence in the database. 
 	// This is effectively trimming the beginning and end of the animation in the database (not in the original sequence).
 	// If set to [0, 0] it will be the entire frame range of the original sequence.
-	UPROPERTY(EditAnywhere, Category="AnimMontage", meta = (DisplayPriority = 1))
+	UPROPERTY(EditAnywhere, Category="Settings", meta = (DisplayPriority = 2))
 	FFloatInterval SamplingRange = FFloatInterval(0.f, 0.f);
 
-	// This allows users to set if this animation is original only (no mirrored data), original and mirrored, or only the mirrored version of this animation.
-	// It requires the mirror table to be set up in the config file.
-	UPROPERTY(EditAnywhere, Category = "AnimMontage", meta = (DisplayPriority = 2))
-	EPoseSearchMirrorOption MirrorOption = EPoseSearchMirrorOption::UnmirroredOnly;
-
 	UAnimationAsset* GetAnimationAsset() const override;
 	UClass* GetAnimationAssetStaticClass() const override;
 	bool IsLooping() const override;
 	const FString GetName() const override;
-	bool IsEnabled() const override { return bEnabled; }
-	void SetIsEnabled(bool bInIsEnabled) override { bEnabled = bInIsEnabled; }
 	bool IsRootMotionEnabled() const override;
-	EPoseSearchMirrorOption GetMirrorOption() const override { return MirrorOption; }
 	FFloatInterval GetSamplingRange() const override { return SamplingRange; }
 };
 
