motion matching - customizing UPoseSearchSchemaFactory::ConfigureProperties to create a skeleton picker modal window for UPoseSearchSchema

diff --git a/Engine/Plugins/Experimental/Animation/PoseSearch/Source/Editor/Private/PoseSearchSchemaFactory.cpp b/Engine/Plugins/Experimental/Animation/PoseSearch/Source/Editor/Private/PoseSearchSchemaFactory.cpp
--- a/Engine/Plugins/Experimental/Animation/PoseSearch/Source/Editor/Private/PoseSearchSchemaFactory.cpp
+++ b/Engine/Plugins/Experimental/Animation/PoseSearch/Source/Editor/Private/PoseSearchSchemaFactory.cpp
@@ -1,7 +1,12 @@
 // Copyright Epic Games, Inc. All Rights Reserved.
 
 #include "PoseSearchSchemaFactory.h"
+#include "Animation/Skeleton.h"
+#include "ContentBrowserModule.h"
+#include "Editor.h"
+#include "IContentBrowserSingleton.h"
 #include "PoseSearch/PoseSearchSchema.h"
+#include "Widgets/Layout/SBorder.h"
 
 #define LOCTEXT_NAMESPACE "PoseSearchEditor"
 
@@ -13,10 +18,56 @@ UPoseSearchSchemaFactory::UPoseSearchSchemaFactory(const FObjectInitializer& Obj
 	SupportedClass = UPoseSearchSchema::StaticClass();
 }
 
+void UPoseSearchSchemaFactory::OnTargetSkeletonSelected(const FAssetData& SelectedAsset)
+{
+	TargetSkeleton = Cast<USkeleton>(SelectedAsset.GetAsset());
+	PickerWindow->RequestDestroyWindow();
+}
+
+bool UPoseSearchSchemaFactory::ConfigureProperties()
+{
+	// Load the content browser module to display an asset picker
+	FContentBrowserModule& ContentBrowserModule = FModuleManager::LoadModuleChecked<FContentBrowserModule>("ContentBrowser");
+
+	FAssetPickerConfig AssetPickerConfig;
+	
+	/** The asset picker will only show skeletons */
+	AssetPickerConfig.Filter.ClassPaths.Add(USkeleton::StaticClass()->GetClassPathName());
+	AssetPickerConfig.Filter.bRecursiveClasses = true;
+
+	/** The delegate that fires when an asset was selected */
+	AssetPickerConfig.OnAssetSelected = FOnAssetSelected::CreateUObject(this, &UPoseSearchSchemaFactory::OnTargetSkeletonSelected);
+
+	/** The default view mode should be a list view */
+	AssetPickerConfig.InitialAssetViewType = EAssetViewType::List;
+
+	PickerWindow = SNew(SWindow)
+	.Title(LOCTEXT("CreatePoseSearchSchemaOptions", "Pick Skeleton"))
+	.ClientSize(FVector2D(500, 600))
+	.SupportsMinimize(false) .SupportsMaximize(false)
+	[
+		SNew(SBorder)
+		.BorderImage( FAppStyle::GetBrush("Menu.Background") )
+		[
+			ContentBrowserModule.Get().CreateAssetPicker(AssetPickerConfig)
+		]
+	];
+
+	GEditor->EditorAddModalWindow(PickerWindow.ToSharedRef());
+	PickerWindow.Reset();
+
+	return TargetSkeleton != nullptr;
+}
+
 UObject* UPoseSearchSchemaFactory::FactoryCreateNew(UClass* Class, UObject* InParent, FName Name, EObjectFlags Flags, UObject* Context, FFeedbackContext* Warn)
 {
-	UPoseSearchSchema* Schema = NewObject<UPoseSearchSchema>(InParent, Class, Name, Flags);
-	Schema->AddDefaultChannels();
+	UPoseSearchSchema* Schema = nullptr;
+	if (TargetSkeleton)
+	{
+		Schema = NewObject<UPoseSearchSchema>(InParent, Class, Name, Flags);
+		Schema->Skeleton = TargetSkeleton;
+		Schema->AddDefaultChannels();
+	}
 	return Schema;
 }
 
diff --git a/Engine/Plugins/Experimental/Animation/PoseSearch/Source/Editor/Public/PoseSearchSchemaFactory.h b/Engine/Plugins/Experimental/Animation/PoseSearch/Source/Editor/Public/PoseSearchSchemaFactory.h
--- a/Engine/Plugins/Experimental/Animation/PoseSearch/Source/Editor/Public/PoseSearchSchemaFactory.h
+++ b/Engine/Plugins/Experimental/Animation/PoseSearch/Source/Editor/Public/PoseSearchSchemaFactory.h
@@ -7,13 +7,24 @@
 #include "Factories/Factory.h"
 #include "PoseSearchSchemaFactory.generated.h"
 
+struct FAssetData;
+class SWindow;
+
 UCLASS()
 class POSESEARCHEDITOR_API UPoseSearchSchemaFactory : public UFactory
 {
 	GENERATED_UCLASS_BODY()
 
+	UPROPERTY()
+	TObjectPtr<class USkeleton> TargetSkeleton;
+
+	TSharedPtr<SWindow> PickerWindow;
+
 	// UFactory interface
+	virtual bool ConfigureProperties() override;
 	virtual UObject* FactoryCreateNew(UClass* Class, UObject* InParent, FName Name, EObjectFlags Flags, UObject* Context, FFeedbackContext* Warn) override;
 	virtual FString GetDefaultNewAssetName() const override;
 	// End of UFactory interface
+
+	void OnTargetSkeletonSelected(const FAssetData& SelectedAsset);
 };
