motion matching - moving ContinuingPoseCostBias, BaseCostBias, and LoopingCostBias from UPoseSearchSchema to UPoseSearchDatabase

diff --git a/Engine/Plugins/Experimental/Animation/PoseSearch/Source/Runtime/Private/PoseSearchAssetIndexer.cpp b/Engine/Plugins/Experimental/Animation/PoseSearch/Source/Runtime/Private/PoseSearchAssetIndexer.cpp
--- a/Engine/Plugins/Experimental/Animation/PoseSearch/Source/Runtime/Private/PoseSearchAssetIndexer.cpp
+++ b/Engine/Plugins/Experimental/Animation/PoseSearch/Source/Runtime/Private/PoseSearchAssetIndexer.cpp
@@ -12,6 +12,7 @@
 #include "PoseSearch/PoseSearchFeatureChannel.h"
 #include "PoseSearch/PoseSearchIndex.h"
 #include "PoseSearch/PoseSearchSchema.h"
+#include "PoseSearch/PoseSearchDatabase.h"
 
 namespace UE::PoseSearch
 {
@@ -80,13 +81,16 @@ static FSamplingParam WrapOrClampSamplingParam(bool bCanWrap, float SamplingPara
 
 //////////////////////////////////////////////////////////////////////////
 // FAssetSamplingContext
-void FAssetSamplingContext::Init(const UMirrorDataTable* InMirrorDataTable, const FBoneContainer& BoneContainer)
+FAssetSamplingContext::FAssetSamplingContext(const UPoseSearchDatabase& Database, const FBoneContainer& BoneContainer)
 {
-	MirrorDataTable = InMirrorDataTable;
+	check(Database.Schema);
+	MirrorDataTable = Database.Schema->MirrorDataTable;
+	BaseCostBias = Database.BaseCostBias;
+	LoopingCostBias = Database.LoopingCostBias;
 
-	if (InMirrorDataTable)
+	if (MirrorDataTable)
 	{
-		InMirrorDataTable->FillCompactPoseAndComponentRefRotations(BoneContainer, CompactPoseMirrorBones, ComponentSpaceRefRotations);
+		MirrorDataTable->FillCompactPoseAndComponentRefRotations(BoneContainer, CompactPoseMirrorBones, ComponentSpaceRefRotations);
 	}
 	else
 	{
@@ -133,8 +137,7 @@ void FAssetIndexer::Process(int32 AssetIdx)
 	for (int32 SampleIdx = GetBeginSampleIdx(); SampleIdx != GetEndSampleIdx(); ++SampleIdx)
 	{
 		const float SampleTime = FMath::Min(CalculateSampleTime(SampleIdx), SequenceLength);
-		float CostAddend = Schema.BaseCostBias;
-		float ContinuingPoseCostAddend = Schema.ContinuingPoseCostBias;
+		float CostAddend = SamplingContext.BaseCostBias;
 		bool bBlockTransition = false;
 
 		TArray<UAnimNotifyState_PoseSearchBase*> NotifyStates;
@@ -153,7 +156,7 @@ void FAssetIndexer::Process(int32 AssetIdx)
 
 		if (AssetSampler.IsLoopable())
 		{
-			CostAddend += Schema.LoopingCostBias;
+			CostAddend += SamplingContext.LoopingCostBias;
 		}
 
 		const int32 ValueOffset = (StartPoseIdx + GetVectorIdx(SampleIdx)) * Schema.SchemaCardinality;
diff --git a/Engine/Plugins/Experimental/Animation/PoseSearch/Source/Runtime/Private/PoseSearchDatabase.cpp b/Engine/Plugins/Experimental/Animation/PoseSearch/Source/Runtime/Private/PoseSearchDatabase.cpp
--- a/Engine/Plugins/Experimental/Animation/PoseSearch/Source/Runtime/Private/PoseSearchDatabase.cpp
+++ b/Engine/Plugins/Experimental/Animation/PoseSearch/Source/Runtime/Private/PoseSearchDatabase.cpp
@@ -778,7 +778,7 @@ FPoseSearchCost UPoseSearchDatabase::SearchContinuingPose(UE::PoseSearch::FSearc
 	}
 #endif // WITH_EDITOR
 
-	// extracting notifies from the database animation asset at time SampleTime to search for UAnimNotifyState_PoseSearchOverrideContinuingPoseCostBias eventually overriding the schema ContinuingPoseCostBias
+	// extracting notifies from the database animation asset at time SampleTime to search for UAnimNotifyState_PoseSearchOverrideContinuingPoseCostBias eventually overriding the database ContinuingPoseCostBias
 	const FSearchIndex& SearchIndex = GetSearchIndex();
 	const int32 PoseIdx = SearchContext.GetCurrentResult().PoseIdx;
 	const FSearchIndexAsset& SearchIndexAsset = SearchIndex.GetAssetForPose(PoseIdx);
@@ -791,19 +791,19 @@ FPoseSearchCost UPoseSearchDatabase::SearchContinuingPose(UE::PoseSearch::FSearc
 	TArray<UAnimNotifyState_PoseSearchBase*> NotifyStates;
 	SequenceBaseSampler.ExtractPoseSearchNotifyStates(SampleTime, NotifyStates);
 
-	float ContinuingPoseCostBias = Schema->ContinuingPoseCostBias;
+	float UpdatedContinuingPoseCostBias = ContinuingPoseCostBias;
 	for (const UAnimNotifyState_PoseSearchBase* PoseSearchNotify : NotifyStates)
 	{
 		if (const UAnimNotifyState_PoseSearchOverrideContinuingPoseCostBias* ContinuingPoseCostBiasNotify = Cast<const UAnimNotifyState_PoseSearchOverrideContinuingPoseCostBias>(PoseSearchNotify))
 		{
-			ContinuingPoseCostBias = ContinuingPoseCostBiasNotify->CostAddend;
+			UpdatedContinuingPoseCostBias = ContinuingPoseCostBiasNotify->CostAddend;
 			break;
 		}
 	}
 
-	// since any PoseCost calculated here is at least SearchIndex.MinCostAddend + ContinuingPoseCostBias,
+	// since any PoseCost calculated here is at least SearchIndex.MinCostAddend + UpdatedContinuingPoseCostBias,
 	// there's no point in performing the search if CurrentBestTotalCost is already better than that
-	if (!GetSkipSearchIfPossible() || SearchContext.GetCurrentBestTotalCost() > SearchIndex.MinCostAddend + ContinuingPoseCostBias)
+	if (!GetSkipSearchIfPossible() || SearchContext.GetCurrentBestTotalCost() > SearchIndex.MinCostAddend + UpdatedContinuingPoseCostBias)
 	{
 		const int32 NumDimensions = Schema->SchemaCardinality;
 		// FMemory_Alloca is forced 16 bytes aligned
@@ -815,12 +815,12 @@ FPoseSearchCost UPoseSearchDatabase::SearchContinuingPose(UE::PoseSearch::FSearc
 		// is the data padded at 16 bytes (and 16 bytes aligned by construction)?
 		if (NumDimensions % 4 == 0)
 		{
-			ContinuingPoseCost = SearchIndex.CompareAlignedPoses(ContinuingPoseIdx, ContinuingPoseCostBias, PoseValues, SearchContext.GetOrBuildQuery(Schema).GetValues());
+			ContinuingPoseCost = SearchIndex.CompareAlignedPoses(ContinuingPoseIdx, UpdatedContinuingPoseCostBias, PoseValues, SearchContext.GetOrBuildQuery(Schema).GetValues());
 		}
 		// data is not 16 bytes padded
 		else
 		{
-			ContinuingPoseCost = SearchIndex.ComparePoses(ContinuingPoseIdx, ContinuingPoseCostBias, PoseValues, SearchContext.GetOrBuildQuery(Schema).GetValues());
+			ContinuingPoseCost = SearchIndex.ComparePoses(ContinuingPoseIdx, UpdatedContinuingPoseCostBias, PoseValues, SearchContext.GetOrBuildQuery(Schema).GetValues());
 		}
 
 #if UE_POSE_SEARCH_TRACE_ENABLED
diff --git a/Engine/Plugins/Experimental/Animation/PoseSearch/Source/Runtime/Private/PoseSearchDerivedData.cpp b/Engine/Plugins/Experimental/Animation/PoseSearch/Source/Runtime/Private/PoseSearchDerivedData.cpp
--- a/Engine/Plugins/Experimental/Animation/PoseSearch/Source/Runtime/Private/PoseSearchDerivedData.cpp
+++ b/Engine/Plugins/Experimental/Animation/PoseSearch/Source/Runtime/Private/PoseSearchDerivedData.cpp
@@ -779,17 +779,6 @@ static bool IndexDatabase(FSearchIndexBase& SearchIndexBase, const UPoseSearchDa
 	const UPoseSearchSchema* Schema = Database.Schema;
 	check(Schema);
 
-	FBoneContainer BoneContainer;
-	BoneContainer.InitializeTo(Schema->BoneIndicesWithParents, UE::Anim::FCurveFilterSettings(UE::Anim::ECurveFilterMode::DisallowAll), *Schema->Skeleton);
-
-	FAssetSamplingContext SamplingContext;
-	SamplingContext.Init(Schema->MirrorDataTable, BoneContainer);
-
-	if (Owner.IsCanceled())
-	{
-		return false;
-	}
-
 	// Prepare samplers for all animation assets.
 	TArray<FAnimationAssetSampler> Samplers;
 	Samplers.Reserve(256);
@@ -834,6 +823,8 @@ static bool IndexDatabase(FSearchIndexBase& SearchIndexBase, const UPoseSearchDa
 		}
 	}
 
+	FBoneContainer BoneContainer;
+	BoneContainer.InitializeTo(Schema->BoneIndicesWithParents, UE::Anim::FCurveFilterSettings(UE::Anim::ECurveFilterMode::DisallowAll), *Schema->Skeleton);
 	ParallelFor(Samplers.Num(), [&Samplers, &BoneContainer](int32 SamplerIdx) { Samplers[SamplerIdx].Process(BoneContainer); }, ParallelForFlags);
 
 	if (Owner.IsCanceled())
@@ -844,6 +835,7 @@ static bool IndexDatabase(FSearchIndexBase& SearchIndexBase, const UPoseSearchDa
 	// prepare indexers
 	TArray<FAssetIndexer> Indexers;
 	Indexers.Reserve(SearchIndexBase.Assets.Num());
+	FAssetSamplingContext SamplingContext(Database, BoneContainer);
 
 	int32 TotalPoses = 0;
 	for (int32 AssetIdx = 0; AssetIdx != SearchIndexBase.Assets.Num(); ++AssetIdx)
diff --git a/Engine/Plugins/Experimental/Animation/PoseSearch/Source/Runtime/Public/PoseSearch/PoseSearchAssetIndexer.h b/Engine/Plugins/Experimental/Animation/PoseSearch/Source/Runtime/Public/PoseSearch/PoseSearchAssetIndexer.h
--- a/Engine/Plugins/Experimental/Animation/PoseSearch/Source/Runtime/Public/PoseSearch/PoseSearchAssetIndexer.h
+++ b/Engine/Plugins/Experimental/Animation/PoseSearch/Source/Runtime/Public/PoseSearch/PoseSearchAssetIndexer.h
@@ -10,6 +10,8 @@
 #include "PoseSearch/PoseSearchFeatureChannel.h"
 #include "PoseSearch/PoseSearchSchema.h"
 
+class UPoseSearchDatabase;
+
 namespace UE::PoseSearch
 {
 
@@ -22,6 +24,9 @@ struct FAssetSamplingContext
 	// Time delta used for computing pose derivatives
 	static constexpr float FiniteDelta = 1 / 60.0f;
 
+	float BaseCostBias = 0.f;
+	float LoopingCostBias = 0.f;
+
 	// Mirror data table pointer copied from Schema for convenience
 	TObjectPtr<const UMirrorDataTable> MirrorDataTable;
 
@@ -32,7 +37,7 @@ struct FAssetSamplingContext
 	// Only initialized and used when a mirroring table is specified
 	TCustomBoneIndexArray<FQuat, FCompactPoseBoneIndex> ComponentSpaceRefRotations;
 
-	void Init(const UMirrorDataTable* InMirrorDataTable, const FBoneContainer& BoneContainer);
+	FAssetSamplingContext(const UPoseSearchDatabase& Database, const FBoneContainer& BoneContainer);
 	FTransform MirrorTransform(const FTransform& Transform) const;
 };
 
diff --git a/Engine/Plugins/Experimental/Animation/PoseSearch/Source/Runtime/Public/PoseSearch/PoseSearchDatabase.h b/Engine/Plugins/Experimental/Animation/PoseSearch/Source/Runtime/Public/PoseSearch/PoseSearchDatabase.h
--- a/Engine/Plugins/Experimental/Animation/PoseSearch/Source/Runtime/Public/PoseSearch/PoseSearchDatabase.h
+++ b/Engine/Plugins/Experimental/Animation/PoseSearch/Source/Runtime/Public/PoseSearch/PoseSearchDatabase.h
@@ -228,11 +228,36 @@ public:
 	UPROPERTY(EditAnywhere, BlueprintReadOnly, Category="Database", DisplayName="Config")
 	TObjectPtr<const UPoseSearchSchema> Schema;
 
+	// Cost added to the continuing pose from this database. This allows users to apply a cost bias (positive or negative) to the continuing pose.
+	// This is useful to help the system stay in one animation segment longer, or shorter depending on how you set this bias.
+	// Negative values make it more likely to be picked, or stayed in, positive values make it less likely to be picked or stay in.
+	// Note: excluded from DDC hash, since used only at runtime in SearchContinuingPose
+	UPROPERTY(EditAnywhere, Category = "Database", meta = (ExcludeFromHash)) 
+	float ContinuingPoseCostBias = -0.01f;
+
+	// Base Cost added or removed to all poses from this database. It can be overridden by Anim Notify: Pose Search Modify Cost at the frame level of animation data.
+	// Negative values make it more likely to be picked, or stayed in, Positive values make it less likely to be picked or stay in.
+	UPROPERTY(EditAnywhere, Category = "Database")
+	float BaseCostBias = 0.f;
+
+	// Cost added to all looping animation assets in this database. This allows users to make it more or less likely to pick the looping animation segments.
+	// Negative values make it more likely to be picked, or stayed in, Positive values make it less likely to be picked or stay in.
+	UPROPERTY(EditAnywhere, Category = "Database")
+	float LoopingCostBias = -0.005f;
+
 	// These settings allow users to trim the start and end of animations in the database to preserve start/end frames for blending, and prevent the system from selecting the very last frames before it blends out.
 	UPROPERTY(EditAnywhere, Category = "Database")
 	FPoseSearchExcludeFromDatabaseParameters ExcludeFromDatabaseParameters;
 
+	UPROPERTY(EditAnywhere, Category="Database")
+	TArray<FInstancedStruct> AnimationAssets;
+
 #if WITH_EDITORONLY_DATA
+	// This optional asset defines a list of databases you want to normalize together. Without it, it would be difficult to compare costs from separately normalized databases containing different types of animation,
+	// like only idles versus only runs animations, given that the range of movement would be dramatically different.
+	UPROPERTY(EditAnywhere, BlueprintReadOnly, Category = "Database")
+	TObjectPtr<const UPoseSearchNormalizationSet> NormalizationSet;
+
 	// Sequences and Blendspaces are deprecated and its data will be part of the AnimationAssets.
 	// All sequences and blend spaces will be added to the AnimationAssets in PostLoad().
 	UPROPERTY()
@@ -246,9 +271,6 @@ public:
 	TObjectPtr<USkeletalMesh> PreviewMesh = nullptr;
 #endif // WITH_EDITORONLY_DATA
 
-	UPROPERTY(EditAnywhere, Category="Database")
-	TArray<FInstancedStruct> AnimationAssets;
-
 	// This dictates how the database will perform the search.
 	UPROPERTY(EditAnywhere, Category = "Performance")
 	EPoseSearchMode PoseSearchMode = EPoseSearchMode::PCAKDTree;
@@ -286,17 +308,15 @@ public:
 	UPROPERTY(EditAnywhere, Category = "Performance", meta = (DisplayName = "KNNQueryNumNeighborsWithDuplicates", EditCondition = "PoseSearchMode == EPoseSearchMode::PCAKDTree && PCAValuesPruningSimilarityThreshold > 0", EditConditionHides, ClampMin = "0", ClampMax = "600", UIMin = "1"))
 	int32 KDTreeQueryNumNeighborsWithDuplicates = 0;
 	
-#if WITH_EDITORONLY_DATA
-	// This optional asset defines a list of databases you want to normalize together. Without it, it would be difficult to compare costs from separately normalized databases containing different types of animation,
-	// like only idles versus only runs animations, given that the range of movement would be dramatically different.
-	UPROPERTY(EditAnywhere, BlueprintReadOnly, Category = "Database")
-	TObjectPtr<const UPoseSearchNormalizationSet> NormalizationSet;
-#endif // WITH_EDITORONLY_DATA
-
 private:
 	// Do not use it directly. Use GetSearchIndex / SetSearchIndex interact with it and validate that is ok to do so.
 	UE::PoseSearch::FSearchIndex SearchIndexPrivate;
 
+#if WITH_EDITOR
+	DECLARE_MULTICAST_DELEGATE(FOnDerivedDataRebuildMulticaster);
+	FOnDerivedDataRebuildMulticaster OnDerivedDataRebuild;
+#endif // WITH_EDITOR
+
 public:
 	virtual ~UPoseSearchDatabase();
 
@@ -334,10 +354,6 @@ public:
 	virtual void BeginCacheForCookedPlatformData(const ITargetPlatform* TargetPlatform) override;
 	virtual bool IsCachedCookedPlatformDataLoaded(const ITargetPlatform* TargetPlatform) override;
 
-private:
-	DECLARE_MULTICAST_DELEGATE(FOnDerivedDataRebuildMulticaster);
-	FOnDerivedDataRebuildMulticaster OnDerivedDataRebuild;
-public:
 	typedef FOnDerivedDataRebuildMulticaster::FDelegate FOnDerivedDataRebuild;
 	void RegisterOnDerivedDataRebuild(const FOnDerivedDataRebuild& Delegate);
 	void UnregisterOnDerivedDataRebuild(void* Unregister);
diff --git a/Engine/Plugins/Experimental/Animation/PoseSearch/Source/Runtime/Public/PoseSearch/PoseSearchSchema.h b/Engine/Plugins/Experimental/Animation/PoseSearch/Source/Runtime/Public/PoseSearch/PoseSearchSchema.h
--- a/Engine/Plugins/Experimental/Animation/PoseSearch/Source/Runtime/Public/PoseSearch/PoseSearchSchema.h
+++ b/Engine/Plugins/Experimental/Animation/PoseSearch/Source/Runtime/Public/PoseSearch/PoseSearchSchema.h
@@ -71,22 +71,6 @@ public:
 	UPROPERTY(Transient)
 	TArray<uint16> BoneIndicesWithParents;
 
-	// Cost added to the continuing pose from databases that uses this config. This allows users to apply a cost bias (positive or negative) to the continuing pose.
-	// This is useful to help the system stay in one animation segment longer, or shorter depending on how you set this bias.
-	// Negative values make it more likely to be picked, or stayed in, positive values make it less likely to be picked or stay in.
-	UPROPERTY(EditAnywhere, Category = "Bias")
-	float ContinuingPoseCostBias = -0.01f;
-
-	// Base Cost added or removed to all poses from databases that use this config. It can be overridden by Anim Notify: Pose Search Modify Cost at the frame level of animation data.
-	// Negative values make it more likely to be picked, or stayed in, Positive values make it less likely to be picked or stay in.
-	UPROPERTY(EditAnywhere, Category = "Bias")
-	float BaseCostBias = 0.f;
-
-	// Cost added to all looping animation assets in a database that uses this config. This allows users to make it more or less likely to pick the looping animation segments.
-	// Negative values make it more likely to be picked, or stayed in, Positive values make it less likely to be picked or stay in.
-	UPROPERTY(EditAnywhere, Category = "Bias")
-	float LoopingCostBias = -0.005f;
-
 	// How many times the animation assets of the database using this schema will be indexed.
 	UPROPERTY(EditAnywhere, Category = "Permutations", meta = (ClampMin = "1"))
 	int32 NumberOfPermutations = 1;
