motion matching - adding customizable DebugColor to pose history node

diff --git a/Engine/Plugins/Experimental/Animation/PoseSearch/Source/Runtime/Private/AnimNode_PoseSearchHistoryCollector.cpp b/Engine/Plugins/Experimental/Animation/PoseSearch/Source/Runtime/Private/AnimNode_PoseSearchHistoryCollector.cpp
--- a/Engine/Plugins/Experimental/Animation/PoseSearch/Source/Runtime/Private/AnimNode_PoseSearchHistoryCollector.cpp
+++ b/Engine/Plugins/Experimental/Animation/PoseSearch/Source/Runtime/Private/AnimNode_PoseSearchHistoryCollector.cpp
@@ -133,7 +133,7 @@ void FAnimNode_PoseSearchHistoryCollector::Evaluate_AnyThread(FPoseContext& Outp
 #if ENABLE_DRAW_DEBUG && ENABLE_ANIM_DEBUG
 	if (CVarAnimPoseHistoryDebugDraw.GetValueOnAnyThread())
 	{
-		PoseHistory.DebugDraw(*Output.AnimInstanceProxy);
+		PoseHistory.DebugDraw(*Output.AnimInstanceProxy, DebugColor.ToFColor(true));
 	}
 #endif // ENABLE_DRAW_DEBUG && ENABLE_ANIM_DEBUG
 }
@@ -185,7 +185,7 @@ void FAnimNode_PoseSearchComponentSpaceHistoryCollector::EvaluateComponentSpace_
 #if ENABLE_DRAW_DEBUG && ENABLE_ANIM_DEBUG
 	if (CVarAnimPoseHistoryDebugDraw.GetValueOnAnyThread())
 	{
-		PoseHistory.DebugDraw(*Output.AnimInstanceProxy);
+		PoseHistory.DebugDraw(*Output.AnimInstanceProxy, DebugColor.ToFColor(true));
 	}
 #endif // ENABLE_DRAW_DEBUG && ENABLE_ANIM_DEBUG
 }
diff --git a/Engine/Plugins/Experimental/Animation/PoseSearch/Source/Runtime/Private/PoseSearchHistory.cpp b/Engine/Plugins/Experimental/Animation/PoseSearch/Source/Runtime/Private/PoseSearchHistory.cpp
--- a/Engine/Plugins/Experimental/Animation/PoseSearch/Source/Runtime/Private/PoseSearchHistory.cpp
+++ b/Engine/Plugins/Experimental/Animation/PoseSearch/Source/Runtime/Private/PoseSearchHistory.cpp
@@ -311,17 +311,8 @@ float FPoseHistory::GetSampleTimeInterval() const
 }
 
 #if ENABLE_DRAW_DEBUG && ENABLE_ANIM_DEBUG
-void FPoseHistory::DebugDraw(FAnimInstanceProxy& AnimInstanceProxy) const
+void FPoseHistory::DebugDraw(FAnimInstanceProxy& AnimInstanceProxy, FColor Color) const
 {
-	auto LerpColor = [](FColor A, FColor B, float T) -> FColor
-	{
-		return FColor(
-			FMath::RoundToInt(float(A.R) * (1.f - T) + float(B.R) * T),
-			FMath::RoundToInt(float(A.G) * (1.f - T) + float(B.G) * T),
-			FMath::RoundToInt(float(A.B) * (1.f - T) + float(B.B) * T),
-			FMath::RoundToInt(float(A.A) * (1.f - T) + float(B.A) * T));
-	};
-
 	TArray<FTransform> PrevGlobalTransforms;
 	for (int32 EntryIndex = 0; EntryIndex < Entries.Num(); ++EntryIndex)
 	{
@@ -340,9 +331,6 @@ void FPoseHistory::DebugDraw(FAnimInstanceProxy& AnimInstanceProxy) const
 		}
 		else
 		{
-			const float LerpFactor = float(EntryIndex - 1) / float(Entries.Num() - 1);
-			const FColor Color = LerpColor(FColorList::Red, FColorList::Orange, LerpFactor);
-
 			for (int32 i = 0; i < Entry.ComponentSpaceTransforms.Num(); ++i)
 			{
 				const FTransform GlobalTransforms = Entry.ComponentSpaceTransforms[i] * Entry.RootTransform;
@@ -426,19 +414,10 @@ void FExtendedPoseHistory::AddFuturePose(float SecondsInTheFuture, FCSPose<FComp
 }
 
 #if ENABLE_DRAW_DEBUG && ENABLE_ANIM_DEBUG
-void FExtendedPoseHistory::DebugDraw(FAnimInstanceProxy& AnimInstanceProxy) const
+void FExtendedPoseHistory::DebugDraw(FAnimInstanceProxy& AnimInstanceProxy, FColor PastColor, FColor FutureColor) const
 {
 	check(PoseHistory);
 
-	auto LerpColor = [](FColor A, FColor B, float T) -> FColor
-	{
-		return FColor(
-			FMath::RoundToInt(float(A.R) * (1.f - T) + float(B.R) * T),
-			FMath::RoundToInt(float(A.G) * (1.f - T) + float(B.G) * T),
-			FMath::RoundToInt(float(A.B) * (1.f - T) + float(B.B) * T),
-			FMath::RoundToInt(float(A.A) * (1.f - T) + float(B.A) * T));
-	};
-
 	TArray<FTransform> PrevGlobalTransforms;
 	for (int32 EntryIndex = 0; EntryIndex < FutureEntries.Num(); ++EntryIndex)
 	{
@@ -457,21 +436,18 @@ void FExtendedPoseHistory::DebugDraw(FAnimInstanceProxy& AnimInstanceProxy) cons
 		}
 		else
 		{
-			const float LerpFactor = float(EntryIndex - 1) / float(FutureEntries.Num() - 1);
-			const FColor Color = LerpColor(FColorList::Green, FColorList::Violet, LerpFactor);
-
 			for (int32 i = 0; i < Entry.ComponentSpaceTransforms.Num(); ++i)
 			{
 				const FTransform GlobalTransforms = Entry.ComponentSpaceTransforms[i] * Entry.RootTransform;
 
-				AnimInstanceProxy.AnimDrawDebugLine(PrevGlobalTransforms[i].GetTranslation(), GlobalTransforms.GetTranslation(), Color, false, 0.f, ESceneDepthPriorityGroup::SDPG_Foreground);
+				AnimInstanceProxy.AnimDrawDebugLine(PrevGlobalTransforms[i].GetTranslation(), GlobalTransforms.GetTranslation(), FutureColor, false, 0.f, ESceneDepthPriorityGroup::SDPG_Foreground);
 
 				PrevGlobalTransforms[i] = GlobalTransforms;
 			}
 		}
 	}
 
-	PoseHistory->DebugDraw(AnimInstanceProxy);
+	PoseHistory->DebugDraw(AnimInstanceProxy, PastColor);
 }
 #endif // ENABLE_DRAW_DEBUG && ENABLE_ANIM_DEBUG
 
diff --git a/Engine/Plugins/Experimental/Animation/PoseSearch/Source/Runtime/Private/PoseSearchLibrary.cpp b/Engine/Plugins/Experimental/Animation/PoseSearch/Source/Runtime/Private/PoseSearchLibrary.cpp
--- a/Engine/Plugins/Experimental/Animation/PoseSearch/Source/Runtime/Private/PoseSearchLibrary.cpp
+++ b/Engine/Plugins/Experimental/Animation/PoseSearch/Source/Runtime/Private/PoseSearchLibrary.cpp
@@ -459,6 +459,10 @@ void UPoseSearchLibrary::MotionMatch(
 	FMemMark Mark(FMemStack::Get());
 	if (Database && AnimInstance)
 	{
+#if ENABLE_DRAW_DEBUG && ENABLE_ANIM_DEBUG
+		FColor HistoryCollectorColor = FColor::Red;
+#endif // ENABLE_DRAW_DEBUG && ENABLE_ANIM_DEBUG
+
 		const FPoseSearchQueryTrajectory TrajectoryRootSpace = ProcessTrajectory(Trajectory, 0.f, 0.f, TrajectorySpeedMultiplier);
 
 		// ExtendedPoseHistory will hold future poses to match AssetSamplerBase (at FutureAnimationStartTime) TimeToFutureAnimationStart seconds in the future
@@ -470,6 +474,9 @@ void UPoseSearchLibrary::MotionMatch(
 				if (const FAnimNode_PoseSearchHistoryCollector_Base* PoseHistoryNode = TagSubsystem->FindNodeByTag<FAnimNode_PoseSearchHistoryCollector_Base>(PoseHistoryName, AnimInstance))
 				{
 					ExtendedPoseHistory.Init(&PoseHistoryNode->GetPoseHistory());
+#if ENABLE_DRAW_DEBUG && ENABLE_ANIM_DEBUG
+					HistoryCollectorColor = PoseHistoryNode->DebugColor.ToFColor(true);
+#endif // ENABLE_DRAW_DEBUG && ENABLE_ANIM_DEBUG
 				}
 			}
 		}
@@ -529,7 +536,7 @@ void UPoseSearchLibrary::MotionMatch(
 			{
 				if (FAnimInstanceProxy* AnimInstanceProxy = UAnimInstanceProxyProvider::GetAnimInstanceProxy(AnimInstance))
 				{
-					ExtendedPoseHistory.DebugDraw(*AnimInstanceProxy);
+					ExtendedPoseHistory.DebugDraw(*AnimInstanceProxy, HistoryCollectorColor, FColor::Green);
 				}
 			}
 #endif // ENABLE_DRAW_DEBUG && ENABLE_ANIM_DEBUG
diff --git a/Engine/Plugins/Experimental/Animation/PoseSearch/Source/Runtime/Public/PoseSearch/AnimNode_PoseSearchHistoryCollector.h b/Engine/Plugins/Experimental/Animation/PoseSearch/Source/Runtime/Public/PoseSearch/AnimNode_PoseSearchHistoryCollector.h
--- a/Engine/Plugins/Experimental/Animation/PoseSearch/Source/Runtime/Public/PoseSearch/AnimNode_PoseSearchHistoryCollector.h
+++ b/Engine/Plugins/Experimental/Animation/PoseSearch/Source/Runtime/Public/PoseSearch/AnimNode_PoseSearchHistoryCollector.h
@@ -32,6 +32,9 @@ public:
 	UPROPERTY(EditAnywhere, Category = Settings)
 	bool bResetOnBecomingRelevant = true;
 
+	UPROPERTY(EditAnywhere, Category = Settings)
+	FLinearColor DebugColor = FLinearColor::Red;
+
 	// Update Counter for detecting being relevant
 	FGraphTraversalCounter UpdateCounter;
 
diff --git a/Engine/Plugins/Experimental/Animation/PoseSearch/Source/Runtime/Public/PoseSearch/PoseSearchHistory.h b/Engine/Plugins/Experimental/Animation/PoseSearch/Source/Runtime/Public/PoseSearch/PoseSearchHistory.h
--- a/Engine/Plugins/Experimental/Animation/PoseSearch/Source/Runtime/Public/PoseSearch/PoseSearchHistory.h
+++ b/Engine/Plugins/Experimental/Animation/PoseSearch/Source/Runtime/Public/PoseSearch/PoseSearchHistory.h
@@ -73,7 +73,7 @@ struct FPoseHistory : public IPoseHistory
 	static bool LerpEntries(const FPoseHistoryEntry& PrevEntry, const FPoseHistoryEntry& NextEntry, float LerpValue, const USkeleton* BoneIndexSkeleton, const USkeleton* LastUpdateSkeleton, const FBoneToTransformMap& BoneToTransformMap, FBoneIndexType BoneIndexType, FBoneIndexType ReferenceBoneIndexType, FTransform& OutBoneTransform);
 
 #if ENABLE_DRAW_DEBUG && ENABLE_ANIM_DEBUG
-	void DebugDraw(FAnimInstanceProxy& AnimInstanceProxy) const;
+	void DebugDraw(FAnimInstanceProxy& AnimInstanceProxy, FColor Color) const;
 #endif
 
 private:
@@ -104,7 +104,7 @@ struct FExtendedPoseHistory : public IPoseHistory
 	void AddFuturePose(float SecondsInTheFuture, FCSPose<FCompactPose>& ComponentSpacePose, const FTransform& ComponentTransform);
 
 #if ENABLE_DRAW_DEBUG && ENABLE_ANIM_DEBUG
-	void DebugDraw(FAnimInstanceProxy& AnimInstanceProxy) const;
+	void DebugDraw(FAnimInstanceProxy& AnimInstanceProxy, FColor PastColor, FColor FutureColor) const;
 #endif
 
 private:
