motion matching - consolidating OriginBone / OriginTimeOffset / PermutationTimeType features between Position/Heading / Velocity channels

diff --git a/Engine/Plugins/Experimental/Animation/PoseSearch/Source/Runtime/Private/PoseSearchFeatureChannel_Heading.cpp b/Engine/Plugins/Experimental/Animation/PoseSearch/Source/Runtime/Private/PoseSearchFeatureChannel_Heading.cpp
--- a/Engine/Plugins/Experimental/Animation/PoseSearch/Source/Runtime/Private/PoseSearchFeatureChannel_Heading.cpp
+++ b/Engine/Plugins/Experimental/Animation/PoseSearch/Source/Runtime/Private/PoseSearchFeatureChannel_Heading.cpp
@@ -14,13 +14,13 @@ UPoseSearchFeatureChannel_Heading::UPoseSearchFeatureChannel_Heading()
 	bUseBlueprintQueryOverride = Cast<UBlueprintGeneratedClass>(GetClass()) != nullptr;
 }
 
-void UPoseSearchFeatureChannel_Heading::FindOrAddToSchema(UPoseSearchSchema* Schema, float SampleTimeOffset, const FName& BoneName, EHeadingAxis HeadingAxis)
+void UPoseSearchFeatureChannel_Heading::FindOrAddToSchema(UPoseSearchSchema* Schema, float SampleTimeOffset, const FName& BoneName, EHeadingAxis HeadingAxis, EPermutationTimeType PermutationTimeType)
 {
-	if (!Schema->FindChannel([SampleTimeOffset, &BoneName, HeadingAxis](const UPoseSearchFeatureChannel* Channel) -> const UPoseSearchFeatureChannel_Heading*
+	if (!Schema->FindChannel([SampleTimeOffset, &BoneName, HeadingAxis, PermutationTimeType](const UPoseSearchFeatureChannel* Channel) -> const UPoseSearchFeatureChannel_Heading*
 		{
 			if (const UPoseSearchFeatureChannel_Heading* Heading = Cast<UPoseSearchFeatureChannel_Heading>(Channel))
 			{
-				if (Heading->Bone.BoneName == BoneName && Heading->SampleTimeOffset == SampleTimeOffset && Heading->OriginTimeOffset == 0.f && Heading->HeadingAxis == HeadingAxis)
+				if (Heading->Bone.BoneName == BoneName && Heading->OriginBone.BoneName == NAME_None && Heading->SampleTimeOffset == SampleTimeOffset && Heading->OriginTimeOffset == 0.f && Heading->HeadingAxis == HeadingAxis && Heading->PermutationTimeType == PermutationTimeType)
 				{
 					return Heading;
 				}
@@ -35,6 +35,7 @@ void UPoseSearchFeatureChannel_Heading::FindOrAddToSchema(UPoseSearchSchema* Sch
 		Heading->HeadingAxis = HeadingAxis;
 		// @todo: perhaps add a tunable color for injected channels
 		Heading->DebugColor = FLinearColor::Gray;
+		Heading->PermutationTimeType = PermutationTimeType;
 		Schema->AddTemporaryChannel(Heading);
 	}
 }
@@ -44,7 +45,9 @@ void UPoseSearchFeatureChannel_Heading::Finalize(UPoseSearchSchema* Schema)
 	ChannelDataOffset = Schema->SchemaCardinality;
 	ChannelCardinality = UE::PoseSearch::FFeatureVectorHelper::GetVectorCardinality(ComponentStripping);
 	Schema->SchemaCardinality += ChannelCardinality;
+
 	SchemaBoneIdx = Schema->AddBoneReference(Bone);
+	SchemaOriginBoneIdx = Schema->AddBoneReference(OriginBone);
 }
 
 void UPoseSearchFeatureChannel_Heading::AddDependentChannels(UPoseSearchSchema* Schema) const
@@ -89,7 +92,7 @@ void UPoseSearchFeatureChannel_Heading::BuildQuery(UE::PoseSearch::FSearchContex
 	if (bUseBlueprintQueryOverride)
 	{
 		const FQuat BoneRotationWorld = BP_GetWorldRotation(SearchContext.GetAnimInstance());
-		const FQuat BoneRotation = SearchContext.GetSampleRotation(SampleTimeOffset, OriginTimeOffset, InOutQuery.GetSchema(), SchemaBoneIdx, RootSchemaBoneIdx, /*!bIsRootBone*/ true, EPermutationTimeType::UseSampleTime, &BoneRotationWorld);
+		const FQuat BoneRotation = SearchContext.GetSampleRotation(SampleTimeOffset, OriginTimeOffset, InOutQuery.GetSchema(), SchemaBoneIdx, SchemaOriginBoneIdx, /*!bIsRootBone*/ true, EPermutationTimeType::UseSampleTime, &BoneRotationWorld);
 		FFeatureVectorHelper::EncodeVector(InOutQuery.EditValues(), ChannelDataOffset, GetAxis(BoneRotation), ComponentStripping);
 	}
 	else
@@ -111,7 +114,7 @@ void UPoseSearchFeatureChannel_Heading::BuildQuery(UE::PoseSearch::FSearchContex
 		else
 		{
 			// calculating the BoneRotation in component space for the bone indexed by SchemaBoneIdx
-			const FQuat BoneRotation = SearchContext.GetSampleRotation(SampleTimeOffset, OriginTimeOffset, InOutQuery.GetSchema(), SchemaBoneIdx, RootSchemaBoneIdx, !bIsRootBone);
+			const FQuat BoneRotation = SearchContext.GetSampleRotation(SampleTimeOffset, OriginTimeOffset, InOutQuery.GetSchema(), SchemaBoneIdx, SchemaOriginBoneIdx, !bIsRootBone);
 			FFeatureVectorHelper::EncodeVector(InOutQuery.EditValues(), ChannelDataOffset, GetAxis(BoneRotation), ComponentStripping);
 		}
 	}
@@ -147,7 +150,7 @@ bool UPoseSearchFeatureChannel_Heading::IndexAsset(UE::PoseSearch::FAssetIndexer
 	FQuat SampleRotation = FQuat::Identity;
 	for (int32 SampleIdx = Indexer.GetBeginSampleIdx(); SampleIdx != Indexer.GetEndSampleIdx(); ++SampleIdx)
 	{
-		if (Indexer.GetSampleRotation(SampleRotation, SampleTimeOffset, OriginTimeOffset, SampleIdx, SchemaBoneIdx, RootSchemaBoneIdx, EPermutationTimeType::UseSampleTime, SamplingAttributeId))
+		if (Indexer.GetSampleRotation(SampleRotation, SampleTimeOffset, OriginTimeOffset, SampleIdx, SchemaBoneIdx, SchemaOriginBoneIdx, EPermutationTimeType::UseSampleTime, SamplingAttributeId))
 		{
 			FFeatureVectorHelper::EncodeVector(Indexer.GetPoseVector(SampleIdx), ChannelDataOffset, GetAxis(SampleRotation), ComponentStripping);
 		}
@@ -199,6 +202,12 @@ FString UPoseSearchFeatureChannel_Heading::GetLabel() const
 		Label.Append(Schema->BoneReferences[SchemaBoneIdx].BoneName.ToString());
 	}
 
+	if (SchemaOriginBoneIdx != RootSchemaBoneIdx)
+	{
+		Label.Append(TEXT("_"));
+		Label.Append(Schema->BoneReferences[SchemaOriginBoneIdx].BoneName.ToString());
+	}
+
 	Label.Appendf(TEXT(" %.2f"), SampleTimeOffset);
 
 	if (!FMath::IsNearlyZero(OriginTimeOffset))
diff --git a/Engine/Plugins/Experimental/Animation/PoseSearch/Source/Runtime/Private/PoseSearchFeatureChannel_Heading.h b/Engine/Plugins/Experimental/Animation/PoseSearch/Source/Runtime/Private/PoseSearchFeatureChannel_Heading.h
--- a/Engine/Plugins/Experimental/Animation/PoseSearch/Source/Runtime/Private/PoseSearchFeatureChannel_Heading.h
+++ b/Engine/Plugins/Experimental/Animation/PoseSearch/Source/Runtime/Private/PoseSearchFeatureChannel_Heading.h
@@ -26,15 +26,19 @@ public:
 	UPROPERTY(EditAnywhere, Category = "Settings")
 	FBoneReference Bone;
 
+	UPROPERTY(EditAnywhere, Category = "Settings")
+	FBoneReference OriginBone;
+
 	UPROPERTY(EditAnywhere, Category = "Settings")
 	float Weight = 1.f;
 
+
 	// if SamplingAttributeId >= 0, ALL the animations contained in the pose search database referencing the schema containing this channel are expected to have 
 	// UAnimNotifyState_PoseSearchSamplingAttribute notify state with a matching SamplingAttributeId, and the UAnimNotifyState_PoseSearchSamplingAttribute properties
-	// will be used as source of data instead of this channel "Bone".
+	// will be used as source of data instead of this channel "Bone". UAnimNotifyState_PoseSearchSamplingAttribute properties will be then converted into OriginBone space
 	UPROPERTY(EditAnywhere, Category = "Settings")
 	int32 SamplingAttributeId = INDEX_NONE;
-
+	
 	// the data relative to the sampling time associated to this channel will be offsetted by SampleTimeOffset seconds.
 	// For example, if Bone is the head bone, and SampleTimeOffset is 0.5, this channel will try to match the future heading of the character head bone 0.5 seconds ahead
 	UPROPERTY(EditAnywhere, Category = "Settings")
@@ -53,6 +57,9 @@ public:
 	UPROPERTY(Transient)
 	int8 SchemaBoneIdx = 0;
 
+	UPROPERTY(Transient)
+	int8 SchemaOriginBoneIdx = 0;
+
 	UPROPERTY(EditAnywhere, Category = "Settings", meta = (ExcludeFromHash, DisplayPriority = 0))
 	FLinearColor DebugColor = FLinearColor::White;
 
@@ -62,6 +69,9 @@ public:
 	UPROPERTY(EditAnywhere, Category = "Settings")
 	EComponentStrippingVector ComponentStripping = EComponentStrippingVector::None;
 
+	UPROPERTY(EditAnywhere, Category = "Settings")
+	EPermutationTimeType PermutationTimeType = EPermutationTimeType::UseSampleTime;
+
 	UFUNCTION(BlueprintPure, BlueprintImplementableEvent, meta=(BlueprintThreadSafe, DisplayName = "Get World Rotation"), Category = "Settings")
 	FQuat BP_GetWorldRotation(const UAnimInstance* AnimInstance) const;
 
@@ -73,6 +83,7 @@ public:
 	virtual void Finalize(UPoseSearchSchema* Schema) override;
 	virtual void BuildQuery(UE::PoseSearch::FSearchContext& SearchContext, UE::PoseSearch::FFeatureVectorBuilder& InOutQuery) const override;
 
+	virtual EPermutationTimeType GetPermutationTimeType() const override { return PermutationTimeType; }
 	virtual void AddDependentChannels(UPoseSearchSchema* Schema) const override;
 
 #if ENABLE_DRAW_DEBUG
@@ -87,5 +98,5 @@ public:
 
 	FVector GetAxis(const FQuat& Rotation) const;
 
-	static void FindOrAddToSchema(UPoseSearchSchema* Schema, float SampleTimeOffset, const FName& BoneName = NAME_None, EHeadingAxis HeadingAxis = EHeadingAxis::X);
+	static void FindOrAddToSchema(UPoseSearchSchema* Schema, float SampleTimeOffset, const FName& BoneName = NAME_None, EHeadingAxis HeadingAxis = EHeadingAxis::X, EPermutationTimeType PermutationTimeType = EPermutationTimeType::UseSampleTime);
 };
\ No newline at end of file
diff --git a/Engine/Plugins/Experimental/Animation/PoseSearch/Source/Runtime/Private/PoseSearchFeatureChannel_Position.cpp b/Engine/Plugins/Experimental/Animation/PoseSearch/Source/Runtime/Private/PoseSearchFeatureChannel_Position.cpp
--- a/Engine/Plugins/Experimental/Animation/PoseSearch/Source/Runtime/Private/PoseSearchFeatureChannel_Position.cpp
+++ b/Engine/Plugins/Experimental/Animation/PoseSearch/Source/Runtime/Private/PoseSearchFeatureChannel_Position.cpp
@@ -19,7 +19,7 @@ void UPoseSearchFeatureChannel_Position::FindOrAddToSchema(UPoseSearchSchema* Sc
 		{
 			if (const UPoseSearchFeatureChannel_Position* Position = Cast<UPoseSearchFeatureChannel_Position>(Channel))
 			{
-				if (Position->Bone.BoneName == BoneName && Position->OriginBone.BoneName == NAME_None && Position->SampleTimeOffset == SampleTimeOffset && Position->PermutationTimeType == PermutationTimeType)
+				if (Position->Bone.BoneName == BoneName && Position->OriginBone.BoneName == NAME_None && Position->SampleTimeOffset == SampleTimeOffset && Position->OriginTimeOffset == 0.f && Position->PermutationTimeType == PermutationTimeType)
 				{
 					return Position;
 				}
@@ -69,7 +69,7 @@ void UPoseSearchFeatureChannel_Position::BuildQuery(UE::PoseSearch::FSearchConte
 	if (bUseBlueprintQueryOverride)
 	{
 		const FVector BonePositionWorld = BP_GetWorldPosition(SearchContext.GetAnimInstance());
-		const FVector BonePosition = SearchContext.GetSamplePosition(SampleTimeOffset, 0.f, InOutQuery.GetSchema(), SchemaBoneIdx, SchemaOriginBoneIdx, /*!bIsRootBone*/ true, PermutationTimeType, &BonePositionWorld);
+		const FVector BonePosition = SearchContext.GetSamplePosition(SampleTimeOffset, OriginTimeOffset, InOutQuery.GetSchema(), SchemaBoneIdx, SchemaOriginBoneIdx, /*!bIsRootBone*/ true, PermutationTimeType, &BonePositionWorld);
   		FFeatureVectorHelper::EncodeVector(InOutQuery.EditValues(), ChannelDataOffset, BonePosition, ComponentStripping);
 	}
 	else
@@ -91,7 +91,7 @@ void UPoseSearchFeatureChannel_Position::BuildQuery(UE::PoseSearch::FSearchConte
 		else
 		{
 			// calculating the BonePosition in root bone space for the bone indexed by SchemaBoneIdx
-			const FVector BonePosition = SearchContext.GetSamplePosition(SampleTimeOffset, 0.f, InOutQuery.GetSchema(), SchemaBoneIdx, SchemaOriginBoneIdx, !bIsRootBone, PermutationTimeType);
+			const FVector BonePosition = SearchContext.GetSamplePosition(SampleTimeOffset, OriginTimeOffset, InOutQuery.GetSchema(), SchemaBoneIdx, SchemaOriginBoneIdx, !bIsRootBone, PermutationTimeType);
 			FFeatureVectorHelper::EncodeVector(InOutQuery.EditValues(), ChannelDataOffset, BonePosition, ComponentStripping);
 		}
 	}
@@ -139,7 +139,7 @@ bool UPoseSearchFeatureChannel_Position::IndexAsset(UE::PoseSearch::FAssetIndexe
 	FVector BonePosition;
 	for (int32 SampleIdx = Indexer.GetBeginSampleIdx(); SampleIdx != Indexer.GetEndSampleIdx(); ++SampleIdx)
 	{
-		if (Indexer.GetSamplePosition(BonePosition, SampleTimeOffset, 0.f, SampleIdx, SchemaBoneIdx, SchemaOriginBoneIdx, PermutationTimeType, SamplingAttributeId))
+		if (Indexer.GetSamplePosition(BonePosition, SampleTimeOffset, OriginTimeOffset, SampleIdx, SchemaBoneIdx, SchemaOriginBoneIdx, PermutationTimeType, SamplingAttributeId))
 		{
 			FFeatureVectorHelper::EncodeVector(Indexer.GetPoseVector(SampleIdx), ChannelDataOffset, BonePosition, ComponentStripping);
 		}
@@ -186,6 +186,12 @@ FString UPoseSearchFeatureChannel_Position::GetLabel() const
 	}
 
 	Label.Appendf(TEXT(" %.2f"), SampleTimeOffset);
+
+	if (!FMath::IsNearlyZero(OriginTimeOffset))
+	{
+		Label.Appendf(TEXT("-%.2f"), OriginTimeOffset);
+	}
+
 	return Label.ToString();
 }
 #endif
\ No newline at end of file
diff --git a/Engine/Plugins/Experimental/Animation/PoseSearch/Source/Runtime/Private/PoseSearchFeatureChannel_Position.h b/Engine/Plugins/Experimental/Animation/PoseSearch/Source/Runtime/Private/PoseSearchFeatureChannel_Position.h
--- a/Engine/Plugins/Experimental/Animation/PoseSearch/Source/Runtime/Private/PoseSearchFeatureChannel_Position.h
+++ b/Engine/Plugins/Experimental/Animation/PoseSearch/Source/Runtime/Private/PoseSearchFeatureChannel_Position.h
@@ -33,6 +33,12 @@ public:
 	UPROPERTY(EditAnywhere, Category = "Settings")
 	float SampleTimeOffset = 0.f;
 
+	// the data relative to the sampling time associated to this channel origin (root / trajectory bone) will be offsetted by OriginTimeOffset seconds.
+	// For example, if Bone is the head bone, SampleTimeOffset is 0.5, and OriginTimeOffset is 0.5, this channel will try to match 
+	// the future position of the character head bone 0.5 seconds ahead, relative to the future root bone 0.5 seconds ahead
+	UPROPERTY(EditAnywhere, Category = "Settings")
+	float OriginTimeOffset = 0.f;
+
 	// index referencing the associated bone in UPoseSearchSchema::BoneReferences
 	UPROPERTY(Transient)
 	int8 SchemaBoneIdx = 0;
diff --git a/Engine/Plugins/Experimental/Animation/PoseSearch/Source/Runtime/Private/PoseSearchFeatureChannel_Velocity.cpp b/Engine/Plugins/Experimental/Animation/PoseSearch/Source/Runtime/Private/PoseSearchFeatureChannel_Velocity.cpp
--- a/Engine/Plugins/Experimental/Animation/PoseSearch/Source/Runtime/Private/PoseSearchFeatureChannel_Velocity.cpp
+++ b/Engine/Plugins/Experimental/Animation/PoseSearch/Source/Runtime/Private/PoseSearchFeatureChannel_Velocity.cpp
@@ -22,6 +22,7 @@ void UPoseSearchFeatureChannel_Velocity::Finalize(UPoseSearchSchema* Schema)
 	Schema->SchemaCardinality += ChannelCardinality;
 
 	SchemaBoneIdx = Schema->AddBoneReference(Bone);
+	SchemaOriginBoneIdx = Schema->AddBoneReference(OriginBone);
 }
 
 void UPoseSearchFeatureChannel_Velocity::AddDependentChannels(UPoseSearchSchema* Schema) const
@@ -42,7 +43,7 @@ void UPoseSearchFeatureChannel_Velocity::BuildQuery(UE::PoseSearch::FSearchConte
 	{
 		const FVector LinearVelocityWorld = BP_GetWorldVelocity(SearchContext.GetAnimInstance());
 
-		FVector LinearVelocity = SearchContext.GetSampleVelocity(SampleTimeOffset, 0.f, InOutQuery.GetSchema(), SchemaBoneIdx, RootSchemaBoneIdx, bUseCharacterSpaceVelocities, /*!bIsRootBone*/ true, EPermutationTimeType::UseSampleTime, &LinearVelocityWorld);
+		FVector LinearVelocity = SearchContext.GetSampleVelocity(SampleTimeOffset, OriginTimeOffset, InOutQuery.GetSchema(), SchemaBoneIdx, SchemaOriginBoneIdx, bUseCharacterSpaceVelocities, /*!bIsRootBone*/ true, EPermutationTimeType::UseSampleTime, &LinearVelocityWorld);
 		if (bNormalize)
 		{
 			LinearVelocity = LinearVelocity.GetClampedToMaxSize(1.f);
@@ -68,7 +69,7 @@ void UPoseSearchFeatureChannel_Velocity::BuildQuery(UE::PoseSearch::FSearchConte
 		else
 		{
 			// calculating the LinearVelocity for the bone indexed by SchemaBoneIdx
-			FVector LinearVelocity = SearchContext.GetSampleVelocity(SampleTimeOffset, 0.f, InOutQuery.GetSchema(), SchemaBoneIdx, RootSchemaBoneIdx, bUseCharacterSpaceVelocities, !bIsRootBone, EPermutationTimeType::UseSampleTime);
+			FVector LinearVelocity = SearchContext.GetSampleVelocity(SampleTimeOffset, OriginTimeOffset, InOutQuery.GetSchema(), SchemaBoneIdx, SchemaOriginBoneIdx, bUseCharacterSpaceVelocities, !bIsRootBone, EPermutationTimeType::UseSampleTime);
 			if (bNormalize)
 			{
 				LinearVelocity = LinearVelocity.GetClampedToMaxSize(1.f);
@@ -111,7 +112,7 @@ bool UPoseSearchFeatureChannel_Velocity::IndexAsset(UE::PoseSearch::FAssetIndexe
 	FVector LinearVelocity;
 	for (int32 SampleIdx = Indexer.GetBeginSampleIdx(); SampleIdx != Indexer.GetEndSampleIdx(); ++SampleIdx)
 	{
-		if (Indexer.GetSampleVelocity(LinearVelocity, SampleTimeOffset, 0.f, SampleIdx, SchemaBoneIdx, RootSchemaBoneIdx, bUseCharacterSpaceVelocities, EPermutationTimeType::UseSampleTime, SamplingAttributeId))
+		if (Indexer.GetSampleVelocity(LinearVelocity, SampleTimeOffset, OriginTimeOffset, SampleIdx, SchemaBoneIdx, SchemaOriginBoneIdx, bUseCharacterSpaceVelocities, EPermutationTimeType::UseSampleTime, SamplingAttributeId))
 		{
 			if (bNormalize)
 			{
@@ -159,7 +160,19 @@ FString UPoseSearchFeatureChannel_Velocity::GetLabel() const
 		Label.Append(Schema->BoneReferences[SchemaBoneIdx].BoneName.ToString());
 	}
 
+	if (SchemaOriginBoneIdx != RootSchemaBoneIdx)
+	{
+		Label.Append(TEXT("_"));
+		Label.Append(Schema->BoneReferences[SchemaOriginBoneIdx].BoneName.ToString());
+	}
+
 	Label.Appendf(TEXT(" %.2f"), SampleTimeOffset);
+
+	if (!FMath::IsNearlyZero(OriginTimeOffset))
+	{
+		Label.Appendf(TEXT("-%.2f"), OriginTimeOffset);
+	}
+
 	return Label.ToString();
 }
 #endif
\ No newline at end of file
diff --git a/Engine/Plugins/Experimental/Animation/PoseSearch/Source/Runtime/Private/PoseSearchFeatureChannel_Velocity.h b/Engine/Plugins/Experimental/Animation/PoseSearch/Source/Runtime/Private/PoseSearchFeatureChannel_Velocity.h
--- a/Engine/Plugins/Experimental/Animation/PoseSearch/Source/Runtime/Private/PoseSearchFeatureChannel_Velocity.h
+++ b/Engine/Plugins/Experimental/Animation/PoseSearch/Source/Runtime/Private/PoseSearchFeatureChannel_Velocity.h
@@ -15,24 +15,37 @@ public:
 	UPROPERTY(EditAnywhere, Category = "Settings")
 	FBoneReference Bone;
 
+	UPROPERTY(EditAnywhere, Category = "Settings")
+	FBoneReference OriginBone;
+
 	UPROPERTY(EditAnywhere, Category = "Settings")
 	float Weight = 1.f;
 
+
 	// if SamplingAttributeId >= 0, ALL the animations contained in the pose search database referencing the schema containing this channel are expected to have 
 	// UAnimNotifyState_PoseSearchSamplingAttribute notify state with a matching SamplingAttributeId, and the UAnimNotifyState_PoseSearchSamplingAttribute properties
-	// will be used as source of data instead of this channel "Bone".
+	// will be used as source of data instead of this channel "Bone". UAnimNotifyState_PoseSearchSamplingAttribute properties will be then converted into OriginBone space
 	UPROPERTY(EditAnywhere, Category = "Settings")
 	int32 SamplingAttributeId = INDEX_NONE;
-
+	
 	// the data relative to the sampling time associated to this channel will be offsetted by SampleTimeOffset seconds.
 	// For example, if Bone is the head bone, and SampleTimeOffset is 0.5, this channel will try to match the future velocity of the character head bone 0.5 seconds ahead
 	UPROPERTY(EditAnywhere, Category = "Settings")
 	float SampleTimeOffset = 0.f;
 
+	// the data relative to the sampling time associated to this channel origin (root / trajectory bone) will be offsetted by OriginTimeOffset seconds.
+	// For example, if Bone is the head bone, SampleTimeOffset is 0.5, and OriginTimeOffset is 0.5, this channel will try to match 
+	// the future velocity of the character head bone 0.5 seconds ahead, relative to the future root bone 0.5 seconds ahead
+	UPROPERTY(EditAnywhere, Category = "Settings")
+	float OriginTimeOffset = 0.f;
+
 	// index referencing the associated bone in UPoseSearchSchema::BoneReferences
 	UPROPERTY(Transient)
 	int8 SchemaBoneIdx = 0;
 
+	UPROPERTY(Transient)
+	int8 SchemaOriginBoneIdx = 0;
+
 	UPROPERTY(EditAnywhere, Category = "Settings", meta = (ExcludeFromHash, DisplayPriority = 0))
 	FLinearColor DebugColor = FLinearColor::Green;
 
@@ -50,6 +63,9 @@ public:
 	UPROPERTY(EditAnywhere, Category = "Settings")
 	EComponentStrippingVector ComponentStripping = EComponentStrippingVector::None;
 
+	UPROPERTY(EditAnywhere, Category = "Settings")
+	EPermutationTimeType PermutationTimeType = EPermutationTimeType::UseSampleTime;
+
 	UFUNCTION(BlueprintPure, BlueprintImplementableEvent, meta=(BlueprintThreadSafe, DisplayName = "Get World Velocity"), Category = "Settings")
 	FVector BP_GetWorldVelocity(const UAnimInstance* AnimInstance) const;
 
@@ -61,6 +77,7 @@ public:
 	virtual void Finalize(UPoseSearchSchema* Schema) override;
 	virtual void BuildQuery(UE::PoseSearch::FSearchContext& SearchContext, UE::PoseSearch::FFeatureVectorBuilder& InOutQuery) const override;
 
+	virtual EPermutationTimeType GetPermutationTimeType() const override { return PermutationTimeType; }
 	virtual void AddDependentChannels(UPoseSearchSchema* Schema) const override;
 
 #if ENABLE_DRAW_DEBUG
