motion matching - removing unnecessary BlendStackNode from FAnimNode_MotionMatching

diff --git a/Engine/Plugins/Experimental/Animation/PoseSearch/Source/Runtime/Public/PoseSearch/AnimNode_MotionMatching.h b/Engine/Plugins/Experimental/Animation/PoseSearch/Source/Runtime/Public/PoseSearch/AnimNode_MotionMatching.h
--- a/Engine/Plugins/Experimental/Animation/PoseSearch/Source/Runtime/Public/PoseSearch/AnimNode_MotionMatching.h
+++ b/Engine/Plugins/Experimental/Animation/PoseSearch/Source/Runtime/Public/PoseSearch/AnimNode_MotionMatching.h
@@ -110,9 +110,6 @@ private:
 	UPROPERTY(EditAnywhere, Category = RootMotion, meta = (PinHiddenByDefault))
 	float YawFromAnimationBlendRate = -1.f;
 
-	UPROPERTY()
-	FAnimNode_BlendStack_Standalone BlendStackNode;
-
 	// Encapsulated motion matching algorithm and internal state
 	FMotionMatchingState MotionMatchingState;
 
