motion matching - fix for bDisableReselection, introduced CVarMotionMatchValidateKNNSearch to validate KDTree search consistency

diff --git a/Engine/Plugins/Experimental/Animation/PoseSearch/Source/Runtime/Private/PoseSearchDatabase.cpp b/Engine/Plugins/Experimental/Animation/PoseSearch/Source/Runtime/Private/PoseSearchDatabase.cpp
--- a/Engine/Plugins/Experimental/Animation/PoseSearch/Source/Runtime/Private/PoseSearchDatabase.cpp
+++ b/Engine/Plugins/Experimental/Animation/PoseSearch/Source/Runtime/Private/PoseSearchDatabase.cpp
@@ -32,6 +32,7 @@ namespace UE::PoseSearch
 {
 #if WITH_EDITOR && ENABLE_ANIM_DEBUG
 static TAutoConsoleVariable<bool> CVarMotionMatchCompareAgainstBruteForce(TEXT("a.MotionMatch.CompareAgainstBruteForce"), false, TEXT("Compare optimized search against brute force search"));
+static TAutoConsoleVariable<bool> CVarMotionMatchValidateKNNSearch(TEXT("a.MotionMatch.ValidateKNNSearch"), false, TEXT("Validate KNN search"));
 #endif
 
 typedef TArray<size_t, TInlineAllocator<256>> FNonSelectableIdx;
@@ -51,7 +52,7 @@ static void PopulateNonSelectableIdx(FNonSelectableIdx& NonSelectableIdx, FSearc
 		{
 			const FPoseSearchDatabaseAnimationAssetBase* DatabaseAnimationAssetBase = Database->GetAnimationAssetStruct(CurrentIndexAsset->SourceAssetIdx).GetPtr<FPoseSearchDatabaseAnimationAssetBase>();
 			check(DatabaseAnimationAssetBase);
-			if (!DatabaseAnimationAssetBase->bDisableReselection)
+			if (DatabaseAnimationAssetBase->bDisableReselection)
 			{
 				// excluding all the poses from DatabaseAnimationAssetBase
 				// @todo: optimize this code!
@@ -870,15 +871,66 @@ UE::PoseSearch::FSearchResult UPoseSearchDatabase::SearchPCAKDTree(UE::PoseSearc
 #endif // UE_POSE_SEARCH_TRACE_ENABLED
 		);
 
-		// if bArePCAValuesPruned we filter out the NonSelectableIdx after kdtree search, otherwise during kdtree search
-		FKDTree::FKNNResultSet ResultSet(ClampedKDTreeQueryNumNeighbors, ResultIndexes, ResultDistanceSqr, bArePCAValuesPruned ? TConstArrayView<size_t>() : NonSelectableIdx);
+		bool bRunNonSelectableIdxPostKDTree = bArePCAValuesPruned;
+#if WITH_EDITOR && ENABLE_ANIM_DEBUG
+		const bool bValidateKNNSearch = CVarMotionMatchValidateKNNSearch.GetValueOnAnyThread();
+		bRunNonSelectableIdxPostKDTree |= bValidateKNNSearch;
+#endif // WITH_EDITOR && ENABLE_ANIM_DEBUG
+
+		// if bRunNonSelectableIdxPostKDTree we filter out the NonSelectableIdx after kdtree search, otherwise during kdtree search
+		FKDTree::FKNNResultSet ResultSet(ClampedKDTreeQueryNumNeighbors, ResultIndexes, ResultDistanceSqr, bRunNonSelectableIdxPostKDTree ? TConstArrayView<size_t>() : NonSelectableIdx);
 
 		check(QueryValues.Num() == NumDimensions);
-		// projecting QueryValues into the PCA space ProjectedQueryValues and query the KDTree
-		SearchIndex.KDTree.FindNeighbors(ResultSet, SearchIndex.PCAProject(QueryValues, ProjectedQueryValues));
+		// projecting QueryValues into the PCA space 
+		TConstArrayView<float> PCAQueryValues = SearchIndex.PCAProject(QueryValues, ProjectedQueryValues);
+		check(PCAQueryValues.Num() == ClampedNumberOfPrincipalComponents);
+
+		// Querying the KDTree with PCAQueryValues, the projected in PCA space QueryValues
+		SearchIndex.KDTree.FindNeighbors(ResultSet, PCAQueryValues);
+
+#if WITH_EDITOR && ENABLE_ANIM_DEBUG
+		if (bValidateKNNSearch)
+		{
+			const int32 NumPCAValuesVectors = SearchIndex.GetNumPCAValuesVectors(ClampedNumberOfPrincipalComponents);
+
+			TArray<TPair<int32, float>> PCAValueIndexCost;
+			PCAValueIndexCost.SetNumUninitialized(NumPCAValuesVectors);
+
+			// validating that the best n "ClampedKDTreeQueryNumNeighbors" are actually the best candidates
+			for (int32 PCAValueIndex = 0; PCAValueIndex < NumPCAValuesVectors; ++PCAValueIndex)
+			{
+				PCAValueIndexCost[PCAValueIndex].Key = PCAValueIndex;
+				PCAValueIndexCost[PCAValueIndex].Value = CompareFeatureVectors(SearchIndex.GetPCAPoseValues(PCAValueIndex), PCAQueryValues);
+			}
+
+			PCAValueIndexCost.Sort([](const TPair<int32, float>& A, const TPair<int32, float>& B)
+				{
+					return A.Value < B.Value;
+				});
+
+			for (size_t ResultIndex = 0; ResultIndex < ResultSet.Num(); ++ResultIndex)
+			{
+				if (PCAValueIndexCost[ResultIndex].Key != ResultIndexes[ResultIndex])
+				{
+					if (!FMath::IsNearlyEqual(PCAValueIndexCost[ResultIndex].Value, ResultDistanceSqr[ResultIndex], UE_KINDA_SMALL_NUMBER))
+					{
+						UE_LOG(LogPoseSearch, Error, TEXT("UPoseSearchDatabase::SearchPCAKDTree - KDTree search order is inconsistent with exaustive search in PCA space"));
+					}
+					else
+					{
+						UE_LOG(LogPoseSearch, Log, TEXT("UPoseSearchDatabase::SearchPCAKDTree - found two points at the same distance from the query in different order between KDTree and exaustive search"));
+					}
+				}
+				else if (!FMath::IsNearlyEqual(PCAValueIndexCost[ResultIndex].Value, ResultDistanceSqr[ResultIndex], UE_KINDA_SMALL_NUMBER))
+				{
+					UE_LOG(LogPoseSearch, Error, TEXT("UPoseSearchDatabase::SearchPCAKDTree - KDTree search cost is inconsistent with exaustive search in PCA space"));
+				}
+			}
+		}
+#endif // WITH_EDITOR && ENABLE_ANIM_DEBUG
 
 		// NonSelectableIdx are already filtered out inside the kdtree search
-		const FSearchFilters SearchFilters(Schema, bArePCAValuesPruned ? NonSelectableIdx : TConstArrayView<size_t>(), SearchIndex.bAnyBlockTransition);
+		const FSearchFilters SearchFilters(Schema, bRunNonSelectableIdxPostKDTree ? NonSelectableIdx : TConstArrayView<size_t>(), SearchIndex.bAnyBlockTransition);
 		
 		// are the PCAValues pruned out of duplicates (multiple poses are associated with the same PCAValuesVectorIdx)
 		if (bArePCAValuesPruned)
