motion matching - fix for pose search debugger "ShowAllPoses": avoid listing poses from non searched databases aside continuing pose search

diff --git a/Engine/Plugins/Experimental/Animation/PoseSearch/Source/Editor/Private/PoseSearchDebuggerDatabaseView.cpp b/Engine/Plugins/Experimental/Animation/PoseSearch/Source/Editor/Private/PoseSearchDebuggerDatabaseView.cpp
--- a/Engine/Plugins/Experimental/Animation/PoseSearch/Source/Editor/Private/PoseSearchDebuggerDatabaseView.cpp
+++ b/Engine/Plugins/Experimental/Animation/PoseSearch/Source/Editor/Private/PoseSearchDebuggerDatabaseView.cpp
@@ -285,17 +285,27 @@ void SDebuggerDatabaseView::Update(const FTraceMotionMatchingStateMessage& State
 
 			if (bShowAllPoses)
 			{
+				bool bAddMissingPoseEntries = false;
+
+				// excluding the already added PoseEntry
 				TSet<int32> PoseEntriesIdx;
 				for (const FTraceMotionMatchingStatePoseEntry& PoseEntry : DbEntry.PoseEntries)
 				{
 					PoseEntriesIdx.Add(PoseEntry.DbPoseIdx);
+
+					const bool bIsContinuingPose = EnumHasAnyFlags(PoseEntry.PoseCandidateFlags, EPoseCandidateFlags::Valid_ContinuingPose);
+					bAddMissingPoseEntries |= !bIsContinuingPose;
 				}
 
-				for (int32 DbPoseIdx = 0; DbPoseIdx < SearchIndex.GetNumPoses(); ++DbPoseIdx)
+				// adding the missing ones
+				if (bAddMissingPoseEntries)
 				{
-					if (!PoseEntriesIdx.Find(DbPoseIdx))
+					for (int32 DbPoseIdx = 0; DbPoseIdx < SearchIndex.GetNumPoses(); ++DbPoseIdx)
 					{
-						AddUnfilteredDatabaseRow(Database, UnfilteredDatabaseRows, SharedData, DbPoseIdx, EPoseCandidateFlags::DiscardedBy_Search, PoseToPCAValuesVectorIndexes);
+						if (!PoseEntriesIdx.Find(DbPoseIdx))
+						{
+							AddUnfilteredDatabaseRow(Database, UnfilteredDatabaseRows, SharedData, DbPoseIdx, EPoseCandidateFlags::DiscardedBy_Search, PoseToPCAValuesVectorIndexes);
+						}
 					}
 				}
 			}
@@ -510,14 +520,7 @@ void SDebuggerDatabaseView::OnFilterTextChanged(const FText& SearchText)
 
 void SDebuggerDatabaseView::OnShowAllPosesCheckboxChanged(ECheckBoxState State)
 {
-	if (State == ECheckBoxState::Checked)
-	{
-		bShowAllPoses = true;
-	}
-	else
-	{
-		bShowAllPoses = false;
-	}
+	bShowAllPoses = (State == ECheckBoxState::Checked);
 
 	if (TSharedPtr<SDebuggerView> DebuggerView = ParentDebuggerViewPtr.Pin())
 	{
