motion matching - reintroducing AnimPoseSearchProvider previously nuked with CL 23157119, and related FPoseSearchProviderModule implemented with the current PoseSearch plugin

diff --git a/Engine/Plugins/Experimental/Animation/PoseSearch/Source/Editor/Private/PoseSearchDatabaseEditor.cpp b/Engine/Plugins/Experimental/Animation/PoseSearch/Source/Editor/Private/PoseSearchDatabaseEditor.cpp
--- a/Engine/Plugins/Experimental/Animation/PoseSearch/Source/Editor/Private/PoseSearchDatabaseEditor.cpp
+++ b/Engine/Plugins/Experimental/Animation/PoseSearch/Source/Editor/Private/PoseSearchDatabaseEditor.cpp
@@ -221,6 +221,9 @@ namespace UE::PoseSearch
 				// Ensure any database changes are reflected
 				DatabaseAsset->UnregisterOnDerivedDataRebuild(this);
 				DatabaseAsset->RegisterOnDerivedDataRebuild(UPoseSearchDatabase::FOnDerivedDataRebuild::CreateSP(this, &FDatabaseEditor::RefreshStatisticsWidgetInformation));
+
+				DatabaseAsset->UnregisterOnSynchronizeWithExternalDependencies(this);
+				DatabaseAsset->RegisterOnSynchronizeWithExternalDependencies(UPoseSearchDatabase::FOnDerivedDataRebuild::CreateSP(this, &FDatabaseEditor::RefreshEditor));
 			}
 		}
 		
@@ -295,6 +298,7 @@ namespace UE::PoseSearch
 			if (UPoseSearchDatabase* DatabaseAsset = ViewModel->GetPoseSearchDatabase())
 			{
 				DatabaseAsset->UnregisterOnDerivedDataRebuild(this);
+				DatabaseAsset->UnregisterOnSynchronizeWithExternalDependencies(this);
 			}
 		}
 	}
@@ -636,6 +640,11 @@ namespace UE::PoseSearch
 		Statistics->Initialize(GetPoseSearchDatabase());
 		StatisticsOverviewWidget->SetObject(Statistics);
 	}
+
+	void FDatabaseEditor::RefreshEditor()
+	{
+		AssetTreeWidget->RefreshTreeView(false, true);
+	}
 }
 
 #undef LOCTEXT_NAMESPACE
diff --git a/Engine/Plugins/Experimental/Animation/PoseSearch/Source/Editor/Public/PoseSearchDatabaseEditor.h b/Engine/Plugins/Experimental/Animation/PoseSearch/Source/Editor/Public/PoseSearchDatabaseEditor.h
--- a/Engine/Plugins/Experimental/Animation/PoseSearch/Source/Editor/Public/PoseSearchDatabaseEditor.h
+++ b/Engine/Plugins/Experimental/Animation/PoseSearch/Source/Editor/Public/PoseSearchDatabaseEditor.h
@@ -69,6 +69,7 @@ namespace UE::PoseSearch
 			ESelectInfo::Type SelectionType);
 
 		void RefreshStatisticsWidgetInformation();
+		void RefreshEditor();
 
 		struct FSelectionWidget
 		{
diff --git a/Engine/Plugins/Experimental/Animation/PoseSearch/Source/Runtime/Private/PoseSearchDatabase.cpp b/Engine/Plugins/Experimental/Animation/PoseSearch/Source/Runtime/Private/PoseSearchDatabase.cpp
--- a/Engine/Plugins/Experimental/Animation/PoseSearch/Source/Runtime/Private/PoseSearchDatabase.cpp
+++ b/Engine/Plugins/Experimental/Animation/PoseSearch/Source/Runtime/Private/PoseSearchDatabase.cpp
@@ -16,6 +16,10 @@
 #include "PoseSearch/PoseSearchSchema.h"
 #include "UObject/ObjectSaveContext.h"
 
+#if WITH_EDITOR
+#include "AssetRegistry/AssetRegistryModule.h"
+#endif //WITH_EDITOR
+
 #if WITH_EDITOR && WITH_ENGINE
 #include "Editor/EditorEngine.h"
 #endif //WITH_EDITOR && WITH_ENGINE
@@ -580,18 +584,88 @@ void UPoseSearchDatabase::PostLoad()
 }
 
 #if WITH_EDITOR
-void UPoseSearchDatabase::RegisterOnDerivedDataRebuild(const FOnDerivedDataRebuild& Delegate)
-{
-	OnDerivedDataRebuild.Add(Delegate);
-}
-void UPoseSearchDatabase::UnregisterOnDerivedDataRebuild(void* Unregister)
+void UPoseSearchDatabase::SynchronizeWithExternalDependencies()
 {
-	OnDerivedDataRebuild.RemoveAll(Unregister);
+	TArray<FTopLevelAssetPath> AncestorClassNames;
+
+	IAssetRegistry& AssetRegistry = FModuleManager::LoadModuleChecked<FAssetRegistryModule>("AssetRegistry").Get();
+
+	TArray<FAssetIdentifier> Referencers;
+	AssetRegistry.GetReferencers(GetPackage()->GetFName(), Referencers);
+
+	for (const FAssetIdentifier& Referencer : Referencers)
+	{
+		TArray<FAssetData> Assets;
+		AssetRegistry.GetAssetsByPackageName(Referencer.PackageName, Assets);
+
+		for (const FAssetData& Asset : Assets)
+		{
+			if (Asset.IsInstanceOf(UAnimSequenceBase::StaticClass()))
+			{
+				SynchronizeWithExternalDependencies(CastChecked<UAnimSequence>(Asset.FastGetAsset(true)));
+			}
+		}
+	}
 }
 
-void UPoseSearchDatabase::NotifyDerivedDataRebuild() const
+void UPoseSearchDatabase::SynchronizeWithExternalDependencies(UAnimSequenceBase* SequenceBase)
 {
-	OnDerivedDataRebuild.Broadcast();
+	if (SequenceBase)
+	{
+		bool bModified = false;
+		for (const FAnimNotifyEvent& NotifyEvent : SequenceBase->Notifies)
+		{
+			if (const UAnimNotifyState_PoseSearchBranchIn* PoseSearchBranchIn = Cast<UAnimNotifyState_PoseSearchBranchIn>(NotifyEvent.NotifyStateClass))
+			{
+				if (PoseSearchBranchIn->Database == this)
+				{
+					if (!bModified)
+					{
+						for (int32 AnimationAssetIndex = AnimationAssets.Num() - 1; AnimationAssetIndex >= 0; --AnimationAssetIndex)
+						{
+							if (const FPoseSearchDatabaseAnimationAssetBase* AnimationAssetBase = AnimationAssets[AnimationAssetIndex].GetPtr<FPoseSearchDatabaseAnimationAssetBase>())
+							{
+								if (AnimationAssetBase->GetAnimationAsset() == SequenceBase)
+								{
+									AnimationAssets.RemoveAt(AnimationAssetIndex);
+								}
+							}
+						}
+
+						bModified = true;
+					}
+
+					if (UAnimSequence* Sequence = Cast<UAnimSequence>(SequenceBase))
+					{
+						FPoseSearchDatabaseSequence DatabaseSequence;
+						DatabaseSequence.Sequence = Sequence;
+						DatabaseSequence.SamplingRange = FFloatInterval(NotifyEvent.GetTriggerTime(), NotifyEvent.GetEndTriggerTime());
+						AnimationAssets.Add(FInstancedStruct::Make(DatabaseSequence));
+					}
+					else if (UAnimComposite* AnimComposite = Cast<UAnimComposite>(SequenceBase))
+					{
+						FPoseSearchDatabaseAnimComposite DatabaseAnimComposite;
+						DatabaseAnimComposite.AnimComposite = AnimComposite;
+						DatabaseAnimComposite.SamplingRange = FFloatInterval(NotifyEvent.GetTriggerTime(), NotifyEvent.GetEndTriggerTime());
+						AnimationAssets.Add(FInstancedStruct::Make(DatabaseAnimComposite));
+					}
+					else if (UAnimMontage* AnimMontage = Cast<UAnimMontage>(SequenceBase))
+					{
+						FPoseSearchDatabaseAnimMontage DatabaseAnimMontage;
+						DatabaseAnimMontage.AnimMontage = AnimMontage;
+						DatabaseAnimMontage.SamplingRange = FFloatInterval(NotifyEvent.GetTriggerTime(), NotifyEvent.GetEndTriggerTime());
+						AnimationAssets.Add(FInstancedStruct::Make(DatabaseAnimMontage));
+					}
+				}
+			}
+		}
+
+		if (bModified)
+		{
+			Modify();
+			NotifySynchronizeWithExternalDependencies();
+		}
+	}
 }
 
 void UPoseSearchDatabase::BeginCacheForCookedPlatformData(const ITargetPlatform* TargetPlatform)
@@ -609,6 +683,13 @@ bool UPoseSearchDatabase::IsCachedCookedPlatformDataLoaded(const ITargetPlatform
 }
 #endif // WITH_EDITOR
 
+void UPoseSearchDatabase::PreSaveRoot(FObjectPreSaveRootContext ObjectSaveContext)
+{
+#if WITH_EDITOR
+	SynchronizeWithExternalDependencies();
+#endif
+}
+
 void UPoseSearchDatabase::PostSaveRoot(FObjectPostSaveRootContext ObjectSaveContext)
 {
 #if WITH_EDITOR
diff --git a/Engine/Plugins/Experimental/Animation/PoseSearch/Source/Runtime/Private/PoseSearchDerivedData.cpp b/Engine/Plugins/Experimental/Animation/PoseSearch/Source/Runtime/Private/PoseSearchDerivedData.cpp
--- a/Engine/Plugins/Experimental/Animation/PoseSearch/Source/Runtime/Private/PoseSearchDerivedData.cpp
+++ b/Engine/Plugins/Experimental/Animation/PoseSearch/Source/Runtime/Private/PoseSearchDerivedData.cpp
@@ -23,10 +23,13 @@
 #include "PoseSearch/PoseSearchSchema.h"
 #include "PoseSearchEigenHelper.h"
 #include "ProfilingDebugging/CookStats.h"
+#include "ScopedTransaction.h"
 #include "Serialization/BulkDataRegistry.h"
 #include "UObject/NoExportTypes.h"
 #include "UObject/PackageReload.h"
 
+#define LOCTEXT_NAMESPACE "PoseSearchDerivedData"
+
 namespace UE::PoseSearch
 {
 #if ENABLE_ANIM_DEBUG
@@ -34,40 +37,40 @@ namespace UE::PoseSearch
 enum EMotionMatchTestFlags
 {
 	// no additional tests will be performed
-	None = 0x0,
+	None = 0,
 
 	// cache will be invalidated every frame (to stress test DDC cancellation while tasks are flying if !WaitForTaskCompletion)
-	InvalidateCache = 0x1,
+	InvalidateCache = 1 << 0,
 
 	// cache will be invalidated once the flying tasks are ended
-	WaitForTaskCompletion = 0x2,
+	WaitForTaskCompletion = 1 << 1,
 
 	// we'll force the database re-indexing and compare the result SearchIndex with the one retrieved via DDC
-	ForceIndexing = 0x4,
+	ForceIndexing = 1 << 2,
 
 	// test KDTree Construct determinism
-	TestKDTreeConstructDeterminism = 0x8,
+	TestKDTreeConstructDeterminism = 1 << 3,
 
 	// validating the kdtree construction
-	ValidateKDTreeConstruct = 0x10,
+	ValidateKDTreeConstruct = 1 << 4,
 
 	// test VPTree Construct determinism
-	TestVPTreeConstructDeterminism = 0x20,
+	TestVPTreeConstructDeterminism = 1 << 5,
 
 	// validating the vptree construction
-	ValidateVPTreeConstruct = 0x40,
+	ValidateVPTreeConstruct = 1 << 6,
 
 	// test IndexDatabase determinism
-	TestIndexDatabaseDeterminism = 0x80,
+	TestIndexDatabaseDeterminism = 1 << 7,
 
 	// test PruneDuplicateValues determinism
-	TestPruneDuplicateValuesDeterminism = 0x100,
+	TestPruneDuplicateValuesDeterminism = 1 << 8,
 
 	// test PruneDuplicatePCAValues determinism
-	TestPruneDuplicatePCAValuesDeterminism = 0x200,
+	TestPruneDuplicatePCAValuesDeterminism = 1 << 9,
 
 	// validating the data we gave to DDC is stored correctly
-	ValidateDDC = 0x400,
+	ValidateDDC = 1 << 10,
 };
 static TAutoConsoleVariable<int32> CVarMotionMatchTestFlags(TEXT("a.MotionMatch.TestFlags"), EMotionMatchTestFlags::None, TEXT("Test Motion Matching using EMotionMatchTestFlags"));
 static TAutoConsoleVariable<int32> CVarMotionMatchTestNumIterations(TEXT("a.MotionMatch.TestNumIterations"), 10, TEXT("Test Motion Matching Num Iterations"));
@@ -1736,6 +1739,7 @@ FAsyncPoseSearchDatabasesManagement::FAsyncPoseSearchDatabasesManagement()
 	FScopeLock Lock(&Mutex);
 
 	OnObjectModifiedHandle = FCoreUObjectDelegates::OnObjectModified.AddRaw(this, &FAsyncPoseSearchDatabasesManagement::OnObjectModified);
+	OnObjectTransactedHandle = FCoreUObjectDelegates::OnObjectTransacted.AddRaw(this, &FAsyncPoseSearchDatabasesManagement::OnObjectTransacted);
 	OnPackageReloadedHandle = FCoreUObjectDelegates::OnPackageReloaded.AddRaw(this, &FAsyncPoseSearchDatabasesManagement::OnPackageReloaded);
 
 	FCoreDelegates::OnPreExit.AddRaw(this, &FAsyncPoseSearchDatabasesManagement::Shutdown);
@@ -1768,6 +1772,45 @@ void FAsyncPoseSearchDatabasesManagement::OnObjectModified(UObject* Object)
 	}
 }
 
+void FAsyncPoseSearchDatabasesManagement::OnObjectTransacted(UObject* Object, const FTransactionObjectEvent& TransactionObjectEvent)
+{
+	check(IsInGameThread());
+
+	FScopeLock Lock(&Mutex);
+
+	// synchronize the UPoseSearchDatabase with eventual UAnimSequenceBase containing UAnimNotifyState_PoseSearchBranchIn(s)
+	UAnimSequenceBase* SequenceBaseToSynchronizeWith = nullptr;
+	UPoseSearchDatabase* Database = nullptr;
+	if (UAnimNotifyState_PoseSearchBranchIn* BranchIn = Cast<UAnimNotifyState_PoseSearchBranchIn>(Object))
+	{
+		if (BranchIn->Database)
+		{
+			SequenceBaseToSynchronizeWith = Cast<UAnimSequenceBase>(BranchIn->GetOuter());
+			Database = BranchIn->Database;
+		}
+	}
+	else if (UAnimSequenceBase* SequenceBase = Cast<UAnimSequenceBase>(Object))
+	{
+		for (const FAnimNotifyEvent& NotifyEvent : SequenceBase->Notifies)
+		{
+			if (const UAnimNotifyState_PoseSearchBranchIn* NotifyEventBranchIn = Cast<UAnimNotifyState_PoseSearchBranchIn>(NotifyEvent.NotifyStateClass))
+			{
+				if (NotifyEventBranchIn->Database)
+				{
+					SequenceBaseToSynchronizeWith = SequenceBase;
+					Database = NotifyEventBranchIn->Database;
+					break;
+				}
+			}
+		}
+	}
+
+	if (SequenceBaseToSynchronizeWith)
+	{
+		Database->SynchronizeWithExternalDependencies(SequenceBaseToSynchronizeWith);
+	}
+}
+
 void FAsyncPoseSearchDatabasesManagement::OnPackageReloaded(const EPackageReloadPhase InPackageReloadPhase, FPackageReloadedEvent* InPackageReloadedEvent)
 {
 	check(IsInGameThread());
@@ -1926,4 +1969,7 @@ bool FAsyncPoseSearchDatabasesManagement::RequestAsyncBuildIndex(const UPoseSear
 }
 
 } // namespace UE::PoseSearch
+
+#undef LOCTEXT_NAMESPACE
+
 #endif // WITH_EDITOR
diff --git a/Engine/Plugins/Experimental/Animation/PoseSearch/Source/Runtime/Private/PoseSearchLibrary.cpp b/Engine/Plugins/Experimental/Animation/PoseSearch/Source/Runtime/Private/PoseSearchLibrary.cpp
--- a/Engine/Plugins/Experimental/Animation/PoseSearch/Source/Runtime/Private/PoseSearchLibrary.cpp
+++ b/Engine/Plugins/Experimental/Animation/PoseSearch/Source/Runtime/Private/PoseSearchLibrary.cpp
@@ -15,6 +15,7 @@
 #include "InstancedStruct.h"
 #include "PoseSearch/AnimNode_MotionMatching.h"
 #include "PoseSearch/AnimNode_PoseSearchHistoryCollector.h"
+#include "PoseSearch/PoseSearchAnimNotifies.h"
 #include "PoseSearch/PoseSearchDatabase.h"
 #include "PoseSearch/PoseSearchDerivedData.h"
 #include "PoseSearch/PoseSearchSchema.h"
@@ -143,10 +144,7 @@ void UPoseSearchLibrary::TraceMotionMatchingState(
 	int32 NodeId,
 	float DeltaTime,
 	bool bSearch,
-	float RecordingTime,
-	float SearchBestCost,
-	float SearchBruteForceCost,
-	int32 SearchBestPosePos)
+	float RecordingTime)
 {
 	using namespace UE::PoseSearch;
 	
@@ -210,12 +208,15 @@ void UPoseSearchLibrary::TraceMotionMatchingState(
 	if (DeltaTime > SMALL_NUMBER)
 	{
 		// simulation
-		const FTransform PrevRoot = SearchContext.GetWorldRootBoneTransformAtTime(-DeltaTime);
-		const FTransform CurrRoot = SearchContext.GetWorldRootBoneTransformAtTime(0.f);
-		const FTransform SimDelta = CurrRoot.GetRelativeTransform(PrevRoot);
+		if (SearchContext.IsTrajectoryValid())
+		{
+			const FTransform PrevRoot = SearchContext.GetWorldRootBoneTransformAtTime(-DeltaTime);
+			const FTransform CurrRoot = SearchContext.GetWorldRootBoneTransformAtTime(0.f);
+			const FTransform SimDelta = CurrRoot.GetRelativeTransform(PrevRoot);
 
-		TraceState.SimLinearVelocity = SimDelta.GetTranslation().Size() / DeltaTime;
-		TraceState.SimAngularVelocity = FMath::RadiansToDegrees(SimDelta.GetRotation().GetAngle()) / DeltaTime;
+			TraceState.SimLinearVelocity = SimDelta.GetTranslation().Size() / DeltaTime;
+			TraceState.SimAngularVelocity = FMath::RadiansToDegrees(SimDelta.GetRotation().GetAngle()) / DeltaTime;
+		}
 
 		// animation
 		TraceState.AnimLinearVelocity = RootMotionTransformDelta.GetTranslation().Size() / DeltaTime;
@@ -227,9 +228,9 @@ void UPoseSearchLibrary::TraceMotionMatchingState(
 	TraceState.DeltaTime = DeltaTime;
 
 	TraceState.RecordingTime = RecordingTime;
-	TraceState.SearchBestCost = SearchBestCost;
-	TraceState.SearchBruteForceCost = SearchBruteForceCost;
-	TraceState.SearchBestPosePos = SearchBestPosePos;
+	TraceState.SearchBestCost = CurrentResult.PoseCost.GetTotalCost();
+	TraceState.SearchBruteForceCost = CurrentResult.BruteForcePoseCost.GetTotalCost();
+	TraceState.SearchBestPosePos = CurrentResult.BestPosePos;
 
 	TraceState.Trajectory = Trajectory;
 
@@ -351,11 +352,9 @@ void UPoseSearchLibrary::UpdateMotionMatchingState(
 	// Record debugger details
 	if (IsTracing(Context))
 	{
-		const float SearchBestCost = InOutMotionMatchingState.CurrentSearchResult.PoseCost.GetTotalCost();
-		const float SearchBruteForceCost = InOutMotionMatchingState.CurrentSearchResult.BruteForcePoseCost.GetTotalCost();
 		TraceMotionMatchingState(Trajectory, SearchContext, InOutMotionMatchingState.CurrentSearchResult, InOutMotionMatchingState.ElapsedPoseSearchTime,
 			InOutMotionMatchingState.RootMotionTransformDelta, Context.AnimInstanceProxy->GetAnimInstanceObject(), Context.GetCurrentNodeId(), DeltaTime, bSearch,
-			AnimInstance ? FObjectTrace::GetWorldElapsedTime(AnimInstance->GetWorld()) : 0.f, SearchBestCost, SearchBruteForceCost, InOutMotionMatchingState.CurrentSearchResult.BestPosePos);
+			AnimInstance ? FObjectTrace::GetWorldElapsedTime(AnimInstance->GetWorld()) : 0.f);
 	}
 #endif // UE_POSE_SEARCH_TRACE_ENABLED
 
@@ -578,13 +577,94 @@ void UPoseSearchLibrary::MotionMatch(
 		}
 #endif // ENABLE_DRAW_DEBUG && ENABLE_ANIM_DEBUG
 
+#if UE_POSE_SEARCH_TRACE_ENABLED
+		
+		TraceMotionMatchingState(Trajectory, SearchContext, SearchResult, 0.f, FTransform::Identity, AnimInstance, DebugSessionUniqueIdentifier,
+			AnimInstance->GetDeltaSeconds(), true, FObjectTrace::GetWorldElapsedTime(AnimInstance->GetWorld()));
+#endif // UE_POSE_SEARCH_TRACE_ENABLED
+	}
+}
+
+UE::PoseSearch::FSearchResult UPoseSearchLibrary::MotionMatch(const FAnimationBaseContext& Context, const UObject* Object)
+{
+	//using namespace UE::Anim;
+	using namespace UE::PoseSearch;
+
+	FSearchResult SearchResult;
+
+	// @todo: we currently support only animation assets, but in the future we should support choosers and proxy (via some common interface)
+	if (const UAnimSequenceBase* SequenceBase = Cast<const UAnimSequenceBase>(Object))
+	{
+		const IPoseHistory* History = nullptr;
+		if (IPoseHistoryProvider* PoseHistoryProvider = Context.GetMessage<IPoseHistoryProvider>())
+		{
+			History = &PoseHistoryProvider->GetPoseHistory();
+		}
+
+		const UAnimInstance* AnimInstance = Cast<const UAnimInstance>(Context.AnimInstanceProxy->GetAnimInstanceObject());
+		check(AnimInstance);
+
+		FMemMark Mark(FMemStack::Get());
+		FSearchContext SearchContext(AnimInstance, nullptr, History);
+		for (const FAnimNotifyEvent& NotifyEvent : SequenceBase->Notifies)
+		{
+			if (const UAnimNotifyState_PoseSearchBranchIn* PoseSearchBranchIn = Cast<UAnimNotifyState_PoseSearchBranchIn>(NotifyEvent.NotifyStateClass))
+			{
+				if (PoseSearchBranchIn->Database)
+				{
+#if ENABLE_ANIM_DEBUG
+					// @todo: implement database filtering based on subsets of Database->AnimationAssets!
+					//		  FSearchContext should hold a list of UAnimationAsset(s) to use to narrow down the database search,
+					//		  but for now we log an error
+					for (int32 AnimationAssetIndex = 0; AnimationAssetIndex < PoseSearchBranchIn->Database->AnimationAssets.Num(); ++AnimationAssetIndex)
+					{
+						if (const FPoseSearchDatabaseAnimationAssetBase* AnimationAssetBase = PoseSearchBranchIn->Database->GetAnimationAssetBase(AnimationAssetIndex))
+						{
+							if (AnimationAssetBase->GetAnimationAsset() != Object)
+							{
+								UE_LOG(LogPoseSearch, Error, TEXT("UAnimNotifyState_PoseSearchBranchIn doesn't yet support database between multiple animation assets"));
+							}
+						}
+					}
+#endif // ENABLE_ANIM_DEBUG
+
+					const FSearchResult NewSearchResult = PoseSearchBranchIn->Database->Search(SearchContext);
+
+					if (NewSearchResult.PoseCost.GetTotalCost() < SearchResult.PoseCost.GetTotalCost())
+					{
+						SearchResult = NewSearchResult;
+						SearchContext.UpdateCurrentBestCost(SearchResult.PoseCost);
+					}
+				}
+			}
+		}
+
+#if ENABLE_DRAW_DEBUG && ENABLE_ANIM_DEBUG
+		if (SearchResult.IsValid())
+		{
+			if (CVarAnimMotionMatchDrawMatchEnable.GetValueOnAnyThread())
+			{
+				UE::PoseSearch::FDebugDrawParams DrawParams(Context.AnimInstanceProxy, Context.AnimInstanceProxy->GetComponentTransform(), SearchResult.Database.Get());
+				DrawParams.DrawFeatureVector(SearchResult.PoseIdx);
+			}
+
+			if (CVarAnimMotionMatchDrawQueryEnable.GetValueOnAnyThread())
+			{
+				UE::PoseSearch::FDebugDrawParams DrawParams(Context.AnimInstanceProxy, Context.AnimInstanceProxy->GetComponentTransform(), SearchResult.Database.Get(), EDebugDrawFlags::DrawQuery);
+				DrawParams.DrawFeatureVector(SearchContext.GetOrBuildQuery(SearchResult.Database->Schema).GetValues());
+			}
+		}
+#endif // ENABLE_DRAW_DEBUG && ENABLE_ANIM_DEBUG
+
 #if UE_POSE_SEARCH_TRACE_ENABLED
 		const float SearchBestCost = SearchResult.PoseCost.GetTotalCost();
 		const float SearchBruteForceCost = SearchResult.BruteForcePoseCost.GetTotalCost();
-		TraceMotionMatchingState(Trajectory, SearchContext, SearchResult, 0.f, FTransform::Identity, AnimInstance, DebugSessionUniqueIdentifier,
-			AnimInstance->GetDeltaSeconds(), true, FObjectTrace::GetWorldElapsedTime(AnimInstance->GetWorld()), SearchBestCost, SearchBruteForceCost, SearchResult.BestPosePos);
+		TraceMotionMatchingState(FPoseSearchQueryTrajectory(), SearchContext, SearchResult, 0.f, FTransform::Identity, AnimInstance, Context.GetCurrentNodeId(),
+			AnimInstance->GetDeltaSeconds(), true, FObjectTrace::GetWorldElapsedTime(AnimInstance->GetWorld()));
 #endif // UE_POSE_SEARCH_TRACE_ENABLED
 	}
+
+	return SearchResult;
 }
 
 #undef LOCTEXT_NAMESPACE
diff --git a/Engine/Plugins/Experimental/Animation/PoseSearch/Source/Runtime/Private/PoseSearchModule.cpp b/Engine/Plugins/Experimental/Animation/PoseSearch/Source/Runtime/Private/PoseSearchModule.cpp
--- a/Engine/Plugins/Experimental/Animation/PoseSearch/Source/Runtime/Private/PoseSearchModule.cpp
+++ b/Engine/Plugins/Experimental/Animation/PoseSearch/Source/Runtime/Private/PoseSearchModule.cpp
@@ -1,5 +1,37 @@
-﻿// Copyright Epic Games, Inc. All Rights Reserved.
+// Copyright Epic Games, Inc. All Rights Reserved.
 
+#include "Animation/AnimPoseSearchProvider.h"
+#include "Features/IModularFeatures.h"
 #include "Modules/ModuleManager.h"
+#include "PoseSearch/PoseSearchLibrary.h"
 
-IMPLEMENT_MODULE(FDefaultModuleImpl, PoseSearch);
\ No newline at end of file
+class FPoseSearchModule final : public IModuleInterface, public UE::Anim::IPoseSearchProvider
+{
+public:
+	
+	// IModuleInterface
+	virtual void StartupModule() override
+	{
+		IModularFeatures::Get().RegisterModularFeature(UE::Anim::IPoseSearchProvider::GetModularFeatureName(), this);
+	}
+
+	virtual void ShutdownModule() override
+	{
+		IModularFeatures::Get().UnregisterModularFeature(UE::Anim::IPoseSearchProvider::GetModularFeatureName(), this);
+	}
+
+	// IPoseSearchProvider
+	virtual UE::Anim::IPoseSearchProvider::FSearchResult Search(const FAnimationBaseContext& GraphContext, const UObject* Object) override
+	{
+		const UE::PoseSearch::FSearchResult SearchResult = UPoseSearchLibrary::MotionMatch(GraphContext, Object);
+		
+		UE::Anim::IPoseSearchProvider::FSearchResult ProviderResult;
+		ProviderResult.Dissimilarity = SearchResult.PoseCost.GetTotalCost();
+		ProviderResult.PoseIdx = SearchResult.PoseIdx;
+		ProviderResult.TimeOffsetSeconds = SearchResult.AssetTime;
+
+		return ProviderResult;
+	}
+};
+
+IMPLEMENT_MODULE(FPoseSearchModule, PoseSearch);
\ No newline at end of file
diff --git a/Engine/Plugins/Experimental/Animation/PoseSearch/Source/Runtime/Public/PoseSearch/PoseSearchAnimNotifies.h b/Engine/Plugins/Experimental/Animation/PoseSearch/Source/Runtime/Public/PoseSearch/PoseSearchAnimNotifies.h
--- a/Engine/Plugins/Experimental/Animation/PoseSearch/Source/Runtime/Public/PoseSearch/PoseSearchAnimNotifies.h
+++ b/Engine/Plugins/Experimental/Animation/PoseSearch/Source/Runtime/Public/PoseSearch/PoseSearchAnimNotifies.h
@@ -6,6 +6,8 @@
 #include "BoneContainer.h"
 #include "PoseSearchAnimNotifies.generated.h"
 
+class UPoseSearchDatabase;
+
 // Base class for pose search anim notify states
 UCLASS(Abstract)
 class POSESEARCH_API UAnimNotifyState_PoseSearchBase : public UAnimNotifyState
@@ -87,4 +89,15 @@ public:
 	FVector LinearVelocity = FVector::ZeroVector;
 
 #endif // WITH_EDITORONLY_DATA
-};
\ No newline at end of file
+};
+
+UCLASS(Blueprintable, meta = (DisplayName = "Pose Matching: Motion Matched Branch In"))
+class POSESEARCH_API UAnimNotifyState_PoseSearchBranchIn : public UAnimNotifyState_PoseSearchBase
+{
+	GENERATED_BODY()
+
+public:
+
+	UPROPERTY(EditAnywhere, Category = Config)
+	TObjectPtr<UPoseSearchDatabase> Database;
+};
diff --git a/Engine/Plugins/Experimental/Animation/PoseSearch/Source/Runtime/Public/PoseSearch/PoseSearchContext.h b/Engine/Plugins/Experimental/Animation/PoseSearch/Source/Runtime/Public/PoseSearch/PoseSearchContext.h
--- a/Engine/Plugins/Experimental/Animation/PoseSearch/Source/Runtime/Public/PoseSearch/PoseSearchContext.h
+++ b/Engine/Plugins/Experimental/Animation/PoseSearch/Source/Runtime/Public/PoseSearch/PoseSearchContext.h
@@ -142,7 +142,7 @@ private:
 
 struct POSESEARCH_API FSearchContext
 {
-	FSearchContext(const UAnimInstance* InAnimInstance, const FPoseSearchQueryTrajectory* InTrajectory, const IPoseHistory* InHistory, float InDesiredPermutationTimeOffset, const FPoseIndicesHistory* InPoseIndicesHistory = nullptr,
+	FSearchContext(const UAnimInstance* InAnimInstance, const FPoseSearchQueryTrajectory* InTrajectory, const IPoseHistory* InHistory, float InDesiredPermutationTimeOffset = 0.f, const FPoseIndicesHistory* InPoseIndicesHistory = nullptr,
 		const FSearchResult& InCurrentResult = FSearchResult(), const FFloatInterval& InPoseJumpThresholdTime = FFloatInterval(0.f, 0.f), bool bInForceInterrupt = false);
 
 	// Returns the rotation of the bone Schema.BoneReferences[SchemaSampleBoneIdx] at an offset time of SampleTimeOffset relative to the
@@ -181,6 +181,7 @@ struct POSESEARCH_API FSearchContext
 	const FFloatInterval& GetPoseJumpThresholdTime() const { return PoseJumpThresholdTime; }
 	const FPoseIndicesHistory* GetPoseIndicesHistory() const { return PoseIndicesHistory; }
 	bool IsHistoryValid() const { return History != nullptr; }
+	bool IsTrajectoryValid() const { return Trajectory != nullptr; }
 	float GetDesiredPermutationTimeOffset() const { return DesiredPermutationTimeOffset; }
 	bool IsForceInterrupt() const { return bForceInterrupt; }
 	FTransform GetWorldRootBoneTransformAtTime(float SampleTime, bool bUseHistoryRoot = false, bool bExtrapolate = true) const;
diff --git a/Engine/Plugins/Experimental/Animation/PoseSearch/Source/Runtime/Public/PoseSearch/PoseSearchDatabase.h b/Engine/Plugins/Experimental/Animation/PoseSearch/Source/Runtime/Public/PoseSearch/PoseSearchDatabase.h
--- a/Engine/Plugins/Experimental/Animation/PoseSearch/Source/Runtime/Public/PoseSearch/PoseSearchDatabase.h
+++ b/Engine/Plugins/Experimental/Animation/PoseSearch/Source/Runtime/Public/PoseSearch/PoseSearchDatabase.h
@@ -315,6 +315,9 @@ private:
 #if WITH_EDITOR
 	DECLARE_MULTICAST_DELEGATE(FOnDerivedDataRebuildMulticaster);
 	FOnDerivedDataRebuildMulticaster OnDerivedDataRebuild;
+
+	DECLARE_MULTICAST_DELEGATE(FOnSynchronizeWithExternalDependenciesMulticaster);
+	FOnSynchronizeWithExternalDependenciesMulticaster OnSynchronizeWithExternalDependencies;
 #endif // WITH_EDITOR
 
 public:
@@ -343,6 +346,7 @@ public:
 
 	// Begin UObject
 	virtual void PostLoad() override;
+	virtual void PreSaveRoot(FObjectPreSaveRootContext ObjectSaveContext) override;
 	virtual void PostSaveRoot(FObjectPostSaveRootContext ObjectSaveContext) override;
 	virtual void Serialize(FArchive& Ar) override;
 	// End UObject
@@ -355,13 +359,21 @@ public:
 	virtual bool IsCachedCookedPlatformDataLoaded(const ITargetPlatform* TargetPlatform) override;
 
 	typedef FOnDerivedDataRebuildMulticaster::FDelegate FOnDerivedDataRebuild;
-	void RegisterOnDerivedDataRebuild(const FOnDerivedDataRebuild& Delegate);
-	void UnregisterOnDerivedDataRebuild(void* Unregister);
-	void NotifyDerivedDataRebuild() const;
+	void RegisterOnDerivedDataRebuild(const FOnDerivedDataRebuild& Delegate) { OnDerivedDataRebuild.Add(Delegate); }
+	void UnregisterOnDerivedDataRebuild(void* Unregister) { OnDerivedDataRebuild.RemoveAll(Unregister); }
+	void NotifyDerivedDataRebuild() const { OnDerivedDataRebuild.Broadcast(); }
+
+	typedef FOnSynchronizeWithExternalDependenciesMulticaster::FDelegate FOnSynchronizeWithExternalDependencies;
+	void RegisterOnSynchronizeWithExternalDependencies(const FOnSynchronizeWithExternalDependencies& Delegate) { OnSynchronizeWithExternalDependencies.Add(Delegate); }
+	void UnregisterOnSynchronizeWithExternalDependencies(void* Unregister) { OnSynchronizeWithExternalDependencies.RemoveAll(Unregister); }
+	void NotifySynchronizeWithExternalDependencies() const { OnSynchronizeWithExternalDependencies.Broadcast(); }
+
+	void SynchronizeWithExternalDependencies();
+	void SynchronizeWithExternalDependencies(UAnimSequenceBase* SequenceBase);
 #endif // WITH_EDITOR
 
 private:
 	UE::PoseSearch::FSearchResult SearchPCAKDTree(UE::PoseSearch::FSearchContext& SearchContext) const;
 	UE::PoseSearch::FSearchResult SearchVPTree(UE::PoseSearch::FSearchContext& SearchContext) const;
 	UE::PoseSearch::FSearchResult SearchBruteForce(UE::PoseSearch::FSearchContext& SearchContext) const;
-};
\ No newline at end of file
+};
diff --git a/Engine/Plugins/Experimental/Animation/PoseSearch/Source/Runtime/Public/PoseSearch/PoseSearchDerivedData.h b/Engine/Plugins/Experimental/Animation/PoseSearch/Source/Runtime/Public/PoseSearch/PoseSearchDerivedData.h
--- a/Engine/Plugins/Experimental/Animation/PoseSearch/Source/Runtime/Public/PoseSearch/PoseSearchDerivedData.h
+++ b/Engine/Plugins/Experimental/Animation/PoseSearch/Source/Runtime/Public/PoseSearch/PoseSearchDerivedData.h
@@ -34,6 +34,7 @@ namespace UE::PoseSearch
 		static FAsyncPoseSearchDatabasesManagement& Get();
 
 		void OnObjectModified(UObject* Object);
+		void OnObjectTransacted(UObject* Object, const FTransactionObjectEvent& TransactionObjectEvent);
 		void OnPackageReloaded(const EPackageReloadPhase InPackageReloadPhase, FPackageReloadedEvent* InPackageReloadedEvent);
 
 		void Shutdown();
@@ -58,6 +59,7 @@ namespace UE::PoseSearch
 		
 		FPoseSearchDatabaseAsyncCacheTasks& Tasks;
 		FDelegateHandle OnObjectModifiedHandle;
+		FDelegateHandle OnObjectTransactedHandle;
 		FDelegateHandle OnPackageReloadedHandle;
 		
 		static FCriticalSection Mutex;
diff --git a/Engine/Plugins/Experimental/Animation/PoseSearch/Source/Runtime/Public/PoseSearch/PoseSearchLibrary.h b/Engine/Plugins/Experimental/Animation/PoseSearch/Source/Runtime/Public/PoseSearch/PoseSearchLibrary.h
--- a/Engine/Plugins/Experimental/Animation/PoseSearch/Source/Runtime/Public/PoseSearch/PoseSearchLibrary.h
+++ b/Engine/Plugins/Experimental/Animation/PoseSearch/Source/Runtime/Public/PoseSearch/PoseSearchLibrary.h
@@ -84,10 +84,7 @@ class POSESEARCH_API UPoseSearchLibrary : public UBlueprintFunctionLibrary
 		int32 NodeId,
 		float DeltaTime,
 		bool bSearch,
-		float RecordingTime,
-		float SearchBestCost,
-		float SearchBruteForceCost,
-		int32 SearchBestPosePos);
+		float RecordingTime);
 #endif // UE_POSE_SEARCH_TRACE_ENABLED
 
 	static FPoseSearchQueryTrajectory ProcessTrajectory(
@@ -165,5 +162,8 @@ public:
 		float FutureAnimationStartTime = 0.f,
 		float TimeToFutureAnimationStart = 0.f,
 		const int32 DebugSessionUniqueIdentifier = 6174);
+
+
+	static UE::PoseSearch::FSearchResult MotionMatch(const FAnimationBaseContext& Context, const UObject* Object);
 };
 
diff --git a/Engine/Source/Editor/AnimGraph/Private/AnimGraphNode_SequencePlayer.cpp b/Engine/Source/Editor/AnimGraph/Private/AnimGraphNode_SequencePlayer.cpp
--- a/Engine/Source/Editor/AnimGraph/Private/AnimGraphNode_SequencePlayer.cpp
+++ b/Engine/Source/Editor/AnimGraph/Private/AnimGraphNode_SequencePlayer.cpp
@@ -16,6 +16,7 @@
 #include "EditorCategoryUtils.h"
 #include "Animation/AnimComposite.h"
 #include "Animation/AnimSequence.h"
+#include "Animation/AnimPoseSearchProvider.h"
 #include "Animation/AnimRootMotionProvider.h"
 #include "UObject/UE5MainStreamObjectVersion.h"
 #include "IAnimBlueprintNodeOverrideAssetsContext.h"
@@ -251,6 +252,16 @@ UScriptStruct* UAnimGraphNode_SequencePlayer::GetTimePropertyStruct() const
 	return FAnimNode_SequencePlayer::StaticStruct();
 }
 
+void UAnimGraphNode_SequencePlayer::CustomizeDetails(IDetailLayoutBuilder& DetailBuilder)
+{
+	Super::CustomizeDetails(DetailBuilder);
+
+	if (!UE::Anim::IPoseSearchProvider::IsAvailable())
+	{
+		DetailBuilder.HideCategory(TEXT("PoseMatching"));
+	}
+}
+
 void UAnimGraphNode_SequencePlayer::CustomizePinData(UEdGraphPin* Pin, FName SourcePropertyName, int32 ArrayIndex) const
 {
 	Super::CustomizePinData(Pin, SourcePropertyName, ArrayIndex);
diff --git a/Engine/Source/Editor/AnimGraph/Public/AnimGraphNode_SequencePlayer.h b/Engine/Source/Editor/AnimGraph/Public/AnimGraphNode_SequencePlayer.h
--- a/Engine/Source/Editor/AnimGraph/Public/AnimGraphNode_SequencePlayer.h
+++ b/Engine/Source/Editor/AnimGraph/Public/AnimGraphNode_SequencePlayer.h
@@ -31,6 +31,7 @@ class UAnimGraphNode_SequencePlayer : public UAnimGraphNode_AssetPlayerBase
 	// End of UEdGraphNode interface
 
 	// UAnimGraphNode_Base interface
+	virtual void CustomizeDetails(IDetailLayoutBuilder& DetailBuilder) override;
 	virtual void CustomizePinData(UEdGraphPin* Pin, FName SourcePropertyName, int32 ArrayIndex) const override;
 	virtual void ValidateAnimNodeDuringCompilation(class USkeleton* ForSkeleton, class FCompilerResultsLog& MessageLog) override;
 	virtual void PreloadRequiredAssets() override;		
diff --git a/Engine/Source/Runtime/Engine/Classes/Animation/AnimNode_SequencePlayer.h b/Engine/Source/Runtime/Engine/Classes/Animation/AnimNode_SequencePlayer.h
--- a/Engine/Source/Runtime/Engine/Classes/Animation/AnimNode_SequencePlayer.h
+++ b/Engine/Source/Runtime/Engine/Classes/Animation/AnimNode_SequencePlayer.h
@@ -38,7 +38,6 @@ public:
 	// End of FAnimNode_Base interface
 
 	ENGINE_API float GetTimeFromEnd(float CurrentNodeTime) const;
-	UE_DEPRECATED(5.1, "GetEffectiveStartPosition is no longer supported. Please use GetStartPosition instead")
 	ENGINE_API float GetEffectiveStartPosition(const FAnimationBaseContext& Context) const;
 
 	// The animation sequence asset to play
diff --git a/Engine/Source/Runtime/Engine/Private/Animation/AnimNode_SequencePlayer.cpp b/Engine/Source/Runtime/Engine/Private/Animation/AnimNode_SequencePlayer.cpp
--- a/Engine/Source/Runtime/Engine/Private/Animation/AnimNode_SequencePlayer.cpp
+++ b/Engine/Source/Runtime/Engine/Private/Animation/AnimNode_SequencePlayer.cpp
@@ -5,6 +5,7 @@
 #include "Animation/AnimMontage.h"
 #include "Animation/AnimInstanceProxy.h"
 #include "Animation/AnimationPoseData.h"
+#include "Animation/AnimPoseSearchProvider.h"
 #include "Animation/ExposedValueHandler.h"
 #include "Logging/TokenizedMessage.h"
 
@@ -61,12 +62,14 @@ void FAnimNode_SequencePlayerBase::Initialize_AnyThread(const FAnimationInitiali
 
 	if (CurrentSequence != nullptr)
 	{
+		const float EffectiveStartPosition = GetEffectiveStartPosition(Context);
 		const float CurrentPlayRate = GetPlayRate();
 		const float CurrentPlayRateBasis = GetPlayRateBasis();
 
+		InternalTimeAccumulator = FMath::Clamp(EffectiveStartPosition, 0.f, CurrentSequence->GetPlayLength());
 		const float AdjustedPlayRate = PlayRateScaleBiasClampState.ApplyTo(GetPlayRateScaleBiasClampConstants(), FMath::IsNearlyZero(CurrentPlayRateBasis) ? 0.f : (CurrentPlayRate / CurrentPlayRateBasis), 0.f);
 		const float EffectivePlayrate = CurrentSequence->RateScale * AdjustedPlayRate;
-		if ((InternalTimeAccumulator == 0.f) && (EffectivePlayrate < 0.f))
+		if ((EffectiveStartPosition == 0.f) && (EffectivePlayrate < 0.f))
 		{
 			InternalTimeAccumulator = CurrentSequence->GetPlayLength();
 		}
@@ -155,6 +158,21 @@ float FAnimNode_SequencePlayerBase::GetTimeFromEnd(float CurrentNodeTime) const
 
 float FAnimNode_SequencePlayerBase::GetEffectiveStartPosition(const FAnimationBaseContext& Context) const
 {
+	// Override the start position if pose matching is enabled
+	UAnimSequenceBase* CurrentSequence = GetSequence();
+	if (CurrentSequence != nullptr && GetStartFromMatchingPose())
+	{
+		UE::Anim::IPoseSearchProvider* PoseSearchProvider = UE::Anim::IPoseSearchProvider::Get();
+		if (PoseSearchProvider)
+		{
+			UE::Anim::IPoseSearchProvider::FSearchResult Result = PoseSearchProvider->Search(Context, CurrentSequence);
+			if (Result.PoseIdx >= 0)
+			{
+				return Result.TimeOffsetSeconds;
+			}
+		}
+	}
+
 	return GetStartPosition();
 }
 
diff --git a/Engine/Source/Runtime/Engine/Private/Animation/AnimPoseSearchProvider.cpp b/Engine/Source/Runtime/Engine/Private/Animation/AnimPoseSearchProvider.cpp
new file mode 100644
--- /dev/null
+++ b/Engine/Source/Runtime/Engine/Private/Animation/AnimPoseSearchProvider.cpp
@@ -0,0 +1,33 @@
+// Copyright Epic Games, Inc. All Rights Reserved.
+
+#include "Animation/AnimPoseSearchProvider.h"
+#include "Features/IModularFeatures.h"
+
+namespace UE::Anim
+{
+
+FName IPoseSearchProvider::GetModularFeatureName()
+{
+	static FName FeatureName = FName(TEXT("AnimPoseSearch"));
+	return FeatureName;
+}
+	
+IPoseSearchProvider* IPoseSearchProvider::Get()
+{
+	IModularFeatures::FScopedLockModularFeatureList ScopedLockModularFeatureList;
+
+	if (IModularFeatures::Get().IsModularFeatureAvailable(GetModularFeatureName()))
+	{
+		return &IModularFeatures::Get().GetModularFeature<IPoseSearchProvider>(GetModularFeatureName());
+	}
+
+	return nullptr;
+}
+
+bool IPoseSearchProvider::IsAvailable()
+{
+	IModularFeatures::FScopedLockModularFeatureList ScopedLockModularFeatureList;
+	return IModularFeatures::Get().IsModularFeatureAvailable(GetModularFeatureName());
+}
+
+} // namespace UE::Anim
\ No newline at end of file
diff --git a/Engine/Source/Runtime/Engine/Public/Animation/AnimPoseSearchProvider.h b/Engine/Source/Runtime/Engine/Public/Animation/AnimPoseSearchProvider.h
new file mode 100644
--- /dev/null
+++ b/Engine/Source/Runtime/Engine/Public/Animation/AnimPoseSearchProvider.h
@@ -0,0 +1,43 @@
+// Copyright Epic Games, Inc. All Rights Reserved.
+
+#pragma once
+
+#include "CoreTypes.h"
+#include "CoreMinimal.h"
+#include "Features/IModularFeature.h"
+
+struct FAnimationBaseContext;
+class UAnimationAsset;
+
+namespace UE::Anim
+{
+
+/** Modular feature interface for PoseSearch */
+class ENGINE_API IPoseSearchProvider : public IModularFeature
+{
+public:
+	virtual ~IPoseSearchProvider() {}
+
+	static FName GetModularFeatureName();
+	static bool IsAvailable();
+	static IPoseSearchProvider* Get();
+	
+	struct FSearchResult
+	{
+		int32 PoseIdx = INDEX_NONE;
+		float TimeOffsetSeconds = 0.f;
+		float Dissimilarity = MAX_flt;
+	};
+
+	/**
+	* Finds a matching pose in the input Object given the current graph context
+	* 
+	* @param	GraphContext	Graph execution context used to construct a pose search query
+	* @param	Object			The object to search for the pose query
+	* 
+	* @return	The pose in the animation asset that most closely matches the query
+	*/
+	virtual FSearchResult Search(const FAnimationBaseContext& GraphContext, const UObject* Object) = 0;
+};
+
+} // namespace UE::Anim
\ No newline at end of file
