motion matching
- replacing some check with UE_LOG to improve QOL
- removing unnecessary UPoseSearchDatabase::bSkipSearchIfPossible (it should always be true!)
- removed allocation in FAnimationAssetSampler::ExtractPoseSearchNotifyStates

diff --git a/Engine/Plugins/Experimental/Animation/PoseSearch/Source/Editor/Private/PoseSearchDebuggerView.cpp b/Engine/Plugins/Experimental/Animation/PoseSearch/Source/Editor/Private/PoseSearchDebuggerView.cpp
--- a/Engine/Plugins/Experimental/Animation/PoseSearch/Source/Editor/Private/PoseSearchDebuggerView.cpp
+++ b/Engine/Plugins/Experimental/Animation/PoseSearch/Source/Editor/Private/PoseSearchDebuggerView.cpp
@@ -275,7 +275,7 @@ bool SDebuggerView::UpdateNodeSelection()
 				.Padding(10.0f)
 				[
 					SNew(SButton)
-					.Text(FText::FromString(Model->GetCurrentDatabase()->GetName()))
+					.Text(FText::FromString(GetNameSafe(Model->GetCurrentDatabase())))
 					.HAlign(HAlign_Center)
 					.VAlign(VAlign_Center)
 					.ContentPadding(10.0f)
diff --git a/Engine/Plugins/Experimental/Animation/PoseSearch/Source/Runtime/Private/PoseSearchAssetIndexer.cpp b/Engine/Plugins/Experimental/Animation/PoseSearch/Source/Runtime/Private/PoseSearchAssetIndexer.cpp
--- a/Engine/Plugins/Experimental/Animation/PoseSearch/Source/Runtime/Private/PoseSearchAssetIndexer.cpp
+++ b/Engine/Plugins/Experimental/Animation/PoseSearch/Source/Runtime/Private/PoseSearchAssetIndexer.cpp
@@ -140,19 +140,18 @@ void FAssetIndexer::Process(int32 AssetIdx)
 		float CostAddend = SamplingContext.BaseCostBias;
 		bool bBlockTransition = false;
 
-		TArray<UAnimNotifyState_PoseSearchBase*> NotifyStates;
-		AssetSampler.ExtractPoseSearchNotifyStates(SampleTime, NotifyStates);
-		for (const UAnimNotifyState_PoseSearchBase* PoseSearchNotify : NotifyStates)
-		{
-			if (PoseSearchNotify->GetClass()->IsChildOf<UAnimNotifyState_PoseSearchBlockTransition>())
-			{
-				bBlockTransition = true;
-			}
-			else if (const UAnimNotifyState_PoseSearchModifyCost* ModifyCostNotify = Cast<const UAnimNotifyState_PoseSearchModifyCost>(PoseSearchNotify))
+		AssetSampler.ExtractPoseSearchNotifyStates(SampleTime, [&bBlockTransition, &CostAddend](const UAnimNotifyState_PoseSearchBase* PoseSearchNotify)
 			{
-				CostAddend = ModifyCostNotify->CostAddend;
-			}
-		}
+				if (PoseSearchNotify->GetClass()->IsChildOf<UAnimNotifyState_PoseSearchBlockTransition>())
+				{
+					bBlockTransition = true;
+				}
+				else if (const UAnimNotifyState_PoseSearchModifyCost* ModifyCostNotify = Cast<const UAnimNotifyState_PoseSearchModifyCost>(PoseSearchNotify))
+				{
+					CostAddend = ModifyCostNotify->CostAddend;
+				}
+				return true;
+			});
 
 		if (AssetSampler.IsLoopable())
 		{
diff --git a/Engine/Plugins/Experimental/Animation/PoseSearch/Source/Runtime/Private/PoseSearchAssetSampler.cpp b/Engine/Plugins/Experimental/Animation/PoseSearch/Source/Runtime/Private/PoseSearchAssetSampler.cpp
--- a/Engine/Plugins/Experimental/Animation/PoseSearch/Source/Runtime/Private/PoseSearchAssetSampler.cpp
+++ b/Engine/Plugins/Experimental/Animation/PoseSearch/Source/Runtime/Private/PoseSearchAssetSampler.cpp
@@ -581,7 +581,7 @@ void FAnimationAssetSampler::Process(const FBoneContainer& BoneContainer)
 	}
 }
 
-void FAnimationAssetSampler::ExtractPoseSearchNotifyStates(float Time, TArray<UAnimNotifyState_PoseSearchBase*>& NotifyStates) const
+void FAnimationAssetSampler::ExtractPoseSearchNotifyStates(float Time, TFunction<bool(UAnimNotifyState_PoseSearchBase*)> ProcessPoseSearchBase) const
 {
 	float SampleTime = Time;
 	FAnimNotifyContext NotifyContext;
@@ -642,7 +642,10 @@ void FAnimationAssetSampler::ExtractPoseSearchNotifyStates(float Time, TArray<UA
 
 		if (UAnimNotifyState_PoseSearchBase* PoseSearchAnimNotify = Cast<UAnimNotifyState_PoseSearchBase>(NotifyEvent->NotifyStateClass))
 		{
-			NotifyStates.Add(PoseSearchAnimNotify);
+			if (!ProcessPoseSearchBase(PoseSearchAnimNotify))
+			{
+				break;
+			}
 		}
 	}
 }
diff --git a/Engine/Plugins/Experimental/Animation/PoseSearch/Source/Runtime/Private/PoseSearchContext.cpp b/Engine/Plugins/Experimental/Animation/PoseSearch/Source/Runtime/Private/PoseSearchContext.cpp
--- a/Engine/Plugins/Experimental/Animation/PoseSearch/Source/Runtime/Private/PoseSearchContext.cpp
+++ b/Engine/Plugins/Experimental/Animation/PoseSearch/Source/Runtime/Private/PoseSearchContext.cpp
@@ -432,15 +432,35 @@ FTransform FSearchContext::GetWorldRootBoneTransformAtTime(float SampleTime, boo
 	FTransform WorldRootBoneTransform = FTransform::Identity;
 	if (bUseHistoryRoot)
 	{
+#if WITH_EDITOR
+		if (!History)
+		{
+			UE_LOG(LogPoseSearch, Error, TEXT("FSearchContext::GetWorldRootBoneTransformAtTime - missing History is required!"));
+		}
+		else
+#else // WITH_EDITOR
 		check(History);
-		// asking for root bone (RootBoneIndexType) in world space (WorldSpaceIndexType)
-		History->GetTransformAtTime(SampleTime, WorldRootBoneTransform, AnimInstance->CurrentSkeleton, RootBoneIndexType, WorldSpaceIndexType, bExtrapolate);
+#endif // WITH_EDITOR
+		{
+			// asking for root bone (RootBoneIndexType) in world space (WorldSpaceIndexType)
+			History->GetTransformAtTime(SampleTime, WorldRootBoneTransform, AnimInstance->CurrentSkeleton, RootBoneIndexType, WorldSpaceIndexType, bExtrapolate);
+		}
 	}
 	else
 	{
+#if WITH_EDITOR
+		if (!Trajectory)
+		{
+			UE_LOG(LogPoseSearch, Error, TEXT("FSearchContext::GetWorldRootBoneTransformAtTime - missing Trajectory is required!"));
+		}
+		else
+#else // WITH_EDITOR
 		check(Trajectory);
-		// Trajectory is already in root bone world space (transformed in UPoseSearchLibrary::ProcessTrajectory), so we just ask for a sample at the proper SampleTime 
-		WorldRootBoneTransform = Trajectory->GetSampleAtTime(SampleTime, bExtrapolate).GetTransform();
+#endif // WITH_EDITOR
+		{
+			// Trajectory is already in root bone world space (transformed in UPoseSearchLibrary::ProcessTrajectory), so we just ask for a sample at the proper SampleTime 
+			WorldRootBoneTransform = Trajectory->GetSampleAtTime(SampleTime, bExtrapolate).GetTransform();
+		}
 	}
 
 	return WorldRootBoneTransform;
diff --git a/Engine/Plugins/Experimental/Animation/PoseSearch/Source/Runtime/Private/PoseSearchDatabase.cpp b/Engine/Plugins/Experimental/Animation/PoseSearch/Source/Runtime/Private/PoseSearchDatabase.cpp
--- a/Engine/Plugins/Experimental/Animation/PoseSearch/Source/Runtime/Private/PoseSearchDatabase.cpp
+++ b/Engine/Plugins/Experimental/Animation/PoseSearch/Source/Runtime/Private/PoseSearchDatabase.cpp
@@ -546,7 +546,7 @@ bool UPoseSearchDatabase::GetSkipSearchIfPossible() const
 		return false;
 	}
 #endif // WITH_EDITOR && ENABLE_ANIM_DEBUG
-	return bSkipSearchIfPossible;
+	return true;
 }
 
 void UPoseSearchDatabase::PostLoad()
@@ -869,19 +869,16 @@ FPoseSearchCost UPoseSearchDatabase::SearchContinuingPose(UE::PoseSearch::FSearc
 	const FAnimationAssetSampler SequenceBaseSampler(DatabaseAnimationAssetBase->GetAnimationAsset(), SearchIndexAsset.BlendParameters);
 	const float SampleTime = GetNormalizedAssetTime(PoseIdx);
 
-	// @todo: change ExtractPoseSearchNotifyStates api to avoid NotifyStates allocation
-	TArray<UAnimNotifyState_PoseSearchBase*> NotifyStates;
-	SequenceBaseSampler.ExtractPoseSearchNotifyStates(SampleTime, NotifyStates);
-
 	float UpdatedContinuingPoseCostBias = ContinuingPoseCostBias;
-	for (const UAnimNotifyState_PoseSearchBase* PoseSearchNotify : NotifyStates)
-	{
-		if (const UAnimNotifyState_PoseSearchOverrideContinuingPoseCostBias* ContinuingPoseCostBiasNotify = Cast<const UAnimNotifyState_PoseSearchOverrideContinuingPoseCostBias>(PoseSearchNotify))
+	SequenceBaseSampler.ExtractPoseSearchNotifyStates(SampleTime, [&UpdatedContinuingPoseCostBias](const UAnimNotifyState_PoseSearchBase* PoseSearchNotify)
 		{
-			UpdatedContinuingPoseCostBias = ContinuingPoseCostBiasNotify->CostAddend;
-			break;
-		}
-	}
+			if (const UAnimNotifyState_PoseSearchOverrideContinuingPoseCostBias* ContinuingPoseCostBiasNotify = Cast<const UAnimNotifyState_PoseSearchOverrideContinuingPoseCostBias>(PoseSearchNotify))
+			{
+				UpdatedContinuingPoseCostBias = ContinuingPoseCostBiasNotify->CostAddend;
+				return false;
+			}
+			return true;
+		});
 
 	// since any PoseCost calculated here is at least SearchIndex.MinCostAddend + UpdatedContinuingPoseCostBias,
 	// there's no point in performing the search if CurrentBestTotalCost is already better than that
diff --git a/Engine/Plugins/Experimental/Animation/PoseSearch/Source/Runtime/Public/PoseSearch/PoseSearchAssetSampler.h b/Engine/Plugins/Experimental/Animation/PoseSearch/Source/Runtime/Public/PoseSearch/PoseSearchAssetSampler.h
--- a/Engine/Plugins/Experimental/Animation/PoseSearch/Source/Runtime/Public/PoseSearch/PoseSearchAssetSampler.h
+++ b/Engine/Plugins/Experimental/Animation/PoseSearch/Source/Runtime/Public/PoseSearch/PoseSearchAssetSampler.h
@@ -41,7 +41,7 @@ struct POSESEARCH_API FAnimationAssetSampler
 
 	// Extracts notify states inheriting from UAnimNotifyState_PoseSearchBase present in the sequence at Time.
 	// The function does not empty NotifyStates before adding new notifies!
-	void ExtractPoseSearchNotifyStates(float Time, TArray<UAnimNotifyState_PoseSearchBase*>& NotifyStates) const;
+	void ExtractPoseSearchNotifyStates(float Time, TFunction<bool(UAnimNotifyState_PoseSearchBase*)> ProcessPoseSearchBase) const;
 	
 	TConstArrayView<FAnimNotifyEvent> GetAllAnimNotifyEvents() const;
 
diff --git a/Engine/Plugins/Experimental/Animation/PoseSearch/Source/Runtime/Public/PoseSearch/PoseSearchDatabase.h b/Engine/Plugins/Experimental/Animation/PoseSearch/Source/Runtime/Public/PoseSearch/PoseSearchDatabase.h
--- a/Engine/Plugins/Experimental/Animation/PoseSearch/Source/Runtime/Public/PoseSearch/PoseSearchDatabase.h
+++ b/Engine/Plugins/Experimental/Animation/PoseSearch/Source/Runtime/Public/PoseSearch/PoseSearchDatabase.h
@@ -287,11 +287,6 @@ public:
 	UPROPERTY(EditAnywhere, Category = "Performance", meta = (DisplayName = "KNNQueryNumNeighbors", EditCondition = "PoseSearchMode == EPoseSearchMode::PCAKDTree || PoseSearchMode == VPTree", EditConditionHides, ClampMin = "1", ClampMax = "600", UIMin = "1"))
 	int32 KDTreeQueryNumNeighbors = 200;
 
-	// When evaluating multiple searches, including the continuing pose search, the system keeps track of the best pose and associated cost.
-	// if the current database cannot possibly improve the current cost, the database search will be skipped entirely.
-	UPROPERTY(EditAnywhere, Category = "Performance")
-	bool bSkipSearchIfPossible = true;
-
 	// if two poses values (multi dimensional point with the schema cardinality) are closer than PosePruningSimilarityThreshold,
 	// only one will be saved into the database FSearchIndexBase (to save memory) and accessed by the two different pose indexes
 	UPROPERTY(EditAnywhere, Category = "Performance")
