motion matching - using int as consistent type for pose and feature vectors indexes

diff --git a/Engine/Plugins/Experimental/Animation/PoseSearch/Source/Runtime/Private/KDTree.cpp b/Engine/Plugins/Experimental/Animation/PoseSearch/Source/Runtime/Private/KDTree.cpp
--- a/Engine/Plugins/Experimental/Animation/PoseSearch/Source/Runtime/Private/KDTree.cpp
+++ b/Engine/Plugins/Experimental/Animation/PoseSearch/Source/Runtime/Private/KDTree.cpp
@@ -19,7 +19,7 @@ namespace UE::PoseSearch
 {
 
 #if UE_POSE_SEARCH_USE_NANOFLANN
-using FKDTreeImplementationBase = nanoflann::KDTreeSingleIndexAdaptor<nanoflann::L2_Simple_Adaptor<float, FKDTree::FDataSource>, FKDTree::FDataSource>;
+using FKDTreeImplementationBase = nanoflann::KDTreeSingleIndexAdaptor<nanoflann::L2_Simple_Adaptor<float, FKDTree::FDataSource>, FKDTree::FDataSource, -1, AccessorType>;
 struct FKDTreeImplementation : FKDTreeImplementationBase
 {
 	using FKDTreeImplementationBase::FKDTreeImplementationBase;
@@ -36,13 +36,13 @@ struct FKDTreeImplementation : FKDTreeImplementationBase
 			return false;
 		}
 
-		const uint32 RootBBoxSize = root_bbox.size();
+		const AccessorType RootBBoxSize = root_bbox.size();
 		if (RootBBoxSize != Other.root_bbox.size())
 		{
 			return false;
 		}
 
-		for (uint32 Index = 0; Index < RootBBoxSize; ++Index)
+		for (AccessorType Index = 0; Index < RootBBoxSize; ++Index)
 		{
 			const Interval& ThisInterval = root_bbox[Index];
 			const Interval& OtherInterval = Other.root_bbox[Index];
@@ -145,7 +145,7 @@ private:
 };
 #endif
 
-FKDTree::FKDTree(int32 Count, int32 Dim, const float* Data, int32 MaxLeafSize)
+FKDTree::FKDTree(AccessorType Count, AccessorType Dim, const float* Data, AccessorType MaxLeafSize)
 : DataSource(Count, Dim, Data)
 , Impl(nullptr)
 {
@@ -207,31 +207,31 @@ FKDTree::FKDTree(const FKDTree& Other)
 
 		DataSource = Other.DataSource;
 
-		check(Other.Impl->m_size <= UINT_MAX);
+		check(Other.Impl->m_size <= AccessorTypeMax);
 		Impl->m_size = Other.Impl->m_size;
 
 		if (Impl->m_size > 0)
 		{
 			Impl->dim = Other.Impl->dim;
 
-			check(Other.Impl->root_bbox.size() <= UINT_MAX);
-			const uint32 root_bbox_size = Other.Impl->root_bbox.size();
+			check(Other.Impl->root_bbox.size() <= AccessorTypeMax);
+			const AccessorType root_bbox_size = Other.Impl->root_bbox.size();
 			Impl->root_bbox.resize(root_bbox_size);
 
-			for (uint32 i = 0; i < root_bbox_size; ++i)
+			for (AccessorType i = 0; i < root_bbox_size; ++i)
 			{
 				Impl->root_bbox[i] = Other.Impl->root_bbox[i];
 			}
 
-			check(Other.Impl->m_leaf_max_size <= UINT_MAX);
-			const uint32 KDTreeLeafMaxSize = Other.Impl->m_leaf_max_size;
+			check(Other.Impl->m_leaf_max_size <= AccessorTypeMax);
+			const AccessorType KDTreeLeafMaxSize = Other.Impl->m_leaf_max_size;
 			Impl->m_leaf_max_size = KDTreeLeafMaxSize;
 
-			check(Other.Impl->vAcc.size() <= UINT_MAX);
-			const uint32 VAccSize = Other.Impl->vAcc.size();
+			check(Other.Impl->vAcc.size() <= AccessorTypeMax);
+			const AccessorType VAccSize = Other.Impl->vAcc.size();
 			Impl->vAcc.resize(VAccSize);
 			
-			for (uint32 i = 0; i < VAccSize; ++i)
+			for (AccessorType i = 0; i < VAccSize; ++i)
 			{
 				Impl->vAcc[i] = Other.Impl->vAcc[i];
 			}
@@ -288,7 +288,7 @@ void FKDTree::Reset()
 	DataSource = FDataSource();
 }
 
-void FKDTree::Construct(int32 Count, int32 Dim, const float* Data, int32 MaxLeafSize)
+void FKDTree::Construct(AccessorType Count, AccessorType Dim, const float* Data, AccessorType MaxLeafSize)
 {
 	Reset();
 	new(this)FKDTree(Count, Dim, Data, MaxLeafSize);
@@ -373,11 +373,11 @@ FArchive& SerializeSubTree(FArchive& Ar, FKDTree& KDTree, FKDTreeImplementation:
 	const bool bIsLeafNode = !bAnyNodeChild1 && !bAnyNodeChild2;
 	if (bIsLeafNode)
 	{
-		check(KDTreeNode->node_type.lr.left <= UINT_MAX);
-		check(KDTreeNode->node_type.lr.right <= UINT_MAX);
+		check(KDTreeNode->node_type.lr.left <= AccessorTypeMax);
+		check(KDTreeNode->node_type.lr.right <= AccessorTypeMax);
 
-		uint32 OffsetLeft = KDTreeNode->node_type.lr.left;
-		uint32 OffsetRight = KDTreeNode->node_type.lr.right;
+		AccessorType OffsetLeft = KDTreeNode->node_type.lr.left;
+		AccessorType OffsetRight = KDTreeNode->node_type.lr.right;
 
 		Ar << OffsetLeft;
 		Ar << OffsetRight;
@@ -417,9 +417,9 @@ FArchive& SerializeSubTree(FArchive& Ar, FKDTree& KDTree, FKDTreeImplementation:
 FArchive& Serialize(FArchive& Ar, FKDTree& KDTree, const float* KDTreeData)
 {
 #if UE_POSE_SEARCH_USE_NANOFLANN
-	check(!KDTree.Impl || KDTree.Impl->m_size <= UINT_MAX);
+	check(!KDTree.Impl || KDTree.Impl->m_size <= AccessorTypeMax);
 
-	uint32 KDTreeSize = KDTree.Impl ? KDTree.Impl->m_size : 0;
+	AccessorType KDTreeSize = KDTree.Impl ? KDTree.Impl->m_size : 0;
 
 	Ar << KDTreeSize;
 
@@ -434,8 +434,8 @@ FArchive& Serialize(FArchive& Ar, FKDTree& KDTree, const float* KDTreeData)
 
 		Ar << KDTree.Impl->dim;
 
-		uint32 root_bbox_size = KDTree.Impl->root_bbox.size();
-		check(KDTree.Impl->root_bbox.size() <= UINT_MAX);
+		AccessorType root_bbox_size = KDTree.Impl->root_bbox.size();
+		check(KDTree.Impl->root_bbox.size() <= AccessorTypeMax);
 		Ar << root_bbox_size;
 
 		if (Ar.IsLoading())
@@ -452,19 +452,19 @@ FArchive& Serialize(FArchive& Ar, FKDTree& KDTree, const float* KDTreeData)
 			Ar.Serialize(&el, sizeof(FKDTreeImplementation::Interval));
 		}
 
-		check(KDTree.Impl->m_leaf_max_size <= UINT_MAX);
-		uint32 KDTreeLeafMaxSize = KDTree.Impl->m_leaf_max_size;
+		check(KDTree.Impl->m_leaf_max_size <= AccessorTypeMax);
+		AccessorType KDTreeLeafMaxSize = KDTree.Impl->m_leaf_max_size;
 		Ar << KDTreeLeafMaxSize;
 		KDTree.Impl->m_leaf_max_size = KDTreeLeafMaxSize;
 
-		check(KDTree.Impl->vAcc.size() <= UINT_MAX);
-		uint32 VAccSize = KDTree.Impl->vAcc.size();
+		check(KDTree.Impl->vAcc.size() <= AccessorTypeMax);
+		AccessorType VAccSize = KDTree.Impl->vAcc.size();
 		Ar << VAccSize;
 		if (Ar.IsLoading())
 		{
 			KDTree.Impl->vAcc.resize(VAccSize);
 		}
-		for (uint32_t& el : KDTree.Impl->vAcc)
+		for (AccessorType& el : KDTree.Impl->vAcc)
 		{
 			Ar << el;
 		}
diff --git a/Engine/Plugins/Experimental/Animation/PoseSearch/Source/Runtime/Private/PoseSearchDatabase.cpp b/Engine/Plugins/Experimental/Animation/PoseSearch/Source/Runtime/Private/PoseSearchDatabase.cpp
--- a/Engine/Plugins/Experimental/Animation/PoseSearch/Source/Runtime/Private/PoseSearchDatabase.cpp
+++ b/Engine/Plugins/Experimental/Animation/PoseSearch/Source/Runtime/Private/PoseSearchDatabase.cpp
@@ -39,7 +39,7 @@ static TAutoConsoleVariable<bool> CVarMotionMatchCompareAgainstBruteForce(TEXT("
 static TAutoConsoleVariable<bool> CVarMotionMatchValidateKNNSearch(TEXT("a.MotionMatch.ValidateKNNSearch"), false, TEXT("Validate KNN search"));
 #endif
 
-typedef TArray<size_t, TInlineAllocator<256>> FNonSelectableIdx;
+typedef TArray<int32, TInlineAllocator<256>> FNonSelectableIdx;
 static void PopulateNonSelectableIdx(FNonSelectableIdx& NonSelectableIdx, FSearchContext& SearchContext, const UPoseSearchDatabase* Database
 #if UE_POSE_SEARCH_TRACE_ENABLED
 	, TConstArrayView<float> QueryValues
@@ -155,7 +155,7 @@ static void PopulateNonSelectableIdx(FNonSelectableIdx& NonSelectableIdx, FSearc
 
 struct FSearchFilters
 {
-	FSearchFilters(const UPoseSearchSchema* Schema, TConstArrayView<size_t> NonSelectableIdx, bool bAnyBlockTransition)
+	FSearchFilters(const UPoseSearchSchema* Schema, TConstArrayView<int32> NonSelectableIdx, bool bAnyBlockTransition)
 	{
 		NonSelectableIdxFilter.NonSelectableIdx = NonSelectableIdx;
 
@@ -164,7 +164,7 @@ struct FSearchFilters
 			Filters.Add(&BlockTransitionFilter);
 		}
 
-		if (NonSelectableIdxFilter.IsFilterActive())
+		if (NonSelectableIdxFilter.Init(NonSelectableIdx).IsFilterActive())
 		{
 			Filters.Add(&NonSelectableIdxFilter);
 		}
@@ -213,6 +213,13 @@ struct FSearchFilters
 private:
 	struct FNonSelectableIdxFilter : public IPoseSearchFilter
 	{
+		const FNonSelectableIdxFilter& Init(TConstArrayView<int32> InNonSelectableIdx)
+		{
+			check(Algo::IsSorted(InNonSelectableIdx));
+			NonSelectableIdx = InNonSelectableIdx;
+			return *this;
+		}
+
 		virtual bool IsFilterActive() const override
 		{
 			return !NonSelectableIdx.IsEmpty();
@@ -223,7 +230,7 @@ private:
 			return Algo::BinarySearch(NonSelectableIdx, PoseIdx) == INDEX_NONE;
 		}
 
-		TConstArrayView<size_t> NonSelectableIdx;
+		TConstArrayView<int32> NonSelectableIdx;
 	};
 
 	struct FBlockTransitionFilter : public IPoseSearchFilter
@@ -808,7 +815,7 @@ UE::PoseSearch::FSearchResult UPoseSearchDatabase::Search(UE::PoseSearch::FSearc
 
 template<bool bReconstructPoseValues, bool bAlignedAndPadded>
 static inline void EvaluatePoseKernel(UE::PoseSearch::FSearchResult& Result, const UE::PoseSearch::FSearchIndex& SearchIndex, TConstArrayView<float> QueryValues, TArrayView<float> ReconstructedPoseValuesBuffer,
-	int32 PoseIdx, const UE::PoseSearch::FSearchFilters& SearchFilters, UE::PoseSearch::FSearchContext& SearchContext, const UPoseSearchDatabase* Database, bool bUpdateBestCandidates, size_t ResultIndex)
+	int32 PoseIdx, const UE::PoseSearch::FSearchFilters& SearchFilters, UE::PoseSearch::FSearchContext& SearchContext, const UPoseSearchDatabase* Database, bool bUpdateBestCandidates, int32 ResultIndex = -1)
 {
 	using namespace UE::PoseSearch;
 
@@ -936,7 +943,7 @@ UE::PoseSearch::FSearchResult UPoseSearchDatabase::SearchPCAKDTree(UE::PoseSearc
 		const bool bArePCAValuesPruned = SearchIndex.PCAValuesVectorToPoseIndexes.Num() > 0;
 
 		//stack allocated temporaries
-		TArrayView<size_t> ResultIndexes((size_t*)FMemory_Alloca((ClampedKDTreeQueryNumNeighbors + 1) * sizeof(size_t)), ClampedKDTreeQueryNumNeighbors + 1);
+		TArrayView<int32> ResultIndexes((int32*)FMemory_Alloca((ClampedKDTreeQueryNumNeighbors + 1) * sizeof(int32)), ClampedKDTreeQueryNumNeighbors + 1);
 		TArrayView<float> ResultDistanceSqr((float*)FMemory_Alloca((ClampedKDTreeQueryNumNeighbors + 1) * sizeof(float)), ClampedKDTreeQueryNumNeighbors + 1);
 		TArrayView<float> ProjectedQueryValues((float*)FMemory_Alloca(ClampedNumberOfPrincipalComponents * sizeof(float)), ClampedNumberOfPrincipalComponents);
 	
@@ -956,7 +963,7 @@ UE::PoseSearch::FSearchResult UPoseSearchDatabase::SearchPCAKDTree(UE::PoseSearc
 #endif // WITH_EDITOR && ENABLE_ANIM_DEBUG
 
 		// if bRunNonSelectableIdxPostKDTree we filter out the NonSelectableIdx after kdtree search, otherwise during kdtree search
-		FKDTree::FKNNResultSet ResultSet(ClampedKDTreeQueryNumNeighbors, ResultIndexes, ResultDistanceSqr, bRunNonSelectableIdxPostKDTree ? TConstArrayView<size_t>() : NonSelectableIdx);
+		FKDTree::FKNNResultSet ResultSet(ClampedKDTreeQueryNumNeighbors, ResultIndexes, ResultDistanceSqr, bRunNonSelectableIdxPostKDTree ? TConstArrayView<int32>() : NonSelectableIdx);
 
 		check(QueryValues.Num() == NumDimensions);
 		// projecting QueryValues into the PCA space 
@@ -986,7 +993,7 @@ UE::PoseSearch::FSearchResult UPoseSearchDatabase::SearchPCAKDTree(UE::PoseSearc
 					return A.Value < B.Value;
 				});
 
-			for (size_t ResultIndex = 0; ResultIndex < ResultSet.Num(); ++ResultIndex)
+			for (int32 ResultIndex = 0; ResultIndex < ResultSet.Num(); ++ResultIndex)
 			{
 				if (PCAValueIndexCost[ResultIndex].Key != ResultIndexes[ResultIndex])
 				{
@@ -1008,7 +1015,7 @@ UE::PoseSearch::FSearchResult UPoseSearchDatabase::SearchPCAKDTree(UE::PoseSearc
 #endif // WITH_EDITOR && ENABLE_ANIM_DEBUG
 
 		// NonSelectableIdx are already filtered out inside the kdtree search
-		const FSearchFilters SearchFilters(Schema, bRunNonSelectableIdxPostKDTree ? NonSelectableIdx : TConstArrayView<size_t>(), SearchIndex.bAnyBlockTransition);
+		const FSearchFilters SearchFilters(Schema, bRunNonSelectableIdxPostKDTree ? NonSelectableIdx : TConstArrayView<int32>(), SearchIndex.bAnyBlockTransition);
 		
 		// are the PCAValues pruned out of duplicates (multiple poses are associated with the same PCAValuesVectorIdx)
 		if (bArePCAValuesPruned)
@@ -1021,9 +1028,9 @@ UE::PoseSearch::FSearchResult UPoseSearchDatabase::SearchPCAKDTree(UE::PoseSearc
 			if (NumDimensions % 4 == 0)
 			{
 				int32 NumEvaluatePoseKernelCalls = 0;
-				for (size_t ResultIndex = 0; ResultIndex < ResultSet.Num(); ++ResultIndex)
+				for (int32 ResultIndex = 0; ResultIndex < ResultSet.Num(); ++ResultIndex)
 				{
-					const TConstArrayView<uint32> PoseIndexes = SearchIndex.PCAValuesVectorToPoseIndexes[ResultIndexes[ResultIndex]];
+					const TConstArrayView<int32> PoseIndexes = SearchIndex.PCAValuesVectorToPoseIndexes[ResultIndexes[ResultIndex]];
 					for (int32 Index = 0; Index < PoseIndexes.Num() && NumEvaluatePoseKernelCalls < MaxNumEvaluatePoseKernelCalls; ++Index, ++NumEvaluatePoseKernelCalls)
 					{
 						EvaluatePoseKernel<false, true>(Result, SearchIndex, QueryValues, TArrayView<float>(), PoseIndexes[Index], SearchFilters, SearchContext, this, true, ResultIndex);
@@ -1033,9 +1040,9 @@ UE::PoseSearch::FSearchResult UPoseSearchDatabase::SearchPCAKDTree(UE::PoseSearc
 			else
 			{
 				int32 NumEvaluatePoseKernelCalls = 0;
-				for (size_t ResultIndex = 0; ResultIndex < ResultSet.Num(); ++ResultIndex)
+				for (int32 ResultIndex = 0; ResultIndex < ResultSet.Num(); ++ResultIndex)
 				{
-					const TConstArrayView<uint32> PoseIndexes = SearchIndex.PCAValuesVectorToPoseIndexes[ResultIndexes[ResultIndex]];
+					const TConstArrayView<int32> PoseIndexes = SearchIndex.PCAValuesVectorToPoseIndexes[ResultIndexes[ResultIndex]];
 					for (int32 Index = 0; Index < PoseIndexes.Num() && NumEvaluatePoseKernelCalls < MaxNumEvaluatePoseKernelCalls; ++Index, ++NumEvaluatePoseKernelCalls)
 					{
 						EvaluatePoseKernel<false, false>(Result, SearchIndex, QueryValues, TArrayView<float>(), PoseIndexes[Index], SearchFilters, SearchContext, this, true, ResultIndex);
@@ -1049,7 +1056,7 @@ UE::PoseSearch::FSearchResult UPoseSearchDatabase::SearchPCAKDTree(UE::PoseSearc
 			// FMemory_Alloca is forced 16 bytes aligned
 			TArrayView<float> ReconstructedPoseValuesBuffer((float*)FMemory_Alloca(NumDimensions * sizeof(float)), NumDimensions);
 			check(IsAligned(ReconstructedPoseValuesBuffer.GetData(), alignof(VectorRegister4Float)));
-			for (size_t ResultIndex = 0; ResultIndex < ResultSet.Num(); ++ResultIndex)
+			for (int32 ResultIndex = 0; ResultIndex < ResultSet.Num(); ++ResultIndex)
 			{
 				EvaluatePoseKernel<true, false>(Result, SearchIndex, QueryValues, ReconstructedPoseValuesBuffer, ResultIndexes[ResultIndex], SearchFilters, SearchContext, this, true, ResultIndex);
 			}
@@ -1057,7 +1064,7 @@ UE::PoseSearch::FSearchResult UPoseSearchDatabase::SearchPCAKDTree(UE::PoseSearc
 		// is the data padded at 16 bytes (and 16 bytes aligned by construction)?
 		else if (NumDimensions % 4 == 0)
 		{
-			for (size_t ResultIndex = 0; ResultIndex < ResultSet.Num(); ++ResultIndex)
+			for (int32 ResultIndex = 0; ResultIndex < ResultSet.Num(); ++ResultIndex)
 			{
 				EvaluatePoseKernel<false, true>(Result, SearchIndex, QueryValues, TArrayView<float>(), ResultIndexes[ResultIndex], SearchFilters, SearchContext, this, true, ResultIndex);
 			}
@@ -1065,7 +1072,7 @@ UE::PoseSearch::FSearchResult UPoseSearchDatabase::SearchPCAKDTree(UE::PoseSearc
 		// no reconstruction, but data is not 16 bytes padded
 		else
 		{
-			for (size_t ResultIndex = 0; ResultIndex < ResultSet.Num(); ++ResultIndex)
+			for (int32 ResultIndex = 0; ResultIndex < ResultSet.Num(); ++ResultIndex)
 			{
 				EvaluatePoseKernel<false, false>(Result, SearchIndex, QueryValues, TArrayView<float>(), ResultIndexes[ResultIndex], SearchFilters, SearchContext, this, true, ResultIndex);
 			}
@@ -1118,7 +1125,6 @@ UE::PoseSearch::FSearchResult UPoseSearchDatabase::SearchVPTree(UE::PoseSearch::
 			, QueryValues
 #endif // UE_POSE_SEARCH_TRACE_ENABLED
 		);
-		check(Algo::IsSorted(NonSelectableIdx));
 
 		const int32 NumDimensions = Schema->SchemaCardinality;
 		check(QueryValues.Num() == NumDimensions);
@@ -1142,7 +1148,7 @@ UE::PoseSearch::FSearchResult UPoseSearchDatabase::SearchVPTree(UE::PoseSearch::
 				const FIndexDistance& IndexDistance = UnsortedResults[ResultIndex];
 
 				// @todo: IndexDistance.Distance is the Sqrt(DissimilarityCost), so there's no need to calculate it again in SearchIndex.ComparePoses
-				const TConstArrayView<uint32> PoseIndexes = SearchIndex.ValuesVectorToPoseIndexes[IndexDistance.Index];
+				const TConstArrayView<int32> PoseIndexes = SearchIndex.ValuesVectorToPoseIndexes[IndexDistance.Index];
 				for (int32 Index = 0; Index < PoseIndexes.Num() && NumEvaluatePoseKernelCalls < MaxNumEvaluatePoseKernelCalls; ++Index, ++NumEvaluatePoseKernelCalls)
 				{
 					EvaluatePoseKernel<false, false>(Result, SearchIndex, QueryValues, TArrayView<float>(), PoseIndexes[Index], SearchFilters, SearchContext, this, true, ResultIndex);
@@ -1206,7 +1212,6 @@ UE::PoseSearch::FSearchResult UPoseSearchDatabase::SearchBruteForce(UE::PoseSear
 			, QueryValues
 #endif // UE_POSE_SEARCH_TRACE_ENABLED
 		);
-		check(Algo::IsSorted(NonSelectableIdx));
 
 		const int32 NumDimensions = Schema->SchemaCardinality;
 		const FSearchFilters SearchFilters(Schema, NonSelectableIdx, SearchIndex.bAnyBlockTransition);
diff --git a/Engine/Plugins/Experimental/Animation/PoseSearch/Source/Runtime/Private/PoseSearchDerivedData.cpp b/Engine/Plugins/Experimental/Animation/PoseSearch/Source/Runtime/Private/PoseSearchDerivedData.cpp
--- a/Engine/Plugins/Experimental/Animation/PoseSearch/Source/Runtime/Private/PoseSearchDerivedData.cpp
+++ b/Engine/Plugins/Experimental/Animation/PoseSearch/Source/Runtime/Private/PoseSearchDerivedData.cpp
@@ -466,7 +466,7 @@ static void PreprocessSearchIndexWeights(FSearchIndex& SearchIndex, const UPoseS
 }
 
 // it calculates Mean, PCAValues, and PCAProjectionMatrix
-static void PreprocessSearchIndexPCAData(FSearchIndex& SearchIndex, int32 NumDimensions, uint32 NumberOfPrincipalComponents, EPoseSearchMode PoseSearchMode)
+static void PreprocessSearchIndexPCAData(FSearchIndex& SearchIndex, int32 NumDimensions, int32 NumberOfPrincipalComponents, EPoseSearchMode PoseSearchMode)
 {
 	// binding SearchIndex.Values and SearchIndex.PCAValues Eigen row major matrix maps
 	const int32 NumPoses = SearchIndex.GetNumPoses();
@@ -531,13 +531,13 @@ static void PreprocessSearchIndexPCAData(FSearchIndex& SearchIndex, int32 NumDim
 
 		// sorting EigenVectors by EigenValues, so we pick the most significant ones to compose our PCA projection matrix.
 		const RowMajorVector EigenValues = EigenSolver.eigenvalues().real();
-		TArray<size_t> Indexer;
+		TArray<int32> Indexer;
 		Indexer.Reserve(NumDimensions);
-		for (size_t DimensionIndex = 0; DimensionIndex < NumDimensions; ++DimensionIndex)
+		for (int32 DimensionIndex = 0; DimensionIndex < NumDimensions; ++DimensionIndex)
 		{
 			Indexer.Push(DimensionIndex);
 		}
-		Indexer.Sort([&EigenValues](size_t a, size_t b)
+		Indexer.Sort([&EigenValues](int32 a, int32 b)
 		{
 			return EigenValues[a] > EigenValues[b];
 		});
@@ -545,7 +545,7 @@ static void PreprocessSearchIndexPCAData(FSearchIndex& SearchIndex, int32 NumDim
 		// composing the PCA projection matrix with the PCANumComponents most significant EigenVectors
 		ColMajorMatrixMap PCAProjectionMatrix(SearchIndex.PCAProjectionMatrix.GetData(), NumDimensions, NumberOfPrincipalComponents);
 		float AccumulatedVariance = 0.f;
-		for (size_t PCAComponentIndex = 0; PCAComponentIndex < NumberOfPrincipalComponents; ++PCAComponentIndex)
+		for (int32 PCAComponentIndex = 0; PCAComponentIndex < NumberOfPrincipalComponents; ++PCAComponentIndex)
 		{
 			PCAProjectionMatrix.col(PCAComponentIndex) = EigenVectors.col(Indexer[PCAComponentIndex]);
 			AccumulatedVariance += EigenValues[Indexer[PCAComponentIndex]];
@@ -628,20 +628,20 @@ static void PreprocessSearchIndexKDTree(FSearchIndex& SearchIndex, const UPoseSe
 			// testing the KDTree is returning the proper searches for all the points in pca space
 			const int32 KDTreeQueryNumNeighbors = Database->KDTreeQueryNumNeighbors;
 
-			TArray<size_t> ResultIndexes;
+			TArray<int32> ResultIndexes;
 			TArray<float> ResultDistanceSqr;
 			ResultIndexes.SetNum(NumPCAValuesVectors + 1);
 			ResultDistanceSqr.SetNum(NumPCAValuesVectors + 1);
 
-			size_t MaxNumNeighborToFindAPoint = 0;
-			for (size_t PointIndex = 0; PointIndex < NumPCAValuesVectors; ++PointIndex)
+			int32 MaxNumNeighborToFindAPoint = 0;
+			for (int32 PointIndex = 0; PointIndex < NumPCAValuesVectors; ++PointIndex)
 			{
 				// searching the kdtree for PointIndex
 				FKDTree::FRadiusResultSet ResultSet(UE_SMALL_NUMBER, NumPCAValuesVectors, ResultIndexes, ResultDistanceSqr);
 				SearchIndex.KDTree.FindNeighbors(ResultSet, MakeArrayView(&SearchIndex.PCAValues[PointIndex * NumberOfPrincipalComponents], NumberOfPrincipalComponents));
 
 				bool bFound = false;
-				for (size_t ResultIndex = 0; ResultIndex < ResultSet.Num(); ++ResultIndex)
+				for (int32 ResultIndex = 0; ResultIndex < ResultSet.Num(); ++ResultIndex)
 				{
 					if (PointIndex == ResultIndexes[ResultIndex])
 					{
@@ -671,12 +671,12 @@ static void PreprocessSearchIndexKDTree(FSearchIndex& SearchIndex, const UPoseSe
 			{
 				// testing the KDTree is returning the proper searches for all the original points transformed in pca space
 				TArrayView<float> ProjectedValues((float*)FMemory_Alloca(NumberOfPrincipalComponents * sizeof(float)), NumberOfPrincipalComponents);
-				for (size_t PointIndex = 0; PointIndex < NumPCAValuesVectors; ++PointIndex)
+				for (int32 PointIndex = 0; PointIndex < NumPCAValuesVectors; ++PointIndex)
 				{
 					FKDTree::FKNNResultSet ResultSet(KDTreeQueryNumNeighbors, ResultIndexes, ResultDistanceSqr);
 					SearchIndex.KDTree.FindNeighbors(ResultSet, SearchIndex.PCAProject(SearchIndex.GetPoseValuesBase(PointIndex, NumDimensions), ProjectedValues));
 
-					size_t ResultIndex = 0;
+					int32 ResultIndex = 0;
 					for (; ResultIndex < ResultSet.Num(); ++ResultIndex)
 					{
 						if (PointIndex == ResultIndexes[ResultIndex])
diff --git a/Engine/Plugins/Experimental/Animation/PoseSearch/Source/Runtime/Private/PoseSearchDerivedDataKey.cpp b/Engine/Plugins/Experimental/Animation/PoseSearch/Source/Runtime/Private/PoseSearchDerivedDataKey.cpp
--- a/Engine/Plugins/Experimental/Animation/PoseSearch/Source/Runtime/Private/PoseSearchDerivedDataKey.cpp
+++ b/Engine/Plugins/Experimental/Animation/PoseSearch/Source/Runtime/Private/PoseSearchDerivedDataKey.cpp
@@ -40,7 +40,7 @@ FKeyBuilder::FKeyBuilder(const UObject* Object, bool bUseDataVer, bool bPerformC
 	if (bUseDataVer)
 	{
 		// used to invalidate the key without having to change POSESEARCHDB_DERIVEDDATA_VER all the times
-		int32 POSESEARCHDB_DERIVEDDATA_VER_SMALL = 199;
+		int32 POSESEARCHDB_DERIVEDDATA_VER_SMALL = 217;
 		FGuid VersionGuid = FDevSystemGuids::GetSystemGuid(FDevSystemGuids::Get().POSESEARCHDB_DERIVEDDATA_VER);
 
 		*this << VersionGuid;
diff --git a/Engine/Plugins/Experimental/Animation/PoseSearch/Source/Runtime/Private/PoseSearchIndex.cpp b/Engine/Plugins/Experimental/Animation/PoseSearch/Source/Runtime/Private/PoseSearchIndex.cpp
--- a/Engine/Plugins/Experimental/Animation/PoseSearch/Source/Runtime/Private/PoseSearchIndex.cpp
+++ b/Engine/Plugins/Experimental/Animation/PoseSearch/Source/Runtime/Private/PoseSearchIndex.cpp
@@ -116,7 +116,7 @@ static bool CalculateSimilarities(TArray<FPosePairSimilarity>& PosePairSimilarit
 		check(Values.Num() == NumPoses * DataCardinality);
 		FKDTree KDTree(NumPoses, DataCardinality, Values.GetData());
 
-		TArray<size_t> ResultIndexes;
+		TArray<int32> ResultIndexes;
 		TArray<float> ResultDistanceSqr;
 		ResultIndexes.SetNum(NumPoses + 1);
 		ResultDistanceSqr.SetNum(NumPoses + 1);
@@ -129,7 +129,7 @@ static bool CalculateSimilarities(TArray<FPosePairSimilarity>& PosePairSimilarit
 			FKDTree::FRadiusResultSet ResultSet(SimilarityThreshold, NumPoses, ResultIndexes, ResultDistanceSqr);
 			KDTree.FindNeighbors(ResultSet, ValuesA);
 
-			for (size_t ResultIndex = 0; ResultIndex < ResultSet.Num(); ++ResultIndex)
+			for (int32 ResultIndex = 0; ResultIndex < ResultSet.Num(); ++ResultIndex)
 			{
 				const int32 ResultPoseIdx = ResultIndexes[ResultIndex];
 				if (PoseIdx != ResultPoseIdx)
@@ -390,7 +390,7 @@ void FSearchIndexBase::Reset()
 
 void FSearchIndexBase::PruneDuplicateValues(float SimilarityThreshold, int32 DataCardinality, bool bDoNotGenerateValuesVectorToPoseIndexes)
 {
-	ValuesVectorToPoseIndexes = FSparsePoseMultiMap<uint32>();
+	ValuesVectorToPoseIndexes = FSparsePoseMultiMap<int32>();
 
 	const int32 NumPoses = GetNumPoses();
 	if (SimilarityThreshold > 0.f && NumPoses >= 2)
@@ -406,7 +406,7 @@ void FSearchIndexBase::PruneDuplicateValues(float SimilarityThreshold, int32 Dat
 
 		if (!bDoNotGenerateValuesVectorToPoseIndexes)
 		{
-			TMap<uint32, TArray<uint32>> ValuesVectorToPoseIndexesMap;
+			TMap<int32, TArray<int32>> ValuesVectorToPoseIndexesMap;
 			ValuesVectorToPoseIndexesMap.Reserve(NumPoses);
 			for (int32 PoseIdx = 0; PoseIdx < NumPoses; ++PoseIdx)
 			{
@@ -414,32 +414,32 @@ void FSearchIndexBase::PruneDuplicateValues(float SimilarityThreshold, int32 Dat
 				check(Metadata.GetValueOffset() % DataCardinality == 0);
 				const int32 ValuesVectorIdx = Metadata.GetValueOffset() / DataCardinality;
 
-				TArray<uint32>& PoseIndexes = ValuesVectorToPoseIndexesMap.FindOrAdd(ValuesVectorIdx);
+				TArray<int32>& PoseIndexes = ValuesVectorToPoseIndexesMap.FindOrAdd(ValuesVectorIdx);
 				PoseIndexes.Add(PoseIdx);
 			}
 
 			// sorting ValuesVectorToPoseIndexesMap keys to create a deterministic FSparsePoseMultiMap later on
 			// we're not using TSortedMap for performance reasons, because ValuesVectorToPoseIndexesMap can be quite big
-			TArray<uint32> SortedKeys;
+			TArray<int32> SortedKeys;
 			SortedKeys.Reserve(ValuesVectorToPoseIndexesMap.Num());
-			for (const TPair<uint32, TArray<uint32>>& Pair : ValuesVectorToPoseIndexesMap)
+			for (const TPair<int32, TArray<int32>>& Pair : ValuesVectorToPoseIndexesMap)
 			{
 				SortedKeys.Add(Pair.Key);
 			}
 			SortedKeys.Sort();
 
-			FSparsePoseMultiMap<uint32> SparsePoseMultiMap(ValuesVectorToPoseIndexesMap.Num(), NumPoses - 1);
-			for (const uint32& Key : SortedKeys)
+			FSparsePoseMultiMap<int32> SparsePoseMultiMap(ValuesVectorToPoseIndexesMap.Num(), NumPoses - 1);
+			for (const int32& Key : SortedKeys)
 			{
-				const uint32 PCAValuesVectorIdx = Key;
-				const TArray<uint32>& PoseIndexes = ValuesVectorToPoseIndexesMap[Key];
+				const int32 PCAValuesVectorIdx = Key;
+				const TArray<int32>& PoseIndexes = ValuesVectorToPoseIndexesMap[Key];
 				SparsePoseMultiMap.Insert(PCAValuesVectorIdx, PoseIndexes);
 			}
 
-			for (uint32 ValuesVectorIdx = 0; ValuesVectorIdx < SparsePoseMultiMap.Num(); ++ValuesVectorIdx)
+			for (int32 ValuesVectorIdx = 0; ValuesVectorIdx < SparsePoseMultiMap.Num(); ++ValuesVectorIdx)
 			{
-				const TConstArrayView<uint32> PoseIndexes = SparsePoseMultiMap[ValuesVectorIdx];
-				const TArray<uint32>& TestPoseIndexes = ValuesVectorToPoseIndexesMap[ValuesVectorIdx];
+				const TConstArrayView<int32> PoseIndexes = SparsePoseMultiMap[ValuesVectorIdx];
+				const TArray<int32>& TestPoseIndexes = ValuesVectorToPoseIndexesMap[ValuesVectorIdx];
 				check(PoseIndexes == TestPoseIndexes);
 			}
 
@@ -582,7 +582,7 @@ TConstArrayView<float> FSearchIndex::PCAProject(TConstArrayView<float> PoseValue
 
 void FSearchIndex::PruneDuplicatePCAValues(float SimilarityThreshold, int32 NumberOfPrincipalComponents)
 {
-	PCAValuesVectorToPoseIndexes = FSparsePoseMultiMap<uint32>();
+	PCAValuesVectorToPoseIndexes = FSparsePoseMultiMap<int32>();
 
 	const uint32 NumPoses = GetNumPoses();
 	if (SimilarityThreshold > 0.f && NumPoses >= 2 && NumberOfPrincipalComponents > 0)
@@ -592,7 +592,7 @@ void FSearchIndex::PruneDuplicatePCAValues(float SimilarityThreshold, int32 Numb
 		// so far we support only pruning an original PCAValues set, where there's a 1:1 mapping between PCAValuesVectors and Poses
 		check(NumPCAValuesVectors == NumPoses);
 
-		TArray<uint32> PoseToPCAValueOffset;
+		TArray<int32> PoseToPCAValueOffset;
 		PoseToPCAValueOffset.AddUninitialized(NumPoses);
 		for (uint32 PoseIdx = 0; PoseIdx < NumPoses; ++PoseIdx)
 		{
@@ -605,42 +605,42 @@ void FSearchIndex::PruneDuplicatePCAValues(float SimilarityThreshold, int32 Numb
 		{
 			if (PruneValues(NumberOfPrincipalComponents, NumPoses, PosePairSimilarities, PCAValues,
 				[&PoseToPCAValueOffset](int32 PoseIdx) { return PoseToPCAValueOffset[PoseIdx]; },
-				[&PoseToPCAValueOffset](int32 PoseIdx, uint32 ValueOffset) { PoseToPCAValueOffset[PoseIdx] = ValueOffset; }))
+				[&PoseToPCAValueOffset](int32 PoseIdx, int32 ValueOffset) { PoseToPCAValueOffset[PoseIdx] = ValueOffset; }))
 			{
 				// we pruned some PCAValues: we need to construct a mapping between PCAValuesVectorIdx to PoseIdx(s)
-				TMap<uint32, TArray<uint32>> PCAValuesVectorToPoseIndexesMap;
+				TMap<int32, TArray<int32>> PCAValuesVectorToPoseIndexesMap;
 				PCAValuesVectorToPoseIndexesMap.Reserve(NumPoses);
 				for (uint32 PoseIdx = 0; PoseIdx < NumPoses; ++PoseIdx)
 				{
 					check(PoseToPCAValueOffset[PoseIdx] % NumberOfPrincipalComponents == 0);
-					const uint32 PCAValuesVectorIdx = PoseToPCAValueOffset[PoseIdx] / NumberOfPrincipalComponents;
-					TArray<uint32>& PoseIndexes = PCAValuesVectorToPoseIndexesMap.FindOrAdd(PCAValuesVectorIdx);
+					const int32 PCAValuesVectorIdx = PoseToPCAValueOffset[PoseIdx] / NumberOfPrincipalComponents;
+					TArray<int32>& PoseIndexes = PCAValuesVectorToPoseIndexesMap.FindOrAdd(PCAValuesVectorIdx);
 					check(!PoseIndexes.Contains(PoseIdx));
 					PoseIndexes.Add(PoseIdx);
 				}
 
 				// sorting PCAValuesVectorToPoseIndexesMap keys to create a deterministic FSparsePoseMultiMap later on
 				// we're not using TSortedMap for performance reasons, because PCAValuesVectorToPoseIndexesMap can be quite big
-				TArray<uint32> SortedKeys;
+				TArray<int32> SortedKeys;
 				SortedKeys.Reserve(PCAValuesVectorToPoseIndexesMap.Num());
-				for (const TPair<uint32, TArray<uint32>>& Pair : PCAValuesVectorToPoseIndexesMap)
+				for (const TPair<int32, TArray<int32>>& Pair : PCAValuesVectorToPoseIndexesMap)
 				{
 					SortedKeys.Add(Pair.Key);
 				}
 				SortedKeys.Sort();
 
-				FSparsePoseMultiMap<uint32> SparsePoseMultiMap(PCAValuesVectorToPoseIndexesMap.Num(), NumPoses - 1);
-				for (const uint32& Key : SortedKeys)
+				FSparsePoseMultiMap<int32> SparsePoseMultiMap(PCAValuesVectorToPoseIndexesMap.Num(), NumPoses - 1);
+				for (const int32& Key : SortedKeys)
 				{
-					const uint32 PCAValuesVectorIdx = Key;
-					const TArray<uint32>& PoseIndexes = PCAValuesVectorToPoseIndexesMap[Key];
+					const int32 PCAValuesVectorIdx = Key;
+					const TArray<int32>& PoseIndexes = PCAValuesVectorToPoseIndexesMap[Key];
 					SparsePoseMultiMap.Insert(PCAValuesVectorIdx, PoseIndexes);
 				}
 
-				for (uint32 PCAValuesVectorIdx = 0; PCAValuesVectorIdx < SparsePoseMultiMap.Num(); ++PCAValuesVectorIdx)
+				for (int32 PCAValuesVectorIdx = 0; PCAValuesVectorIdx < SparsePoseMultiMap.Num(); ++PCAValuesVectorIdx)
 				{
-					const TConstArrayView<uint32> PoseIndexes = SparsePoseMultiMap[PCAValuesVectorIdx];
-					const TArray<uint32>& TestPoseIndexes = PCAValuesVectorToPoseIndexesMap[PCAValuesVectorIdx];
+					const TConstArrayView<int32> PoseIndexes = SparsePoseMultiMap[PCAValuesVectorIdx];
+					const TArray<int32>& TestPoseIndexes = PCAValuesVectorToPoseIndexesMap[PCAValuesVectorIdx];
 					check(PoseIndexes == TestPoseIndexes);
 				}
 
@@ -714,16 +714,16 @@ void FSearchIndex::GetPoseToPCAValuesVectorIndexes(TArray<uint32>& PoseToPCAValu
 {
 	if (PCAValuesVectorToPoseIndexes.Num() > 0)
 	{
-		PoseToPCAValuesVectorIndexes.Init(uint32(INDEX_NONE), PCAValuesVectorToPoseIndexes.MaxValue + 1);
-		for (uint32 PCAValuesVectorIdx = 0; PCAValuesVectorIdx < PCAValuesVectorToPoseIndexes.Num(); ++PCAValuesVectorIdx)
+		PoseToPCAValuesVectorIndexes.Init(INDEX_NONE, PCAValuesVectorToPoseIndexes.MaxValue + 1);
+		for (int32 PCAValuesVectorIdx = 0; PCAValuesVectorIdx < PCAValuesVectorToPoseIndexes.Num(); ++PCAValuesVectorIdx)
 		{
-			for (uint32 PoseIdx : PCAValuesVectorToPoseIndexes[PCAValuesVectorIdx])
+			for (int32 PoseIdx : PCAValuesVectorToPoseIndexes[PCAValuesVectorIdx])
 			{
 				PoseToPCAValuesVectorIndexes[PoseIdx] = PCAValuesVectorIdx;
 			}
 		}
 
-		check(!PoseToPCAValuesVectorIndexes.Contains(uint32(INDEX_NONE)));
+		check(!PoseToPCAValuesVectorIndexes.Contains(INDEX_NONE));
 	}
 	else
 	{
diff --git a/Engine/Plugins/Experimental/Animation/PoseSearch/Source/Runtime/Public/PoseSearch/KDTree.h b/Engine/Plugins/Experimental/Animation/PoseSearch/Source/Runtime/Public/PoseSearch/KDTree.h
--- a/Engine/Plugins/Experimental/Animation/PoseSearch/Source/Runtime/Public/PoseSearch/KDTree.h
+++ b/Engine/Plugins/Experimental/Animation/PoseSearch/Source/Runtime/Public/PoseSearch/KDTree.h
@@ -10,13 +10,16 @@ class FArchive;
 
 namespace UE::PoseSearch
 {
+enum { AccessorTypeMax = INT_MAX };
+typedef int32 AccessorType;
+
 struct FKDTreeImplementation;
 
 struct FKDTree
 {
 	struct FDataSource
 	{
-		FDataSource(int32 pointCount, int32 pointDim, const float* data)
+		FDataSource(AccessorType pointCount, AccessorType pointDim, const float* data)
 		: PointCount((pointDim > 0 && data) ? pointCount : 0)
 		, PointDim(pointDim)
 		, Data(data)
@@ -47,8 +50,8 @@ struct FKDTree
 				return false;
 			}
 
-			const int32 DataSize = PointCount * PointDim;
-			for (int32 DataIndex = 0; DataIndex < DataSize; ++DataIndex)
+			const AccessorType DataSize = PointCount * PointDim;
+			for (AccessorType DataIndex = 0; DataIndex < DataSize; ++DataIndex)
 			{
 				if (Data[DataIndex] != Other.Data[DataIndex])
 				{
@@ -60,10 +63,10 @@ struct FKDTree
 		}
 
 		// Must return the number of data points
-		inline size_t kdtree_get_point_count() const { return PointCount; }
+		inline AccessorType kdtree_get_point_count() const { return PointCount; }
 
 		// this method is called by nanoflann::findNeighbors -> nanoflann::searchLevel, so it has to be as fast as possible
-		inline float kdtree_get_pt(const size_t idx, const size_t dim) const { return Data[idx * PointDim + dim]; }
+		inline float kdtree_get_pt(const AccessorType idx, const AccessorType dim) const { return Data[idx * PointDim + dim]; }
 
 		// Optional bounding-box computation: return false to default to a standard
 		// bbox computation loop.
@@ -72,14 +75,14 @@ struct FKDTree
 		//   find out the expected dimensionality (e.g. 2 or 3 for point clouds)
 		template <class BBOX> bool kdtree_get_bbox(BBOX& /* bb */) const { return false; }
 
-		int32 PointCount = 0;
-		int32 PointDim = 0;
+		AccessorType PointCount = 0;
+		AccessorType PointDim = 0;
 		const float* Data = nullptr;
 	};
 
 	struct FKNNResultSet
 	{
-		inline FKNNResultSet(size_t InNumNeighbors, TArrayView<size_t> InIndexes, TArrayView<float> InDistances, TConstArrayView<size_t> InExcludeFromSearchIndexes = TConstArrayView<size_t>())
+		inline FKNNResultSet(AccessorType InNumNeighbors, TArrayView<AccessorType> InIndexes, TArrayView<float> InDistances, TConstArrayView<AccessorType> InExcludeFromSearchIndexes = TConstArrayView<AccessorType>())
 		: Indexes(InIndexes)
 		, Distances(InDistances)
 		, NumNeighbors(InNumNeighbors)
@@ -95,7 +98,7 @@ struct FKDTree
 			Distances[NumNeighbors - 1] = UE_BIG_NUMBER;
 		}
 
-		inline size_t Num() const
+		inline AccessorType Num() const
 		{
 			return Count;
 		}
@@ -105,7 +108,7 @@ struct FKDTree
 			return Count == NumNeighbors;
 		}
 
-		inline bool addPoint(float dist, size_t index)
+		inline bool addPoint(float dist, AccessorType index)
 		{
 			if (Algo::BinarySearch(ExcludeFromSearchIndexes, index) != INDEX_NONE)
 			{
@@ -113,7 +116,7 @@ struct FKDTree
 			}
 
 			// shifting Distances[i] and Indexes[i] to make space for "dist" and "index" at the right "i"th slot
-			size_t i;
+			AccessorType i;
 			for (i = Count; (i > 0) && (Distances[i - 1] > dist); --i)
 			{
 				// no need to check "if (i < capacity)" since dists and indices can contains more items than capacity_ 
@@ -138,16 +141,16 @@ struct FKDTree
 		inline float worstDist() const { return Distances[NumNeighbors - 1]; }
 
 	private:
-		TArrayView<size_t> Indexes;
+		TArrayView<AccessorType> Indexes;
 		TArrayView<float> Distances;
-		size_t NumNeighbors;
-		size_t Count;
-		TConstArrayView<size_t> ExcludeFromSearchIndexes; // sorted array view
+		AccessorType NumNeighbors;
+		AccessorType Count;
+		TConstArrayView<AccessorType> ExcludeFromSearchIndexes; // sorted array view
 	};
 
 	struct FRadiusResultSet
 	{
-		inline FRadiusResultSet(float Radius, size_t MaxNumNeighbors, TArrayView<size_t> InIndexes, TArrayView<float> InDistances)
+		inline FRadiusResultSet(float Radius, AccessorType MaxNumNeighbors, TArrayView<AccessorType> InIndexes, TArrayView<float> InDistances)
 		: Indexes(InIndexes)
 		, Distances(InDistances)
 		, NumNeighbors(MaxNumNeighbors)
@@ -161,7 +164,7 @@ struct FKDTree
 			Distances[NumNeighbors - 1] = Radius;
 		}
 
-		inline size_t Num() const
+		inline AccessorType Num() const
 		{
 			return Count;
 		}
@@ -171,12 +174,12 @@ struct FKDTree
 			return Count == NumNeighbors;
 		}
 
-		inline bool addPoint(float dist, size_t index)
+		inline bool addPoint(float dist, AccessorType index)
 		{
 			if (dist < worstDist())
 			{
 				// shifting Distances[i] and Indexes[i] to make space for "dist" and "index" at the right "i"th slot
-				size_t i;
+				AccessorType i;
 				for (i = Count; (i > 0) && (Distances[i - 1] > dist); --i)
 				{
 					// no need to check "if (i < capacity)" since dists and indices can contains more items than capacity_ 
@@ -202,14 +205,14 @@ struct FKDTree
 		inline float worstDist() const { return Distances[NumNeighbors - 1]; }
 
 	private:
-		TArrayView<size_t> Indexes;
+		TArrayView<AccessorType> Indexes;
 		TArrayView<float> Distances;
-		size_t NumNeighbors;
-		size_t Count;
+		AccessorType NumNeighbors;
+		AccessorType Count;
 	};
 	
 
-	FKDTree(int32 Count, int32 Dim, const float* Data, int32 MaxLeafSize = 16);
+	FKDTree(AccessorType Count, AccessorType Dim, const float* Data, AccessorType MaxLeafSize = 16);
 	FKDTree();
 	FKDTree(const FKDTree& Other);
 	FKDTree(FKDTree&& Other) = delete;
@@ -221,7 +224,7 @@ struct FKDTree
 	bool operator==(const FKDTree& Other) const;
 	
 	void Reset();
-	void Construct(int32 Count, int32 Dim, const float* Data, int32 MaxLeafSize = 16);
+	void Construct(AccessorType Count, AccessorType Dim, const float* Data, AccessorType MaxLeafSize = 16);
 	bool FindNeighbors(FKNNResultSet& Result, TConstArrayView<float> Query) const;
 	bool FindNeighbors(FRadiusResultSet& Result, TConstArrayView<float> Query) const;
 	POSESEARCH_API SIZE_T GetAllocatedSize() const;
diff --git a/Engine/Plugins/Experimental/Animation/PoseSearch/Source/Runtime/Public/PoseSearch/PoseSearchIndex.h b/Engine/Plugins/Experimental/Animation/PoseSearch/Source/Runtime/Public/PoseSearch/PoseSearchIndex.h
--- a/Engine/Plugins/Experimental/Animation/PoseSearch/Source/Runtime/Public/PoseSearch/PoseSearchIndex.h
+++ b/Engine/Plugins/Experimental/Animation/PoseSearch/Source/Runtime/Public/PoseSearch/PoseSearchIndex.h
@@ -188,7 +188,7 @@ struct FSearchStats
 };
 
 // compact representation of an array of arrays
-template <typename Type = uint32>
+template <typename Type = int32>
 struct FSparsePoseMultiMap
 {
 	FSparsePoseMultiMap(Type InMaxKey = Type(0), Type InMaxValue = Type(0))
@@ -309,7 +309,7 @@ struct FSparsePoseMultiMap
 struct FSearchIndexBase
 {
 	TAlignedArray<float> Values;
-	FSparsePoseMultiMap<uint32> ValuesVectorToPoseIndexes;
+	FSparsePoseMultiMap<int32> ValuesVectorToPoseIndexes;
 	TAlignedArray<FPoseMetadata> PoseMetadata;
 	bool bAnyBlockTransition = false;
 	TAlignedArray<FSearchIndexAsset> Assets;
@@ -375,7 +375,7 @@ struct FSearchIndex : public FSearchIndexBase
 	// since (VA - VB).square() could lead to big numbers, and VW being multiplied by the variance of the dataset
 	TAlignedArray<float> WeightsSqrt;
 	TAlignedArray<float> PCAValues;
-	FSparsePoseMultiMap<uint32> PCAValuesVectorToPoseIndexes;
+	FSparsePoseMultiMap<int32> PCAValuesVectorToPoseIndexes;
 	TAlignedArray<float> PCAProjectionMatrix;
 	TAlignedArray<float> Mean;
 
