motion matching - FSearchContext::FBestPoseCandidates multiple database support fix

diff --git a/Engine/Plugins/Experimental/Animation/PoseSearch/Source/Runtime/Private/PoseSearchDatabase.cpp b/Engine/Plugins/Experimental/Animation/PoseSearch/Source/Runtime/Private/PoseSearchDatabase.cpp
--- a/Engine/Plugins/Experimental/Animation/PoseSearch/Source/Runtime/Private/PoseSearchDatabase.cpp
+++ b/Engine/Plugins/Experimental/Animation/PoseSearch/Source/Runtime/Private/PoseSearchDatabase.cpp
@@ -913,12 +913,6 @@ FPoseSearchCost UPoseSearchDatabase::SearchContinuingPose(UE::PoseSearch::FSearc
 		SearchContext.BestCandidates.Add(ContinuingPoseCost, ContinuingPoseIdx, this, EPoseCandidateFlags::Valid_ContinuingPose);
 #endif // UE_POSE_SEARCH_TRACE_ENABLED
 	}
-	else
-	{
-#if UE_POSE_SEARCH_TRACE_ENABLED
-		SearchContext.BestCandidates.Add(this);
-#endif // UE_POSE_SEARCH_TRACE_ENABLED
-	}
 
 	return ContinuingPoseCost;
 }
@@ -1088,10 +1082,6 @@ UE::PoseSearch::FSearchResult UPoseSearchDatabase::SearchPCAKDTree(UE::PoseSearc
 #endif // UE_POSE_SEARCH_TRACE_ENABLED
 	}
 
-#if UE_POSE_SEARCH_TRACE_ENABLED
-	SearchContext.BestCandidates.Add(this);
-#endif // UE_POSE_SEARCH_TRACE_ENABLED
-
 	// finalizing Result properties
 	if (Result.PoseIdx != INDEX_NONE)
 	{
@@ -1176,10 +1166,6 @@ UE::PoseSearch::FSearchResult UPoseSearchDatabase::SearchVPTree(UE::PoseSearch::
 #endif // UE_POSE_SEARCH_TRACE_ENABLED
 	}
 
-#if UE_POSE_SEARCH_TRACE_ENABLED
-	SearchContext.BestCandidates.Add(this);
-#endif // UE_POSE_SEARCH_TRACE_ENABLED
-
 	// finalizing Result properties
 	if (Result.PoseIdx != INDEX_NONE)
 	{
@@ -1255,10 +1241,6 @@ UE::PoseSearch::FSearchResult UPoseSearchDatabase::SearchBruteForce(UE::PoseSear
 #endif // UE_POSE_SEARCH_TRACE_ENABLED
 	}
 
-#if UE_POSE_SEARCH_TRACE_ENABLED
-	SearchContext.BestCandidates.Add(this);
-#endif // UE_POSE_SEARCH_TRACE_ENABLED
-
 	// finalizing Result properties
 	if (Result.PoseIdx != INDEX_NONE)
 	{
diff --git a/Engine/Plugins/Experimental/Animation/PoseSearch/Source/Runtime/Private/PoseSearchLibrary.cpp b/Engine/Plugins/Experimental/Animation/PoseSearch/Source/Runtime/Private/PoseSearchLibrary.cpp
--- a/Engine/Plugins/Experimental/Animation/PoseSearch/Source/Runtime/Private/PoseSearchLibrary.cpp
+++ b/Engine/Plugins/Experimental/Animation/PoseSearch/Source/Runtime/Private/PoseSearchLibrary.cpp
@@ -174,13 +174,7 @@ void UPoseSearchLibrary::TraceMotionMatchingState(
 
 	const int32 CurrentPoseIdx = bSearch && CurrentResult.PoseCost.IsValid() ? CurrentResult.PoseIdx : INDEX_NONE;
 	FTraceMotionMatchingState TraceState;
-
-	for (const UPoseSearchDatabase* SearchedDatabase : SearchContext.BestCandidates.GetSearchedDatabases())
-	{
-		AddUniqueDatabase(TraceState.DatabaseEntries, SearchedDatabase, SearchContext);
-	}
-
-	while (!SearchContext.BestCandidates.IsEmpty())
+	while (!SearchContext.BestCandidates.IsEmpty()) 
 	{
 		FSearchContext::FPoseCandidate PoseCandidate;
 		SearchContext.BestCandidates.Pop(PoseCandidate);
@@ -194,6 +188,8 @@ void UPoseSearchLibrary::TraceMotionMatchingState(
 		PoseEntry.PoseCandidateFlags = PoseCandidate.PoseCandidateFlags;
 		if (CurrentPoseIdx == PoseCandidate.PoseIdx && CurrentResult.Database.Get() == PoseCandidate.Database)
 		{
+			check(EnumHasAnyFlags(PoseEntry.PoseCandidateFlags, EPoseCandidateFlags::Valid_Pose | EPoseCandidateFlags::Valid_ContinuingPose));
+
 			EnumAddFlags(PoseEntry.PoseCandidateFlags, EPoseCandidateFlags::Valid_CurrentPose);
 			
 			TraceState.CurrentDbEntryIdx = DbEntryIdx;
diff --git a/Engine/Plugins/Experimental/Animation/PoseSearch/Source/Runtime/Public/PoseSearch/PoseSearchContext.h b/Engine/Plugins/Experimental/Animation/PoseSearch/Source/Runtime/Public/PoseSearch/PoseSearchContext.h
--- a/Engine/Plugins/Experimental/Animation/PoseSearch/Source/Runtime/Public/PoseSearch/PoseSearchContext.h
+++ b/Engine/Plugins/Experimental/Animation/PoseSearch/Source/Runtime/Public/PoseSearch/PoseSearchContext.h
@@ -218,18 +218,31 @@ private:
 	
 #if UE_POSE_SEARCH_TRACE_ENABLED
 
-public:
-	struct FPoseCandidateBase
+	struct FPoseCandidateId
 	{
-		FPoseSearchCost Cost;
 		int32 PoseIdx = 0;
 		const UPoseSearchDatabase* Database = nullptr;
 
-		bool operator<(const FPoseCandidateBase& Other) const { return Other.Cost < Cost; } // Reverse compare because BestCandidates is a max heap
-		bool operator==(const FSearchResult& SearchResult) const { return (PoseIdx == SearchResult.PoseIdx) && (Database == SearchResult.Database.Get()); }
+		bool operator==(const FPoseCandidateId& Other) const
+		{
+			return PoseIdx == Other.PoseIdx && Database == Other.Database;
+		}
+		
+		friend uint32 GetTypeHash(const FPoseCandidateId& PoseCandidateId)
+		{
+			return HashCombineFast(GetTypeHash(PoseCandidateId.PoseIdx), GetTypeHash(PoseCandidateId.Database));
+		}
+	};
+
+	struct FPoseCandidateIdCost : public FPoseCandidateId
+	{
+		FPoseSearchCost Cost;
+		bool operator<(const FPoseCandidateIdCost& Other) const { return Other.Cost < Cost; } // Reverse compare because BestCandidates is a max heap
 	};
+	
+public:
 
-	struct FPoseCandidate : public FPoseCandidateBase
+	struct FPoseCandidate : public FPoseCandidateIdCost
 	{
 		EPoseCandidateFlags PoseCandidateFlags = EPoseCandidateFlags::None;
 	};
@@ -239,19 +252,17 @@ public:
 		FBestPoseCandidates()
 		{
 			// preallocating memory to avoid multiple reallocations / rehashing
-			SearchedDatabases.Empty(8);
 			PoseCandidateHeap.Reserve(MaxPoseCandidates);
 			PoseIdxToFlags.Empty(MaxPoseCandidates);
 		}
 
-		void Add(const UPoseSearchDatabase* Database)
-		{
-			SearchedDatabases.Add(Database);
-		}
-
 		void Add(const FPoseSearchCost& Cost, int32 PoseIdx, const UPoseSearchDatabase* Database, EPoseCandidateFlags PoseCandidateFlags)
 		{
-			if (EPoseCandidateFlags* PoseIdxPoseCandidateFlags = PoseIdxToFlags.Find(PoseIdx))
+			FPoseCandidate PoseCandidate;
+			PoseCandidate.PoseIdx = PoseIdx;
+			PoseCandidate.Database = Database;
+
+			if (EPoseCandidateFlags* PoseIdxPoseCandidateFlags = PoseIdxToFlags.Find(PoseCandidate))
 			{
 				*PoseIdxPoseCandidateFlags |= PoseCandidateFlags;
 			}
@@ -279,25 +290,19 @@ public:
 					FPoseCandidate PoppedPoseCandidate;
 					Pop(PoppedPoseCandidate);
 					PoseCandidateHeap.HeapPush(ContinuingPoseCandidate);
-					PoseIdxToFlags.Add(ContinuingPoseCandidate.PoseIdx, ContinuingPoseCandidate.PoseCandidateFlags);
+					PoseIdxToFlags.Add(ContinuingPoseCandidate, ContinuingPoseCandidate.PoseCandidateFlags);
 				}
 
-				FPoseCandidate PoseCandidate;
 				PoseCandidate.Cost = Cost;
-				PoseCandidate.PoseIdx = PoseIdx;
-				PoseCandidate.Database = Database;
-				
 				PoseCandidateHeap.HeapPush(PoseCandidate);
-				PoseIdxToFlags.Add(PoseIdx, PoseCandidateFlags);
+				PoseIdxToFlags.Add(PoseCandidate, PoseCandidateFlags);
 			}
-
-			Add(Database);
 		}
 
 		void Pop(FPoseCandidate& OutItem)
 		{
 			PoseCandidateHeap.HeapPop(OutItem, false);
-			OutItem.PoseCandidateFlags = PoseIdxToFlags.FindAndRemoveChecked(OutItem.PoseIdx);
+			OutItem.PoseCandidateFlags = PoseIdxToFlags.FindAndRemoveChecked(OutItem);
 		}
 
 		bool IsEmpty() const
@@ -310,15 +315,9 @@ public:
 			MaxPoseCandidates = Value;
 		}
 
-		const TSet<const UPoseSearchDatabase*>& GetSearchedDatabases() const
-		{
-			return SearchedDatabases;
-		}
-
 	private:
-		TSet<const UPoseSearchDatabase*> SearchedDatabases;
-		TArray<FPoseCandidateBase> PoseCandidateHeap;
-		TMap<int32, EPoseCandidateFlags> PoseIdxToFlags;
+		TArray<FPoseCandidateIdCost> PoseCandidateHeap;
+		TMap<FPoseCandidateId, EPoseCandidateFlags> PoseIdxToFlags;
 		int32 MaxPoseCandidates = 200;
 	};
 	
