motion matching - IPoseSearchProvider improvements to support searching multiple AnimationAsset(s)

needed multiple preflights to pass all steps:
6514a306314218e67fb7e5eb
651497b2d03a3ae564064232
65149331c061787671ba70e5
651486b68b36c3ec41b58634

diff --git a/Engine/Plugins/Experimental/Animation/PoseSearch/Source/Editor/Private/PoseSearchDebuggerDatabaseColumns.h b/Engine/Plugins/Experimental/Animation/PoseSearch/Source/Editor/Private/PoseSearchDebuggerDatabaseColumns.h
--- a/Engine/Plugins/Experimental/Animation/PoseSearch/Source/Editor/Private/PoseSearchDebuggerDatabaseColumns.h
+++ b/Engine/Plugins/Experimental/Animation/PoseSearch/Source/Editor/Private/PoseSearchDebuggerDatabaseColumns.h
@@ -631,6 +631,11 @@ struct FPoseCandidateFlags : ITextColumn
 				Sring.Append("F ");
 			}
 
+			if (EnumHasAnyFlags(Row->PoseCandidateFlags, EPoseCandidateFlags::DiscardedBy_AssetIdxFilter))
+			{
+				Sring.Append("A ");
+			}
+			
 			if (EnumHasAnyFlags(Row->PoseCandidateFlags, EPoseCandidateFlags::DiscardedBy_Search))
 			{
 				Sring.Append("S ");
@@ -669,6 +674,11 @@ struct FPoseCandidateFlags : ITextColumn
 				TextBuilder.AppendLine(LOCTEXT("DiscardedBy_PoseFilter_Tooltip", "(F) Filter"));
 			}
 
+			if (EnumHasAnyFlags(Row->PoseCandidateFlags, EPoseCandidateFlags::DiscardedBy_AssetIdxFilter))
+			{
+				TextBuilder.AppendLine(LOCTEXT("DiscardedBy_AssetIdxFilter_Tooltip", "(A) Asset Idx Filter"));
+			}
+			
 			if (EnumHasAnyFlags(Row->PoseCandidateFlags, EPoseCandidateFlags::DiscardedBy_Search))
 			{
 				TextBuilder.AppendLine(LOCTEXT("DiscardedBy_Search_Tooltip", "(S) Search"));
diff --git a/Engine/Plugins/Experimental/Animation/PoseSearch/Source/Runtime/Private/PoseSearchContext.cpp b/Engine/Plugins/Experimental/Animation/PoseSearch/Source/Runtime/Private/PoseSearchContext.cpp
--- a/Engine/Plugins/Experimental/Animation/PoseSearch/Source/Runtime/Private/PoseSearchContext.cpp
+++ b/Engine/Plugins/Experimental/Animation/PoseSearch/Source/Runtime/Private/PoseSearchContext.cpp
@@ -351,11 +351,13 @@ void FDebugDrawParams::DrawFeatureVector(int32 PoseIdx)
 
 //////////////////////////////////////////////////////////////////////////
 // FSearchContext
-FSearchContext::FSearchContext(const UAnimInstance* InAnimInstance, const FPoseSearchQueryTrajectory* InTrajectory, const IPoseHistory* InHistory, float InDesiredPermutationTimeOffset,
-	const FPoseIndicesHistory* InPoseIndicesHistory, const FSearchResult& InCurrentResult, const FFloatInterval& InPoseJumpThresholdTime, bool bInForceInterrupt)
+FSearchContext::FSearchContext(const UAnimInstance* InAnimInstance, const IPoseHistory* InHistory, TConstArrayView<const UAnimationAsset*> InAnimationsToConsider,
+		const FPoseSearchQueryTrajectory* InTrajectory, float InDesiredPermutationTimeOffset, const FPoseIndicesHistory* InPoseIndicesHistory,
+		const FSearchResult& InCurrentResult, const FFloatInterval& InPoseJumpThresholdTime, bool bInForceInterrupt)
 : AnimInstance(InAnimInstance)
-, Trajectory(InTrajectory)
 , History(InHistory)
+, AnimationsToConsider(InAnimationsToConsider)
+, Trajectory(InTrajectory)
 , DesiredPermutationTimeOffset(InDesiredPermutationTimeOffset)
 , PoseIndicesHistory(InPoseIndicesHistory)
 , CurrentResult(InCurrentResult)
diff --git a/Engine/Plugins/Experimental/Animation/PoseSearch/Source/Runtime/Private/PoseSearchDatabase.cpp b/Engine/Plugins/Experimental/Animation/PoseSearch/Source/Runtime/Private/PoseSearchDatabase.cpp
--- a/Engine/Plugins/Experimental/Animation/PoseSearch/Source/Runtime/Private/PoseSearchDatabase.cpp
+++ b/Engine/Plugins/Experimental/Animation/PoseSearch/Source/Runtime/Private/PoseSearchDatabase.cpp
@@ -39,6 +39,28 @@ static TAutoConsoleVariable<bool> CVarMotionMatchCompareAgainstBruteForce(TEXT("
 static TAutoConsoleVariable<bool> CVarMotionMatchValidateKNNSearch(TEXT("a.MotionMatch.ValidateKNNSearch"), false, TEXT("Validate KNN search"));
 #endif
 
+typedef TArray<int32, TInlineAllocator<256>> FSelectableAssetIdx;
+static void PopulateSelectableAssetIdx(FSelectableAssetIdx& SelectableAssetIdx, TConstArrayView<const UAnimationAsset*> AnimationsToConsider, const UPoseSearchDatabase* Database)
+{
+	check(Database);
+	SelectableAssetIdx.Reset();
+	if (!AnimationsToConsider.IsEmpty())
+	{
+		const FSearchIndex& SearchIndex = Database->GetSearchIndex();
+
+		for (int32 AssetIndex = 0; AssetIndex < SearchIndex.Assets.Num(); ++AssetIndex)
+		{
+			if (const FPoseSearchDatabaseAnimationAssetBase* DatabaseAnimationAssetBase = Database->GetAnimationAssetBase(SearchIndex.Assets[AssetIndex]))
+			{
+				if (AnimationsToConsider.Contains(DatabaseAnimationAssetBase->GetAnimationAsset()))
+				{
+					SelectableAssetIdx.Add(AssetIndex);
+				}
+			}
+		}
+	}
+}
+
 typedef TArray<int32, TInlineAllocator<256>> FNonSelectableIdx;
 static void PopulateNonSelectableIdx(FNonSelectableIdx& NonSelectableIdx, FSearchContext& SearchContext, const UPoseSearchDatabase* Database
 #if UE_POSE_SEARCH_TRACE_ENABLED
@@ -155,10 +177,8 @@ static void PopulateNonSelectableIdx(FNonSelectableIdx& NonSelectableIdx, FSearc
 
 struct FSearchFilters
 {
-	FSearchFilters(const UPoseSearchSchema* Schema, TConstArrayView<int32> NonSelectableIdx, bool bAnyBlockTransition)
+	FSearchFilters(const UPoseSearchSchema* Schema, TConstArrayView<int32> NonSelectableIdx, TConstArrayView<int32> SelectableAssetIdx, bool bAnyBlockTransition)
 	{
-		NonSelectableIdxFilter.NonSelectableIdx = NonSelectableIdx;
-
 		if (bAnyBlockTransition)
 		{
 			Filters.Add(&BlockTransitionFilter);
@@ -169,6 +189,11 @@ struct FSearchFilters
 			Filters.Add(&NonSelectableIdxFilter);
 		}
 
+		if (SelectableAssetIdxFilter.Init(SelectableAssetIdx).IsFilterActive())
+		{
+			Filters.Add(&SelectableAssetIdxFilter);
+		}
+
 		for (const IPoseSearchFilter* Filter : Schema->GetChannels())
 		{
 			if (Filter->IsFilterActive())
@@ -193,6 +218,11 @@ struct FSearchFilters
 				{
 					// candidate already added to SearchContext.BestCandidates by PopulateNonSelectableIdx
 				}
+				else if (Filter == &SelectableAssetIdxFilter)
+				{
+					const FPoseSearchCost PoseCost = SearchIndex.ComparePoses(PoseIdx, 0.f, PoseValues, QueryValues);
+					SearchContext.BestCandidates.Add(PoseCost, PoseIdx, Database, EPoseCandidateFlags::DiscardedBy_AssetIdxFilter);
+				}
 				else if (Filter == &BlockTransitionFilter)
 				{
 					const FPoseSearchCost PoseCost = SearchIndex.ComparePoses(PoseIdx, 0.f, PoseValues, QueryValues);
@@ -233,6 +263,28 @@ private:
 		TConstArrayView<int32> NonSelectableIdx;
 	};
 
+	struct FSelectableAssetIdxFilter : public IPoseSearchFilter
+	{
+		const FSelectableAssetIdxFilter& Init(TConstArrayView<int32> InSelectableAssetIdxFilter)
+		{
+			check(Algo::IsSorted(InSelectableAssetIdxFilter));
+			SelectableAssetIdxFilter = InSelectableAssetIdxFilter;
+			return *this;
+		}
+
+		virtual bool IsFilterActive() const override
+		{
+			return !SelectableAssetIdxFilter.IsEmpty();
+		}
+
+		virtual bool IsFilterValid(TConstArrayView<float> PoseValues, TConstArrayView<float> QueryValues, int32 PoseIdx, const FPoseMetadata& Metadata) const override
+		{
+			return Algo::BinarySearch(SelectableAssetIdxFilter, int32(Metadata.GetAssetIndex())) != INDEX_NONE;
+		}
+
+		TConstArrayView<int32> SelectableAssetIdxFilter;
+	};
+
 	struct FBlockTransitionFilter : public IPoseSearchFilter
 	{
 		virtual bool IsFilterActive() const override
@@ -247,6 +299,7 @@ private:
 	};
 
 	FNonSelectableIdxFilter NonSelectableIdxFilter;
+	FSelectableAssetIdxFilter SelectableAssetIdxFilter;
 	FBlockTransitionFilter BlockTransitionFilter;
 
 	TArray<const IPoseSearchFilter*, TInlineAllocator<64>> Filters;
@@ -943,6 +996,9 @@ UE::PoseSearch::FSearchResult UPoseSearchDatabase::SearchPCAKDTree(UE::PoseSearc
 	
 		TConstArrayView<float> QueryValues = SearchContext.GetOrBuildQuery(Schema).GetValues();
 
+		FSelectableAssetIdx SelectableAssetIdx;
+		PopulateSelectableAssetIdx(SelectableAssetIdx, SearchContext.GetAnimationsToConsider(), this);
+
 		FNonSelectableIdx NonSelectableIdx;
 		PopulateNonSelectableIdx(NonSelectableIdx, SearchContext, this
 #if UE_POSE_SEARCH_TRACE_ENABLED
@@ -1009,7 +1065,7 @@ UE::PoseSearch::FSearchResult UPoseSearchDatabase::SearchPCAKDTree(UE::PoseSearc
 #endif // WITH_EDITOR && ENABLE_ANIM_DEBUG
 
 		// NonSelectableIdx are already filtered out inside the kdtree search
-		const FSearchFilters SearchFilters(Schema, bRunNonSelectableIdxPostKDTree ? NonSelectableIdx : TConstArrayView<int32>(), SearchIndex.bAnyBlockTransition);
+		const FSearchFilters SearchFilters(Schema, bRunNonSelectableIdxPostKDTree ? NonSelectableIdx : TConstArrayView<int32>(), SelectableAssetIdx, SearchIndex.bAnyBlockTransition);
 		
 		// are the PCAValues pruned out of duplicates (multiple poses are associated with the same PCAValuesVectorIdx)
 		if (bArePCAValuesPruned)
@@ -1108,6 +1164,9 @@ UE::PoseSearch::FSearchResult UPoseSearchDatabase::SearchVPTree(UE::PoseSearch::
 	{
 		TConstArrayView<float> QueryValues = SearchContext.GetOrBuildQuery(Schema).GetValues();
 
+		FSelectableAssetIdx SelectableAssetIdx;
+		PopulateSelectableAssetIdx(SelectableAssetIdx, SearchContext.GetAnimationsToConsider(), this);
+
 		// @todo: implement filtering within the VPTree as KDTree does
 		FNonSelectableIdx NonSelectableIdx;
 		PopulateNonSelectableIdx(NonSelectableIdx, SearchContext, this
@@ -1119,7 +1178,7 @@ UE::PoseSearch::FSearchResult UPoseSearchDatabase::SearchVPTree(UE::PoseSearch::
 		const int32 NumDimensions = Schema->SchemaCardinality;
 		check(QueryValues.Num() == NumDimensions);
 		
-		const FSearchFilters SearchFilters(Schema, NonSelectableIdx, SearchIndex.bAnyBlockTransition);
+		const FSearchFilters SearchFilters(Schema, NonSelectableIdx, SelectableAssetIdx, SearchIndex.bAnyBlockTransition);
 
 		// @todo: implement a FVPTreeDataSource for aligned and padded features vector like CompareAlignedPoses does 
 		FVPTreeDataSource DataSource(SearchIndex);
@@ -1192,6 +1251,9 @@ UE::PoseSearch::FSearchResult UPoseSearchDatabase::SearchBruteForce(UE::PoseSear
 	{
 		TConstArrayView<float> QueryValues = SearchContext.GetOrBuildQuery(Schema).GetValues();
 
+		FSelectableAssetIdx SelectableAssetIdx;
+		PopulateSelectableAssetIdx(SelectableAssetIdx, SearchContext.GetAnimationsToConsider(), this);
+
 		FNonSelectableIdx NonSelectableIdx;
 		PopulateNonSelectableIdx(NonSelectableIdx, SearchContext, this
 #if UE_POSE_SEARCH_TRACE_ENABLED
@@ -1200,7 +1262,7 @@ UE::PoseSearch::FSearchResult UPoseSearchDatabase::SearchBruteForce(UE::PoseSear
 		);
 
 		const int32 NumDimensions = Schema->SchemaCardinality;
-		const FSearchFilters SearchFilters(Schema, NonSelectableIdx, SearchIndex.bAnyBlockTransition);
+		const FSearchFilters SearchFilters(Schema, NonSelectableIdx, SelectableAssetIdx, SearchIndex.bAnyBlockTransition);
 		const bool bUpdateBestCandidates = PoseSearchMode == EPoseSearchMode::BruteForce;
 
 		// do we need to reconstruct pose values?
diff --git a/Engine/Plugins/Experimental/Animation/PoseSearch/Source/Runtime/Private/PoseSearchDerivedDataKey.cpp b/Engine/Plugins/Experimental/Animation/PoseSearch/Source/Runtime/Private/PoseSearchDerivedDataKey.cpp
--- a/Engine/Plugins/Experimental/Animation/PoseSearch/Source/Runtime/Private/PoseSearchDerivedDataKey.cpp
+++ b/Engine/Plugins/Experimental/Animation/PoseSearch/Source/Runtime/Private/PoseSearchDerivedDataKey.cpp
@@ -40,7 +40,7 @@ FKeyBuilder::FKeyBuilder(const UObject* Object, bool bUseDataVer, bool bPerformC
 	if (bUseDataVer)
 	{
 		// used to invalidate the key without having to change POSESEARCHDB_DERIVEDDATA_VER all the times
-		int32 POSESEARCHDB_DERIVEDDATA_VER_SMALL = 217;
+		int32 POSESEARCHDB_DERIVEDDATA_VER_SMALL = 219;
 		FGuid VersionGuid = FDevSystemGuids::GetSystemGuid(FDevSystemGuids::Get().POSESEARCHDB_DERIVEDDATA_VER);
 
 		*this << VersionGuid;
diff --git a/Engine/Plugins/Experimental/Animation/PoseSearch/Source/Runtime/Private/PoseSearchLibrary.cpp b/Engine/Plugins/Experimental/Animation/PoseSearch/Source/Runtime/Private/PoseSearchLibrary.cpp
--- a/Engine/Plugins/Experimental/Animation/PoseSearch/Source/Runtime/Private/PoseSearchLibrary.cpp
+++ b/Engine/Plugins/Experimental/Animation/PoseSearch/Source/Runtime/Private/PoseSearchLibrary.cpp
@@ -282,7 +282,8 @@ void UPoseSearchLibrary::UpdateMotionMatchingState(
 	FMemMark Mark(FMemStack::Get());
 	const UAnimInstance* AnimInstance = Cast<const UAnimInstance>(Context.AnimInstanceProxy->GetAnimInstanceObject());
 	check(AnimInstance);
-	FSearchContext SearchContext(AnimInstance, &TrajectoryRootSpace, History, 0.f, &InOutMotionMatchingState.PoseIndicesHistory, InOutMotionMatchingState.CurrentSearchResult, PoseJumpThresholdTime, bForceInterrupt);
+	FSearchContext SearchContext(AnimInstance, History, TConstArrayView<const UAnimationAsset*>(), &TrajectoryRootSpace, 0.f,
+		&InOutMotionMatchingState.PoseIndicesHistory, InOutMotionMatchingState.CurrentSearchResult, PoseJumpThresholdTime, bForceInterrupt);
 
 	const bool bCanAdvance = InOutMotionMatchingState.CurrentSearchResult.CanAdvance(DeltaTime);
 
@@ -537,7 +538,7 @@ void UPoseSearchLibrary::MotionMatch(
 #endif // ENABLE_DRAW_DEBUG && ENABLE_ANIM_DEBUG
 		}
 
-		FSearchContext SearchContext(AnimInstance, &TrajectoryRootSpace, ExtendedPoseHistory.IsInitialized() ? &ExtendedPoseHistory : nullptr, TimeToFutureAnimationStart);
+		FSearchContext SearchContext(AnimInstance, ExtendedPoseHistory.IsInitialized() ? &ExtendedPoseHistory : nullptr, TConstArrayView<const UAnimationAsset*>(), &TrajectoryRootSpace, TimeToFutureAnimationStart);
 
 		FSearchResult SearchResult = Database->Search(SearchContext);
 		if (SearchResult.IsValid())
@@ -561,13 +562,13 @@ void UPoseSearchLibrary::MotionMatch(
 			FAnimInstanceProxy* AnimInstanceProxy = UAnimInstanceProxyProvider::GetAnimInstanceProxy(AnimInstance);
 			if (CVarAnimMotionMatchDrawMatchEnable.GetValueOnAnyThread())
 			{
-				UE::PoseSearch::FDebugDrawParams DrawParams(AnimInstanceProxy, AnimInstanceProxy->GetComponentTransform(), SearchResult.Database.Get());
+				FDebugDrawParams DrawParams(AnimInstanceProxy, AnimInstanceProxy->GetComponentTransform(), SearchResult.Database.Get());
 				DrawParams.DrawFeatureVector(SearchResult.PoseIdx);
 			}
 
 			if (CVarAnimMotionMatchDrawQueryEnable.GetValueOnAnyThread())
 			{
-				UE::PoseSearch::FDebugDrawParams DrawParams(AnimInstanceProxy, AnimInstanceProxy->GetComponentTransform(), SearchResult.Database.Get(), EDebugDrawFlags::DrawQuery);
+				FDebugDrawParams DrawParams(AnimInstanceProxy, AnimInstanceProxy->GetComponentTransform(), SearchResult.Database.Get(), EDebugDrawFlags::DrawQuery);
 				DrawParams.DrawFeatureVector(SearchContext.GetOrBuildQuery(SearchResult.Database->Schema).GetValues());
 			}
 		}
@@ -581,15 +582,40 @@ void UPoseSearchLibrary::MotionMatch(
 	}
 }
 
-UE::PoseSearch::FSearchResult UPoseSearchLibrary::MotionMatch(const FAnimationBaseContext& Context, const UObject* Object)
+UE::PoseSearch::FSearchResult UPoseSearchLibrary::MotionMatch(const FAnimationBaseContext& Context, TConstArrayView<UAnimationAsset*> AnimationAssets)
 {
-	//using namespace UE::Anim;
 	using namespace UE::PoseSearch;
 
 	FSearchResult SearchResult;
 
-	// @todo: we currently support only animation assets, but in the future we should support choosers and proxy (via some common interface)
-	if (const UAnimSequenceBase* SequenceBase = Cast<const UAnimSequenceBase>(Object))
+	// budgeting some stack allocations for simple use cases. bigger requests of AnimationAssets contining 
+	// UAnimNotifyState_PoseSearchBranchIn referencing multiple datbases will default to slower heap allocations
+	enum { MAX_STACK_ALLOCATED_ANIMATIONS = 16 };
+	enum { MAX_STACK_ALLOCATED_SETS = 2 };
+	typedef	TArray<const UAnimationAsset*, TInlineAllocator<MAX_STACK_ALLOCATED_ANIMATIONS>> TDbAnims;
+	typedef TMap<const UPoseSearchDatabase*, TDbAnims, TInlineSetAllocator<MAX_STACK_ALLOCATED_SETS>> FPerDbAnimMap;
+	typedef TPair<const UPoseSearchDatabase*, TDbAnims> FPerDbAnimPair;
+	FPerDbAnimMap PerDbAnimMap;
+	
+	// colecting all the UAnimSequenceBase to consider for each database
+	for (const UAnimationAsset* AnimationAsset : AnimationAssets)
+	{
+		if (const UAnimSequenceBase* SequenceBase = Cast<const UAnimSequenceBase>(AnimationAsset))
+		{
+			for (const FAnimNotifyEvent& NotifyEvent : SequenceBase->Notifies)
+			{
+				if (const UAnimNotifyState_PoseSearchBranchIn* PoseSearchBranchIn = Cast<UAnimNotifyState_PoseSearchBranchIn>(NotifyEvent.NotifyStateClass))
+				{
+					if (PoseSearchBranchIn->Database)
+					{
+						PerDbAnimMap.FindOrAdd(PoseSearchBranchIn->Database).AddUnique(SequenceBase);
+					}
+				}
+			}
+		}
+	}
+
+	if (!PerDbAnimMap.IsEmpty())
 	{
 		const IPoseHistory* History = nullptr;
 		if (IPoseHistoryProvider* PoseHistoryProvider = Context.GetMessage<IPoseHistoryProvider>())
@@ -601,37 +627,17 @@ UE::PoseSearch::FSearchResult UPoseSearchLibrary::MotionMatch(const FAnimationBa
 		check(AnimInstance);
 
 		FMemMark Mark(FMemStack::Get());
-		FSearchContext SearchContext(AnimInstance, nullptr, History);
-		for (const FAnimNotifyEvent& NotifyEvent : SequenceBase->Notifies)
-		{
-			if (const UAnimNotifyState_PoseSearchBranchIn* PoseSearchBranchIn = Cast<UAnimNotifyState_PoseSearchBranchIn>(NotifyEvent.NotifyStateClass))
-			{
-				if (PoseSearchBranchIn->Database)
-				{
-#if ENABLE_ANIM_DEBUG
-					// @todo: implement database filtering based on subsets of Database->AnimationAssets!
-					//		  FSearchContext should hold a list of UAnimationAsset(s) to use to narrow down the database search,
-					//		  but for now we log an error
-					for (int32 AnimationAssetIndex = 0; AnimationAssetIndex < PoseSearchBranchIn->Database->AnimationAssets.Num(); ++AnimationAssetIndex)
-					{
-						if (const FPoseSearchDatabaseAnimationAssetBase* AnimationAssetBase = PoseSearchBranchIn->Database->GetAnimationAssetBase(AnimationAssetIndex))
-						{
-							if (AnimationAssetBase->GetAnimationAsset() != Object)
-							{
-								UE_LOG(LogPoseSearch, Error, TEXT("UAnimNotifyState_PoseSearchBranchIn doesn't yet support database between multiple animation assets"));
-							}
-						}
-					}
-#endif // ENABLE_ANIM_DEBUG
+		FSearchContext SearchContext(AnimInstance, History);
 
-					const FSearchResult NewSearchResult = PoseSearchBranchIn->Database->Search(SearchContext);
+		for (const FPerDbAnimPair& PerDbAnimPair : PerDbAnimMap)
+		{
+			SearchContext.SetAnimationsToConsider(PerDbAnimPair.Value);
 
-					if (NewSearchResult.PoseCost.GetTotalCost() < SearchResult.PoseCost.GetTotalCost())
-					{
-						SearchResult = NewSearchResult;
-						SearchContext.UpdateCurrentBestCost(SearchResult.PoseCost);
-					}
-				}
+			const FSearchResult NewSearchResult = PerDbAnimPair.Key->Search(SearchContext);
+			if (NewSearchResult.PoseCost.GetTotalCost() < SearchResult.PoseCost.GetTotalCost())
+			{
+				SearchResult = NewSearchResult;
+				SearchContext.UpdateCurrentBestCost(SearchResult.PoseCost);
 			}
 		}
 
@@ -640,13 +646,13 @@ UE::PoseSearch::FSearchResult UPoseSearchLibrary::MotionMatch(const FAnimationBa
 		{
 			if (CVarAnimMotionMatchDrawMatchEnable.GetValueOnAnyThread())
 			{
-				UE::PoseSearch::FDebugDrawParams DrawParams(Context.AnimInstanceProxy, Context.AnimInstanceProxy->GetComponentTransform(), SearchResult.Database.Get());
+				FDebugDrawParams DrawParams(Context.AnimInstanceProxy, Context.AnimInstanceProxy->GetComponentTransform(), SearchResult.Database.Get());
 				DrawParams.DrawFeatureVector(SearchResult.PoseIdx);
 			}
 
 			if (CVarAnimMotionMatchDrawQueryEnable.GetValueOnAnyThread())
 			{
-				UE::PoseSearch::FDebugDrawParams DrawParams(Context.AnimInstanceProxy, Context.AnimInstanceProxy->GetComponentTransform(), SearchResult.Database.Get(), EDebugDrawFlags::DrawQuery);
+				FDebugDrawParams DrawParams(Context.AnimInstanceProxy, Context.AnimInstanceProxy->GetComponentTransform(), SearchResult.Database.Get(), EDebugDrawFlags::DrawQuery);
 				DrawParams.DrawFeatureVector(SearchContext.GetOrBuildQuery(SearchResult.Database->Schema).GetValues());
 			}
 		}
diff --git a/Engine/Plugins/Experimental/Animation/PoseSearch/Source/Runtime/Private/PoseSearchModule.cpp b/Engine/Plugins/Experimental/Animation/PoseSearch/Source/Runtime/Private/PoseSearchModule.cpp
--- a/Engine/Plugins/Experimental/Animation/PoseSearch/Source/Runtime/Private/PoseSearchModule.cpp
+++ b/Engine/Plugins/Experimental/Animation/PoseSearch/Source/Runtime/Private/PoseSearchModule.cpp
@@ -3,6 +3,7 @@
 #include "Animation/AnimPoseSearchProvider.h"
 #include "Features/IModularFeatures.h"
 #include "Modules/ModuleManager.h"
+#include "PoseSearch/PoseSearchDatabase.h"
 #include "PoseSearch/PoseSearchLibrary.h"
 
 class FPoseSearchModule final : public IModuleInterface, public UE::Anim::IPoseSearchProvider
@@ -21,14 +22,19 @@ public:
 	}
 
 	// IPoseSearchProvider
-	virtual UE::Anim::IPoseSearchProvider::FSearchResult Search(const FAnimationBaseContext& GraphContext, const UObject* Object) override
+	virtual UE::Anim::IPoseSearchProvider::FSearchResult Search(const FAnimationBaseContext& GraphContext, TConstArrayView<UAnimationAsset*> AnimationAssets) override
 	{
-		const UE::PoseSearch::FSearchResult SearchResult = UPoseSearchLibrary::MotionMatch(GraphContext, Object);
-		
+		const UE::PoseSearch::FSearchResult SearchResult = UPoseSearchLibrary::MotionMatch(GraphContext, AnimationAssets);
 		UE::Anim::IPoseSearchProvider::FSearchResult ProviderResult;
-		ProviderResult.Dissimilarity = SearchResult.PoseCost.GetTotalCost();
-		ProviderResult.PoseIdx = SearchResult.PoseIdx;
-		ProviderResult.TimeOffsetSeconds = SearchResult.AssetTime;
+		if (const UE::PoseSearch::FSearchIndexAsset* SearchIndexAsset = SearchResult.GetSearchIndexAsset())
+		{
+			if (const FPoseSearchDatabaseAnimationAssetBase* DatabaseAnimationAssetBase = SearchResult.Database->GetAnimationAssetBase(*SearchIndexAsset))
+			{
+				ProviderResult.AnimationAsset = DatabaseAnimationAssetBase->GetAnimationAsset();
+				ProviderResult.Dissimilarity = SearchResult.PoseCost.GetTotalCost();
+				ProviderResult.TimeOffsetSeconds = SearchResult.AssetTime;
+			}
+		}
 
 		return ProviderResult;
 	}
diff --git a/Engine/Plugins/Experimental/Animation/PoseSearch/Source/Runtime/Public/PoseSearch/PoseSearchContext.h b/Engine/Plugins/Experimental/Animation/PoseSearch/Source/Runtime/Public/PoseSearch/PoseSearchContext.h
--- a/Engine/Plugins/Experimental/Animation/PoseSearch/Source/Runtime/Public/PoseSearch/PoseSearchContext.h
+++ b/Engine/Plugins/Experimental/Animation/PoseSearch/Source/Runtime/Public/PoseSearch/PoseSearchContext.h
@@ -27,7 +27,7 @@ enum class EDebugDrawFlags : uint32
 };
 ENUM_CLASS_FLAGS(EDebugDrawFlags);
 
-enum class EPoseCandidateFlags : uint8
+enum class EPoseCandidateFlags : uint32
 {
 	None = 0,
 
@@ -41,9 +41,10 @@ enum class EPoseCandidateFlags : uint8
 	DiscardedBy_PoseReselectHistory = 1 << 4,
 	DiscardedBy_BlockTransition = 1 << 5,
 	DiscardedBy_PoseFilter = 1 << 6,
-	DiscardedBy_Search = 1 << 7,
+	DiscardedBy_AssetIdxFilter = 1 << 7,
+	DiscardedBy_Search = 1 << 8,
 
-	AnyDiscardedMask = DiscardedBy_PoseJumpThresholdTime | DiscardedBy_PoseReselectHistory | DiscardedBy_BlockTransition | DiscardedBy_PoseFilter | DiscardedBy_Search,
+	AnyDiscardedMask = DiscardedBy_PoseJumpThresholdTime | DiscardedBy_PoseReselectHistory | DiscardedBy_BlockTransition | DiscardedBy_PoseFilter | DiscardedBy_AssetIdxFilter | DiscardedBy_Search,
 };
 ENUM_CLASS_FLAGS(EPoseCandidateFlags);
 
@@ -142,7 +143,8 @@ private:
 
 struct POSESEARCH_API FSearchContext
 {
-	FSearchContext(const UAnimInstance* InAnimInstance, const FPoseSearchQueryTrajectory* InTrajectory, const IPoseHistory* InHistory, float InDesiredPermutationTimeOffset = 0.f, const FPoseIndicesHistory* InPoseIndicesHistory = nullptr,
+	FSearchContext(const UAnimInstance* InAnimInstance, const IPoseHistory* InHistory, TConstArrayView<const UAnimationAsset*> InAnimationsToConsider = TConstArrayView<const UAnimationAsset*>(),
+		const FPoseSearchQueryTrajectory* InTrajectory = nullptr, float InDesiredPermutationTimeOffset = 0.f, const FPoseIndicesHistory* InPoseIndicesHistory = nullptr,
 		const FSearchResult& InCurrentResult = FSearchResult(), const FFloatInterval& InPoseJumpThresholdTime = FFloatInterval(0.f, 0.f), bool bInForceInterrupt = false);
 
 	// Returns the rotation of the bone Schema.BoneReferences[SchemaSampleBoneIdx] at an offset time of SampleTimeOffset relative to the
@@ -186,6 +188,9 @@ struct POSESEARCH_API FSearchContext
 	bool IsForceInterrupt() const { return bForceInterrupt; }
 	FTransform GetWorldRootBoneTransformAtTime(float SampleTime, bool bUseHistoryRoot = false, bool bExtrapolate = true) const;
 	const UAnimInstance* GetAnimInstance() const { return AnimInstance; }
+
+	void SetAnimationsToConsider(TConstArrayView<const UAnimationAsset*> InAnimationsToConsider) { AnimationsToConsider = InAnimationsToConsider; }
+	TConstArrayView<const UAnimationAsset*> GetAnimationsToConsider() const { return AnimationsToConsider; }
 	
 private:
 	// returns the world space transform of the bone SchemaBoneIdx at time SampleTime
@@ -198,9 +203,13 @@ private:
 	FQuat GetSampleRotationInternal(float SampleTime, float OriginTime, const UPoseSearchSchema* Schema, int8 SchemaSampleBoneIdx = RootSchemaBoneIdx, int8 SchemaOriginBoneIdx = RootSchemaBoneIdx, bool bUseHistoryRoot = false, const FQuat* SampleBoneRotationWorldOverride = nullptr);
 
 	const UAnimInstance* AnimInstance = nullptr;
+	const IPoseHistory* History = nullptr;
+
+	// if AnimationsToConsider is not empty, we'll search only for poses from UAnimationAsset(s) that are in the AnimationsToConsider
+	TConstArrayView<const UAnimationAsset*> AnimationsToConsider;
+
 	// Trajectory has been transformed in root bone world space reference system
 	const FPoseSearchQueryTrajectory* Trajectory = nullptr;
-	const IPoseHistory* History = nullptr;
 	float DesiredPermutationTimeOffset = 0.f;
 	const FPoseIndicesHistory* PoseIndicesHistory = nullptr;
 	FSearchResult CurrentResult;
diff --git a/Engine/Plugins/Experimental/Animation/PoseSearch/Source/Runtime/Public/PoseSearch/PoseSearchLibrary.h b/Engine/Plugins/Experimental/Animation/PoseSearch/Source/Runtime/Public/PoseSearch/PoseSearchLibrary.h
--- a/Engine/Plugins/Experimental/Animation/PoseSearch/Source/Runtime/Public/PoseSearch/PoseSearchLibrary.h
+++ b/Engine/Plugins/Experimental/Animation/PoseSearch/Source/Runtime/Public/PoseSearch/PoseSearchLibrary.h
@@ -163,7 +163,6 @@ public:
 		float TimeToFutureAnimationStart = 0.f,
 		const int32 DebugSessionUniqueIdentifier = 6174);
 
-
-	static UE::PoseSearch::FSearchResult MotionMatch(const FAnimationBaseContext& Context, const UObject* Object);
+	static UE::PoseSearch::FSearchResult MotionMatch(const FAnimationBaseContext& Context, TConstArrayView<UAnimationAsset*> AnimationAssets);
 };
 
diff --git a/Engine/Source/Runtime/Engine/Private/Animation/AnimNode_SequencePlayer.cpp b/Engine/Source/Runtime/Engine/Private/Animation/AnimNode_SequencePlayer.cpp
--- a/Engine/Source/Runtime/Engine/Private/Animation/AnimNode_SequencePlayer.cpp
+++ b/Engine/Source/Runtime/Engine/Private/Animation/AnimNode_SequencePlayer.cpp
@@ -159,14 +159,14 @@ float FAnimNode_SequencePlayerBase::GetTimeFromEnd(float CurrentNodeTime) const
 float FAnimNode_SequencePlayerBase::GetEffectiveStartPosition(const FAnimationBaseContext& Context) const
 {
 	// Override the start position if pose matching is enabled
-	UAnimSequenceBase* CurrentSequence = GetSequence();
+	UAnimationAsset* CurrentSequence = GetSequence();
 	if (CurrentSequence != nullptr && GetStartFromMatchingPose())
 	{
 		UE::Anim::IPoseSearchProvider* PoseSearchProvider = UE::Anim::IPoseSearchProvider::Get();
 		if (PoseSearchProvider)
 		{
-			UE::Anim::IPoseSearchProvider::FSearchResult Result = PoseSearchProvider->Search(Context, CurrentSequence);
-			if (Result.PoseIdx >= 0)
+			UE::Anim::IPoseSearchProvider::FSearchResult Result = PoseSearchProvider->Search(Context, MakeArrayView(&CurrentSequence, 1));
+			if (Result.AnimationAsset != nullptr)
 			{
 				return Result.TimeOffsetSeconds;
 			}
diff --git a/Engine/Source/Runtime/Engine/Public/Animation/AnimPoseSearchProvider.h b/Engine/Source/Runtime/Engine/Public/Animation/AnimPoseSearchProvider.h
--- a/Engine/Source/Runtime/Engine/Public/Animation/AnimPoseSearchProvider.h
+++ b/Engine/Source/Runtime/Engine/Public/Animation/AnimPoseSearchProvider.h
@@ -24,7 +24,7 @@ public:
 	
 	struct FSearchResult
 	{
-		int32 PoseIdx = INDEX_NONE;
+		UAnimationAsset* AnimationAsset = nullptr;
 		float TimeOffsetSeconds = 0.f;
 		float Dissimilarity = MAX_flt;
 	};
@@ -33,11 +33,11 @@ public:
 	* Finds a matching pose in the input Object given the current graph context
 	* 
 	* @param	GraphContext	Graph execution context used to construct a pose search query
-	* @param	Object			The object to search for the pose query
+	* @param	AnimationAssets	The animation assets to search for the pose query
 	* 
-	* @return	The pose in the animation asset that most closely matches the query
+	* @return	The pose in the AnimationAssets that most closely matches the query
 	*/
-	virtual FSearchResult Search(const FAnimationBaseContext& GraphContext, const UObject* Object) = 0;
+	virtual FSearchResult Search(const FAnimationBaseContext& GraphContext, TConstArrayView<UAnimationAsset*> AnimationAssets) = 0;
 };
 
 } // namespace UE::Anim
\ No newline at end of file
