motion matching - consolidating MM and PoseHistory nodex to use FPoseSearchQueryTrajectory as only historical world space skeletal mesh component transforms

diff --git a/Engine/Plugins/Experimental/Animation/PoseSearch/Source/Editor/Private/PoseSearchMeshComponent.cpp b/Engine/Plugins/Experimental/Animation/PoseSearch/Source/Editor/Private/PoseSearchMeshComponent.cpp
--- a/Engine/Plugins/Experimental/Animation/PoseSearch/Source/Editor/Private/PoseSearchMeshComponent.cpp
+++ b/Engine/Plugins/Experimental/Animation/PoseSearch/Source/Editor/Private/PoseSearchMeshComponent.cpp
@@ -51,6 +51,8 @@ void UPoseSearchMeshComponent::ResetToStart()
 
 void UPoseSearchMeshComponent::UpdatePose(const FUpdateContext& UpdateContext)
 {
+	using namespace UE::PoseSearch;
+
 	FMemMark Mark(FMemStack::Get());
 
 	FCompactPose CompactPose;
diff --git a/Engine/Plugins/Experimental/Animation/PoseSearch/Source/Runtime/Private/AnimNode_MotionMatching.cpp b/Engine/Plugins/Experimental/Animation/PoseSearch/Source/Runtime/Private/AnimNode_MotionMatching.cpp
--- a/Engine/Plugins/Experimental/Animation/PoseSearch/Source/Runtime/Private/AnimNode_MotionMatching.cpp
+++ b/Engine/Plugins/Experimental/Animation/PoseSearch/Source/Runtime/Private/AnimNode_MotionMatching.cpp
@@ -20,10 +20,11 @@
 #define LOCTEXT_NAMESPACE "AnimNode_MotionMatching"
 
 #if ENABLE_ANIM_DEBUG
-static TAutoConsoleVariable<int32> CVarAnimNodeMotionMatchingDrawQuery(TEXT("a.AnimNode.MotionMatching.DebugDrawQuery"), 0, TEXT("Draw input query"));
-static TAutoConsoleVariable<int32> CVarAnimNodeMotionMatchingDrawCurResult(TEXT("a.AnimNode.MotionMatching.DebugDrawCurResult"), 0, TEXT("Draw current result"));
-static TAutoConsoleVariable<int32> CVarAnimNodeMotionMatchingDrawInfo(TEXT("a.AnimNode.MotionMatching.DebugDrawInfo"), 0, TEXT("Draw info like current databases and asset"));
-static TAutoConsoleVariable<int32> CVarAnimNodeMotionMatchingDrawInfoVerbose(TEXT("a.AnimNode.MotionMatching.DebugDrawInfoVerbose"), 1, TEXT("Draw additional info like blend stack"));
+static TAutoConsoleVariable<bool> CVarAnimNodeMotionMatchingDrawQuery(TEXT("a.AnimNode.MotionMatching.DebugDrawQuery"), false, TEXT("Draw input query"));
+static TAutoConsoleVariable<bool> CVarAnimNodeMotionMatchingDrawCurResult(TEXT("a.AnimNode.MotionMatching.DebugDrawCurResult"), false, TEXT("Draw current result"));
+static TAutoConsoleVariable<bool> CVarAnimNodeMotionMatchingDrawPoseHistory(TEXT("a.AnimNode.MotionMatching.DebugDrawPoseHistory"), false, TEXT("Draw Pose History"));
+static TAutoConsoleVariable<bool> CVarAnimNodeMotionMatchingDrawInfo(TEXT("a.AnimNode.MotionMatching.DebugDrawInfo"), false, TEXT("Draw info like current databases and asset"));
+static TAutoConsoleVariable<bool> CVarAnimNodeMotionMatchingDrawInfoVerbose(TEXT("a.AnimNode.MotionMatching.DebugDrawInfoVerbose"), true, TEXT("Draw additional info like blend stack"));
 static TAutoConsoleVariable<float> CVarAnimNodeMotionMatchingDrawInfoHeight(TEXT("a.AnimNode.MotionMatching.DebugDrawInfoHeight"), 50.f, TEXT("Vertical offset for DebugDrawInfo"));
 #endif
 
@@ -45,6 +46,8 @@ void FAnimNode_MotionMatching::Evaluate_AnyThread(FPoseContext& Output)
 	DECLARE_SCOPE_HIERARCHICAL_COUNTER_ANIMNODE(Evaluate_AnyThread);
 	ANIM_MT_SCOPE_CYCLE_COUNTER_VERBOSE(MotionMatching, !IsInGameThread());
 
+	using namespace UE::PoseSearch;
+
 	FAnimNode_BlendStack_Standalone::Evaluate_AnyThread(Output);
 
 	const UE::Anim::IAnimRootMotionProvider* RootMotionProvider = UE::Anim::IAnimRootMotionProvider::Get();
@@ -132,7 +135,7 @@ void FAnimNode_MotionMatching::UpdateAssetPlayer(const FAnimationUpdateContext&
 	}
 
 #if ENABLE_ANIM_DEBUG
-	if (CVarAnimNodeMotionMatchingDrawInfo.GetValueOnAnyThread() > 0)
+	if (CVarAnimNodeMotionMatchingDrawInfo.GetValueOnAnyThread())
 	{
 		const UPoseSearchDatabase* CurrentDatabase = MotionMatchingState.CurrentSearchResult.Database;
 		const UAnimationAsset* CurrentAnimationAsset = AnimPlayers.IsEmpty() ? nullptr : AnimPlayers.First().GetAnimationAsset();
@@ -140,7 +143,7 @@ void FAnimNode_MotionMatching::UpdateAssetPlayer(const FAnimationUpdateContext&
 		FString DebugInfo = FString::Printf(TEXT("bForceInterruptNextUpdate(%d)\n"), bForceInterruptNextUpdate);
 		DebugInfo += FString::Printf(TEXT("Current Database(%s)\n"), *GetNameSafe(CurrentDatabase));
 		DebugInfo += FString::Printf(TEXT("Current Asset(%s)\n"), *GetNameSafe(CurrentAnimationAsset));
-		if (CVarAnimNodeMotionMatchingDrawInfoVerbose.GetValueOnAnyThread() > 0)
+		if (CVarAnimNodeMotionMatchingDrawInfoVerbose.GetValueOnAnyThread())
 		{
 			DebugInfo += FString::Printf(TEXT("Databases to search:\n"));
 			for (const UPoseSearchDatabase* DatabaseToSearch : DatabasesToSearch)
@@ -175,8 +178,9 @@ void FAnimNode_MotionMatching::UpdateAssetPlayer(const FAnimationUpdateContext&
 		bForceInterruptNextUpdate,
 		bShouldSearch
 		#if ENABLE_ANIM_DEBUG
-		, CVarAnimNodeMotionMatchingDrawQuery.GetValueOnAnyThread() > 0
-		, CVarAnimNodeMotionMatchingDrawCurResult.GetValueOnAnyThread() > 0
+		, CVarAnimNodeMotionMatchingDrawQuery.GetValueOnAnyThread()
+		, CVarAnimNodeMotionMatchingDrawCurResult.GetValueOnAnyThread()
+		, CVarAnimNodeMotionMatchingDrawPoseHistory.GetValueOnAnyThread()
 		#endif // ENABLE_ANIM_DEBUG
 	);
 
diff --git a/Engine/Plugins/Experimental/Animation/PoseSearch/Source/Runtime/Private/AnimNode_PoseSearchHistoryCollector.cpp b/Engine/Plugins/Experimental/Animation/PoseSearch/Source/Runtime/Private/AnimNode_PoseSearchHistoryCollector.cpp
--- a/Engine/Plugins/Experimental/Animation/PoseSearch/Source/Runtime/Private/AnimNode_PoseSearchHistoryCollector.cpp
+++ b/Engine/Plugins/Experimental/Animation/PoseSearch/Source/Runtime/Private/AnimNode_PoseSearchHistoryCollector.cpp
@@ -80,7 +80,7 @@ void FAnimNode_PoseSearchHistoryCollector_Base::CacheBones_AnyThread(const FAnim
 		Pose.SetBoneContainer(&Context.AnimInstanceProxy->GetRequiredBones());
 		FCSPose<FCompactPose> ComponentSpacePose;
 		ComponentSpacePose.InitPose(Pose);
-		PoseHistory.Update(0.f, ComponentSpacePose, Context.AnimInstanceProxy->GetComponentTransform());
+		PoseHistory.Update(0.f, ComponentSpacePose);
 	}
 }
 
@@ -128,12 +128,12 @@ void FAnimNode_PoseSearchHistoryCollector::Evaluate_AnyThread(FPoseContext& Outp
 
 	FCSPose<FCompactPose> ComponentSpacePose;
 	ComponentSpacePose.InitPose(Output.Pose);
-	PoseHistory.Update(Output.AnimInstanceProxy->GetDeltaSeconds(), ComponentSpacePose, Output.AnimInstanceProxy->GetComponentTransform());
+	PoseHistory.Update(Output.AnimInstanceProxy->GetDeltaSeconds(), ComponentSpacePose);
 
 #if ENABLE_DRAW_DEBUG && ENABLE_ANIM_DEBUG
 	if (CVarAnimPoseHistoryDebugDraw.GetValueOnAnyThread())
 	{
-		PoseHistory.DebugDraw(*Output.AnimInstanceProxy, DebugColor.ToFColor(true));
+		PoseHistory.DebugDraw(*Output.AnimInstanceProxy, DebugColor.ToFColor(true), &DebugDrawTrajectory);
 	}
 #endif // ENABLE_DRAW_DEBUG && ENABLE_ANIM_DEBUG
 }
@@ -180,12 +180,12 @@ void FAnimNode_PoseSearchComponentSpaceHistoryCollector::EvaluateComponentSpace_
 	Super::EvaluateComponentSpace_AnyThread(Output);
 	Source.EvaluateComponentSpace(Output);
 
-	PoseHistory.Update(Output.AnimInstanceProxy->GetDeltaSeconds(), Output.Pose, Output.AnimInstanceProxy->GetComponentTransform());
+	PoseHistory.Update(Output.AnimInstanceProxy->GetDeltaSeconds(), Output.Pose);
 
 #if ENABLE_DRAW_DEBUG && ENABLE_ANIM_DEBUG
 	if (CVarAnimPoseHistoryDebugDraw.GetValueOnAnyThread())
 	{
-		PoseHistory.DebugDraw(*Output.AnimInstanceProxy, DebugColor.ToFColor(true));
+		PoseHistory.DebugDraw(*Output.AnimInstanceProxy, DebugColor.ToFColor(true), &DebugDrawTrajectory);
 	}
 #endif // ENABLE_DRAW_DEBUG && ENABLE_ANIM_DEBUG
 }
diff --git a/Engine/Plugins/Experimental/Animation/PoseSearch/Source/Runtime/Private/PoseSearchAssetIndexer.cpp b/Engine/Plugins/Experimental/Animation/PoseSearch/Source/Runtime/Private/PoseSearchAssetIndexer.cpp
--- a/Engine/Plugins/Experimental/Animation/PoseSearch/Source/Runtime/Private/PoseSearchAssetIndexer.cpp
+++ b/Engine/Plugins/Experimental/Animation/PoseSearch/Source/Runtime/Private/PoseSearchAssetIndexer.cpp
@@ -187,26 +187,24 @@ void FAssetIndexer::ComputeStats()
 {
 	Stats = FStats();
 
-	check(SamplingContext.FiniteDelta > UE_KINDA_SMALL_NUMBER);
-
 	for (int32 SampleIdx = GetBeginSampleIdx(); SampleIdx != GetEndSampleIdx(); ++SampleIdx)
 	{
 		const float SampleTime = FMath::Min(CalculateSampleTime(SampleIdx), AssetSampler.GetPlayLength());
 
 		bool AnyClamped = false;
-		const FTransform TrajTransformsPast = GetTransform(SampleTime - SamplingContext.FiniteDelta, AnyClamped);
+		const FTransform TrajTransformsPast = GetTransform(SampleTime - FiniteDelta, AnyClamped);
 		if (!AnyClamped)
 		{
 			const FTransform TrajTransformsPresent = GetTransform(SampleTime, AnyClamped);
 			if (!AnyClamped)
 			{
-				const FTransform TrajTransformsFuture = GetTransform(SampleTime + SamplingContext.FiniteDelta, AnyClamped);
+				const FTransform TrajTransformsFuture = GetTransform(SampleTime + FiniteDelta, AnyClamped);
 				if (!AnyClamped)
 				{
 					// if any transform is clamped we just skip the sample entirely
-					const FVector LinearVelocityPresent = (TrajTransformsPresent.GetTranslation() - TrajTransformsPast.GetTranslation()) / SamplingContext.FiniteDelta;
-					const FVector LinearVelocityFuture = (TrajTransformsFuture.GetTranslation() - TrajTransformsPresent.GetTranslation()) / SamplingContext.FiniteDelta;
-					const FVector LinearAcceleration = (LinearVelocityFuture - LinearVelocityPresent) / SamplingContext.FiniteDelta;
+					const FVector LinearVelocityPresent = (TrajTransformsPresent.GetTranslation() - TrajTransformsPast.GetTranslation()) / FiniteDelta;
+					const FVector LinearVelocityFuture = (TrajTransformsFuture.GetTranslation() - TrajTransformsPresent.GetTranslation()) / FiniteDelta;
+					const FVector LinearAcceleration = (LinearVelocityFuture - LinearVelocityPresent) / FiniteDelta;
 
 					const float Speed = LinearVelocityPresent.Length();
 					const float Acceleration = LinearAcceleration.Length();
@@ -391,6 +389,8 @@ FAssetIndexer::CachedEntry& FAssetIndexer::GetEntry(float SampleTime)
 // returns the transform in component space for the bone indexed by Schema->BoneReferences[SchemaBoneIdx] at SampleTime seconds
 FTransform FAssetIndexer::GetComponentSpaceTransform(float SampleTime, bool& bClamped, int8 SchemaBoneIdx)
 {
+	using namespace UE::PoseSearch;
+
 	CachedEntry& Entry = GetEntry(SampleTime);
 	bClamped = Entry.bClamped;
 
@@ -407,6 +407,8 @@ FTransform FAssetIndexer::GetComponentSpaceTransform(float SampleTime, bool& bCl
 // returns the transform in animation space for the bone indexed by Schema->BoneReferences[SchemaBoneIdx] at SampleTime seconds
 FTransform FAssetIndexer::GetTransform(float SampleTime, bool& bClamped, int8 SchemaBoneIdx)
 {
+	using namespace UE::PoseSearch;
+
 	CachedEntry& Entry = GetEntry(SampleTime);
 	bClamped = Entry.bClamped;
 
@@ -447,6 +449,8 @@ float FAssetIndexer::CalculateSampleTime(int32 SampleIdx) const
 
 bool FAssetIndexer::GetSampleRotation(FQuat& OutSampleRotation, float SampleTimeOffset, float OriginTimeOffset, int32 SampleIdx, int8 SchemaSampleBoneIdx, int8 SchemaOriginBoneIdx, EPermutationTimeType PermutationTimeType, int32 SamplingAttributeId)
 {
+	using namespace UE::PoseSearch;
+
 	float PermutationSampleTimeOffset = 0.f;
 	float PermutationOriginTimeOffset = 0.f;
 	UPoseSearchFeatureChannel::GetPermutationTimeOffsets(PermutationTimeType, CalculatePermutationTimeOffset(), PermutationSampleTimeOffset, PermutationOriginTimeOffset);
@@ -537,6 +541,8 @@ bool FAssetIndexer::GetSamplePosition(FVector& OutSamplePosition, float SampleTi
 
 bool FAssetIndexer::GetSamplePositionInternal(FVector& OutSamplePosition, float SampleTime, float OriginTime, bool& bClamped, int8 SchemaSampleBoneIdx, int8 SchemaOriginBoneIdx, int32 SamplingAttributeId)
 {
+	using namespace UE::PoseSearch;
+
 	bool bUnused;
 	if (SamplingAttributeId >= 0)
 	{
@@ -633,7 +639,6 @@ bool FAssetIndexer::GetSampleVelocity(FVector& OutSampleVelocity, float SampleTi
 	const float Time = CalculateSampleTime(SampleIdx);
 	const float SampleTime = Time + SampleTimeOffset + PermutationSampleTimeOffset;
 	const float OriginTime = Time + OriginTimeOffset + PermutationOriginTimeOffset;
-	const float FiniteDelta = SamplingContext.FiniteDelta;
 
 	bool bUnused, bClampedPast;
 	FVector BonePositionPast, BonePositionPresent;
diff --git a/Engine/Plugins/Experimental/Animation/PoseSearch/Source/Runtime/Private/PoseSearchContext.cpp b/Engine/Plugins/Experimental/Animation/PoseSearch/Source/Runtime/Private/PoseSearchContext.cpp
--- a/Engine/Plugins/Experimental/Animation/PoseSearch/Source/Runtime/Private/PoseSearchContext.cpp
+++ b/Engine/Plugins/Experimental/Animation/PoseSearch/Source/Runtime/Private/PoseSearchContext.cpp
@@ -381,7 +381,7 @@ FSearchContext::FSearchContext(const UAnimInstance* InAnimInstance, const IPoseH
 	}
 }
 
-FQuat FSearchContext::GetSampleRotation(float SampleTimeOffset, float OriginTimeOffset, const UPoseSearchSchema* Schema, int8 SchemaSampleBoneIdx, int8 SchemaOriginBoneIdx, bool bUseHistoryRoot, EPermutationTimeType PermutationTimeType, const FQuat* SampleBoneRotationWorldOverride)
+FQuat FSearchContext::GetSampleRotation(float SampleTimeOffset, float OriginTimeOffset, const UPoseSearchSchema* Schema, int8 SchemaSampleBoneIdx, int8 SchemaOriginBoneIdx, EPermutationTimeType PermutationTimeType, const FQuat* SampleBoneRotationWorldOverride)
 {
 	float PermutationSampleTimeOffset = 0.f;
 	float PermutationOriginTimeOffset = 0.f;
@@ -390,10 +390,10 @@ FQuat FSearchContext::GetSampleRotation(float SampleTimeOffset, float OriginTime
 	const float SampleTime = SampleTimeOffset + PermutationSampleTimeOffset;
 	const float OriginTime = OriginTimeOffset + PermutationOriginTimeOffset;
 
-	return GetSampleRotationInternal(SampleTime, OriginTime, Schema, SchemaSampleBoneIdx, SchemaOriginBoneIdx, bUseHistoryRoot, SampleBoneRotationWorldOverride);
+	return GetSampleRotationInternal(SampleTime, OriginTime, Schema, SchemaSampleBoneIdx, SchemaOriginBoneIdx, SampleBoneRotationWorldOverride);
 }
 
-FVector FSearchContext::GetSamplePosition(float SampleTimeOffset, float OriginTimeOffset, const UPoseSearchSchema* Schema, int8 SchemaSampleBoneIdx, int8 SchemaOriginBoneIdx, bool bUseHistoryRoot, EPermutationTimeType PermutationTimeType, const FVector* SampleBonePositionWorldOverride)
+FVector FSearchContext::GetSamplePosition(float SampleTimeOffset, float OriginTimeOffset, const UPoseSearchSchema* Schema, int8 SchemaSampleBoneIdx, int8 SchemaOriginBoneIdx, EPermutationTimeType PermutationTimeType, const FVector* SampleBonePositionWorldOverride)
 {
 	float PermutationSampleTimeOffset = 0.f;
 	float PermutationOriginTimeOffset = 0.f;
@@ -401,76 +401,53 @@ FVector FSearchContext::GetSamplePosition(float SampleTimeOffset, float OriginTi
 
 	const float SampleTime = SampleTimeOffset + PermutationSampleTimeOffset;
 	const float OriginTime = OriginTimeOffset + PermutationOriginTimeOffset;
-	return GetSamplePositionInternal(SampleTime, OriginTime, Schema, SchemaSampleBoneIdx, SchemaOriginBoneIdx, bUseHistoryRoot, SampleBonePositionWorldOverride);
+	return GetSamplePositionInternal(SampleTime, OriginTime, Schema, SchemaSampleBoneIdx, SchemaOriginBoneIdx, SampleBonePositionWorldOverride);
 }
 
-FVector FSearchContext::GetSampleVelocity(float SampleTimeOffset, float OriginTimeOffset, const UPoseSearchSchema* Schema, int8 SchemaSampleBoneIdx, int8 SchemaOriginBoneIdx, bool bUseCharacterSpaceVelocities, bool bUseHistoryRoot, EPermutationTimeType PermutationTimeType, const FVector* SampleBoneVelocityWorldOverride)
+FVector FSearchContext::GetSampleVelocity(float SampleTimeOffset, float OriginTimeOffset, const UPoseSearchSchema* Schema, int8 SchemaSampleBoneIdx, int8 SchemaOriginBoneIdx, bool bUseCharacterSpaceVelocities, EPermutationTimeType PermutationTimeType, const FVector* SampleBoneVelocityWorldOverride)
 {
+	using namespace UE::PoseSearch;
+
 	float PermutationSampleTimeOffset = 0.f;
 	float PermutationOriginTimeOffset = 0.f;
 	UPoseSearchFeatureChannel::GetPermutationTimeOffsets(PermutationTimeType, DesiredPermutationTimeOffset, PermutationSampleTimeOffset, PermutationOriginTimeOffset);
 
 	const float SampleTime = SampleTimeOffset + PermutationSampleTimeOffset;
 	const float OriginTime = OriginTimeOffset + PermutationOriginTimeOffset;
-	const float FiniteDelta = History ? History->GetSampleTimeInterval() : 1 / 60.0f;
-	check(FiniteDelta > UE_KINDA_SMALL_NUMBER);
 
 	if (SampleBoneVelocityWorldOverride)
 	{
-		const FTransform RootBoneTransform = GetWorldBoneTransformAtTime(OriginTime, Schema, RootSchemaBoneIdx, bUseHistoryRoot);
+		const FTransform RootBoneTransform = GetWorldBoneTransformAtTime(OriginTime, Schema, RootSchemaBoneIdx);
 		return RootBoneTransform.InverseTransformVector(*SampleBoneVelocityWorldOverride);
 	}
 
 	// calculating the local Position for the bone indexed by SchemaSampleBoneIdx
-	const FVector PreviousTranslation = GetSamplePositionInternal(SampleTime - FiniteDelta, bUseCharacterSpaceVelocities ? OriginTime - FiniteDelta : OriginTime, Schema, SchemaSampleBoneIdx, SchemaOriginBoneIdx, bUseHistoryRoot);
-	const FVector CurrentTranslation = GetSamplePositionInternal(SampleTime, OriginTime, Schema, SchemaSampleBoneIdx, SchemaOriginBoneIdx, bUseHistoryRoot);
+	const FVector PreviousTranslation = GetSamplePositionInternal(SampleTime - FiniteDelta, bUseCharacterSpaceVelocities ? OriginTime - FiniteDelta : OriginTime, Schema, SchemaSampleBoneIdx, SchemaOriginBoneIdx);
+	const FVector CurrentTranslation = GetSamplePositionInternal(SampleTime, OriginTime, Schema, SchemaSampleBoneIdx, SchemaOriginBoneIdx);
 
 	const FVector LinearVelocity = (CurrentTranslation - PreviousTranslation) / FiniteDelta;
 	return LinearVelocity;
 }
 
-FTransform FSearchContext::GetWorldRootBoneTransformAtTime(float SampleTime, bool bUseHistoryRoot, bool bExtrapolate) const
+FTransform FSearchContext::GetWorldRootBoneTransformAtTime(float SampleTime, bool bExtrapolate) const
 {
-	FTransform WorldRootBoneTransform = FTransform::Identity;
-	if (bUseHistoryRoot)
-	{
 #if WITH_EDITOR
-		if (!History)
-		{
-			UE_LOG(LogPoseSearch, Error, TEXT("FSearchContext::GetWorldRootBoneTransformAtTime - missing History is required!"));
-		}
-		else
-#else // WITH_EDITOR
-		check(History);
-#endif // WITH_EDITOR
-		{
-			// asking for root bone (RootBoneIndexType) in world space (WorldSpaceIndexType)
-			History->GetTransformAtTime(SampleTime, WorldRootBoneTransform, AnimInstance->CurrentSkeleton, RootBoneIndexType, WorldSpaceIndexType, bExtrapolate);
-		}
-	}
-	else
+	if (!Trajectory)
 	{
-#if WITH_EDITOR
-		if (!Trajectory)
-		{
-			UE_LOG(LogPoseSearch, Error, TEXT("FSearchContext::GetWorldRootBoneTransformAtTime - missing Trajectory is required!"));
-		}
-		else
-#else // WITH_EDITOR
-		check(Trajectory);
-#endif // WITH_EDITOR
-		{
-			// Trajectory is already in root bone world space (transformed in UPoseSearchLibrary::ProcessTrajectory), so we just ask for a sample at the proper SampleTime 
-			WorldRootBoneTransform = Trajectory->GetSampleAtTime(SampleTime, bExtrapolate).GetTransform();
-		}
+		UE_LOG(LogPoseSearch, Error, TEXT("FSearchContext::GetWorldRootBoneTransformAtTime - missing Trajectory is required!"));
+		return FTransform::Identity;
 	}
+#endif // WITH_EDITOR
 
-	return WorldRootBoneTransform;
+	check(Trajectory);
+	
+	// Trajectory is already in root bone world space (transformed in UPoseSearchLibrary::ProcessTrajectory), so we just ask for a sample at the proper SampleTime 
+	return Trajectory->GetSampleAtTime(SampleTime, bExtrapolate).GetTransform();
 }
 
-FTransform FSearchContext::GetWorldBoneTransformAtTime(float SampleTime, const UPoseSearchSchema* Schema, int8 SchemaBoneIdx, bool bUseHistoryRoot)
+FTransform FSearchContext::GetWorldBoneTransformAtTime(float SampleTime, const UPoseSearchSchema* Schema, int8 SchemaBoneIdx)
 {
-	const FTransform WorldRootBoneTransform = GetWorldRootBoneTransformAtTime(SampleTime, bUseHistoryRoot);
+	const FTransform WorldRootBoneTransform = GetWorldRootBoneTransformAtTime(SampleTime);
 
 	const FBoneIndexType BoneIndexType = Schema->GetBoneIndexType(SchemaBoneIdx);
 	if (BoneIndexType != RootBoneIndexType)
@@ -519,17 +496,17 @@ FTransform FSearchContext::GetLocalBoneTransformAtTime(float SampleTime, const U
 	return FTransform::Identity;
 }
 
-FVector FSearchContext::GetSamplePositionInternal(float SampleTime, float OriginTime, const UPoseSearchSchema* Schema, int8 SchemaSampleBoneIdx, int8 SchemaOriginBoneIdx, bool bUseHistoryRoot, const FVector* SampleBonePositionWorldOverride)
+FVector FSearchContext::GetSamplePositionInternal(float SampleTime, float OriginTime, const UPoseSearchSchema* Schema, int8 SchemaSampleBoneIdx, int8 SchemaOriginBoneIdx, const FVector* SampleBonePositionWorldOverride)
 {
 	if (SampleBonePositionWorldOverride)
 	{
-		const FTransform RootBoneTransform = GetWorldBoneTransformAtTime(OriginTime, Schema, RootSchemaBoneIdx, bUseHistoryRoot);
+		const FTransform RootBoneTransform = GetWorldBoneTransformAtTime(OriginTime, Schema, RootSchemaBoneIdx);
 		if (SchemaOriginBoneIdx == RootSchemaBoneIdx)
 		{
 			return RootBoneTransform.InverseTransformPosition(*SampleBonePositionWorldOverride);
 		}
 
-		const FTransform OriginBoneTransform = GetWorldBoneTransformAtTime(OriginTime, Schema, SchemaOriginBoneIdx, bUseHistoryRoot);
+		const FTransform OriginBoneTransform = GetWorldBoneTransformAtTime(OriginTime, Schema, SchemaOriginBoneIdx);
 		const FVector DeltaBoneTranslation = *SampleBonePositionWorldOverride - OriginBoneTransform.GetTranslation();
 		return RootBoneTransform.InverseTransformVector(DeltaBoneTranslation);
 	}
@@ -546,29 +523,29 @@ FVector FSearchContext::GetSamplePositionInternal(float SampleTime, float Origin
 		return SampleBonePosition - OriginBonePosition;
 	}
 
-	const FTransform RootBoneTransform = GetWorldBoneTransformAtTime(OriginTime, Schema, RootSchemaBoneIdx, bUseHistoryRoot);
-	const FTransform SampleBoneTransform = GetWorldBoneTransformAtTime(SampleTime, Schema, SchemaSampleBoneIdx, bUseHistoryRoot);
+	const FTransform RootBoneTransform = GetWorldBoneTransformAtTime(OriginTime, Schema, RootSchemaBoneIdx);
+	const FTransform SampleBoneTransform = GetWorldBoneTransformAtTime(SampleTime, Schema, SchemaSampleBoneIdx);
 	if (SchemaOriginBoneIdx == RootSchemaBoneIdx)
 	{
 		return RootBoneTransform.InverseTransformPosition(SampleBoneTransform.GetTranslation());
 	}
 
-	const FTransform OriginBoneTransform = GetWorldBoneTransformAtTime(OriginTime, Schema, SchemaOriginBoneIdx, bUseHistoryRoot);
+	const FTransform OriginBoneTransform = GetWorldBoneTransformAtTime(OriginTime, Schema, SchemaOriginBoneIdx);
 	const FVector DeltaBoneTranslation = SampleBoneTransform.GetTranslation() - OriginBoneTransform.GetTranslation();
 	return RootBoneTransform.InverseTransformVector(DeltaBoneTranslation);
 }
 
-FQuat FSearchContext::GetSampleRotationInternal(float SampleTime, float OriginTime, const UPoseSearchSchema* Schema, int8 SchemaSampleBoneIdx, int8 SchemaOriginBoneIdx, bool bUseHistoryRoot, const FQuat* SampleBoneRotationWorldOverride)
+FQuat FSearchContext::GetSampleRotationInternal(float SampleTime, float OriginTime, const UPoseSearchSchema* Schema, int8 SchemaSampleBoneIdx, int8 SchemaOriginBoneIdx, const FQuat* SampleBoneRotationWorldOverride)
 {
 	if (SampleBoneRotationWorldOverride)
 	{
-		const FTransform RootBoneTransform = GetWorldBoneTransformAtTime(OriginTime, Schema, RootSchemaBoneIdx, bUseHistoryRoot);
+		const FTransform RootBoneTransform = GetWorldBoneTransformAtTime(OriginTime, Schema, RootSchemaBoneIdx);
 		if (SchemaOriginBoneIdx == RootSchemaBoneIdx)
 		{
 			return RootBoneTransform.InverseTransformRotation(*SampleBoneRotationWorldOverride);
 		}
 
-		const FTransform OriginBoneTransform = GetWorldBoneTransformAtTime(OriginTime, Schema, SchemaOriginBoneIdx, bUseHistoryRoot);
+		const FTransform OriginBoneTransform = GetWorldBoneTransformAtTime(OriginTime, Schema, SchemaOriginBoneIdx);
 		const FQuat DeltaBoneRotation = OriginBoneTransform.InverseTransformRotation(*SampleBoneRotationWorldOverride);
 		return RootBoneTransform.InverseTransformRotation(DeltaBoneRotation);
 	}
@@ -585,14 +562,14 @@ FQuat FSearchContext::GetSampleRotationInternal(float SampleTime, float OriginTi
 		return OriginBoneTransform.InverseTransformRotation(SampleBoneTransform.GetRotation());
 	}
 
-	const FTransform RootBoneTransform = GetWorldBoneTransformAtTime(OriginTime, Schema, RootSchemaBoneIdx, bUseHistoryRoot);
-	const FTransform SampleBoneTransform = GetWorldBoneTransformAtTime(SampleTime, Schema, SchemaSampleBoneIdx, bUseHistoryRoot);
+	const FTransform RootBoneTransform = GetWorldBoneTransformAtTime(OriginTime, Schema, RootSchemaBoneIdx);
+	const FTransform SampleBoneTransform = GetWorldBoneTransformAtTime(SampleTime, Schema, SchemaSampleBoneIdx);
 	if (SchemaOriginBoneIdx == RootSchemaBoneIdx)
 	{
 		return RootBoneTransform.InverseTransformRotation(SampleBoneTransform.GetRotation());
 	}
 
-	const FTransform OriginBoneTransform = GetWorldBoneTransformAtTime(OriginTime, Schema, SchemaOriginBoneIdx, bUseHistoryRoot);
+	const FTransform OriginBoneTransform = GetWorldBoneTransformAtTime(OriginTime, Schema, SchemaOriginBoneIdx);
 	const FQuat DeltaBoneRotation = OriginBoneTransform.InverseTransformRotation(SampleBoneTransform.GetRotation());
 	return RootBoneTransform.InverseTransformRotation(DeltaBoneRotation);
 }
diff --git a/Engine/Plugins/Experimental/Animation/PoseSearch/Source/Runtime/Private/PoseSearchFeatureChannel_Heading.cpp b/Engine/Plugins/Experimental/Animation/PoseSearch/Source/Runtime/Private/PoseSearchFeatureChannel_Heading.cpp
--- a/Engine/Plugins/Experimental/Animation/PoseSearch/Source/Runtime/Private/PoseSearchFeatureChannel_Heading.cpp
+++ b/Engine/Plugins/Experimental/Animation/PoseSearch/Source/Runtime/Private/PoseSearchFeatureChannel_Heading.cpp
@@ -92,7 +92,7 @@ void UPoseSearchFeatureChannel_Heading::BuildQuery(UE::PoseSearch::FSearchContex
 	if (bUseBlueprintQueryOverride)
 	{
 		const FQuat BoneRotationWorld = BP_GetWorldRotation(SearchContext.GetAnimInstance());
-		const FQuat BoneRotation = SearchContext.GetSampleRotation(SampleTimeOffset, OriginTimeOffset, InOutQuery.GetSchema(), SchemaBoneIdx, SchemaOriginBoneIdx, /*!bIsRootBone*/ true, EPermutationTimeType::UseSampleTime, &BoneRotationWorld);
+		const FQuat BoneRotation = SearchContext.GetSampleRotation(SampleTimeOffset, OriginTimeOffset, InOutQuery.GetSchema(), SchemaBoneIdx, SchemaOriginBoneIdx, EPermutationTimeType::UseSampleTime, &BoneRotationWorld);
 		FFeatureVectorHelper::EncodeVector(InOutQuery.EditValues(), ChannelDataOffset, GetAxis(BoneRotation), ComponentStripping);
 	}
 	else
@@ -114,7 +114,7 @@ void UPoseSearchFeatureChannel_Heading::BuildQuery(UE::PoseSearch::FSearchContex
 		else
 		{
 			// calculating the BoneRotation in component space for the bone indexed by SchemaBoneIdx
-			const FQuat BoneRotation = SearchContext.GetSampleRotation(SampleTimeOffset, OriginTimeOffset, InOutQuery.GetSchema(), SchemaBoneIdx, SchemaOriginBoneIdx, !bIsRootBone);
+			const FQuat BoneRotation = SearchContext.GetSampleRotation(SampleTimeOffset, OriginTimeOffset, InOutQuery.GetSchema(), SchemaBoneIdx, SchemaOriginBoneIdx, PermutationTimeType);
 			FFeatureVectorHelper::EncodeVector(InOutQuery.EditValues(), ChannelDataOffset, GetAxis(BoneRotation), ComponentStripping);
 		}
 	}
@@ -127,7 +127,7 @@ void UPoseSearchFeatureChannel_Heading::DebugDraw(const UE::PoseSearch::FDebugDr
 
 	const FColor Color = DebugColor.ToFColor(true);
 	const FVector BoneHeading = DrawParams.ExtractRotation(PoseVector, OriginTimeOffset).RotateVector(FFeatureVectorHelper::DecodeVector(PoseVector, ChannelDataOffset, ComponentStripping));
-	const FVector BonePos = DrawParams.ExtractPosition(PoseVector, SampleTimeOffset, SchemaBoneIdx, EPermutationTimeType::UseSampleTime, SamplingAttributeId);
+	const FVector BonePos = DrawParams.ExtractPosition(PoseVector, SampleTimeOffset, SchemaBoneIdx, PermutationTimeType, SamplingAttributeId);
 
 	DrawParams.DrawPoint(BonePos, Color, 3.f);
 	DrawParams.DrawLine(BonePos + BoneHeading * 4.f, BonePos + BoneHeading * 15.f, Color);
@@ -150,7 +150,7 @@ bool UPoseSearchFeatureChannel_Heading::IndexAsset(UE::PoseSearch::FAssetIndexer
 	FQuat SampleRotation = FQuat::Identity;
 	for (int32 SampleIdx = Indexer.GetBeginSampleIdx(); SampleIdx != Indexer.GetEndSampleIdx(); ++SampleIdx)
 	{
-		if (Indexer.GetSampleRotation(SampleRotation, SampleTimeOffset, OriginTimeOffset, SampleIdx, SchemaBoneIdx, SchemaOriginBoneIdx, EPermutationTimeType::UseSampleTime, SamplingAttributeId))
+		if (Indexer.GetSampleRotation(SampleRotation, SampleTimeOffset, OriginTimeOffset, SampleIdx, SchemaBoneIdx, SchemaOriginBoneIdx, PermutationTimeType, SamplingAttributeId))
 		{
 			FFeatureVectorHelper::EncodeVector(Indexer.GetPoseVector(SampleIdx), ChannelDataOffset, GetAxis(SampleRotation), ComponentStripping);
 		}
@@ -164,6 +164,8 @@ bool UPoseSearchFeatureChannel_Heading::IndexAsset(UE::PoseSearch::FAssetIndexer
 
 FString UPoseSearchFeatureChannel_Heading::GetLabel() const
 {
+	using namespace UE::PoseSearch;
+
 	TStringBuilder<256> Label;
 	if (const UPoseSearchFeatureChannel* OuterChannel = Cast<UPoseSearchFeatureChannel>(GetOuter()))
 	{
diff --git a/Engine/Plugins/Experimental/Animation/PoseSearch/Source/Runtime/Private/PoseSearchFeatureChannel_Phase.cpp b/Engine/Plugins/Experimental/Animation/PoseSearch/Source/Runtime/Private/PoseSearchFeatureChannel_Phase.cpp
--- a/Engine/Plugins/Experimental/Animation/PoseSearch/Source/Runtime/Private/PoseSearchFeatureChannel_Phase.cpp
+++ b/Engine/Plugins/Experimental/Animation/PoseSearch/Source/Runtime/Private/PoseSearchFeatureChannel_Phase.cpp
@@ -406,6 +406,8 @@ bool UPoseSearchFeatureChannel_Phase::IndexAsset(UE::PoseSearch::FAssetIndexer&
 
 FString UPoseSearchFeatureChannel_Phase::GetLabel() const
 {
+	using namespace UE::PoseSearch;
+
 	TStringBuilder<256> Label;
 	if (const UPoseSearchFeatureChannel* OuterChannel = Cast<UPoseSearchFeatureChannel>(GetOuter()))
 	{
diff --git a/Engine/Plugins/Experimental/Animation/PoseSearch/Source/Runtime/Private/PoseSearchFeatureChannel_Position.cpp b/Engine/Plugins/Experimental/Animation/PoseSearch/Source/Runtime/Private/PoseSearchFeatureChannel_Position.cpp
--- a/Engine/Plugins/Experimental/Animation/PoseSearch/Source/Runtime/Private/PoseSearchFeatureChannel_Position.cpp
+++ b/Engine/Plugins/Experimental/Animation/PoseSearch/Source/Runtime/Private/PoseSearchFeatureChannel_Position.cpp
@@ -50,6 +50,8 @@ void UPoseSearchFeatureChannel_Position::Finalize(UPoseSearchSchema* Schema)
 
 void UPoseSearchFeatureChannel_Position::AddDependentChannels(UPoseSearchSchema* Schema) const
 {
+	using namespace UE::PoseSearch;
+
 	if (Schema->bInjectAdditionalDebugChannels)
 	{
 		if (SchemaOriginBoneIdx != RootSchemaBoneIdx)
@@ -65,17 +67,17 @@ void UPoseSearchFeatureChannel_Position::BuildQuery(UE::PoseSearch::FSearchConte
 	using namespace UE::PoseSearch;
 
 	check(InOutQuery.GetSchema());
-	const bool bIsRootBone = SchemaBoneIdx == RootSchemaBoneIdx;
 	if (bUseBlueprintQueryOverride)
 	{
 		const FVector BonePositionWorld = BP_GetWorldPosition(SearchContext.GetAnimInstance());
-		const FVector BonePosition = SearchContext.GetSamplePosition(SampleTimeOffset, OriginTimeOffset, InOutQuery.GetSchema(), SchemaBoneIdx, SchemaOriginBoneIdx, /*!bIsRootBone*/ true, PermutationTimeType, &BonePositionWorld);
+		const FVector BonePosition = SearchContext.GetSamplePosition(SampleTimeOffset, OriginTimeOffset, InOutQuery.GetSchema(), SchemaBoneIdx, SchemaOriginBoneIdx, PermutationTimeType, &BonePositionWorld);
   		FFeatureVectorHelper::EncodeVector(InOutQuery.EditValues(), ChannelDataOffset, BonePosition, ComponentStripping);
 	}
 	else
 	{
 		const bool bIsCurrentResultValid = SearchContext.GetCurrentResult().IsValid() && SearchContext.GetCurrentResult().Database->Schema == InOutQuery.GetSchema();
 		const bool bSkip = InputQueryPose != EInputQueryPose::UseCharacterPose && bIsCurrentResultValid;
+		const bool bIsRootBone = SchemaBoneIdx == RootSchemaBoneIdx;
 		if (bSkip || (!SearchContext.IsHistoryValid() && !bIsRootBone))
 		{
 			if (bIsCurrentResultValid)
@@ -91,7 +93,7 @@ void UPoseSearchFeatureChannel_Position::BuildQuery(UE::PoseSearch::FSearchConte
 		else
 		{
 			// calculating the BonePosition in root bone space for the bone indexed by SchemaBoneIdx
-			const FVector BonePosition = SearchContext.GetSamplePosition(SampleTimeOffset, OriginTimeOffset, InOutQuery.GetSchema(), SchemaBoneIdx, SchemaOriginBoneIdx, !bIsRootBone, PermutationTimeType);
+			const FVector BonePosition = SearchContext.GetSamplePosition(SampleTimeOffset, OriginTimeOffset, InOutQuery.GetSchema(), SchemaBoneIdx, SchemaOriginBoneIdx, PermutationTimeType);
 			FFeatureVectorHelper::EncodeVector(InOutQuery.EditValues(), ChannelDataOffset, BonePosition, ComponentStripping);
 		}
 	}
@@ -153,6 +155,8 @@ bool UPoseSearchFeatureChannel_Position::IndexAsset(UE::PoseSearch::FAssetIndexe
 
 FString UPoseSearchFeatureChannel_Position::GetLabel() const
 {
+	using namespace UE::PoseSearch;
+
 	TStringBuilder<256> Label;
 	if (const UPoseSearchFeatureChannel* OuterChannel = Cast<UPoseSearchFeatureChannel>(GetOuter()))
 	{
diff --git a/Engine/Plugins/Experimental/Animation/PoseSearch/Source/Runtime/Private/PoseSearchFeatureChannel_Velocity.cpp b/Engine/Plugins/Experimental/Animation/PoseSearch/Source/Runtime/Private/PoseSearchFeatureChannel_Velocity.cpp
--- a/Engine/Plugins/Experimental/Animation/PoseSearch/Source/Runtime/Private/PoseSearchFeatureChannel_Velocity.cpp
+++ b/Engine/Plugins/Experimental/Animation/PoseSearch/Source/Runtime/Private/PoseSearchFeatureChannel_Velocity.cpp
@@ -43,7 +43,7 @@ void UPoseSearchFeatureChannel_Velocity::BuildQuery(UE::PoseSearch::FSearchConte
 	{
 		const FVector LinearVelocityWorld = BP_GetWorldVelocity(SearchContext.GetAnimInstance());
 
-		FVector LinearVelocity = SearchContext.GetSampleVelocity(SampleTimeOffset, OriginTimeOffset, InOutQuery.GetSchema(), SchemaBoneIdx, SchemaOriginBoneIdx, bUseCharacterSpaceVelocities, /*!bIsRootBone*/ true, EPermutationTimeType::UseSampleTime, &LinearVelocityWorld);
+		FVector LinearVelocity = SearchContext.GetSampleVelocity(SampleTimeOffset, OriginTimeOffset, InOutQuery.GetSchema(), SchemaBoneIdx, SchemaOriginBoneIdx, bUseCharacterSpaceVelocities, EPermutationTimeType::UseSampleTime, &LinearVelocityWorld);
 		if (bNormalize)
 		{
 			LinearVelocity = LinearVelocity.GetClampedToMaxSize(1.f);
@@ -69,7 +69,7 @@ void UPoseSearchFeatureChannel_Velocity::BuildQuery(UE::PoseSearch::FSearchConte
 		else
 		{
 			// calculating the LinearVelocity for the bone indexed by SchemaBoneIdx
-			FVector LinearVelocity = SearchContext.GetSampleVelocity(SampleTimeOffset, OriginTimeOffset, InOutQuery.GetSchema(), SchemaBoneIdx, SchemaOriginBoneIdx, bUseCharacterSpaceVelocities, !bIsRootBone, EPermutationTimeType::UseSampleTime);
+			FVector LinearVelocity = SearchContext.GetSampleVelocity(SampleTimeOffset, OriginTimeOffset, InOutQuery.GetSchema(), SchemaBoneIdx, SchemaOriginBoneIdx, bUseCharacterSpaceVelocities, PermutationTimeType);
 			if (bNormalize)
 			{
 				LinearVelocity = LinearVelocity.GetClampedToMaxSize(1.f);
@@ -90,7 +90,7 @@ void UPoseSearchFeatureChannel_Velocity::DebugDraw(const UE::PoseSearch::FDebugD
 
 	const FVector LinearVelocity = DrawParams.GetRootTransform().TransformVector(FFeatureVectorHelper::DecodeVector(PoseVector, ChannelDataOffset, ComponentStripping));
 	const FVector BoneVelDirection = LinearVelocity.GetSafeNormal();
-	const FVector BonePos = DrawParams.ExtractPosition(PoseVector, SampleTimeOffset, SchemaBoneIdx, EPermutationTimeType::UseSampleTime, SamplingAttributeId);
+	const FVector BonePos = DrawParams.ExtractPosition(PoseVector, SampleTimeOffset, SchemaBoneIdx, PermutationTimeType, SamplingAttributeId);
 
 	DrawParams.DrawLine(BonePos, BonePos + LinearVelocity * LinearVelocityScale, Color);
 }
@@ -112,7 +112,7 @@ bool UPoseSearchFeatureChannel_Velocity::IndexAsset(UE::PoseSearch::FAssetIndexe
 	FVector LinearVelocity;
 	for (int32 SampleIdx = Indexer.GetBeginSampleIdx(); SampleIdx != Indexer.GetEndSampleIdx(); ++SampleIdx)
 	{
-		if (Indexer.GetSampleVelocity(LinearVelocity, SampleTimeOffset, OriginTimeOffset, SampleIdx, SchemaBoneIdx, SchemaOriginBoneIdx, bUseCharacterSpaceVelocities, EPermutationTimeType::UseSampleTime, SamplingAttributeId))
+		if (Indexer.GetSampleVelocity(LinearVelocity, SampleTimeOffset, OriginTimeOffset, SampleIdx, SchemaBoneIdx, SchemaOriginBoneIdx, bUseCharacterSpaceVelocities, PermutationTimeType, SamplingAttributeId))
 		{
 			if (bNormalize)
 			{
@@ -130,6 +130,8 @@ bool UPoseSearchFeatureChannel_Velocity::IndexAsset(UE::PoseSearch::FAssetIndexe
 
 FString UPoseSearchFeatureChannel_Velocity::GetLabel() const
 {
+	using namespace UE::PoseSearch;
+
 	TStringBuilder<256> Label;
 	if (const UPoseSearchFeatureChannel* OuterChannel = Cast<UPoseSearchFeatureChannel>(GetOuter()))
 	{
diff --git a/Engine/Plugins/Experimental/Animation/PoseSearch/Source/Runtime/Private/PoseSearchHistory.cpp b/Engine/Plugins/Experimental/Animation/PoseSearch/Source/Runtime/Private/PoseSearchHistory.cpp
--- a/Engine/Plugins/Experimental/Animation/PoseSearch/Source/Runtime/Private/PoseSearchHistory.cpp
+++ b/Engine/Plugins/Experimental/Animation/PoseSearch/Source/Runtime/Private/PoseSearchHistory.cpp
@@ -12,6 +12,7 @@
 #include "PoseSearch/PoseSearchResult.h"
 #include "PoseSearch/PoseSearchDatabase.h"
 #include "PoseSearch/PoseSearchDefines.h"
+#include "PoseSearch/PoseSearchTrajectoryTypes.h"
 
 IMPLEMENT_ANIMGRAPH_MESSAGE(UE::PoseSearch::IPoseHistoryProvider);
 
@@ -66,10 +67,9 @@ FORCEINLINE auto LowerBound(IteratorType First, IteratorType Last, const ValueTy
 
 //////////////////////////////////////////////////////////////////////////
 // FPoseHistoryEntry
-void FPoseHistoryEntry::Update(float InTime, FCSPose<FCompactPose>& ComponentSpacePose, const FTransform& ComponentTransform, const FBoneToTransformMap& BoneToTransformMap)
+void FPoseHistoryEntry::Update(float InTime, FCSPose<FCompactPose>& ComponentSpacePose, const FBoneToTransformMap& BoneToTransformMap)
 {
 	Time = InTime;
-	RootTransform = ComponentTransform;
 
 	const FBoneContainer& BoneContainer = ComponentSpacePose.GetPose().GetBoneContainer();
 	const USkeleton* SkeletonAsset = BoneContainer.GetSkeletonAsset();
@@ -145,7 +145,7 @@ FBoneIndexType FPoseHistory::GetRemappedBoneIndexType(FBoneIndexType BoneIndexTy
 FComponentSpaceTransformIndex FPoseHistory::GetRemappedComponentSpaceTransformIndex(const USkeleton* BoneIndexSkeleton, const USkeleton* LastUpdateSkeleton, const FBoneToTransformMap& BoneToTransformMap, FBoneIndexType BoneIndexType, bool& bSuccess)
 {
 	FComponentSpaceTransformIndex BoneTransformIndex = FComponentSpaceTransformIndex(BoneIndexType);
-	if (BoneIndexType < WorldSpaceIndexType)
+	if (BoneIndexType != ComponentSpaceIndexType)
 	{
 		BoneIndexType = GetRemappedBoneIndexType(BoneIndexType, BoneIndexSkeleton, LastUpdateSkeleton);
 
@@ -165,35 +165,38 @@ FComponentSpaceTransformIndex FPoseHistory::GetRemappedComponentSpaceTransformIn
 	return BoneTransformIndex;
 }
 
-bool FPoseHistory::LerpEntries(const FPoseHistoryEntry& PrevEntry, const FPoseHistoryEntry& NextEntry, float LerpValue, const USkeleton* BoneIndexSkeleton, const USkeleton* LastUpdateSkeleton, const FBoneToTransformMap& BoneToTransformMap, FBoneIndexType BoneIndexType, FBoneIndexType ReferenceBoneIndexType, FTransform& OutBoneTransform)
+bool FPoseHistory::LerpEntries(float Time, bool bExtrapolate, const FPoseHistoryEntry& PrevEntry, const FPoseHistoryEntry& NextEntry, const USkeleton* BoneIndexSkeleton, const USkeleton* LastUpdateSkeleton,
+	const FBoneToTransformMap& BoneToTransformMap, FBoneIndexType BoneIndexType, FBoneIndexType ReferenceBoneIndexType, FTransform& OutBoneTransform)
 {
 	bool bSuccess = true;
 
+	const float Denominator = NextEntry.Time - PrevEntry.Time;
+	float LerpValue = 0.f;
+	if (!FMath::IsNearlyZero(Denominator))
+	{
+		const float SecondsAgo = -Time;
+		const float Numerator = SecondsAgo - PrevEntry.Time;
+		LerpValue = bExtrapolate ? Numerator / Denominator : FMath::Clamp(Numerator / Denominator, 0.f, 1.f);
+	}
+
 	const FComponentSpaceTransformIndex BoneTransformIndex = GetRemappedComponentSpaceTransformIndex(BoneIndexSkeleton, LastUpdateSkeleton, BoneToTransformMap, BoneIndexType, bSuccess);
 	const FComponentSpaceTransformIndex ReferenceBoneTransformIndex = GetRemappedComponentSpaceTransformIndex(BoneIndexSkeleton, LastUpdateSkeleton, BoneToTransformMap, ReferenceBoneIndexType, bSuccess);
 
-	if (BoneTransformIndex < WorldSpaceIndexType)
+	if (BoneTransformIndex != ComponentSpaceIndexType)
 	{
-		switch (ReferenceBoneTransformIndex)
+		if (ReferenceBoneTransformIndex == ComponentSpaceIndexType)
 		{
-		case ComponentSpaceIndexType:
 			OutBoneTransform.Blend(
 				PrevEntry.ComponentSpaceTransforms[BoneTransformIndex],
 				NextEntry.ComponentSpaceTransforms[BoneTransformIndex],
 				LerpValue);
-			break;
-		case WorldSpaceIndexType:
-			OutBoneTransform.Blend(
-				PrevEntry.ComponentSpaceTransforms[BoneTransformIndex] * PrevEntry.RootTransform,
-				NextEntry.ComponentSpaceTransforms[BoneTransformIndex] * NextEntry.RootTransform,
-				LerpValue);
-			break;
-		default:
+		}
+		else
+		{
 			OutBoneTransform.Blend(
 				PrevEntry.ComponentSpaceTransforms[BoneTransformIndex] * PrevEntry.ComponentSpaceTransforms[ReferenceBoneTransformIndex].Inverse(),
 				NextEntry.ComponentSpaceTransforms[BoneTransformIndex] * NextEntry.ComponentSpaceTransforms[ReferenceBoneTransformIndex].Inverse(),
 				LerpValue);
-			break;
 		}
 	}
 	else
@@ -209,8 +212,8 @@ bool FPoseHistory::LerpEntries(const FPoseHistoryEntry& PrevEntry, const FPoseHi
 
 bool FPoseHistory::GetTransformAtTime(float Time, FTransform& OutBoneTransform, const USkeleton* BoneIndexSkeleton, FBoneIndexType BoneIndexType, FBoneIndexType ReferenceBoneIndexType, bool bExtrapolate) const
 {
-	static_assert(RootBoneIndexType == 0 && WorldSpaceIndexType < ComponentSpaceIndexType); // some assumptions
-	check(BoneIndexType < WorldSpaceIndexType);
+	static_assert(RootBoneIndexType == 0 && ComponentSpaceIndexType == FBoneIndexType(-1)); // some assumptions
+	check(BoneIndexType != ComponentSpaceIndexType);
 	
 	const int32 Num = Entries.Num();
 	
@@ -231,15 +234,7 @@ bool FPoseHistory::GetTransformAtTime(float Time, FTransform& OutBoneTransform,
 		const FPoseHistoryEntry& PrevEntry = Entries[PrevIdx];
 		const FPoseHistoryEntry& NextEntry = Entries[NextIdx];
 
-		const float Denominator = NextEntry.Time - PrevEntry.Time;
-		float LerpValue = 0.f;
-		if (!FMath::IsNearlyZero(Denominator))
-		{
-			const float Numerator = SecondsAgo - PrevEntry.Time;
-			LerpValue = bExtrapolate ? Numerator / Denominator : FMath::Clamp(Numerator / Denominator, 0.f, 1.f);
-		}
-
-		return LerpEntries(PrevEntry, NextEntry, LerpValue, BoneIndexSkeleton, GetLastUpdateSkeleton(), BoneToTransformMap, BoneIndexType, ReferenceBoneIndexType, OutBoneTransform);
+		return LerpEntries(Time, bExtrapolate, PrevEntry, NextEntry, BoneIndexSkeleton, GetLastUpdateSkeleton(), BoneToTransformMap, BoneIndexType, ReferenceBoneIndexType, OutBoneTransform);
 	}
 	
 	OutBoneTransform = FTransform::Identity;
@@ -256,7 +251,7 @@ void FPoseHistory::ClearHistory()
 	Entries.Reset();
 }
 
-void FPoseHistory::Update(float SecondsElapsed, FCSPose<FCompactPose>& ComponentSpacePose, const FTransform& ComponentTransform)
+void FPoseHistory::Update(float SecondsElapsed, FCSPose<FCompactPose>& ComponentSpacePose)
 {
 	const USkeleton* Skeleton = ComponentSpacePose.GetPose().GetBoneContainer().GetSkeletonAsset();
 	if (LastUpdateSkeleton != Skeleton)
@@ -272,7 +267,8 @@ void FPoseHistory::Update(float SecondsElapsed, FCSPose<FCompactPose>& Component
 		Entry.Time += SecondsElapsed;
 	}
 
-	if (Entries.Num() != Entries.Max())
+	const int32 EntriesMax = Entries.Max();
+	if (Entries.Num() != EntriesMax)
 	{
 		// Consume every pose until the queue is full
 		Entries.Emplace();
@@ -284,7 +280,9 @@ void FPoseHistory::Update(float SecondsElapsed, FCSPose<FCompactPose>& Component
 		// want to evenly distribute poses across the entire history buffer so we only push additional
 		// poses when enough time has elapsed.
 
-		const float SampleInterval = GetSampleTimeInterval();
+		check(EntriesMax >= 2 && TimeHorizon > UE_KINDA_SMALL_NUMBER);
+		// Reserve one pose for computing derivatives at the time horizon
+		const float SampleInterval = TimeHorizon / (EntriesMax - 1);
 
 		bool bCanEvictOldest = Entries[1].Time >= TimeHorizon + SampleInterval;
 		bool bShouldPushNewest = Entries[Entries.Num() - 2].Time >= SampleInterval;
@@ -298,21 +296,14 @@ void FPoseHistory::Update(float SecondsElapsed, FCSPose<FCompactPose>& Component
 	}
 
 	// Regardless of the retention policy, we always update the most recent Entry
-	Entries.Last().Update(0.f, ComponentSpacePose, ComponentTransform, BoneToTransformMap);
-}
-
-float FPoseHistory::GetSampleTimeInterval() const
-{
-	const int32 EntriesMax = Entries.Max();
-	check(EntriesMax >= 2 && TimeHorizon > UE_KINDA_SMALL_NUMBER);
-
-	// Reserve one pose for computing derivatives at the time horizon
-	return TimeHorizon / (EntriesMax - 1);
+	Entries.Last().Update(0.f, ComponentSpacePose, BoneToTransformMap);
 }
 
 #if ENABLE_DRAW_DEBUG && ENABLE_ANIM_DEBUG
-void FPoseHistory::DebugDraw(FAnimInstanceProxy& AnimInstanceProxy, FColor Color) const
+void FPoseHistory::DebugDraw(FAnimInstanceProxy& AnimInstanceProxy, FColor Color, const FPoseSearchQueryTrajectory* Trajectory) const
 {
+	const bool bValidTrajectory = Trajectory && !Trajectory->Samples.IsEmpty();
+
 	TArray<FTransform> PrevGlobalTransforms;
 	for (int32 EntryIndex = 0; EntryIndex < Entries.Num(); ++EntryIndex)
 	{
@@ -326,14 +317,16 @@ void FPoseHistory::DebugDraw(FAnimInstanceProxy& AnimInstanceProxy, FColor Color
 			PrevGlobalTransforms.SetNum(Entry.ComponentSpaceTransforms.Num());
 			for (int32 i = 0; i < Entry.ComponentSpaceTransforms.Num(); ++i)
 			{
-				PrevGlobalTransforms[i] = Entry.ComponentSpaceTransforms[i] * Entry.RootTransform;
+				const FTransform RootTransform = bValidTrajectory ? Trajectory->GetSampleAtTime(-Entry.Time).GetTransform() : AnimInstanceProxy.GetComponentTransform();
+				PrevGlobalTransforms[i] = Entry.ComponentSpaceTransforms[i] * RootTransform;
 			}
 		}
 		else
 		{
 			for (int32 i = 0; i < Entry.ComponentSpaceTransforms.Num(); ++i)
 			{
-				const FTransform GlobalTransforms = Entry.ComponentSpaceTransforms[i] * Entry.RootTransform;
+				const FTransform RootTransform = bValidTrajectory ? Trajectory->GetSampleAtTime(-Entry.Time).GetTransform() : AnimInstanceProxy.GetComponentTransform();
+				const FTransform GlobalTransforms = Entry.ComponentSpaceTransforms[i] * RootTransform;
 
 				AnimInstanceProxy.AnimDrawDebugLine(PrevGlobalTransforms[i].GetTranslation(), GlobalTransforms.GetTranslation(), Color, false, 0.f, ESceneDepthPriorityGroup::SDPG_Foreground);
 
@@ -357,12 +350,6 @@ bool FExtendedPoseHistory::IsInitialized() const
 	return PoseHistory != nullptr;
 }
 
-float FExtendedPoseHistory::GetSampleTimeInterval() const
-{
-	check(PoseHistory);
-	return PoseHistory->GetSampleTimeInterval();
-}
-
 bool FExtendedPoseHistory::GetTransformAtTime(float Time, FTransform& OutBoneTransform, const USkeleton* BoneIndexSkeleton, FBoneIndexType BoneIndexType, FBoneIndexType ReferenceBoneIndexType, bool bExtrapolate) const
 {
 	check(PoseHistory);
@@ -377,15 +364,8 @@ bool FExtendedPoseHistory::GetTransformAtTime(float Time, FTransform& OutBoneTra
 			const FPoseHistoryEntries& PastEntries = PoseHistory->GetEntries();
 			const FPoseHistoryEntry& NextEntry = FutureEntries[NextIdx];
 			const FPoseHistoryEntry& PrevEntry = NextIdx > 0 ? FutureEntries[NextIdx - 1] : !PastEntries.IsEmpty() ? PastEntries.First() : NextEntry;
-
-			const float Denominator = NextEntry.Time - PrevEntry.Time;
-			float LerpValue = 0.f;
-			if (!FMath::IsNearlyZero(Denominator))
-			{
-				const float Numerator = SecondsAgo - PrevEntry.Time;
-				LerpValue = bExtrapolate ? Numerator / Denominator : FMath::Clamp(Numerator / Denominator, 0.f, 1.f);
-			}
-			return FPoseHistory::LerpEntries(PrevEntry, NextEntry, LerpValue, BoneIndexSkeleton, PoseHistory->GetLastUpdateSkeleton(), PoseHistory->GetBoneToTransformMap(), BoneIndexType, ReferenceBoneIndexType, OutBoneTransform);
+						
+			return FPoseHistory::LerpEntries(Time, bExtrapolate, PrevEntry, NextEntry, BoneIndexSkeleton, PoseHistory->GetLastUpdateSkeleton(), PoseHistory->GetBoneToTransformMap(), BoneIndexType, ReferenceBoneIndexType, OutBoneTransform);
 		}
 	}
 	
@@ -410,14 +390,16 @@ void FExtendedPoseHistory::AddFuturePose(float SecondsInTheFuture, FCSPose<FComp
 	check(PoseHistory);	
 	const float SecondsAgo = -SecondsInTheFuture;
 	const int32 LowerBoundIdx = Algo::LowerBound(FutureEntries, SecondsAgo, [](const FPoseHistoryEntry& Entry, float Value) { return Value < Entry.Time; });
-	FutureEntries.InsertDefaulted_GetRef(LowerBoundIdx).Update(SecondsAgo, ComponentSpacePose, ComponentTransform, PoseHistory->GetBoneToTransformMap());
+	FutureEntries.InsertDefaulted_GetRef(LowerBoundIdx).Update(SecondsAgo, ComponentSpacePose, PoseHistory->GetBoneToTransformMap());
 }
 
 #if ENABLE_DRAW_DEBUG && ENABLE_ANIM_DEBUG
-void FExtendedPoseHistory::DebugDraw(FAnimInstanceProxy& AnimInstanceProxy, FColor PastColor, FColor FutureColor) const
+void FExtendedPoseHistory::DebugDraw(FAnimInstanceProxy& AnimInstanceProxy, FColor Color, const FPoseSearchQueryTrajectory* Trajectory) const
 {
 	check(PoseHistory);
 
+	const bool bValidTrajectory = Trajectory && !Trajectory->Samples.IsEmpty();
+
 	TArray<FTransform> PrevGlobalTransforms;
 	for (int32 EntryIndex = 0; EntryIndex < FutureEntries.Num(); ++EntryIndex)
 	{
@@ -431,23 +413,25 @@ void FExtendedPoseHistory::DebugDraw(FAnimInstanceProxy& AnimInstanceProxy, FCol
 			PrevGlobalTransforms.SetNum(Entry.ComponentSpaceTransforms.Num());
 			for (int32 i = 0; i < Entry.ComponentSpaceTransforms.Num(); ++i)
 			{
-				PrevGlobalTransforms[i] = Entry.ComponentSpaceTransforms[i] * Entry.RootTransform;
+				const FTransform RootTransform = bValidTrajectory ? Trajectory->GetSampleAtTime(-Entry.Time).GetTransform() : AnimInstanceProxy.GetComponentTransform();
+				PrevGlobalTransforms[i] = Entry.ComponentSpaceTransforms[i] * RootTransform;
 			}
 		}
 		else
 		{
 			for (int32 i = 0; i < Entry.ComponentSpaceTransforms.Num(); ++i)
 			{
-				const FTransform GlobalTransforms = Entry.ComponentSpaceTransforms[i] * Entry.RootTransform;
+				const FTransform RootTransform = bValidTrajectory ? Trajectory->GetSampleAtTime(-Entry.Time).GetTransform() : AnimInstanceProxy.GetComponentTransform();
+				const FTransform GlobalTransforms = Entry.ComponentSpaceTransforms[i] * RootTransform;
 
-				AnimInstanceProxy.AnimDrawDebugLine(PrevGlobalTransforms[i].GetTranslation(), GlobalTransforms.GetTranslation(), FutureColor, false, 0.f, ESceneDepthPriorityGroup::SDPG_Foreground);
+				AnimInstanceProxy.AnimDrawDebugLine(PrevGlobalTransforms[i].GetTranslation(), GlobalTransforms.GetTranslation(), Color, false, 0.f, ESceneDepthPriorityGroup::SDPG_Foreground);
 
 				PrevGlobalTransforms[i] = GlobalTransforms;
 			}
 		}
 	}
 
-	PoseHistory->DebugDraw(AnimInstanceProxy, PastColor);
+	PoseHistory->DebugDraw(AnimInstanceProxy, Color, Trajectory);
 }
 #endif // ENABLE_DRAW_DEBUG && ENABLE_ANIM_DEBUG
 
diff --git a/Engine/Plugins/Experimental/Animation/PoseSearch/Source/Runtime/Private/PoseSearchLibrary.cpp b/Engine/Plugins/Experimental/Animation/PoseSearch/Source/Runtime/Private/PoseSearchLibrary.cpp
--- a/Engine/Plugins/Experimental/Animation/PoseSearch/Source/Runtime/Private/PoseSearchLibrary.cpp
+++ b/Engine/Plugins/Experimental/Animation/PoseSearch/Source/Runtime/Private/PoseSearchLibrary.cpp
@@ -251,12 +251,15 @@ void UPoseSearchLibrary::UpdateMotionMatchingState(
 	bool bForceInterrupt,
 	bool bShouldSearch,
 	bool bDebugDrawQuery,
-	bool bDebugDrawCurResult)
+	bool bDebugDrawCurResult,
+	bool bDebugDrawPoseHistory)
 {
 	QUICK_SCOPE_CYCLE_COUNTER(STAT_PoseSearch_Update);
 
 	using namespace UE::PoseSearch;
 
+	check(Context.AnimInstanceProxy);
+
 	if (Databases.IsEmpty())
 	{
 		Context.LogMessage(
@@ -275,8 +278,15 @@ void UPoseSearchLibrary::UpdateMotionMatchingState(
 	if (IPoseHistoryProvider* PoseHistoryProvider = Context.GetMessage<IPoseHistoryProvider>())
 	{
 		History = &PoseHistoryProvider->GetPoseHistory();
+
+#if ENABLE_DRAW_DEBUG && ENABLE_ANIM_DEBUG
+		if (bDebugDrawPoseHistory)
+		{
+			History->DebugDraw(*Context.AnimInstanceProxy, FColor::Orange, &Trajectory);
+		}
+#endif // ENABLE_DRAW_DEBUG && ENABLE_ANIM_DEBUG
 	}
-	
+
 	const FPoseSearchQueryTrajectory TrajectoryRootSpace = ProcessTrajectory(Trajectory, InOutMotionMatchingState.ComponentDeltaYaw, YawFromAnimationTrajectoryBlendTime, TrajectorySpeedMultiplier);
 
 	FMemMark Mark(FMemStack::Get());
@@ -442,8 +452,6 @@ void UPoseSearchLibrary::MotionMatch(
 	using namespace UE::Anim;
 	using namespace UE::PoseSearch;
 
-	static constexpr float FiniteDelta = 1 / 60.0f;
-
 	Result.SelectedAnimation = nullptr;
 	Result.SelectedTime = 0.f;
 	Result.bLoop = false;
@@ -532,7 +540,7 @@ void UPoseSearchLibrary::MotionMatch(
 			{
 				if (FAnimInstanceProxy* AnimInstanceProxy = UAnimInstanceProxyProvider::GetAnimInstanceProxy(AnimInstance))
 				{
-					ExtendedPoseHistory.DebugDraw(*AnimInstanceProxy, HistoryCollectorColor, FColor::Green);
+					ExtendedPoseHistory.DebugDraw(*AnimInstanceProxy, HistoryCollectorColor, &TrajectoryRootSpace);
 				}
 			}
 #endif // ENABLE_DRAW_DEBUG && ENABLE_ANIM_DEBUG
diff --git a/Engine/Plugins/Experimental/Animation/PoseSearch/Source/Runtime/Private/PoseSearchSchema.cpp b/Engine/Plugins/Experimental/Animation/PoseSearch/Source/Runtime/Private/PoseSearchSchema.cpp
--- a/Engine/Plugins/Experimental/Animation/PoseSearch/Source/Runtime/Private/PoseSearchSchema.cpp
+++ b/Engine/Plugins/Experimental/Animation/PoseSearch/Source/Runtime/Private/PoseSearchSchema.cpp
@@ -54,6 +54,8 @@ void UPoseSearchSchema::AddDefaultChannels()
 
 int8 UPoseSearchSchema::AddBoneReference(const FBoneReference& BoneReference)
 {
+	using namespace UE::PoseSearch;
+
 	int32 SchemaBoneIdx = 0;
 	check(Skeleton);
 	
@@ -87,6 +89,8 @@ int8 UPoseSearchSchema::AddBoneReference(const FBoneReference& BoneReference)
 
 void UPoseSearchSchema::Finalize()
 {
+	using namespace UE::PoseSearch;
+
 	BoneReferences.Reset();
 	BoneIndicesWithParents.Reset();
 	FinalizedChannels.Reset();
diff --git a/Engine/Plugins/Experimental/Animation/PoseSearch/Source/Runtime/Public/PoseSearch/AnimNode_PoseSearchHistoryCollector.h b/Engine/Plugins/Experimental/Animation/PoseSearch/Source/Runtime/Public/PoseSearch/AnimNode_PoseSearchHistoryCollector.h
--- a/Engine/Plugins/Experimental/Animation/PoseSearch/Source/Runtime/Public/PoseSearch/AnimNode_PoseSearchHistoryCollector.h
+++ b/Engine/Plugins/Experimental/Animation/PoseSearch/Source/Runtime/Public/PoseSearch/AnimNode_PoseSearchHistoryCollector.h
@@ -4,6 +4,7 @@
 
 #include "Animation/AnimNodeBase.h"
 #include "PoseSearch/PoseSearchHistory.h"
+#include "PoseSearch/PoseSearchTrajectoryTypes.h"
 #include "AnimNode_PoseSearchHistoryCollector.generated.h"
 
 USTRUCT(BlueprintInternalUseOnly)
@@ -32,12 +33,16 @@ public:
 	UPROPERTY(EditAnywhere, Category = Settings)
 	bool bResetOnBecomingRelevant = true;
 
-	UPROPERTY(EditAnywhere, Category = Settings)
-	FLinearColor DebugColor = FLinearColor::Red;
-
 	// Update Counter for detecting being relevant
 	FGraphTraversalCounter UpdateCounter;
 
+	UPROPERTY(EditAnywhere, Category = Debug)
+	FLinearColor DebugColor = FLinearColor::Red;
+
+	// Optional Trajectory for debug drawing purposes only: if not provided, all the collected pose transforms will be drawn using the Graph context
+	UPROPERTY(EditAnywhere, Category = Debug, meta = (PinHiddenByDefault))
+	FPoseSearchQueryTrajectory DebugDrawTrajectory;
+
 	// FAnimNode_Base interface
 	virtual void CacheBones_AnyThread(const FAnimationCacheBonesContext& Context) override;
 	virtual void Update_AnyThread(const FAnimationUpdateContext& Context) override;
diff --git a/Engine/Plugins/Experimental/Animation/PoseSearch/Source/Runtime/Public/PoseSearch/PoseSearchAssetIndexer.h b/Engine/Plugins/Experimental/Animation/PoseSearch/Source/Runtime/Public/PoseSearch/PoseSearchAssetIndexer.h
--- a/Engine/Plugins/Experimental/Animation/PoseSearch/Source/Runtime/Public/PoseSearch/PoseSearchAssetIndexer.h
+++ b/Engine/Plugins/Experimental/Animation/PoseSearch/Source/Runtime/Public/PoseSearch/PoseSearchAssetIndexer.h
@@ -21,9 +21,6 @@ struct FPoseMetadata;
 
 struct FAssetSamplingContext
 {
-	// Time delta used for computing pose derivatives
-	static constexpr float FiniteDelta = 1 / 60.0f;
-
 	float BaseCostBias = 0.f;
 	float LoopingCostBias = 0.f;
 
diff --git a/Engine/Plugins/Experimental/Animation/PoseSearch/Source/Runtime/Public/PoseSearch/PoseSearchContext.h b/Engine/Plugins/Experimental/Animation/PoseSearch/Source/Runtime/Public/PoseSearch/PoseSearchContext.h
--- a/Engine/Plugins/Experimental/Animation/PoseSearch/Source/Runtime/Public/PoseSearch/PoseSearchContext.h
+++ b/Engine/Plugins/Experimental/Animation/PoseSearch/Source/Runtime/Public/PoseSearch/PoseSearchContext.h
@@ -150,21 +150,18 @@ struct POSESEARCH_API FSearchContext
 	// Returns the rotation of the bone Schema.BoneReferences[SchemaSampleBoneIdx] at an offset time of SampleTimeOffset relative to the
 	// transform of the bone Schema.BoneReferences[SchemaOriginBoneIdx] at an offset time of time OriginTimeOffset 
 	// Times will be processed by GetPermutationTimeOffsets(PermutationTimeType, ...)
-	// if bUseHistoryRoot is true, eventually required root transforms will be gathered from the pose history node data, otherwise the context trajectory will be used
-	FQuat GetSampleRotation(float SampleTimeOffset, float OriginTimeOffset, const UPoseSearchSchema* Schema, int8 SchemaSampleBoneIdx = RootSchemaBoneIdx, int8 SchemaOriginBoneIdx = RootSchemaBoneIdx, bool bUseHistoryRoot = false, EPermutationTimeType PermutationTimeType = EPermutationTimeType::UseSampleTime, const FQuat* SampleBoneRotationWorldOverride = nullptr);
+	FQuat GetSampleRotation(float SampleTimeOffset, float OriginTimeOffset, const UPoseSearchSchema* Schema, int8 SchemaSampleBoneIdx = RootSchemaBoneIdx, int8 SchemaOriginBoneIdx = RootSchemaBoneIdx, EPermutationTimeType PermutationTimeType = EPermutationTimeType::UseSampleTime, const FQuat* SampleBoneRotationWorldOverride = nullptr);
 	
 	// Returns the position of the bone Schema.BoneReferences[SchemaSampleBoneIdx] at an offset time of SampleTimeOffset relative to the
 	// transform of the bone Schema.BoneReferences[SchemaOriginBoneIdx] at an offset time of time OriginTimeOffset 
 	// Times will be processed by GetPermutationTimeOffsets(PermutationTimeType, ...)
-	// if bUseHistoryRoot is true, eventually required root transforms will be gathered from the pose history node data, otherwise the context trajectory will be used
-	FVector GetSamplePosition(float SampleTimeOffset, float OriginTimeOffset, const UPoseSearchSchema* Schema, int8 SchemaSampleBoneIdx = RootSchemaBoneIdx, int8 SchemaOriginBoneIdx = RootSchemaBoneIdx, bool bUseHistoryRoot = false, EPermutationTimeType PermutationTimeType = EPermutationTimeType::UseSampleTime, const FVector* SampleBonePositionWorldOverride = nullptr);
+	FVector GetSamplePosition(float SampleTimeOffset, float OriginTimeOffset, const UPoseSearchSchema* Schema, int8 SchemaSampleBoneIdx = RootSchemaBoneIdx, int8 SchemaOriginBoneIdx = RootSchemaBoneIdx, EPermutationTimeType PermutationTimeType = EPermutationTimeType::UseSampleTime, const FVector* SampleBonePositionWorldOverride = nullptr);
 	
 	// Returns the delta velocity of the velocity of the bone Schema.BoneReferences[SchemaSampleBoneIdx] at an offset time of SampleTimeOffset minus
 	// the velocity of the bone Schema.BoneReferences[SchemaOriginBoneIdx] at an offset time of time OriginTimeOffset 
 	// Times will be processed by GetPermutationTimeOffsets(PermutationTimeType, ...)
 	// if bUseCharacterSpaceVelocities is true, velocities will be computed in root bone space, rather than world space
-	// if bUseHistoryRoot is true, eventually required root transforms will be gathered from the pose history node data, otherwise the context trajectory will be used
-	FVector GetSampleVelocity(float SampleTimeOffset, float OriginTimeOffset, const UPoseSearchSchema* Schema, int8 SchemaSampleBoneIdx = RootSchemaBoneIdx, int8 SchemaOriginBoneIdx = RootSchemaBoneIdx, bool bUseCharacterSpaceVelocities = true, bool bUseHistoryRoot = false, EPermutationTimeType PermutationTimeType = EPermutationTimeType::UseSampleTime, const FVector* SampleBoneVelocityWorldOverride = nullptr);
+	FVector GetSampleVelocity(float SampleTimeOffset, float OriginTimeOffset, const UPoseSearchSchema* Schema, int8 SchemaSampleBoneIdx = RootSchemaBoneIdx, int8 SchemaOriginBoneIdx = RootSchemaBoneIdx, bool bUseCharacterSpaceVelocities = true, EPermutationTimeType PermutationTimeType = EPermutationTimeType::UseSampleTime, const FVector* SampleBoneVelocityWorldOverride = nullptr);
 
 	void ClearCachedEntries();
 
@@ -186,7 +183,7 @@ struct POSESEARCH_API FSearchContext
 	bool IsTrajectoryValid() const { return Trajectory != nullptr; }
 	float GetDesiredPermutationTimeOffset() const { return DesiredPermutationTimeOffset; }
 	bool IsForceInterrupt() const { return bForceInterrupt; }
-	FTransform GetWorldRootBoneTransformAtTime(float SampleTime, bool bUseHistoryRoot = false, bool bExtrapolate = true) const;
+	FTransform GetWorldRootBoneTransformAtTime(float SampleTime, bool bExtrapolate = true) const;
 	const UAnimInstance* GetAnimInstance() const { return AnimInstance; }
 
 	void SetAnimationsToConsider(TConstArrayView<const UAnimationAsset*> InAnimationsToConsider) { AnimationsToConsider = InAnimationsToConsider; }
@@ -194,13 +191,13 @@ struct POSESEARCH_API FSearchContext
 	
 private:
 	// returns the world space transform of the bone SchemaBoneIdx at time SampleTime
-	FTransform GetWorldBoneTransformAtTime(float SampleTime, const UPoseSearchSchema* Schema, int8 SchemaBoneIdx = RootSchemaBoneIdx, bool bUseHistoryRoot = false);
+	FTransform GetWorldBoneTransformAtTime(float SampleTime, const UPoseSearchSchema* Schema, int8 SchemaBoneIdx = RootSchemaBoneIdx);
 	
 	// returns the local space transform relative to the root bone of the bone SchemaBoneIdx at time SampleTime
 	FTransform GetLocalBoneTransformAtTime(float SampleTime, const UPoseSearchSchema* Schema, int8 SchemaSampleBoneIdx = RootSchemaBoneIdx);
 	
-	FVector GetSamplePositionInternal(float SampleTime, float OriginTime, const UPoseSearchSchema* Schema, int8 SchemaSampleBoneIdx = RootSchemaBoneIdx, int8 SchemaOriginBoneIdx = RootSchemaBoneIdx, bool bUseHistoryRoot = false, const FVector* SampleBonePositionWorldOverride = nullptr);
-	FQuat GetSampleRotationInternal(float SampleTime, float OriginTime, const UPoseSearchSchema* Schema, int8 SchemaSampleBoneIdx = RootSchemaBoneIdx, int8 SchemaOriginBoneIdx = RootSchemaBoneIdx, bool bUseHistoryRoot = false, const FQuat* SampleBoneRotationWorldOverride = nullptr);
+	FVector GetSamplePositionInternal(float SampleTime, float OriginTime, const UPoseSearchSchema* Schema, int8 SchemaSampleBoneIdx = RootSchemaBoneIdx, int8 SchemaOriginBoneIdx = RootSchemaBoneIdx, const FVector* SampleBonePositionWorldOverride = nullptr);
+	FQuat GetSampleRotationInternal(float SampleTime, float OriginTime, const UPoseSearchSchema* Schema, int8 SchemaSampleBoneIdx = RootSchemaBoneIdx, int8 SchemaOriginBoneIdx = RootSchemaBoneIdx, const FQuat* SampleBoneRotationWorldOverride = nullptr);
 
 	const UAnimInstance* AnimInstance = nullptr;
 	const IPoseHistory* History = nullptr;
diff --git a/Engine/Plugins/Experimental/Animation/PoseSearch/Source/Runtime/Public/PoseSearch/PoseSearchDefines.h b/Engine/Plugins/Experimental/Animation/PoseSearch/Source/Runtime/Public/PoseSearch/PoseSearchDefines.h
--- a/Engine/Plugins/Experimental/Animation/PoseSearch/Source/Runtime/Public/PoseSearch/PoseSearchDefines.h
+++ b/Engine/Plugins/Experimental/Animation/PoseSearch/Source/Runtime/Public/PoseSearch/PoseSearchDefines.h
@@ -21,15 +21,16 @@ constexpr EParallelForFlags ParallelForFlags = EParallelForFlags::None;
 
 DECLARE_LOG_CATEGORY_EXTERN(LogPoseSearch, Log, All);
 
+namespace UE::PoseSearch
+{
 static constexpr int8 RootSchemaBoneIdx = 0;
 static constexpr FBoneIndexType RootBoneIndexType = 0;
 static constexpr FBoneIndexType ComponentSpaceIndexType = FBoneIndexType(-1);
-static constexpr FBoneIndexType WorldSpaceIndexType = FBoneIndexType(-2);
 
 static constexpr int32 PreallocatedCachedQueriesNum = 8;
 
-namespace UE::PoseSearch
-{
+static constexpr float FiniteDelta = 1 / 60.0f; // Time delta used for computing pose derivatives
+
 constexpr int32 TMax(int32 A, int32 B) { return (A > B ? A : B); }
 template<typename ElementType> constexpr int32 TAlignOf() { return TMax(alignof(ElementType), 16); }
 template<typename ElementType> using TAlignedArray = TArray<ElementType, TAlignedHeapAllocator<TAlignOf<ElementType>()>>;
diff --git a/Engine/Plugins/Experimental/Animation/PoseSearch/Source/Runtime/Public/PoseSearch/PoseSearchHistory.h b/Engine/Plugins/Experimental/Animation/PoseSearch/Source/Runtime/Public/PoseSearch/PoseSearchHistory.h
--- a/Engine/Plugins/Experimental/Animation/PoseSearch/Source/Runtime/Public/PoseSearch/PoseSearchHistory.h
+++ b/Engine/Plugins/Experimental/Animation/PoseSearch/Source/Runtime/Public/PoseSearch/PoseSearchHistory.h
@@ -10,6 +10,7 @@
 #include "UObject/ObjectKey.h"
 
 struct FAnimInstanceProxy;
+struct FPoseSearchQueryTrajectory;
 class USkeleton;
 class UWorld;
 
@@ -24,27 +25,25 @@ typedef TMap<FBoneIndexType, FComponentSpaceTransformIndex> FBoneToTransformMap;
 struct POSESEARCH_API IPoseHistory
 {
 	virtual ~IPoseHistory() {}
-	virtual float GetSampleTimeInterval() const = 0;
 	
 	// returns the BoneIndexType transform relative to ReferenceBoneIndexType: 
 	// if ReferenceBoneIndexType is 0 (RootBoneIndexType), OutBoneTransform is in root bone space
 	// if ReferenceBoneIndexType is FBoneIndexType(-1) (ComponentSpaceIndexType), OutBoneTransform is in component space
-	// if ReferenceBoneIndexType is FBoneIndexType(-2) (WorldSpaceIndexType), OutBoneTransform is in world space
 	virtual bool GetTransformAtTime(float Time, FTransform& OutBoneTransform, const USkeleton* BoneIndexSkeleton = nullptr, FBoneIndexType BoneIndexType = RootBoneIndexType, FBoneIndexType ReferenceBoneIndexType = ComponentSpaceIndexType, bool bExtrapolate = true) const = 0;
-
 	virtual bool IsEmpty() const = 0;
+#if ENABLE_DRAW_DEBUG && ENABLE_ANIM_DEBUG
+	// debug draw the pose history. If Trajectory is not provided all the poses will be drawn at the AnimInstanceProxy transform
+	virtual void DebugDraw(FAnimInstanceProxy& AnimInstanceProxy, FColor Color, const FPoseSearchQueryTrajectory* Trajectory = nullptr) const = 0;
+#endif
 };
 
 struct FPoseHistoryEntry
 {
-	// collected root component transform in world space (not the root bone transform)
-	FTransform RootTransform;
-
 	// collected bones transforms in component space
 	TArray<FTransform> ComponentSpaceTransforms;
 	float Time = 0.f;
 
-	void Update(float InTime, FCSPose<FCompactPose>& ComponentSpacePose, const FTransform& ComponentTransform, const FBoneToTransformMap& BoneToTransformMap);
+	void Update(float InTime, FCSPose<FCompactPose>& ComponentSpacePose, const FBoneToTransformMap& BoneToTransformMap);
 };
 
 typedef TRingBuffer<FPoseHistoryEntry> FPoseHistoryEntries;
@@ -53,7 +52,7 @@ typedef TArray<FPoseHistoryEntry> FPoseHistoryFutureEntries;
 struct FPoseHistory : public IPoseHistory
 {
 	void Init(int32 InNumPoses, float InTimeHorizon, const TArray<FBoneIndexType>& RequiredBones);
-	void Update(float SecondsElapsed, FCSPose<FCompactPose>& ComponentSpacePose, const FTransform& ComponentTransform);
+	void Update(float SecondsElapsed, FCSPose<FCompactPose>& ComponentSpacePose);
 	float GetTimeHorizon() const { return TimeHorizon; }
 
 	const FBoneToTransformMap& GetBoneToTransformMap() const { return BoneToTransformMap; }
@@ -61,20 +60,18 @@ struct FPoseHistory : public IPoseHistory
 	const USkeleton* GetLastUpdateSkeleton() const { return LastUpdateSkeleton.Get(); }
 
 	// IPoseHistory interface
-	virtual float GetSampleTimeInterval() const override;
-	virtual bool GetTransformAtTime(float Time, FTransform& OutBoneTransform, const USkeleton* BoneIndexSkeleton = nullptr, FBoneIndexType BoneIndexType = RootBoneIndexType, FBoneIndexType ReferenceBoneIndexType = ComponentSpaceIndexType, bool bExtrapolate = true) const;
+	virtual bool GetTransformAtTime(float Time, FTransform& OutBoneTransform, const USkeleton* BoneIndexSkeleton = nullptr, FBoneIndexType BoneIndexType = RootBoneIndexType, FBoneIndexType ReferenceBoneIndexType = ComponentSpaceIndexType, bool bExtrapolate = true) const override;
 	virtual bool IsEmpty() const override;
+#if ENABLE_DRAW_DEBUG && ENABLE_ANIM_DEBUG
+	virtual void DebugDraw(FAnimInstanceProxy& AnimInstanceProxy, FColor Color, const FPoseSearchQueryTrajectory* Trajectory = nullptr) const override;
+#endif
 	// End of IPoseHistory interface
 
 	void ClearHistory();
 
 	static FBoneIndexType GetRemappedBoneIndexType(FBoneIndexType BoneIndexType, const USkeleton* BoneIndexSkeleton, const USkeleton* LastUpdateSkeleton);
 	static FComponentSpaceTransformIndex GetRemappedComponentSpaceTransformIndex(const USkeleton* BoneIndexSkeleton, const USkeleton* LastUpdateSkeleton, const FBoneToTransformMap& BoneToTransformMap, FBoneIndexType BoneIndexType, bool& bSuccess);
-	static bool LerpEntries(const FPoseHistoryEntry& PrevEntry, const FPoseHistoryEntry& NextEntry, float LerpValue, const USkeleton* BoneIndexSkeleton, const USkeleton* LastUpdateSkeleton, const FBoneToTransformMap& BoneToTransformMap, FBoneIndexType BoneIndexType, FBoneIndexType ReferenceBoneIndexType, FTransform& OutBoneTransform);
-
-#if ENABLE_DRAW_DEBUG && ENABLE_ANIM_DEBUG
-	void DebugDraw(FAnimInstanceProxy& AnimInstanceProxy, FColor Color) const;
-#endif
+	static bool LerpEntries(float Time, bool bExtrapolate, const FPoseHistoryEntry& PrevEntry, const FPoseHistoryEntry& NextEntry, const USkeleton* BoneIndexSkeleton, const USkeleton* LastUpdateSkeleton, const FBoneToTransformMap& BoneToTransformMap, FBoneIndexType BoneIndexType, FBoneIndexType ReferenceBoneIndexType, FTransform& OutBoneTransform);
 
 private:
 	// skeleton from the last Update, to keep tracking skeleton changes, and support compatible skeletons
@@ -95,18 +92,16 @@ struct FExtendedPoseHistory : public IPoseHistory
 	bool IsInitialized() const;
 
 	// IPoseHistory interface
-	virtual float GetSampleTimeInterval() const override;
-	virtual bool GetTransformAtTime(float Time, FTransform& OutBoneTransform, const USkeleton* BoneIndexSkeleton = nullptr, FBoneIndexType BoneIndexType = RootBoneIndexType, FBoneIndexType ReferenceBoneIndexType = ComponentSpaceIndexType, bool bExtrapolate = true) const;
+	virtual bool GetTransformAtTime(float Time, FTransform& OutBoneTransform, const USkeleton* BoneIndexSkeleton = nullptr, FBoneIndexType BoneIndexType = RootBoneIndexType, FBoneIndexType ReferenceBoneIndexType = ComponentSpaceIndexType, bool bExtrapolate = true) const override;
 	virtual bool IsEmpty() const override;
+#if ENABLE_DRAW_DEBUG && ENABLE_ANIM_DEBUG
+	virtual void DebugDraw(FAnimInstanceProxy& AnimInstanceProxy, FColor Color, const FPoseSearchQueryTrajectory* Trajectory = nullptr) const override;	
+#endif
 	// End of IPoseHistory interface
 
 	void ResetFuturePoses();
 	void AddFuturePose(float SecondsInTheFuture, FCSPose<FCompactPose>& ComponentSpacePose, const FTransform& ComponentTransform);
 
-#if ENABLE_DRAW_DEBUG && ENABLE_ANIM_DEBUG
-	void DebugDraw(FAnimInstanceProxy& AnimInstanceProxy, FColor PastColor, FColor FutureColor) const;
-#endif
-
 private:
 	const FPoseHistory* PoseHistory = nullptr;
 	FPoseHistoryFutureEntries FutureEntries;
diff --git a/Engine/Plugins/Experimental/Animation/PoseSearch/Source/Runtime/Public/PoseSearch/PoseSearchLibrary.h b/Engine/Plugins/Experimental/Animation/PoseSearch/Source/Runtime/Public/PoseSearch/PoseSearchLibrary.h
--- a/Engine/Plugins/Experimental/Animation/PoseSearch/Source/Runtime/Public/PoseSearch/PoseSearchLibrary.h
+++ b/Engine/Plugins/Experimental/Animation/PoseSearch/Source/Runtime/Public/PoseSearch/PoseSearchLibrary.h
@@ -129,7 +129,8 @@ public:
 		bool bForceInterrupt = false,
 		bool bShouldSearch = true,
 		bool bDebugDrawQuery = false,
-		bool bDebugDrawCurResult = false);
+		bool bDebugDrawCurResult = false,
+		bool bDebugDrawPoseHistory = false);
 
 	/**
 	* Implementation of the core motion matching algorithm
