motion matching - adding UPoseSearchSchema picker for UPoseSearchDatabaseFactory

diff --git a/Engine/Plugins/Experimental/Animation/PoseSearch/Source/Editor/Private/PoseSearchDatabaseFactory.cpp b/Engine/Plugins/Experimental/Animation/PoseSearch/Source/Editor/Private/PoseSearchDatabaseFactory.cpp
--- a/Engine/Plugins/Experimental/Animation/PoseSearch/Source/Editor/Private/PoseSearchDatabaseFactory.cpp
+++ b/Engine/Plugins/Experimental/Animation/PoseSearch/Source/Editor/Private/PoseSearchDatabaseFactory.cpp
@@ -1,7 +1,12 @@
 // Copyright Epic Games, Inc. All Rights Reserved.
 
 #include "PoseSearchDatabaseFactory.h"
+#include "ContentBrowserModule.h"
+#include "Editor.h"
+#include "IContentBrowserSingleton.h"
 #include "PoseSearch/PoseSearchDatabase.h"
+#include "PoseSearch/PoseSearchSchema.h"
+#include "Widgets/Layout/SBorder.h"
 
 #define LOCTEXT_NAMESPACE "PoseSearchEditor"
 
@@ -13,9 +18,57 @@ UPoseSearchDatabaseFactory::UPoseSearchDatabaseFactory(const FObjectInitializer&
 	SupportedClass = UPoseSearchDatabase::StaticClass();
 }
 
+void UPoseSearchDatabaseFactory::OnTargetSchemaSelected(const FAssetData& SelectedAsset)
+{
+	TargetSchema = Cast<UPoseSearchSchema>(SelectedAsset.GetAsset());
+	PickerWindow->RequestDestroyWindow();
+}
+
+bool UPoseSearchDatabaseFactory::ConfigureProperties()
+{
+	// Load the content browser module to display an asset picker
+	FContentBrowserModule& ContentBrowserModule = FModuleManager::LoadModuleChecked<FContentBrowserModule>("ContentBrowser");
+
+	FAssetPickerConfig AssetPickerConfig;
+	
+	/** The asset picker will only show UPoseSearchSchema(s) */
+	AssetPickerConfig.Filter.ClassPaths.Add(UPoseSearchSchema::StaticClass()->GetClassPathName());
+	AssetPickerConfig.Filter.bRecursiveClasses = true;
+
+	/** The delegate that fires when an asset was selected */
+	AssetPickerConfig.OnAssetSelected = FOnAssetSelected::CreateUObject(this, &UPoseSearchDatabaseFactory::OnTargetSchemaSelected);
+
+	/** The default view mode should be a list view */
+	AssetPickerConfig.InitialAssetViewType = EAssetViewType::List;
+
+	PickerWindow = SNew(SWindow)
+	.Title(LOCTEXT("CreatePoseSearchDatabaseOptions", "Pick Config"))
+	.ClientSize(FVector2D(500, 600))
+	.SupportsMinimize(false)
+	.SupportsMaximize(false)
+	[
+		SNew(SBorder)
+		.BorderImage( FAppStyle::GetBrush("Menu.Background") )
+		[
+			ContentBrowserModule.Get().CreateAssetPicker(AssetPickerConfig)
+		]
+	];
+
+	GEditor->EditorAddModalWindow(PickerWindow.ToSharedRef());
+	PickerWindow.Reset();
+
+	return TargetSchema != nullptr;
+}
+
 UObject* UPoseSearchDatabaseFactory::FactoryCreateNew(UClass* Class, UObject* InParent, FName Name, EObjectFlags Flags, UObject* Context, FFeedbackContext* Warn)
 {
-	return NewObject<UPoseSearchDatabase>(InParent, Class, Name, Flags);
+	UPoseSearchDatabase* Database = nullptr;
+	if (TargetSchema)
+	{
+		Database = NewObject<UPoseSearchDatabase>(InParent, Class, Name, Flags);
+		Database->Schema = TargetSchema;
+	}
+	return Database;
 }
 
 FString UPoseSearchDatabaseFactory::GetDefaultNewAssetName() const
diff --git a/Engine/Plugins/Experimental/Animation/PoseSearch/Source/Editor/Public/PoseSearchDatabaseFactory.h b/Engine/Plugins/Experimental/Animation/PoseSearch/Source/Editor/Public/PoseSearchDatabaseFactory.h
--- a/Engine/Plugins/Experimental/Animation/PoseSearch/Source/Editor/Public/PoseSearchDatabaseFactory.h
+++ b/Engine/Plugins/Experimental/Animation/PoseSearch/Source/Editor/Public/PoseSearchDatabaseFactory.h
@@ -7,13 +7,24 @@
 #include "Factories/Factory.h"
 #include "PoseSearchDatabaseFactory.generated.h"
 
+struct FAssetData;
+class SWindow;
+
 UCLASS()
 class POSESEARCHEDITOR_API UPoseSearchDatabaseFactory : public UFactory
 {
 	GENERATED_UCLASS_BODY()
 
+	UPROPERTY()
+	TObjectPtr<class UPoseSearchSchema> TargetSchema;
+
+	TSharedPtr<SWindow> PickerWindow;
+
 	// UFactory interface
+	virtual bool ConfigureProperties() override;
 	virtual UObject* FactoryCreateNew(UClass* Class, UObject* InParent, FName Name, EObjectFlags Flags, UObject* Context, FFeedbackContext* Warn) override;
 	virtual FString GetDefaultNewAssetName() const override;
 	// End of UFactory interface
+
+	void OnTargetSchemaSelected(const FAssetData& SelectedAsset);
 };
