motion matching - database synchronization with external assets improvements

diff --git a/Engine/Plugins/Experimental/Animation/PoseSearch/Source/Editor/Private/PoseSearchDatabaseAssetListItem.cpp b/Engine/Plugins/Experimental/Animation/PoseSearch/Source/Editor/Private/PoseSearchDatabaseAssetListItem.cpp
--- a/Engine/Plugins/Experimental/Animation/PoseSearch/Source/Editor/Private/PoseSearchDatabaseAssetListItem.cpp
+++ b/Engine/Plugins/Experimental/Animation/PoseSearch/Source/Editor/Private/PoseSearchDatabaseAssetListItem.cpp
@@ -34,7 +34,7 @@
 
 namespace UE::PoseSearch
 {
-	static constexpr FLinearColor DisabledColor = FLinearColor(1.0f, 1.0f, 1.0f, 0.25f);
+	static constexpr FLinearColor DisabledColor = FLinearColor(1.f, 1.f, 1.f, 0.25f);
 	
 	void SDatabaseAssetListItem::Construct(
 		const FArguments& InArgs,
@@ -459,8 +459,6 @@ namespace UE::PoseSearch
 		const TSharedPtr<FDatabaseViewModel> ViewModelPtr = EditorViewModel.Pin();
 		const TSharedPtr<FDatabaseAssetTreeNode> TreeNodePtr = WeakAssetTreeNode.Pin();
 
-		ViewModelPtr->GetPoseSearchDatabase()->Modify();
-		
 		ViewModelPtr->SetIsEnabled(TreeNodePtr->SourceAssetIdx, NewCheckboxState == ECheckBoxState::Checked ? true : false);
 
 		SkeletonView.Pin()->RefreshTreeView(false, true);
@@ -469,7 +467,24 @@ namespace UE::PoseSearch
 
 	FSlateColor SDatabaseAssetListItem::GetNameTextColorAndOpacity() const
 	{
-		return GetAssetEnabledChecked() == ECheckBoxState::Checked ? FLinearColor::White : DisabledColor;
+		TSharedPtr<FDatabaseViewModel> ViewModelPtr = EditorViewModel.Pin();
+		TSharedPtr<FDatabaseAssetTreeNode> TreeNodePtr = WeakAssetTreeNode.Pin();
+		const UPoseSearchDatabase* Database = ViewModelPtr->GetPoseSearchDatabase();
+
+		if (const FPoseSearchDatabaseAnimationAssetBase* DatabaseAnimationAssetBase = Database->GetAnimationAssetBase(TreeNodePtr->SourceAssetIdx))
+		{
+			if (DatabaseAnimationAssetBase->IsEnabled())
+			{
+				if (DatabaseAnimationAssetBase->bSynchronizeWithExternalDependency)
+				{
+					return FColor::Turquoise;
+				}
+
+				return FLinearColor::White;
+			}
+		}
+
+		return DisabledColor;
 	}
 
 	FSlateColor SDatabaseAssetListItem::GetLoopingColorAndOpacity() const
diff --git a/Engine/Plugins/Experimental/Animation/PoseSearch/Source/Editor/Private/PoseSearchDatabaseAssetTree.cpp b/Engine/Plugins/Experimental/Animation/PoseSearch/Source/Editor/Private/PoseSearchDatabaseAssetTree.cpp
--- a/Engine/Plugins/Experimental/Animation/PoseSearch/Source/Editor/Private/PoseSearchDatabaseAssetTree.cpp
+++ b/Engine/Plugins/Experimental/Animation/PoseSearch/Source/Editor/Private/PoseSearchDatabaseAssetTree.cpp
@@ -616,10 +616,7 @@ namespace UE::PoseSearch
 		FScopedTransaction Transaction(LOCTEXT("DeleteAsset", "Delete Asset"));
 		const TSharedPtr<FDatabaseViewModel> ViewModel = EditorViewModel.Pin();
 		
-		ViewModel->GetPoseSearchDatabase()->Modify();
-		ViewModel->DeleteFromDatabase(Node->SourceAssetIdx);
-
-		if (bFinalizeChanges)
+		if (ViewModel->DeleteFromDatabase(Node->SourceAssetIdx) && bFinalizeChanges)
 		{
 			FinalizeTreeChanges();
 		}
@@ -685,8 +682,6 @@ namespace UE::PoseSearch
 			const FScopedTransaction Transaction(LOCTEXT("DeletePoseSearchDatabaseNodes", "Delete selected items from Pose Search Database"));
 			const TSharedPtr<FDatabaseViewModel> ViewModel = EditorViewModel.Pin();
 
-			ViewModel->GetPoseSearchDatabase()->Modify();
-			
 			SelectedNodes.Sort([](const TSharedPtr<FDatabaseAssetTreeNode>& A, const TSharedPtr<FDatabaseAssetTreeNode>& B)
 				{
 					return B->SourceAssetIdx < A->SourceAssetIdx;
diff --git a/Engine/Plugins/Experimental/Animation/PoseSearch/Source/Editor/Private/PoseSearchDatabaseViewModel.cpp b/Engine/Plugins/Experimental/Animation/PoseSearch/Source/Editor/Private/PoseSearchDatabaseViewModel.cpp
--- a/Engine/Plugins/Experimental/Animation/PoseSearch/Source/Editor/Private/PoseSearchDatabaseViewModel.cpp
+++ b/Engine/Plugins/Experimental/Animation/PoseSearch/Source/Editor/Private/PoseSearchDatabaseViewModel.cpp
@@ -14,6 +14,7 @@
 #include "GameFramework/PlayerController.h"
 #include "InstancedStruct.h"
 #include "Modules/ModuleManager.h"
+#include "PoseSearch/PoseSearchAnimNotifies.h"
 #include "PoseSearch/PoseSearchDatabase.h"
 #include "PoseSearch/PoseSearchDefines.h"
 #include "PoseSearch/PoseSearchDerivedData.h"
@@ -356,9 +357,20 @@ namespace UE::PoseSearch
 		PoseSearchDatabase->AnimationAssets.Add(FInstancedStruct::Make(NewAsset));
 	}
 
-	void FDatabaseViewModel::DeleteFromDatabase(int32 AnimationAssetIndex)
+	bool FDatabaseViewModel::DeleteFromDatabase(int32 AnimationAssetIndex)
 	{
-		PoseSearchDatabase->AnimationAssets.RemoveAt(AnimationAssetIndex);
+		if (const FPoseSearchDatabaseAnimationAssetBase* DatabaseAnimationAssetBase = PoseSearchDatabase->GetAnimationAssetBase(AnimationAssetIndex))
+		{
+			if (!DatabaseAnimationAssetBase->bSynchronizeWithExternalDependency)
+			{
+				PoseSearchDatabase->AnimationAssets.RemoveAt(AnimationAssetIndex);
+				PoseSearchDatabase->Modify();
+		
+				return true;
+			}
+		}
+
+		return false;
 	}
 
 	void FDatabaseViewModel::SetDisableReselection(int32 AnimationAssetIndex, bool bEnabled)
@@ -383,6 +395,8 @@ namespace UE::PoseSearch
 	{
 		if (FPoseSearchDatabaseAnimationAssetBase* DatabaseAnimationAsset = PoseSearchDatabase->GetMutableAnimationAssetBase(AnimationAssetIndex))
 		{
+			GetPoseSearchDatabase()->Modify();
+
 			DatabaseAnimationAsset->SetIsEnabled(bEnabled);
 		}
 	}
diff --git a/Engine/Plugins/Experimental/Animation/PoseSearch/Source/Editor/Public/PoseSearchDatabaseViewModel.h b/Engine/Plugins/Experimental/Animation/PoseSearch/Source/Editor/Public/PoseSearchDatabaseViewModel.h
--- a/Engine/Plugins/Experimental/Animation/PoseSearch/Source/Editor/Public/PoseSearchDatabaseViewModel.h
+++ b/Engine/Plugins/Experimental/Animation/PoseSearch/Source/Editor/Public/PoseSearchDatabaseViewModel.h
@@ -106,7 +106,7 @@ namespace UE::PoseSearch
 		void AddBlendSpaceToDatabase(UBlendSpace* BlendSpace);
 		void AddAnimCompositeToDatabase(UAnimComposite* AnimComposite);
 		void AddAnimMontageToDatabase(UAnimMontage* AnimMontage);
-		void DeleteFromDatabase(int32 AnimationAssetIndex);
+		bool DeleteFromDatabase(int32 AnimationAssetIndex);
 
 		void SetDisableReselection(int32 AnimationAssetIndex, bool bEnabled);
 		bool IsDisableReselection(int32 AnimationAssetIndex) const;
diff --git a/Engine/Plugins/Experimental/Animation/PoseSearch/Source/Runtime/Private/PoseSearchDatabase.cpp b/Engine/Plugins/Experimental/Animation/PoseSearch/Source/Runtime/Private/PoseSearchDatabase.cpp
--- a/Engine/Plugins/Experimental/Animation/PoseSearch/Source/Runtime/Private/PoseSearchDatabase.cpp
+++ b/Engine/Plugins/Experimental/Animation/PoseSearch/Source/Runtime/Private/PoseSearchDatabase.cpp
@@ -658,73 +658,101 @@ void UPoseSearchDatabase::SynchronizeWithExternalDependencies()
 		TArray<FAssetData> Assets;
 		AssetRegistry.GetAssetsByPackageName(Referencer.PackageName, Assets);
 
+		TArray<UAnimSequenceBase*> SequencesBase;
 		for (const FAssetData& Asset : Assets)
 		{
 			if (Asset.IsInstanceOf(UAnimSequenceBase::StaticClass()))
 			{
-				SynchronizeWithExternalDependencies(CastChecked<UAnimSequence>(Asset.FastGetAsset(true)));
+				if (UAnimSequenceBase* SequenceBase = CastChecked<UAnimSequenceBase>(Asset.FastGetAsset(true)))
+				{
+					for (const FAnimNotifyEvent& NotifyEvent : SequenceBase->Notifies)
+					{
+						if (const UAnimNotifyState_PoseSearchBranchIn* BranchIn = Cast<UAnimNotifyState_PoseSearchBranchIn>(NotifyEvent.NotifyStateClass))
+						{
+							if (BranchIn->Database == this)
+							{
+								SequencesBase.Add(SequenceBase);
+								break;
+							}
+						}
+					}
+				}
 			}
 		}
+		
+		if (!SequencesBase.IsEmpty())
+		{
+			SynchronizeWithExternalDependencies(SequencesBase);
+		}
 	}
 }
 
-void UPoseSearchDatabase::SynchronizeWithExternalDependencies(UAnimSequenceBase* SequenceBase)
+void UPoseSearchDatabase::SynchronizeWithExternalDependencies(TConstArrayView<UAnimSequenceBase*> SequencesBase)
 {
-	if (SequenceBase)
+	// @todo: improve the logic to identify if the asset is actually modified (this method could potentially remove and readd exatly the same AnimationAssets)
+	bool bModified = false;
+
+	// removing all the SequencesBase references from the database
+	for (int32 AnimationAssetIndex = AnimationAssets.Num() - 1; AnimationAssetIndex >= 0; --AnimationAssetIndex)
+	{
+		if (const FPoseSearchDatabaseAnimationAssetBase* AnimationAssetBase = AnimationAssets[AnimationAssetIndex].GetPtr<FPoseSearchDatabaseAnimationAssetBase>())
+		{
+			if (AnimationAssetBase->bSynchronizeWithExternalDependency && SequencesBase.Contains(AnimationAssetBase->GetAnimationAsset()))
+			{
+				AnimationAssets.RemoveAt(AnimationAssetIndex);
+				bModified = true;
+			}
+		}
+	}
+
+	// readding the required / updated references
+	for (UAnimSequenceBase* SequenceBase : SequencesBase)
 	{
-		bool bModified = false;
-		for (const FAnimNotifyEvent& NotifyEvent : SequenceBase->Notifies)
+		if (SequenceBase)
 		{
-			if (const UAnimNotifyState_PoseSearchBranchIn* PoseSearchBranchIn = Cast<UAnimNotifyState_PoseSearchBranchIn>(NotifyEvent.NotifyStateClass))
+			for (const FAnimNotifyEvent& NotifyEvent : SequenceBase->Notifies)
 			{
-				if (PoseSearchBranchIn->Database == this)
+				if (const UAnimNotifyState_PoseSearchBranchIn* PoseSearchBranchIn = Cast<UAnimNotifyState_PoseSearchBranchIn>(NotifyEvent.NotifyStateClass))
 				{
-					if (!bModified)
+					if (PoseSearchBranchIn->Database == this)
 					{
-						for (int32 AnimationAssetIndex = AnimationAssets.Num() - 1; AnimationAssetIndex >= 0; --AnimationAssetIndex)
+						if (UAnimSequence* Sequence = Cast<UAnimSequence>(SequenceBase))
 						{
-							if (const FPoseSearchDatabaseAnimationAssetBase* AnimationAssetBase = AnimationAssets[AnimationAssetIndex].GetPtr<FPoseSearchDatabaseAnimationAssetBase>())
-							{
-								if (AnimationAssetBase->GetAnimationAsset() == SequenceBase)
-								{
-									AnimationAssets.RemoveAt(AnimationAssetIndex);
-								}
-							}
+							FPoseSearchDatabaseSequence DatabaseSequence;
+							DatabaseSequence.Sequence = Sequence;
+							DatabaseSequence.SamplingRange = FFloatInterval(NotifyEvent.GetTriggerTime(), NotifyEvent.GetEndTriggerTime());
+							DatabaseSequence.bSynchronizeWithExternalDependency = true;
+							AnimationAssets.Add(FInstancedStruct::Make(DatabaseSequence));
+							bModified = true;
+						}
+						else if (UAnimComposite* AnimComposite = Cast<UAnimComposite>(SequenceBase))
+						{
+							FPoseSearchDatabaseAnimComposite DatabaseAnimComposite;
+							DatabaseAnimComposite.AnimComposite = AnimComposite;
+							DatabaseAnimComposite.SamplingRange = FFloatInterval(NotifyEvent.GetTriggerTime(), NotifyEvent.GetEndTriggerTime());
+							DatabaseAnimComposite.bSynchronizeWithExternalDependency = true;
+							AnimationAssets.Add(FInstancedStruct::Make(DatabaseAnimComposite));
+							bModified = true;
+						}
+						else if (UAnimMontage* AnimMontage = Cast<UAnimMontage>(SequenceBase))
+						{
+							FPoseSearchDatabaseAnimMontage DatabaseAnimMontage;
+							DatabaseAnimMontage.AnimMontage = AnimMontage;
+							DatabaseAnimMontage.SamplingRange = FFloatInterval(NotifyEvent.GetTriggerTime(), NotifyEvent.GetEndTriggerTime());
+							DatabaseAnimMontage.bSynchronizeWithExternalDependency = true;
+							AnimationAssets.Add(FInstancedStruct::Make(DatabaseAnimMontage));
+							bModified = true;
 						}
-
-						bModified = true;
-					}
-
-					if (UAnimSequence* Sequence = Cast<UAnimSequence>(SequenceBase))
-					{
-						FPoseSearchDatabaseSequence DatabaseSequence;
-						DatabaseSequence.Sequence = Sequence;
-						DatabaseSequence.SamplingRange = FFloatInterval(NotifyEvent.GetTriggerTime(), NotifyEvent.GetEndTriggerTime());
-						AnimationAssets.Add(FInstancedStruct::Make(DatabaseSequence));
-					}
-					else if (UAnimComposite* AnimComposite = Cast<UAnimComposite>(SequenceBase))
-					{
-						FPoseSearchDatabaseAnimComposite DatabaseAnimComposite;
-						DatabaseAnimComposite.AnimComposite = AnimComposite;
-						DatabaseAnimComposite.SamplingRange = FFloatInterval(NotifyEvent.GetTriggerTime(), NotifyEvent.GetEndTriggerTime());
-						AnimationAssets.Add(FInstancedStruct::Make(DatabaseAnimComposite));
-					}
-					else if (UAnimMontage* AnimMontage = Cast<UAnimMontage>(SequenceBase))
-					{
-						FPoseSearchDatabaseAnimMontage DatabaseAnimMontage;
-						DatabaseAnimMontage.AnimMontage = AnimMontage;
-						DatabaseAnimMontage.SamplingRange = FFloatInterval(NotifyEvent.GetTriggerTime(), NotifyEvent.GetEndTriggerTime());
-						AnimationAssets.Add(FInstancedStruct::Make(DatabaseAnimMontage));
 					}
 				}
 			}
 		}
+	}
 
-		if (bModified)
-		{
-			Modify();
-			NotifySynchronizeWithExternalDependencies();
-		}
+	if (bModified)
+	{
+		Modify();
+		NotifySynchronizeWithExternalDependencies();
 	}
 }
 
@@ -746,8 +774,11 @@ bool UPoseSearchDatabase::IsCachedCookedPlatformDataLoaded(const ITargetPlatform
 void UPoseSearchDatabase::PreSaveRoot(FObjectPreSaveRootContext ObjectSaveContext)
 {
 #if WITH_EDITOR
+	// in case the database desynchronized with the UAnimNotifyState_PoseSearchBranchIn referencing it, we need to resyncrhonize
 	SynchronizeWithExternalDependencies();
 #endif
+
+	Super::PreSaveRoot(ObjectSaveContext);
 }
 
 void UPoseSearchDatabase::PostSaveRoot(FObjectPostSaveRootContext ObjectSaveContext)
diff --git a/Engine/Plugins/Experimental/Animation/PoseSearch/Source/Runtime/Private/PoseSearchDerivedData.cpp b/Engine/Plugins/Experimental/Animation/PoseSearch/Source/Runtime/Private/PoseSearchDerivedData.cpp
--- a/Engine/Plugins/Experimental/Animation/PoseSearch/Source/Runtime/Private/PoseSearchDerivedData.cpp
+++ b/Engine/Plugins/Experimental/Animation/PoseSearch/Source/Runtime/Private/PoseSearchDerivedData.cpp
@@ -1755,6 +1755,62 @@ FAsyncPoseSearchDatabasesManagement::~FAsyncPoseSearchDatabasesManagement()
 	delete &Tasks;
 }
 
+// given Object it figures out a map of databases to UAnimSequenceBase(s) containing UAnimNotifyState_PoseSearchBranchIn
+void FAsyncPoseSearchDatabasesManagement::CollectDatabasesToSynchronize(UObject* Object)
+{
+	FScopeLock Lock(&Mutex);
+
+	if (UAnimSequenceBase* SequenceBase = Cast<UAnimSequenceBase>(Object))
+	{
+		for (const FAnimNotifyEvent& NotifyEvent : SequenceBase->Notifies)
+		{
+			if (const UAnimNotifyState_PoseSearchBranchIn* BranchIn = Cast<UAnimNotifyState_PoseSearchBranchIn>(NotifyEvent.NotifyStateClass))
+			{
+				if (BranchIn->Database)
+				{
+					DatabasesToSynchronize.FindOrAdd(BranchIn->Database).AddUnique(SequenceBase);
+				}
+			}
+		}
+	}
+	else if (UAnimNotifyState_PoseSearchBranchIn* BranchIn = Cast<UAnimNotifyState_PoseSearchBranchIn>(Object))
+	{
+		if (BranchIn->Database)
+		{
+			if (UAnimSequenceBase* OuterSequenceBase = Cast<UAnimSequenceBase>(BranchIn->GetOuter()))
+			{
+				DatabasesToSynchronize.FindOrAdd(BranchIn->Database).AddUnique(OuterSequenceBase);
+			}
+		}
+	}
+}
+
+void FAsyncPoseSearchDatabasesManagement::SynchronizeDatabases()
+{
+	FScopeLock Lock(&Mutex);
+
+	// copying DatabasesToSynchronize because modifying the database will call OnObjectModified that could populate DatabasesToSynchronize again
+	const TDatabasesToSynchronize DatabasesToSynchronizeCopy = DatabasesToSynchronize;
+	DatabasesToSynchronize.Reset();
+
+	for (const TDatabasesToSynchronizePair& Pair : DatabasesToSynchronizeCopy)
+	{
+		if (Pair.Key.IsValid())
+		{
+			TArray<UAnimSequenceBase*> SequencesBase;
+			for (const TWeakObjectPtr<UAnimSequenceBase>& SequenceBase : Pair.Value)
+			{
+				if (SequenceBase.IsValid())
+				{
+					SequencesBase.Add(SequenceBase.Get());
+				}
+			}
+
+			Pair.Key->SynchronizeWithExternalDependencies(SequencesBase);
+		}
+	}
+}
+
 // we're listening to OnObjectModified to cancel any pending Task indexing databases depending from Object to avoid multi threading issues
 void FAsyncPoseSearchDatabasesManagement::OnObjectModified(UObject* Object)
 {
@@ -1770,6 +1826,9 @@ void FAsyncPoseSearchDatabasesManagement::OnObjectModified(UObject* Object)
 			Tasks.RemoveAtSwap(TaskIndex, 1, false);
 		}
 	}
+
+	// collecting databases to synchronize prior modifying the Object
+	CollectDatabasesToSynchronize(Object);
 }
 
 void FAsyncPoseSearchDatabasesManagement::OnObjectTransacted(UObject* Object, const FTransactionObjectEvent& TransactionObjectEvent)
@@ -1778,37 +1837,10 @@ void FAsyncPoseSearchDatabasesManagement::OnObjectTransacted(UObject* Object, co
 
 	FScopeLock Lock(&Mutex);
 
-	// synchronize the UPoseSearchDatabase with eventual UAnimSequenceBase containing UAnimNotifyState_PoseSearchBranchIn(s)
-	UAnimSequenceBase* SequenceBaseToSynchronizeWith = nullptr;
-	UPoseSearchDatabase* Database = nullptr;
-	if (UAnimNotifyState_PoseSearchBranchIn* BranchIn = Cast<UAnimNotifyState_PoseSearchBranchIn>(Object))
-	{
-		if (BranchIn->Database)
-		{
-			SequenceBaseToSynchronizeWith = Cast<UAnimSequenceBase>(BranchIn->GetOuter());
-			Database = BranchIn->Database;
-		}
-	}
-	else if (UAnimSequenceBase* SequenceBase = Cast<UAnimSequenceBase>(Object))
-	{
-		for (const FAnimNotifyEvent& NotifyEvent : SequenceBase->Notifies)
-		{
-			if (const UAnimNotifyState_PoseSearchBranchIn* NotifyEventBranchIn = Cast<UAnimNotifyState_PoseSearchBranchIn>(NotifyEvent.NotifyStateClass))
-			{
-				if (NotifyEventBranchIn->Database)
-				{
-					SequenceBaseToSynchronizeWith = SequenceBase;
-					Database = NotifyEventBranchIn->Database;
-					break;
-				}
-			}
-		}
-	}
+	// collecting databases to synchronize on Object transacted, and merging the results with the OnObjectModified collection
+	CollectDatabasesToSynchronize(Object);
 
-	if (SequenceBaseToSynchronizeWith)
-	{
-		Database->SynchronizeWithExternalDependencies(SequenceBaseToSynchronizeWith);
-	}
+	SynchronizeDatabases();
 }
 
 void FAsyncPoseSearchDatabasesManagement::OnPackageReloaded(const EPackageReloadPhase InPackageReloadPhase, FPackageReloadedEvent* InPackageReloadedEvent)
@@ -1843,6 +1875,9 @@ void FAsyncPoseSearchDatabasesManagement::Shutdown()
 	FCoreUObjectDelegates::OnObjectModified.Remove(OnObjectModifiedHandle);
 	OnObjectModifiedHandle.Reset();
 
+	FCoreUObjectDelegates::OnObjectTransacted.Remove(OnObjectTransactedHandle);
+	OnObjectTransactedHandle.Reset();
+
 	FCoreUObjectDelegates::OnPackageReloaded.Remove(OnPackageReloadedHandle);
 	OnPackageReloadedHandle.Reset();
 }
diff --git a/Engine/Plugins/Experimental/Animation/PoseSearch/Source/Runtime/Public/PoseSearch/PoseSearchDatabase.h b/Engine/Plugins/Experimental/Animation/PoseSearch/Source/Runtime/Public/PoseSearch/PoseSearchDatabase.h
--- a/Engine/Plugins/Experimental/Animation/PoseSearch/Source/Runtime/Public/PoseSearch/PoseSearchDatabase.h
+++ b/Engine/Plugins/Experimental/Animation/PoseSearch/Source/Runtime/Public/PoseSearch/PoseSearchDatabase.h
@@ -99,6 +99,13 @@ struct POSESEARCH_API FPoseSearchDatabaseAnimationAssetBase
 	// It requires the mirror table to be set up in the config file.
 	UPROPERTY(EditAnywhere, Category = "Settings", meta = (DisplayPriority = 3))
 	EPoseSearchMirrorOption MirrorOption = EPoseSearchMirrorOption::UnmirroredOnly;
+
+#if WITH_EDITORONLY_DATA
+	// SynchronizeWithExternalDependency is true when this asset has been added via SynchronizeWithExternalDependencies.
+	// To delete it, remove the PoseSearchBranchIn notify state
+	UPROPERTY(VisibleAnywhere, Category = "Settings", meta = (DisplayPriority = 20))
+	bool bSynchronizeWithExternalDependency = false;
+#endif // WITH_EDITORONLY_DATA
 };
 
 /** A sequence entry in a UPoseSearchDatabase. */
@@ -364,7 +371,7 @@ public:
 	void NotifySynchronizeWithExternalDependencies() const { OnSynchronizeWithExternalDependencies.Broadcast(); }
 
 	void SynchronizeWithExternalDependencies();
-	void SynchronizeWithExternalDependencies(UAnimSequenceBase* SequenceBase);
+	void SynchronizeWithExternalDependencies(TConstArrayView<UAnimSequenceBase*> SequencesBase);
 #endif // WITH_EDITOR
 
 private:
diff --git a/Engine/Plugins/Experimental/Animation/PoseSearch/Source/Runtime/Public/PoseSearch/PoseSearchDerivedData.h b/Engine/Plugins/Experimental/Animation/PoseSearch/Source/Runtime/Public/PoseSearch/PoseSearchDerivedData.h
--- a/Engine/Plugins/Experimental/Animation/PoseSearch/Source/Runtime/Public/PoseSearch/PoseSearchDerivedData.h
+++ b/Engine/Plugins/Experimental/Animation/PoseSearch/Source/Runtime/Public/PoseSearch/PoseSearchDerivedData.h
@@ -7,6 +7,7 @@
 #include "UObject/GCObject.h"
 #include "UObject/ObjectSaveContext.h"
 
+class UAnimSequenceBase;
 class UPoseSearchDatabase;
 
 namespace UE::PoseSearch
@@ -57,11 +58,19 @@ namespace UE::PoseSearch
 		virtual FString GetReferencerName() const override { return TEXT("FAsyncPoseSearchDatabaseManagement"); }
 		// End FGCObject
 		
+		void CollectDatabasesToSynchronize(UObject* Object);
+		void SynchronizeDatabases();
+		
+		// map of all those databases UAnimSequenceBase(s) that were containing or are containing UAnimNotifyState_PoseSearchBranchIn that requires synchronization
+		typedef TMap<TWeakObjectPtr<UPoseSearchDatabase>, TArray<TWeakObjectPtr<UAnimSequenceBase>>> TDatabasesToSynchronize;
+		typedef TPair<TWeakObjectPtr<UPoseSearchDatabase>, TArray<TWeakObjectPtr<UAnimSequenceBase>>> TDatabasesToSynchronizePair;
+		TDatabasesToSynchronize DatabasesToSynchronize;
+
 		FPoseSearchDatabaseAsyncCacheTasks& Tasks;
 		FDelegateHandle OnObjectModifiedHandle;
 		FDelegateHandle OnObjectTransactedHandle;
 		FDelegateHandle OnPackageReloadedHandle;
-		
+
 		static FCriticalSection Mutex;
 	};
 } // namespace UE::PoseSearch
